- Improve interpolation of caches by including previous interpolation

diff --git a/Engine/Plugins/FX/Niagara/Source/Niagara/Classes/NiagaraConstants.h b/Engine/Plugins/FX/Niagara/Source/Niagara/Classes/NiagaraConstants.h
--- a/Engine/Plugins/FX/Niagara/Source/Niagara/Classes/NiagaraConstants.h
+++ b/Engine/Plugins/FX/Niagara/Source/Niagara/Classes/NiagaraConstants.h
@@ -244,6 +244,7 @@ struct NIAGARA_API FNiagaraConstants
 	static const FString SystemNamespaceString;
 	static const FString UserNamespaceString;
 	static const FString InternalNamespaceString;
+	static const FString PreviousNamespaceString;
 
 	static const int32 MaxCategoryNameLength;
 	static const int32 MaxParameterLength;
diff --git a/Engine/Plugins/FX/Niagara/Source/Niagara/Classes/NiagaraSimCache.h b/Engine/Plugins/FX/Niagara/Source/Niagara/Classes/NiagaraSimCache.h
--- a/Engine/Plugins/FX/Niagara/Source/Niagara/Classes/NiagaraSimCache.h
+++ b/Engine/Plugins/FX/Niagara/Source/Niagara/Classes/NiagaraSimCache.h
@@ -229,15 +229,19 @@ struct FNiagaraSimCacheDataBuffersLayout
 	struct FVariableCopyContext
 	{
 		float					FrameFraction		= 0.0f;
-		float					RecipDt				= 0.0f;
+		float					FrameDeltaSeconds	= 0.0f;
+		float					SimDeltaSeconds		= 0.0f;
+		float					PrevFrameFraction	= 0.0f;
 		uint32					NumInstances		= 0;
-		uint8*					Dest				= nullptr;
+		uint8*					DestCurr			= nullptr;
+		uint8*					DestPrev			= nullptr;
 		uint32					DestStride			= 0;
-		const uint8*			SourceAComponent	= nullptr;
+		const uint8*			SourceACurr			= nullptr;
+		const uint8*			SourceAPrev			= nullptr;
 		uint32					SourceAStride		= 0;
-		const uint8*			SourceBComponent	= nullptr;
+		const uint8*			SourceBCurr			= nullptr;
 		uint32					SourceBStride		= 0;
-		const uint8*			VelocityComponent	= nullptr;
+		const uint8*			Velocity			= nullptr;
 		FTransform				RebaseTransform;
 		TConstArrayView<uint32>	InterpMappings;
 	};
@@ -248,14 +252,26 @@ struct FNiagaraSimCacheDataBuffersLayout
 	{
 		FVariableCopyMapping() = default;
 		explicit FVariableCopyMapping(uint16 InComponentFrom, uint16 InComponentTo, FVariableCopyFunction InCopyFunc)
-			: ComponentFrom(InComponentFrom)
-			, ComponentTo(InComponentTo)
+			: CurrComponentFrom(InComponentFrom)
+			, PrevComponentFrom(InComponentFrom)
+			, CurrComponentTo(InComponentTo)
+			, PrevComponentTo(InComponentTo)
+			, CopyFunc(InCopyFunc)
+		{
+		}
+		explicit FVariableCopyMapping(uint16 InCurrComponentFrom, uint16 InPrevComponentFrom, uint16 InCurrComponentTo, uint16 InPrevComponentTo, FVariableCopyFunction InCopyFunc)
+			: CurrComponentFrom(InCurrComponentFrom)
+			, PrevComponentFrom(InPrevComponentFrom)
+			, CurrComponentTo(InCurrComponentTo)
+			, PrevComponentTo(InPrevComponentTo)
 			, CopyFunc(InCopyFunc)
 		{
 		}
 
-		uint16					ComponentFrom = 0;
-		uint16					ComponentTo = 0;
+		uint16					CurrComponentFrom = 0;
+		uint16					PrevComponentFrom = 0;
+		uint16					CurrComponentTo = 0;
+		uint16					PrevComponentTo = 0;
 		FVariableCopyFunction	CopyFunc;
 	};
 
@@ -292,12 +308,19 @@ struct FNiagaraSimCacheDataBuffersLayout
 	UPROPERTY()
 	TArray<FName> InterpVariableNames;
 
+	UPROPERTY()
+	uint16 ComponentVelocity = INDEX_NONE;
+
+	// Transient value used to map particles between frames for interpolation
+	uint16 ComponentUniqueID = INDEX_NONE;
+
 	TArray<uint16> ComponentMappingsToDataBuffer;
 	TArray<FVariableCopyMapping> VariableCopyMappingsToDataBuffer;
-	uint16 ComponentVelocity = INDEX_NONE;
 
 	TArray<uint16>	ComponentMappingsFromDataBuffer;
-	uint16			ComponentUniqueID = INDEX_NONE;
+
+	int32 IndexOfCacheVariable(const FNiagaraVariableBase& InVariable) const;
+	const FNiagaraSimCacheVariable* FindCacheVariable(const FNiagaraVariableBase& InVariable) const;
 };
 
 USTRUCT()
diff --git a/Engine/Plugins/FX/Niagara/Source/Niagara/Private/NiagaraConstants.cpp b/Engine/Plugins/FX/Niagara/Source/Niagara/Private/NiagaraConstants.cpp
--- a/Engine/Plugins/FX/Niagara/Source/Niagara/Private/NiagaraConstants.cpp
+++ b/Engine/Plugins/FX/Niagara/Source/Niagara/Private/NiagaraConstants.cpp
@@ -60,6 +60,7 @@ const FString FNiagaraConstants::StackContextNamespaceString(TEXT("StackContext"
 const FString FNiagaraConstants::SystemNamespaceString(TEXT("System"));
 const FString FNiagaraConstants::UserNamespaceString(TEXT("User"));
 const FString FNiagaraConstants::InternalNamespaceString(TEXT("__INTERNAL__"));
+const FString FNiagaraConstants::PreviousNamespaceString(TEXT("Previous"));
 
 const FName FNiagaraConstants::InputScopeName(TEXT("Input"));
 const FName FNiagaraConstants::OutputScopeName(TEXT("Output"));
diff --git a/Engine/Plugins/FX/Niagara/Source/Niagara/Private/NiagaraSimCache.cpp b/Engine/Plugins/FX/Niagara/Source/Niagara/Private/NiagaraSimCache.cpp
--- a/Engine/Plugins/FX/Niagara/Source/Niagara/Private/NiagaraSimCache.cpp
+++ b/Engine/Plugins/FX/Niagara/Source/Niagara/Private/NiagaraSimCache.cpp
@@ -35,6 +35,22 @@ FNiagaraSimCacheFeedbackContext::~FNiagaraSimCacheFeedbackContext()
 	}
 }
 
+int32 FNiagaraSimCacheDataBuffersLayout::IndexOfCacheVariable(const FNiagaraVariableBase& InVariable) const
+{
+	return Variables.IndexOfByPredicate(
+		[InVariable](const FNiagaraSimCacheVariable& CacheVariable)
+		{
+			return CacheVariable.Variable == InVariable;
+		}
+	);
+}
+
+const FNiagaraSimCacheVariable* FNiagaraSimCacheDataBuffersLayout::FindCacheVariable(const FNiagaraVariableBase& InVariable) const
+{
+	const int32 Index = IndexOfCacheVariable(InVariable);
+	return Index != INDEX_NONE ? &Variables[Index] : nullptr;
+}
+
 UNiagaraSimCache::UNiagaraSimCache(const FObjectInitializer& ObjectInitializer)
 {
 }
@@ -285,7 +301,7 @@ bool UNiagaraSimCache::WriteFrame(UNiagaraComponent* NiagaraComponent, FNiagaraS
 	FNiagaraSimCacheFrame& CacheFrame = CacheFrames.AddDefaulted_GetRef();
 	CacheFrame.LocalToWorld = Helper.SystemInstance->GatheredInstanceParameters.ComponentTrans;
 	CacheFrame.LWCTile = Helper.SystemInstance->GetLWCTile();
-	CacheFrame.SimulationAge = FMath::TruncToFloat(Helper.SystemInstance->GetAge() * CacheAgeResolution) / CacheAgeResolution;
+	CacheFrame.SimulationAge = FMath::RoundToFloat(Helper.SystemInstance->GetAge() * CacheAgeResolution) / CacheAgeResolution;
 
 	CacheFrame.SystemData.LocalBounds = Helper.SystemInstance->GetLocalBounds();
 
@@ -534,7 +550,7 @@ bool UNiagaraSimCache::CanRead(UNiagaraSystem* NiagaraSystem)
 bool UNiagaraSimCache::Read(float TimeSeconds, FNiagaraSystemInstance* SystemInstance) const
 {
 	// Adjust time to match our allow resolution
-	TimeSeconds = FMath::TruncToFloat(TimeSeconds * CacheAgeResolution) / CacheAgeResolution;
+	TimeSeconds = FMath::RoundToFloat(TimeSeconds * CacheAgeResolution) / CacheAgeResolution;
 
 	const float RelativeTime = FMath::Max(TimeSeconds - StartSeconds, 0.0f);
 	if ( RelativeTime > DurationSeconds )
@@ -605,9 +621,10 @@ bool UNiagaraSimCache::ReadFrame(int32 FrameIndex, float FrameFraction, FNiagara
 
 	const int32 NextFrameIndex = FMath::Min(FrameIndex + 1, CacheFrames.Num() - 1);
 	const float FrameDeltaSeconds = CacheFrames[NextFrameIndex].SimulationAge - CacheFrame.SimulationAge;
+	const float SimDeltaSeconds = Helper.SystemInstance->CachedDeltaSeconds;
 
 	Helper.SystemInstance->LocalBounds = CacheFrame.SystemData.LocalBounds;
-	Helper.ReadDataBuffer(FrameFraction, FrameDeltaSeconds, RebaseTransform, CacheLayout.SystemLayout, CacheFrame.SystemData.SystemDataBuffers, CacheFrame.SystemData.SystemDataBuffers, Helper.GetSystemSimulationDataSet());
+	Helper.ReadDataBuffer(FrameFraction, FrameDeltaSeconds, SimDeltaSeconds, RebaseTransform, CacheLayout.SystemLayout, CacheFrame.SystemData.SystemDataBuffers, CacheFrame.SystemData.SystemDataBuffers, Helper.GetSystemSimulationDataSet());
 
 	const int32 NumEmitters = CacheLayout.EmitterLayouts.Num();
 	for (int32 i=0; i < NumEmitters; ++i)
@@ -620,11 +637,11 @@ bool UNiagaraSimCache::ReadFrame(int32 FrameIndex, float FrameFraction, FNiagara
 		const FNiagaraSimCacheEmitterFrame& CacheEmitterFrameB = CacheFrames[NextFrameIndex].EmitterData[i];
 		if (CacheLayout.EmitterLayouts[i].SimTarget == ENiagaraSimTarget::GPUComputeSim)
 		{
-			Helper.ReadDataBufferGPU(FrameFraction, FrameDeltaSeconds, RebaseTransform, EmitterInstance, CacheLayout.EmitterLayouts[i], CacheEmitterFrame.ParticleDataBuffers, CacheEmitterFrameB.ParticleDataBuffers, EmitterInstance.GetData(), PendingCommandsInFlight);
+			Helper.ReadDataBufferGPU(FrameFraction, FrameDeltaSeconds, SimDeltaSeconds, RebaseTransform, EmitterInstance, CacheLayout.EmitterLayouts[i], CacheEmitterFrame.ParticleDataBuffers, CacheEmitterFrameB.ParticleDataBuffers, EmitterInstance.GetData(), PendingCommandsInFlight);
 		}
 		else
 		{
-			Helper.ReadDataBuffer(FrameFraction, FrameDeltaSeconds, RebaseTransform, CacheLayout.EmitterLayouts[i], CacheEmitterFrame.ParticleDataBuffers, CacheEmitterFrameB.ParticleDataBuffers, EmitterInstance.GetData());
+			Helper.ReadDataBuffer(FrameFraction, FrameDeltaSeconds, SimDeltaSeconds, RebaseTransform, CacheLayout.EmitterLayouts[i], CacheEmitterFrame.ParticleDataBuffers, CacheEmitterFrameB.ParticleDataBuffers, EmitterInstance.GetData());
 		}
 	}
 
diff --git a/Engine/Plugins/FX/Niagara/Source/Niagara/Private/NiagaraSimCacheHelper.h b/Engine/Plugins/FX/Niagara/Source/Niagara/Private/NiagaraSimCacheHelper.h
--- a/Engine/Plugins/FX/Niagara/Source/Niagara/Private/NiagaraSimCacheHelper.h
+++ b/Engine/Plugins/FX/Niagara/Source/Niagara/Private/NiagaraSimCacheHelper.h
@@ -13,6 +13,61 @@
 
 struct FNiagaraSimCacheHelper
 {
+	/////////////////////////////////////////////////////////////////////////////////
+	template<typename TBaseType, int32 NumComponents>
+	struct FSoAReader
+	{
+		const TBaseType* Components[NumComponents];
+		FSoAReader(const uint8* InData, const int32 InDataStride)
+		{
+			for (int32 i=0; i < NumComponents; ++i)
+			{
+				Components[i] = reinterpret_cast<const TBaseType*>(InData);
+				InData += InDataStride;
+			}
+		}
+	};
+
+	struct FSoAVec3Reader : FSoAReader<float, 3>
+	{
+		FSoAVec3Reader(const uint8* InData, const int32 InDataStride) : FSoAReader<float, 3>(InData, InDataStride) {}
+		FVector3f Get(int iInstance) const { return FVector3f(Components[0][iInstance], Components[1][iInstance], Components[2][iInstance]); }
+	};
+
+	struct FSoAQuatReader : FSoAReader<float, 4>
+	{
+		FSoAQuatReader(const uint8* InData, const int32 InDataStride) : FSoAReader<float, 4>(InData, InDataStride) {}
+		FQuat4f Get(int iInstance) const { return FQuat4f(Components[0][iInstance], Components[1][iInstance], Components[2][iInstance], Components[3][iInstance]); }
+	};
+
+	/////////////////////////////////////////////////////////////////////////////////
+	template<typename TBaseType, int32 NumComponents>
+	struct FSoAWriter
+	{
+		TBaseType* Components[NumComponents];
+		FSoAWriter(uint8* InData, const int32 InDataStride)
+		{
+			for (int32 i = 0; i < NumComponents; ++i)
+			{
+				Components[i] = reinterpret_cast<TBaseType*>(InData);
+				InData += InDataStride;
+			}
+		}
+	};
+
+	struct FSoAVec3Writer : FSoAWriter<float, 3>
+	{
+		FSoAVec3Writer(uint8* InData, const int32 InDataStride) : FSoAWriter<float, 3>(InData, InDataStride) {}
+		void Set(int iInstance, const FVector3f& Value) const { Components[0][iInstance] = Value.X; Components[1][iInstance] = Value.Y; Components[2][iInstance] = Value.Z; }
+	};
+
+	struct FSoAQuatWriter : FSoAWriter<float, 4>
+	{
+		FSoAQuatWriter(uint8* InData, const int32 InDataStride) : FSoAWriter<float, 4>(InData, InDataStride) {}
+		void Set(int iInstance, const FQuat4f& Value) const { Components[0][iInstance] = Value.X; Components[1][iInstance] = Value.Y; Components[2][iInstance] = Value.Z; Components[3][iInstance] = Value.W; }
+	};
+	/////////////////////////////////////////////////////////////////////////////////
+
 	explicit FNiagaraSimCacheHelper(FNiagaraSystemInstance* InSystemInstance)
 	{
 		SystemInstance = InSystemInstance;
@@ -69,6 +124,8 @@ struct FNiagaraSimCacheHelper
 	{
 		CacheLayout.LayoutName = LayoutName;
 		CacheLayout.SimTarget = CompiledData.SimTarget;
+		CacheLayout.ComponentVelocity = INDEX_NONE;
+		CacheLayout.ComponentUniqueID = INDEX_NONE;
 
 		// Determine the components to cache
 		int32 TotalCacheComponents = 0;
@@ -87,9 +144,9 @@ struct FNiagaraSimCacheHelper
 		}
 
 		// We need to preserve the velocity attribute if we want to use velocity based extrapolation of positions
+		const FNiagaraVariableBase VelocityVariable(FNiagaraTypeDefinition::GetVec3Def(), "Velocity");
 		if (CreateParameters.bAllowVelocityExtrapolation)
 		{
-			const FNiagaraVariableBase VelocityVariable(FNiagaraTypeDefinition::GetVec3Def(), "Velocity");
 			const int32 VelocityAttributeIndex = CompiledData.Variables.IndexOfByKey(VelocityVariable);
 			if (VelocityAttributeIndex != INDEX_NONE)
 			{
@@ -131,6 +188,14 @@ struct FNiagaraSimCacheHelper
 			CacheLayout.Int32Count += uint16(DataSetVariableLayout.GetNumInt32Components());
 		}
 
+		if (CreateParameters.bAllowVelocityExtrapolation)
+		{
+			if (const FNiagaraSimCacheVariable* VelocityCacheVariable = CacheLayout.FindCacheVariable(VelocityVariable) )
+			{
+				CacheLayout.ComponentVelocity = uint16(VelocityCacheVariable->FloatOffset);
+			}
+		}
+
 		if (CreateParameters.bAllowInterpolation)
 		{
 			const FNiagaraVariableBase UniqueIDVariable(FNiagaraTypeDefinition::GetIntDef(), "UniqueID");
@@ -370,7 +435,38 @@ struct FNiagaraSimCacheHelper
 		CacheLayout.ComponentMappingsToDataBuffer.Empty(CacheTotalComponents);
 		CacheLayout.ComponentMappingsToDataBuffer.AddDefaulted(CacheTotalComponents);
 		CacheLayout.VariableCopyMappingsToDataBuffer.Empty(0);
-		CacheLayout.ComponentVelocity = INDEX_NONE;
+
+		//-OPT: We don't have to do this all the time, figure out a better way
+		TArray<FNiagaraVariableBase> InterpCurrAttributes;
+		TArray<FNiagaraVariableBase> InterpPrevAttributes;
+		if (CacheLayout.InterpVariableNames.Num() > 0)
+		{
+			for (const FNiagaraSimCacheVariable& PreviousVariable : CacheLayout.Variables)
+			{
+				// We only support previous interpolation for types that are used in velocity calculation
+				const bool bInterpVariable =
+					CacheLayout.InterpVariableNames.Contains(PreviousVariable.Variable.GetName()) &&
+					(PreviousVariable.Variable.GetType() == FNiagaraTypeDefinition::GetPositionDef() || PreviousVariable.Variable.GetType() == FNiagaraTypeDefinition::GetQuatDef());
+				if (!bInterpVariable)
+				{
+					continue;
+				}
+
+				FNiagaraVariableBase CurrentVariable = PreviousVariable.Variable;
+				if (!CurrentVariable.RemoveRootNamespace(FNiagaraConstants::PreviousNamespaceString))
+				{
+					continue;
+				}
+
+				if (CompiledData.Variables.IndexOfByKey(PreviousVariable.Variable) == INDEX_NONE || CompiledData.Variables.IndexOfByKey(CurrentVariable) == INDEX_NONE)
+				{
+					continue;
+				}
+
+				InterpCurrAttributes.Emplace(CurrentVariable);
+				InterpPrevAttributes.Emplace(PreviousVariable.Variable);
+			}
+		}
 
 		int32 FloatOffset = 0;
 		int32 HalfOffset = CacheLayout.FloatCount;
@@ -378,7 +474,7 @@ struct FNiagaraSimCacheHelper
 		for (const FNiagaraSimCacheVariable& SourceVariable : CacheLayout.Variables)
 		{
 			// Find variable, if it doesn't exist that's ok as the cache contains more data than is required
-			const int32 DataSetVariableIndex = CompiledData.Variables.IndexOfByPredicate([&](const FNiagaraVariableBase& DataSetVariable) { return DataSetVariable == SourceVariable.Variable; });
+			const int32 DataSetVariableIndex = CompiledData.Variables.IndexOfByKey(SourceVariable.Variable);
 			const FNiagaraVariableLayoutInfo* DestVariableLayout = nullptr;
 			if (DataSetVariableIndex != INDEX_NONE)
 			{
@@ -393,16 +489,10 @@ struct FNiagaraSimCacheHelper
 				}
 			}
 
-			// If this is our velocity component track it
-			//-OPT: Don't need to do this each time
-			if (CacheLayout.bAllowVelocityExtrapolation)
+			// We won't copy these attributes here as they will be part of a combined copy step to interpolate previous
+			if (InterpPrevAttributes.Contains(SourceVariable.Variable))
 			{
-				const FNiagaraVariableBase VelocityVariable(FNiagaraTypeDefinition::GetVec3Def(), "Velocity");
-
-				if (SourceVariable.Variable == VelocityVariable)
-				{
-					CacheLayout.ComponentVelocity = uint16(FloatOffset);
-				}
+				DestVariableLayout = nullptr;
 			}
 
 			// Is this a type that requires conversion / re-basing?
@@ -411,22 +501,41 @@ struct FNiagaraSimCacheHelper
 				const bool bInterpVariable = CacheLayout.InterpVariableNames.Contains(SourceVariable.Variable.GetName());
 				const bool bRebaseVariable = CacheLayout.RebaseVariableNames.Contains(SourceVariable.Variable.GetName());
 
-				if (bInterpVariable || bRebaseVariable)
+				if (DestVariableLayout && (bInterpVariable || bRebaseVariable))
 				{
 					if (SourceVariable.Variable.GetType() == FNiagaraTypeDefinition::GetPositionDef())
 					{
 						check(SourceVariable.FloatCount == 3);
 						if (bInterpVariable)
 						{
-							CacheLayout.VariableCopyMappingsToDataBuffer.Emplace(uint16(FloatOffset), uint16(DestVariableLayout->FloatComponentStart), bRebaseVariable ? &FNiagaraSimCacheHelper::InterpPositions<true> : &FNiagaraSimCacheHelper::InterpPositions<false>);
+							const int32 InterpVarIndex = InterpCurrAttributes.IndexOfByKey(SourceVariable.Variable);
+							if (InterpVarIndex != INDEX_NONE)
+							{
+								const int32 PreviousDataSetVariableIndex = CompiledData.Variables.IndexOfByKey(InterpPrevAttributes[InterpVarIndex]);
+								check(PreviousDataSetVariableIndex != INDEX_NONE);
+
+								const FNiagaraSimCacheVariable* PrevSourceVariable = CacheLayout.FindCacheVariable(InterpPrevAttributes[InterpVarIndex]);
+								check(PrevSourceVariable != nullptr);
+
+								CacheLayout.VariableCopyMappingsToDataBuffer.Emplace(
+									SourceVariable.FloatOffset, PrevSourceVariable->FloatOffset,
+									uint16(DestVariableLayout->FloatComponentStart), uint16(CompiledData.VariableLayouts[PreviousDataSetVariableIndex].FloatComponentStart),
+									bRebaseVariable ? &FNiagaraSimCacheHelper::InterpPositions<true, true> : &FNiagaraSimCacheHelper::InterpPositions<false, true>
+									
+								);
+							}
+							else
+							{
+								CacheLayout.VariableCopyMappingsToDataBuffer.Emplace(SourceVariable.FloatOffset, uint16(DestVariableLayout->FloatComponentStart), bRebaseVariable ? &FNiagaraSimCacheHelper::InterpPositions<true, false> : &FNiagaraSimCacheHelper::InterpPositions<false, false>);
+							}
 						}
 						else if ( CacheLayout.bAllowVelocityExtrapolation )
 						{
-							CacheLayout.VariableCopyMappingsToDataBuffer.Emplace(uint16(FloatOffset), uint16(DestVariableLayout->FloatComponentStart), bRebaseVariable ? &FNiagaraSimCacheHelper::ExtrapolatePositions<true> : &FNiagaraSimCacheHelper::ExtrapolatePositions<false>);
+							CacheLayout.VariableCopyMappingsToDataBuffer.Emplace(SourceVariable.FloatOffset, uint16(DestVariableLayout->FloatComponentStart), bRebaseVariable ? &FNiagaraSimCacheHelper::ExtrapolatePositions<true> : &FNiagaraSimCacheHelper::ExtrapolatePositions<false>);
 						}
 						else
 						{
-							CacheLayout.VariableCopyMappingsToDataBuffer.Emplace(uint16(FloatOffset), uint16(DestVariableLayout->FloatComponentStart), &FNiagaraSimCacheHelper::CopyPositions);
+							CacheLayout.VariableCopyMappingsToDataBuffer.Emplace(SourceVariable.FloatOffset, uint16(DestVariableLayout->FloatComponentStart), &FNiagaraSimCacheHelper::RebasePositions);
 						}
 						DestVariableLayout = nullptr;
 					}
@@ -435,19 +544,41 @@ struct FNiagaraSimCacheHelper
 						check(SourceVariable.FloatCount == 4);
 						if (bInterpVariable)
 						{
-							CacheLayout.VariableCopyMappingsToDataBuffer.Emplace(uint16(FloatOffset), uint16(DestVariableLayout->FloatComponentStart), bInterpVariable  ? &FNiagaraSimCacheHelper::InterpQuaternions<true> : &FNiagaraSimCacheHelper::InterpQuaternions<false>);
+							const int32 InterpVarIndex = InterpCurrAttributes.IndexOfByKey(SourceVariable.Variable);
+							if (InterpVarIndex != INDEX_NONE)
+							{
+								const int32 PreviousDataSetVariableIndex = CompiledData.Variables.IndexOfByKey(InterpPrevAttributes[InterpVarIndex]);
+								check(PreviousDataSetVariableIndex != INDEX_NONE);
+
+								const FNiagaraSimCacheVariable* PrevSourceVariable = CacheLayout.FindCacheVariable(InterpPrevAttributes[InterpVarIndex]);
+								check(PrevSourceVariable != nullptr);
+
+								CacheLayout.VariableCopyMappingsToDataBuffer.Emplace(
+									SourceVariable.FloatOffset, PrevSourceVariable->FloatOffset,
+									uint16(DestVariableLayout->FloatComponentStart), uint16(CompiledData.VariableLayouts[PreviousDataSetVariableIndex].FloatComponentStart),
+									bRebaseVariable ? &FNiagaraSimCacheHelper::InterpQuaternions<true, true> : &FNiagaraSimCacheHelper::InterpQuaternions<false, true>
+									
+								);
+							}
+							else
+							{
+								CacheLayout.VariableCopyMappingsToDataBuffer.Emplace(SourceVariable.FloatOffset, uint16(DestVariableLayout->FloatComponentStart), bInterpVariable  ? &FNiagaraSimCacheHelper::InterpQuaternions<true, false> : &FNiagaraSimCacheHelper::InterpQuaternions<false, false>);
+							}
 						}
 						else
 						{
-							CacheLayout.VariableCopyMappingsToDataBuffer.Emplace(uint16(FloatOffset), uint16(DestVariableLayout->FloatComponentStart), &FNiagaraSimCacheHelper::CopyQuaternions);
+							CacheLayout.VariableCopyMappingsToDataBuffer.Emplace(SourceVariable.FloatOffset, uint16(DestVariableLayout->FloatComponentStart), &FNiagaraSimCacheHelper::RebaseQuaternions);
 						}
 						DestVariableLayout = nullptr;
 					}
 					else if (SourceVariable.Variable.GetType() == FNiagaraTypeDefinition::GetMatrix4Def())
 					{
-						check(SourceVariable.FloatCount == 16);
-						CacheLayout.VariableCopyMappingsToDataBuffer.Emplace(uint16(FloatOffset), uint16(DestVariableLayout->FloatComponentStart), &FNiagaraSimCacheHelper::CopyMatrices);
-						DestVariableLayout = nullptr;
+						if (bRebaseVariable)
+						{
+							check(SourceVariable.FloatCount == 16);
+							CacheLayout.VariableCopyMappingsToDataBuffer.Emplace(SourceVariable.FloatOffset, uint16(DestVariableLayout->FloatComponentStart), &FNiagaraSimCacheHelper::RebaseMatrices);
+							DestVariableLayout = nullptr;
+						}
 					}
 				}
 			}
@@ -595,7 +726,7 @@ struct FNiagaraSimCacheHelper
 		}
 	}
 
-	static void ReadCustomBuffers(float FrameFraction, float FrameDeltaSeconds, const FTransform& RebaseTransform, const FNiagaraSimCacheDataBuffersLayout& CacheLayout, const FNiagaraSimCacheDataBuffers& CacheBufferA, const FNiagaraSimCacheDataBuffers& CacheBufferB, uint8* DestBuffer, uint32 DestStride)
+	static void ReadCustomBuffers(float FrameFraction, float FrameDeltaSeconds, float SimDeltaSeconds, const FTransform& RebaseTransform, const FNiagaraSimCacheDataBuffersLayout& CacheLayout, const FNiagaraSimCacheDataBuffers& CacheBufferA, const FNiagaraSimCacheDataBuffers& CacheBufferB, uint8* DestBuffer, uint32 DestStride)
 	{
 		if (CacheLayout.VariableCopyMappingsToDataBuffer.Num() == 0)
 		{
@@ -603,33 +734,37 @@ struct FNiagaraSimCacheHelper
 		}
 
 		FNiagaraSimCacheDataBuffersLayout::FVariableCopyContext VariableCopyDataContext;
-		VariableCopyDataContext.FrameFraction	= FrameFraction;
-		VariableCopyDataContext.RecipDt			= FrameDeltaSeconds;
-		VariableCopyDataContext.NumInstances	= CacheBufferA.NumInstances;
-		VariableCopyDataContext.RebaseTransform	= RebaseTransform;
-		VariableCopyDataContext.InterpMappings	= CacheBufferA.InterpMapping;
-		VariableCopyDataContext.DestStride		= DestStride;
-		VariableCopyDataContext.SourceAStride	= CacheBufferA.NumInstances * sizeof(float);
-		VariableCopyDataContext.SourceBStride	= CacheBufferB.NumInstances * sizeof(float);
+		VariableCopyDataContext.FrameFraction		= FrameFraction;
+		VariableCopyDataContext.FrameDeltaSeconds	= FrameDeltaSeconds;
+		VariableCopyDataContext.SimDeltaSeconds		= SimDeltaSeconds;
+		VariableCopyDataContext.PrevFrameFraction	= 1.0f - FMath::Clamp(SimDeltaSeconds / (FrameDeltaSeconds * (1.0f + FrameFraction)), 0.0f, 1.0f);		//-TODO: We are assuming both frames have the same DT here
+		VariableCopyDataContext.NumInstances		= CacheBufferA.NumInstances;
+		VariableCopyDataContext.RebaseTransform		= RebaseTransform;
+		VariableCopyDataContext.InterpMappings		= CacheBufferA.InterpMapping;
+		VariableCopyDataContext.DestStride			= DestStride;
+		VariableCopyDataContext.SourceAStride		= CacheBufferA.NumInstances * sizeof(float);
+		VariableCopyDataContext.SourceBStride		= CacheBufferB.NumInstances * sizeof(float);
 		if (CacheLayout.ComponentVelocity != uint16(INDEX_NONE))
 		{
-			VariableCopyDataContext.VelocityComponent = CacheBufferA.FloatData.GetData() + (CacheLayout.ComponentVelocity * VariableCopyDataContext.SourceAStride);
+			VariableCopyDataContext.Velocity = CacheBufferA.FloatData.GetData() + (CacheLayout.ComponentVelocity * VariableCopyDataContext.SourceAStride);
 		}
 		else
 		{
-			VariableCopyDataContext.VelocityComponent = nullptr;
+			VariableCopyDataContext.Velocity = nullptr;
 		}
 
 		for (const FNiagaraSimCacheDataBuffersLayout::FVariableCopyMapping& VariableCopyMapping : CacheLayout.VariableCopyMappingsToDataBuffer)
 		{
-			VariableCopyDataContext.Dest				= DestBuffer + (uint32(VariableCopyMapping.ComponentTo) * DestStride);
-			VariableCopyDataContext.SourceAComponent	= CacheBufferA.FloatData.GetData() + (VariableCopyMapping.ComponentFrom * VariableCopyDataContext.SourceAStride);
-			VariableCopyDataContext.SourceBComponent	= CacheBufferB.FloatData.GetData() + (VariableCopyMapping.ComponentFrom * VariableCopyDataContext.SourceBStride);
+			VariableCopyDataContext.DestCurr	= DestBuffer + (uint32(VariableCopyMapping.CurrComponentTo) * DestStride);
+			VariableCopyDataContext.DestPrev	= DestBuffer + (uint32(VariableCopyMapping.PrevComponentTo) * DestStride);
+			VariableCopyDataContext.SourceACurr	= CacheBufferA.FloatData.GetData() + (VariableCopyMapping.CurrComponentFrom * VariableCopyDataContext.SourceAStride);
+			VariableCopyDataContext.SourceAPrev = CacheBufferA.FloatData.GetData() + (VariableCopyMapping.PrevComponentFrom * VariableCopyDataContext.SourceAStride);
+			VariableCopyDataContext.SourceBCurr	= CacheBufferB.FloatData.GetData() + (VariableCopyMapping.CurrComponentFrom * VariableCopyDataContext.SourceBStride);
 			VariableCopyMapping.CopyFunc(VariableCopyDataContext);
 		}
 	}
 
-	void ReadDataBuffer(float FrameFraction, float FrameDeltaSeconds, const FTransform& RebaseTransform, const FNiagaraSimCacheDataBuffersLayout& CacheLayout, const FNiagaraSimCacheDataBuffers& CacheBufferA, const FNiagaraSimCacheDataBuffers& CacheBufferB, FNiagaraDataSet& DataSet)
+	void ReadDataBuffer(float FrameFraction, float FrameDeltaSeconds, float SimDeltaSeconds, const FTransform& RebaseTransform, const FNiagaraSimCacheDataBuffersLayout& CacheLayout, const FNiagaraSimCacheDataBuffers& CacheBufferA, const FNiagaraSimCacheDataBuffers& CacheBufferB, FNiagaraDataSet& DataSet)
 	{
 		FNiagaraDataBuffer& DataBuffer = DataSet.BeginSimulate();
 		DataBuffer.Allocate(CacheBufferA.NumInstances);
@@ -640,7 +775,7 @@ struct FNiagaraSimCacheHelper
 			ReadFloatBuffers(iComponent, CacheLayout, CacheBufferA, MakeArrayView(DataBuffer.GetComponentPtrFloat(0), DataBuffer.GetFloatBuffer().Num()), DataBuffer.GetFloatStride());
 			ReadHalfBuffers(iComponent, CacheLayout, CacheBufferA, MakeArrayView(DataBuffer.GetComponentPtrHalf(0), DataBuffer.GetHalfBuffer().Num()), DataBuffer.GetHalfStride());
 			ReadInt32Buffers(iComponent, CacheLayout, CacheBufferA, MakeArrayView(DataBuffer.GetComponentPtrInt32(0), DataBuffer.GetInt32Buffer().Num()), DataBuffer.GetInt32Stride());
-			ReadCustomBuffers(FrameFraction, FrameDeltaSeconds, RebaseTransform, CacheLayout, CacheBufferA, CacheBufferB, DataBuffer.GetComponentPtrFloat(0), DataBuffer.GetFloatStride());
+			ReadCustomBuffers(FrameFraction, FrameDeltaSeconds, SimDeltaSeconds, RebaseTransform, CacheLayout, CacheBufferA, CacheBufferB, DataBuffer.GetComponentPtrFloat(0), DataBuffer.GetFloatStride());
 		}
 
 		//-TODO:DestinationDataBuffer.SetIDTable(CacheBufferA.IDToIndexTable);
@@ -649,7 +784,7 @@ struct FNiagaraSimCacheHelper
 		DataSet.EndSimulate();
 	}
 
-	void ReadDataBufferGPU(float InFrameFraction, float InFrameDeltaSeconds, const FTransform& InRebaseTransform, FNiagaraEmitterInstance& EmitterInstance, const FNiagaraSimCacheDataBuffersLayout& InCacheLayout, const FNiagaraSimCacheDataBuffers& InCacheBufferA, const FNiagaraSimCacheDataBuffers& InCacheBufferB, FNiagaraDataSet& InDataSet, std::atomic<int32>& InPendingCommandsCounter)
+	void ReadDataBufferGPU(float InFrameFraction, float InFrameDeltaSeconds, float InSimDeltaSeconds, const FTransform& InRebaseTransform, FNiagaraEmitterInstance& EmitterInstance, const FNiagaraSimCacheDataBuffersLayout& InCacheLayout, const FNiagaraSimCacheDataBuffers& InCacheBufferA, const FNiagaraSimCacheDataBuffers& InCacheBufferB, FNiagaraDataSet& InDataSet, std::atomic<int32>& InPendingCommandsCounter)
 	{
 		if (EmitterInstance.IsDisabled())
 		{
@@ -662,7 +797,7 @@ struct FNiagaraSimCacheHelper
 
 		FNiagaraGpuComputeDispatchInterface* DispathInterface = EmitterInstance.GetParentSystemInstance()->GetComputeDispatchInterface();
 		ENQUEUE_RENDER_COMMAND(NiagaraSimCacheGpuReadFrame)(
-			[DispathInterface, GPUExecContext=EmitterInstance.GetGPUContext(), FrameFraction=InFrameFraction, FrameDeltaSeconds=InFrameDeltaSeconds, RebaseTransform=InRebaseTransform, CacheLayout=&InCacheLayout, CacheBufferA=&InCacheBufferA, CacheBufferB=&InCacheBufferB, DataSet=&InDataSet, PendingCommandsCounter=&InPendingCommandsCounter](FRHICommandListImmediate& RHICmdList)
+			[DispathInterface, GPUExecContext=EmitterInstance.GetGPUContext(), FrameFraction=InFrameFraction, FrameDeltaSeconds=InFrameDeltaSeconds, SimDeltaSeconds=InSimDeltaSeconds, RebaseTransform=InRebaseTransform, CacheLayout=&InCacheLayout, CacheBufferA=&InCacheBufferA, CacheBufferB=&InCacheBufferB, DataSet=&InDataSet, PendingCommandsCounter=&InPendingCommandsCounter](FRHICommandListImmediate& RHICmdList)
 			{
 				const int32 NumInstances = CacheBufferA->NumInstances;
 
@@ -698,7 +833,7 @@ struct FNiagaraSimCacheHelper
 						FRWBuffer& RWBuffer = DataBuffer.GetGPUBufferFloat();
 						uint8* RWBufferMemory = reinterpret_cast<uint8*>(RHILockBuffer(RWBuffer.Buffer, 0, RWBuffer.NumBytes, RLM_WriteOnly));
 						ReadFloatBuffers(iComponent, *CacheLayout, *CacheBufferA, MakeArrayView(RWBufferMemory, RWBuffer.NumBytes), DataBuffer.GetFloatStride());
-						ReadCustomBuffers(FrameFraction, FrameDeltaSeconds, RebaseTransform, *CacheLayout, *CacheBufferA, *CacheBufferB, RWBufferMemory, DataBuffer.GetFloatStride());
+						ReadCustomBuffers(FrameFraction, FrameDeltaSeconds, SimDeltaSeconds, RebaseTransform, *CacheLayout, *CacheBufferA, *CacheBufferB, RWBufferMemory, DataBuffer.GetFloatStride());
 						RHIUnlockBuffer(RWBuffer.Buffer);
 					}
 
@@ -748,102 +883,147 @@ struct FNiagaraSimCacheHelper
 	template<bool bWithRebase>
 	static void ExtrapolatePositions(const FNiagaraSimCacheDataBuffersLayout::FVariableCopyContext& CopyDataContext)
 	{
-		float* DstPositions[3] = { reinterpret_cast<float*>(CopyDataContext.Dest + (CopyDataContext.DestStride * 0)), reinterpret_cast<float*>(CopyDataContext.Dest + (CopyDataContext.DestStride * 1)), reinterpret_cast<float*>(CopyDataContext.Dest + (CopyDataContext.DestStride * 2))};
-		const float* SrcPositions[3] = { reinterpret_cast<const float*>(CopyDataContext.SourceAComponent + (CopyDataContext.SourceAStride * 0)), reinterpret_cast<const float*>(CopyDataContext.SourceAComponent + (CopyDataContext.SourceAStride * 1)), reinterpret_cast<const float*>(CopyDataContext.SourceAComponent + (CopyDataContext.SourceAStride * 2)) };
-		const float* SrcVelocities[3] = { reinterpret_cast<const float*>(CopyDataContext.VelocityComponent + (CopyDataContext.SourceAStride * 0)), reinterpret_cast<const float*>(CopyDataContext.VelocityComponent + (CopyDataContext.SourceAStride * 1)), reinterpret_cast<const float*>(CopyDataContext.VelocityComponent + (CopyDataContext.SourceAStride * 2)) };
+		FSoAVec3Reader SrcPositions(CopyDataContext.SourceACurr, CopyDataContext.SourceAStride);
+		FSoAVec3Reader SrcVelocities(CopyDataContext.Velocity, CopyDataContext.SourceAStride);
+		FSoAVec3Writer DstPositions(CopyDataContext.DestCurr, CopyDataContext.DestStride);
 
 		for (uint32 i = 0; i < CopyDataContext.NumInstances; ++i)
 		{
-			const FVector3f CachePosition(SrcPositions[0][i], SrcPositions[1][i], SrcPositions[2][i]);
-			const FVector3f CacheVelocity(SrcVelocities[0][i], SrcVelocities[1][i], SrcVelocities[2][i]);
-			const FVector3f Position = CachePosition + (CacheVelocity * CopyDataContext.FrameFraction * CopyDataContext.RecipDt);
-			const FVector3f RebasedPosition = bWithRebase ? FVector3f(CopyDataContext.RebaseTransform.TransformPosition(FVector(Position))) : Position;
-			DstPositions[0][i] = RebasedPosition.X;
-			DstPositions[1][i] = RebasedPosition.Y;
-			DstPositions[2][i] = RebasedPosition.Z;
+			const FVector3f Position = SrcPositions.Get(i);
+			const FVector3f Velocity = SrcVelocities.Get(i);
+			const FVector3f ExtrapolatedPosition = Position + (Velocity * CopyDataContext.FrameFraction * CopyDataContext.FrameDeltaSeconds);
+			const FVector3f RebasedPosition = bWithRebase ? FVector3f(CopyDataContext.RebaseTransform.TransformPosition(FVector(ExtrapolatedPosition))) : ExtrapolatedPosition;
+			DstPositions.Set(i, RebasedPosition);
 		}
 	}
 
-	template<bool bWithRebase>
+	template<bool bWithRebase, bool bWithPrevious>
 	static void InterpPositions(const FNiagaraSimCacheDataBuffersLayout::FVariableCopyContext& CopyDataContext)
 	{
-		float* DstPositions[3] = { reinterpret_cast<float*>(CopyDataContext.Dest + (CopyDataContext.DestStride * 0)), reinterpret_cast<float*>(CopyDataContext.Dest + (CopyDataContext.DestStride * 1)), reinterpret_cast<float*>(CopyDataContext.Dest + (CopyDataContext.DestStride * 2)) };
-		const float* SrcAPositions[3] = { reinterpret_cast<const float*>(CopyDataContext.SourceAComponent + (CopyDataContext.SourceAStride * 0)), reinterpret_cast<const float*>(CopyDataContext.SourceAComponent + (CopyDataContext.SourceAStride * 1)), reinterpret_cast<const float*>(CopyDataContext.SourceAComponent + (CopyDataContext.SourceAStride * 2)) };
-		const float* SrcBPositions[3] = { reinterpret_cast<const float*>(CopyDataContext.SourceBComponent + (CopyDataContext.SourceBStride * 0)), reinterpret_cast<const float*>(CopyDataContext.SourceBComponent + (CopyDataContext.SourceBStride * 1)), reinterpret_cast<const float*>(CopyDataContext.SourceBComponent + (CopyDataContext.SourceBStride * 2)) };
 
-		for (uint32 iInstanceA=0; iInstanceA < CopyDataContext.NumInstances; ++iInstanceA)
-		{
-			const uint32 iInstanceB = CopyDataContext.InterpMappings[iInstanceA];
+		FSoAVec3Reader SrcAPositions(CopyDataContext.SourceACurr, CopyDataContext.SourceAStride);
+		FSoAVec3Reader SrcBPositions(CopyDataContext.SourceBCurr, CopyDataContext.SourceBStride);
+		FSoAVec3Writer DstPositions(CopyDataContext.DestCurr, CopyDataContext.DestStride);
 
-			const FVector3f CachePositionA(SrcAPositions[0][iInstanceA], SrcAPositions[1][iInstanceA], SrcAPositions[2][iInstanceA]);
-			const FVector3f CachePositionB = iInstanceB == INDEX_NONE ? CachePositionA : FVector3f(SrcBPositions[0][iInstanceB], SrcBPositions[1][iInstanceB], SrcBPositions[2][iInstanceB]);
-			const FVector3f Position(FMath::Lerp(CachePositionA, CachePositionB, CopyDataContext.FrameFraction));
+		if (bWithPrevious)
+		{
+			FSoAVec3Reader SrcAPrevPositions(CopyDataContext.SourceAPrev, CopyDataContext.SourceAStride);
+			FSoAVec3Writer DstPrevPositions(CopyDataContext.DestPrev, CopyDataContext.DestStride);
 
-			const FVector3f RebasedPosition = bWithRebase ? FVector3f(CopyDataContext.RebaseTransform.TransformPosition(FVector(Position))) : Position;
-			DstPositions[0][iInstanceA] = float(RebasedPosition.X);
-			DstPositions[1][iInstanceA] = float(RebasedPosition.Y);
-			DstPositions[2][iInstanceA] = float(RebasedPosition.Z);
+			for (uint32 iInstanceA = 0; iInstanceA < CopyDataContext.NumInstances; ++iInstanceA)
+			{
+				const uint32 iInstanceB = CopyDataContext.InterpMappings[iInstanceA];
+				const FVector3f CurrPosA = SrcAPositions.Get(iInstanceA);
+				const FVector3f CurrPosB = iInstanceB == INDEX_NONE ? CurrPosA : SrcBPositions.Get(iInstanceB);
+				const FVector3f CurrPos = FMath::Lerp(CurrPosA, CurrPosB, CopyDataContext.FrameFraction);
+				DstPositions.Set(iInstanceA, bWithRebase ? FVector3f(CopyDataContext.RebaseTransform.TransformPosition(FVector(CurrPos))) : CurrPos);
+
+				const FVector3f PrevPosA = SrcAPrevPositions.Get(iInstanceA);
+				const FVector3f PrevPos = FMath::Lerp(PrevPosA, CurrPosA, CopyDataContext.PrevFrameFraction);
+				DstPrevPositions.Set(iInstanceA, bWithRebase ? FVector3f(CopyDataContext.RebaseTransform.TransformPosition(FVector(PrevPos))) : PrevPos);
+			}
+		}
+		else
+		{
+			for (uint32 iInstanceA = 0; iInstanceA < CopyDataContext.NumInstances; ++iInstanceA)
+			{
+				const uint32 iInstanceB = CopyDataContext.InterpMappings[iInstanceA];
+				const FVector3f CurrPosA = SrcAPositions.Get(iInstanceA);
+				const FVector3f CurrPosB = iInstanceB == INDEX_NONE ? CurrPosA : SrcBPositions.Get(iInstanceB);
+				const FVector3f CurrPos  = FMath::Lerp(CurrPosA, CurrPosB, CopyDataContext.FrameFraction);
+				DstPositions.Set(iInstanceA, bWithRebase ? FVector3f(CopyDataContext.RebaseTransform.TransformPosition(FVector(CurrPos))) : CurrPos);
+			}
 		}
 	}
 
-	static void CopyPositions(const FNiagaraSimCacheDataBuffersLayout::FVariableCopyContext& CopyDataContext)
+	static void RebasePositions(const FNiagaraSimCacheDataBuffersLayout::FVariableCopyContext& CopyDataContext)
 	{
-		float* DstPositions[3] = { reinterpret_cast<float*>(CopyDataContext.Dest + (CopyDataContext.DestStride * 0)), reinterpret_cast<float*>(CopyDataContext.Dest + (CopyDataContext.DestStride * 1)), reinterpret_cast<float*>(CopyDataContext.Dest + (CopyDataContext.DestStride * 2)) };
-		const float* SrcPositions[3] = { reinterpret_cast<const float*>(CopyDataContext.SourceAComponent + (CopyDataContext.SourceAStride * 0)), reinterpret_cast<const float*>(CopyDataContext.SourceAComponent + (CopyDataContext.SourceAStride * 1)), reinterpret_cast<const float*>(CopyDataContext.SourceAComponent + (CopyDataContext.SourceAStride * 2)) };
+		FSoAVec3Reader SrcPositions(CopyDataContext.SourceACurr, CopyDataContext.SourceAStride);
+		FSoAVec3Writer DstPositions(CopyDataContext.DestCurr, CopyDataContext.DestStride);
 
 		for (uint32 i = 0; i < CopyDataContext.NumInstances; ++i)
 		{
-			const FVector3f CachePosition(SrcPositions[0][i], SrcPositions[1][i], SrcPositions[2][i]);
+			const FVector3f CachePosition = SrcPositions.Get(i);
 			const FVector3f RebasedPosition = FVector3f(CopyDataContext.RebaseTransform.TransformPosition(FVector(CachePosition)));
-			DstPositions[0][i] = RebasedPosition.X;
-			DstPositions[1][i] = RebasedPosition.Y;
-			DstPositions[2][i] = RebasedPosition.Z;
+			DstPositions.Set(i, RebasedPosition);
 		}
 	}
 
-	template<bool bWithRebase>
+	////template<bool bWithRebase>
+	////static void CopyVelocities(const FNiagaraSimCacheDataBuffersLayout::FVariableCopyContext& CopyDataContext)
+	////{
+	////	float* DstVelocity[3] = { reinterpret_cast<float*>(CopyDataContext.Dest + (CopyDataContext.DestStride * 0)), reinterpret_cast<float*>(CopyDataContext.Dest + (CopyDataContext.DestStride * 1)), reinterpret_cast<float*>(CopyDataContext.Dest + (CopyDataContext.DestStride * 2)) };
+	////	const float* SrcVelocity[3] = { reinterpret_cast<const float*>(CopyDataContext.SourceAComponent + (CopyDataContext.SourceAStride * 0)), reinterpret_cast<const float*>(CopyDataContext.SourceAComponent + (CopyDataContext.SourceAStride * 1)), reinterpret_cast<const float*>(CopyDataContext.SourceAComponent + (CopyDataContext.SourceAStride * 2)) };
+
+	////	for (uint32 i = 0; i < CopyDataContext.NumInstances; ++i)
+	////	{
+	////		const FVector3f CacheVelocity(SrcVelocity[0][i], SrcVelocity[1][i], SrcVelocity[2][i]);
+	////		const FVector3f RebasedVelocity = bWithRebase ? FVector3f(CopyDataContext.RebaseTransform.TransformVector(FVector(CacheVelocity))) : CacheVelocity;
+	////		DstVelocity[0][i] = RebasedVelocity.X * CopyDataContext.SimDeltaSeconds;
+	////		DstVelocity[1][i] = RebasedVelocity.Y * CopyDataContext.SimDeltaSeconds;
+	////		DstVelocity[2][i] = RebasedVelocity.Z * CopyDataContext.SimDeltaSeconds;
+	////	}
+	////}
+
+	template<bool bWithRebase, bool bWithPrevious>
 	static void InterpQuaternions(const FNiagaraSimCacheDataBuffersLayout::FVariableCopyContext& CopyDataContext)
 	{
-		float* DstQuats[4] = { reinterpret_cast<float*>(CopyDataContext.Dest + (CopyDataContext.DestStride * 0)), reinterpret_cast<float*>(CopyDataContext.Dest + (CopyDataContext.DestStride * 1)), reinterpret_cast<float*>(CopyDataContext.Dest + (CopyDataContext.DestStride * 2)), reinterpret_cast<float*>(CopyDataContext.Dest + (CopyDataContext.DestStride * 3)) };
-		const float* SrcAQuats[4] = { reinterpret_cast<const float*>(CopyDataContext.SourceAComponent + (CopyDataContext.SourceAStride * 0)), reinterpret_cast<const float*>(CopyDataContext.SourceAComponent + (CopyDataContext.SourceAStride * 1)), reinterpret_cast<const float*>(CopyDataContext.SourceAComponent + (CopyDataContext.SourceAStride * 2)), reinterpret_cast<const float*>(CopyDataContext.SourceAComponent + (CopyDataContext.SourceAStride * 3)) };
-		const float* SrcBQuats[4] = { reinterpret_cast<const float*>(CopyDataContext.SourceBComponent + (CopyDataContext.SourceBStride * 0)), reinterpret_cast<const float*>(CopyDataContext.SourceBComponent + (CopyDataContext.SourceBStride * 1)), reinterpret_cast<const float*>(CopyDataContext.SourceBComponent + (CopyDataContext.SourceBStride * 2)), reinterpret_cast<const float*>(CopyDataContext.SourceBComponent + (CopyDataContext.SourceBStride * 3)) };
-
-		for (uint32 iInstanceA=0; iInstanceA < CopyDataContext.NumInstances; ++iInstanceA)
-		{
-			const uint32 iInstanceB = CopyDataContext.InterpMappings[iInstanceA];
-			const FQuat4f CacheRotationA(SrcAQuats[0][iInstanceA], SrcAQuats[1][iInstanceA], SrcAQuats[2][iInstanceA], SrcAQuats[3][iInstanceA]);
-			const FQuat4f CacheRotationB = iInstanceB == INDEX_NONE ? CacheRotationA : FQuat4f(SrcBQuats[0][iInstanceB], SrcBQuats[1][iInstanceB], SrcBQuats[2][iInstanceB], SrcBQuats[3][iInstanceB]);
-			const FQuat4f CacheRotation(FQuat4f::Slerp(CacheRotationA, CacheRotationB, CopyDataContext.FrameFraction));
-			const FQuat4f RebasedQuat = bWithRebase ? CacheRotation * FQuat4f(CopyDataContext.RebaseTransform.GetRotation()) : CacheRotation;
-			DstQuats[0][iInstanceA] = RebasedQuat.X;
-			DstQuats[1][iInstanceA] = RebasedQuat.Y;
-			DstQuats[2][iInstanceA] = RebasedQuat.Z;
-			DstQuats[3][iInstanceA] = RebasedQuat.W;
+		FSoAQuatReader SrcAQuats(CopyDataContext.SourceACurr, CopyDataContext.SourceAStride);
+		FSoAQuatReader SrcBQuats(CopyDataContext.SourceBCurr, CopyDataContext.SourceBStride);
+		FSoAQuatWriter DstQuats(CopyDataContext.DestCurr, CopyDataContext.DestStride);
+
+		const FQuat4f RebaseQuat(CopyDataContext.RebaseTransform.GetRotation());
+		if (bWithPrevious)
+		{
+			FSoAQuatReader SrcAPrevQuats(CopyDataContext.SourceAPrev, CopyDataContext.SourceAStride);
+			FSoAQuatWriter DstPrevQuats(CopyDataContext.DestPrev, CopyDataContext.DestStride);
+
+			for (uint32 iInstanceA = 0; iInstanceA < CopyDataContext.NumInstances; ++iInstanceA)
+			{
+				const uint32 iInstanceB = CopyDataContext.InterpMappings[iInstanceA];
+				const FQuat4f CurrQuatA = SrcAQuats.Get(iInstanceA);
+				const FQuat4f CurrQuatB = iInstanceB == INDEX_NONE ? CurrQuatA : SrcBQuats.Get(iInstanceB);
+				const FQuat4f CurrQuat  = FQuat4f::Slerp(CurrQuatA, CurrQuatB, CopyDataContext.FrameFraction);
+				DstQuats.Set(iInstanceA, bWithRebase ? CurrQuat * RebaseQuat : CurrQuat);
+
+				const FQuat4f PrevQuatA = SrcAPrevQuats.Get(iInstanceA);
+				const FQuat4f PrevQuat = FQuat4f::Slerp(PrevQuatA, CurrQuat, CopyDataContext.PrevFrameFraction);
+				DstPrevQuats.Set(iInstanceA, bWithRebase ? PrevQuat * RebaseQuat : PrevQuat);
+			}
+		}
+		else
+		{
+			for (uint32 iInstanceA=0; iInstanceA < CopyDataContext.NumInstances; ++iInstanceA)
+			{
+				const uint32 iInstanceB = CopyDataContext.InterpMappings[iInstanceA];
+				const FQuat4f CurrQuatA = SrcAQuats.Get(iInstanceA);
+				const FQuat4f CurrQuatB = iInstanceB == INDEX_NONE ? CurrQuatA : SrcBQuats.Get(iInstanceB);
+				const FQuat4f CurrQuat = FQuat4f::Slerp(CurrQuatA, CurrQuatB, CopyDataContext.FrameFraction);
+				DstQuats.Set(iInstanceA, bWithRebase ? CurrQuat * RebaseQuat : CurrQuat);
+			}
 		}
 	}
 
-	static void CopyQuaternions(const FNiagaraSimCacheDataBuffersLayout::FVariableCopyContext& CopyDataContext)
+	static void RebaseQuaternions(const FNiagaraSimCacheDataBuffersLayout::FVariableCopyContext& CopyDataContext)
 	{
-		float* DstQuats[4] = { reinterpret_cast<float*>(CopyDataContext.Dest + (CopyDataContext.DestStride * 0)), reinterpret_cast<float*>(CopyDataContext.Dest + (CopyDataContext.DestStride * 1)), reinterpret_cast<float*>(CopyDataContext.Dest + (CopyDataContext.DestStride * 2)), reinterpret_cast<float*>(CopyDataContext.Dest + (CopyDataContext.DestStride * 3)) };
-		const float* SrcQuats[4] = { reinterpret_cast<const float*>(CopyDataContext.SourceAComponent + (CopyDataContext.SourceAStride * 0)), reinterpret_cast<const float*>(CopyDataContext.SourceAComponent + (CopyDataContext.SourceAStride * 1)), reinterpret_cast<const float*>(CopyDataContext.SourceAComponent + (CopyDataContext.SourceAStride * 2)), reinterpret_cast<const float*>(CopyDataContext.SourceAComponent + (CopyDataContext.SourceAStride * 3)) };
+		FSoAQuatReader SrcQuats(CopyDataContext.SourceACurr, CopyDataContext.SourceAStride);
+		FSoAQuatWriter DstQuats(CopyDataContext.DestCurr, CopyDataContext.DestStride);
 
+		const FQuat4f RebaseQuat(CopyDataContext.RebaseTransform.GetRotation());
 		for (uint32 i = 0; i < CopyDataContext.NumInstances; ++i)
 		{
-			const FQuat4f CacheRotation(SrcQuats[0][i], SrcQuats[1][i], SrcQuats[2][i], SrcQuats[3][i]);
-			const FQuat4f RebasedQuat = CacheRotation * FQuat4f(CopyDataContext.RebaseTransform.GetRotation());
-			DstQuats[0][i] = RebasedQuat.X;
-			DstQuats[1][i] = RebasedQuat.Y;
-			DstQuats[2][i] = RebasedQuat.Z;
-			DstQuats[3][i] = RebasedQuat.W;
+			const FQuat4f CacheRotation = SrcQuats.Get(i);
+			const FQuat4f RebasedQuat = CacheRotation * RebaseQuat;
+			DstQuats.Set(i, RebasedQuat);
 		}
 	}
 
-	static void CopyMatrices(const FNiagaraSimCacheDataBuffersLayout::FVariableCopyContext& CopyDataContext)
+	static void RebaseMatrices(const FNiagaraSimCacheDataBuffersLayout::FVariableCopyContext& CopyDataContext)
 	{
-		float* DstFloats = reinterpret_cast<float*>(CopyDataContext.Dest);
+		float* DstFloats = reinterpret_cast<float*>(CopyDataContext.DestCurr);
 		const uint32 DstStride = CopyDataContext.DestStride >> 2;
 
-		const float* SrcFloats = reinterpret_cast<const float*>(CopyDataContext.SourceAComponent);
+		const float* SrcFloats = reinterpret_cast<const float*>(CopyDataContext.SourceACurr);
 		const uint32 SrcStride = CopyDataContext.SourceAStride >> 2;
 
 		const FMatrix44d RebaseMatrix = CopyDataContext.RebaseTransform.ToMatrixWithScale();
