Removing FViewInfo from FX interfaces in favor of FSceneView.

diff --git a/Engine/Plugins/FX/Niagara/Source/Niagara/Private/NiagaraAsyncGpuTraceHelper.cpp b/Engine/Plugins/FX/Niagara/Source/Niagara/Private/NiagaraAsyncGpuTraceHelper.cpp
--- a/Engine/Plugins/FX/Niagara/Source/Niagara/Private/NiagaraAsyncGpuTraceHelper.cpp
+++ b/Engine/Plugins/FX/Niagara/Source/Niagara/Private/NiagaraAsyncGpuTraceHelper.cpp
@@ -116,7 +116,7 @@ void FNiagaraAsyncGpuTraceHelper::BeginFrame(FRHICommandList& RHICmdList, FNiaga
 	}
 }
 
-void FNiagaraAsyncGpuTraceHelper::PostRenderOpaque(FRHICommandList& RHICmdList, FNiagaraGpuComputeDispatchInterface* Dispatcher, TConstArrayView<FViewInfo> Views)
+void FNiagaraAsyncGpuTraceHelper::PostRenderOpaque(FRHICommandList& RHICmdList, FNiagaraGpuComputeDispatchInterface* Dispatcher, TConstStridedView<FSceneView> Views)
 {
 #if NIAGARA_ASYNC_GPU_TRACE_COLLISION_GROUPS
 	if (bCollisionGroupMapDirty)
@@ -128,7 +128,7 @@ void FNiagaraAsyncGpuTraceHelper::PostRenderOpaque(FRHICommandList& RHICmdList,
 
 	FNiagaraAsyncGpuTraceProvider::FCollisionGroupHashMap* CollisionHashMap = GetCollisionGroupHashMap();
 
-	for (const auto& TraceProvider : TraceProviders)
+	for (const TUniquePtr<FNiagaraAsyncGpuTraceProvider>& TraceProvider : TraceProviders)
 	{
 		if (TraceProvider->IsAvailable())
 		{
diff --git a/Engine/Plugins/FX/Niagara/Source/Niagara/Private/NiagaraAsyncGpuTraceHelper.h b/Engine/Plugins/FX/Niagara/Source/Niagara/Private/NiagaraAsyncGpuTraceHelper.h
--- a/Engine/Plugins/FX/Niagara/Source/Niagara/Private/NiagaraAsyncGpuTraceHelper.h
+++ b/Engine/Plugins/FX/Niagara/Source/Niagara/Private/NiagaraAsyncGpuTraceHelper.h
@@ -22,7 +22,7 @@ class FRHIUniformBuffer;
 class FRHIUnorderedAccessView;
 struct FRWBuffer;
 class FScene;
-class FViewInfo;
+class FSceneView;
 
 // we currently support assigning, from blueprint, primitives into collision groups which can be excluded
 // from ray traces.  Currently only supported with HWRT
@@ -68,7 +68,7 @@ public:
 	void Reset();
 
 	void BeginFrame(FRHICommandList& RHICmdList, FNiagaraGpuComputeDispatchInterface* Dispatcher);
-	void PostRenderOpaque(FRHICommandList& RHICmdList, FNiagaraGpuComputeDispatchInterface* Dispatcher, TConstArrayView<FViewInfo> Views);
+	void PostRenderOpaque(FRHICommandList& RHICmdList, FNiagaraGpuComputeDispatchInterface* Dispatcher, TConstStridedView<FSceneView> Views);
 	void EndFrame(FRHICommandList& RHICmdList, FNiagaraGpuComputeDispatchInterface* Dispatcher);
 
 	/** Accumulates ray requests from user DIs into single dispatches per DI. */
diff --git a/Engine/Plugins/FX/Niagara/Source/Niagara/Private/NiagaraGpuComputeDebug.cpp b/Engine/Plugins/FX/Niagara/Source/Niagara/Private/NiagaraGpuComputeDebug.cpp
--- a/Engine/Plugins/FX/Niagara/Source/Niagara/Private/NiagaraGpuComputeDebug.cpp
+++ b/Engine/Plugins/FX/Niagara/Source/Niagara/Private/NiagaraGpuComputeDebug.cpp
@@ -236,7 +236,7 @@ bool FNiagaraGpuComputeDebug::ShouldDrawDebug() const
 	return GNiagaraGpuComputeDebug_DrawDebugEnabled && (VisualizeTextures.Num() > 0);
 }
 
-void FNiagaraGpuComputeDebug::DrawDebug(FRDGBuilder& GraphBuilder, const FViewInfo& View, const FScreenPassRenderTarget& Output)
+void FNiagaraGpuComputeDebug::DrawDebug(FRDGBuilder& GraphBuilder, const FSceneView& View, const FScreenPassRenderTarget& Output)
 {
 	if (!GNiagaraGpuComputeDebug_DrawDebugEnabled || (VisualizeTextures.Num() == 0))
 	{
@@ -285,7 +285,7 @@ void FNiagaraGpuComputeDebug::DrawDebug(FRDGBuilder& GraphBuilder, const FViewIn
 	}
 }
 
-void FNiagaraGpuComputeDebug::DrawSceneDebug(class FRDGBuilder& GraphBuilder, const class FViewInfo& View, FRDGTextureRef SceneColor, FRDGTextureRef SceneDepth)
+void FNiagaraGpuComputeDebug::DrawSceneDebug(FRDGBuilder& GraphBuilder, const FSceneView& View, FRDGTextureRef SceneColor, FRDGTextureRef SceneDepth)
 {
 	if (!GNiagaraGpuComputeDebug_DrawDebugEnabled)
 	{
diff --git a/Engine/Plugins/FX/Niagara/Source/Niagara/Private/NiagaraGpuComputeDebug.h b/Engine/Plugins/FX/Niagara/Source/Niagara/Private/NiagaraGpuComputeDebug.h
--- a/Engine/Plugins/FX/Niagara/Source/Niagara/Private/NiagaraGpuComputeDebug.h
+++ b/Engine/Plugins/FX/Niagara/Source/Niagara/Private/NiagaraGpuComputeDebug.h
@@ -84,10 +84,10 @@ public:
 	bool ShouldDrawDebug() const;
 
 	// Draw all the debug information for the system
-	void DrawDebug(FRDGBuilder& GraphBuilder, const class FViewInfo& View, const struct FScreenPassRenderTarget& Output);
+	void DrawDebug(FRDGBuilder& GraphBuilder, const FSceneView& View, const struct FScreenPassRenderTarget& Output);
 
 	// Draw debug information that requires rendering into the scene
-	void DrawSceneDebug(FRDGBuilder& GraphBuilder, const class FViewInfo& View, FRDGTextureRef SceneColor, FRDGTextureRef SceneDepth);
+	void DrawSceneDebug(FRDGBuilder& GraphBuilder, const FSceneView& View, FRDGTextureRef SceneColor, FRDGTextureRef SceneDepth);
 
 private:
 	ERHIFeatureLevel::Type FeatureLevel;
diff --git a/Engine/Plugins/FX/Niagara/Source/Niagara/Private/NiagaraGpuComputeDispatch.cpp b/Engine/Plugins/FX/Niagara/Source/Niagara/Private/NiagaraGpuComputeDispatch.cpp
--- a/Engine/Plugins/FX/Niagara/Source/Niagara/Private/NiagaraGpuComputeDispatch.cpp
+++ b/Engine/Plugins/FX/Niagara/Source/Niagara/Private/NiagaraGpuComputeDispatch.cpp
@@ -37,8 +37,6 @@
 #include "TextureResource.h"
 #include "FXRenderingUtils.h"
 
-#include "ScenePrivate.h"
-
 DECLARE_CYCLE_STAT(TEXT("GPU Dispatch Setup [RT]"), STAT_NiagaraGPUDispatchSetup_RT, STATGROUP_Niagara);
 DECLARE_CYCLE_STAT(TEXT("GPU Emitter Dispatch [RT]"), STAT_NiagaraGPUSimTick_RT, STATGROUP_Niagara);
 DECLARE_CYCLE_STAT(TEXT("GPU Data Readback [RT]"), STAT_NiagaraGPUReadback_RT, STATGROUP_Niagara);
@@ -414,9 +412,9 @@ void FNiagaraGpuComputeDispatch::ProcessPendingTicksFlush(FRHICommandListImmedia
 			ViewInitOptions.ProjectionMatrix = CachedViewInitOptions.ProjectionMatrix;
 
 			GetRendererModule().CreateAndInitSingleView(RHICmdList, &ViewFamily, &ViewInitOptions);
-			const FViewInfo* DummyView = (const FViewInfo*)ViewFamily.Views[0];
 
-			TConstArrayView<FViewInfo> DummyViews = MakeArrayView(DummyView, 1);
+			// Only one element, don't need to fully stride this array
+			TConstStridedView<FSceneView> DummyViews = MakeStridedView<const FSceneView>(0, ViewFamily.Views[0], 1);
 			const bool bAllowGPUParticleUpdate = true;
 
 			// Notify that we are about to begin rendering the 'scene' this is required because some RHIs will ClearState
@@ -1016,7 +1014,7 @@ void FNiagaraGpuComputeDispatch::PrepareAllTicks(FRHICommandListImmediate& RHICm
 	}
 }
 
-void FNiagaraGpuComputeDispatch::ExecuteTicks(FRDGBuilder& GraphBuilder, TConstArrayView<FViewInfo> Views, ENiagaraGpuComputeTickStage::Type TickStage)
+void FNiagaraGpuComputeDispatch::ExecuteTicks(FRDGBuilder& GraphBuilder, TConstStridedView<FSceneView> Views, ENiagaraGpuComputeTickStage::Type TickStage)
 {
 #if WITH_MGPU
 	if (StageToWaitForGPUTransfers == TickStage)
@@ -1046,8 +1044,7 @@ void FNiagaraGpuComputeDispatch::ExecuteTicks(FRDGBuilder& GraphBuilder, TConstA
 	RDG_RHI_EVENT_SCOPE(GraphBuilder, NiagaraGpuComputeDispatch);	//-TODO:RDG: Show TickStage
 
 	// Setup Parameters that can be read from data interfaces
-	SimulationSceneViews = MakeStridedViewOfBase<const FSceneView>(Views);
-	SimulationViewInfos = Views;
+	SimulationSceneViews = Views;
 
 	if (FSceneInterface::GetShadingPath(FeatureLevel) == EShadingPath::Deferred)
 	{
@@ -1333,7 +1330,6 @@ void FNiagaraGpuComputeDispatch::ExecuteTicks(FRDGBuilder& GraphBuilder, TConstA
 
 	// Tear down for tick pass
 	SimulationSceneViews = TConstStridedView<FSceneView>();
-	SimulationViewInfos = TConstArrayView<FViewInfo>();
 	SceneTexturesUniformParams = nullptr;
 	MobileSceneTexturesUniformParams = nullptr;
 
@@ -1812,7 +1808,7 @@ void FNiagaraGpuComputeDispatch::PreInitViews(FRDGBuilder& GraphBuilder, bool bA
 				DumpDebugFrame();
 			}
 
-			ExecuteTicks(GraphBuilder, TConstArrayView<FViewInfo>(), ENiagaraGpuComputeTickStage::PreInitViews);
+			ExecuteTicks(GraphBuilder, TConstStridedView<FSceneView>(), ENiagaraGpuComputeTickStage::PreInitViews);
 		}
 	}
 	else
@@ -1844,11 +1840,11 @@ void FNiagaraGpuComputeDispatch::PreInitViews(FRDGBuilder& GraphBuilder, bool bA
 #endif
 }
 
-void FNiagaraGpuComputeDispatch::PostInitViews(FRDGBuilder& GraphBuilder, TArrayView<const class FViewInfo> Views, bool bAllowGPUParticleUpdate)
+void FNiagaraGpuComputeDispatch::PostInitViews(FRDGBuilder& GraphBuilder, TConstStridedView<FSceneView> Views, bool bAllowGPUParticleUpdate)
 {
 	LLM_SCOPE(ELLMTag::Niagara);
 
-	bAllowGPUParticleUpdate = bAllowGPUParticleUpdate && GetReferenceAllowGPUUpdate(Views);
+	bAllowGPUParticleUpdate = bAllowGPUParticleUpdate && Views.Num() > 0 && Views[0].AllowGPUParticleUpdate();
 
 	if (bAllowGPUParticleUpdate && FNiagaraUtilities::AllowGPUParticles(GetShaderPlatform()))
 	{
@@ -1873,11 +1869,11 @@ void FNiagaraGpuComputeDispatch::PostInitViews(FRDGBuilder& GraphBuilder, TArray
 	}
 }
 
-void FNiagaraGpuComputeDispatch::PostRenderOpaque(FRDGBuilder& GraphBuilder, TConstArrayView<FViewInfo> Views, bool bAllowGPUParticleUpdate)
+void FNiagaraGpuComputeDispatch::PostRenderOpaque(FRDGBuilder& GraphBuilder, TConstStridedView<FSceneView> Views, bool bAllowGPUParticleUpdate)
 {
 	LLM_SCOPE(ELLMTag::Niagara);
 
-	bAllowGPUParticleUpdate = bAllowGPUParticleUpdate && GetReferenceAllowGPUUpdate(Views);
+	bAllowGPUParticleUpdate = bAllowGPUParticleUpdate && Views.Num() > 0 && Views[0].AllowGPUParticleUpdate();
 
 	// Cache view information which will be used if we have to flush simulation commands in the future
 	if ( bAllowGPUParticleUpdate && Views.IsValidIndex(0) )
@@ -1888,7 +1884,7 @@ void FNiagaraGpuComputeDispatch::PostRenderOpaque(FRDGBuilder& GraphBuilder, TCo
 			CachedViewInitOptions.GammaCorrection	= ViewFamily->GammaCorrection;
 		}
 
-		CachedViewInitOptions.ViewRect				= Views[0].ViewRect;
+		CachedViewInitOptions.ViewRect				= UE::FXRenderingUtils::GetRawViewRectUnsafe(Views[0]);
 		CachedViewInitOptions.ViewOrigin			= Views[0].SceneViewInitOptions.ViewOrigin;
 		CachedViewInitOptions.ViewRotationMatrix	= Views[0].SceneViewInitOptions.ViewRotationMatrix;
 		CachedViewInitOptions.ProjectionMatrix		= Views[0].SceneViewInitOptions.ProjectionMatrix;
@@ -2023,7 +2019,7 @@ bool FNiagaraGpuComputeDispatch::RequiresRayTracingScene() const
 	return NumProxiesThatRequireRayTracingScene > 0;
 }
 
-void FNiagaraGpuComputeDispatch::PreRender(FRDGBuilder& GraphBuilder, TConstArrayView<FViewInfo> Views, bool bAllowGPUParticleUpdate)
+void FNiagaraGpuComputeDispatch::PreRender(FRDGBuilder& GraphBuilder, TConstStridedView<FSceneView> Views, bool bAllowGPUParticleUpdate)
 {
 	if (!FNiagaraUtilities::AllowGPUParticles(GetShaderPlatform()))
 	{
@@ -2318,7 +2314,7 @@ bool FNiagaraGpuComputeDispatch::ShouldDebugDraw_RenderThread() const
 	return false;
 }
 
-void FNiagaraGpuComputeDispatch::DrawDebug_RenderThread(class FRDGBuilder& GraphBuilder, const class FViewInfo& View, const struct FScreenPassRenderTarget& Output)
+void FNiagaraGpuComputeDispatch::DrawDebug_RenderThread(FRDGBuilder& GraphBuilder, const FSceneView& View, const struct FScreenPassRenderTarget& Output)
 {
 #if NIAGARA_COMPUTEDEBUG_ENABLED
 	if (FNiagaraGpuComputeDebug* GpuComputeDebug = GpuComputeDebugPtr.Get())
@@ -2328,7 +2324,7 @@ void FNiagaraGpuComputeDispatch::DrawDebug_RenderThread(class FRDGBuilder& Graph
 #endif
 }
 
-void FNiagaraGpuComputeDispatch::DrawSceneDebug_RenderThread(class FRDGBuilder& GraphBuilder, const class FViewInfo& View, FRDGTextureRef SceneColor, FRDGTextureRef SceneDepth)
+void FNiagaraGpuComputeDispatch::DrawSceneDebug_RenderThread(FRDGBuilder& GraphBuilder, const FSceneView& View, FRDGTextureRef SceneColor, FRDGTextureRef SceneDepth)
 {
 #if NIAGARA_COMPUTEDEBUG_ENABLED
 	if (FNiagaraGpuComputeDebug* GpuComputeDebug = GpuComputeDebugPtr.Get())
diff --git a/Engine/Plugins/FX/Niagara/Source/Niagara/Private/NiagaraGpuComputeDispatch.h b/Engine/Plugins/FX/Niagara/Source/Niagara/Private/NiagaraGpuComputeDispatch.h
--- a/Engine/Plugins/FX/Niagara/Source/Niagara/Private/NiagaraGpuComputeDispatch.h
+++ b/Engine/Plugins/FX/Niagara/Source/Niagara/Private/NiagaraGpuComputeDispatch.h
@@ -109,23 +109,23 @@ public:
 
 	virtual void DrawDebug(FCanvas* Canvas) override {}
 	virtual bool ShouldDebugDraw_RenderThread() const override;
-	virtual void DrawDebug_RenderThread(class FRDGBuilder& GraphBuilder, const class FViewInfo& View, const struct FScreenPassRenderTarget& Output) override;
-	virtual void DrawSceneDebug_RenderThread(class FRDGBuilder& GraphBuilder, const class FViewInfo& View, FRDGTextureRef SceneColor, FRDGTextureRef SceneDepth) override;
+	virtual void DrawDebug_RenderThread(FRDGBuilder& GraphBuilder, const FSceneView& View, const struct FScreenPassRenderTarget& Output) override;
+	virtual void DrawSceneDebug_RenderThread(FRDGBuilder& GraphBuilder, const FSceneView& View, FRDGTextureRef SceneColor, FRDGTextureRef SceneDepth) override;
 	virtual void AddVectorField(UVectorFieldComponent* VectorFieldComponent) override {}
 	virtual void RemoveVectorField(UVectorFieldComponent* VectorFieldComponent) override {}
 	virtual void UpdateVectorField(UVectorFieldComponent* VectorFieldComponent) override {}
 	virtual void PreInitViews(FRDGBuilder& GraphBuilder, bool bAllowGPUParticleUpdate) override;
-	virtual void PostInitViews(FRDGBuilder& GraphBuilder, TConstArrayView<FViewInfo> Views, bool bAllowGPUParticleUpdate) override;
+	virtual void PostInitViews(FRDGBuilder& GraphBuilder, TConstStridedView<FSceneView> Views, bool bAllowGPUParticleUpdate) override;
 	virtual bool UsesGlobalDistanceField() const override;
 	virtual bool UsesDepthBuffer() const override;
 	virtual bool RequiresEarlyViewUniformBuffer() const override;
 	virtual bool RequiresRayTracingScene() const override;
-	virtual void PreRender(FRDGBuilder& GraphBuilder, TConstArrayView<FViewInfo> Views, bool bAllowGPUParticleUpdate) override;
+	virtual void PreRender(FRDGBuilder& GraphBuilder, TConstStridedView<FSceneView> Views, bool bAllowGPUParticleUpdate) override;
 	virtual void OnDestroy() override; // Called on the gamethread to delete the batcher on the renderthread.
 
 	virtual void Tick(UWorld* World, float DeltaTime) override;
 
-	virtual void PostRenderOpaque(FRDGBuilder& GraphBuilder, TConstArrayView<FViewInfo> Views, bool bAllowGPUParticleUpdate) override;
+	virtual void PostRenderOpaque(FRDGBuilder& GraphBuilder, TConstStridedView<FSceneView> Views, bool bAllowGPUParticleUpdate) override;
 
 	// FNiagaraGpuComputeDispatchInterface Impl
 	virtual void FlushPendingTicks_GameThread() override;
@@ -184,7 +184,7 @@ private:
 	void UpdateInstanceCountManager(FRHICommandListImmediate& RHICmdList);
 	void PrepareTicksForProxy(FRHICommandListImmediate& RHICmdList, FNiagaraSystemGpuComputeProxy* ComputeProxy, FNiagaraGpuDispatchList& GpuDispatchList);
 	void PrepareAllTicks(FRHICommandListImmediate& RHICmdList);
-	void ExecuteTicks(FRDGBuilder& GraphBuilder, TConstArrayView<FViewInfo> Views, ENiagaraGpuComputeTickStage::Type TickStage);
+	void ExecuteTicks(FRDGBuilder& GraphBuilder, TConstStridedView<FSceneView> Views, ENiagaraGpuComputeTickStage::Type TickStage);
 	void DispatchStage(FRDGBuilder& GraphBuilder, const FNiagaraGPUSystemTick& Tick, const FNiagaraComputeInstanceData& InstanceData, const FNiagaraSimStageData& SimStageData);
 
 	/**
diff --git a/Engine/Plugins/FX/Niagara/Source/Niagara/Private/NiagaraGpuScratchPad.h b/Engine/Plugins/FX/Niagara/Source/Niagara/Private/NiagaraGpuScratchPad.h
--- a/Engine/Plugins/FX/Niagara/Source/Niagara/Private/NiagaraGpuScratchPad.h
+++ b/Engine/Plugins/FX/Niagara/Source/Niagara/Private/NiagaraGpuScratchPad.h
@@ -9,7 +9,7 @@
 #include "RHIUtilities.h"
 
 class FScene;
-class FViewInfo;
+class FSceneView;
 class FRayTracingPipelineState;
 class FRHICommandList;
 class FRHIRayTracingScene;
diff --git a/Engine/Plugins/FX/Niagara/Source/Niagara/Public/NiagaraGpuComputeDispatchInterface.h b/Engine/Plugins/FX/Niagara/Source/Niagara/Public/NiagaraGpuComputeDispatchInterface.h
--- a/Engine/Plugins/FX/Niagara/Source/Niagara/Public/NiagaraGpuComputeDispatchInterface.h
+++ b/Engine/Plugins/FX/Niagara/Source/Niagara/Public/NiagaraGpuComputeDispatchInterface.h
@@ -9,9 +9,9 @@
 #include "NiagaraGPUInstanceCountManager.h"
 #include "FXSystem.h"
 
-class RENDERCORE_API FRDGBuilder;
+class FRDGBuilder;
 class FRDGExternalAccessQueue;
-class FViewInfo;
+class FSceneView;
 
 class FNiagaraAsyncGpuTraceHelper;
 struct FNiagaraComputeExecutionContext;
@@ -79,7 +79,6 @@ public:
 	List is only valid during graph building (i.e. during ExecuteTicks) and for simulations in PostInitViews / PostRenderOpaque.
 	*/
 	TConstStridedView<FSceneView> GetSimulationSceneViews() const { return SimulationSceneViews; }
-	TConstArrayView<FViewInfo> GetSimulationViewInfos() const { return SimulationViewInfos; }
 
 	/** Get access to the instance count manager. */
 	FORCEINLINE FNiagaraGPUInstanceCountManager& GetGPUInstanceCounterManager() { check(IsInRenderingThread()); return GPUInstanceCounterManager; }
@@ -184,7 +183,6 @@ protected:
 	TArray<TPair<FName, TUniquePtr<FNiagaraGpuComputeDataManager>>> GpuDataManagers;
 
 	TConstStridedView<FSceneView>			SimulationSceneViews;
-	TConstArrayView<FViewInfo>				SimulationViewInfos;
 
 	bool									bIsOutsideSceneRenderer = false;
 };
diff --git a/Engine/Plugins/FX/Niagara/Source/NiagaraShader/Private/NiagaraAsyncGpuTraceProvider.cpp b/Engine/Plugins/FX/Niagara/Source/NiagaraShader/Private/NiagaraAsyncGpuTraceProvider.cpp
--- a/Engine/Plugins/FX/Niagara/Source/NiagaraShader/Private/NiagaraAsyncGpuTraceProvider.cpp
+++ b/Engine/Plugins/FX/Niagara/Source/NiagaraShader/Private/NiagaraAsyncGpuTraceProvider.cpp
@@ -2,6 +2,7 @@
 
 #include "NiagaraAsyncGpuTraceProvider.h"
 
+#include "Containers/StridedView.h"
 #include "GlobalShader.h"
 #include "NiagaraAsyncGpuTraceProviderGsdf.h"
 #include "NiagaraAsyncGpuTraceProviderHwrt.h"
@@ -280,3 +281,7 @@ void FNiagaraAsyncGpuTraceProvider::BuildCollisionGroupHashMap(FRHICommandList&
 	RHICmdList.Transition(FRHITransitionInfo(Result.PrimIdHashTable.UAV, ERHIAccess::UAVCompute, ERHIAccess::SRVCompute));
 	RHICmdList.Transition(FRHITransitionInfo(Result.HashToCollisionGroups.UAV, ERHIAccess::UAVCompute, ERHIAccess::SRVCompute));
 }
+
+void FNiagaraAsyncGpuTraceProvider::PostRenderOpaque(FRHICommandList& RHICmdList, TConstStridedView<FSceneView> Views, FCollisionGroupHashMap* CollisionGroupHash)
+{
+}
diff --git a/Engine/Plugins/FX/Niagara/Source/NiagaraShader/Private/NiagaraAsyncGpuTraceProviderGsdf.cpp b/Engine/Plugins/FX/Niagara/Source/NiagaraShader/Private/NiagaraAsyncGpuTraceProviderGsdf.cpp
--- a/Engine/Plugins/FX/Niagara/Source/NiagaraShader/Private/NiagaraAsyncGpuTraceProviderGsdf.cpp
+++ b/Engine/Plugins/FX/Niagara/Source/NiagaraShader/Private/NiagaraAsyncGpuTraceProviderGsdf.cpp
@@ -13,8 +13,6 @@
 #include "SceneManagement.h"
 #include "ShaderParameterStruct.h"
 
-#include "ScenePrivate.h"
-
 static int GNiagaraAsyncGpuTraceGsdfEnabled = 1;
 static FAutoConsoleVariableRef CVarNiagaraAsyncGpuTraceGsdfEnabled(
 	TEXT("fx.Niagara.AsyncGpuTrace.GlobalSdfEnabled"),
@@ -108,9 +106,9 @@ bool FNiagaraAsyncGpuTraceProviderGsdf::IsAvailable() const
 	return true;
 }
 
-void FNiagaraAsyncGpuTraceProviderGsdf::PostRenderOpaque(FRHICommandList& RHICmdList, TConstArrayView<FViewInfo> Views, FCollisionGroupHashMap* CollisionGroupHash)
+void FNiagaraAsyncGpuTraceProviderGsdf::PostRenderOpaque(FRHICommandList& RHICmdList, TConstStridedView<FSceneView> Views, FCollisionGroupHashMap* CollisionGroupHash)
 {
-	if (const FGlobalDistanceFieldParameterData* DistanceFieldData = UE::FXRenderingUtils::GetGlobalDistanceFieldParameterData(MakeStridedViewOfBase<const FSceneView>(Views)))
+	if (const FGlobalDistanceFieldParameterData* DistanceFieldData = UE::FXRenderingUtils::GetGlobalDistanceFieldParameterData(Views))
 	{
 		m_DistanceFieldData = *DistanceFieldData;
 	}
diff --git a/Engine/Plugins/FX/Niagara/Source/NiagaraShader/Private/NiagaraAsyncGpuTraceProviderGsdf.h b/Engine/Plugins/FX/Niagara/Source/NiagaraShader/Private/NiagaraAsyncGpuTraceProviderGsdf.h
--- a/Engine/Plugins/FX/Niagara/Source/NiagaraShader/Private/NiagaraAsyncGpuTraceProviderGsdf.h
+++ b/Engine/Plugins/FX/Niagara/Source/NiagaraShader/Private/NiagaraAsyncGpuTraceProviderGsdf.h
@@ -15,7 +15,7 @@ public:
 	FNiagaraAsyncGpuTraceProviderGsdf(EShaderPlatform InShaderPlatform, FNiagaraGpuComputeDispatchInterface* Dispatcher);
 
 	static bool IsSupported();
-	virtual void PostRenderOpaque(FRHICommandList& RHICmdList, TConstArrayView<FViewInfo> Views, FCollisionGroupHashMap* CollisionGroupHash) override;
+	virtual void PostRenderOpaque(FRHICommandList& RHICmdList, TConstStridedView<FSceneView> Views, FCollisionGroupHashMap* CollisionGroupHash) override;
 	virtual bool IsAvailable() const override;
 	virtual void IssueTraces(FRHICommandList& RHICmdList, const FDispatchRequest& Request, FCollisionGroupHashMap* CollisionGroupHash) override;
 	virtual void Reset() override;
diff --git a/Engine/Plugins/FX/Niagara/Source/NiagaraShader/Private/NiagaraAsyncGpuTraceProviderHwrt.cpp b/Engine/Plugins/FX/Niagara/Source/NiagaraShader/Private/NiagaraAsyncGpuTraceProviderHwrt.cpp
--- a/Engine/Plugins/FX/Niagara/Source/NiagaraShader/Private/NiagaraAsyncGpuTraceProviderHwrt.cpp
+++ b/Engine/Plugins/FX/Niagara/Source/NiagaraShader/Private/NiagaraAsyncGpuTraceProviderHwrt.cpp
@@ -16,8 +16,6 @@
 #include "Templates/Function.h"
 #include "RayTracingPayloadType.h"
 
-#include "ScenePrivate.h"
-
 static int GNiagaraAsyncGpuTraceHwrtEnabled = 1;
 static FAutoConsoleVariableRef CVarNiagaraAsyncGpuTraceHwrtEnabled(
 	TEXT("fx.Niagara.AsyncGpuTrace.HWRayTraceEnabled"),
@@ -311,7 +309,7 @@ bool FNiagaraAsyncGpuTraceProviderHwrt::IsAvailable() const
 	return UE::FXRenderingUtils::RayTracing::HasRayTracingScene(Dispatcher->GetSceneInterface());
 }
 
-void FNiagaraAsyncGpuTraceProviderHwrt::PostRenderOpaque(FRHICommandList& RHICmdList, TConstArrayView<FViewInfo> Views, FCollisionGroupHashMap* CollisionGroupHash)
+void FNiagaraAsyncGpuTraceProviderHwrt::PostRenderOpaque(FRHICommandList& RHICmdList, TConstStridedView<FSceneView> Views, FCollisionGroupHashMap* CollisionGroupHash)
 {
 	check(IsAvailable());
 	check(Views.Num() > 0);
diff --git a/Engine/Plugins/FX/Niagara/Source/NiagaraShader/Private/NiagaraAsyncGpuTraceProviderHwrt.h b/Engine/Plugins/FX/Niagara/Source/NiagaraShader/Private/NiagaraAsyncGpuTraceProviderHwrt.h
--- a/Engine/Plugins/FX/Niagara/Source/NiagaraShader/Private/NiagaraAsyncGpuTraceProviderHwrt.h
+++ b/Engine/Plugins/FX/Niagara/Source/NiagaraShader/Private/NiagaraAsyncGpuTraceProviderHwrt.h
@@ -16,7 +16,7 @@ public:
 	FNiagaraAsyncGpuTraceProviderHwrt(EShaderPlatform InShaderPlatform, FNiagaraGpuComputeDispatchInterface* Dispatcher);
 
 	static bool IsSupported();
-	virtual void PostRenderOpaque(FRHICommandList& RHICmdList, TConstArrayView<FViewInfo> Views, FCollisionGroupHashMap* CollisionGroupHash) override;
+	virtual void PostRenderOpaque(FRHICommandList& RHICmdList, TConstStridedView<FSceneView> Views, FCollisionGroupHashMap* CollisionGroupHash) override;
 	virtual bool IsAvailable() const override;
 	virtual void IssueTraces(FRHICommandList& RHICmdList, const FDispatchRequest& Request, FCollisionGroupHashMap* CollisionGroupHash) override;
 	virtual void Reset() override;
diff --git a/Engine/Plugins/FX/Niagara/Source/NiagaraShader/Private/NiagaraDebugShaders.cpp b/Engine/Plugins/FX/Niagara/Source/NiagaraShader/Private/NiagaraDebugShaders.cpp
--- a/Engine/Plugins/FX/Niagara/Source/NiagaraShader/Private/NiagaraDebugShaders.cpp
+++ b/Engine/Plugins/FX/Niagara/Source/NiagaraShader/Private/NiagaraDebugShaders.cpp
@@ -2,11 +2,11 @@
 
 #include "NiagaraDebugShaders.h"
 
+#include "FXRenderingUtils.h"
 #include "RHI.h"
 #include "ScreenRendering.h"
 #include "CommonRenderResources.h"
 #include "Modules/ModuleManager.h"
-
 #include "PipelineStateCache.h"
 #include "RenderGraphBuilder.h"
 #include "RenderGraphUtils.h"
@@ -197,20 +197,24 @@ void NiagaraDebugShaders::ClearUAV(FRDGBuilder& GraphBuilder, FRDGBufferUAVRef U
 }
 
 void NiagaraDebugShaders::DrawDebugLines(
-	class FRDGBuilder& GraphBuilder, const class FViewInfo& View, FRDGTextureRef SceneColor, FRDGTextureRef SceneDepth,
+	class FRDGBuilder& GraphBuilder, const class FSceneView& View, FRDGTextureRef SceneColor, FRDGTextureRef SceneDepth,
 	const uint32 LineInstanceCount, FRDGBufferRef LineBuffer
 )
 {
-	TShaderMapRef<FNiagaraDebugDrawLineVS> VertexShader(GetGlobalShaderMap(View.GetFeatureLevel()));
-	TShaderMapRef<FNiagaraDebugDrawLinePS> PixelShader(GetGlobalShaderMap(View.GetFeatureLevel()));
+	const FGlobalShaderMap* ShaderMap = GetGlobalShaderMap(View.GetFeatureLevel());
+
+	TShaderMapRef<FNiagaraDebugDrawLineVS> VertexShader(ShaderMap);
+	TShaderMapRef<FNiagaraDebugDrawLinePS> PixelShader(ShaderMap);
+
+	const FIntRect ViewRect = UE::FXRenderingUtils::GetRawViewRectUnsafe(View);
 
 	FNiagaraDebugDrawLineParameters* PassParameters = GraphBuilder.AllocParameters<FNiagaraDebugDrawLineParameters>();
 	PassParameters->VSParameters.View					= View.ViewUniformBuffer;
 	PassParameters->VSParameters.GpuLineBuffer			= GraphBuilder.CreateSRV(LineBuffer, PF_R32_UINT);
 
 	PassParameters->PSParameters.OutputInvResolution	= FVector2f(1.0f / View.UnscaledViewRect.Width(), 1.0f / View.UnconstrainedViewRect.Height());
-	PassParameters->PSParameters.OriginalViewRectMin	= FVector2f(View.ViewRect.Min);
-	PassParameters->PSParameters.OriginalViewSize		= FVector2f(View.ViewRect.Width(), View.ViewRect.Height());
+	PassParameters->PSParameters.OriginalViewRectMin	= FVector2f(ViewRect.Min);
+	PassParameters->PSParameters.OriginalViewSize		= FVector2f(ViewRect.Width(), ViewRect.Height());
 	PassParameters->PSParameters.OriginalBufferInvSize	= FVector2f(1.f / SceneDepth->Desc.Extent.X, 1.f / SceneDepth->Desc.Extent.Y);
 	PassParameters->PSParameters.OccludedColorScale		= GNiagaraGpuComputeDebug_OccludedLineColorScale;
 	PassParameters->PSParameters.DepthTexture			= SceneDepth;
@@ -242,10 +246,14 @@ void NiagaraDebugShaders::DrawDebugLines(
 	);
 }
 
-void NiagaraDebugShaders::DrawDebugLines(FRDGBuilder& GraphBuilder, const FViewInfo& View, FRDGTextureRef SceneColor, FRDGTextureRef SceneDepth, FRDGBufferRef ArgsBuffer, FRDGBufferRef LineBuffer)
+void NiagaraDebugShaders::DrawDebugLines(FRDGBuilder& GraphBuilder, const FSceneView& View, FRDGTextureRef SceneColor, FRDGTextureRef SceneDepth, FRDGBufferRef ArgsBuffer, FRDGBufferRef LineBuffer)
 {
-	TShaderMapRef<FNiagaraDebugDrawLineVS> VertexShader(GetGlobalShaderMap(View.GetFeatureLevel()));
-	TShaderMapRef<FNiagaraDebugDrawLinePS> PixelShader(GetGlobalShaderMap(View.GetFeatureLevel()));
+	const FGlobalShaderMap* ShaderMap = GetGlobalShaderMap(View.GetFeatureLevel());
+
+	TShaderMapRef<FNiagaraDebugDrawLineVS> VertexShader(ShaderMap);
+	TShaderMapRef<FNiagaraDebugDrawLinePS> PixelShader(ShaderMap);
+
+	const FIntRect ViewRect = UE::FXRenderingUtils::GetRawViewRectUnsafe(View);
 
 	FNiagaraDebugDrawLineIndirectParameters* PassParameters = GraphBuilder.AllocParameters<FNiagaraDebugDrawLineIndirectParameters>();
 	PassParameters->AccessIndirectDrawArgsBuffer		= ArgsBuffer;
@@ -254,8 +262,8 @@ void NiagaraDebugShaders::DrawDebugLines(FRDGBuilder& GraphBuilder, const FViewI
 	PassParameters->VSParameters.GpuLineBuffer			= GraphBuilder.CreateSRV(LineBuffer, PF_R32_UINT);
 
 	PassParameters->PSParameters.OutputInvResolution	= FVector2f(1.0f / View.UnscaledViewRect.Width(), 1.0f / View.UnconstrainedViewRect.Height());
-	PassParameters->PSParameters.OriginalViewRectMin	= FVector2f(View.ViewRect.Min);
-	PassParameters->PSParameters.OriginalViewSize		= FVector2f(View.ViewRect.Width(), View.ViewRect.Height());
+	PassParameters->PSParameters.OriginalViewRectMin	= FVector2f(ViewRect.Min);
+	PassParameters->PSParameters.OriginalViewSize		= FVector2f(ViewRect.Width(), ViewRect.Height());
 	PassParameters->PSParameters.OriginalBufferInvSize	= FVector2f(1.f / SceneDepth->Desc.Extent.X, 1.f / SceneDepth->Desc.Extent.Y);
 	PassParameters->PSParameters.OccludedColorScale		= GNiagaraGpuComputeDebug_OccludedLineColorScale;
 	PassParameters->PSParameters.DepthTexture			= SceneDepth;
@@ -288,7 +296,7 @@ void NiagaraDebugShaders::DrawDebugLines(FRDGBuilder& GraphBuilder, const FViewI
 }
 
 void NiagaraDebugShaders::VisualizeTexture(
-	class FRDGBuilder& GraphBuilder, const FViewInfo& View, const FScreenPassRenderTarget& Output,
+	class FRDGBuilder& GraphBuilder, const FSceneView& View, const FScreenPassRenderTarget& Output,
 	const FIntPoint& Location, const int32& DisplayHeight,
 	const FIntVector4& InAttributesToVisualize, FRDGTextureRef Texture, const FIntVector4& NumTextureAttributes, uint32 TickCounter,
 	const FVector2D& PreviewDisplayRange
diff --git a/Engine/Plugins/FX/Niagara/Source/NiagaraShader/Public/NiagaraAsyncGpuTraceProvider.h b/Engine/Plugins/FX/Niagara/Source/NiagaraShader/Public/NiagaraAsyncGpuTraceProvider.h
--- a/Engine/Plugins/FX/Niagara/Source/NiagaraShader/Public/NiagaraAsyncGpuTraceProvider.h
+++ b/Engine/Plugins/FX/Niagara/Source/NiagaraShader/Public/NiagaraAsyncGpuTraceProvider.h
@@ -6,9 +6,9 @@
 #include "NiagaraSettings.h"
 #include "RHIDefinitions.h"
 
-class FScene;
 class FNiagaraGpuComputeDispatchInterface;
-class FViewInfo;
+class FScene;
+class FSceneView;
 
 // mirrors structure in Engine\Plugins\FX\Niagara\Shaders\Private\NiagaraAsyncGpuTraceCommon.ush
 struct FNiagaraAsyncGpuTrace
@@ -77,9 +77,7 @@ public:
 	virtual bool IsAvailable() const = 0;
 	virtual EProviderType GetType() const = 0;
 
-	virtual void PostRenderOpaque(FRHICommandList& RHICmdList, TConstArrayView<FViewInfo> Views, FCollisionGroupHashMap* CollisionGroupHash)
-	{
-	}
+	virtual void PostRenderOpaque(FRHICommandList& RHICmdList, TConstStridedView<FSceneView> Views, FCollisionGroupHashMap* CollisionGroupHash);
 
 	virtual void IssueTraces(FRHICommandList& RHICmdList, const FDispatchRequest& Request, FCollisionGroupHashMap* CollisionGroupHash)
 	{
diff --git a/Engine/Plugins/FX/Niagara/Source/NiagaraShader/Public/NiagaraDebugShaders.h b/Engine/Plugins/FX/Niagara/Source/NiagaraShader/Public/NiagaraDebugShaders.h
--- a/Engine/Plugins/FX/Niagara/Source/NiagaraShader/Public/NiagaraDebugShaders.h
+++ b/Engine/Plugins/FX/Niagara/Source/NiagaraShader/Public/NiagaraDebugShaders.h
@@ -12,17 +12,17 @@ namespace NiagaraDebugShaders
 	NIAGARASHADER_API void ClearUAV(FRDGBuilder& GraphBuilder, FRDGBufferUAVRef UAV, FUintVector4 ClearValues, uint32 UIntsToSet);
 
 	NIAGARASHADER_API void DrawDebugLines(
-		class FRDGBuilder& GraphBuilder, const class FViewInfo& View, FRDGTextureRef SceneColor, FRDGTextureRef SceneDepth,
+		class FRDGBuilder& GraphBuilder, const class FSceneView& View, FRDGTextureRef SceneColor, FRDGTextureRef SceneDepth,
 		const uint32 LineInstanceCount, FRDGBufferRef LineBuffer
 	);
 
 	NIAGARASHADER_API void DrawDebugLines(
-		class FRDGBuilder& GraphBuilder, const class FViewInfo& View, FRDGTextureRef SceneColor, FRDGTextureRef SceneDepth,
+		class FRDGBuilder& GraphBuilder, const class FSceneView& View, FRDGTextureRef SceneColor, FRDGTextureRef SceneDepth,
 		FRDGBufferRef ArgsBuffer, FRDGBufferRef LineBuffer
 	);
 
 	NIAGARASHADER_API void VisualizeTexture(
-		class FRDGBuilder& GraphBuilder, const FViewInfo& View, const FScreenPassRenderTarget& Output,
+		class FRDGBuilder& GraphBuilder, const FSceneView& View, const FScreenPassRenderTarget& Output,
 		const FIntPoint& Location, const int32& DisplayHeight,
 		const FIntVector4& AttributesToVisualize, FRDGTextureRef Texture, const FIntVector4& NumTextureAttributes, uint32 TickCounter,
 		const FVector2D& PreviewDisplayRange = FVector2D::ZeroVector
diff --git a/Engine/Source/Runtime/Engine/Private/Particles/FXSystem.cpp b/Engine/Source/Runtime/Engine/Private/Particles/FXSystem.cpp
--- a/Engine/Source/Runtime/Engine/Private/Particles/FXSystem.cpp
+++ b/Engine/Source/Runtime/Engine/Private/Particles/FXSystem.cpp
@@ -15,6 +15,8 @@
 #include "SceneInterface.h"
 #include "SceneRendering.h" // needed for STATGROUP_CommandListMarkers
 #include "DataDrivenShaderPlatformInfo.h"
+#include "FXRenderingUtils.h"
+#include "Containers/StridedView.h"
 
 
 TMap<FName, FCreateCustomFXSystemDelegate> FFXSystemInterface::CreateCustomFXDelegates;
@@ -110,25 +112,6 @@ void FFXSystemInterface::UnregisterCustomFXSystem(const FName& InterfaceName)
 	CreateCustomFXDelegates.Remove(InterfaceName);
 }
 
-TUniformBufferRef<FViewUniformShaderParameters> FFXSystemInterface::GetReferenceViewUniformBuffer(TConstArrayView<FViewInfo> Views)
-{
-	check(Views.Num() > 0);
-	return Views[0].ViewUniformBuffer;
-}
-
-bool FFXSystemInterface::GetReferenceAllowGPUUpdate(TConstArrayView<FViewInfo> Views)
-{
-	check(Views.Num() > 0);
-	return Views[0].AllowGPUParticleUpdate();
-}
-
-const FGlobalDistanceFieldParameterData* FFXSystemInterface::GetReferenceGlobalDistanceFieldData(TConstArrayView<FViewInfo> Views)
-{
-	check(Views.Num() > 0);
-	const FViewInfo& ReferenceView = Views[0];
-	return &ReferenceView.GlobalDistanceFieldInfo.ParameterData;
-}
-
 /*------------------------------------------------------------------------------
 	FX system console variables.
 ------------------------------------------------------------------------------*/
@@ -446,7 +429,7 @@ void FFXSystem::PreInitViews(FRDGBuilder& GraphBuilder, bool bAllowGPUParticleUp
 	}
 }
 
-void FFXSystem::PostInitViews(FRDGBuilder& GraphBuilder, TConstArrayView<FViewInfo> Views, bool bAllowGPUParticleUpdate)
+void FFXSystem::PostInitViews(FRDGBuilder& GraphBuilder, TConstStridedView<FSceneView> Views, bool bAllowGPUParticleUpdate)
 {
 	// nothing to do here
 }
@@ -501,17 +484,17 @@ DECLARE_CYCLE_STAT(TEXT("FXPreRender_FinalizeCDF"), STAT_CLM_FXPreRender_Finaliz
 DECLARE_GPU_DRAWCALL_STAT(FXSystemPreRender);
 DECLARE_GPU_DRAWCALL_STAT(FXSystemPostRenderOpaque);
 
-void FFXSystem::PreRender(FRDGBuilder& GraphBuilder, TConstArrayView<FViewInfo> Views, bool bAllowGPUParticleSceneUpdate)
+void FFXSystem::PreRender(FRDGBuilder& GraphBuilder, TConstStridedView<FSceneView> Views, bool bAllowGPUParticleSceneUpdate)
 {
-	bAllowGPUParticleSceneUpdate = bAllowGPUParticleSceneUpdate && GetReferenceAllowGPUUpdate(Views);
+	bAllowGPUParticleSceneUpdate = bAllowGPUParticleSceneUpdate && Views.Num() > 0 && Views[0].AllowGPUParticleUpdate();
 
 	if (RHISupportsGPUParticles() && bAllowGPUParticleSceneUpdate)
 	{
 		RDG_GPU_STAT_SCOPE(GraphBuilder, FXSystemPreRender);
 		RDG_CSV_STAT_EXCLUSIVE_SCOPE(GraphBuilder, FXSystem);
 
-		TUniformBufferRef<FViewUniformShaderParameters> ViewUniformBuffer = GetReferenceViewUniformBuffer(Views);
-		const FGlobalDistanceFieldParameterData* GlobalDistanceFieldParameterData = GetReferenceGlobalDistanceFieldData(Views);
+		TUniformBufferRef<FViewUniformShaderParameters> ViewUniformBuffer = Views[0].ViewUniformBuffer;
+		const FGlobalDistanceFieldParameterData* GlobalDistanceFieldParameterData = &static_cast<const FViewInfo&>(Views[0]).GlobalDistanceFieldInfo.ParameterData;
 
 		AddPass(
 			GraphBuilder,
@@ -548,16 +531,16 @@ void FFXSystem::PreRender(FRDGBuilder& GraphBuilder, TConstArrayView<FViewInfo>
     }
 }
 
-void FFXSystem::PostRenderOpaque(FRDGBuilder& GraphBuilder, TConstArrayView<FViewInfo> Views, bool bAllowGPUParticleUpdate)
+void FFXSystem::PostRenderOpaque(FRDGBuilder& GraphBuilder, TConstStridedView<FSceneView> Views, bool bAllowGPUParticleUpdate)
 {
-	bAllowGPUParticleUpdate = bAllowGPUParticleUpdate && GetReferenceAllowGPUUpdate(Views);
+	bAllowGPUParticleUpdate = bAllowGPUParticleUpdate && Views.Num() > 0 && Views[0].AllowGPUParticleUpdate();
 
 	if (RHISupportsGPUParticles() && IsParticleCollisionModeSupported(GetShaderPlatform(), PCM_DepthBuffer) && bAllowGPUParticleUpdate)
 	{
 		RDG_GPU_STAT_SCOPE(GraphBuilder, FXSystemPostRenderOpaque);
 		RDG_CSV_STAT_EXCLUSIVE_SCOPE(GraphBuilder, FXSystem);
 
-		TUniformBufferRef<FViewUniformShaderParameters> ViewUniformBuffer = GetReferenceViewUniformBuffer(Views);
+		TUniformBufferRef<FViewUniformShaderParameters> ViewUniformBuffer = Views[0].ViewUniformBuffer;
 
 		AddPass(GraphBuilder, RDG_EVENT_NAME("FFXSystem::PostRenderOpaque"), 
 			[this, ViewUniformBuffer](FRHICommandListImmediate& RHICmdList)
diff --git a/Engine/Source/Runtime/Engine/Private/Particles/FXSystemPrivate.h b/Engine/Source/Runtime/Engine/Private/Particles/FXSystemPrivate.h
--- a/Engine/Source/Runtime/Engine/Private/Particles/FXSystemPrivate.h
+++ b/Engine/Source/Runtime/Engine/Private/Particles/FXSystemPrivate.h
@@ -203,13 +203,13 @@ public:
 	virtual void UpdateVectorField(UVectorFieldComponent* VectorFieldComponent) override;
 	FParticleEmitterInstance* CreateGPUSpriteEmitterInstance(FGPUSpriteEmitterInfo& EmitterInfo);
 	virtual void PreInitViews(class FRDGBuilder& GraphBuilder, bool bAllowGPUParticleUpdate) override;
-	virtual void PostInitViews(FRDGBuilder& GraphBuilder, TConstArrayView<FViewInfo> Views, bool bAllowGPUParticleUpdate) override;
+	virtual void PostInitViews(FRDGBuilder& GraphBuilder, TConstStridedView<FSceneView> Views, bool bAllowGPUParticleUpdate) override;
 	virtual bool UsesGlobalDistanceField() const override;
 	virtual bool UsesDepthBuffer() const override;
 	virtual bool RequiresEarlyViewUniformBuffer() const override;
 	virtual bool RequiresRayTracingScene() const override;
-	virtual void PreRender(FRDGBuilder& GraphBuilder, TConstArrayView<FViewInfo> Views, bool bAllowGPUParticleUpdate) override;
-	virtual void PostRenderOpaque(FRDGBuilder& GraphBuilder, TConstArrayView<FViewInfo> Views, bool bAllowGPUParticleUpdate) override;
+	virtual void PreRender(FRDGBuilder& GraphBuilder, TConstStridedView<FSceneView> Views, bool bAllowGPUParticleUpdate) override;
+	virtual void PostRenderOpaque(FRDGBuilder& GraphBuilder, TConstStridedView<FSceneView> Views, bool bAllowGPUParticleUpdate) override;
 	// End FFXSystemInterface.
 
 	/*--------------------------------------------------------------------------
diff --git a/Engine/Source/Runtime/Engine/Private/Particles/FXSystemSet.cpp b/Engine/Source/Runtime/Engine/Private/Particles/FXSystemSet.cpp
--- a/Engine/Source/Runtime/Engine/Private/Particles/FXSystemSet.cpp
+++ b/Engine/Source/Runtime/Engine/Private/Particles/FXSystemSet.cpp
@@ -5,7 +5,9 @@
 =============================================================================*/
 
 #include "FXSystemSet.h"
+#include "FXRenderingUtils.h"
 #include "GPUSortManager.h"
+#include "Containers/StridedView.h"
 
 FFXSystemSet::FFXSystemSet(FGPUSortManager* InGPUSortManager)
 	: GPUSortManager(InGPUSortManager)
@@ -79,7 +81,7 @@ bool FFXSystemSet::ShouldDebugDraw_RenderThread() const
 	return false;
 }
 
-void FFXSystemSet::DrawDebug_RenderThread(class FRDGBuilder& GraphBuilder, const class FViewInfo& View, const struct FScreenPassRenderTarget& Output)
+void FFXSystemSet::DrawDebug_RenderThread(FRDGBuilder& GraphBuilder, const FSceneView& View, const struct FScreenPassRenderTarget& Output)
 {
 	for (FFXSystemInterface* FXSystem : FXSystems)
 	{
@@ -88,7 +90,7 @@ void FFXSystemSet::DrawDebug_RenderThread(class FRDGBuilder& GraphBuilder, const
 	}
 }
 
-void FFXSystemSet::DrawSceneDebug_RenderThread(class FRDGBuilder& GraphBuilder, const class FViewInfo& View, FRDGTextureRef SceneColor, FRDGTextureRef SceneDepth)
+void FFXSystemSet::DrawSceneDebug_RenderThread(FRDGBuilder& GraphBuilder, const FSceneView& View, FRDGTextureRef SceneColor, FRDGTextureRef SceneDepth)
 {
 	for (FFXSystemInterface* FXSystem : FXSystems)
 	{
@@ -124,7 +126,7 @@ void FFXSystemSet::UpdateVectorField(UVectorFieldComponent* VectorFieldComponent
 	}
 }
 
-void FFXSystemSet::PreInitViews(class FRDGBuilder& GraphBuilder, bool bAllowGPUParticleUpdate)
+void FFXSystemSet::PreInitViews(FRDGBuilder& GraphBuilder, bool bAllowGPUParticleUpdate)
 {
 	for (FFXSystemInterface* FXSystem : FXSystems)
 	{
@@ -133,7 +135,7 @@ void FFXSystemSet::PreInitViews(class FRDGBuilder& GraphBuilder, bool bAllowGPUP
 	}
 }
 
-void FFXSystemSet::PostInitViews(class FRDGBuilder& GraphBuilder, TConstArrayView<FViewInfo> Views, bool bAllowGPUParticleUpdate)
+void FFXSystemSet::PostInitViews(FRDGBuilder& GraphBuilder, TConstStridedView<FSceneView> Views, bool bAllowGPUParticleUpdate)
 {
 	for (FFXSystemInterface* FXSystem : FXSystems)
 	{
@@ -194,7 +196,7 @@ bool FFXSystemSet::RequiresRayTracingScene() const
 	return false;
 }
 
-void FFXSystemSet::PreRender(class FRDGBuilder& GraphBuilder, TConstArrayView<FViewInfo> Views, bool bAllowGPUParticleSceneUpdate)
+void FFXSystemSet::PreRender(FRDGBuilder& GraphBuilder, TConstStridedView<FSceneView> Views, bool bAllowGPUParticleSceneUpdate)
 {
 	for (FFXSystemInterface* FXSystem : FXSystems)
 	{
@@ -212,7 +214,7 @@ void FFXSystemSet::SetSceneTexturesUniformBuffer(const TUniformBufferRef<FSceneT
 	}
 }
 
-void FFXSystemSet::PostRenderOpaque(FRDGBuilder& GraphBuilder, TConstArrayView<const class FViewInfo> Views, bool bAllowGPUParticleSceneUpdate)
+void FFXSystemSet::PostRenderOpaque(FRDGBuilder& GraphBuilder, TConstStridedView<FSceneView> Views, bool bAllowGPUParticleSceneUpdate)
 {
 	for (FFXSystemInterface* FXSystem : FXSystems)
 	{
diff --git a/Engine/Source/Runtime/Engine/Private/Particles/FXSystemSet.h b/Engine/Source/Runtime/Engine/Private/Particles/FXSystemSet.h
--- a/Engine/Source/Runtime/Engine/Private/Particles/FXSystemSet.h
+++ b/Engine/Source/Runtime/Engine/Private/Particles/FXSystemSet.h
@@ -34,19 +34,19 @@ public:
 
 	virtual void DrawDebug(FCanvas* Canvas) override;
 	virtual bool ShouldDebugDraw_RenderThread() const override;
-	virtual void DrawDebug_RenderThread(class FRDGBuilder& GraphBuilder, const class FViewInfo& View, const struct FScreenPassRenderTarget& Output) override;
-	virtual void DrawSceneDebug_RenderThread(class FRDGBuilder& GraphBuilder, const class FViewInfo& View, FRDGTextureRef SceneColor, FRDGTextureRef SceneDepth) override;
+	virtual void DrawDebug_RenderThread(FRDGBuilder& GraphBuilder, const FSceneView& View, const struct FScreenPassRenderTarget& Output) override;
+	virtual void DrawSceneDebug_RenderThread(FRDGBuilder& GraphBuilder, const FSceneView& View, FRDGTextureRef SceneColor, FRDGTextureRef SceneDepth) override;
 	virtual void AddVectorField(UVectorFieldComponent* VectorFieldComponent) override;
 	virtual void RemoveVectorField(UVectorFieldComponent* VectorFieldComponent) override;
 	virtual void UpdateVectorField(UVectorFieldComponent* VectorFieldComponent) override;
 	virtual void PreInitViews(FRDGBuilder& GraphBuilder, bool bAllowGPUParticleUpdate) override;
-	virtual void PostInitViews(FRDGBuilder& GraphBuilder, TConstArrayView<FViewInfo> Views, bool bAllowGPUParticleUpdate) override;
+	virtual void PostInitViews(FRDGBuilder& GraphBuilder, TConstStridedView<FSceneView> Views, bool bAllowGPUParticleUpdate) override;
 	virtual bool UsesGlobalDistanceField() const override;
 	virtual bool UsesDepthBuffer() const override;
 	virtual bool RequiresEarlyViewUniformBuffer() const override;
 	virtual bool RequiresRayTracingScene() const override;
-	virtual void PreRender(FRDGBuilder& GraphBuilder, TConstArrayView<FViewInfo> Views, bool bAllowGPUParticleSceneUpdate) override;
-	virtual void PostRenderOpaque(FRDGBuilder& GraphBuilder, TConstArrayView<FViewInfo> Views, bool bAllowGPUParticleSceneUpdate) override;
+	virtual void PreRender(FRDGBuilder& GraphBuilder, TConstStridedView<FSceneView> Views, bool bAllowGPUParticleSceneUpdate) override;
+	virtual void PostRenderOpaque(FRDGBuilder& GraphBuilder, TConstStridedView<FSceneView> Views, bool bAllowGPUParticleSceneUpdate) override;
 
 	virtual void OnDestroy() override;
 	virtual void DestroyGPUSimulation() override;
diff --git a/Engine/Source/Runtime/Engine/Public/FXSystem.h b/Engine/Source/Runtime/Engine/Public/FXSystem.h
--- a/Engine/Source/Runtime/Engine/Public/FXSystem.h
+++ b/Engine/Source/Runtime/Engine/Public/FXSystem.h
@@ -207,12 +207,18 @@ public:
 	 * @param View The view we are rendering for
 	 * @param Output The output buffer information
 	 */
-	virtual void DrawDebug_RenderThread(class FRDGBuilder& GraphBuilder, const class FViewInfo& View, const struct FScreenPassRenderTarget& Output) {}
+	virtual void DrawDebug_RenderThread(FRDGBuilder& GraphBuilder, const FSceneView& View, const struct FScreenPassRenderTarget& Output) {}
+	
+	UE_DEPRECATED(5.3, "Passing FViewInfo is deprecated in favor of FSceneView. This function will do nothing.")
+	inline void DrawDebug_RenderThread(FRDGBuilder& GraphBuilder, const FViewInfo& View, const struct FScreenPassRenderTarget& Output) {}
 
 	/**
 	 * Call to handle debug drawing to the scene (i.e. where depth is available)
 	 */
-	virtual void DrawSceneDebug_RenderThread(class FRDGBuilder& GraphBuilder, const class FViewInfo& View, FRDGTextureRef SceneColor, FRDGTextureRef SceneDepth) {}
+	virtual void DrawSceneDebug_RenderThread(FRDGBuilder& GraphBuilder, const FSceneView& View, FRDGTextureRef SceneColor, FRDGTextureRef SceneDepth) {}
+
+	UE_DEPRECATED(5.3, "Passing FViewInfo is deprecated in favor of FSceneView. This function will do nothing.")
+	inline void DrawSceneDebug_RenderThread(FRDGBuilder& GraphBuilder, const FViewInfo& View, FRDGTextureRef SceneColor, FRDGTextureRef SceneDepth) {}
 
 	/**
 	 * Add a vector field to the FX system.
@@ -238,7 +244,10 @@ public:
 	 */
 	virtual void PreInitViews(FRDGBuilder& GraphBuilder, bool bAllowGPUParticleUpdate) = 0;
 
-	virtual void PostInitViews(FRDGBuilder& GraphBuilder, TConstArrayView<FViewInfo> Views, bool bAllowGPUParticleUpdate) = 0;
+	virtual void PostInitViews(FRDGBuilder& GraphBuilder, TConstStridedView<FSceneView> Views, bool bAllowGPUParticleUpdate) = 0;
+
+	UE_DEPRECATED(5.3, "Passing an array of FViewInfo is deprecated in favor of FSceneView. This function will do nothing.")
+	inline void PostInitViews(FRDGBuilder& GraphBuilder, TConstArrayView<FViewInfo> Views, bool bAllowGPUParticleUpdate) {}
 
 	virtual bool UsesGlobalDistanceField() const = 0;
 
@@ -252,12 +261,18 @@ public:
 	 * Notification from the renderer that it is about to draw FX belonging to
 	 * this system.
 	 */
-	virtual void PreRender(FRDGBuilder& GraphBuilder, TConstArrayView<FViewInfo> Views, bool bAllowGPUParticleUpdate) = 0;
+	virtual void PreRender(FRDGBuilder& GraphBuilder, TConstStridedView<FSceneView> Views, bool bAllowGPUParticleUpdate) = 0;
+
+	UE_DEPRECATED(5.3, "Passing an array of FViewInfo is deprecated in favor of FSceneView. This function will do nothing.")
+	inline void PreRender(FRDGBuilder& GraphBuilder, TConstArrayView<FViewInfo> Views, bool bAllowGPUParticleUpdate) {}
 
 	/**
 	 * Notification from the renderer that opaque primitives have rendered.
 	 */
-	virtual void PostRenderOpaque(FRDGBuilder& GraphBuilder, TConstArrayView<FViewInfo> Views, bool bAllowGPUParticleUpdate) = 0;
+	virtual void PostRenderOpaque(FRDGBuilder& GraphBuilder, TConstStridedView<FSceneView> Views, bool bAllowGPUParticleUpdate) = 0;
+
+	UE_DEPRECATED(5.3, "Passing an array of FViewInfo is deprecated in favor of FSceneView. This function will do nothing.")
+	inline void PostRenderOpaque(FRDGBuilder& GraphBuilder, TConstArrayView<FViewInfo> Views, bool bAllowGPUParticleUpdate) {}
 
 	bool IsPendingKill() const { return bIsPendingKill; }
 
@@ -279,11 +294,6 @@ protected:
 	/** By making the destructor protected, an instance must be destroyed via FFXSystemInterface::Destroy. */
 	ENGINE_API virtual ~FFXSystemInterface() {}
 
-	
-	ENGINE_API static TUniformBufferRef<FViewUniformShaderParameters> GetReferenceViewUniformBuffer(TConstArrayView<FViewInfo> Views);
-	ENGINE_API static bool GetReferenceAllowGPUUpdate(TConstArrayView<FViewInfo> Views);
-	ENGINE_API static const FGlobalDistanceFieldParameterData* GetReferenceGlobalDistanceFieldData(TConstArrayView<FViewInfo> Views);
-
 private:
 
 	bool bIsPendingKill = false;
diff --git a/Engine/Source/Runtime/Renderer/Private/DeferredShadingRenderer.cpp b/Engine/Source/Runtime/Renderer/Private/DeferredShadingRenderer.cpp
--- a/Engine/Source/Runtime/Renderer/Private/DeferredShadingRenderer.cpp
+++ b/Engine/Source/Runtime/Renderer/Private/DeferredShadingRenderer.cpp
@@ -623,7 +623,7 @@ END_SHADER_PARAMETER_STRUCT()
 
 static void RenderOpaqueFX(
 	FRDGBuilder& GraphBuilder,
-	TArrayView<const FViewInfo> Views,
+	TConstStridedView<FSceneView> Views,
 	FFXSystemInterface* FXSystem,
 	TRDGUniformBufferRef<FSceneTextureUniformParameters> SceneTexturesUniformBuffer)
 {
@@ -2802,7 +2802,7 @@ void FDeferredShadingSceneRenderer::Render(FRDGBuilder& GraphBuilder)
 	{
 		SCOPE_CYCLE_COUNTER(STAT_FDeferredShadingSceneRenderer_FXSystem_PreRender);
 		GraphBuilder.SetCommandListStat(GET_STATID(STAT_CLM_FXPreRender));
-		FXSystem->PreRender(GraphBuilder, Views, true /*bAllowGPUParticleUpdate*/);
+		FXSystem->PreRender(GraphBuilder, GetSceneViews(), true /*bAllowGPUParticleUpdate*/);
 		if (FGPUSortManager* GPUSortManager = FXSystem->GetGPUSortManager())
 		{
 			GPUSortManager->OnPreRender(GraphBuilder);
@@ -3927,7 +3927,7 @@ void FDeferredShadingSceneRenderer::Render(FRDGBuilder& GraphBuilder)
 
 	FRDGTextureRef ExposureIlluminance = AddCalculateExposureIlluminancePass(GraphBuilder, Views, SceneTextures, TranslucencyLightingVolumeTextures, ExposureIlluminanceSetup);
 
-	RenderOpaqueFX(GraphBuilder, Views, FXSystem, SceneTextures.UniformBuffer);
+	RenderOpaqueFX(GraphBuilder, GetSceneViews(), FXSystem, SceneTextures.UniformBuffer);
 
 	FRendererModule& RendererModule = static_cast<FRendererModule&>(GetRendererModule());
 	RendererModule.RenderPostOpaqueExtensions(GraphBuilder, Views, SceneTextures);
diff --git a/Engine/Source/Runtime/Renderer/Private/FXRenderingUtils.cpp b/Engine/Source/Runtime/Renderer/Private/FXRenderingUtils.cpp
--- a/Engine/Source/Runtime/Renderer/Private/FXRenderingUtils.cpp
+++ b/Engine/Source/Runtime/Renderer/Private/FXRenderingUtils.cpp
@@ -2,6 +2,7 @@
 
 #include "FXRenderingUtils.h"
 
+#include "Containers/StridedView.h"
 #include "DistanceFieldLightingShared.h"
 #include "Lumen/LumenScreenProbeGather.h"
 #include "MaterialShared.h"
@@ -10,6 +11,16 @@
 #include "ScenePrivate.h"
 #include "SystemTextures.h"
 
+TConstStridedView<FSceneView> UE::FXRenderingUtils::ConvertViewArray(TConstArrayView<FViewInfo> Views)
+{
+	return MakeStridedViewOfBase<const FSceneView>(Views);
+}
+
+FIntRect UE::FXRenderingUtils::GetRawViewRectUnsafe(const FSceneView& View)
+{
+	return static_cast<const FViewInfo&>(View).ViewRect;
+}
+
 bool UE::FXRenderingUtils::CanMaterialRenderBeforeFXPostOpaque(
 	const FSceneViewFamily& ViewFamily,
 	const FPrimitiveSceneProxy& SceneProxy,
diff --git a/Engine/Source/Runtime/Renderer/Private/MobileShadingRenderer.cpp b/Engine/Source/Runtime/Renderer/Private/MobileShadingRenderer.cpp
--- a/Engine/Source/Runtime/Renderer/Private/MobileShadingRenderer.cpp
+++ b/Engine/Source/Runtime/Renderer/Private/MobileShadingRenderer.cpp
@@ -169,7 +169,7 @@ static FMobileCustomDepthStencilUsage GetCustomDepthStencilUsage(const FViewInfo
 
 static void RenderOpaqueFX(
 	FRDGBuilder& GraphBuilder,
-	TArrayView<const FViewInfo> Views,
+	TConstStridedView<FSceneView> Views,
 	FFXSystemInterface* FXSystem,
 	TRDGUniformBufferRef<FMobileSceneTextureUniformParameters> MobileSceneTexturesUniformBuffer)
 {
@@ -465,7 +465,7 @@ void FMobileSceneRenderer::InitViews(FRDGBuilder& GraphBuilder, FSceneTexturesCo
 		// This is to init the ViewUniformBuffer before rendering for the Niagara compute shader.
 		// This needs to run before ComputeViewVisibility() is called, but the views normally initialize the ViewUniformBuffer after that (at the end of this method).
 		Views[0].InitRHIResources();
-		FXSystem->PostInitViews(GraphBuilder, Views, !ViewFamily.EngineShowFlags.HitProxies);
+		FXSystem->PostInitViews(GraphBuilder, GetSceneViews(), !ViewFamily.EngineShowFlags.HitProxies);
 	}
 	ComputeViewVisibility(RHICmdList, BasePassDepthStencilAccess, ViewCommandsPerView, DynamicIndexBuffer, DynamicVertexBuffer, DynamicReadBuffer, InstanceCullingManager, VirtualTextureUpdater, ComputeViewVisibilityCallbacks);
 	PostVisibilityFrameSetup(ILCTaskData);
@@ -959,7 +959,7 @@ void FMobileSceneRenderer::Render(FRDGBuilder& GraphBuilder)
 	// Notify the FX system that the scene is about to be rendered.
 	if (FXSystem && ViewFamily.EngineShowFlags.Particles)
 	{
-		FXSystem->PreRender(GraphBuilder, Views, true /*bAllowGPUParticleUpdate*/);
+		FXSystem->PreRender(GraphBuilder, GetSceneViews(), true /*bAllowGPUParticleUpdate*/);
 		if (FGPUSortManager* GPUSortManager = FXSystem->GetGPUSortManager())
 		{
 			GPUSortManager->OnPreRender(GraphBuilder);
@@ -1075,7 +1075,7 @@ void FMobileSceneRenderer::Render(FRDGBuilder& GraphBuilder)
 	FRendererModule& RendererModule = static_cast<FRendererModule&>(GetRendererModule());
 	RendererModule.RenderPostOpaqueExtensions(GraphBuilder, Views, SceneTextures);
 
-	RenderOpaqueFX(GraphBuilder, Views, FXSystem, SceneTextures.MobileUniformBuffer);
+	RenderOpaqueFX(GraphBuilder, GetSceneViews(), FXSystem, SceneTextures.MobileUniformBuffer);
 
 	if (bRequiresPixelProjectedPlanarRelfectionPass)
 	{
diff --git a/Engine/Source/Runtime/Renderer/Private/PostProcess/PostProcessing.cpp b/Engine/Source/Runtime/Renderer/Private/PostProcess/PostProcessing.cpp
--- a/Engine/Source/Runtime/Renderer/Private/PostProcess/PostProcessing.cpp
+++ b/Engine/Source/Runtime/Renderer/Private/PostProcess/PostProcessing.cpp
@@ -1466,7 +1466,7 @@ void AddPostProcessingPasses(
 		{
 			if (FFXSystemInterface* FXSystem = View.Family->Scene->GetFXSystem())
 			{
-				FXSystem->DrawSceneDebug_RenderThread(GraphBuilder, View, SceneColor.Texture, SceneDepth.Texture);
+				FXSystem->DrawSceneDebug_RenderThread(GraphBuilder, (const FSceneView&)View, SceneColor.Texture, SceneDepth.Texture);
 			}
 		}
 	}
diff --git a/Engine/Source/Runtime/Renderer/Private/SceneHitProxyRendering.cpp b/Engine/Source/Runtime/Renderer/Private/SceneHitProxyRendering.cpp
--- a/Engine/Source/Runtime/Renderer/Private/SceneHitProxyRendering.cpp
+++ b/Engine/Source/Runtime/Renderer/Private/SceneHitProxyRendering.cpp
@@ -683,14 +683,14 @@ void FDeferredShadingSceneRenderer::RenderHitProxies(FRDGBuilder& GraphBuilder)
 	if (FXSystem && Views.IsValidIndex(0))
 	{
 		FGPUSortManager* GPUSortManager = FXSystem->GetGPUSortManager();
-		FXSystem->PreRender(GraphBuilder, Views, false);
+		FXSystem->PreRender(GraphBuilder, GetSceneViews(), false);
 		if (GPUSortManager)
 		{
 			GPUSortManager->OnPreRender(GraphBuilder);
 		}
 		// Call PostRenderOpaque now as this is irrelevant for when rendering hit proxies.
 		// because we don't tick the particles in the render loop (see last param being "false").
-		FXSystem->PostRenderOpaque(GraphBuilder, Views, false /*bAllowGPUParticleUpdate*/);
+		FXSystem->PostRenderOpaque(GraphBuilder, GetSceneViews(), false /*bAllowGPUParticleUpdate*/);
 		if (GPUSortManager)
 		{
 			GPUSortManager->OnPostRenderOpaque(GraphBuilder);
diff --git a/Engine/Source/Runtime/Renderer/Private/SceneRendering.cpp b/Engine/Source/Runtime/Renderer/Private/SceneRendering.cpp
--- a/Engine/Source/Runtime/Renderer/Private/SceneRendering.cpp
+++ b/Engine/Source/Runtime/Renderer/Private/SceneRendering.cpp
@@ -3619,7 +3619,7 @@ void FSceneRenderer::RenderFinish(FRDGBuilder& GraphBuilder, FRDGTextureRef View
 					});
 					if (bFxDebugDraw)
 					{
-						FXInterface->DrawDebug_RenderThread(GraphBuilder, View, Output);
+						FXInterface->DrawDebug_RenderThread(GraphBuilder, (const FSceneView&)View, Output);
 					}
 				}
 			}
diff --git a/Engine/Source/Runtime/Renderer/Private/SceneRendering.h b/Engine/Source/Runtime/Renderer/Private/SceneRendering.h
--- a/Engine/Source/Runtime/Renderer/Private/SceneRendering.h
+++ b/Engine/Source/Runtime/Renderer/Private/SceneRendering.h
@@ -2254,6 +2254,8 @@ public:
 	DECLARE_MULTICAST_DELEGATE_OneParam(FSceneOnScreenMessagesDelegate, FScreenMessageWriter&);
 	FSceneOnScreenMessagesDelegate OnGetOnScreenMessages;
 
+	inline TConstStridedView<FSceneView> GetSceneViews() const { return MakeStridedViewOfBase<const FSceneView>(MakeArrayView(Views)); }
+
 protected:
 
 	/** Size of the family. */
diff --git a/Engine/Source/Runtime/Renderer/Private/SceneVisibility.cpp b/Engine/Source/Runtime/Renderer/Private/SceneVisibility.cpp
--- a/Engine/Source/Runtime/Renderer/Private/SceneVisibility.cpp
+++ b/Engine/Source/Runtime/Renderer/Private/SceneVisibility.cpp
@@ -5294,7 +5294,7 @@ void FDeferredShadingSceneRenderer::BeginInitViews(
 				InstancedView->InitRHIResources();
 			}
 			Views[0].InitRHIResources();
-			FXSystem->PostInitViews(GraphBuilder, Views, !ViewFamily.EngineShowFlags.HitProxies);
+			FXSystem->PostInitViews(GraphBuilder, GetSceneViews(), !ViewFamily.EngineShowFlags.HitProxies);
 		}
 	}
 
diff --git a/Engine/Source/Runtime/Renderer/Private/ScreenPass.h b/Engine/Source/Runtime/Renderer/Private/ScreenPass.h
--- a/Engine/Source/Runtime/Renderer/Private/ScreenPass.h
+++ b/Engine/Source/Runtime/Renderer/Private/ScreenPass.h
@@ -635,7 +635,7 @@ template <typename TFunction>
 FORCEINLINE void AddDrawCanvasPass(
 	FRDGBuilder& GraphBuilder,
 	FRDGEventName&& PassName,
-	const FViewInfo& View,
+	const FSceneView& View,
 	FScreenPassRenderTarget Output,
 	TFunction Function)
 {
diff --git a/Engine/Source/Runtime/Renderer/Public/FXRenderingUtils.h b/Engine/Source/Runtime/Renderer/Public/FXRenderingUtils.h
--- a/Engine/Source/Runtime/Renderer/Public/FXRenderingUtils.h
+++ b/Engine/Source/Runtime/Renderer/Public/FXRenderingUtils.h
@@ -28,6 +28,9 @@ class FVisibleRayTracingMeshCommand;
 
 namespace UE::FXRenderingUtils
 {
+	RENDERER_API TConstStridedView<FSceneView> ConvertViewArray(TConstArrayView<FViewInfo> Views);
+	RENDERER_API FIntRect GetRawViewRectUnsafe(const FSceneView& View);
+
 	RENDERER_API bool CanMaterialRenderBeforeFXPostOpaque(const FSceneViewFamily& ViewFamily, const FPrimitiveSceneProxy& SceneProxy, const FMaterial& Material);
 
 	RENDERER_API const FGlobalDistanceFieldParameterData* GetGlobalDistanceFieldParameterData(TConstStridedView<FSceneView> Views);
