Fixed a bug where disabled modules could generate compile warnings

diff --git a/Engine/Plugins/FX/Niagara/Source/NiagaraEditor/Private/NiagaraHlslTranslator.cpp b/Engine/Plugins/FX/Niagara/Source/NiagaraEditor/Private/NiagaraHlslTranslator.cpp
--- a/Engine/Plugins/FX/Niagara/Source/NiagaraEditor/Private/NiagaraHlslTranslator.cpp
+++ b/Engine/Plugins/FX/Niagara/Source/NiagaraEditor/Private/NiagaraHlslTranslator.cpp
@@ -3,11 +3,9 @@
 #include "NiagaraHlslTranslator.h"
 
 #include "EdGraphSchema_Niagara.h"
-#include "EdGraphUtilities.h"
 #include "INiagaraEditorTypeUtilities.h"
 #include "Modules/ModuleManager.h"
 #include "NiagaraCommon.h"
-#include "NiagaraComponent.h"
 #include "NiagaraConstants.h"
 #include "NiagaraDataInterface.h"
 #include "NiagaraDataInterfaceVector2DCurve.h"
@@ -33,14 +31,13 @@
 #include "NiagaraNodeOp.h"
 #include "NiagaraParameterCollection.h"
 #include "NiagaraScriptSource.h"
-#include "NiagaraScriptVariable.h"
 #include "NiagaraSettings.h"
 #include "NiagaraShared.h"
 #include "NiagaraSimulationStageBase.h"
 #include "NiagaraTrace.h"
 #include "ShaderCore.h"
 #include "Misc/FileHelper.h"
-
+#include "Misc/CoreMiscDefines.h"
 #include "NiagaraDataInterfaceUtilities.h"
 
 #define LOCTEXT_NAMESPACE "NiagaraCompiler"
@@ -9417,34 +9414,17 @@ int32 FHlslNiagaraTranslator::CompileOutputPin(const UEdGraphPin* InPin)
 	}
 
 	UNiagaraNode* Node = Cast<UNiagaraNode>(Pin->GetOwningNode());
-	const UEdGraphPin* OriginalPin = Pin;
 
-	// The node can also replace our pin with another pin (e.g. in the case of static switches), so we need to make sure we don't run into a circular dependency
-	/*TSet<UEdGraphPin*> SeenPins;
-	while (Node->SubstituteCompiledPin(this, &Pin))
+	if (!Node->IsNodeEnabled())
 	{
-		bool bIsAlreadyInSet = false;
-		SeenPins.Add(Pin, &bIsAlreadyInSet);
-		Node = Cast<UNiagaraNode>(Pin->GetOwningNode());
-		if (bIsAlreadyInSet)
+		// if the node is disabled (which commonly happens when a module is disabled in the stack), we skip it and follow the execution pin to the next node
+		UEdGraphPin* InputPin = Node->GetInputPin(0);
+		if (Node->GetOutputPin(0) == Pin && InputPin && InputPin->LinkedTo.Num() == 1 && Node->IsParameterMapPin(Pin))
 		{
-			Error(LOCTEXT("CircularGraphSubstitutionError", "Circular dependency detected, please check your module graph."), Node, Pin);
-			return INDEX_NONE;
+			return CompileOutputPin(InputPin->LinkedTo[0]);
 		}
-	}*/
-
-	// It is possible that the output pin was substituted by an input pin (e.g. the default value pin on a node).
-	// If that is the case we try to compile that pin directly.
-	if (Pin->Direction == EGPD_Input)
-	{
-		int32* ExistingChunk = PinToCodeChunks.Last().Find(OriginalPin); // Check if the pin was already compiled before
-		if (ExistingChunk)
-		{
-			return *ExistingChunk;
-		}
-		int32 Chunk = CompilePin(Pin);
-		PinToCodeChunks.Last().Add(OriginalPin, Chunk);
-		return Chunk;
+		Error(LOCTEXT("TraceDisabledPinFailed", "Failed to trace output pin of disabled node to a valid input!"), Node, Pin);
+		return INDEX_NONE;
 	}
 
 	int32 Ret = INDEX_NONE;
diff --git a/Engine/Plugins/FX/Niagara/Source/NiagaraEditor/Private/ViewModels/Stack/NiagaraStackSummaryViewInputCollection.cpp b/Engine/Plugins/FX/Niagara/Source/NiagaraEditor/Private/ViewModels/Stack/NiagaraStackSummaryViewInputCollection.cpp
--- a/Engine/Plugins/FX/Niagara/Source/NiagaraEditor/Private/ViewModels/Stack/NiagaraStackSummaryViewInputCollection.cpp
+++ b/Engine/Plugins/FX/Niagara/Source/NiagaraEditor/Private/ViewModels/Stack/NiagaraStackSummaryViewInputCollection.cpp
@@ -113,7 +113,7 @@ void UNiagaraStackSummaryViewObject::AppendEmitterCategory(FFunctionCallNodesSta
 		bool bIsFirst = true;
 		for (UNiagaraNodeFunctionCall* ModuleNode : ModuleNodes)
 		{
-			if (ModuleNode && (ModuleNode->HasValidScriptAndGraph() || ModuleNode->Signature.IsValid()))
+			if (ModuleNode && ModuleNode->IsNodeEnabled() && (ModuleNode->HasValidScriptAndGraph() || ModuleNode->Signature.IsValid()))
 			{
 				AppendInputsForFunctionCall(State, ModuleNode, ModuleNode, NewIssues, true);
 			}
