Fixed missing support for copy pasting user parameters when pasting emitters

diff --git a/Engine/Plugins/FX/Niagara/Source/Niagara/Private/NiagaraParameterStore.cpp b/Engine/Plugins/FX/Niagara/Source/Niagara/Private/NiagaraParameterStore.cpp
--- a/Engine/Plugins/FX/Niagara/Source/Niagara/Private/NiagaraParameterStore.cpp
+++ b/Engine/Plugins/FX/Niagara/Source/Niagara/Private/NiagaraParameterStore.cpp
@@ -1310,11 +1310,18 @@ const FNiagaraVariableBase* FNiagaraParameterStore::FindVariable(const UNiagaraD
 
 void FNiagaraParameterStore::CopyParameterData(FNiagaraParameterStore& DestStore, const FNiagaraVariable& Parameter) const
 {
-	int32 DestIndex = DestStore.IndexOf(Parameter);
-	int32 SrcIndex = IndexOf(Parameter);
+	CopyParameterData(DestStore, Parameter, Parameter);
+}
+
+void FNiagaraParameterStore::CopyParameterData(FNiagaraParameterStore& DestStore, const FNiagaraVariable& SourceParameter, const FNiagaraVariable& TargetParameter) const
+{
+	ensure(SourceParameter.GetType() == TargetParameter.GetType());
+	
+	int32 DestIndex = DestStore.IndexOf(TargetParameter);
+	int32 SrcIndex = IndexOf(SourceParameter);
 	if (DestIndex != INDEX_NONE && SrcIndex != INDEX_NONE)
 	{
-		if (Parameter.IsDataInterface())
+		if (SourceParameter.IsDataInterface())
 		{
 			ensure(DestStore.DataInterfaces.IsValidIndex(DestIndex));
 			UNiagaraDataInterface* DestDataInterface = DestStore.DataInterfaces[DestIndex];
@@ -1322,8 +1329,8 @@ void FNiagaraParameterStore::CopyParameterData(FNiagaraParameterStore& DestStore
 			{
 				if (ensureMsgf(DestStore.Owner != nullptr, TEXT("Destination data interface pointer was null and a new one couldn't be created because the destination store's owner pointer was also null.")))
 				{
-					UE_LOG(LogNiagara, Warning, TEXT("While trying to copy parameter data the destination data interface was null, creating a new one.  Parameter: %s Destination Store Owner: %s"), *Parameter.GetName().ToString(), *GetPathNameSafe(DestStore.Owner.Get()));
-					DestDataInterface = NewObject<UNiagaraDataInterface>(DestStore.Owner.Get(), Parameter.GetType().GetClass(), NAME_None, RF_Transactional | RF_Public);
+					UE_LOG(LogNiagara, Warning, TEXT("While trying to copy parameter data the destination data interface was null, creating a new one.  Parameter: %s Destination Store Owner: %s"), *SourceParameter.GetName().ToString(), *GetPathNameSafe(DestStore.Owner.Get()));
+					DestDataInterface = NewObject<UNiagaraDataInterface>(DestStore.Owner.Get(), SourceParameter.GetType().GetClass(), NAME_None, RF_Transactional | RF_Public);
 					DestStore.DataInterfaces[DestIndex] = DestDataInterface;
 				}
 				else
@@ -1334,18 +1341,18 @@ void FNiagaraParameterStore::CopyParameterData(FNiagaraParameterStore& DestStore
 			DataInterfaces[SrcIndex]->CopyTo(DestDataInterface);
 			DestStore.OnInterfaceChange();
 		}
-		else if (Parameter.IsUObject())
+		else if (SourceParameter.IsUObject())
 		{
 			DestStore.SetUObject(GetUObject(SrcIndex), DestIndex);
 		}
-		else if (Parameter.GetType() == FNiagaraTypeDefinition::GetPositionDef())
+		else if (SourceParameter.GetType() == FNiagaraTypeDefinition::GetPositionDef())
 		{
-			const FVector* SourceVector = GetPositionParameterValue(Parameter.GetName());
-			DestStore.SetPositionParameterValue(SourceVector ? *SourceVector : FVector::ZeroVector, Parameter.GetName());
+			const FVector* SourceVector = GetPositionParameterValue(SourceParameter.GetName());
+			DestStore.SetPositionParameterValue(SourceVector ? *SourceVector : FVector::ZeroVector, TargetParameter.GetName());
 		}
 		else
 		{
-			DestStore.SetParameterData(GetParameterData(SrcIndex), DestIndex, Parameter.GetSizeInBytes());
+			DestStore.SetParameterData(GetParameterData(SrcIndex), DestIndex, SourceParameter.GetSizeInBytes());
 		}
 	}
 }
diff --git a/Engine/Plugins/FX/Niagara/Source/Niagara/Public/NiagaraParameterStore.h b/Engine/Plugins/FX/Niagara/Source/Niagara/Public/NiagaraParameterStore.h
--- a/Engine/Plugins/FX/Niagara/Source/Niagara/Public/NiagaraParameterStore.h
+++ b/Engine/Plugins/FX/Niagara/Source/Niagara/Public/NiagaraParameterStore.h
@@ -477,7 +477,8 @@ public:
 
 	/** Copies the passed parameter from this parameter store into another. */
 	void CopyParameterData(FNiagaraParameterStore& DestStore, const FNiagaraVariable& Parameter) const;
-
+	void CopyParameterData(FNiagaraParameterStore& DestStore, const FNiagaraVariable& SourceParameter, const FNiagaraVariable& TargetParameter) const;
+	
 	enum class EDataInterfaceCopyMethod
 	{
 		/** A new data interface will be created and it will be synchronized using the CopyTo method. */
diff --git a/Engine/Plugins/FX/Niagara/Source/NiagaraEditor/Private/Customizations/NiagaraComponentDetails.cpp b/Engine/Plugins/FX/Niagara/Source/NiagaraEditor/Private/Customizations/NiagaraComponentDetails.cpp
--- a/Engine/Plugins/FX/Niagara/Source/NiagaraEditor/Private/Customizations/NiagaraComponentDetails.cpp
+++ b/Engine/Plugins/FX/Niagara/Source/NiagaraEditor/Private/Customizations/NiagaraComponentDetails.cpp
@@ -1032,6 +1032,7 @@ void FNiagaraSystemUserParameterBuilder::GenerateChildContent(IDetailChildrenBui
 			TSharedPtr<FNiagaraUserParameterPanelViewModel> UserParameterPanelViewModel = SystemViewModel.Pin()->GetUserParameterPanelViewModel();
 			if(ensure(UserParameterPanelViewModel.IsValid()))
 			{
+				UserParameterPanelViewModel->OnRefreshRequested().BindSP(this, &FNiagaraSystemUserParameterBuilder::Rebuild);
 				UserParameterPanelViewModel->OnParameterAdded().BindSP(this, &FNiagaraSystemUserParameterBuilder::RequestRename);
 			}
 			
diff --git a/Engine/Plugins/FX/Niagara/Source/NiagaraEditor/Private/Customizations/NiagaraComponentDetails.h b/Engine/Plugins/FX/Niagara/Source/NiagaraEditor/Private/Customizations/NiagaraComponentDetails.h
--- a/Engine/Plugins/FX/Niagara/Source/NiagaraEditor/Private/Customizations/NiagaraComponentDetails.h
+++ b/Engine/Plugins/FX/Niagara/Source/NiagaraEditor/Private/Customizations/NiagaraComponentDetails.h
@@ -179,6 +179,14 @@ public:
 		if (System.IsValid() && bDelegatesInitialized)
 		{
 			System->GetExposedParameters().OnStructureChanged().RemoveAll(this);
+
+			if(SystemViewModel.IsValid() && SystemViewModel.Pin()->GetUserParameterPanelViewModel().IsValid())
+			{
+				if (TSharedPtr<FNiagaraUserParameterPanelViewModel> PanelViewModel = SystemViewModel.Pin()->GetUserParameterPanelViewModel())
+				{
+					PanelViewModel->OnRefreshRequested().Unbind();
+				}
+			}
 			
 			if(SystemViewModel.IsValid() && SystemViewModel.Pin()->GetUserParametersHierarchyViewModel())
 			{
diff --git a/Engine/Plugins/FX/Niagara/Source/NiagaraEditor/Private/NiagaraEditorUtilities.cpp b/Engine/Plugins/FX/Niagara/Source/NiagaraEditor/Private/NiagaraEditorUtilities.cpp
--- a/Engine/Plugins/FX/Niagara/Source/NiagaraEditor/Private/NiagaraEditorUtilities.cpp
+++ b/Engine/Plugins/FX/Niagara/Source/NiagaraEditor/Private/NiagaraEditorUtilities.cpp
@@ -22,12 +22,14 @@
 #include "NiagaraNodeFunctionCall.h"
 #include "NiagaraNodeInput.h"
 #include "NiagaraNodeOutput.h"
+#include "NiagaraNodeParameterMapGet.h"
 #include "NiagaraNodeParameterMapSet.h"
 #include "NiagaraNodeStaticSwitch.h"
 #include "NiagaraOverviewNode.h"
 #include "NiagaraParameterMapHistory.h"
 #include "NiagaraParameterDefinitions.h"
 #include "NiagaraScript.h"
+#include "NiagaraScriptGraphViewModel.h"
 #include "NiagaraScriptMergeManager.h"
 #include "NiagaraScriptSource.h"
 #include "NiagaraSettings.h"
@@ -60,6 +62,7 @@
 #include "Misc/ScopeExit.h"
 #include "Styling/CoreStyle.h"
 #include "Subsystems/AssetEditorSubsystem.h"
+#include "UObject/PropertyIterator.h"
 #include "UObject/StructOnScope.h"
 #include "Widgets/Images/SImage.h"
 #include "Widgets/Notifications/SNotificationList.h"
@@ -2211,6 +2214,179 @@ bool FNiagaraEditorUtilities::AddParameter(FNiagaraVariable& NewParameterVariabl
 	return bSuccess;
 }
 
+TArray<FNiagaraVariable> FNiagaraEditorUtilities::GetReferencedUserParametersFromEmitter(TSharedRef<FNiagaraEmitterViewModel> EmitterViewModel)
+{
+	TArray<FNiagaraVariable> UserParameters;
+	
+	TSharedRef<FNiagaraScriptGraphViewModel> GraphViewModel = EmitterViewModel->GetSharedScriptViewModel()->GetGraphViewModel();
+
+	TArray<UNiagaraNodeParameterMapGet*> MapGetNodes;
+	GraphViewModel->GetGraph()->GetNodesOfClass<UNiagaraNodeParameterMapGet>(MapGetNodes);
+
+	// first we check for direct references on parameter map get nodes
+	for(UNiagaraNodeParameterMapGet* MapGetNode : MapGetNodes)
+	{
+		TArray<UEdGraphPin*> OutputPins;
+		MapGetNode->GetOutputPins(OutputPins);
+
+		for(UEdGraphPin* Pin : OutputPins)
+		{
+			if(!Pin->bOrphanedPin && Pin->PinType.PinSubCategory == UNiagaraNodeParameterMapBase::ParameterPinSubCategory)
+			{
+				FNiagaraVariable Variable = UEdGraphSchema_Niagara::PinToNiagaraVariable(Pin);
+				FNiagaraParameterHandle Handle(Variable.GetName());
+				if(Handle.IsUserHandle())
+				{
+					UserParameters.AddUnique(Variable);
+				}
+			}			
+		}
+	}
+
+	TArray<UNiagaraNodeInput*> InputNodes;
+	GraphViewModel->GetGraph()->GetNodesOfClass<UNiagaraNodeInput>(InputNodes);
+
+	// then we check for DI user parameter binding properties
+	for(UNiagaraNodeInput* Input : InputNodes)
+	{
+		if(Input->GetDataInterface() != nullptr)
+		{
+			for(TFieldIterator<FStructProperty> It(Input->GetDataInterface()->GetClass()); It; ++It)
+			{
+				if((*It)->Struct == StaticStruct<FNiagaraUserParameterBinding>())
+				{
+					FNiagaraUserParameterBinding ParameterBinding;
+					(*It)->GetValue_InContainer(Input->GetDataInterface(), &ParameterBinding);
+
+					if(ParameterBinding.Parameter.IsValid())
+					{
+						FNiagaraParameterHandle ReferencedParameter(ParameterBinding.Parameter.GetName());
+						if(ReferencedParameter.IsUserHandle())
+						{
+							UserParameters.AddUnique(ParameterBinding.Parameter);
+						}
+					}
+				}
+			}
+		}
+	}
+
+	// then we check for user parameter bindings on renderers
+	const TArray<UNiagaraRendererProperties*>& Renderers = EmitterViewModel->GetEmitter().GetEmitterData()->GetRenderers();
+	for(UNiagaraRendererProperties* Renderer : Renderers)
+	{
+		for(TFieldIterator<FStructProperty> It(Renderer->GetClass()); It; ++It)
+		{
+			if((*It)->Struct == StaticStruct<FNiagaraUserParameterBinding>())
+			{
+				FNiagaraUserParameterBinding ParameterBinding;
+				(*It)->GetValue_InContainer(Renderer, &ParameterBinding);
+
+				if(ParameterBinding.Parameter.IsValid())
+				{
+					FNiagaraParameterHandle ReferencedParameter(ParameterBinding.Parameter.GetName());
+					if(ReferencedParameter.IsUserHandle())
+					{
+						UserParameters.AddUnique(ParameterBinding.Parameter);
+					}
+				}
+			}
+		}
+	}
+	
+	return UserParameters;
+}
+
+TArray<UNiagaraNodeParameterMapGet*> FNiagaraEditorUtilities::GetParameterMapGetNodesWithUserParameter(TSharedRef<FNiagaraEmitterViewModel> EmitterViewModel, FNiagaraVariable UserParameter)
+{
+	TArray<UNiagaraNodeParameterMapGet*> Result;
+
+	TSharedRef<FNiagaraScriptGraphViewModel> GraphViewModel = EmitterViewModel->GetSharedScriptViewModel()->GetGraphViewModel();
+
+	TArray<UNiagaraNodeParameterMapGet*> MapGetNodes;
+	GraphViewModel->GetGraph()->GetNodesOfClass<UNiagaraNodeParameterMapGet>(MapGetNodes);
+
+	for(UNiagaraNodeParameterMapGet* MapGetNode : MapGetNodes)
+	{
+		TArray<UEdGraphPin*> OutputPins;
+		MapGetNode->GetOutputPins(OutputPins);
+
+		for(UEdGraphPin* Pin : OutputPins)
+		{
+			if(!Pin->bOrphanedPin && Pin->PinType.PinSubCategory == UNiagaraNodeParameterMapBase::ParameterPinSubCategory)
+			{
+				FNiagaraVariable Variable = UEdGraphSchema_Niagara::PinToNiagaraVariable(Pin);
+				FNiagaraParameterHandle Handle(Variable.GetName());
+				if(Handle.IsUserHandle() && Variable == UserParameter)
+				{
+					Result.AddUnique(MapGetNode);
+					break;
+				}
+			}			
+		}
+	}
+
+	return Result;
+}
+
+TArray<FNiagaraUserParameterBinding*> FNiagaraEditorUtilities::GetUserParameterBindingsForUserParameter(TSharedRef<FNiagaraEmitterViewModel> EmitterViewModel, FNiagaraVariable UserParameter)
+{
+	TArray<FNiagaraUserParameterBinding*> Bindings;
+
+	// First we gather the bindings on renderers of an emitter
+	const TArray<UNiagaraRendererProperties*>& Renderers = EmitterViewModel->GetEmitter().GetEmitterData()->GetRenderers();
+	for(UNiagaraRendererProperties* Renderer : Renderers)
+	{
+		for(TFieldIterator<FStructProperty> It(Renderer->GetClass()); It; ++It)
+		{
+			if((*It)->Struct == StaticStruct<FNiagaraUserParameterBinding>())
+			{				
+				FNiagaraUserParameterBinding* ParameterBinding = (FNiagaraUserParameterBinding*) (*It)->ContainerPtrToValuePtr<void>(Renderer, 0);
+				
+				if(ParameterBinding->Parameter.IsValid())
+				{
+					FNiagaraParameterHandle ReferencedParameter(ParameterBinding->Parameter.GetName());
+					if(ReferencedParameter.IsUserHandle() && ParameterBinding->Parameter == UserParameter)
+					{
+						Bindings.Add(ParameterBinding);
+					}
+				}
+			}
+		}
+	}
+
+	// then we get the bindings on DIs which live on input nodes in the emitter graph
+	TSharedRef<FNiagaraScriptGraphViewModel> GraphViewModel = EmitterViewModel->GetSharedScriptViewModel()->GetGraphViewModel();
+
+	TArray<UNiagaraNodeInput*> InputNodes;
+	GraphViewModel->GetGraph()->GetNodesOfClass<UNiagaraNodeInput>(InputNodes);
+
+	for(UNiagaraNodeInput* Input : InputNodes)
+	{
+		if(Input->GetDataInterface() != nullptr)
+		{
+			for(TFieldIterator<FStructProperty> It(Input->GetDataInterface()->GetClass()); It; ++It)
+			{
+				if((*It)->Struct == StaticStruct<FNiagaraUserParameterBinding>())
+				{					
+					FNiagaraUserParameterBinding* ParameterBinding = (FNiagaraUserParameterBinding*) (*It)->ContainerPtrToValuePtr<void>(Input->GetDataInterface(), 0);
+
+					if(ParameterBinding->Parameter.IsValid())
+					{
+						FNiagaraParameterHandle ReferencedParameter(ParameterBinding->Parameter.GetName());
+						if(ReferencedParameter.IsUserHandle() && ParameterBinding->Parameter == UserParameter)
+						{
+							Bindings.Add(ParameterBinding);
+						}
+					}
+				}
+			}
+		}
+	}
+	
+	return Bindings;
+}
+
 TObjectPtr<UNiagaraScriptVariable> FNiagaraEditorUtilities::GetScriptVariableForUserParameter(const FNiagaraVariable& UserParameter, TSharedPtr<FNiagaraSystemViewModel> SystemViewModel)
 {
 	return Cast<UNiagaraSystemEditorData>(SystemViewModel->GetSystem().GetEditorData())->FindOrAddUserScriptVariable(UserParameter, SystemViewModel->GetSystem());
@@ -2226,6 +2402,37 @@ TObjectPtr<UNiagaraScriptVariable> FNiagaraEditorUtilities::FindScriptVariableFo
 	return Cast<UNiagaraSystemEditorData>(System.GetEditorData())->FindUserScriptVariable(UserParameterGuid);
 }
 
+void FNiagaraEditorUtilities::ReplaceUserParameterReferences(TSharedRef<FNiagaraEmitterViewModel> EmitterViewModel,	FNiagaraVariable OldUserParameter, FNiagaraVariable NewUserParameter)
+{
+	TArray<FNiagaraUserParameterBinding*> ReferencingBindings = FNiagaraEditorUtilities::GetUserParameterBindingsForUserParameter(EmitterViewModel, OldUserParameter);
+
+	for(FNiagaraUserParameterBinding* ReferencingBinding : ReferencingBindings)
+	{
+		ReferencingBinding->Parameter = NewUserParameter;
+	}
+
+	TArray<UNiagaraNodeParameterMapGet*> ReferencingMapGetNodes = FNiagaraEditorUtilities::GetParameterMapGetNodesWithUserParameter(EmitterViewModel, OldUserParameter);
+
+	for(UNiagaraNodeParameterMapGet* MapGet : ReferencingMapGetNodes)
+	{
+		TArray<UEdGraphPin*> OutputPins;
+		MapGet->GetOutputPins(OutputPins);
+
+		for(UEdGraphPin* GraphPin : OutputPins)
+		{
+			if(GraphPin->PinType.PinSubCategory == UNiagaraNodeParameterMapBase::ParameterPinSubCategory && GraphPin->bOrphanedPin == false)
+			{
+				FNiagaraVariable RepresentingVariable = UEdGraphSchema_Niagara::PinToNiagaraVariable(GraphPin);
+
+				if(RepresentingVariable == OldUserParameter)
+				{
+					GraphPin->PinName = NewUserParameter.GetName();
+				}
+			}
+		}
+	}
+}
+
 void FNiagaraEditorUtilities::ShowParentEmitterInContentBrowser(TSharedRef<FNiagaraEmitterViewModel> Emitter)
 {
 	FContentBrowserModule& ContentBrowserModule = FModuleManager::Get().LoadModuleChecked<FContentBrowserModule>(TEXT("ContentBrowser"));
diff --git a/Engine/Plugins/FX/Niagara/Source/NiagaraEditor/Private/ViewModels/NiagaraSystemViewModel.cpp b/Engine/Plugins/FX/Niagara/Source/NiagaraEditor/Private/ViewModels/NiagaraSystemViewModel.cpp
--- a/Engine/Plugins/FX/Niagara/Source/NiagaraEditor/Private/ViewModels/NiagaraSystemViewModel.cpp
+++ b/Engine/Plugins/FX/Niagara/Source/NiagaraEditor/Private/ViewModels/NiagaraSystemViewModel.cpp
@@ -35,6 +35,7 @@
 #include "NiagaraNodeInput.h"
 #include "NiagaraNodeOutput.h"
 #include "NiagaraNodeParameterMapGet.h"
+#include "NiagaraNotificationWidgetProvider.h"
 #include "NiagaraOverviewNode.h"
 #include "NiagaraParameterDefinitionsSubscriber.h"
 #include "NiagaraScriptGraphViewModel.h"
@@ -64,6 +65,7 @@
 #include "ViewModels/TNiagaraViewModelManager.h"
 #include "ViewModels/HierarchyEditor/NiagaraUserParametersHierarchyViewModel.h"
 #include "ViewModels/NiagaraParameterPanelViewModel.h"
+#include "Widgets/Notifications/SNotificationList.h"
 
 
 DECLARE_CYCLE_STAT(TEXT("Niagara - SystemViewModel - CompileSystem"), STAT_NiagaraEditor_SystemViewModel_CompileSystem, STATGROUP_NiagaraEditor);
@@ -582,6 +584,8 @@ void FNiagaraSystemViewModel::DuplicateEmitters(TArray<FEmitterHandleToDuplicate
 	}
 
 	GetSystem().Modify();
+
+	TMap<FGuid, UNiagaraSystem*> NewHandleToOriginalSystemMap;
 	for (FEmitterHandleToDuplicate& EmitterHandleToDuplicate : EmitterHandlesToDuplicate)
 	{
 		UNiagaraSystem* SourceSystem = nullptr;
@@ -613,6 +617,7 @@ void FNiagaraSystemViewModel::DuplicateEmitters(TArray<FEmitterHandleToDuplicate
 			const FNiagaraEmitterHandle& EmitterHandle = GetSystem().DuplicateEmitterHandle(HandleToDuplicate, FNiagaraUtilities::GetUniqueName(HandleToDuplicate.GetName(), EmitterHandleNames));
 			FNiagaraScratchPadUtilities::FixExternalScratchPadScriptsForEmitter(*SourceSystem, EmitterHandle.GetInstance());
 			EmitterHandleNames.Add(EmitterHandle.GetName());
+			NewHandleToOriginalSystemMap.Add(EmitterHandle.GetId(), SourceSystem);
 			if (EmitterHandleToDuplicate.OverviewNode)
 			{
 				EmitterHandleToDuplicate.OverviewNode->Initialize(&GetSystem(), EmitterHandle.GetId());
@@ -624,6 +629,127 @@ void FNiagaraSystemViewModel::DuplicateEmitters(TArray<FEmitterHandleToDuplicate
 	GetEditorData().SynchronizeOverviewGraphWithSystem(GetSystem());
 	RefreshAll();
 	bForceAutoCompileOnce = true;
+
+	// We copy over referenced User Parameters now
+	TArray<TSharedPtr<FNiagaraEmitterHandleViewModel>> NewEmitterHandleViewModels;
+	FNiagaraUserRedirectionParameterStore& TargetParameterStore = GetSystem().GetExposedParameters();
+
+	// we have to keep track of parameters we have already added but renamed.
+	// i.e. if we copy paste 2 emitters into a new system and both reference the same parameter, but there was already one with matching name and mismatched type
+	// we create a new parameter with unique name _once_ that both emitters should point to.
+	TMap<FNiagaraVariable, FNiagaraVariable> RenamedUserParameters;
+	
+	for(auto& Element : NewHandleToOriginalSystemMap)
+	{
+		// if the source system of the current emitter is the system we are currently editing, we skip the copy paste of user parameters
+		if(Element.Value == &GetSystem())
+		{
+			continue;
+		}
+		
+		TSharedPtr<FNiagaraEmitterHandleViewModel> EmitterHandleViewModel = GetEmitterHandleViewModelById(Element.Key);
+		TArray<FNiagaraVariable> ReferencedUserParameters = FNiagaraEditorUtilities::GetReferencedUserParametersFromEmitter(EmitterHandleViewModel->GetEmitterViewModel());
+
+		UNiagaraSystem* OriginalSystem = Element.Value;
+		
+		for(const FNiagaraVariable& UserParameter : ReferencedUserParameters)
+		{
+			// in case we already have a parameter with the same name & type, we simply skip the parameter by default.
+			// if the user wants to add a new User Parameter, he can do that too.
+			if(TargetParameterStore.FindParameterVariable(UserParameter, false) != nullptr)
+			{
+				continue;
+			}
+			// if we have a parameter with the same name but a different type, we need to create a new parameter with a different name, copy the data and fixup the references
+			// if we have already added a new unique user parameter with the same name, we reuse that parameter
+			else if(TargetParameterStore.FindParameterVariable(UserParameter, true) != nullptr)
+			{
+				TArray<FNiagaraVariable> ExistingUserParameters;
+				TargetParameterStore.GetUserParameters(ExistingUserParameters);
+
+				TSet<FName> ExistingUserParameterNames;
+				for(const FNiagaraVariable& ExistingUserParameter : ExistingUserParameters)
+				{
+					FNiagaraVariable RedirectedUserParameter = ExistingUserParameter;
+					TargetParameterStore.RedirectUserVariable(RedirectedUserParameter);
+					ExistingUserParameterNames.Add(RedirectedUserParameter.GetName());
+				}
+
+				bool bParameterWasRenamed = RenamedUserParameters.Contains(UserParameter);
+
+				if(!bParameterWasRenamed)
+				{
+					FNiagaraVariable NewUserParameter = FNiagaraVariable(UserParameter.GetType(), FNiagaraUtilities::GetUniqueName(UserParameter.GetName(), ExistingUserParameterNames));
+					TargetParameterStore.AddParameter(NewUserParameter);
+					OriginalSystem->GetExposedParameters().CopyParameterData(TargetParameterStore, UserParameter, NewUserParameter);
+					FNiagaraEditorUtilities::ReplaceUserParameterReferences(EmitterHandleViewModel->GetEmitterViewModel(), UserParameter, NewUserParameter);
+					RenamedUserParameters.Add(UserParameter, NewUserParameter);
+				}
+				else
+				{
+					FNiagaraVariable RenamedUserParameter = RenamedUserParameters[UserParameter];
+					FNiagaraEditorUtilities::ReplaceUserParameterReferences(EmitterHandleViewModel->GetEmitterViewModel(), UserParameter, RenamedUserParameter);
+				}
+			}
+			// if the parameter doesn't exist at all, we just add it, copy the data and we're done
+			else
+			{
+				const FNiagaraVariableWithOffset* ExistingUserParameter = OriginalSystem->GetExposedParameters().FindParameterVariable(UserParameter);
+				if(ensure(ExistingUserParameter != nullptr))
+				{
+					TargetParameterStore.AddParameter(UserParameter);
+					OriginalSystem->GetExposedParameters().CopyParameterData(GetSystem().GetExposedParameters(), UserParameter);
+				}
+			}
+		}
+	}
+
+	if(RenamedUserParameters.Num() > 0)
+	{
+		TSharedRef<SVerticalBox> WidgetContent = SNew(SVerticalBox)
+		+ SVerticalBox::Slot()
+		.AutoHeight()
+		[
+			SNew(STextBlock).Text(LOCTEXT("RenamedParametersDuringPasteNotificationStart", "Renamed the following User Parameters on the new emitters.\n"))
+		];
+		
+		for(auto& RenamedVariable : RenamedUserParameters)
+		{
+			WidgetContent->AddSlot()
+			.AutoHeight()
+			[
+				SNew(SHorizontalBox)
+				+SHorizontalBox::Slot()
+				.AutoWidth()
+				[
+					FNiagaraParameterUtilities::GetParameterWidget(RenamedVariable.Key, false, false)
+				]
+				+SHorizontalBox::Slot()
+				.AutoWidth()
+				.VAlign(VAlign_Center)
+				.Padding(3.f)
+				[
+					SNew(STextBlock).Text(LOCTEXT("RenamedParametersDuringPasteParameterConnector", "to"))
+				]
+				+SHorizontalBox::Slot()
+				.AutoWidth()
+				[
+					FNiagaraParameterUtilities::GetParameterWidget(RenamedVariable.Value, false, false)
+				]
+			];
+		}
+		
+		TSharedRef<FNiagaraParameterNotificationWidgetProvider> WidgetProvider = MakeShared<FNiagaraParameterNotificationWidgetProvider>(WidgetContent);
+		FNotificationInfo Info(WidgetProvider);
+		Info.WidthOverride = FOptionalSize();
+		Info.ExpireDuration = 5.f;
+		FSlateNotificationManager::Get().AddNotification(Info);
+	}
+
+	for(TSharedPtr<FNiagaraEmitterHandleViewModel> EmitterHandleViewModel : NewEmitterHandleViewModels)
+	{
+		EmitterHandleViewModel->GetEmitterStackViewModel()->GetRootEntry()->RefreshChildren();
+	}
 }
 
 FGuid FNiagaraSystemViewModel::GetMessageLogGuid() const
diff --git a/Engine/Plugins/FX/Niagara/Source/NiagaraEditor/Private/ViewModels/NiagaraUserParameterPanelViewModel.cpp b/Engine/Plugins/FX/Niagara/Source/NiagaraEditor/Private/ViewModels/NiagaraUserParameterPanelViewModel.cpp
--- a/Engine/Plugins/FX/Niagara/Source/NiagaraEditor/Private/ViewModels/NiagaraUserParameterPanelViewModel.cpp
+++ b/Engine/Plugins/FX/Niagara/Source/NiagaraEditor/Private/ViewModels/NiagaraUserParameterPanelViewModel.cpp
@@ -16,3 +16,8 @@ void FNiagaraUserParameterPanelViewModel::Initialize(TSharedRef<FNiagaraSystemVi
 	SystemViewModelWeak = InSystemViewModel;
 	RegisteredHandle = RegisterViewModelWithMap(&(SystemViewModelWeak.Pin()->GetSystem()), this);
 }
+
+void FNiagaraUserParameterPanelViewModel::Refresh() const
+{
+	OnRefreshRequestedDelegate.ExecuteIfBound();
+}
diff --git a/Engine/Plugins/FX/Niagara/Source/NiagaraEditor/Public/NiagaraEditorUtilities.h b/Engine/Plugins/FX/Niagara/Source/NiagaraEditor/Public/NiagaraEditorUtilities.h
--- a/Engine/Plugins/FX/Niagara/Source/NiagaraEditor/Public/NiagaraEditorUtilities.h
+++ b/Engine/Plugins/FX/Niagara/Source/NiagaraEditor/Public/NiagaraEditorUtilities.h
@@ -16,6 +16,7 @@
 class UNiagaraNodeInput;
 class UNiagaraNodeOutput;
 class UNiagaraNodeFunctionCall;
+class UNiagaraNodeParameterMapGet;
 struct FNiagaraVariable;
 struct FNiagaraTypeDefinition;
 class UNiagaraGraph;
@@ -278,11 +279,15 @@ namespace FNiagaraEditorUtilities
 	 * @returns Bool for whether adding the parameter succeeded.
 	 */
 	bool AddParameter(FNiagaraVariable& NewParameterVariable, FNiagaraParameterStore& TargetParameterStore, UObject& ParameterStoreOwner, UNiagaraStackEditorData* StackEditorData);
-
+	
+	TArray<FNiagaraVariable> GetReferencedUserParametersFromEmitter(TSharedRef<FNiagaraEmitterViewModel> EmitterViewModel);
+	TArray<UNiagaraNodeParameterMapGet*> GetParameterMapGetNodesWithUserParameter(TSharedRef<FNiagaraEmitterViewModel> EmitterViewModel, FNiagaraVariable UserParameter);
+	TArray<FNiagaraUserParameterBinding*> GetUserParameterBindingsForUserParameter(TSharedRef<FNiagaraEmitterViewModel> EmitterViewModel, FNiagaraVariable UserParameter);
 	NIAGARAEDITOR_API TObjectPtr<UNiagaraScriptVariable> GetScriptVariableForUserParameter(const FNiagaraVariable& UserParameter, TSharedPtr<FNiagaraSystemViewModel> SystemViewModel);
 	NIAGARAEDITOR_API TObjectPtr<UNiagaraScriptVariable> GetScriptVariableForUserParameter(const FNiagaraVariable& UserParameter, UNiagaraSystem& System);
 	NIAGARAEDITOR_API TObjectPtr<UNiagaraScriptVariable> FindScriptVariableForUserParameter(const FGuid& UserParameterGuid, UNiagaraSystem& System);
-	
+	void ReplaceUserParameterReferences(TSharedRef<FNiagaraEmitterViewModel> EmitterViewModel, FNiagaraVariable OldUserParameter, FNiagaraVariable NewUserParameter);
+
 	NIAGARAEDITOR_API bool AddEmitterContextMenuActions(FMenuBuilder& MenuBuilder, const TSharedPtr<FNiagaraEmitterHandleViewModel>& EmitterHandleViewModel);
 
 	void ShowParentEmitterInContentBrowser(TSharedRef<FNiagaraEmitterViewModel> EmitterViewModel);
diff --git a/Engine/Plugins/FX/Niagara/Source/NiagaraEditor/Public/ViewModels/NiagaraUserParameterPanelViewModel.h b/Engine/Plugins/FX/Niagara/Source/NiagaraEditor/Public/ViewModels/NiagaraUserParameterPanelViewModel.h
--- a/Engine/Plugins/FX/Niagara/Source/NiagaraEditor/Public/ViewModels/NiagaraUserParameterPanelViewModel.h
+++ b/Engine/Plugins/FX/Niagara/Source/NiagaraEditor/Public/ViewModels/NiagaraUserParameterPanelViewModel.h
@@ -13,14 +13,19 @@ class FNiagaraUserParameterPanelViewModel
     , public TNiagaraViewModelManager<UNiagaraSystem, FNiagaraUserParameterPanelViewModel>
 {
 public:
+	DECLARE_DELEGATE(FOnRefreshRequested)
     DECLARE_DELEGATE_OneParam(FOnParameterAdded, FNiagaraVariable);
 
 	virtual ~FNiagaraUserParameterPanelViewModel();
 	
 	void Initialize(TSharedRef<FNiagaraSystemViewModel> InSystemViewModel);
-	
+
+	void Refresh() const;
+
+	FOnRefreshRequested& OnRefreshRequested() { return OnRefreshRequestedDelegate; }
     FOnParameterAdded& OnParameterAdded() { return OnParameterAddedDelegate; }
 private:
+	FOnRefreshRequested OnRefreshRequestedDelegate;
     FOnParameterAdded OnParameterAddedDelegate;
 	TWeakPtr<FNiagaraSystemViewModel> SystemViewModelWeak;
     TNiagaraViewModelManager<UNiagaraSystem, FNiagaraUserParameterPanelViewModel>::Handle RegisteredHandle;
