Move "IsHeterogeneousVolumes" property of  Niagara Mesh Renderer Properties from the "Mesh Renderer" section to the "Rendering" section

diff --git a/Engine/Plugins/FX/Niagara/Source/Niagara/Public/NiagaraMeshRendererProperties.h b/Engine/Plugins/FX/Niagara/Source/Niagara/Public/NiagaraMeshRendererProperties.h
--- a/Engine/Plugins/FX/Niagara/Source/Niagara/Public/NiagaraMeshRendererProperties.h
+++ b/Engine/Plugins/FX/Niagara/Source/Niagara/Public/NiagaraMeshRendererProperties.h
@@ -251,7 +251,7 @@ public:
 	uint32 bOverrideMaterials : 1;
 
 	/** Whether or not to render with heterogeneous volumes.*/
-	UPROPERTY(EditAnywhere, Category = "Mesh Rendering")
+	UPROPERTY(EditAnywhere, Category = "Rendering")
 	uint32 bUseHeterogeneousVolumes : 1;
 
 	/** If true, the particles are only sorted when using a translucent material. */
