- Add tick count to the cache data, editoronly as it's useful for debugging

diff --git a/Engine/Plugins/FX/Niagara/Source/Niagara/Classes/NiagaraSimCache.h b/Engine/Plugins/FX/Niagara/Source/Niagara/Classes/NiagaraSimCache.h
--- a/Engine/Plugins/FX/Niagara/Source/Niagara/Classes/NiagaraSimCache.h
+++ b/Engine/Plugins/FX/Niagara/Source/Niagara/Classes/NiagaraSimCache.h
@@ -187,6 +187,11 @@ struct FNiagaraSimCacheFrame
 	UPROPERTY()
 	float SimulationAge = 0.0f;
 
+#if WITH_EDITORONLY_DATA
+	UPROPERTY()
+	int32 SimulationTickCount = 0;
+#endif
+
 	UPROPERTY()
 	FNiagaraSimCacheSystemFrame SystemData;
 
diff --git a/Engine/Plugins/FX/Niagara/Source/Niagara/Private/NiagaraSimCache.cpp b/Engine/Plugins/FX/Niagara/Source/Niagara/Private/NiagaraSimCache.cpp
--- a/Engine/Plugins/FX/Niagara/Source/Niagara/Private/NiagaraSimCache.cpp
+++ b/Engine/Plugins/FX/Niagara/Source/Niagara/Private/NiagaraSimCache.cpp
@@ -302,6 +302,9 @@ bool UNiagaraSimCache::WriteFrame(UNiagaraComponent* NiagaraComponent, FNiagaraS
 	CacheFrame.LocalToWorld = Helper.SystemInstance->GatheredInstanceParameters.ComponentTrans;
 	CacheFrame.LWCTile = Helper.SystemInstance->GetLWCTile();
 	CacheFrame.SimulationAge = FMath::RoundToFloat(Helper.SystemInstance->GetAge() * CacheAgeResolution) / CacheAgeResolution;
+#if WITH_EDITORONLY_DATA
+	CacheFrame.SimulationTickCount = Helper.SystemInstance->GetTickCount();
+#endif
 
 	CacheFrame.SystemData.LocalBounds = Helper.SystemInstance->GetLocalBounds();
 
