- Fix crash if intepolation failed because unique ID was missing

#rnx

diff --git a/Engine/Plugins/FX/Niagara/Source/Niagara/Private/NiagaraSimCacheHelper.h b/Engine/Plugins/FX/Niagara/Source/Niagara/Private/NiagaraSimCacheHelper.h
--- a/Engine/Plugins/FX/Niagara/Source/Niagara/Private/NiagaraSimCacheHelper.h
+++ b/Engine/Plugins/FX/Niagara/Source/Niagara/Private/NiagaraSimCacheHelper.h
@@ -208,6 +208,10 @@ struct FNiagaraSimCacheHelper
 				CacheLayout.bAllowInterpolation = true;
 			}
 		}
+		if (!CacheLayout.bAllowInterpolation)
+		{
+			CacheLayout.InterpVariableNames.Empty();
+		}
 
 		// Build write mappings we will build read mappings in a separate path
 		int32 FloatOffset = 0;
@@ -439,7 +443,7 @@ struct FNiagaraSimCacheHelper
 		//-OPT: We don't have to do this all the time, figure out a better way
 		TArray<FNiagaraVariableBase> InterpCurrAttributes;
 		TArray<FNiagaraVariableBase> InterpPrevAttributes;
-		if (CacheLayout.InterpVariableNames.Num() > 0)
+		if (CacheLayout.bAllowInterpolation)
 		{
 			for (const FNiagaraSimCacheVariable& PreviousVariable : CacheLayout.Variables)
 			{
@@ -498,7 +502,7 @@ struct FNiagaraSimCacheHelper
 			// Is this a type that requires conversion / re-basing?
 			if (DestVariableLayout != nullptr)
 			{
-				const bool bInterpVariable = CacheLayout.InterpVariableNames.Contains(SourceVariable.Variable.GetName());
+				const bool bInterpVariable = CacheLayout.bAllowInterpolation && CacheLayout.InterpVariableNames.Contains(SourceVariable.Variable.GetName());
 				const bool bRebaseVariable = CacheLayout.RebaseVariableNames.Contains(SourceVariable.Variable.GetName());
 
 				if (DestVariableLayout && (bInterpVariable || bRebaseVariable))
