Allow Niagara renderers to be moved via context action

diff --git a/Engine/Plugins/FX/Niagara/Source/Niagara/Classes/NiagaraEmitter.h b/Engine/Plugins/FX/Niagara/Source/Niagara/Classes/NiagaraEmitter.h
--- a/Engine/Plugins/FX/Niagara/Source/Niagara/Classes/NiagaraEmitter.h
+++ b/Engine/Plugins/FX/Niagara/Source/Niagara/Classes/NiagaraEmitter.h
@@ -879,6 +879,7 @@ public:
 
 	void NIAGARA_API AddRenderer(UNiagaraRendererProperties* Renderer, FGuid EmitterVersion);
 	void NIAGARA_API RemoveRenderer(UNiagaraRendererProperties* Renderer, FGuid EmitterVersion);
+	void NIAGARA_API MoveRenderer(UNiagaraRendererProperties* Renderer, int32 NewIndex, FGuid EmitterVersion);
 	void NIAGARA_API AddEventHandler(FNiagaraEventScriptProperties EventHandler, FGuid EmitterVersion);
 	void NIAGARA_API RemoveEventHandlerByUsageId(FGuid EventHandlerUsageId, FGuid EmitterVersion);
 	void NIAGARA_API AddSimulationStage(UNiagaraSimulationStageBase* SimulationStage, FGuid EmitterVersion);
diff --git a/Engine/Plugins/FX/Niagara/Source/Niagara/Private/NiagaraEmitter.cpp b/Engine/Plugins/FX/Niagara/Source/Niagara/Private/NiagaraEmitter.cpp
--- a/Engine/Plugins/FX/Niagara/Source/Niagara/Private/NiagaraEmitter.cpp
+++ b/Engine/Plugins/FX/Niagara/Source/Niagara/Private/NiagaraEmitter.cpp
@@ -2972,6 +2972,32 @@ void UNiagaraEmitter::RemoveRenderer(UNiagaraRendererProperties* Renderer, FGuid
 	EmitterData->RebuildRendererBindings(*this);
 }
 
+void UNiagaraEmitter::MoveRenderer(UNiagaraRendererProperties* Renderer, int32 NewIndex, FGuid EmitterVersion)
+{
+	FVersionedNiagaraEmitterData* EmitterData = GetEmitterData(EmitterVersion);
+	int32 CurrentIndex = EmitterData->RendererProperties.IndexOfByKey(Renderer);
+	if (CurrentIndex == INDEX_NONE || CurrentIndex == NewIndex || !EmitterData->RendererProperties.IsValidIndex(NewIndex))
+	{
+		return;
+	}
+	
+	FNiagaraSystemUpdateContext UpdateContext;
+	UpdateContext.SetDestroyOnAdd(true);
+	if (UNiagaraSystem* Owner = GetTypedOuter<UNiagaraSystem>())
+	{
+		UpdateContext.Add(Owner, true);
+	}
+
+	Modify();
+	EmitterData->RendererProperties.RemoveAt(CurrentIndex);
+	EmitterData->RendererProperties.Insert(Renderer, NewIndex);
+#if WITH_EDITOR
+	UpdateChangeId(TEXT("Renderer moved"));
+	OnRenderersChangedDelegate.Broadcast();
+#endif
+	EmitterData->RebuildRendererBindings(*this);
+}
+
 FNiagaraEventScriptProperties* FVersionedNiagaraEmitterData::GetEventHandlerByIdUnsafe(FGuid ScriptUsageId)
 {
 	for (FNiagaraEventScriptProperties& EventScriptProperties : EventHandlerScriptProps)
diff --git a/Engine/Plugins/FX/Niagara/Source/NiagaraEditor/Private/ViewModels/Stack/NiagaraStackRendererItem.cpp b/Engine/Plugins/FX/Niagara/Source/NiagaraEditor/Private/ViewModels/Stack/NiagaraStackRendererItem.cpp
--- a/Engine/Plugins/FX/Niagara/Source/NiagaraEditor/Private/ViewModels/Stack/NiagaraStackRendererItem.cpp
+++ b/Engine/Plugins/FX/Niagara/Source/NiagaraEditor/Private/ViewModels/Stack/NiagaraStackRendererItem.cpp
@@ -24,7 +24,6 @@
 #include "NiagaraClipboard.h"
 #include "NiagaraEmitterEditorData.h"
 #include "ViewModels/NiagaraParameterPanelViewModel.h"
-
 #include "ViewModels/NiagaraEmitterHandleViewModel.h"
 #include "ViewModels/Stack/NiagaraStackViewModel.h"
 
@@ -100,6 +99,64 @@ const UNiagaraStackRendererItem::FCollectedUsageData& UNiagaraStackRendererItem:
 	return CachedCollectedUsageData.GetValue();
 }
 
+bool UNiagaraStackRendererItem::CanMoveRendererUp() const
+{
+	if (HasBaseRenderer() || RendererProperties.IsValid() == false)
+	{
+		return false;
+	}
+	if (FVersionedNiagaraEmitterData* EmitterData = GetEmitterViewModel()->GetEmitter().GetEmitterData())
+	{
+		return EmitterData->GetRenderers().IndexOfByKey(RendererProperties.Get()) > 0;
+	}
+	return false;
+}
+
+void UNiagaraStackRendererItem::MoveRendererUp() const
+{
+	if (!CanMoveRendererUp())
+	{
+		return;
+	}
+
+	FVersionedNiagaraEmitter VersionedEmitter = GetEmitterViewModel()->GetEmitter();
+	if (FVersionedNiagaraEmitterData* EmitterData = VersionedEmitter.GetEmitterData())
+	{
+		FScopedTransaction ScopedTransaction(LOCTEXT("MoveRendererUpTransaction", "Move renderer up"));
+		int32 CurrentIndex = EmitterData->GetRenderers().IndexOfByKey(RendererProperties.Get());
+		VersionedEmitter.Emitter->MoveRenderer(RendererProperties.Get(), CurrentIndex - 1, VersionedEmitter.Version);
+	}
+}
+
+bool UNiagaraStackRendererItem::CanMoveRendererDown() const
+{
+	if (HasBaseRenderer() || RendererProperties.IsValid() == false)
+	{
+		return false;
+	}
+	if (FVersionedNiagaraEmitterData* EmitterData = GetEmitterViewModel()->GetEmitter().GetEmitterData())
+	{
+		return EmitterData->GetRenderers().IndexOfByKey(RendererProperties.Get()) < EmitterData->GetRenderers().Num() - 1;
+	}
+	return false;
+}
+
+void UNiagaraStackRendererItem::MoveRendererDown() const
+{
+	if (!CanMoveRendererDown())
+	{
+		return;
+	}
+
+	FVersionedNiagaraEmitter VersionedEmitter = GetEmitterViewModel()->GetEmitter();
+	if (FVersionedNiagaraEmitterData* EmitterData = VersionedEmitter.GetEmitterData())
+	{
+		FScopedTransaction ScopedTransaction(LOCTEXT("MoveRendererDownTransaction", "Move renderer down"));
+		int32 CurrentIndex = EmitterData->GetRenderers().IndexOfByKey(RendererProperties.Get());
+		VersionedEmitter.Emitter->MoveRenderer(RendererProperties.Get(), CurrentIndex + 1, VersionedEmitter.Version);
+	}
+}
+
 TArray<FNiagaraVariable> UNiagaraStackRendererItem::GetMissingVariables(UNiagaraRendererProperties* RendererProperties, const FVersionedNiagaraEmitterData* EmitterData)
 {
 	TArray<FNiagaraVariable> MissingAttributes;
diff --git a/Engine/Plugins/FX/Niagara/Source/NiagaraEditor/Public/ViewModels/Stack/NiagaraStackRendererItem.h b/Engine/Plugins/FX/Niagara/Source/NiagaraEditor/Public/ViewModels/Stack/NiagaraStackRendererItem.h
--- a/Engine/Plugins/FX/Niagara/Source/NiagaraEditor/Public/ViewModels/Stack/NiagaraStackRendererItem.h
+++ b/Engine/Plugins/FX/Niagara/Source/NiagaraEditor/Public/ViewModels/Stack/NiagaraStackRendererItem.h
@@ -36,15 +36,15 @@ public:
 	virtual bool TestCanCopyWithMessage(FText& OutMessage) const override;
 	virtual void Copy(UNiagaraClipboardContent* ClipboardContent) const override;
 
-	virtual bool SupportsPaste() const { return true; }
+	virtual bool SupportsPaste() const override { return true; }
 	virtual bool TestCanPasteWithMessage(const UNiagaraClipboardContent* ClipboardContent, FText& OutMessage) const override;
 	virtual FText GetPasteTransactionText(const UNiagaraClipboardContent* ClipboardContent) const override;
 	virtual void Paste(const UNiagaraClipboardContent* ClipboardContent, FText& OutPasteWarning) override;
 
 	virtual bool SupportsDelete() const override { return true; }
-	virtual bool TestCanDeleteWithMessage(FText& OutCanDeleteMessage) const;
+	virtual bool TestCanDeleteWithMessage(FText& OutCanDeleteMessage) const override;
 	virtual FText GetDeleteTransactionText() const override;
-	virtual void Delete();
+	virtual void Delete() override;
 
 	virtual bool SupportsInheritance() const override { return true; }
 	virtual bool GetIsInherited() const override;
@@ -71,6 +71,11 @@ public:
 	bool IsOwningEmitterExcludedFromScalability() const;
 
 	virtual const FCollectedUsageData& GetCollectedUsageData() const override;
+
+	bool CanMoveRendererUp() const;
+	void MoveRendererUp() const;
+	bool CanMoveRendererDown() const;
+	void MoveRendererDown() const;
 protected:
 	virtual void FinalizeInternal() override;
 
diff --git a/Engine/Plugins/FX/Niagara/Source/NiagaraEditorWidgets/Private/NiagaraEditorWidgetsUtilities.cpp b/Engine/Plugins/FX/Niagara/Source/NiagaraEditorWidgets/Private/NiagaraEditorWidgetsUtilities.cpp
--- a/Engine/Plugins/FX/Niagara/Source/NiagaraEditorWidgets/Private/NiagaraEditorWidgetsUtilities.cpp
+++ b/Engine/Plugins/FX/Niagara/Source/NiagaraEditorWidgets/Private/NiagaraEditorWidgetsUtilities.cpp
@@ -22,14 +22,12 @@
 #include "Editor.h"
 #include "EditorFontGlyphs.h"
 #include "NiagaraEditorCommands.h"
-#include "NiagaraEmitter.h"
 #include "NiagaraEmitterEditorData.h"
-#include "NiagaraMessages.h"
 #include "NiagaraSystemEditorData.h"
 #include "ViewModels/NiagaraEmitterViewModel.h"
 #include "ViewModels/NiagaraSystemSelectionViewModel.h"
-#include "ViewModels/NiagaraSystemViewModel.h"
 #include "ViewModels/Stack/NiagaraStackInputCategory.h"
+#include "ViewModels/Stack/NiagaraStackRendererItem.h"
 
 #define LOCTEXT_NAMESPACE "NiagaraStackEditorWidgetsUtilities"
 
@@ -375,29 +373,27 @@ bool FNiagaraStackEditorWidgetsUtilities::AddStackItemContextMenuActions(FMenuBu
 					EUserInterfaceActionType::Check);
 			}
 
-			UNiagaraStackModuleItem* ModuleItem = Cast<UNiagaraStackModuleItem>(&StackItem);
-
-			if (ModuleItem && ModuleItem->GetIsEnabled() == false)
-			{
-				MenuBuilder.AddMenuEntry(FNiagaraEditorCommands::Get().HideDisabledModules);
-			}
-			
-			if (ModuleItem && ModuleItem->GetModuleNode().ContainsDebugSwitch())
+			if (UNiagaraStackModuleItem* ModuleItem = Cast<UNiagaraStackModuleItem>(&StackItem))
 			{
-				FUIAction Action(FExecuteAction::CreateStatic(&ToggleShouldDebugDraw, TWeakObjectPtr<UNiagaraStackItem>(&StackItem)),
-					FCanExecuteAction(),
-					FIsActionChecked::CreateUObject(ModuleItem, &UNiagaraStackModuleItem::IsDebugDrawEnabled));
-				MenuBuilder.AddMenuEntry(
-					LOCTEXT("ShouldDebugDraw", "Enable Debug Draw"),
-					LOCTEXT("ToggleShouldDebugDrawToolTip", "Toggle debug draw enable/disabled"),
-					FSlateIcon(),
-					Action,
-					NAME_None,
-					EUserInterfaceActionType::Check);
-			}
+				if (ModuleItem->GetIsEnabled() == false)
+				{
+					MenuBuilder.AddMenuEntry(FNiagaraEditorCommands::Get().HideDisabledModules);
+				}
+				
+				if (ModuleItem->GetModuleNode().ContainsDebugSwitch())
+				{
+					FUIAction Action(FExecuteAction::CreateStatic(&ToggleShouldDebugDraw, TWeakObjectPtr<UNiagaraStackItem>(&StackItem)),
+						FCanExecuteAction(),
+						FIsActionChecked::CreateUObject(ModuleItem, &UNiagaraStackModuleItem::IsDebugDrawEnabled));
+					MenuBuilder.AddMenuEntry(
+						LOCTEXT("ShouldDebugDraw", "Enable Debug Draw"),
+						LOCTEXT("ToggleShouldDebugDrawToolTip", "Toggle debug draw enable/disabled"),
+						FSlateIcon(),
+						Action,
+						NAME_None,
+						EUserInterfaceActionType::Check);
+				}
 
-			if(ModuleItem)
-			{
 				MenuBuilder.AddMenuEntry(
 					LOCTEXT("AddNote", "Add Note"),
 					LOCTEXT("AddNoteToolTip", "Add a note to this module item."),
@@ -405,6 +401,24 @@ bool FNiagaraStackEditorWidgetsUtilities::AddStackItemContextMenuActions(FMenuBu
 					FUIAction(FExecuteAction::CreateStatic(&EnableNoteMode, TWeakObjectPtr<UNiagaraStackModuleItem>(ModuleItem))));
 			}
 
+			if (UNiagaraStackRendererItem* RendererItem = Cast<UNiagaraStackRendererItem>(&StackItem))
+			{
+				FUIAction UpAction(FExecuteAction::CreateUObject(RendererItem, &UNiagaraStackRendererItem::MoveRendererUp),
+					FCanExecuteAction::CreateUObject(RendererItem, &UNiagaraStackRendererItem::CanMoveRendererUp));
+				MenuBuilder.AddMenuEntry(
+					LOCTEXT("MoveRendererUp", "Move up"),
+					LOCTEXT("MoveRendererUpToolTip", "Moves the renderer up, so it renders before other renderers with the same render order hint"),
+					FSlateIcon(FAppStyle::GetAppStyleSetName(), "Icons.ArrowUp"),
+					UpAction);
+
+				FUIAction DownAction(FExecuteAction::CreateUObject(RendererItem, &UNiagaraStackRendererItem::MoveRendererDown),
+					FCanExecuteAction::CreateUObject(RendererItem, &UNiagaraStackRendererItem::CanMoveRendererDown));
+				MenuBuilder.AddMenuEntry(
+					LOCTEXT("MoveRendererUp", "Move down"),
+					LOCTEXT("MoveRendererUpToolTip", "Moves the renderer down, so it renders after other renderers with the same render order hint"),
+					FSlateIcon(FAppStyle::GetAppStyleSetName(), "Icons.ArrowDown"),
+					DownAction);
+			}
 		}
 		MenuBuilder.EndSection();
 		return true;
diff --git a/Engine/Plugins/FX/Niagara/Source/NiagaraEditorWidgets/Private/Stack/SNiagaraStackItemGroupAddMenu.cpp b/Engine/Plugins/FX/Niagara/Source/NiagaraEditorWidgets/Private/Stack/SNiagaraStackItemGroupAddMenu.cpp
--- a/Engine/Plugins/FX/Niagara/Source/NiagaraEditorWidgets/Private/Stack/SNiagaraStackItemGroupAddMenu.cpp
+++ b/Engine/Plugins/FX/Niagara/Source/NiagaraEditorWidgets/Private/Stack/SNiagaraStackItemGroupAddMenu.cpp
@@ -5,18 +5,11 @@
 #include "NiagaraEditorStyle.h"
 #include "NiagaraEditorUtilities.h"
 #include "ViewModels/Stack/INiagaraStackItemGroupAddUtilities.h"
-#include "Widgets/Input/SEditableTextBox.h"
-#include "EdGraph/EdGraph.h"
 #include "SGraphActionMenu.h"
 #include "SNiagaraGraphActionWidget.h"
 #include "Framework/Application/SlateApplication.h"
-#include "Widgets/Input/SCheckBox.h"
-#include "Styling/SlateTypes.h"
 #include "Widgets/SItemSelector.h"
-#include "Widgets/SNiagaraLibraryOnlyToggleHeader.h"
 #include "Widgets/Text/SRichTextBlock.h"
-#include "Widgets/Layout/SSeparator.h"
-#include "Widgets/Layout/SSpacer.h"
 
 #define LOCTEXT_NAMESPACE "NiagaraStackItemGroupAddMenu"
 
diff --git a/Engine/Plugins/FX/Niagara/Source/NiagaraEditorWidgets/Private/Stack/SNiagaraStackModuleItem.cpp b/Engine/Plugins/FX/Niagara/Source/NiagaraEditorWidgets/Private/Stack/SNiagaraStackModuleItem.cpp
--- a/Engine/Plugins/FX/Niagara/Source/NiagaraEditorWidgets/Private/Stack/SNiagaraStackModuleItem.cpp
+++ b/Engine/Plugins/FX/Niagara/Source/NiagaraEditorWidgets/Private/Stack/SNiagaraStackModuleItem.cpp
@@ -2,7 +2,6 @@
 
 #include "Stack/SNiagaraStackModuleItem.h"
 
-#include "Editor.h"
 #include "EditorFontGlyphs.h"
 #include "NiagaraActions.h"
 #include "NiagaraConstants.h"
@@ -18,11 +17,8 @@
 #include "ScopedTransaction.h"
 #include "SDropTarget.h"
 #include "SGraphActionMenu.h"
-#include "SNiagaraGraphActionWidget.h"
 #include "Framework/Application/SlateApplication.h"
 #include "Framework/MultiBox/MultiBoxBuilder.h"
-#include "Subsystems/AssetEditorSubsystem.h"
-#include "ViewModels/NiagaraEmitterViewModel.h"
 #include "ViewModels/NiagaraScratchPadScriptViewModel.h"
 #include "ViewModels/NiagaraScratchPadViewModel.h"
 #include "ViewModels/NiagaraSystemViewModel.h"
@@ -34,7 +30,6 @@
 #include "Widgets/Images/SImage.h"
 #include "Widgets/Input/SButton.h"
 #include "Widgets/Input/SComboButton.h"
-#include "Widgets/Input/SEditableTextBox.h"
 #include "Widgets/Input/SMultiLineEditableTextBox.h"
 #include "Widgets/Layout/SBox.h"
 
