Fixed localization warnings

diff --git a/Engine/Plugins/FX/Niagara/Source/NiagaraEditorWidgets/Private/NiagaraEditorWidgetsUtilities.cpp b/Engine/Plugins/FX/Niagara/Source/NiagaraEditorWidgets/Private/NiagaraEditorWidgetsUtilities.cpp
--- a/Engine/Plugins/FX/Niagara/Source/NiagaraEditorWidgets/Private/NiagaraEditorWidgetsUtilities.cpp
+++ b/Engine/Plugins/FX/Niagara/Source/NiagaraEditorWidgets/Private/NiagaraEditorWidgetsUtilities.cpp
@@ -414,8 +414,8 @@ bool FNiagaraStackEditorWidgetsUtilities::AddStackItemContextMenuActions(FMenuBu
 				FUIAction DownAction(FExecuteAction::CreateUObject(RendererItem, &UNiagaraStackRendererItem::MoveRendererDown),
 					FCanExecuteAction::CreateUObject(RendererItem, &UNiagaraStackRendererItem::CanMoveRendererDown));
 				MenuBuilder.AddMenuEntry(
-					LOCTEXT("MoveRendererUp", "Move down"),
-					LOCTEXT("MoveRendererUpToolTip", "Moves the renderer down, so it renders after other renderers with the same render order hint"),
+					LOCTEXT("MoveRendererDown", "Move down"),
+					LOCTEXT("MoveRendererDownToolTip", "Moves the renderer down, so it renders after other renderers with the same render order hint"),
 					FSlateIcon(FAppStyle::GetAppStyleSetName(), "Icons.ArrowDown"),
 					DownAction);
 			}
