Fixed a crash when trying to set uobject-attributes in the stack

diff --git a/Engine/Plugins/FX/Niagara/Source/NiagaraEditor/Private/ViewModels/Stack/NiagaraStackFunctionInput.cpp b/Engine/Plugins/FX/Niagara/Source/NiagaraEditor/Private/ViewModels/Stack/NiagaraStackFunctionInput.cpp
--- a/Engine/Plugins/FX/Niagara/Source/NiagaraEditor/Private/ViewModels/Stack/NiagaraStackFunctionInput.cpp
+++ b/Engine/Plugins/FX/Niagara/Source/NiagaraEditor/Private/ViewModels/Stack/NiagaraStackFunctionInput.cpp
@@ -3025,7 +3025,7 @@ void UNiagaraStackFunctionInput::UpdateValuesFromScriptDefaults(FInputValues& In
 					// Data interfaces are handled differently than other values types so collect them here.
 					GetDefaultDataInterfaceValueFromDefaultPin(DefaultPin, InInputValues);
 				}
-				else
+				else if (InputType.IsUObject() == false)
 				{
 					// Otherwise check for local and linked values.
 					if (DefaultPin->LinkedTo.Num() == 0)
diff --git a/Engine/Plugins/FX/Niagara/Source/NiagaraEditor/Private/ViewModels/Stack/NiagaraStackScriptItemGroup.cpp b/Engine/Plugins/FX/Niagara/Source/NiagaraEditor/Private/ViewModels/Stack/NiagaraStackScriptItemGroup.cpp
--- a/Engine/Plugins/FX/Niagara/Source/NiagaraEditor/Private/ViewModels/Stack/NiagaraStackScriptItemGroup.cpp
+++ b/Engine/Plugins/FX/Niagara/Source/NiagaraEditor/Private/ViewModels/Stack/NiagaraStackScriptItemGroup.cpp
@@ -348,6 +348,13 @@ private:
 		{
 			Vars.Add(ParameterVariable);
 			DefaultVals.Add(FNiagaraConstants::GetAttributeDefaultValue(ParameterVariable));
+
+			FNiagaraTypeDefinition Type = ParameterVariable.GetType();
+			if (Type.IsUObject() && Type.IsDataInterface() == false)
+			{
+				// we currently don't support setting uobjects in the stack (except for data interfaces)
+				return nullptr;
+			}
 		}
 
 		UNiagaraNodeAssignment* NewAssignmentModule = FNiagaraStackGraphUtilities::AddParameterModuleToStack(Vars, *OutputNode, TargetIndex,DefaultVals );
