When Applying modules that are in use in an emitter/system, the module no longer gets deselected, and the UI gets refreshed so we get live changes.
Necessary changes to selection handling:
- Removed "RefreshAll" call when compiling emitters (which happened when we applied a script)
The recreation of emitter tracks in RefreshAll would cause sequencer to overwrite our stack selection

diff --git a/Engine/Plugins/FX/Niagara/Source/NiagaraEditor/Private/NiagaraEditorModule.cpp b/Engine/Plugins/FX/Niagara/Source/NiagaraEditor/Private/NiagaraEditorModule.cpp
--- a/Engine/Plugins/FX/Niagara/Source/NiagaraEditor/Private/NiagaraEditorModule.cpp
+++ b/Engine/Plugins/FX/Niagara/Source/NiagaraEditor/Private/NiagaraEditorModule.cpp
@@ -1975,6 +1975,11 @@ void FNiagaraEditorModule::OnAssetDeleted(UObject* DeletedObject)
 	}
 }
 
+FNiagaraEditorModule::FOnScriptApplied& FNiagaraEditorModule::OnScriptApplied()
+{
+	return OnScriptAppliedDelegate;
+}
+
 UNiagaraParameterCollection* FNiagaraEditorModule::FindCollectionForVariable(const FString& VariableName)
 {
 	auto FindCachedCollectionByPrefix = [this](const FString& Prefix)
@@ -2075,6 +2080,12 @@ void FNiagaraEditorModule::PreloadSelectablePluginAssetsByClass(UClass* InClass)
 	}
 }
 
+void FNiagaraEditorModule::ScriptApplied(UNiagaraScript* Script, FGuid VersionGuid) const
+{
+	ensure(Script != nullptr);
+	OnScriptAppliedDelegate.Broadcast(Script, VersionGuid);
+}
+
 #endif
 
 #undef LOCTEXT_NAMESPACE
diff --git a/Engine/Plugins/FX/Niagara/Source/NiagaraEditor/Private/NiagaraEditorUtilities.cpp b/Engine/Plugins/FX/Niagara/Source/NiagaraEditor/Private/NiagaraEditorUtilities.cpp
--- a/Engine/Plugins/FX/Niagara/Source/NiagaraEditor/Private/NiagaraEditorUtilities.cpp
+++ b/Engine/Plugins/FX/Niagara/Source/NiagaraEditor/Private/NiagaraEditorUtilities.cpp
@@ -948,11 +948,6 @@ void FNiagaraEditorUtilities::CompileExistingEmitters(const TArray<FVersionedNia
 			}
 		}
 	}
-
-	for (TSharedPtr<FNiagaraSystemViewModel> SystemViewModel : ExistingSystemViewModels)
-	{
-		SystemViewModel->RefreshAll();
-	}
 }
 
 bool FNiagaraEditorUtilities::TryGetEventDisplayName(UNiagaraEmitter* Emitter, FGuid EventUsageId, FText& OutEventDisplayName)
diff --git a/Engine/Plugins/FX/Niagara/Source/NiagaraEditor/Private/Toolkits/NiagaraScriptToolkit.cpp b/Engine/Plugins/FX/Niagara/Source/NiagaraEditor/Private/Toolkits/NiagaraScriptToolkit.cpp
--- a/Engine/Plugins/FX/Niagara/Source/NiagaraEditor/Private/Toolkits/NiagaraScriptToolkit.cpp
+++ b/Engine/Plugins/FX/Niagara/Source/NiagaraEditor/Private/Toolkits/NiagaraScriptToolkit.cpp
@@ -775,6 +775,7 @@ void FNiagaraScriptToolkit::OnApply()
 	SCOPE_CYCLE_COUNTER(STAT_NiagaraEditor_ScriptToolkit_OnApply);
 	UE_LOG(LogNiagaraEditor, Log, TEXT("Applying Niagara Script %s"), *GetEditingObjects()[0]->GetName());
 	UpdateOriginalNiagaraScript();
+	FNiagaraEditorModule::Get().ScriptApplied(OriginalNiagaraScript.Script, OriginalNiagaraScript.Version);
 }
 
 bool FNiagaraScriptToolkit::OnApplyEnabled() const
diff --git a/Engine/Plugins/FX/Niagara/Source/NiagaraEditor/Private/ViewModels/NiagaraScratchPadScriptViewModel.cpp b/Engine/Plugins/FX/Niagara/Source/NiagaraEditor/Private/ViewModels/NiagaraScratchPadScriptViewModel.cpp
--- a/Engine/Plugins/FX/Niagara/Source/NiagaraEditor/Private/ViewModels/NiagaraScratchPadScriptViewModel.cpp
+++ b/Engine/Plugins/FX/Niagara/Source/NiagaraEditor/Private/ViewModels/NiagaraScratchPadScriptViewModel.cpp
@@ -229,8 +229,10 @@ void FNiagaraScratchPadScriptViewModel::ApplyChanges()
 		FunctionCallNodeToRefresh->MarkNodeRequiresSynchronization(TEXT("ScratchPadChangesApplied"), true);
 	}
 
+	TSharedRef<FNiagaraScriptViewModel> ViewModel = AsShared();
+	TSharedRef<FNiagaraScratchPadScriptViewModel> ScratchPadScriptViewModel = StaticCastSharedRef<FNiagaraScratchPadScriptViewModel>(ViewModel);
 	OnHasUnappliedChangesChangedDelegate.Broadcast();
-	OnChangesAppliedDelegate.Broadcast();
+	OnChangesAppliedDelegate.Broadcast(ScratchPadScriptViewModel);
 }
 
 void FNiagaraScratchPadScriptViewModel::DiscardChanges()
diff --git a/Engine/Plugins/FX/Niagara/Source/NiagaraEditor/Private/ViewModels/NiagaraScratchPadViewModel.cpp b/Engine/Plugins/FX/Niagara/Source/NiagaraEditor/Private/ViewModels/NiagaraScratchPadViewModel.cpp
--- a/Engine/Plugins/FX/Niagara/Source/NiagaraEditor/Private/ViewModels/NiagaraScratchPadViewModel.cpp
+++ b/Engine/Plugins/FX/Niagara/Source/NiagaraEditor/Private/ViewModels/NiagaraScratchPadViewModel.cpp
@@ -237,7 +237,6 @@ void UNiagaraScratchPadViewModel::RefreshScriptViewModels()
 
 void UNiagaraScratchPadViewModel::ApplyScratchPadChanges()
 {
-	bIsBulkApplying = true;
 	for(const TSharedRef<FNiagaraScratchPadScriptViewModel>& ScriptViewModel : ScriptViewModels)
 	{
 		if(ScriptViewModel->HasUnappliedChanges())
@@ -245,9 +244,6 @@ void UNiagaraScratchPadViewModel::ApplyScratchPadChanges()
 			ScriptViewModel->ApplyChanges();
 		}
 	}
-
-	GetSystemViewModel()->RefreshAll();
-	bIsBulkApplying = false;
 }
 
 const TArray<TSharedRef<FNiagaraScratchPadScriptViewModel>>& UNiagaraScratchPadViewModel::GetScriptViewModels() const
@@ -851,15 +847,10 @@ void UNiagaraScratchPadViewModel::ScriptViewModelHasUnappliedChangesChanged()
 	bHasUnappliedChangesCache.Reset();
 }
 
-void UNiagaraScratchPadViewModel::ScriptViewModelChangesApplied()
+void UNiagaraScratchPadViewModel::ScriptViewModelChangesApplied(TSharedRef<FNiagaraScratchPadScriptViewModel> ViewModel)
 {
 	UpdateChangeId(GetSystemViewModel());
-
-	// if we are bulk applying, we don't want to force refreshes more often than necessary. The bulk apply will refresh explicitly after all scripts have been applied.
-	if(!bIsBulkApplying)
-	{
-		SystemViewModelWeak.Pin()->RefreshAll();
-	}
+	FNiagaraEditorModule::Get().ScriptApplied(ViewModel->GetOriginalScript(), FGuid());
 }
 
 void UNiagaraScratchPadViewModel::ScriptViewModelRequestDiscardChanges(TWeakPtr<FNiagaraScratchPadScriptViewModel> ScriptViewModelWeak)
diff --git a/Engine/Plugins/FX/Niagara/Source/NiagaraEditor/Private/ViewModels/NiagaraSystemViewModel.cpp b/Engine/Plugins/FX/Niagara/Source/NiagaraEditor/Private/ViewModels/NiagaraSystemViewModel.cpp
--- a/Engine/Plugins/FX/Niagara/Source/NiagaraEditor/Private/ViewModels/NiagaraSystemViewModel.cpp
+++ b/Engine/Plugins/FX/Niagara/Source/NiagaraEditor/Private/ViewModels/NiagaraSystemViewModel.cpp
@@ -169,6 +169,25 @@ void FNiagaraSystemViewModel::Initialize(UNiagaraSystem& InSystem, FNiagaraSyste
 		SetupPreviewComponentAndInstance();
 		SetupSequencer();
 	}
+
+	OnScriptAppliedDelegateHandle = FNiagaraEditorModule::Get().OnScriptApplied().AddLambda([this](UNiagaraScript* Script, FGuid VersionGuid)
+	{
+		for(UNiagaraGraph* Graph : GetAllGraphs())
+		{
+			TArray<UNiagaraNodeFunctionCall*> FunctionCallNodes;
+			Graph->GetNodesOfClass(FunctionCallNodes);
+
+			for(UNiagaraNodeFunctionCall* FunctionCall : FunctionCallNodes)
+			{
+				if(FunctionCall->FunctionScript == Script)
+				{		
+					ResetSystem(ETimeResetMode::AllowResetTime, EMultiResetMode::AllowResetAllInstances, EReinitMode::ReinitializeSystem);
+					return;
+				}
+			}
+		}
+	});
+	
 	RefreshAll();
 	AddSystemEventHandlers();
 }
@@ -333,6 +352,8 @@ void FNiagaraSystemViewModel::Cleanup()
 		ParameterPanelViewModel.Pin()->GetOnInvalidateCachedDependencies().RemoveAll(this);
 	}
 
+	FNiagaraEditorModule::Get().OnScriptApplied().Remove(OnScriptAppliedDelegateHandle);
+
 	System = nullptr;
 }
 
diff --git a/Engine/Plugins/FX/Niagara/Source/NiagaraEditor/Private/ViewModels/Stack/NiagaraStackFunctionInputCollection.cpp b/Engine/Plugins/FX/Niagara/Source/NiagaraEditor/Private/ViewModels/Stack/NiagaraStackFunctionInputCollection.cpp
--- a/Engine/Plugins/FX/Niagara/Source/NiagaraEditor/Private/ViewModels/Stack/NiagaraStackFunctionInputCollection.cpp
+++ b/Engine/Plugins/FX/Niagara/Source/NiagaraEditor/Private/ViewModels/Stack/NiagaraStackFunctionInputCollection.cpp
@@ -5,6 +5,7 @@
 #include "EdGraphSchema_Niagara.h"
 #include "NiagaraClipboard.h"
 #include "NiagaraDataInterface.h"
+#include "NiagaraEditorModule.h"
 #include "NiagaraEditorUtilities.h"
 #include "NiagaraGraph.h"
 #include "NiagaraNodeAssignment.h"
@@ -98,11 +99,16 @@ void UNiagaraStackFunctionInputCollection::Initialize(
 	ModuleNode = &InModuleNode;
 	InputFunctionCallNode = &InInputFunctionCallNode;
 	InputFunctionCallNode->OnInputsChanged().AddUObject(this, &UNiagaraStackFunctionInputCollection::OnFunctionInputsChanged);
+
+	FNiagaraEditorModule::Get().OnScriptApplied().AddUObject(this, &UNiagaraStackFunctionInputCollection::OnScriptApplied);
 }
 
 void UNiagaraStackFunctionInputCollection::FinalizeInternal()
 {
 	InputFunctionCallNode->OnInputsChanged().RemoveAll(this);
+	
+	FNiagaraEditorModule::Get().OnScriptApplied().RemoveAll(this);
+	
 	Super::FinalizeInternal();
 }
 
@@ -187,6 +193,14 @@ TArray<UNiagaraStackFunctionInput*> UNiagaraStackFunctionInputCollection::GetInl
 	return OutArray;
 }
 
+void UNiagaraStackFunctionInputCollection::OnScriptApplied(UNiagaraScript* NiagaraScript, FGuid Guid)
+{
+	if(InputFunctionCallNode->FunctionScript == NiagaraScript)
+	{
+		RefreshChildren();
+	}
+}
+
 void UNiagaraStackFunctionInputCollection::RefreshChildrenInternal(const TArray<UNiagaraStackEntry*>& CurrentChildren, TArray<UNiagaraStackEntry*>& NewChildren, TArray<FStackIssue>& NewIssues)
 {
 	RefreshChildrenForFunctionCall(ModuleNode, InputFunctionCallNode, CurrentChildren, NewChildren, NewIssues, false);
diff --git a/Engine/Plugins/FX/Niagara/Source/NiagaraEditor/Public/NiagaraEditorModule.h b/Engine/Plugins/FX/Niagara/Source/NiagaraEditor/Public/NiagaraEditorModule.h
--- a/Engine/Plugins/FX/Niagara/Source/NiagaraEditor/Public/NiagaraEditorModule.h
+++ b/Engine/Plugins/FX/Niagara/Source/NiagaraEditor/Public/NiagaraEditorModule.h
@@ -132,6 +132,7 @@ class FNiagaraEditorModule : public IModuleInterface,
 public:
 	DECLARE_DELEGATE_RetVal_OneParam(UMovieSceneNiagaraParameterTrack*, FOnCreateMovieSceneTrackForParameter, FNiagaraVariable);
 	DECLARE_MULTICAST_DELEGATE_OneParam(FOnCheckScriptToolkitsShouldFocusGraphElement, const FNiagaraScriptIDAndGraphFocusInfo*);
+	DECLARE_MULTICAST_DELEGATE_TwoParams(FOnScriptApplied, UNiagaraScript*, FGuid /** VersionGuid */);
 
 public:
 	FNiagaraEditorModule();
@@ -246,6 +247,10 @@ public:
 
 	void PreloadSelectablePluginAssetsByClass(UClass* InClass);
 
+	void ScriptApplied(UNiagaraScript* Script, FGuid VersionGuid = FGuid()) const;
+	
+	/** Callback whenever a script is applied/updated in the editor. */
+	FOnScriptApplied& OnScriptApplied();
 private:
 	class FDeferredDestructionContainerBase
 	{
@@ -344,6 +349,8 @@ private:
 	TMap<FNiagaraTypeDefinition, TSharedRef<INiagaraEditorTypeUtilities, ESPMode::ThreadSafe>> TypeToEditorUtilitiesMap;
 	TSharedPtr<INiagaraEditorTypeUtilities, ESPMode::ThreadSafe> EnumTypeUtilities;
 
+	FOnScriptApplied OnScriptAppliedDelegate;
+
 	FDelegateHandle CreateEmitterTrackEditorHandle;
 	FDelegateHandle CreateSystemTrackEditorHandle;
 
@@ -366,7 +373,7 @@ private:
 
 	FDelegateHandle OnAssetCreatedHandle;
 	FDelegateHandle OnAssetDeletedHandle;
-
+	
 	USequencerSettings* SequencerSettings;
 
 	TSharedPtr<INiagaraEditorWidgetProvider> WidgetProvider;
diff --git a/Engine/Plugins/FX/Niagara/Source/NiagaraEditor/Public/ViewModels/NiagaraScratchPadScriptViewModel.h b/Engine/Plugins/FX/Niagara/Source/NiagaraEditor/Public/ViewModels/NiagaraScratchPadScriptViewModel.h
--- a/Engine/Plugins/FX/Niagara/Source/NiagaraEditor/Public/ViewModels/NiagaraScratchPadScriptViewModel.h
+++ b/Engine/Plugins/FX/Niagara/Source/NiagaraEditor/Public/ViewModels/NiagaraScratchPadScriptViewModel.h
@@ -16,7 +16,7 @@ public:
 	DECLARE_MULTICAST_DELEGATE(FOnRenamed);
 	DECLARE_MULTICAST_DELEGATE(FOnPinnedChanged);
 	DECLARE_MULTICAST_DELEGATE(FOnHasUnappliedChangesChanged);
-	DECLARE_MULTICAST_DELEGATE(FOnChangesApplied);
+	DECLARE_MULTICAST_DELEGATE_OneParam(FOnChangesApplied, TSharedRef<FNiagaraScratchPadScriptViewModel> /* The applied view model */);
 	DECLARE_MULTICAST_DELEGATE_OneParam(FOnNodeIDFocusRequested, FNiagaraScriptIDAndGraphFocusInfo*  /* FocusInfo */);
 	DECLARE_MULTICAST_DELEGATE_OneParam(FOnPinIDFocusRequested, FNiagaraScriptIDAndGraphFocusInfo*  /* FocusInfo */);
 	DECLARE_MULTICAST_DELEGATE_OneParam(FOnGraphSubObjectSelectionChanged, const UObject*);
diff --git a/Engine/Plugins/FX/Niagara/Source/NiagaraEditor/Public/ViewModels/NiagaraScratchPadViewModel.h b/Engine/Plugins/FX/Niagara/Source/NiagaraEditor/Public/ViewModels/NiagaraScratchPadViewModel.h
--- a/Engine/Plugins/FX/Niagara/Source/NiagaraEditor/Public/ViewModels/NiagaraScratchPadViewModel.h
+++ b/Engine/Plugins/FX/Niagara/Source/NiagaraEditor/Public/ViewModels/NiagaraScratchPadViewModel.h
@@ -105,7 +105,7 @@ private:
 
 	void ScriptViewModelHasUnappliedChangesChanged();
 
-	void ScriptViewModelChangesApplied();
+	void ScriptViewModelChangesApplied(TSharedRef<FNiagaraScratchPadScriptViewModel> ViewModel);
 
 	void ScriptViewModelRequestDiscardChanges(TWeakPtr<FNiagaraScratchPadScriptViewModel> ScriptViewModelWeak);
 
@@ -140,6 +140,4 @@ private:
 	FOnScriptRenamed OnScriptRenamedDelegate;
 
 	FOnScriptDeleted OnScriptDeletedDelegate;
-
-	bool bIsBulkApplying = false;
 };
\ No newline at end of file
diff --git a/Engine/Plugins/FX/Niagara/Source/NiagaraEditor/Public/ViewModels/NiagaraSystemViewModel.h b/Engine/Plugins/FX/Niagara/Source/NiagaraEditor/Public/ViewModels/NiagaraSystemViewModel.h
--- a/Engine/Plugins/FX/Niagara/Source/NiagaraEditor/Public/ViewModels/NiagaraSystemViewModel.h
+++ b/Engine/Plugins/FX/Niagara/Source/NiagaraEditor/Public/ViewModels/NiagaraSystemViewModel.h
@@ -8,6 +8,7 @@
 #include "ISequencerModule.h"
 #include "NiagaraCommon.h"
 #include "NiagaraEditorCommon.h"
+#include "NiagaraEditorModule.h"
 #include "NiagaraOverviewNode.h"
 #include "NiagaraSystemScalabilityViewModel.h"
 #include "NiagaraUserParameterPanelViewModel.h"
@@ -662,6 +663,8 @@ private:
 
 	FOnRequestFocusTab OnRequestFocusTabDelegate;
 
+	FDelegateHandle OnScriptAppliedDelegateHandle;
+
 	/** A flag for preventing reentrancy when syncrhonizing sequencer data. */
 	bool bUpdatingEmittersFromSequencerDataChange;
 
diff --git a/Engine/Plugins/FX/Niagara/Source/NiagaraEditor/Public/ViewModels/Stack/NiagaraStackFunctionInputCollection.h b/Engine/Plugins/FX/Niagara/Source/NiagaraEditor/Public/ViewModels/Stack/NiagaraStackFunctionInputCollection.h
--- a/Engine/Plugins/FX/Niagara/Source/NiagaraEditor/Public/ViewModels/Stack/NiagaraStackFunctionInputCollection.h
+++ b/Engine/Plugins/FX/Niagara/Source/NiagaraEditor/Public/ViewModels/Stack/NiagaraStackFunctionInputCollection.h
@@ -147,6 +147,8 @@ public:
 	TArray<UNiagaraStackFunctionInput*> GetInlineParameterInputs() const;
 
 protected:
+	void OnScriptApplied(UNiagaraScript* NiagaraScript, FGuid Guid);
+	
 	virtual void FinalizeInternal() override;
 
 	virtual void RefreshChildrenInternal(const TArray<UNiagaraStackEntry*>& CurrentChildren, TArray<UNiagaraStackEntry*>& NewChildren, TArray<FStackIssue>& NewIssues) override;
