Removing Niagara dependency on Renderer/Private

diff --git a/Engine/Plugins/FX/Niagara/Source/Niagara/Niagara.Build.cs b/Engine/Plugins/FX/Niagara/Source/Niagara/Niagara.Build.cs
--- a/Engine/Plugins/FX/Niagara/Source/Niagara/Niagara.Build.cs
+++ b/Engine/Plugins/FX/Niagara/Source/Niagara/Niagara.Build.cs
@@ -68,7 +68,6 @@ public class Niagara : ModuleRules
         PrivateIncludePaths.AddRange(
             new string[] {
 				System.IO.Path.Combine(GetModuleDirectory("Engine"), "Private"),
-				System.IO.Path.Combine(GetModuleDirectory("Renderer"), "Private"),
 			});
 
         if (Target.bBuildEditor == true)
diff --git a/Engine/Plugins/FX/Niagara/Source/NiagaraEditor/NiagaraEditor.Build.cs b/Engine/Plugins/FX/Niagara/Source/NiagaraEditor/NiagaraEditor.Build.cs
--- a/Engine/Plugins/FX/Niagara/Source/NiagaraEditor/NiagaraEditor.Build.cs
+++ b/Engine/Plugins/FX/Niagara/Source/NiagaraEditor/NiagaraEditor.Build.cs
@@ -15,7 +15,6 @@ public class NiagaraEditor : ModuleRules
 			"NiagaraEditor/Private/ViewModels",
 			"NiagaraEditor/Private/TypeEditorUtilities",
 			Path.Combine(GetModuleDirectory("PropertyEditor"), "Private"),
-			Path.Combine(GetModuleDirectory("Renderer"), "Private"),
 		});
 
 		PrivateDependencyModuleNames.AddRange(
diff --git a/Engine/Plugins/FX/Niagara/Source/NiagaraEditor/Private/GlobalDistanceFieldCapture.cpp b/Engine/Plugins/FX/Niagara/Source/NiagaraEditor/Private/GlobalDistanceFieldCapture.cpp
--- a/Engine/Plugins/FX/Niagara/Source/NiagaraEditor/Private/GlobalDistanceFieldCapture.cpp
+++ b/Engine/Plugins/FX/Niagara/Source/NiagaraEditor/Private/GlobalDistanceFieldCapture.cpp
@@ -4,7 +4,6 @@
 #include "ContentBrowserModule.h"
 #include "Editor.h"
 #include "Engine/VolumeTexture.h"
-#include "GlobalDistanceField.h"
 #include "GlobalDistanceFieldReadback.h"
 #include "HAL/IConsoleManager.h"
 #include "IContentBrowserSingleton.h"
diff --git a/Engine/Plugins/FX/Niagara/Source/NiagaraShader/NiagaraShader.Build.cs b/Engine/Plugins/FX/Niagara/Source/NiagaraShader/NiagaraShader.Build.cs
--- a/Engine/Plugins/FX/Niagara/Source/NiagaraShader/NiagaraShader.Build.cs
+++ b/Engine/Plugins/FX/Niagara/Source/NiagaraShader/NiagaraShader.Build.cs
@@ -47,7 +47,6 @@ public class NiagaraShader : ModuleRules
             });
 
 		PrivateIncludePaths.AddRange(new string[] {
-			System.IO.Path.Combine(GetModuleDirectory("Renderer"), "Private"),
 		});
 
 		PrivateIncludePathModuleNames.AddRange(
