Restoring Renderer/Private include path to allow time to addres the newer includes.

diff --git a/Engine/Plugins/FX/Niagara/Source/Niagara/Niagara.Build.cs b/Engine/Plugins/FX/Niagara/Source/Niagara/Niagara.Build.cs
--- a/Engine/Plugins/FX/Niagara/Source/Niagara/Niagara.Build.cs
+++ b/Engine/Plugins/FX/Niagara/Source/Niagara/Niagara.Build.cs
@@ -68,6 +68,7 @@ public class Niagara : ModuleRules
         PrivateIncludePaths.AddRange(
             new string[] {
 				System.IO.Path.Combine(GetModuleDirectory("Engine"), "Private"),
+				System.IO.Path.Combine(GetModuleDirectory("Renderer"), "Private"),
 			});
 
         if (Target.bBuildEditor == true)
