- Fix for module / dynamic inputs failing to compile due to script usage when using ENiagaraScriptUsageMask helpers

#rnx

diff --git a/Engine/Plugins/FX/Niagara/Source/Niagara/Public/NiagaraCommon.h b/Engine/Plugins/FX/Niagara/Source/Niagara/Public/NiagaraCommon.h
--- a/Engine/Plugins/FX/Niagara/Source/Niagara/Public/NiagaraCommon.h
+++ b/Engine/Plugins/FX/Niagara/Source/Niagara/Public/NiagaraCommon.h
@@ -916,12 +916,18 @@ namespace ENiagaraScriptUsageMask
 	enum Type
 	{
 		System =
+			(1 << int32(ENiagaraScriptUsage::Module)) |
+			(1 << int32(ENiagaraScriptUsage::DynamicInput)) |
 			(1 << int32(ENiagaraScriptUsage::SystemSpawnScript)) |
 			(1 << int32(ENiagaraScriptUsage::SystemUpdateScript)),
 		Emitter =
+			(1 << int32(ENiagaraScriptUsage::Module)) |
+			(1 << int32(ENiagaraScriptUsage::DynamicInput)) |
 			(1 << int32(ENiagaraScriptUsage::EmitterSpawnScript)) |
 			(1 << int32(ENiagaraScriptUsage::EmitterUpdateScript)),
 		Particle =
+			(1 << int32(ENiagaraScriptUsage::Module)) |
+			(1 << int32(ENiagaraScriptUsage::DynamicInput)) |
 			(1 << int32(ENiagaraScriptUsage::ParticleSpawnScript)) |
 			(1 << int32(ENiagaraScriptUsage::ParticleSpawnScriptInterpolated)) |
 			(1 << int32(ENiagaraScriptUsage::ParticleUpdateScript)) |
