Fixed output & transient parameters appearing outside their own stage. This does not fix drag & drop yet.

diff --git a/Engine/Plugins/FX/Niagara/Source/NiagaraEditor/Private/ViewModels/Stack/NiagaraStackFunctionInput.cpp b/Engine/Plugins/FX/Niagara/Source/NiagaraEditor/Private/ViewModels/Stack/NiagaraStackFunctionInput.cpp
--- a/Engine/Plugins/FX/Niagara/Source/NiagaraEditor/Private/ViewModels/Stack/NiagaraStackFunctionInput.cpp
+++ b/Engine/Plugins/FX/Niagara/Source/NiagaraEditor/Private/ViewModels/Stack/NiagaraStackFunctionInput.cpp
@@ -1519,6 +1519,12 @@ void UNiagaraStackFunctionInput::GetAvailableParameterHandles(TArray<FNiagaraPar
 							}
 						}
 
+						// as transient variables & output variables are only valid within their current stage, we may want to skip the current parameter handle
+						if(OutputNode != CurrentOutputNode && (AvailableHandle.IsOutputHandle() || AvailableHandle.IsTransientHandle()))
+						{
+							continue;
+						}
+
 						if (bWritten)
 						{
 							if (FNiagaraEditorUtilities::AreTypesAssignable(HistoryVariable.GetType(), InputType))
