- Add Set method to simple counter
- This is really exchange under the hood but is more obvious for end users

#okforgithuib public

diff --git a/Engine/Plugins/FX/Niagara/Shaders/Private/NiagaraDataInterfaceSimpleCounterTemplate.ush b/Engine/Plugins/FX/Niagara/Shaders/Private/NiagaraDataInterfaceSimpleCounterTemplate.ush
--- a/Engine/Plugins/FX/Niagara/Shaders/Private/NiagaraDataInterfaceSimpleCounterTemplate.ush
+++ b/Engine/Plugins/FX/Niagara/Shaders/Private/NiagaraDataInterfaceSimpleCounterTemplate.ush
@@ -8,6 +8,15 @@ void Get_{ParameterName}(out int CurrValue)
 	CurrValue = RWInstanceCounts[{ParameterName}_CountOffset];
 }
 
+void Set_{ParameterName}_UEImpureCall(in bool bExecute, in int NewValue)
+{
+	if ( bExecute )
+	{
+		int PrevValue;
+		InterlockedExchange(RWInstanceCounts[{ParameterName}_CountOffset], NewValue, PrevValue);
+	}
+}
+
 void Exchange_{ParameterName}_UEImpureCall(in bool bExecute, in int NewValue, out int PrevValue)
 {
 	if ( bExecute )
diff --git a/Engine/Plugins/FX/Niagara/Source/Niagara/Private/DataInterface/NiagaraDataInterfaceSimpleCounter.cpp b/Engine/Plugins/FX/Niagara/Source/Niagara/Private/DataInterface/NiagaraDataInterfaceSimpleCounter.cpp
--- a/Engine/Plugins/FX/Niagara/Source/Niagara/Private/DataInterface/NiagaraDataInterfaceSimpleCounter.cpp
+++ b/Engine/Plugins/FX/Niagara/Source/Niagara/Private/DataInterface/NiagaraDataInterfaceSimpleCounter.cpp
@@ -24,6 +24,7 @@ namespace NDISimpleCounterLocal
 	static const TCHAR* TemplateShaderFile = TEXT("/Plugin/FX/Niagara/Private/NiagaraDataInterfaceSimpleCounterTemplate.ush");
 	static const FName NAME_GetNextValue_Deprecated(TEXT("GetNextValue"));
 	static const FName NAME_Get(TEXT("Get"));
+	static const FName NAME_Set(TEXT("Set"));
 	static const FName NAME_Exchange(TEXT("Exchange"));
 	static const FName NAME_Add(TEXT("Add"));
 	static const FName NAME_Increment(TEXT("Increment"));
@@ -285,6 +286,17 @@ void UNiagaraDataInterfaceSimpleCounter::GetFunctions(TArray<FNiagaraFunctionSig
 		Sig.Outputs.Add(FNiagaraVariable(FNiagaraTypeDefinition::GetIntDef(), TEXT("Current Value")));
 		Sig.SetDescription(LOCTEXT("GetDesc", "Gets the current value of the counter."));
 	}
+	{
+		FNiagaraFunctionSignature& Sig = OutFunctions.AddDefaulted_GetRef();
+		Sig.Name = NAME_Set;
+		Sig.bMemberFunction = true;
+		Sig.bRequiresContext = false;
+		Sig.bRequiresExecPin = true;
+		Sig.Inputs.Add(FNiagaraVariable(FNiagaraTypeDefinition(GetClass()), TEXT("Counter")));
+		Sig.Inputs.Add_GetRef(FNiagaraVariable(FNiagaraTypeDefinition::GetBoolDef(), TEXT("Execute"))).SetValue(true);
+		Sig.Inputs.Add(FNiagaraVariable(FNiagaraTypeDefinition::GetIntDef(), TEXT("New Value")));
+		Sig.SetDescription(LOCTEXT("SetDesc", "Sets the current value of the counter."));
+	}
 	{
 		FNiagaraFunctionSignature& Sig = OutFunctions.AddDefaulted_GetRef();
 		Sig.Name = NAME_Exchange;
@@ -348,6 +360,10 @@ void UNiagaraDataInterfaceSimpleCounter::GetVMExternalFunction(const FVMExternal
 	{
 		OutFunc = FVMExternalFunction::CreateUObject(this, &UNiagaraDataInterfaceSimpleCounter::VMGet);
 	}
+	else if (BindingInfo.Name == NAME_Set)
+	{
+		OutFunc = FVMExternalFunction::CreateUObject(this, &UNiagaraDataInterfaceSimpleCounter::VMSet);
+	}
 	else if (BindingInfo.Name == NAME_Exchange)
 	{
 		OutFunc = FVMExternalFunction::CreateUObject(this, &UNiagaraDataInterfaceSimpleCounter::VMExchange);
@@ -462,6 +478,29 @@ void UNiagaraDataInterfaceSimpleCounter::VMGet(FVectorVMExternalFunctionContext&
 	}
 }
 
+void UNiagaraDataInterfaceSimpleCounter::VMSet(FVectorVMExternalFunctionContext& Context)
+{
+	VectorVM::FUserPtrHandler<FNDISimpleCounterInstanceData_GameThread> InstanceData(Context);
+	FNDIInputParam<bool> InExecute(Context);
+	FNDIInputParam<int32> InValue(Context);
+
+	for (int32 i=0; i < Context.GetNumInstances(); ++i)
+	{
+		const bool bExecute = InExecute.GetAndAdvance();
+		const int32 NewValue = InValue.GetAndAdvance();
+		if (bExecute)
+		{
+			const int32 PrevValue = InstanceData->Counter.exchange(NewValue);
+		}
+	}
+
+	if (FNiagaraUtilities::ShouldSyncCpuToGpu(GpuSyncMode))
+	{
+		InstanceData->bModified = true;
+		MarkRenderDataDirty();
+	}
+}
+
 void UNiagaraDataInterfaceSimpleCounter::VMExchange(FVectorVMExternalFunctionContext& Context)
 {	
 	VectorVM::FUserPtrHandler<FNDISimpleCounterInstanceData_GameThread> InstanceData(Context);
diff --git a/Engine/Plugins/FX/Niagara/Source/Niagara/Private/DataInterface/NiagaraDataInterfaceSimpleCounter.h b/Engine/Plugins/FX/Niagara/Source/Niagara/Private/DataInterface/NiagaraDataInterfaceSimpleCounter.h
--- a/Engine/Plugins/FX/Niagara/Source/Niagara/Private/DataInterface/NiagaraDataInterfaceSimpleCounter.h
+++ b/Engine/Plugins/FX/Niagara/Source/Niagara/Private/DataInterface/NiagaraDataInterfaceSimpleCounter.h
@@ -56,6 +56,7 @@ public:
 
 	// VM functions
 	void VMGet(FVectorVMExternalFunctionContext& Context);
+	void VMSet(FVectorVMExternalFunctionContext& Context);
 	void VMExchange(FVectorVMExternalFunctionContext& Context);
 	void VMAdd(FVectorVMExternalFunctionContext& Context);
 	void VMIncrement(FVectorVMExternalFunctionContext& Context);
