- Fix previous interpolation bug (CurrPos vs CurrPosA)

#rnx

diff --git a/Engine/Plugins/FX/Niagara/Source/Niagara/Private/NiagaraSimCacheHelper.h b/Engine/Plugins/FX/Niagara/Source/Niagara/Private/NiagaraSimCacheHelper.h
--- a/Engine/Plugins/FX/Niagara/Source/Niagara/Private/NiagaraSimCacheHelper.h
+++ b/Engine/Plugins/FX/Niagara/Source/Niagara/Private/NiagaraSimCacheHelper.h
@@ -904,7 +904,6 @@ struct FNiagaraSimCacheHelper
 	template<bool bWithRebase, bool bWithPrevious>
 	static void InterpPositions(const FNiagaraSimCacheDataBuffersLayout::FVariableCopyContext& CopyDataContext)
 	{
-
 		FSoAVec3Reader SrcAPositions(CopyDataContext.SourceACurr, CopyDataContext.SourceAStride);
 		FSoAVec3Reader SrcBPositions(CopyDataContext.SourceBCurr, CopyDataContext.SourceBStride);
 		FSoAVec3Writer DstPositions(CopyDataContext.DestCurr, CopyDataContext.DestStride);
@@ -923,7 +922,7 @@ struct FNiagaraSimCacheHelper
 				DstPositions.Set(iInstanceA, bWithRebase ? FVector3f(CopyDataContext.RebaseTransform.TransformPosition(FVector(CurrPos))) : CurrPos);
 
 				const FVector3f PrevPosA = SrcAPrevPositions.Get(iInstanceA);
-				const FVector3f PrevPos = FMath::Lerp(PrevPosA, CurrPosA, CopyDataContext.PrevFrameFraction);
+				const FVector3f PrevPos = FMath::Lerp(PrevPosA, CurrPos, CopyDataContext.PrevFrameFraction);
 				DstPrevPositions.Set(iInstanceA, bWithRebase ? FVector3f(CopyDataContext.RebaseTransform.TransformPosition(FVector(PrevPos))) : PrevPos);
 			}
 		}
