- Handle feature level change for materials in system editor

diff --git a/Engine/Plugins/FX/Niagara/Source/NiagaraEditor/Private/Widgets/SNiagaraSystemViewport.cpp b/Engine/Plugins/FX/Niagara/Source/NiagaraEditor/Private/Widgets/SNiagaraSystemViewport.cpp
--- a/Engine/Plugins/FX/Niagara/Source/NiagaraEditor/Private/Widgets/SNiagaraSystemViewport.cpp
+++ b/Engine/Plugins/FX/Niagara/Source/NiagaraEditor/Private/Widgets/SNiagaraSystemViewport.cpp
@@ -22,6 +22,7 @@
 #include "NiagaraSystemEditorData.h"
 #include "SNiagaraSystemViewportToolBar.h"
 #include "UnrealEdGlobals.h"
+#include "Editor/EditorEngine.h"
 #include "Editor/UnrealEdEngine.h"
 #include "Engine/Canvas.h"
 #include "Engine/Font.h"
@@ -499,10 +500,32 @@ void SNiagaraSystemViewport::Construct(const FArguments& InArgs)
 	SEditorViewport::Construct( SEditorViewport::FArguments() );
 
 	Client->EngineShowFlags.SetGrid(Settings->IsShowGridInViewport());
+
+	// Register for preview feature level changes
+	if (UEditorEngine* EditorEngine = Cast<UEditorEngine>(GEngine))
+	{
+		OnPreviewFeatureLevelChangedHandle = EditorEngine->OnPreviewFeatureLevelChanged().AddLambda(
+			[WeakWorld=TWeakObjectPtr<UWorld>(Client->GetWorld())](ERHIFeatureLevel::Type NewFeatureLevel)
+			{
+				if (UWorld* World = WeakWorld.Get())
+				{
+					World->ChangeFeatureLevel(NewFeatureLevel);
+				}
+			}
+		);
+	}
 }
 
 SNiagaraSystemViewport::~SNiagaraSystemViewport()
 {
+	if (OnPreviewFeatureLevelChangedHandle.IsValid())
+	{
+		if (UEditorEngine* EditorEngine = Cast<UEditorEngine>(GEngine))
+		{
+			EditorEngine->OnPreviewFeatureLevelChanged().Remove(OnPreviewFeatureLevelChangedHandle);
+		}
+	}
+
 	if (SystemViewportClient.IsValid())
 	{
 		SystemViewportClient->Viewport = NULL;
diff --git a/Engine/Plugins/FX/Niagara/Source/NiagaraEditor/Private/Widgets/SNiagaraSystemViewport.h b/Engine/Plugins/FX/Niagara/Source/NiagaraEditor/Private/Widgets/SNiagaraSystemViewport.h
--- a/Engine/Plugins/FX/Niagara/Source/NiagaraEditor/Private/Widgets/SNiagaraSystemViewport.h
+++ b/Engine/Plugins/FX/Niagara/Source/NiagaraEditor/Private/Widgets/SNiagaraSystemViewport.h
@@ -138,6 +138,8 @@ private:
 
 	/** True if orbit mode was active before we started a view transition. Used to restore orbit mode at the the end of the transition */
 	bool bShouldActivateOrbitAfterTransitioning = false;
+
+	FDelegateHandle OnPreviewFeatureLevelChangedHandle;
 };
 
 #if NIAGARA_PERF_BASELINES
