- Round SimCache frame to nearest if we do not have interpolation enabled
- This helps reduce / remove ghosting with MRQ TSS renders

diff --git a/Engine/Plugins/FX/Niagara/Source/Niagara/Private/NiagaraSimCache.cpp b/Engine/Plugins/FX/Niagara/Source/Niagara/Private/NiagaraSimCache.cpp
--- a/Engine/Plugins/FX/Niagara/Source/Niagara/Private/NiagaraSimCache.cpp
+++ b/Engine/Plugins/FX/Niagara/Source/Niagara/Private/NiagaraSimCache.cpp
@@ -598,6 +598,11 @@ bool UNiagaraSimCache::Read(float TimeSeconds, FNiagaraSystemInstance* SystemIns
 		const float SimulationAgeA = CacheFrames[FrameIndex].SimulationAge;
 		const float SimulationAgeB = CacheFrames[FMath::Min(FrameIndex + 1, NumFramesMinusOne)].SimulationAge;
 		FrameFraction = (TimeSeconds - SimulationAgeA) / (SimulationAgeB - SimulationAgeA);
+		if (!CreateParameters.bAllowInterpolation && !CreateParameters.bAllowVelocityExtrapolation && FrameFraction > 0.5f)
+		{
+			FrameFraction = 0.0f;
+			FrameIndex = FrameIndex + 1;
+		}
 	}
 
 	return ReadFrame(FrameIndex, FrameFraction, SystemInstance);
