* Fixed "parameter for" node
* Cleaned up some niagara node code
* Fixed pin interactions for parameter map set and get nodes

diff --git a/Engine/Plugins/FX/Niagara/Source/NiagaraEditor/Private/NiagaraNodeParameterMapBase.cpp b/Engine/Plugins/FX/Niagara/Source/NiagaraEditor/Private/NiagaraNodeParameterMapBase.cpp
--- a/Engine/Plugins/FX/Niagara/Source/NiagaraEditor/Private/NiagaraNodeParameterMapBase.cpp
+++ b/Engine/Plugins/FX/Niagara/Source/NiagaraEditor/Private/NiagaraNodeParameterMapBase.cpp
@@ -12,7 +12,6 @@
 #include "NiagaraConstants.h"
 #include "ToolMenus.h"
 #include "NiagaraScriptVariable.h"
-#include "AssetRegistry/AssetRegistryModule.h"
 
 #include UE_INLINE_GENERATED_CPP_BY_NAME(NiagaraNodeParameterMapBase)
 
diff --git a/Engine/Plugins/FX/Niagara/Source/NiagaraEditor/Private/NiagaraNodeParameterMapBase.h b/Engine/Plugins/FX/Niagara/Source/NiagaraEditor/Private/NiagaraNodeParameterMapBase.h
--- a/Engine/Plugins/FX/Niagara/Source/NiagaraEditor/Private/NiagaraNodeParameterMapBase.h
+++ b/Engine/Plugins/FX/Niagara/Source/NiagaraEditor/Private/NiagaraNodeParameterMapBase.h
@@ -33,8 +33,8 @@ public:
 
 	virtual void GetNodeContextMenuActions(class UToolMenu* Menu, class UGraphNodeContextMenuContext* Context) const override;
 	virtual void GetPinHoverText(const UEdGraphPin& Pin, FString& HoverTextOut) const override;
-
-	virtual bool IncludeParentNodeContextMenu() const { return true; }
+	virtual bool CanModifyExistingPins() const override { return false; }
+	virtual bool IncludeParentNodeContextMenu() const override { return true; }
 
 	void SetPinName(UEdGraphPin* InPin, const FName& InName);
 
diff --git a/Engine/Plugins/FX/Niagara/Source/NiagaraEditor/Private/NiagaraNodeParameterMapFor.cpp b/Engine/Plugins/FX/Niagara/Source/NiagaraEditor/Private/NiagaraNodeParameterMapFor.cpp
--- a/Engine/Plugins/FX/Niagara/Source/NiagaraEditor/Private/NiagaraNodeParameterMapFor.cpp
+++ b/Engine/Plugins/FX/Niagara/Source/NiagaraEditor/Private/NiagaraNodeParameterMapFor.cpp
@@ -5,23 +5,14 @@
 #include "NiagaraEditorUtilities.h"
 #include "SNiagaraGraphNodeConvert.h"
 #include "NiagaraHlslTranslator.h"
-#include "Templates/SharedPointer.h"
-#include "NiagaraGraph.h"
-#include "NiagaraConstants.h"
-#include "Framework/MultiBox/MultiBoxBuilder.h"
-#include "Widgets/Layout/SBox.h"
-#include "Widgets/Input/SEditableTextBox.h"
 #include "EdGraph/EdGraphNode.h"
 
-#include "ScopedTransaction.h"
-
 #include UE_INLINE_GENERATED_CPP_BY_NAME(NiagaraNodeParameterMapFor)
 
 #define LOCTEXT_NAMESPACE "NiagaraNodeParameterMapFor" 
 
 UNiagaraNodeParameterMapFor::UNiagaraNodeParameterMapFor() : UNiagaraNodeParameterMapSet()
 {
-
 	UEdGraphNode::NodeUpgradeMessage = LOCTEXT("NodeExperimental", "This node is marked as experimental, use with care!");
 }
 
@@ -30,11 +21,24 @@ void UNiagaraNodeParameterMapFor::AllocateDefaultPins()
 	PinPendingRename = nullptr;
 	const UEdGraphSchema_Niagara* Schema = GetDefault<UEdGraphSchema_Niagara>();
 	CreatePin(EGPD_Input, Schema->TypeDefinitionToPinType(FNiagaraTypeDefinition::GetParameterMapDef()), *UNiagaraNodeParameterMapBase::SourcePinName.ToString());
-	CreatePin(EGPD_Input, Schema->TypeDefinitionToPinType(FNiagaraTypeDefinition::GetIntDef()), TEXT("Module.Iteration Count"));
+	UEdGraphPin* IterationPin = CreatePin(EGPD_Input, Schema->TypeDefinitionToPinType(FNiagaraTypeDefinition::GetIntDef()), TEXT("Iteration Count"));
+	IterationPin->PinToolTip = TEXT("How often the for loop should be executed - will only be evaluated once before executing the loop!");
 	CreatePin(EGPD_Output, Schema->TypeDefinitionToPinType(FNiagaraTypeDefinition::GetParameterMapDef()), *UNiagaraNodeParameterMapBase::DestPinName.ToString());
 	CreateAddPin(EGPD_Input);
 }
 
+void UNiagaraNodeParameterMapFor::GetPinHoverText(const UEdGraphPin& Pin, FString& HoverTextOut) const
+{
+	FPinCollectorArray InputPins;
+	GetInputPins(InputPins);
+	if (&Pin == InputPins[1])
+	{
+		HoverTextOut = LOCTEXT("IterationCountPinTooltip", "How often the loop should be executed - will only be evaluated once before executing the loop!").ToString();
+		return;
+	}
+	Super::GetPinHoverText(Pin, HoverTextOut);
+}
+
 void UNiagaraNodeParameterMapFor::Compile(FHlslNiagaraTranslator* Translator, TArray<int32>& Outputs)
 {
 	FPinCollectorArray InputPins;
@@ -63,6 +67,13 @@ FText UNiagaraNodeParameterMapFor::GetNodeTitle(ENodeTitleType::Type TitleType)
 	return LOCTEXT("UNiagaraNodeParameterMapForName", "Map For");
 }
 
+bool UNiagaraNodeParameterMapFor::SkipPinCompilation(UEdGraphPin* Pin) const
+{
+	FPinCollectorArray InputPins;
+	GetInputPins(InputPins);
+	return Pin == InputPins[1];
+}
+
 UNiagaraNodeParameterMapForWithContinue::UNiagaraNodeParameterMapForWithContinue() : UNiagaraNodeParameterMapFor()
 {
 	UEdGraphNode::NodeUpgradeMessage = LOCTEXT("NodeExperimental", "This node is marked as experimental, use with care!");
@@ -79,11 +90,9 @@ void UNiagaraNodeParameterMapForWithContinue::Compile(FHlslNiagaraTranslator* Tr
 			const int32 IterationCount = Translator->CompilePin(InputPins[1]);
 
 			Translator->ParameterMapForBegin(this, IterationCount);
-			if (InputPins[2])
-			{
-				const int32 IterationEnabledChunk = Translator->CompilePin(InputPins[2]);
-				Translator->ParameterMapForContinue(this, IterationEnabledChunk);
-			}
+			
+			const int32 IterationEnabledChunk = Translator->CompilePin(InputPins[2]);
+			Translator->ParameterMapForContinue(this, IterationEnabledChunk);
 
 			UNiagaraNodeParameterMapSet::Compile(Translator, Outputs);
 
@@ -102,8 +111,8 @@ void UNiagaraNodeParameterMapForWithContinue::AllocateDefaultPins()
 	PinPendingRename = nullptr;
 	const UEdGraphSchema_Niagara* Schema = GetDefault<UEdGraphSchema_Niagara>();
 	CreatePin(EGPD_Input, Schema->TypeDefinitionToPinType(FNiagaraTypeDefinition::GetParameterMapDef()), *UNiagaraNodeParameterMapBase::SourcePinName.ToString());
-	CreatePin(EGPD_Input, Schema->TypeDefinitionToPinType(FNiagaraTypeDefinition::GetIntDef()), TEXT("Module.Iteration Count"));
-	CreatePin(EGPD_Input, Schema->TypeDefinitionToPinType(FNiagaraTypeDefinition::GetBoolDef()), TEXT("Module.Iteration Enabled"));
+	CreatePin(EGPD_Input, Schema->TypeDefinitionToPinType(FNiagaraTypeDefinition::GetIntDef()), TEXT("Iteration Count"));
+	CreatePin(EGPD_Input, Schema->TypeDefinitionToPinType(FNiagaraTypeDefinition::GetBoolDef()), TEXT("Iteration Enabled"));
 	CreatePin(EGPD_Output, Schema->TypeDefinitionToPinType(FNiagaraTypeDefinition::GetParameterMapDef()), *UNiagaraNodeParameterMapBase::DestPinName.ToString());
 	CreateAddPin(EGPD_Input);
 }
@@ -113,6 +122,25 @@ FText UNiagaraNodeParameterMapForWithContinue::GetNodeTitle(ENodeTitleType::Type
 	return LOCTEXT("UNiagaraNodeParameterMapForWithContinueName", "Map For With Continue");
 }
 
+void UNiagaraNodeParameterMapForWithContinue::GetPinHoverText(const UEdGraphPin& Pin, FString& HoverTextOut) const
+{
+	FPinCollectorArray InputPins;
+	GetInputPins(InputPins);
+	if (&Pin == InputPins[2])
+	{
+		HoverTextOut = LOCTEXT("IterationEnabledPinTooltip", "If set to false the loop will continue with the next iteration and execute no code - will be evaluated at the *start* of every iteration.").ToString();
+		return;
+	}
+	Super::GetPinHoverText(Pin, HoverTextOut);
+}
+
+bool UNiagaraNodeParameterMapForWithContinue::SkipPinCompilation(UEdGraphPin* Pin) const
+{
+	FPinCollectorArray InputPins;
+	GetInputPins(InputPins);
+	return Pin == InputPins[1] || Pin == InputPins[2];
+}
+
 UNiagaraNodeParameterMapForIndex::UNiagaraNodeParameterMapForIndex() : UNiagaraNode()
 {
 	UEdGraphNode::NodeUpgradeMessage = LOCTEXT("NodeExperimental", "This node is marked as experimental, use with care!");
diff --git a/Engine/Plugins/FX/Niagara/Source/NiagaraEditor/Private/NiagaraNodeParameterMapFor.h b/Engine/Plugins/FX/Niagara/Source/NiagaraEditor/Private/NiagaraNodeParameterMapFor.h
--- a/Engine/Plugins/FX/Niagara/Source/NiagaraEditor/Private/NiagaraNodeParameterMapFor.h
+++ b/Engine/Plugins/FX/Niagara/Source/NiagaraEditor/Private/NiagaraNodeParameterMapFor.h
@@ -20,6 +20,11 @@ public:
 	virtual void Compile(class FHlslNiagaraTranslator* Translator, TArray<int32>& Outputs) override;
 
 	virtual FText GetNodeTitle(ENodeTitleType::Type TitleType) const override;
+	virtual bool IncludeParentNodeContextMenu() const override { return false; }
+	virtual void GetPinHoverText(const UEdGraphPin& Pin, FString& HoverTextOut) const override;
+	
+protected:
+	virtual bool SkipPinCompilation(UEdGraphPin* Pin) const override;
 };
 
 UCLASS()
@@ -33,6 +38,11 @@ public:
 	virtual void AllocateDefaultPins() override;
 	virtual void Compile(class FHlslNiagaraTranslator* Translator, TArray<int32>& Outputs) override;
 	virtual FText GetNodeTitle(ENodeTitleType::Type TitleType) const override;
+	virtual bool IncludeParentNodeContextMenu() const override { return false; }
+	virtual void GetPinHoverText(const UEdGraphPin& Pin, FString& HoverTextOut) const override;
+	
+protected:
+	virtual bool SkipPinCompilation(UEdGraphPin* Pin) const override;
 };
 
 UCLASS()
diff --git a/Engine/Plugins/FX/Niagara/Source/NiagaraEditor/Private/NiagaraNodeParameterMapGet.cpp b/Engine/Plugins/FX/Niagara/Source/NiagaraEditor/Private/NiagaraNodeParameterMapGet.cpp
--- a/Engine/Plugins/FX/Niagara/Source/NiagaraEditor/Private/NiagaraNodeParameterMapGet.cpp
+++ b/Engine/Plugins/FX/Niagara/Source/NiagaraEditor/Private/NiagaraNodeParameterMapGet.cpp
@@ -3,7 +3,6 @@
 #include "NiagaraNodeParameterMapGet.h"
 #include "EdGraphSchema_Niagara.h"
 #include "NiagaraEditorUtilities.h"
-#include "SNiagaraGraphNodeConvert.h"
 #include "NiagaraGraph.h"
 #include "NiagaraHlslTranslator.h"
 #include "ScopedTransaction.h"
@@ -13,13 +12,9 @@
 #include "Modules/ModuleManager.h"
 #include "Templates/SharedPointer.h"
 #include "NiagaraConstants.h"
-#include "ToolMenus.h"
-#include "Widgets/Layout/SBox.h"
-#include "Widgets/Input/SEditableTextBox.h"
 #include "EdGraph/EdGraphNode.h"
 #include "NiagaraParameterCollection.h"
 #include "NiagaraScriptVariable.h"
-#include "NiagaraConstants.h"
 
 #include UE_INLINE_GENERATED_CPP_BY_NAME(NiagaraNodeParameterMapGet)
 
@@ -49,14 +44,7 @@ bool UNiagaraNodeParameterMapGet::IsPinNameEditable(const UEdGraphPin* GraphPinO
 {
 	const UEdGraphSchema_Niagara* Schema = GetDefault<UEdGraphSchema_Niagara>();
 	FNiagaraTypeDefinition TypeDef = Schema->PinToTypeDefinition(GraphPinObj);
-	if (TypeDef.IsValid() && GraphPinObj && GraphPinObj->Direction == EGPD_Output && CanRenamePin(GraphPinObj))
-	{
-		return true;
-	}
-	else
-	{
-		return false;
-	}
+	return TypeDef.IsValid() && GraphPinObj && GraphPinObj->Direction == EGPD_Output && CanRenamePin(GraphPinObj);
 }
 
 bool UNiagaraNodeParameterMapGet::IsPinNameEditableUponCreation(const UEdGraphPin* GraphPinObj) const
@@ -683,11 +671,5 @@ void UNiagaraNodeParameterMapGet::GetPinHoverText(const UEdGraphPin& Pin, FStrin
 	}
 }
 
-void UNiagaraNodeParameterMapGet::GetNodeContextMenuActions(UToolMenu* Menu, UGraphNodeContextMenuContext* Context) const
-{
-	Super::GetNodeContextMenuActions(Menu, Context);
-
-}
-
 #undef LOCTEXT_NAMESPACE
 
diff --git a/Engine/Plugins/FX/Niagara/Source/NiagaraEditor/Private/NiagaraNodeParameterMapGet.h b/Engine/Plugins/FX/Niagara/Source/NiagaraEditor/Private/NiagaraNodeParameterMapGet.h
--- a/Engine/Plugins/FX/Niagara/Source/NiagaraEditor/Private/NiagaraNodeParameterMapGet.h
+++ b/Engine/Plugins/FX/Niagara/Source/NiagaraEditor/Private/NiagaraNodeParameterMapGet.h
@@ -21,19 +21,18 @@ public:
 	virtual void AllocateDefaultPins() override;
 	virtual TSharedPtr<SGraphNode> CreateVisualWidget() override;
 
-	virtual void GetNodeContextMenuActions(class UToolMenu* Menu, class UGraphNodeContextMenuContext* Context) const override;
-	virtual bool IncludeParentNodeContextMenu() const { return true; }
+	virtual bool IncludeParentNodeContextMenu() const override { return true; }
 	virtual bool IsPinNameEditable(const UEdGraphPin* GraphPinObj) const override;
 	virtual bool IsPinNameEditableUponCreation(const UEdGraphPin* GraphPinObj) const override;
 	virtual bool VerifyEditablePinName(const FText& InName, FText& OutErrorMessage, const UEdGraphPin* InGraphPinObj) const override;
 	virtual bool CommitEditablePinName(const FText& InName, UEdGraphPin* InGraphPinObj, bool bSuppressEvents)  override;
 	virtual bool CancelEditablePinName(const FText& InName, UEdGraphPin* InGraphPinObj) override;
 
-	virtual void Compile(class FHlslNiagaraTranslator* Translator, TArray<int32>& Outputs);
+	virtual void Compile(class FHlslNiagaraTranslator* Translator, TArray<int32>& Outputs) override;
 
 	virtual void BuildParameterMapHistory(FNiagaraParameterMapHistoryBuilder& OutHistory, bool bRecursive = true, bool bFilterForCompilation = true) const override;
 
-	virtual FText GetNodeTitle(ENodeTitleType::Type TitleType) const;
+	virtual FText GetNodeTitle(ENodeTitleType::Type TitleType) const override;
 
 	void AddOrphanedPinPairGuids(UEdGraphPin* OutputPin, UEdGraphPin* DefaultPin);
 	
@@ -44,13 +43,13 @@ public:
 	
 	virtual void PostLoad() override;
 
-	void GatherExternalDependencyData(ENiagaraScriptUsage InUsage, const FGuid& InUsageId, TArray<FNiagaraCompileHash>& InReferencedCompileHashes, TArray<FString>& InReferencedObjs) const override;
+	virtual void GatherExternalDependencyData(ENiagaraScriptUsage InUsage, const FGuid& InUsageId, TArray<FNiagaraCompileHash>& InReferencedCompileHashes, TArray<FString>& InReferencedObjs) const override;
 	virtual void GetPinHoverText(const UEdGraphPin& Pin, FString& HoverTextOut) const override;
 
-	virtual FName GetNewPinDefaultNamespace() const { return PARAM_MAP_MODULE_STR; }
+	virtual FName GetNewPinDefaultNamespace() const override { return PARAM_MAP_MODULE_STR; }
 
 	/** Convenience method to determine whether this Node is a Map Get or Map Set when adding a parameter through the parameter panel. */
-	virtual EEdGraphPinDirection GetPinDirectionForNewParameters() { return EEdGraphPinDirection::EGPD_Output; };
+	virtual EEdGraphPinDirection GetPinDirectionForNewParameters() override { return EEdGraphPinDirection::EGPD_Output; };
 
 protected:
 	virtual void OnNewTypedPinAdded(UEdGraphPin*& NewPin) override;
diff --git a/Engine/Plugins/FX/Niagara/Source/NiagaraEditor/Private/NiagaraNodeParameterMapSet.cpp b/Engine/Plugins/FX/Niagara/Source/NiagaraEditor/Private/NiagaraNodeParameterMapSet.cpp
--- a/Engine/Plugins/FX/Niagara/Source/NiagaraEditor/Private/NiagaraNodeParameterMapSet.cpp
+++ b/Engine/Plugins/FX/Niagara/Source/NiagaraEditor/Private/NiagaraNodeParameterMapSet.cpp
@@ -11,12 +11,8 @@
 #include "NiagaraGraph.h"
 #include "NiagaraScriptVariable.h"
 #include "ScopedTransaction.h"
-#include "SNiagaraGraphNodeConvert.h"
 #include "SNiagaraGraphParameterMapSetNode.h"
 #include "Templates/SharedPointer.h"
-#include "ToolMenus.h"
-#include "Widgets/Input/SEditableTextBox.h"
-#include "Widgets/Layout/SBox.h"
 
 #include UE_INLINE_GENERATED_CPP_BY_NAME(NiagaraNodeParameterMapSet)
 
@@ -45,33 +41,12 @@ bool UNiagaraNodeParameterMapSet::IsPinNameEditable(const UEdGraphPin* GraphPinO
 {
 	const UEdGraphSchema_Niagara* Schema = GetDefault<UEdGraphSchema_Niagara>();
 	FNiagaraTypeDefinition TypeDef = Schema->PinToTypeDefinition(GraphPinObj);
-	if (TypeDef.IsValid() && GraphPinObj && GraphPinObj->Direction == EGPD_Input && CanRenamePin(GraphPinObj))
-	{
-		return true;
-	}
-	else
-	{
-		return false;
-	}
+	return TypeDef.IsValid() && GraphPinObj && GraphPinObj->Direction == EGPD_Input && CanRenamePin(GraphPinObj);
 }
 
 bool UNiagaraNodeParameterMapSet::IsPinNameEditableUponCreation(const UEdGraphPin* GraphPinObj) const
 {
-	if (GraphPinObj == PinPendingRename)
-	{
-
-		return true;
-	}
-	else
-	{
-		return false;
-	}
-}
-
-void UNiagaraNodeParameterMapSet::RemoveDynamicPin(UEdGraphPin* Pin)
-{
-	// Call NiagaraNodeWithDynamicPins::RemoveDynamicPin() instead of base class to fixup the associated pin variable's metadata.
-	Super::RemoveDynamicPin(Pin);
+	return GraphPinObj == PinPendingRename;
 }
 
 bool UNiagaraNodeParameterMapSet::VerifyEditablePinName(const FText& InName, FText& OutErrorMessage, const UEdGraphPin* InGraphPinObj) const
@@ -254,7 +229,7 @@ void UNiagaraNodeParameterMapSet::Compile(class FHlslNiagaraTranslator* Translat
 	CompileInputs.Reserve(InputPins.Num());
 	for (UEdGraphPin* InputPin : InputPins)
 	{
-		if (IsAddPin(InputPin) || Translator->IsFunctionVariableCulledFromCompilation(InputPin->PinName))
+		if (IsAddPin(InputPin) || SkipPinCompilation(InputPin) || Translator->IsFunctionVariableCulledFromCompilation(InputPin->PinName))
 		{
 			continue;
 		}
@@ -346,13 +321,6 @@ void UNiagaraNodeParameterMapSet::BuildParameterMapHistory(FNiagaraParameterMapH
 	OutHistory.RegisterParameterMapPin(ParamMapIdx, GetOutputPin(0));
 }
 
-void UNiagaraNodeParameterMapSet::GetNodeContextMenuActions(UToolMenu* Menu, UGraphNodeContextMenuContext* Context) const
-{
-	Super::GetNodeContextMenuActions(Menu, Context);
-
-
-}
-
 void UNiagaraNodeParameterMapSet::PostLoad()
 {
 	for (UEdGraphPin* Pin : Pins)
diff --git a/Engine/Plugins/FX/Niagara/Source/NiagaraEditor/Private/NiagaraNodeParameterMapSet.h b/Engine/Plugins/FX/Niagara/Source/NiagaraEditor/Private/NiagaraNodeParameterMapSet.h
--- a/Engine/Plugins/FX/Niagara/Source/NiagaraEditor/Private/NiagaraNodeParameterMapSet.h
+++ b/Engine/Plugins/FX/Niagara/Source/NiagaraEditor/Private/NiagaraNodeParameterMapSet.h
@@ -23,26 +23,25 @@ public:
 	virtual TSharedPtr<SGraphNode> CreateVisualWidget() override;
 	virtual bool IsPinNameEditable(const UEdGraphPin* GraphPinObj) const override;
 	virtual bool IsPinNameEditableUponCreation(const UEdGraphPin* GraphPinObj) const override;
-	virtual void RemoveDynamicPin(UEdGraphPin* Pin) override;
 	virtual bool VerifyEditablePinName(const FText& InName, FText& OutErrorMessage, const UEdGraphPin* InGraphPinObj) const override;
 	virtual bool CommitEditablePinName(const FText& InName, UEdGraphPin* InGraphPinObj, bool bSuppressEvents = false)  override;
 	virtual bool CancelEditablePinName(const FText& InName, UEdGraphPin* InGraphPinObj) override;
 
 	virtual void BuildParameterMapHistory(FNiagaraParameterMapHistoryBuilder& OutHistory, bool bRecursive = true, bool bFilterForCompilation = true) const override;
 
-	virtual void Compile(class FHlslNiagaraTranslator* Translator, TArray<int32>& Outputs);
+	virtual void Compile(class FHlslNiagaraTranslator* Translator, TArray<int32>& Outputs) override;
 
-	virtual FText GetNodeTitle(ENodeTitleType::Type TitleType) const;
+	virtual FText GetNodeTitle(ENodeTitleType::Type TitleType) const override;
 
-	virtual void GetNodeContextMenuActions(class UToolMenu* Menu, class UGraphNodeContextMenuContext* Context) const override;
-	virtual bool IncludeParentNodeContextMenu() const { return true; }
+	virtual bool IncludeParentNodeContextMenu() const override { return false; }
 	virtual void PostLoad() override;
-	virtual FName GetNewPinDefaultNamespace() const { return PARAM_MAP_LOCAL_MODULE_STR; }
+	virtual FName GetNewPinDefaultNamespace() const override { return PARAM_MAP_LOCAL_MODULE_STR; }
 
 	/** Convenience method to determine whether this Node is a Map Get or Map Set when adding a parameter through the parameter panel. */
-	virtual EEdGraphPinDirection GetPinDirectionForNewParameters() { return EEdGraphPinDirection::EGPD_Input; };
+	virtual EEdGraphPinDirection GetPinDirectionForNewParameters() override { return EEdGraphPinDirection::EGPD_Input; };
 
 protected:
-	virtual void OnNewTypedPinAdded(UEdGraphPin*& NewPin); 
+	virtual void OnNewTypedPinAdded(UEdGraphPin*& NewPin) override; 
 	virtual void OnPinRenamed(UEdGraphPin* RenamedPin, const FString& OldName) override;
+	virtual bool SkipPinCompilation(UEdGraphPin* Pin) const { return false; };
 };
diff --git a/Engine/Plugins/FX/Niagara/Source/NiagaraEditor/Private/Widgets/SNiagaraGraphParameterMapGetNode.h b/Engine/Plugins/FX/Niagara/Source/NiagaraEditor/Private/Widgets/SNiagaraGraphParameterMapGetNode.h
--- a/Engine/Plugins/FX/Niagara/Source/NiagaraEditor/Private/Widgets/SNiagaraGraphParameterMapGetNode.h
+++ b/Engine/Plugins/FX/Niagara/Source/NiagaraEditor/Private/Widgets/SNiagaraGraphParameterMapGetNode.h
@@ -12,7 +12,7 @@ public:
 	void Construct(const FArguments& InArgs, UEdGraphNode* InGraphNode);
 
 	virtual void AddPin(const TSharedRef<SGraphPin>& PinToAdd) override;
-	virtual TSharedRef<SWidget> CreateNodeContentArea();
+	virtual TSharedRef<SWidget> CreateNodeContentArea() override;
 	virtual void CreatePinWidgets() override;
 
 protected:
diff --git a/Engine/Plugins/FX/Niagara/Source/NiagaraEditor/Private/Widgets/SNiagaraGraphParameterMapSetNode.cpp b/Engine/Plugins/FX/Niagara/Source/NiagaraEditor/Private/Widgets/SNiagaraGraphParameterMapSetNode.cpp
--- a/Engine/Plugins/FX/Niagara/Source/NiagaraEditor/Private/Widgets/SNiagaraGraphParameterMapSetNode.cpp
+++ b/Engine/Plugins/FX/Niagara/Source/NiagaraEditor/Private/Widgets/SNiagaraGraphParameterMapSetNode.cpp
@@ -2,15 +2,10 @@
 
 #include "SNiagaraGraphParameterMapSetNode.h"
 #include "NiagaraNodeParameterMapSet.h"
-#include "Widgets/Input/SButton.h"
 #include "GraphEditorSettings.h"
-#include "Rendering/DrawElements.h"
-#include "Widgets/SBoxPanel.h"
 #include "SGraphPin.h"
 #include "EdGraphSchema_Niagara.h"
-#include "NiagaraScriptVariable.h"
 #include "SDropTarget.h"
-#include "NiagaraEditorStyle.h"
 
 
 #define LOCTEXT_NAMESPACE "SNiagaraGraphParameterMapSetNode"
diff --git a/Engine/Plugins/FX/Niagara/Source/NiagaraEditor/Private/Widgets/SNiagaraGraphParameterMapSetNode.h b/Engine/Plugins/FX/Niagara/Source/NiagaraEditor/Private/Widgets/SNiagaraGraphParameterMapSetNode.h
--- a/Engine/Plugins/FX/Niagara/Source/NiagaraEditor/Private/Widgets/SNiagaraGraphParameterMapSetNode.h
+++ b/Engine/Plugins/FX/Niagara/Source/NiagaraEditor/Private/Widgets/SNiagaraGraphParameterMapSetNode.h
@@ -10,7 +10,7 @@ public:
 	SLATE_END_ARGS();
 
 	void Construct(const FArguments& InArgs, UEdGraphNode* InGraphNode);
-	virtual TSharedRef<SWidget> CreateNodeContentArea();
+	virtual TSharedRef<SWidget> CreateNodeContentArea() override;
 
 protected:
 	FReply OnDroppedOnTarget(const FGeometry& InGeometry, const FDragDropEvent& InDragDropEvent);
diff --git a/Engine/Plugins/FX/Niagara/Source/NiagaraEditor/Public/NiagaraNodeFunctionCall.h b/Engine/Plugins/FX/Niagara/Source/NiagaraEditor/Public/NiagaraNodeFunctionCall.h
--- a/Engine/Plugins/FX/Niagara/Source/NiagaraEditor/Public/NiagaraNodeFunctionCall.h
+++ b/Engine/Plugins/FX/Niagara/Source/NiagaraEditor/Public/NiagaraNodeFunctionCall.h
@@ -134,7 +134,7 @@ public:
 
 	bool FindAutoBoundInput(UNiagaraNodeInput* InputNode, UEdGraphPin* PinToAutoBind, FNiagaraVariable& OutFoundVar, ENiagaraInputNodeUsage& OutNodeUsage);
 
-	void BuildParameterMapHistory(FNiagaraParameterMapHistoryBuilder& OutHistory, bool bRecursive = true, bool bFilterForCompilation = true) const override;
+	virtual void BuildParameterMapHistory(FNiagaraParameterMapHistoryBuilder& OutHistory, bool bRecursive = true, bool bFilterForCompilation = true) const override;
 
 	NIAGARAEDITOR_API void ChangeScriptVersion(FGuid NewScriptVersion, const FNiagaraScriptVersionUpgradeContext& UpgradeContext, bool bShowNotesInStack = false);
 
@@ -196,9 +196,9 @@ protected:
 	virtual bool CanRenamePin(const UEdGraphPin* Pin) const override { return AllowDynamicPins() && Super::CanRenamePin(Pin); }
 	virtual bool CanRemovePin(const UEdGraphPin* Pin) const override { return AllowDynamicPins() && Super::CanRemovePin(Pin); }
 	virtual bool CanMovePin(const UEdGraphPin* Pin, int32 DirectionToMove) const override { return AllowDynamicPins() && Super::CanMovePin(Pin, DirectionToMove); }
-	virtual bool CanModifyExistingPins() const { return false; }
+	virtual bool CanModifyExistingPins() const override { return false; }
 
-	virtual void OnNewTypedPinAdded(UEdGraphPin*& NewPin) { Super::OnNewTypedPinAdded(NewPin); RefreshSignature(); }
+	virtual void OnNewTypedPinAdded(UEdGraphPin*& NewPin) override { Super::OnNewTypedPinAdded(NewPin); RefreshSignature(); }
 	virtual void OnPinRenamed(UEdGraphPin* RenamedPin, const FString& OldPinName) override { Super::OnPinRenamed(RenamedPin, OldPinName); RefreshSignature(); }
 	virtual void RemoveDynamicPin(UEdGraphPin* Pin) override { Super::RemoveDynamicPin(Pin); RefreshSignature(); }
 	virtual void MoveDynamicPin(UEdGraphPin* Pin, int32 MoveAmount)override { Super::MoveDynamicPin(Pin, MoveAmount); RefreshSignature(); }
diff --git a/Engine/Plugins/FX/Niagara/Source/NiagaraEditor/Public/NiagaraNodeWithDynamicPins.h b/Engine/Plugins/FX/Niagara/Source/NiagaraEditor/Public/NiagaraNodeWithDynamicPins.h
--- a/Engine/Plugins/FX/Niagara/Source/NiagaraEditor/Public/NiagaraNodeWithDynamicPins.h
+++ b/Engine/Plugins/FX/Niagara/Source/NiagaraEditor/Public/NiagaraNodeWithDynamicPins.h
@@ -23,7 +23,7 @@ public:
 	//~ UEdGraphNode interface
 	virtual void PinConnectionListChanged(UEdGraphPin* Pin) override;
 	virtual void GetNodeContextMenuActions(class UToolMenu* Menu, class UGraphNodeContextMenuContext* Context) const override;
-	virtual bool IncludeParentNodeContextMenu() const { return true; }
+	virtual bool IncludeParentNodeContextMenu() const override { return true; }
 
 	/** Requests a new pin be added to the node with the specified direction, type, and name. */
 	UEdGraphPin* RequestNewTypedPin(EEdGraphPinDirection Direction, const FNiagaraTypeDefinition& Type, FName InName);
@@ -63,7 +63,7 @@ protected:
 	
 	/** Called in subclasses to restrict renaming.*/
 	/** Verify that the potential rename has produced acceptable results for a pin.*/
-	virtual bool VerifyEditablePinName(const FText& InName, FText& OutErrorMessage, const UEdGraphPin* InGraphPinObj) const { OutErrorMessage = FText::GetEmpty(); return true; }
+	virtual bool VerifyEditablePinName(const FText& InName, FText& OutErrorMessage, const UEdGraphPin* InGraphPinObj) const override { OutErrorMessage = FText::GetEmpty(); return true; }
 
 	/** Called when a pin is renamed. */
 	virtual void OnPinRenamed(UEdGraphPin* RenamedPin, const FString& OldPinName) { }
@@ -90,7 +90,7 @@ protected:
 
 	virtual bool OnVerifyTextChanged(const FText& NewText, FText& OutMessage) { return true; };
 
-	bool IsValidPinToCompile(UEdGraphPin* Pin) const override;
+	virtual bool IsValidPinToCompile(UEdGraphPin* Pin) const override;
 
 private:
 
