Added unit test for Niagara custom LWC struct hlsl generation

diff --git a/Engine/Plugins/FX/Niagara/Source/NiagaraEditor/Private/NiagaraHlslTranslator.cpp b/Engine/Plugins/FX/Niagara/Source/NiagaraEditor/Private/NiagaraHlslTranslator.cpp
--- a/Engine/Plugins/FX/Niagara/Source/NiagaraEditor/Private/NiagaraHlslTranslator.cpp
+++ b/Engine/Plugins/FX/Niagara/Source/NiagaraEditor/Private/NiagaraHlslTranslator.cpp
@@ -4825,6 +4825,7 @@ int32 FHlslNiagaraTranslator::GetConstantDirect(int InConstantValue)
 bool FHlslNiagaraTranslator::GenerateStructInitializer(TStringBuilder<128>& InitializerString, UStruct* UserDefinedStruct, const void* StructData, int32 ByteOffset)
 {
 	//-TODO: Alignment Issues
+	//-TODO: Also adjust automation test NiagaraTest_CustomLWCStructHlsl when this is fixed 
 	//const bool bHasAlignmentDummy = CompileOptions.TargetUsage != ENiagaraScriptUsage::ParticleGPUComputeScript;
 
 	//auto CheckAlignment =
