Fixed an error where hlsl compile errors were not correctly shown in the stack

diff --git a/Engine/Plugins/FX/Niagara/Source/Niagara/Private/NiagaraScript.cpp b/Engine/Plugins/FX/Niagara/Source/Niagara/Private/NiagaraScript.cpp
--- a/Engine/Plugins/FX/Niagara/Source/Niagara/Private/NiagaraScript.cpp
+++ b/Engine/Plugins/FX/Niagara/Source/Niagara/Private/NiagaraScript.cpp
@@ -4032,8 +4032,28 @@ FGuid UNiagaraScript::GetBaseChangeID(const FGuid& VersionGuid) const
 
 ENiagaraScriptCompileStatus UNiagaraScript::GetLastCompileStatus() const
 {
+	bool bHasWarnings = false;
+	if (Usage == ENiagaraScriptUsage::ParticleGPUComputeScript)
+	{
+		const TArray<FString>& Errors = ScriptResource->GetCompileErrors();
+		if (Errors.Num() > 0)
+		{
+			for (const FString& ErrorStr : Errors)
+			{
+				if (ErrorStr.Contains(TEXT("err0r")))
+				{
+					return ENiagaraScriptCompileStatus::NCS_Error;
+				}
+			}
+			bHasWarnings = true;
+		}	
+	}
 	if (CachedScriptVM.IsValid())
 	{
+		if (CachedScriptVM.LastCompileStatus == ENiagaraScriptCompileStatus::NCS_UpToDate && bHasWarnings)
+		{
+			return ENiagaraScriptCompileStatus::NCS_ComputeUpToDateWithWarnings;
+		}
 		return CachedScriptVM.LastCompileStatus;
 	}
 	return ENiagaraScriptCompileStatus::NCS_Unknown;
