- Add render thread mirror for sim cache read mappings
- This fixes a race if the read mappings were modified while we are using them on the RT

diff --git a/Engine/Plugins/FX/Niagara/Source/Niagara/Classes/NiagaraSimCache.h b/Engine/Plugins/FX/Niagara/Source/Niagara/Classes/NiagaraSimCache.h
--- a/Engine/Plugins/FX/Niagara/Source/Niagara/Classes/NiagaraSimCache.h
+++ b/Engine/Plugins/FX/Niagara/Source/Niagara/Classes/NiagaraSimCache.h
@@ -316,13 +316,21 @@ struct FNiagaraSimCacheDataBuffersLayout
 	UPROPERTY()
 	uint16 ComponentVelocity = INDEX_NONE;
 
-	// Transient value used to map particles between frames for interpolation
-	uint16 ComponentUniqueID = INDEX_NONE;
-
-	TArray<uint16> ComponentMappingsToDataBuffer;
-	TArray<FVariableCopyMapping> VariableCopyMappingsToDataBuffer;
+	struct FCacheBufferWriteInfo
+	{
+		uint16			ComponentUniqueID = INDEX_NONE;		// Used during building for interpolation to track particles
+		TArray<uint16>	ComponentMappingsFromDataBuffer;	// Used to map individual components from a Niagara Data Buffer -> Cache Buffer
+	};
+	FCacheBufferWriteInfo CacheBufferWriteInfo;
 
-	TArray<uint16>	ComponentMappingsFromDataBuffer;
+	struct FCacheBufferReadInfo
+	{
+		TArray<uint16>					ComponentMappingsToDataBuffer;		// Used to map individual components from Cache Buffer -> Niagara Data Buffer
+		TArray<FVariableCopyMapping>	VariableCopyMappingsToDataBuffer;	// Used for more complex operations when we need to rebase / interpolate / etc
+	};
+	FCacheBufferReadInfo CacheBufferReadInfo;
+	mutable bool bNeedsCacheBufferReadInfoUpdateForRT = false;
+	mutable FCacheBufferReadInfo CacheBufferReadInfo_RT;
 
 	int32 IndexOfCacheVariable(const FNiagaraVariableBase& InVariable) const;
 	const FNiagaraSimCacheVariable* FindCacheVariable(const FNiagaraVariableBase& InVariable) const;
diff --git a/Engine/Plugins/FX/Niagara/Source/Niagara/Private/NiagaraSimCacheHelper.h b/Engine/Plugins/FX/Niagara/Source/Niagara/Private/NiagaraSimCacheHelper.h
--- a/Engine/Plugins/FX/Niagara/Source/Niagara/Private/NiagaraSimCacheHelper.h
+++ b/Engine/Plugins/FX/Niagara/Source/Niagara/Private/NiagaraSimCacheHelper.h
@@ -125,7 +125,7 @@ struct FNiagaraSimCacheHelper
 		CacheLayout.LayoutName = LayoutName;
 		CacheLayout.SimTarget = CompiledData.SimTarget;
 		CacheLayout.ComponentVelocity = INDEX_NONE;
-		CacheLayout.ComponentUniqueID = INDEX_NONE;
+		CacheLayout.CacheBufferWriteInfo.ComponentUniqueID = INDEX_NONE;
 
 		// Determine the components to cache
 		int32 TotalCacheComponents = 0;
@@ -164,8 +164,8 @@ struct FNiagaraSimCacheHelper
 
 		CacheLayout.Variables.AddDefaulted(NumCacheVariables);
 
-		CacheLayout.ComponentMappingsFromDataBuffer.Empty(TotalCacheComponents);
-		CacheLayout.ComponentMappingsFromDataBuffer.AddDefaulted(TotalCacheComponents);
+		CacheLayout.CacheBufferWriteInfo.ComponentMappingsFromDataBuffer.Empty(TotalCacheComponents);
+		CacheLayout.CacheBufferWriteInfo.ComponentMappingsFromDataBuffer.AddDefaulted(TotalCacheComponents);
 		CacheLayout.RebaseVariableNames = MoveTemp(InRebaseVariableNames);
 		CacheLayout.InterpVariableNames = MoveTemp(InInterpVariableNames);
 
@@ -199,12 +199,12 @@ struct FNiagaraSimCacheHelper
 		if (CreateParameters.bAllowInterpolation)
 		{
 			const FNiagaraVariableBase UniqueIDVariable(FNiagaraTypeDefinition::GetIntDef(), "UniqueID");
-			CacheLayout.ComponentUniqueID = CompiledData.Variables.IndexOfByKey(UniqueIDVariable);
-			if (CacheLayout.ComponentUniqueID != uint16(INDEX_NONE))
+			CacheLayout.CacheBufferWriteInfo.ComponentUniqueID = CompiledData.Variables.IndexOfByKey(UniqueIDVariable);
+			if (CacheLayout.CacheBufferWriteInfo.ComponentUniqueID != uint16(INDEX_NONE))
 			{
-				const FNiagaraVariableLayoutInfo& DataSetVariableLayout = CompiledData.VariableLayouts[CacheLayout.ComponentUniqueID];
+				const FNiagaraVariableLayoutInfo& DataSetVariableLayout = CompiledData.VariableLayouts[CacheLayout.CacheBufferWriteInfo.ComponentUniqueID];
 				check(DataSetVariableLayout.GetNumInt32Components() == 1);
-				CacheLayout.ComponentUniqueID = DataSetVariableLayout.Int32ComponentStart;
+				CacheLayout.CacheBufferWriteInfo.ComponentUniqueID = DataSetVariableLayout.Int32ComponentStart;
 				CacheLayout.bAllowInterpolation = true;
 			}
 		}
@@ -225,19 +225,19 @@ struct FNiagaraSimCacheHelper
 
 			for (int32 iComponent=0; iComponent < CacheVariable.FloatCount; ++iComponent)
 			{
-				CacheLayout.ComponentMappingsFromDataBuffer[FloatOffset] = uint16(DataSetVariableLayout.FloatComponentStart + iComponent);
+				CacheLayout.CacheBufferWriteInfo.ComponentMappingsFromDataBuffer[FloatOffset] = uint16(DataSetVariableLayout.FloatComponentStart + iComponent);
 				++FloatOffset;
 			}
 
 			for (int32 iComponent=0; iComponent < CacheVariable.HalfCount; ++iComponent)
 			{
-				CacheLayout.ComponentMappingsFromDataBuffer[HalfOffset] = uint16(DataSetVariableLayout.HalfComponentStart + iComponent);
+				CacheLayout.CacheBufferWriteInfo.ComponentMappingsFromDataBuffer[HalfOffset] = uint16(DataSetVariableLayout.HalfComponentStart + iComponent);
 				++HalfOffset;
 			}
 
 			for (int32 iComponent=0; iComponent < CacheVariable.Int32Count; ++iComponent)
 			{
-				CacheLayout.ComponentMappingsFromDataBuffer[Int32Offset] = uint16(DataSetVariableLayout.Int32ComponentStart + iComponent);
+				CacheLayout.CacheBufferWriteInfo.ComponentMappingsFromDataBuffer[Int32Offset] = uint16(DataSetVariableLayout.Int32ComponentStart + iComponent);
 				++Int32Offset;
 			}
 		}
@@ -436,9 +436,11 @@ struct FNiagaraSimCacheHelper
 	static bool BuildCacheReadMappings(FNiagaraSimCacheDataBuffersLayout& CacheLayout, const FNiagaraDataSetCompiledData& CompiledData)
 	{
 		const int32 CacheTotalComponents = CacheLayout.FloatCount + CacheLayout.HalfCount + CacheLayout.Int32Count;
-		CacheLayout.ComponentMappingsToDataBuffer.Empty(CacheTotalComponents);
-		CacheLayout.ComponentMappingsToDataBuffer.AddDefaulted(CacheTotalComponents);
-		CacheLayout.VariableCopyMappingsToDataBuffer.Empty(0);
+		FNiagaraSimCacheDataBuffersLayout::FCacheBufferReadInfo& CacheBufferReadInfo = CacheLayout.CacheBufferReadInfo;
+		CacheBufferReadInfo.ComponentMappingsToDataBuffer.Empty(CacheTotalComponents);
+		CacheBufferReadInfo.ComponentMappingsToDataBuffer.AddDefaulted(CacheTotalComponents);
+		CacheBufferReadInfo.VariableCopyMappingsToDataBuffer.Empty(0);
+		CacheLayout.bNeedsCacheBufferReadInfoUpdateForRT = true;
 
 		//-OPT: We don't have to do this all the time, figure out a better way
 		TArray<FNiagaraVariableBase> InterpCurrAttributes;
@@ -521,7 +523,7 @@ struct FNiagaraSimCacheHelper
 								const FNiagaraSimCacheVariable* PrevSourceVariable = CacheLayout.FindCacheVariable(InterpPrevAttributes[InterpVarIndex]);
 								check(PrevSourceVariable != nullptr);
 
-								CacheLayout.VariableCopyMappingsToDataBuffer.Emplace(
+								CacheBufferReadInfo.VariableCopyMappingsToDataBuffer.Emplace(
 									SourceVariable.FloatOffset, PrevSourceVariable->FloatOffset,
 									uint16(DestVariableLayout->FloatComponentStart), uint16(CompiledData.VariableLayouts[PreviousDataSetVariableIndex].FloatComponentStart),
 									bRebaseVariable ? &FNiagaraSimCacheHelper::InterpPositions<true, true> : &FNiagaraSimCacheHelper::InterpPositions<false, true>
@@ -530,16 +532,16 @@ struct FNiagaraSimCacheHelper
 							}
 							else
 							{
-								CacheLayout.VariableCopyMappingsToDataBuffer.Emplace(SourceVariable.FloatOffset, uint16(DestVariableLayout->FloatComponentStart), bRebaseVariable ? &FNiagaraSimCacheHelper::InterpPositions<true, false> : &FNiagaraSimCacheHelper::InterpPositions<false, false>);
+								CacheBufferReadInfo.VariableCopyMappingsToDataBuffer.Emplace(SourceVariable.FloatOffset, uint16(DestVariableLayout->FloatComponentStart), bRebaseVariable ? &FNiagaraSimCacheHelper::InterpPositions<true, false> : &FNiagaraSimCacheHelper::InterpPositions<false, false>);
 							}
 						}
 						else if ( CacheLayout.bAllowVelocityExtrapolation )
 						{
-							CacheLayout.VariableCopyMappingsToDataBuffer.Emplace(SourceVariable.FloatOffset, uint16(DestVariableLayout->FloatComponentStart), bRebaseVariable ? &FNiagaraSimCacheHelper::ExtrapolatePositions<true> : &FNiagaraSimCacheHelper::ExtrapolatePositions<false>);
+							CacheBufferReadInfo.VariableCopyMappingsToDataBuffer.Emplace(SourceVariable.FloatOffset, uint16(DestVariableLayout->FloatComponentStart), bRebaseVariable ? &FNiagaraSimCacheHelper::ExtrapolatePositions<true> : &FNiagaraSimCacheHelper::ExtrapolatePositions<false>);
 						}
 						else
 						{
-							CacheLayout.VariableCopyMappingsToDataBuffer.Emplace(SourceVariable.FloatOffset, uint16(DestVariableLayout->FloatComponentStart), &FNiagaraSimCacheHelper::RebasePositions);
+							CacheBufferReadInfo.VariableCopyMappingsToDataBuffer.Emplace(SourceVariable.FloatOffset, uint16(DestVariableLayout->FloatComponentStart), &FNiagaraSimCacheHelper::RebasePositions);
 						}
 						DestVariableLayout = nullptr;
 					}
@@ -557,7 +559,7 @@ struct FNiagaraSimCacheHelper
 								const FNiagaraSimCacheVariable* PrevSourceVariable = CacheLayout.FindCacheVariable(InterpPrevAttributes[InterpVarIndex]);
 								check(PrevSourceVariable != nullptr);
 
-								CacheLayout.VariableCopyMappingsToDataBuffer.Emplace(
+								CacheBufferReadInfo.VariableCopyMappingsToDataBuffer.Emplace(
 									SourceVariable.FloatOffset, PrevSourceVariable->FloatOffset,
 									uint16(DestVariableLayout->FloatComponentStart), uint16(CompiledData.VariableLayouts[PreviousDataSetVariableIndex].FloatComponentStart),
 									bRebaseVariable ? &FNiagaraSimCacheHelper::InterpQuaternions<true, true> : &FNiagaraSimCacheHelper::InterpQuaternions<false, true>
@@ -566,12 +568,12 @@ struct FNiagaraSimCacheHelper
 							}
 							else
 							{
-								CacheLayout.VariableCopyMappingsToDataBuffer.Emplace(SourceVariable.FloatOffset, uint16(DestVariableLayout->FloatComponentStart), bInterpVariable  ? &FNiagaraSimCacheHelper::InterpQuaternions<true, false> : &FNiagaraSimCacheHelper::InterpQuaternions<false, false>);
+								CacheBufferReadInfo.VariableCopyMappingsToDataBuffer.Emplace(SourceVariable.FloatOffset, uint16(DestVariableLayout->FloatComponentStart), bInterpVariable  ? &FNiagaraSimCacheHelper::InterpQuaternions<true, false> : &FNiagaraSimCacheHelper::InterpQuaternions<false, false>);
 							}
 						}
 						else
 						{
-							CacheLayout.VariableCopyMappingsToDataBuffer.Emplace(SourceVariable.FloatOffset, uint16(DestVariableLayout->FloatComponentStart), &FNiagaraSimCacheHelper::RebaseQuaternions);
+							CacheBufferReadInfo.VariableCopyMappingsToDataBuffer.Emplace(SourceVariable.FloatOffset, uint16(DestVariableLayout->FloatComponentStart), &FNiagaraSimCacheHelper::RebaseQuaternions);
 						}
 						DestVariableLayout = nullptr;
 					}
@@ -580,7 +582,7 @@ struct FNiagaraSimCacheHelper
 						if (bRebaseVariable)
 						{
 							check(SourceVariable.FloatCount == 16);
-							CacheLayout.VariableCopyMappingsToDataBuffer.Emplace(SourceVariable.FloatOffset, uint16(DestVariableLayout->FloatComponentStart), &FNiagaraSimCacheHelper::RebaseMatrices);
+							CacheBufferReadInfo.VariableCopyMappingsToDataBuffer.Emplace(SourceVariable.FloatOffset, uint16(DestVariableLayout->FloatComponentStart), &FNiagaraSimCacheHelper::RebaseMatrices);
 							DestVariableLayout = nullptr;
 						}
 					}
@@ -589,17 +591,17 @@ struct FNiagaraSimCacheHelper
 
 			for (int32 i = 0; i < SourceVariable.FloatCount; ++i)
 			{
-				CacheLayout.ComponentMappingsToDataBuffer[FloatOffset++] = uint16(DestVariableLayout ? DestVariableLayout->FloatComponentStart + i : INDEX_NONE);
+				CacheBufferReadInfo.ComponentMappingsToDataBuffer[FloatOffset++] = uint16(DestVariableLayout ? DestVariableLayout->FloatComponentStart + i : INDEX_NONE);
 			}
 
 			for (int32 i = 0; i < SourceVariable.HalfCount; ++i)
 			{
-				CacheLayout.ComponentMappingsToDataBuffer[HalfOffset++] = uint16(DestVariableLayout ? DestVariableLayout->HalfComponentStart + i : INDEX_NONE);
+				CacheBufferReadInfo.ComponentMappingsToDataBuffer[HalfOffset++] = uint16(DestVariableLayout ? DestVariableLayout->HalfComponentStart + i : INDEX_NONE);
 			}
 
 			for (int32 i = 0; i < SourceVariable.Int32Count; ++i)
 			{
-				CacheLayout.ComponentMappingsToDataBuffer[Int32Offset++] = uint16(DestVariableLayout ? DestVariableLayout->Int32ComponentStart + i : INDEX_NONE);
+				CacheBufferReadInfo.ComponentMappingsToDataBuffer[Int32Offset++] = uint16(DestVariableLayout ? DestVariableLayout->Int32ComponentStart + i : INDEX_NONE);
 			}
 		}
 
@@ -628,7 +630,7 @@ struct FNiagaraSimCacheHelper
 		CacheBuffer.FloatData.AddDefaulted(CacheLayout.FloatCount * NumInstances * sizeof(float));
 		for ( uint32 i=0; i < CacheLayout.FloatCount; ++i )
 		{
-			const uint32 Component = CacheLayout.ComponentMappingsFromDataBuffer[iComponent++];
+			const uint32 Component = CacheLayout.CacheBufferWriteInfo.ComponentMappingsFromDataBuffer[iComponent++];
 			const uint8* Source = DataBuffer.GetComponentPtrFloat(Component) + (FirstInstance * sizeof(float));
 			uint8* Dest = CacheBuffer.FloatData.GetData() + (i * NumInstances * sizeof(float));
 			CheckedMemcpy(CacheBuffer.FloatData, Dest, DataBuffer.GetFloatBuffer(), Source, sizeof(float) * NumInstances);
@@ -639,7 +641,7 @@ struct FNiagaraSimCacheHelper
 		CacheBuffer.HalfData.AddDefaulted(CacheLayout.HalfCount * NumInstances * sizeof(FFloat16));
 		for (uint32 i = 0; i < CacheLayout.HalfCount; ++i)
 		{
-			const uint32 Component = CacheLayout.ComponentMappingsFromDataBuffer[iComponent++];
+			const uint32 Component = CacheLayout.CacheBufferWriteInfo.ComponentMappingsFromDataBuffer[iComponent++];
 			const uint8* Source = DataBuffer.GetComponentPtrHalf(Component) + (FirstInstance * sizeof(FFloat16));
 			uint8* Dest = CacheBuffer.HalfData.GetData() + (i * NumInstances * sizeof(FFloat16));
 			CheckedMemcpy(CacheBuffer.HalfData, Dest, DataBuffer.GetHalfBuffer(), Source, sizeof(FFloat16) * NumInstances);
@@ -650,7 +652,7 @@ struct FNiagaraSimCacheHelper
 		CacheBuffer.Int32Data.AddDefaulted(CacheLayout.Int32Count * NumInstances * sizeof(int32));
 		for (uint32 i = 0; i < CacheLayout.Int32Count; ++i)
 		{
-			const uint32 Component = CacheLayout.ComponentMappingsFromDataBuffer[iComponent++];
+			const uint32 Component = CacheLayout.CacheBufferWriteInfo.ComponentMappingsFromDataBuffer[iComponent++];
 			const uint8* Source = DataBuffer.GetComponentPtrInt32(Component) + (FirstInstance * sizeof(int32));
 			uint8* Dest = CacheBuffer.Int32Data.GetData() + (i * NumInstances * sizeof(int32));
 			CheckedMemcpy(CacheBuffer.Int32Data, Dest, DataBuffer.GetInt32Buffer(), Source, sizeof(int32) * NumInstances);
@@ -666,7 +668,7 @@ struct FNiagaraSimCacheHelper
 		{
 			//-TODO: Persistent ID mapping
 			CacheBuffer.InterpMapping.SetNumUninitialized(DataBuffer.GetNumInstances());
-			const uint8* UniqueIDs = DataBuffer.GetComponentPtrInt32(CacheLayout.ComponentUniqueID);
+			const uint8* UniqueIDs = DataBuffer.GetComponentPtrInt32(CacheLayout.CacheBufferWriteInfo.ComponentUniqueID);
 			FMemory::Memcpy(CacheBuffer.InterpMapping.GetData(), UniqueIDs, DataBuffer.GetNumInstances() * sizeof(int32));
 		}
 	}
@@ -682,12 +684,12 @@ struct FNiagaraSimCacheHelper
 		}
 	}
 
-	static void ReadFloatBuffers(int32& iComponent, const FNiagaraSimCacheDataBuffersLayout& CacheLayout, const FNiagaraSimCacheDataBuffers& CacheBuffer, TArrayView<uint8> DestBuffer, uint32 DestStride)
+	static void ReadFloatBuffers(int32& iComponent, uint32 FloatCount, TConstArrayView<uint16> ComponentMappingsToDataBuffer, const FNiagaraSimCacheDataBuffers& CacheBuffer, TArrayView<uint8> DestBuffer, uint32 DestStride)
 	{
 		const int32 NumInstances = CacheBuffer.NumInstances;
-		for (uint32 i = 0; i < CacheLayout.FloatCount; ++i)
+		for (uint32 i = 0; i < FloatCount; ++i)
 		{
-			const uint32 Component = CacheLayout.ComponentMappingsToDataBuffer[iComponent++];
+			const uint32 Component = ComponentMappingsToDataBuffer[iComponent++];
 			if (Component == InvalidComponent)
 			{
 				continue;
@@ -698,12 +700,12 @@ struct FNiagaraSimCacheHelper
 		}
 	}
 
-	static void ReadHalfBuffers(int32& iComponent, const FNiagaraSimCacheDataBuffersLayout& CacheLayout, const FNiagaraSimCacheDataBuffers& CacheBuffer, TArrayView<uint8> DestBuffer, uint32 DestStride)
+	static void ReadHalfBuffers(int32& iComponent, uint32 HalfCount, TConstArrayView<uint16> ComponentMappingsToDataBuffer, const FNiagaraSimCacheDataBuffers& CacheBuffer, TArrayView<uint8> DestBuffer, uint32 DestStride)
 	{
 		const int32 NumInstances = CacheBuffer.NumInstances;
-		for (uint32 i = 0; i < CacheLayout.HalfCount; ++i)
+		for (uint32 i = 0; i < HalfCount; ++i)
 		{
-			const uint32 Component = CacheLayout.ComponentMappingsToDataBuffer[iComponent++];
+			const uint32 Component = ComponentMappingsToDataBuffer[iComponent++];
 			if (Component == InvalidComponent)
 			{
 				continue;
@@ -714,12 +716,12 @@ struct FNiagaraSimCacheHelper
 		}
 	}
 
-	static void ReadInt32Buffers(int32& iComponent, const FNiagaraSimCacheDataBuffersLayout& CacheLayout, const FNiagaraSimCacheDataBuffers& CacheBuffer, TArrayView<uint8> DestBuffer, uint32 DestStride)
+	static void ReadInt32Buffers(int32& iComponent, uint32 Int32Count, TConstArrayView<uint16> ComponentMappingsToDataBuffer, const FNiagaraSimCacheDataBuffers& CacheBuffer, TArrayView<uint8> DestBuffer, uint32 DestStride)
 	{
 		const int32 NumInstances = CacheBuffer.NumInstances;
-		for (uint32 i = 0; i < CacheLayout.Int32Count; ++i)
+		for (uint32 i = 0; i < Int32Count; ++i)
 		{
-			const uint32 Component = CacheLayout.ComponentMappingsToDataBuffer[iComponent++];
+			const uint32 Component = ComponentMappingsToDataBuffer[iComponent++];
 			if (Component == InvalidComponent)
 			{
 				continue;
@@ -730,9 +732,9 @@ struct FNiagaraSimCacheHelper
 		}
 	}
 
-	static void ReadCustomBuffers(float FrameFraction, float FrameDeltaSeconds, float SimDeltaSeconds, const FTransform& RebaseTransform, const FNiagaraSimCacheDataBuffersLayout& CacheLayout, const FNiagaraSimCacheDataBuffers& CacheBufferA, const FNiagaraSimCacheDataBuffers& CacheBufferB, uint8* DestBuffer, uint32 DestStride)
+	static void ReadCustomBuffers(float FrameFraction, float FrameDeltaSeconds, float SimDeltaSeconds, const FTransform& RebaseTransform, uint16 ComponentVelocity, const TConstArrayView<FNiagaraSimCacheDataBuffersLayout::FVariableCopyMapping>& VariableCopyMappingsToDataBuffer, const FNiagaraSimCacheDataBuffers& CacheBufferA, const FNiagaraSimCacheDataBuffers& CacheBufferB, uint8* DestBuffer, uint32 DestStride)
 	{
-		if (CacheLayout.VariableCopyMappingsToDataBuffer.Num() == 0)
+		if (VariableCopyMappingsToDataBuffer.Num() == 0)
 		{
 			return;
 		}
@@ -748,16 +750,16 @@ struct FNiagaraSimCacheHelper
 		VariableCopyDataContext.DestStride			= DestStride;
 		VariableCopyDataContext.SourceAStride		= CacheBufferA.NumInstances * sizeof(float);
 		VariableCopyDataContext.SourceBStride		= CacheBufferB.NumInstances * sizeof(float);
-		if (CacheLayout.ComponentVelocity != uint16(INDEX_NONE))
+		if (ComponentVelocity != uint16(INDEX_NONE))
 		{
-			VariableCopyDataContext.Velocity = CacheBufferA.FloatData.GetData() + (CacheLayout.ComponentVelocity * VariableCopyDataContext.SourceAStride);
+			VariableCopyDataContext.Velocity = CacheBufferA.FloatData.GetData() + (ComponentVelocity * VariableCopyDataContext.SourceAStride);
 		}
 		else
 		{
 			VariableCopyDataContext.Velocity = nullptr;
 		}
 
-		for (const FNiagaraSimCacheDataBuffersLayout::FVariableCopyMapping& VariableCopyMapping : CacheLayout.VariableCopyMappingsToDataBuffer)
+		for (const FNiagaraSimCacheDataBuffersLayout::FVariableCopyMapping& VariableCopyMapping : VariableCopyMappingsToDataBuffer)
 		{
 			VariableCopyDataContext.DestCurr	= DestBuffer + (uint32(VariableCopyMapping.CurrComponentTo) * DestStride);
 			VariableCopyDataContext.DestPrev	= DestBuffer + (uint32(VariableCopyMapping.PrevComponentTo) * DestStride);
@@ -776,10 +778,10 @@ struct FNiagaraSimCacheHelper
 		if (CacheBufferA.NumInstances > 0 )
 		{
 			int32 iComponent = 0;
-			ReadFloatBuffers(iComponent, CacheLayout, CacheBufferA, MakeArrayView(DataBuffer.GetComponentPtrFloat(0), DataBuffer.GetFloatBuffer().Num()), DataBuffer.GetFloatStride());
-			ReadHalfBuffers(iComponent, CacheLayout, CacheBufferA, MakeArrayView(DataBuffer.GetComponentPtrHalf(0), DataBuffer.GetHalfBuffer().Num()), DataBuffer.GetHalfStride());
-			ReadInt32Buffers(iComponent, CacheLayout, CacheBufferA, MakeArrayView(DataBuffer.GetComponentPtrInt32(0), DataBuffer.GetInt32Buffer().Num()), DataBuffer.GetInt32Stride());
-			ReadCustomBuffers(FrameFraction, FrameDeltaSeconds, SimDeltaSeconds, RebaseTransform, CacheLayout, CacheBufferA, CacheBufferB, DataBuffer.GetComponentPtrFloat(0), DataBuffer.GetFloatStride());
+			ReadFloatBuffers(iComponent, CacheLayout.FloatCount, CacheLayout.CacheBufferReadInfo.ComponentMappingsToDataBuffer, CacheBufferA, MakeArrayView(DataBuffer.GetComponentPtrFloat(0), DataBuffer.GetFloatBuffer().Num()), DataBuffer.GetFloatStride());
+			ReadHalfBuffers(iComponent, CacheLayout.HalfCount, CacheLayout.CacheBufferReadInfo.ComponentMappingsToDataBuffer, CacheBufferA, MakeArrayView(DataBuffer.GetComponentPtrHalf(0), DataBuffer.GetHalfBuffer().Num()), DataBuffer.GetHalfStride());
+			ReadInt32Buffers(iComponent, CacheLayout.Int32Count, CacheLayout.CacheBufferReadInfo.ComponentMappingsToDataBuffer, CacheBufferA, MakeArrayView(DataBuffer.GetComponentPtrInt32(0), DataBuffer.GetInt32Buffer().Num()), DataBuffer.GetInt32Stride());
+			ReadCustomBuffers(FrameFraction, FrameDeltaSeconds, SimDeltaSeconds, RebaseTransform, CacheLayout.ComponentVelocity, CacheLayout.CacheBufferReadInfo.VariableCopyMappingsToDataBuffer, CacheBufferA, CacheBufferB, DataBuffer.GetComponentPtrFloat(0), DataBuffer.GetFloatStride());
 		}
 
 		//-TODO:DestinationDataBuffer.SetIDTable(CacheBufferA.IDToIndexTable);
@@ -799,6 +801,18 @@ struct FNiagaraSimCacheHelper
 
 		check(EmitterInstance.GetGPUContext());
 
+		//-OPT: Rather than a lazy update here we could do it on initial read
+		if (InCacheLayout.bNeedsCacheBufferReadInfoUpdateForRT)
+		{
+			InCacheLayout.bNeedsCacheBufferReadInfoUpdateForRT = false;
+			ENQUEUE_RENDER_COMMAND(NiagaraSimCacheGpuUpdateReadInfo)(
+				[CacheLayout=&InCacheLayout, NewReadInfo=InCacheLayout.CacheBufferReadInfo](FRHICommandListImmediate&)
+				{
+					CacheLayout->CacheBufferReadInfo_RT = NewReadInfo;
+				}
+			);
+		}
+
 		FNiagaraGpuComputeDispatchInterface* DispathInterface = EmitterInstance.GetParentSystemInstance()->GetComputeDispatchInterface();
 		ENQUEUE_RENDER_COMMAND(NiagaraSimCacheGpuReadFrame)(
 			[DispathInterface, GPUExecContext=EmitterInstance.GetGPUContext(), FrameFraction=InFrameFraction, FrameDeltaSeconds=InFrameDeltaSeconds, SimDeltaSeconds=InSimDeltaSeconds, RebaseTransform=InRebaseTransform, CacheLayout=&InCacheLayout, CacheBufferA=&InCacheBufferA, CacheBufferB=&InCacheBufferB, DataSet=&InDataSet, PendingCommandsCounter=&InPendingCommandsCounter](FRHICommandListImmediate& RHICmdList)
@@ -836,8 +850,8 @@ struct FNiagaraSimCacheHelper
 					{
 						FRWBuffer& RWBuffer = DataBuffer.GetGPUBufferFloat();
 						uint8* RWBufferMemory = reinterpret_cast<uint8*>(RHILockBuffer(RWBuffer.Buffer, 0, RWBuffer.NumBytes, RLM_WriteOnly));
-						ReadFloatBuffers(iComponent, *CacheLayout, *CacheBufferA, MakeArrayView(RWBufferMemory, RWBuffer.NumBytes), DataBuffer.GetFloatStride());
-						ReadCustomBuffers(FrameFraction, FrameDeltaSeconds, SimDeltaSeconds, RebaseTransform, *CacheLayout, *CacheBufferA, *CacheBufferB, RWBufferMemory, DataBuffer.GetFloatStride());
+						ReadFloatBuffers(iComponent, CacheLayout->FloatCount, CacheLayout->CacheBufferReadInfo_RT.ComponentMappingsToDataBuffer, *CacheBufferA, MakeArrayView(RWBufferMemory, RWBuffer.NumBytes), DataBuffer.GetFloatStride());
+						ReadCustomBuffers(FrameFraction, FrameDeltaSeconds, SimDeltaSeconds, RebaseTransform, CacheLayout->ComponentVelocity, CacheLayout->CacheBufferReadInfo_RT.VariableCopyMappingsToDataBuffer, *CacheBufferA, *CacheBufferB, RWBufferMemory, DataBuffer.GetFloatStride());
 						RHIUnlockBuffer(RWBuffer.Buffer);
 					}
 
@@ -846,7 +860,7 @@ struct FNiagaraSimCacheHelper
 					{
 						FRWBuffer& RWBuffer = DataBuffer.GetGPUBufferHalf();
 						uint8* RWBufferMemory = reinterpret_cast<uint8*>(RHILockBuffer(RWBuffer.Buffer, 0, RWBuffer.NumBytes, RLM_WriteOnly));
-						ReadHalfBuffers(iComponent, *CacheLayout, *CacheBufferA, MakeArrayView(RWBufferMemory, RWBuffer.NumBytes), DataBuffer.GetHalfStride());
+						ReadHalfBuffers(iComponent, CacheLayout->HalfCount, CacheLayout->CacheBufferReadInfo_RT.ComponentMappingsToDataBuffer, *CacheBufferA, MakeArrayView(RWBufferMemory, RWBuffer.NumBytes), DataBuffer.GetHalfStride());
 						RHIUnlockBuffer(RWBuffer.Buffer);
 					}
 
@@ -855,7 +869,7 @@ struct FNiagaraSimCacheHelper
 					{
 						FRWBuffer& RWBuffer = DataBuffer.GetGPUBufferInt();
 						uint8* RWBufferMemory = reinterpret_cast<uint8*>(RHILockBuffer(RWBuffer.Buffer, 0, RWBuffer.NumBytes, RLM_WriteOnly));
-						ReadInt32Buffers(iComponent, *CacheLayout, *CacheBufferA, MakeArrayView(RWBufferMemory, RWBuffer.NumBytes), DataBuffer.GetInt32Stride());
+						ReadInt32Buffers(iComponent, CacheLayout->Int32Count, CacheLayout->CacheBufferReadInfo_RT.ComponentMappingsToDataBuffer, *CacheBufferA, MakeArrayView(RWBufferMemory, RWBuffer.NumBytes), DataBuffer.GetInt32Stride());
 						RHIUnlockBuffer(RWBuffer.Buffer);
 					}
 				}
