- Allow none particle parameters to be bound in particle mode

#rnx

diff --git a/Engine/Plugins/FX/Niagara/Source/Niagara/Private/NiagaraRendererProperties.cpp b/Engine/Plugins/FX/Niagara/Source/Niagara/Private/NiagaraRendererProperties.cpp
--- a/Engine/Plugins/FX/Niagara/Source/Niagara/Private/NiagaraRendererProperties.cpp
+++ b/Engine/Plugins/FX/Niagara/Source/Niagara/Private/NiagaraRendererProperties.cpp
@@ -356,18 +356,12 @@ bool UNiagaraRendererProperties::IsSupportedVariableForBinding(const FNiagaraVar
 			InSourceForBinding.IsInNameSpace(FNiagaraConstants::EmitterNamespace);
 	}
 
-	switch (GetCurrentSourceMode())
-	{
-		case ENiagaraRendererSourceDataMode::Particles:
-			return InSourceForBinding.IsInNameSpace(FNiagaraConstants::ParticleAttributeNamespaceString);
-
-		default:
-			return
-				InSourceForBinding.IsInNameSpace(FNiagaraConstants::UserNamespaceString) ||
-				InSourceForBinding.IsInNameSpace(FNiagaraConstants::SystemNamespaceString) ||
-				InSourceForBinding.IsInNameSpace(FNiagaraConstants::EmitterNamespaceString);
-	}
-
+	const ENiagaraRendererSourceDataMode CurrentSourceMode = GetCurrentSourceMode();
+	return
+		(CurrentSourceMode == ENiagaraRendererSourceDataMode::Particles && InSourceForBinding.IsInNameSpace(FNiagaraConstants::ParticleAttributeNamespaceString)) ||
+		InSourceForBinding.IsInNameSpace(FNiagaraConstants::UserNamespaceString) ||
+		InSourceForBinding.IsInNameSpace(FNiagaraConstants::SystemNamespaceString) ||
+		InSourceForBinding.IsInNameSpace(FNiagaraConstants::EmitterNamespaceString);
 }
 
 void UNiagaraRendererProperties::RenameEmitter(const FName& InOldName, const UNiagaraEmitter* InRenamedEmitter)
