Niagara - Add an optional setting to specify a required effect type for a project.

diff --git a/Engine/Plugins/FX/Niagara/Source/Niagara/Private/NiagaraSettings.cpp b/Engine/Plugins/FX/Niagara/Source/Niagara/Private/NiagaraSettings.cpp
--- a/Engine/Plugins/FX/Niagara/Source/Niagara/Private/NiagaraSettings.cpp
+++ b/Engine/Plugins/FX/Niagara/Source/Niagara/Private/NiagaraSettings.cpp
@@ -54,10 +54,15 @@ UNiagaraSettings::FOnNiagaraSettingsChanged& UNiagaraSettings::OnSettingsChanged
 }
 
 UNiagaraSettings::FOnNiagaraSettingsChanged UNiagaraSettings::SettingsChangedDelegate;
-#endif
 
-UNiagaraEffectType* UNiagaraSettings::GetDefaultEffectType()const
+UNiagaraEffectType* UNiagaraSettings::GetDefaultEffectType() const
 {
 	return Cast<UNiagaraEffectType>(DefaultEffectType.TryLoad());
 }
 
+UNiagaraEffectType* UNiagaraSettings::GetRequiredEffectType() const
+{
+	return Cast<UNiagaraEffectType>(RequiredEffectType.TryLoad());
+}
+#endif
+
diff --git a/Engine/Plugins/FX/Niagara/Source/Niagara/Public/NiagaraSettings.h b/Engine/Plugins/FX/Niagara/Source/Niagara/Public/NiagaraSettings.h
--- a/Engine/Plugins/FX/Niagara/Source/Niagara/Public/NiagaraSettings.h
+++ b/Engine/Plugins/FX/Niagara/Source/Niagara/Public/NiagaraSettings.h
@@ -144,6 +144,10 @@ class NIAGARA_API UNiagaraSettings : public UDeveloperSettings
 	UPROPERTY(config, EditAnywhere, Category = Niagara, meta = (AllowedClasses = "/Script/Niagara.NiagaraEffectType"))
 	FSoftObjectPath DefaultEffectType;
 
+	/** Specifies a required effect type which must be used for effects in the project. */
+	UPROPERTY(config, EditAnywhere, Category = Niagara, meta = (AllowedClasses = "/Script/Niagara.NiagaraEffectType"))
+	FSoftObjectPath RequiredEffectType;
+
 	/** Position pin type color. The other pin colors are defined in the general editor settings. */
 	UPROPERTY(config, EditAnywhere, Category=Niagara)
 	FLinearColor PositionPinTypeColor;
@@ -228,12 +232,13 @@ class NIAGARA_API UNiagaraSettings : public UDeveloperSettings
 #if WITH_EDITOR
 	void AddEnumParameterType(UEnum* Enum);
 	virtual FText GetSectionText() const override;
-#endif
+
 	// END UDeveloperSettings Interface
 
-	UNiagaraEffectType* GetDefaultEffectType()const;
+	UNiagaraEffectType* GetDefaultEffectType() const;
+
+	UNiagaraEffectType* GetRequiredEffectType() const;
 
-#if WITH_EDITOR
 	virtual void PostEditChangeProperty(FPropertyChangedEvent& PropertyChangedEvent) override;
 
 public:
diff --git a/Engine/Plugins/FX/Niagara/Source/NiagaraEditor/Private/NiagaraSystemFactoryNew.cpp b/Engine/Plugins/FX/Niagara/Source/NiagaraEditor/Private/NiagaraSystemFactoryNew.cpp
--- a/Engine/Plugins/FX/Niagara/Source/NiagaraEditor/Private/NiagaraSystemFactoryNew.cpp
+++ b/Engine/Plugins/FX/Niagara/Source/NiagaraEditor/Private/NiagaraSystemFactoryNew.cpp
@@ -210,7 +210,20 @@ void UNiagaraSystemFactoryNew::InitializeSystem(UNiagaraSystem* System, bool bCr
 
 	const UNiagaraSettings* Settings = GetDefault<UNiagaraSettings>();
 	check(Settings);
-	System->SetEffectType(Settings->GetDefaultEffectType());
+
+	UNiagaraEffectType* RequiredEffectType = Settings->GetRequiredEffectType();
+	if (RequiredEffectType != nullptr)
+	{
+		System->SetEffectType(RequiredEffectType);
+	}
+	else
+	{
+		UNiagaraEffectType* DefaultEffectType = Settings->GetDefaultEffectType();
+		if (DefaultEffectType != nullptr)
+		{
+			System->SetEffectType(DefaultEffectType);
+		}
+	}
 }
 
 #undef LOCTEXT_NAMESPACE
diff --git a/Engine/Plugins/FX/Niagara/Source/NiagaraEditor/Private/ViewModels/Stack/NiagaraStackSystemPropertiesItem.cpp b/Engine/Plugins/FX/Niagara/Source/NiagaraEditor/Private/ViewModels/Stack/NiagaraStackSystemPropertiesItem.cpp
--- a/Engine/Plugins/FX/Niagara/Source/NiagaraEditor/Private/ViewModels/Stack/NiagaraStackSystemPropertiesItem.cpp
+++ b/Engine/Plugins/FX/Niagara/Source/NiagaraEditor/Private/ViewModels/Stack/NiagaraStackSystemPropertiesItem.cpp
@@ -2,6 +2,7 @@
 
 #include "ViewModels/Stack/NiagaraStackSystemPropertiesItem.h"
 #include "NiagaraEmitter.h"
+#include "NiagaraSettings.h"
 #include "ViewModels/Stack/NiagaraStackObject.h"
 #include "NiagaraSystem.h"
 #include "NiagaraSystemDetailsCustomization.h"
@@ -134,6 +135,34 @@ void UNiagaraStackSystemPropertiesItem::RefreshChildrenInternal(const TArray<UNi
 		}
 	}
 
+	const UNiagaraSettings* Settings = GetDefault<UNiagaraSettings>();
+	if (System.IsValid() && Settings->GetRequiredEffectType() != nullptr)
+	{
+		UNiagaraEffectType* RequiredEffectType = Settings->GetRequiredEffectType();
+		if (System->GetEffectType() != RequiredEffectType)
+		{
+			TWeakObjectPtr<UNiagaraSystem> SystemWeak = System;
+			TWeakObjectPtr<UNiagaraEffectType> RequiredEffectTypeWeak = RequiredEffectType;
+			NewIssues.Add(FStackIssue(
+				EStackIssueSeverity::Error,
+				LOCTEXT("SystemNotUsingRequiredEffectTypeIssue", "Incorrect Effect Type In Use"),
+				LOCTEXT("SystemNotUsingRequiredEffectTypeIssueLong", "This project has a required effect type specified and this system is not currently using it."),
+				GetStackEditorDataKey(),
+				false,
+				{
+					FStackIssueFix(
+						LOCTEXT("SwitchToRequiredEffectType", "Switch to the required effect type for this project."),
+						FStackIssueFixDelegate::CreateLambda([SystemWeak, RequiredEffectTypeWeak]() 
+						{ 
+							if (SystemWeak.IsValid() && RequiredEffectTypeWeak.IsValid())
+							{
+								SystemWeak->SetEffectType(RequiredEffectTypeWeak.Get()); 
+							}
+						}))
+				}));
+		}
+	}
+
 	Super::RefreshChildrenInternal(CurrentChildren, NewChildren, NewIssues);
 }
 
