Refactored parameter panel logic to allow objects only in user parameters. The jira was created under the assumption that objects should be creatable as system attribute in the first place, which isn't the case as neither VM nor HLSL understand 'objects'.

diff --git a/Engine/Plugins/FX/Niagara/Source/NiagaraEditor/Private/Customizations/NiagaraComponentDetails.cpp b/Engine/Plugins/FX/Niagara/Source/NiagaraEditor/Private/Customizations/NiagaraComponentDetails.cpp
--- a/Engine/Plugins/FX/Niagara/Source/NiagaraEditor/Private/Customizations/NiagaraComponentDetails.cpp
+++ b/Engine/Plugins/FX/Niagara/Source/NiagaraEditor/Private/Customizations/NiagaraComponentDetails.cpp
@@ -1122,7 +1122,7 @@ TSharedRef<SWidget> FNiagaraSystemUserParameterBuilder::GetAddParameterMenu()
 	AddParameterMenu = SNew(SNiagaraAddParameterFromPanelMenu)
 		.Graphs(SystemViewModel.Pin()->GetParameterPanelViewModel()->GetEditableGraphsConst())
 		.OnNewParameterRequested(this, &FNiagaraSystemUserParameterBuilder::AddParameter)
-		.OnAllowMakeType(this, &FNiagaraSystemUserParameterBuilder::CanMakeNewParameterOfType)
+		.OnAllowMakeType(this, &FNiagaraSystemUserParameterBuilder::OnAllowMakeType)
 		.NamespaceId(UserCategory.NamespaceMetaData.GetGuid())
 		.ShowNamespaceCategory(false)
 		.ShowGraphParameters(false)
@@ -1138,9 +1138,11 @@ void FNiagaraSystemUserParameterBuilder::AddParameter(FNiagaraVariable NewParame
 	SystemViewModel.Pin()->GetUserParameterPanelViewModel()->OnParameterAdded().Execute(NewParameter);
 }
 
-bool FNiagaraSystemUserParameterBuilder::CanMakeNewParameterOfType(const FNiagaraTypeDefinition& InType) const
+bool FNiagaraSystemUserParameterBuilder::OnAllowMakeType(const FNiagaraTypeDefinition& InType) const
 {
-	return SystemViewModel.Pin()->GetParameterPanelViewModel()->CanMakeNewParameterOfType(InType);
+	FNiagaraNamespaceMetadata UserNameSpaceMetaData = GetDefault<UNiagaraEditorSettings>()->GetMetaDataForNamespaces({FNiagaraConstants::UserNamespace});
+	FNiagaraParameterPanelCategory Category(UserNameSpaceMetaData);
+	return SystemViewModel.Pin()->GetParameterPanelViewModel()->CanAddType(InType, Category);
 }
 
 void FNiagaraSystemUserParameterBuilder::ParameterValueChanged()
diff --git a/Engine/Plugins/FX/Niagara/Source/NiagaraEditor/Private/Customizations/NiagaraComponentDetails.h b/Engine/Plugins/FX/Niagara/Source/NiagaraEditor/Private/Customizations/NiagaraComponentDetails.h
--- a/Engine/Plugins/FX/Niagara/Source/NiagaraEditor/Private/Customizations/NiagaraComponentDetails.h
+++ b/Engine/Plugins/FX/Niagara/Source/NiagaraEditor/Private/Customizations/NiagaraComponentDetails.h
@@ -210,7 +210,7 @@ public:
 private:
 	TSharedRef<SWidget> GetAddParameterMenu();
 	void AddParameter(FNiagaraVariable NewParameter) const;
-	bool CanMakeNewParameterOfType(const FNiagaraTypeDefinition& InType) const;
+	bool OnAllowMakeType(const FNiagaraTypeDefinition& InType) const;
 	
 	void ParameterValueChanged();
 
diff --git a/Engine/Plugins/FX/Niagara/Source/NiagaraEditor/Private/ViewModels/NiagaraParameterPanelViewModel.cpp b/Engine/Plugins/FX/Niagara/Source/NiagaraEditor/Private/ViewModels/NiagaraParameterPanelViewModel.cpp
--- a/Engine/Plugins/FX/Niagara/Source/NiagaraEditor/Private/ViewModels/NiagaraParameterPanelViewModel.cpp
+++ b/Engine/Plugins/FX/Niagara/Source/NiagaraEditor/Private/ViewModels/NiagaraParameterPanelViewModel.cpp
@@ -265,6 +265,18 @@ void INiagaraParameterPanelViewModel::RefreshFull(bool bDoCategoryExpansion) con
 	}
 }
 
+bool INiagaraParameterPanelViewModel::CanAddType(const FNiagaraTypeDefinition& InType, TOptional<FNiagaraParameterPanelCategory> Category)
+{
+	bool bCanAddObjects = Category.IsSet() && Category.GetValue().NamespaceMetaData.GetGuid() == FNiagaraEditorGuids::UserNamespaceMetaDataGuid;
+	
+	return !(
+		InType == FNiagaraTypeDefinition::GetParameterMapDef()
+		|| InType == FNiagaraTypeDefinition::GetGenericNumericDef()
+		|| InType.IsInternalType()
+		|| ((!bCanAddObjects && InType.IsUObject()) && !InType.IsDataInterface())
+		);
+}
+
 bool INiagaraParameterPanelViewModel::GetCanDeleteParameterAndToolTip(const FNiagaraParameterPanelItem& ItemToDelete, FText& OutCanDeleteParameterToolTip) const
 {
 	if (ItemToDelete.bExternallyReferenced)
@@ -787,14 +799,6 @@ void INiagaraParameterPanelViewModel::SelectParameterItemByName(const FName Para
 	}
 }
 
-bool INiagaraParameterPanelViewModel::CanMakeNewParameterOfType(const FNiagaraTypeDefinition& InType)
-{
-	return InType != FNiagaraTypeDefinition::GetParameterMapDef()
-		&& InType != FNiagaraTypeDefinition::GetGenericNumericDef()
-		&& !InType.IsInternalType();
-}
-
-
 ///////////////////////////////////////////////////////////////////////////////
 /// System Toolkit Parameter Panel View Model								///
 ///////////////////////////////////////////////////////////////////////////////
@@ -1539,7 +1543,7 @@ FMenuAndSearchBoxWidgets FNiagaraSystemToolkitParameterPanelViewModel::GetParame
 		.OnSpecificParameterRequested(this, &FNiagaraSystemToolkitParameterPanelViewModel::FindOrAddParameter, Category)
 		.OnAddScriptVar(this, &FNiagaraSystemToolkitParameterPanelViewModel::AddScriptVariable)
 		.OnAddParameterDefinitions(this, &FNiagaraSystemToolkitParameterPanelViewModel::AddParameterDefinitions)
-		.OnAllowMakeType_Static(&INiagaraParameterPanelViewModel::CanMakeNewParameterOfType)
+		.OnAllowMakeType(this, &INiagaraParameterPanelViewModel::CanAddType, TOptional<FNiagaraParameterPanelCategory>(Category))
 		.NamespaceId(Category.NamespaceMetaData.GetGuid())
 		.ShowNamespaceCategory(false)
 		.ShowGraphParameters(false)
@@ -1711,7 +1715,7 @@ TSharedRef<SWidget> FNiagaraSystemToolkitParameterPanelViewModel::CreateAddParam
 		.OnSpecificParameterRequested_Lambda(AddExistingParameterLambda)
 		.OnAddScriptVar_Lambda(AddScriptVarLambda)
 		.OnAddParameterDefinitions(this, &FNiagaraSystemToolkitParameterPanelViewModel::AddParameterDefinitions)
-		.OnAllowMakeType_Static(&INiagaraParameterPanelViewModel::CanMakeNewParameterOfType)
+		.OnAllowMakeType(this, &INiagaraParameterPanelViewModel::CanAddType, TOptional<FNiagaraParameterPanelCategory>())
 		.AllowCreatingNew(true)
 		.NamespaceId(FNiagaraEditorUtilities::GetNamespaceIdForUsage(FNiagaraStackGraphUtilities::GetOutputNodeUsage(*AssignmentNode)))
 		.ShowNamespaceCategory(false)
@@ -2914,7 +2918,7 @@ FMenuAndSearchBoxWidgets FNiagaraScriptToolkitParameterPanelViewModel::GetParame
 		.OnSpecificParameterRequested(this, &FNiagaraScriptToolkitParameterPanelViewModel::FindOrAddParameter, Category)
 		.OnAddScriptVar(this, &FNiagaraScriptToolkitParameterPanelViewModel::AddScriptVariable)
 		.OnAddParameterDefinitions(this, &FNiagaraScriptToolkitParameterPanelViewModel::AddParameterDefinitions)
-		.OnAllowMakeType_Static(&INiagaraParameterPanelViewModel::CanMakeNewParameterOfType)
+		.OnAllowMakeType(this, &INiagaraParameterPanelViewModel::CanAddType, TOptional<FNiagaraParameterPanelCategory>(Category))
 		.NamespaceId(Category.NamespaceMetaData.GetGuid())
 		.ShowNamespaceCategory(false)
 		.ShowGraphParameters(false)
@@ -3502,7 +3506,7 @@ FMenuAndSearchBoxWidgets FNiagaraParameterDefinitionsToolkitParameterPanelViewMo
 	TSharedPtr<SNiagaraAddParameterFromPanelMenu> MenuWidget = SAssignNew(ParameterMenuWidget, SNiagaraAddParameterFromPanelMenu)
 	.OnNewParameterRequested(this, &FNiagaraParameterDefinitionsToolkitParameterPanelViewModel::AddParameter, Category, bRequestRename, bMakeUniqueName)
 	.OnSpecificParameterRequested(this, &FNiagaraParameterDefinitionsToolkitParameterPanelViewModel::FindOrAddParameter, Category)
-	.OnAllowMakeType_Static(&INiagaraParameterPanelViewModel::CanMakeNewParameterOfType)
+	.OnAllowMakeType(this, &INiagaraParameterPanelViewModel::CanAddType, TOptional<FNiagaraParameterPanelCategory>(Category))
 	.NamespaceId(Category.NamespaceMetaData.GetGuid())
 	.ShowNamespaceCategory(false)
 	.ShowGraphParameters(false)
diff --git a/Engine/Plugins/FX/Niagara/Source/NiagaraEditor/Public/ViewModels/NiagaraParameterPanelViewModel.h b/Engine/Plugins/FX/Niagara/Source/NiagaraEditor/Public/ViewModels/NiagaraParameterPanelViewModel.h
--- a/Engine/Plugins/FX/Niagara/Source/NiagaraEditor/Public/ViewModels/NiagaraParameterPanelViewModel.h
+++ b/Engine/Plugins/FX/Niagara/Source/NiagaraEditor/Public/ViewModels/NiagaraParameterPanelViewModel.h
@@ -190,7 +190,9 @@ public:
 	virtual void AddParameter(FNiagaraVariable NewVariable, const FNiagaraParameterPanelCategory Category, const bool bRequestRename, const bool bMakeUniqueName)  = 0;
 
 	virtual void FindOrAddParameter(FNiagaraVariable Variable, const FNiagaraParameterPanelCategory Category)  = 0;
-	
+
+	virtual bool CanAddType(const FNiagaraTypeDefinition& InType, TOptional<FNiagaraParameterPanelCategory> Category);
+
 	virtual bool GetCanAddParametersToCategory(FNiagaraParameterPanelCategory Category) const = 0;
 
 	virtual void DeleteParameters(const TArray<FNiagaraParameterPanelItem>& ItemsToDelete)  = 0;
@@ -289,8 +291,6 @@ public:
 	virtual bool IsVariableSelected(FNiagaraVariableBase& InVar) const { return false; };
 
 	virtual TSharedPtr<FNiagaraObjectSelection> GetVariableObjectSelection() { return TSharedPtr<FNiagaraObjectSelection>(); };
-
-	static bool CanMakeNewParameterOfType(const FNiagaraTypeDefinition& InType);
 	
 protected:
 	TWeakPtr<INiagaraParameterPanelViewModel> MainParameterPanelViewModel;
@@ -348,7 +348,7 @@ public:
 	virtual void AddParameter(FNiagaraVariable NewVariable, const FNiagaraParameterPanelCategory Category, const bool bRequestRename, const bool bMakeUniqueName)  override;
 
 	virtual void FindOrAddParameter(FNiagaraVariable Variable, const FNiagaraParameterPanelCategory Category)  override;
-
+	
 	virtual bool GetCanAddParametersToCategory(FNiagaraParameterPanelCategory Category) const override;
 
 	virtual void DeleteParameters(const TArray<FNiagaraParameterPanelItem>& ItemsToDelete)  override;
