Improvements to reduce the cost of UNiagaraSystem::RequestCompile() in the case of a warm DDC
-when using baked out RI parameters we now store prepared RI parameters in the DDC so that with a warm DDC we can apply them without having to run the relatively expesnive PrepareRapidIterationParameters step
-some small optimizations around generating FNiagaraVMExecutableDataId

diff --git a/Engine/Plugins/FX/Niagara/Source/Niagara/Classes/NiagaraScript.h b/Engine/Plugins/FX/Niagara/Source/Niagara/Classes/NiagaraScript.h
--- a/Engine/Plugins/FX/Niagara/Source/Niagara/Classes/NiagaraScript.h
+++ b/Engine/Plugins/FX/Niagara/Source/Niagara/Classes/NiagaraScript.h
@@ -445,6 +445,10 @@ public:
 	UPROPERTY()
 	TArray<FNiagaraCompileDependency> ExternalDependencies;
 
+	/** The default set of rapid iteration parameters defined by the script that this data represents.  In the case of baked
+	    in RI parameters this will be the values that are baked in, otherwise it will be the set of defaults based on the graphs. */
+	UPROPERTY()
+	TArray<FNiagaraVariable> BakedRapidIterationParameters;
 #endif
 
 	UPROPERTY()
@@ -1132,7 +1136,7 @@ public:
 	NIAGARA_API FOnPropertyChanged& OnPropertyChanged();
 
 	/** External call used to identify the values for a successful VM script compilation. OnVMScriptCompiled will be issued in this case.*/
-	void SetVMCompilationResults(const FNiagaraVMExecutableDataId& InCompileId, FNiagaraVMExecutableData& InScriptVM, FString EmitterUniqueName, const TMap<FName, UNiagaraDataInterface*>& ObjectNameMap);
+	void SetVMCompilationResults(const FNiagaraVMExecutableDataId& InCompileId, FNiagaraVMExecutableData& InScriptVM, FString EmitterUniqueName, const TMap<FName, UNiagaraDataInterface*>& ObjectNameMap, bool ApplyRapidIterationParameters);
 
 	/** In the event where we "merge" we duplicate the changes of the source script onto the newly cloned copy. This function will synchronize the compiled script 
 		results assuming that the scripts themselves are bound to the same key. This saves looking things up in the DDC. It returns true if successfully synchronized and 
diff --git a/Engine/Plugins/FX/Niagara/Source/Niagara/Classes/NiagaraSystem.h b/Engine/Plugins/FX/Niagara/Source/Niagara/Classes/NiagaraSystem.h
--- a/Engine/Plugins/FX/Niagara/Source/Niagara/Classes/NiagaraSystem.h
+++ b/Engine/Plugins/FX/Niagara/Source/Niagara/Classes/NiagaraSystem.h
@@ -180,13 +180,30 @@ struct FNiagaraSystemCompileRequest
 	UPROPERTY()
 	TArray<TObjectPtr<UObject>> RootObjects;
 
-	TArray<TSharedPtr<FNiagaraAsyncCompileTask, ESPMode::ThreadSafe>> DDCTasks;
-	
+	using FAsyncTaskPtr = TSharedPtr<FNiagaraAsyncCompileTask, ESPMode::ThreadSafe>;
+
+#if WITH_EDITORONLY_DATA
+	void Abort();
+	FAsyncTaskPtr* FindTask(const UNiagaraScript* Script);
+	bool QueryCompileComplete(UNiagaraSystem* OwningSystem, bool bWait, const double& MaxDuration);
+	bool Resolve(UNiagaraSystem* OwningSystem, FNiagaraParameterStore& ExposedParameters);
+	void Reset();
+	void UpdateSpawnDataInterfaces();
+	void Launch(UNiagaraSystem* OwningSystem, TConstArrayView<UNiagaraScript*> ScriptsNeedingCompile, TConstArrayView<FAsyncTaskPtr> Tasks);
+
 	bool bIsValid = true;
 	bool bForced = false;
 	bool bAllScriptsSynchronized = false;
+	bool bEvaluateParametersPending = false;
 
 	float CombinedCompileTime = 0.0f;
+
+private:
+	TArray<FAsyncTaskPtr> DDCTasks;
+
+	bool bDDCGetCompleted = false;
+
+#endif
 };
 
 struct FNiagaraEmitterExecutionIndex
@@ -259,6 +276,7 @@ public:
 	virtual void PreEditChange(FProperty* PropertyThatWillChange)override;
 	virtual void PostEditChangeProperty(struct FPropertyChangedEvent& PropertyChangedEvent) override; 
 	virtual void BeginCacheForCookedPlatformData(const ITargetPlatform *TargetPlatform) override;
+	virtual bool IsCachedCookedPlatformDataLoaded(const ITargetPlatform* TargetPlatform) override;
 	//~ End UObject interface
 
 	/** Helper method to handle when an internal variable has been renamed. Renames any downstream dependencies in the emitters or exposed variables.*/
@@ -772,6 +790,9 @@ public:
 
 	/** Updates any post compile data based upon data interfaces. */
 	void OnCompiledDataInterfaceChanged();
+
+	/** Updates the rapid iteration parameters for all scripts referenced by the system. */
+	void PrepareRapidIterationParametersForCompilation();
 #endif
 
 private:
@@ -803,7 +824,6 @@ private:
 
 	/** Helper for filling in attribute datasets per emitter. */
 	void InitEmitterDataSetCompiledData(FNiagaraDataSetCompiledData& DataSetToInit, const FNiagaraEmitterHandle& InAssociatedEmitterHandle);
-	void PrepareRapidIterationParametersForCompilation();
 #endif
 
 	void ResolveScalabilitySettings();
diff --git a/Engine/Plugins/FX/Niagara/Source/Niagara/Private/NiagaraAsyncCompile.cpp b/Engine/Plugins/FX/Niagara/Source/Niagara/Private/NiagaraAsyncCompile.cpp
--- a/Engine/Plugins/FX/Niagara/Source/Niagara/Private/NiagaraAsyncCompile.cpp
+++ b/Engine/Plugins/FX/Niagara/Source/Niagara/Private/NiagaraAsyncCompile.cpp
@@ -130,8 +130,12 @@ void FNiagaraAsyncCompileTask::ProcessCurrentState()
 
 	if (CurrentState == ENiagaraCompilationState::CheckDDC)
 	{
-		// check if the DDC data is ready
-		CheckDDCResult();
+		// if FNiagaraSystemCompileRequest is going to handle evaluating the CheckDDC state, then skip doing anything here
+		if (bCheckDDCEnabled)
+		{
+			// check if the DDC data is ready
+			CheckDDCResult();
+		}
 	}
 	else if (CurrentState == ENiagaraCompilationState::Precompile)
 	{
@@ -255,20 +259,13 @@ void FNiagaraAsyncCompileTask::PrecompileData()
 	SCOPE_CYCLE_COUNTER(STAT_Niagara_System_Precompile);
 	UE_LOG(LogNiagara, Verbose, TEXT("Getting precompile data for %s!"), *AssetPath);
 
-	// we compute the vmID here again to check if the script changed between the start of the task and now. If it was changed we need to update the ddc key,
-	// because otherwise we would save the compiled data under the wrong ddc key
-	FNiagaraVMExecutableDataId NewID;
-	ScriptPair.CompiledScript->ComputeVMCompilationId(NewID, FGuid());
-	FString CurrentDDCKey = UNiagaraScript::BuildNiagaraDDCKeyString(NewID, AssetPath);
-	if (DDCKey != CurrentDDCKey && LogNiagara.GetVerbosity() >= ELogVerbosity::Verbose)
+	// we do one more check here to see if the source data has been changed out from under us, if it has then we abort the compilation attempt
+	if (!CompilationIdMatchesRequest())
 	{
-		UE_LOG(LogNiagara, Warning, TEXT("Compile ID for %s changed during the compilation task."), *AssetPath);
-		FNiagaraVMExecutableDataId OldID = ScriptPair.CompileId;
-		LogVMID(OldID, "OldID");
-		LogVMID(NewID, "NewID");
+		AbortTask();
+		return;
 	}
-	DDCKey = CurrentDDCKey;
-	
+
 	ComputedPrecompileData = PrecompileReference->GetPrecompileData(ScriptPair.CompiledScript);
 	TSharedPtr<FNiagaraCompileRequestDuplicateDataBase, ESPMode::ThreadSafe> SystemPrecompileDuplicateData = PrecompileReference->GetPrecompileDuplicateData(ScriptPair.VersionedEmitter.Emitter, ScriptPair.CompiledScript);
 
@@ -373,6 +370,7 @@ void FNiagaraAsyncCompileTask::ProcessResult()
 	ScriptPair.CompileResults = ExeData;
 	ScriptPair.bResultsReady = true;
 	ScriptPair.CompileTime = float(FPlatformTime::Seconds() - StartCompileTime);
+	ExeData->BakedRapidIterationParameters = BakedRapidIterationParameters;
 	ScriptPair.CompiledScript->ExecToBinaryData(ScriptPair.CompiledScript, DDCOutData, *ExeData);
 	UE_LOG(LogNiagara, Verbose, TEXT("Got %i bytes in ddc data for %s"), DDCOutData.Num(), *AssetPath);
 }
@@ -441,10 +439,379 @@ void FNiagaraAsyncCompileTask::PutToDDC()
 	{
 		UE_LOG(LogNiagara, Verbose, TEXT("Writing data for %s to the ddc"), *AssetPath);
 		GetDerivedDataCache()->Put(*DDCKey, DDCOutData, *AssetPath, true);
+
+		if (AlternateDDCKey.IsSet())
+		{
+			GetDerivedDataCache()->Put(*(*AlternateDDCKey), DDCOutData, *AssetPath, true);
+		}
 	}
 }
 
-#endif
+void FNiagaraAsyncCompileTask::AssignInitialCompilationId(const FNiagaraVMExecutableDataId& InitialCompilationId)
+{
+	ScriptPair.CompileId = InitialCompilationId;
+}
+
+void FNiagaraAsyncCompileTask::UpdateCompilationId(const FNiagaraVMExecutableDataId& UpdatedCompilationId)
+{
+	// we need to update our CompileId and add a new DDC key where we'll push the data to
+	ScriptPair.CompileId = UpdatedCompilationId;
+	AlternateDDCKey = UNiagaraScript::BuildNiagaraDDCKeyString(UpdatedCompilationId, AssetPath);
+}
+
+bool FNiagaraAsyncCompileTask::CompilationIdMatchesRequest() const
+{
+	FNiagaraVMExecutableDataId CompilationId;
+	ScriptPair.CompiledScript->ComputeVMCompilationId(CompilationId, FGuid());
+
+	const bool IdMatches = CompilationId == ScriptPair.CompileId;
+	if (!IdMatches)
+	{
+		if (LogNiagara.GetVerbosity() >= ELogVerbosity::Verbose)
+		{
+			UE_LOG(LogNiagara, Warning, TEXT("Compile ID for %s changed during the compilation task."), *AssetPath);
+			FNiagaraVMExecutableDataId OldID = ScriptPair.CompileId;
+			LogVMID(ScriptPair.CompileId, "RequestedID");
+			LogVMID(ScriptPair.CompiledScript->GetComputedVMCompilationId(), "CurrentID");
+		}
+	}
+
+	return IdMatches;
+}
+
+void FNiagaraSystemCompileRequest::Abort()
+{
+	for (auto& AsyncTask : DDCTasks)
+	{
+		AsyncTask->AbortTask();
+	}
+}
+
+FNiagaraSystemCompileRequest::FAsyncTaskPtr* FNiagaraSystemCompileRequest::FindTask(const UNiagaraScript* Script)
+{
+	return DDCTasks.FindByPredicate([Script](const FAsyncTaskPtr& TaskPtr) -> bool
+	{
+		return TaskPtr->GetScriptPair().CompiledScript == Script;
+	});
+}
+
+bool FNiagaraSystemCompileRequest::QueryCompileComplete(UNiagaraSystem* OwningSystem, bool bWait, const double& MaxDuration)
+{
+	const double QueryStartTime = FPlatformTime::Seconds();
+
+	if (bWait)
+	{
+		for (FAsyncTaskPtr& AsyncTask : DDCTasks)
+		{
+			AsyncTask->bWaitForCompileJob = true;
+		}
+	}
+
+	if (bEvaluateParametersPending)
+	{
+		if (!bDDCGetCompleted)
+		{
+			bDDCGetCompleted = true;
+			bool bAllTasksFinalized = true;
+			{
+				for (FAsyncTaskPtr& AsyncTask : DDCTasks)
+				{
+					if (AsyncTask->CurrentState == ENiagaraCompilationState::CheckDDC)
+					{
+						AsyncTask->CheckDDCResult();
+						if (AsyncTask->CurrentState == ENiagaraCompilationState::CheckDDC)
+						{
+							bDDCGetCompleted = false;
+						}
+					}
+
+					if (AsyncTask->CurrentState != ENiagaraCompilationState::Finished)
+					{
+						bAllTasksFinalized = false;
+					}
+				}
+			}
+
+			if (bDDCGetCompleted && bEvaluateParametersPending && !bAllTasksFinalized)
+			{
+				// run through the active scripts and make sure that they are still up to date before we try to do
+				// any additional processing on them
+				for (FAsyncTaskPtr& AsyncTask : DDCTasks)
+				{
+					if (AsyncTask->CurrentState == ENiagaraCompilationState::Finished)
+					{
+						continue;
+					}
+
+					if (!AsyncTask->CompilationIdMatchesRequest())
+					{
+						AsyncTask->AbortTask();
+					}
+				}
+
+				// note - by having this deferred (after we've checked DDC) we do risk that the prepared RI parameters could be
+				// out of sync with what was there at the time the compile was requested.  So, we'll only take this path if we're
+				// not in edit mode.  In edit mode we'll prepare the RI parameters during UNiagaraSystem::RequestCompile()
+				OwningSystem->PrepareRapidIterationParametersForCompilation();
+
+				// go through our tasks and evaluate the rapid iteration parameters.  We'll need to store them in the ExeData
+				// so that subsequent builds can take advantage of the work done in PrepareRapidIterationParamtersForCompilation
+				// by storing it in the ExeData.  Additionally if this evaluation has changed the rapid iteration parameters the
+				// DDC key will be different, and we'll want to push the data both to the old key and to the new key.
+				const int32 TaskCount = DDCTasks.Num();
+				for (int32 TaskIt = 0; TaskIt < TaskCount; ++TaskIt)
+				{
+					FAsyncTaskPtr& AsyncTask = DDCTasks[TaskIt];
+					if (AsyncTask->CurrentState == ENiagaraCompilationState::Finished)
+					{
+						continue;
+					}
+
+					const FEmitterCompiledScriptPair& ScriptPair = AsyncTask->GetScriptPair();
+					if (UNiagaraScript* CompiledScript = ScriptPair.CompiledScript)
+					{
+						TConstArrayView<FNiagaraVariableWithOffset> RapidIterationParameters = CompiledScript->RapidIterationParameters.ReadParameterVariables();
+
+						AsyncTask->BakedRapidIterationParameters.Reserve(RapidIterationParameters.Num());
+						for (const FNiagaraVariableWithOffset& ParamWithOffset : RapidIterationParameters)
+						{
+							FNiagaraVariable& Parameter = AsyncTask->BakedRapidIterationParameters.Add_GetRef(ParamWithOffset);
+							Parameter.SetData(CompiledScript->RapidIterationParameters.GetParameterData(ParamWithOffset.Offset));
+						}
+
+						// if preparing the rapid iteration parameters has adjusted the compile id, then we want to be sure
+						// to populate both entries in the DDC with the results (so subsequent loads don't need to go through
+						// this same code path, as the serialized compileID will find the DDC results)
+						FNiagaraVMExecutableDataId PostRapidIterationParameterID;
+						CompiledScript->ComputeVMCompilationId(PostRapidIterationParameterID, FGuid());
+
+						if (PostRapidIterationParameterID != ScriptPair.CompileId)
+						{
+							AsyncTask->UpdateCompilationId(PostRapidIterationParameterID);
+						}
+					}
+				}
+			}
+		}
+
+		if (!bDDCGetCompleted)
+		{
+			return false;
+		}
+	}
+
+	bool bCompileComplete = true;
+
+	if (bWait)
+	{
+		for (FAsyncTaskPtr& AsyncTask : DDCTasks)
+		{
+			// before we start to wait for the compile results, we start the compilation of all remaining tasks
+			while (!AsyncTask->IsDone() && AsyncTask->CurrentState < ENiagaraCompilationState::AwaitResult)
+			{
+				AsyncTask->ProcessCurrentState();
+			}
+		}
+	}
+
+	// Check to see if ALL of the sub-requests have resolved.
+	{
+		for (auto& AsyncTask : DDCTasks)
+		{
+			// if a task is very expensive we bail and continue on the next frame  
+			if (FPlatformTime::Seconds() - QueryStartTime < MaxDuration)
+			{
+				AsyncTask->ProcessCurrentState();
+			}
+
+			if (AsyncTask->IsDone())
+			{
+				continue;
+			}
+
+			if (!AsyncTask->bFetchedGCObjects && AsyncTask->CurrentState > ENiagaraCompilationState::Precompile && AsyncTask->CurrentState <= ENiagaraCompilationState::ProcessResult)
+			{
+				RootObjects.Append(AsyncTask->PrecompileReference->CompilationRootObjects);
+				AsyncTask->bFetchedGCObjects = true;
+			}
+
+			if (bWait)
+			{
+				AsyncTask->WaitAndResolveResult();
+			}
+			else
+			{
+				bCompileComplete = false;
+			}
+		}
+	}
+
+	return bCompileComplete;
+}
+
+bool FNiagaraSystemCompileRequest::Resolve(UNiagaraSystem* OwningSystem, FNiagaraParameterStore& ExposedParameters)
+{
+	bool bHasCompiledJobs = false;
+
+	for (FAsyncTaskPtr& AsyncTask : DDCTasks)
+	{
+		bHasCompiledJobs |= AsyncTask->bUsedShaderCompilerWorker;
+
+		FEmitterCompiledScriptPair& EmitterCompiledScriptPair = AsyncTask->ScriptPair;
+		if (EmitterCompiledScriptPair.bResultsReady)
+		{
+			TSharedPtr<FNiagaraVMExecutableData> ExeData = EmitterCompiledScriptPair.CompileResults;
+			CombinedCompileTime += EmitterCompiledScriptPair.CompileTime;
+			UNiagaraScript* CompiledScript = EmitterCompiledScriptPair.CompiledScript;
+
+			// generate the ObjectNameMap from the source (from the duplicated data if available).
+			TMap<FName, UNiagaraDataInterface*> ObjectNameMap;
+			if (AsyncTask->ComputedPrecompileDuplicateData.IsValid())
+			{
+				if (const UNiagaraScriptSourceBase* ScriptSource = AsyncTask->ComputedPrecompileDuplicateData->GetScriptSource())
+				{
+					ObjectNameMap = ScriptSource->ComputeObjectNameMap(*OwningSystem, CompiledScript->GetUsage(), CompiledScript->GetUsageId(), AsyncTask->UniqueEmitterName);
+				}
+				else
+				{
+					checkf(false, TEXT("Failed to get ScriptSource from ComputedPrecompileDuplicateData"));
+				}
+
+				// we need to replace the DI in the above generated map with the duplicates that we've created as a part of the duplication process
+				TMap<FName, UNiagaraDataInterface*> DuplicatedObjectNameMap = AsyncTask->ComputedPrecompileDuplicateData->GetObjectNameMap();
+				for (auto ObjectMapIt = ObjectNameMap.CreateIterator(); ObjectMapIt; ++ObjectMapIt)
+				{
+					if (UNiagaraDataInterface** Replacement = DuplicatedObjectNameMap.Find(ObjectMapIt.Key()))
+					{
+						ObjectMapIt.Value() = *Replacement;
+					}
+					else
+					{
+						ObjectMapIt.RemoveCurrent();
+					}
+				}
+			}
+			else
+			{
+				const UNiagaraScriptSourceBase* ScriptSource = CompiledScript->GetLatestSource();
+				ObjectNameMap = ScriptSource->ComputeObjectNameMap(*OwningSystem, CompiledScript->GetUsage(), CompiledScript->GetUsageId(), AsyncTask->UniqueEmitterName);
+			}
+
+			EmitterCompiledScriptPair.CompiledScript->SetVMCompilationResults(EmitterCompiledScriptPair.CompileId, *ExeData, AsyncTask->UniqueEmitterName, ObjectNameMap, bEvaluateParametersPending);
+
+			// Synchronize the variables that we actually encountered during precompile so that we can expose them to the end user.
+			TArray<FNiagaraVariable> OriginalExposedParams;
+			ExposedParameters.GetParameters(OriginalExposedParams);
+			TArray<FNiagaraVariable>& EncounteredExposedVars = AsyncTask->EncounteredExposedVars;
+			for (int32 i = 0; i < EncounteredExposedVars.Num(); i++)
+			{
+				if (OriginalExposedParams.Contains(EncounteredExposedVars[i]) == false)
+				{
+					// Just in case it wasn't added previously..
+					ExposedParameters.AddParameter(EncounteredExposedVars[i], true, false);
+				}
+			}
+		}
+	}
+
+	return bHasCompiledJobs;
+}
+
+void FNiagaraSystemCompileRequest::Reset()
+{
+	// clean up the precompile data
+	for (FAsyncTaskPtr& AsyncTask : DDCTasks)
+	{
+		AsyncTask->ComputedPrecompileData.Reset();
+		if (AsyncTask->ComputedPrecompileDuplicateData.IsValid())
+		{
+			AsyncTask->ComputedPrecompileDuplicateData->ReleaseCompilationCopies();
+			AsyncTask->ComputedPrecompileDuplicateData.Reset();
+		}
+	}
+}
+
+void FNiagaraSystemCompileRequest::UpdateSpawnDataInterfaces()
+{
+	// HACK: This is a temporary hack to fix an issue where data interfaces used by modules and dynamic inputs in the
+	// particle update script aren't being shared by the interpolated spawn script when accessed directly.  This works
+	// properly if the data interface is assigned to a named particle parameter and then linked to an input.
+	// TODO: Bind these data interfaces the same way parameter data interfaces are bound.
+	for (auto& AsyncTask : DDCTasks)
+	{
+		FEmitterCompiledScriptPair& EmitterCompiledScriptPair = AsyncTask->ScriptPair;
+		FVersionedNiagaraEmitter Emitter = EmitterCompiledScriptPair.VersionedEmitter;
+		UNiagaraScript* CompiledScript = EmitterCompiledScriptPair.CompiledScript;
+
+		if (UNiagaraScript::IsEquivalentUsage(CompiledScript->GetUsage(), ENiagaraScriptUsage::ParticleUpdateScript))
+		{
+			UNiagaraScript* SpawnScript = Emitter.GetEmitterData()->SpawnScriptProps.Script;
+			for (const FNiagaraScriptDataInterfaceInfo& UpdateDataInterfaceInfo : CompiledScript->GetCachedDefaultDataInterfaces())
+			{
+				// If the data interface isn't being written to a parameter map then it won't be bound properly so we
+				// assign the update scripts copy of the data interface to the spawn scripts copy by pointer so that they will share
+				// the data interface at runtime and will both be updated in the editor.
+				for (FNiagaraScriptDataInterfaceInfo& SpawnDataInterfaceInfo : SpawnScript->GetCachedDefaultDataInterfaces())
+				{
+					if (SpawnDataInterfaceInfo.RegisteredParameterMapWrite == NAME_None && UpdateDataInterfaceInfo.RegisteredParameterMapWrite == NAME_None && UpdateDataInterfaceInfo.Name == SpawnDataInterfaceInfo.Name)
+					{
+						SpawnDataInterfaceInfo.DataInterface = UpdateDataInterfaceInfo.DataInterface;
+						break;
+					}
+				}
+			}
+		}
+	}
+}
+
+void FNiagaraSystemCompileRequest::Launch(UNiagaraSystem* OwningSystem, TConstArrayView<UNiagaraScript*> ScriptsNeedingCompile, TConstArrayView<FAsyncTaskPtr> Tasks)
+{
+	bEvaluateParametersPending = !OwningSystem->ShouldUseRapidIterationParameters();
+	DDCTasks = Tasks;
+
+	// prepare data for any precompile the ddc tasks need to do
+	TSharedPtr<FNiagaraLazyPrecompileReference, ESPMode::ThreadSafe> PrecompileReference = MakeShared<FNiagaraLazyPrecompileReference, ESPMode::ThreadSafe>();
+	PrecompileReference->System = OwningSystem;
+	PrecompileReference->Scripts = ScriptsNeedingCompile;
+
+	const TArray<FNiagaraEmitterHandle>& EmitterHandles = OwningSystem->GetEmitterHandles();
+	for (int32 i = 0; i < EmitterHandles.Num(); i++)
+	{
+		FNiagaraEmitterHandle Handle = EmitterHandles[i];
+		if (Handle.GetIsEnabled() && Handle.GetEmitterData())
+		{
+			TArray<UNiagaraScript*> EmitterScripts;
+			Handle.GetEmitterData()->GetScripts(EmitterScripts, false, true);
+			check(EmitterScripts.Num() > 0);
+			for (UNiagaraScript* EmitterScript : EmitterScripts)
+			{
+				PrecompileReference->EmitterScriptIndex.Add(EmitterScript, i);
+			}
+		}
+	}
+
+	for (UNiagaraScript* Script : ScriptsNeedingCompile)
+	{
+		if (FAsyncTaskPtr* AsyncTaskPtr = FindTask(Script))
+		{
+			if (ensure(*AsyncTaskPtr))
+			{
+				FNiagaraAsyncCompileTask& CompileTask = *(*AsyncTaskPtr);
+				CompileTask.DDCKey = UNiagaraScript::BuildNiagaraDDCKeyString(CompileTask.GetScriptPair().CompileId, CompileTask.AssetPath);
+				CompileTask.PrecompileReference = PrecompileReference;
+				CompileTask.StartTaskTime = FPlatformTime::Seconds();
+				CompileTask.bCheckDDCEnabled = !bEvaluateParametersPending;
+
+				// fire off all the ddc tasks, which will trigger the compilation if the data is not in the ddc
+				UE_LOG(LogNiagara, Verbose, TEXT("Scheduling async get task for %s"), *CompileTask.AssetPath);
+				CompileTask.TaskHandle = GetDerivedDataCacheRef().GetAsynchronous(*CompileTask.DDCKey, CompileTask.AssetPath);
+			}
+		}
+	}
+}
+
+
+#endif // WITH_EDITORONLY_DATA
 
 #undef LOCTEXT_NAMESPACE // NiagaraAsyncCompile
 
diff --git a/Engine/Plugins/FX/Niagara/Source/Niagara/Private/NiagaraAsyncCompile.h b/Engine/Plugins/FX/Niagara/Source/Niagara/Private/NiagaraAsyncCompile.h
--- a/Engine/Plugins/FX/Niagara/Source/Niagara/Private/NiagaraAsyncCompile.h
+++ b/Engine/Plugins/FX/Niagara/Source/Niagara/Private/NiagaraAsyncCompile.h
@@ -22,7 +22,7 @@ struct FEmitterCompiledScriptPair
 	uint32 PendingJobID = INDEX_NONE; // this is the ID for any active shader compiler worker job
 	FNiagaraVMExecutableDataId CompileId;
 	TSharedPtr<FNiagaraVMExecutableData> CompileResults;
-	float CompileTime = 0.0f;
+	double CompileTime = 0;
 };
 
 UENUM()
@@ -60,6 +60,7 @@ class FNiagaraAsyncCompileTask
 #if WITH_EDITORONLY_DATA
 public:
 	FString DDCKey;
+	TOptional<FString> AlternateDDCKey;
 	FString AssetPath;
 	FString UniqueEmitterName;
 	uint32 TaskHandle = 0;
@@ -71,9 +72,15 @@ public:
 	bool bUsedShaderCompilerWorker = false;
 	bool bFetchedGCObjects = false;
 	bool bExperimentalVMDisabled = true;
+
+	// in order to coordinate between tasks associated with a system CheckDDC can be handled external to the task.
+	// In that case we disable this flag
+	bool bCheckDDCEnabled = true;
+
 	UNiagaraSystem* OwningSystem;
 	FEmitterCompiledScriptPair ScriptPair;
 	TArray<FNiagaraVariable> EncounteredExposedVars;
+	TArray<FNiagaraVariable> BakedRapidIterationParameters;
 
 	ENiagaraCompilationState CurrentState;
 
@@ -103,5 +110,8 @@ public:
 	void ProcessResult();
 	void OptimizeByteCode();
 
+	void AssignInitialCompilationId(const FNiagaraVMExecutableDataId& InitialCompilationId);
+	void UpdateCompilationId(const FNiagaraVMExecutableDataId& UpdatedCompilationId);
+	bool CompilationIdMatchesRequest() const;
 #endif
 };
diff --git a/Engine/Plugins/FX/Niagara/Source/Niagara/Private/NiagaraScript.cpp b/Engine/Plugins/FX/Niagara/Source/Niagara/Private/NiagaraScript.cpp
--- a/Engine/Plugins/FX/Niagara/Source/Niagara/Private/NiagaraScript.cpp
+++ b/Engine/Plugins/FX/Niagara/Source/Niagara/Private/NiagaraScript.cpp
@@ -3,6 +3,7 @@
 #include "NiagaraScript.h"
 
 #include "Algo/RemoveIf.h"
+#include "Algo/Unique.h"
 #include "DataDrivenShaderPlatformInfo.h"
 #include "Modules/ModuleManager.h"
 #include "NiagaraScriptSourceBase.h"
@@ -1099,21 +1100,19 @@ void UNiagaraScript::ComputeVMCompilationId(FNiagaraVMExecutableDataId& Id, FGui
 			{
 				Id.AdditionalDefines.Add(GNiagaraForceSafeScriptAttributeTrim ? TEXT("TrimAttributesSafe") : TEXT("TrimAttributes"));
 
-				TArray<FString> PreserveAttributes;
-				auto AddAttributeToPreserve = 
-					[&](FName AttributeName)
+				TArray<FName> AttributesToPreserve;
+				auto AddAttributeToPreserve = [&](FName AttributeName) -> void
+				{
+					if (AttributeName != NAME_None)
 					{
-						if ( !AttributeName.IsNone() )
-						{
-							PreserveAttributes.AddUnique(FString::Printf(TEXT("PreserveAttribute=%s"), *AttributeName.ToString()));
-						}
-					};
+						AttributesToPreserve.Add(AttributeName);
+					}
+				};
 
 				// preserve the attributes that have been defined on the emitter directly
 				for (const FString& Attribute : EmitterData->AttributesToPreserve)
 				{
-					const FString PreserveDefine = TEXT("PreserveAttribute=") + Attribute;
-					PreserveAttributes.AddUnique(PreserveDefine);
+					AddAttributeToPreserve(*Attribute);
 				}
 
 				// Now preserve the attributes that have been defined on the renderers in use
@@ -1144,9 +1143,20 @@ void UNiagaraScript::ComputeVMCompilationId(FNiagaraVMExecutableDataId& Id, FGui
 				}
 
 				// We sort the keys so that it doesn't matter what order they were defined in.
-				PreserveAttributes.Sort([](const FString& A, const FString& B) -> bool { return A < B; });
+				if (!AttributesToPreserve.IsEmpty())
+				{
+					AttributesToPreserve.Sort(FNameLexicalLess());
+					AttributesToPreserve.SetNum(Algo::Unique(AttributesToPreserve));
 
-				Id.AdditionalDefines.Append(PreserveAttributes);
+					Id.AdditionalDefines.Reserve(Id.AdditionalDefines.Num() + AttributesToPreserve.Num());
+					for (FName AttributeName : AttributesToPreserve)
+					{
+						TStringBuilder<256> PreserveAttributeDefine;
+						PreserveAttributeDefine.Append(TEXT("PreserveAttribute="));
+						AttributeName.AppendString(PreserveAttributeDefine);
+						Id.AdditionalDefines.Add(PreserveAttributeDefine.ToString());
+					}
+				}
 			}
 
 			ComputeVMCompilationId_EmitterShared(Id, Outer, ENiagaraRendererSourceDataMode::Particles);
@@ -2773,7 +2783,7 @@ FName ResolveEmitterAlias(const FName& InName, const FString& InAlias)
 	return Var.GetName();
 }
 
-void UNiagaraScript::SetVMCompilationResults(const FNiagaraVMExecutableDataId& InCompileId, FNiagaraVMExecutableData& InScriptVM, FString EmitterUniqueName, const TMap<FName, UNiagaraDataInterface*>& ObjectNameMap)
+void UNiagaraScript::SetVMCompilationResults(const FNiagaraVMExecutableDataId& InCompileId, FNiagaraVMExecutableData& InScriptVM, FString EmitterUniqueName, const TMap<FName, UNiagaraDataInterface*>& ObjectNameMap, bool bApplyRapidIterationParameters)
 {
 	CachedScriptVMId = InCompileId;
 	CachedScriptVM = InScriptVM;
@@ -2829,6 +2839,19 @@ void UNiagaraScript::SetVMCompilationResults(const FNiagaraVMExecutableDataId& I
 		check(CachedDefaultDataInterfaces[Idx].DataInterface != nullptr);
 	}
 
+	if (bApplyRapidIterationParameters)
+	{
+		const bool bClearBindings = false;
+		RapidIterationParameters.Empty(bClearBindings);
+		for (const FNiagaraVariable& Parameter : InScriptVM.BakedRapidIterationParameters)
+		{
+			const bool bInitialize = false;
+			const bool bTriggerRebind = false;
+			RapidIterationParameters.AddParameter(Parameter, bInitialize, bTriggerRebind);
+		}
+		RapidIterationParameters.TriggerOnLayoutChanged();
+	}
+
 	GenerateStatIDs();
 
 	// Now go ahead and trigger the GPU script compile now that we have a compiled GPU hlsl script.
@@ -2917,7 +2940,7 @@ void UNiagaraScript::RequestCompile(const FGuid& ScriptVersion, bool bForceCompi
 			if (BinaryToExecData(this, OutData, ExeData))
 			{
 				COOK_STAT(Timer.AddHit(OutData.Num()));
-				SetVMCompilationResults(LastGeneratedVMId, ExeData, FString(), RequestDuplicateData->GetObjectNameMap());
+				SetVMCompilationResults(LastGeneratedVMId, ExeData, FString(), RequestDuplicateData->GetObjectNameMap(), false);
 				return;
 			}
 		}
@@ -2930,7 +2953,7 @@ void UNiagaraScript::RequestCompile(const FGuid& ScriptVersion, bool bForceCompi
 		TSharedPtr<FNiagaraVMExecutableData> ExeData = NiagaraModule.GetCompileJobResult(JobHandle, true);
 		if (ExeData)
 		{
-			SetVMCompilationResults(LastGeneratedVMId, *ExeData, FString(), RequestDuplicateData->GetObjectNameMap());
+			SetVMCompilationResults(LastGeneratedVMId, *ExeData, FString(), RequestDuplicateData->GetObjectNameMap(), false);
 			// save result to the ddc
 			if (ExecToBinaryData(this, OutData, *ExeData))
 			{
diff --git a/Engine/Plugins/FX/Niagara/Source/Niagara/Private/NiagaraSystem.cpp b/Engine/Plugins/FX/Niagara/Source/Niagara/Private/NiagaraSystem.cpp
--- a/Engine/Plugins/FX/Niagara/Source/Niagara/Private/NiagaraSystem.cpp
+++ b/Engine/Plugins/FX/Niagara/Source/Niagara/Private/NiagaraSystem.cpp
@@ -266,10 +266,22 @@ void UNiagaraSystem::BeginCacheForCookedPlatformData(const ITargetPlatform *Targ
 	
 	EnsureFullyLoaded();
 #if WITH_EDITORONLY_DATA
-	WaitForCompilationComplete();
+	if (bNeedsRequestCompile)
+	{
+		RequestCompile(false);
+	}
 #endif
 }
 
+bool UNiagaraSystem::IsCachedCookedPlatformDataLoaded(const ITargetPlatform* TargetPlatform)
+{
+	if (Super::IsCachedCookedPlatformDataLoaded(TargetPlatform))
+	{
+		return PollForCompilationComplete(false);
+	}
+	return false;
+}
+
 void UNiagaraSystem::HandleVariableRenamed(const FNiagaraVariable& InOldVariable, const FNiagaraVariable& InNewVariable, bool bUpdateContexts)
 {
 	if (InOldVariable.IsInNameSpace(FNiagaraConstants::UserNamespaceString))
@@ -2216,10 +2228,7 @@ void UNiagaraSystem::KillAllActiveCompilations()
 	InvalidateActiveCompiles();
 	for (FNiagaraSystemCompileRequest& Request : ActiveCompilations)
 	{
-		for (auto& AsyncTask : Request.DDCTasks)
-		{
-			AsyncTask->AbortTask();
-		}
+		Request.Abort();
 	}
 	ActiveCompilations.Empty();
 }
@@ -2366,8 +2375,8 @@ bool UNiagaraSystem::CompilationResultsValid(FNiagaraSystemCompileRequest& Compi
 {
 	// for now the only thing we're concerned about is if we've got results for SystemSpawn and SystemUpdate scripts
 	// then we need to make sure that they agree in terms of the dataset attributes
-	TSharedPtr<FNiagaraAsyncCompileTask, ESPMode::ThreadSafe>* SpawnScriptRequest = CompileRequest.DDCTasks.FindByPredicate([this](const auto& Task) { return Task->GetScriptPair().CompiledScript == SystemSpawnScript; });
-	TSharedPtr<FNiagaraAsyncCompileTask, ESPMode::ThreadSafe>* UpdateScriptRequest = CompileRequest.DDCTasks.FindByPredicate([this](const auto& Task) { return Task->GetScriptPair().CompiledScript == SystemUpdateScript; });
+	FNiagaraSystemCompileRequest::FAsyncTaskPtr* SpawnScriptRequest = CompileRequest.FindTask(SystemSpawnScript);
+	FNiagaraSystemCompileRequest::FAsyncTaskPtr* UpdateScriptRequest = CompileRequest.FindTask(SystemUpdateScript);
 
 	const bool SpawnScriptValid = SpawnScriptRequest
 		&& SpawnScriptRequest->Get()->GetScriptPair().CompileResults.IsValid()
@@ -2669,26 +2678,6 @@ void UNiagaraSystem::EvaluateCompileResultDependencies() const
 	}
 }
 
-void UNiagaraSystem::PreProcessWaitingDDCTasks(bool bProcessForWait)
-{
-	if (!bProcessForWait)
-	{
-		return;
-	}
-	for (FNiagaraSystemCompileRequest& CompileRequest : ActiveCompilations)
-	{
-		for (auto& AsyncTask : CompileRequest.DDCTasks)
-		{
-			AsyncTask->bWaitForCompileJob = true;
-			// before we start to wait for the compile results, we start the compilation of all remaining tasks
-			while (!AsyncTask->IsDone() && AsyncTask->CurrentState < ENiagaraCompilationState::AwaitResult)
-			{
-				AsyncTask->ProcessCurrentState();
-			}
-		}
-	}
-}
-
 bool UNiagaraSystem::QueryCompileComplete(bool bWait, bool bDoPost, bool bDoNotApply)
 {
 	check(IsInGameThread());
@@ -2697,41 +2686,10 @@ bool UNiagaraSystem::QueryCompileComplete(bool bWait, bool bDoPost, bool bDoNotA
 		bCompilationReentrantGuard = true;
 		ON_SCOPE_EXIT { bCompilationReentrantGuard = false; };
 
-		PreProcessWaitingDDCTasks(bWait);
-		
 		FNiagaraSystemCompileRequest& CompileRequest = ActiveCompilations[0];
-		bool bAreWeWaitingForAnyResults = false;
-		double StartTime = FPlatformTime::Seconds();
-		
-		// Check to see if ALL of the sub-requests have resolved.
-		for (auto& AsyncTask : CompileRequest.DDCTasks)
-		{
-			// if a task is very expensive we bail and continue on the next frame  
-			if (FPlatformTime::Seconds() - StartTime < 0.125)
-			{
-				AsyncTask->ProcessCurrentState();
-			}
 
-			if (AsyncTask->IsDone())
-			{
-				continue;
-			}
-
-			if (!AsyncTask->bFetchedGCObjects && AsyncTask->CurrentState > ENiagaraCompilationState::Precompile && AsyncTask->CurrentState <= ENiagaraCompilationState::ProcessResult)
-			{
-				CompileRequest.RootObjects.Append(AsyncTask->PrecompileReference->CompilationRootObjects);
-				AsyncTask->bFetchedGCObjects = true;
-			}
-			
-			if (bWait)
-			{
-				AsyncTask->WaitAndResolveResult();
-			}
-			else
-			{
-				bAreWeWaitingForAnyResults = true;
-			}
-		}
+		constexpr double MaxDurationPerIteration = 0.125;
+		bool bAreWeWaitingForAnyResults = !CompileRequest.QueryCompileComplete(this, bWait, MaxDurationPerIteration);
 
 		// Make sure that we aren't waiting for any results to come back.
 		if (bAreWeWaitingForAnyResults)
@@ -2758,78 +2716,8 @@ bool UNiagaraSystem::QueryCompileComplete(bool bWait, bool bDoPost, bool bDoNotA
 		}
 
 		// Now that the above code says they are all complete, go ahead and resolve them all at once.
-		bool bHasCompiledJobs = false;
-		for (auto& AsyncTask : CompileRequest.DDCTasks)
-		{
-			bHasCompiledJobs |= AsyncTask->bUsedShaderCompilerWorker;
-			
-			FEmitterCompiledScriptPair& EmitterCompiledScriptPair = AsyncTask->ScriptPair;
-			if (EmitterCompiledScriptPair.bResultsReady)
-			{
-				TSharedPtr<FNiagaraVMExecutableData> ExeData = EmitterCompiledScriptPair.CompileResults;
-				CompileRequest.CombinedCompileTime += EmitterCompiledScriptPair.CompileTime;
-				UNiagaraScript* CompiledScript = EmitterCompiledScriptPair.CompiledScript;
-
-				// generate the ObjectNameMap from the source (from the duplicated data if available).
-				TMap<FName, UNiagaraDataInterface*> ObjectNameMap;
-				if (AsyncTask->ComputedPrecompileDuplicateData.IsValid())
-				{
-					if (const UNiagaraScriptSourceBase* ScriptSource = AsyncTask->ComputedPrecompileDuplicateData->GetScriptSource())
-					{
-						ObjectNameMap = ScriptSource->ComputeObjectNameMap(*this, CompiledScript->GetUsage(), CompiledScript->GetUsageId(), AsyncTask->UniqueEmitterName);
-					}
-					else
-					{
-						checkf(false, TEXT("Failed to get ScriptSource from ComputedPrecompileDuplicateData"));
-					}
-
-					// we need to replace the DI in the above generated map with the duplicates that we've created as a part of the duplication process
-					TMap<FName, UNiagaraDataInterface*> DuplicatedObjectNameMap = AsyncTask->ComputedPrecompileDuplicateData->GetObjectNameMap();
-					for (auto ObjectMapIt = ObjectNameMap.CreateIterator(); ObjectMapIt; ++ObjectMapIt)
-					{
-						if (UNiagaraDataInterface** Replacement = DuplicatedObjectNameMap.Find(ObjectMapIt.Key()))
-						{
-							ObjectMapIt.Value() = *Replacement;
-						}
-						else
-						{
-							ObjectMapIt.RemoveCurrent();
-						}
-					}
-				}
-				else
-				{
-					const UNiagaraScriptSourceBase* ScriptSource = CompiledScript->GetLatestSource();
-					ObjectNameMap = ScriptSource->ComputeObjectNameMap(*this, CompiledScript->GetUsage(), CompiledScript->GetUsageId(), AsyncTask->UniqueEmitterName);
-				}
-
-				EmitterCompiledScriptPair.CompiledScript->SetVMCompilationResults(EmitterCompiledScriptPair.CompileId, *ExeData, AsyncTask->UniqueEmitterName, ObjectNameMap);
-
-				// Synchronize the variables that we actually encountered during precompile so that we can expose them to the end user.
-				TArray<FNiagaraVariable> OriginalExposedParams;
-				ExposedParameters.GetParameters(OriginalExposedParams);
-				TArray<FNiagaraVariable>& EncounteredExposedVars = AsyncTask->EncounteredExposedVars;
-				for (int32 i = 0; i < EncounteredExposedVars.Num(); i++)
-				{
-					if (OriginalExposedParams.Contains(EncounteredExposedVars[i]) == false)
-					{
-						// Just in case it wasn't added previously..
-						ExposedParameters.AddParameter(EncounteredExposedVars[i], true, false);
-					}
-				}
-			}
-		}
-
-		// clean up the precompile data
-		for (auto& AsyncTask : CompileRequest.DDCTasks)
-		{
-			AsyncTask->ComputedPrecompileData.Reset();
-			if (AsyncTask->ComputedPrecompileDuplicateData.IsValid())
-			{
-				AsyncTask->ComputedPrecompileDuplicateData->ReleaseCompilationCopies();
-				AsyncTask->ComputedPrecompileDuplicateData.Reset();
-			}
-		}
+		bool bHasCompiledJobs = CompileRequest.Resolve(this, ExposedParameters);
+		CompileRequest.Reset();
 
 		// Once compile results have been set, check dependencies found during compile and evaluate whether dependency compile events should be added.
 		if (FNiagaraCVarUtilities::GetShouldEmitMessagesForFailIfNotSet())
@@ -2858,35 +2746,7 @@ bool UNiagaraSystem::QueryCompileComplete(bool bWait, bool bDoPost, bool bDoNotA
 		InitEmitterCompiledData();
 		InitSystemCompiledData();
 
-		// HACK: This is a temporary hack to fix an issue where data interfaces used by modules and dynamic inputs in the
-		// particle update script aren't being shared by the interpolated spawn script when accessed directly.  This works
-		// properly if the data interface is assigned to a named particle parameter and then linked to an input.
-		// TODO: Bind these data interfaces the same way parameter data interfaces are bound.
-		for (auto& AsyncTask : CompileRequest.DDCTasks)
-		{
-			FEmitterCompiledScriptPair& EmitterCompiledScriptPair = AsyncTask->ScriptPair;
-			FVersionedNiagaraEmitter Emitter = EmitterCompiledScriptPair.VersionedEmitter;
-			UNiagaraScript* CompiledScript = EmitterCompiledScriptPair.CompiledScript;
-
-			if (UNiagaraScript::IsEquivalentUsage(CompiledScript->GetUsage(), ENiagaraScriptUsage::ParticleUpdateScript))
-			{
-				UNiagaraScript* SpawnScript = Emitter.GetEmitterData()->SpawnScriptProps.Script;
-				for (const FNiagaraScriptDataInterfaceInfo& UpdateDataInterfaceInfo : CompiledScript->GetCachedDefaultDataInterfaces())
-				{
-					// If the data interface isn't being written to a parameter map then it won't be bound properly so we
-					// assign the update scripts copy of the data interface to the spawn scripts copy by pointer so that they will share
-					// the data interface at runtime and will both be updated in the editor.
-					for (FNiagaraScriptDataInterfaceInfo& SpawnDataInterfaceInfo : SpawnScript->GetCachedDefaultDataInterfaces())
-					{
-						if (SpawnDataInterfaceInfo.RegisteredParameterMapWrite == NAME_None && UpdateDataInterfaceInfo.RegisteredParameterMapWrite == NAME_None && UpdateDataInterfaceInfo.Name == SpawnDataInterfaceInfo.Name)
-						{
-							SpawnDataInterfaceInfo.DataInterface = UpdateDataInterfaceInfo.DataInterface;
-							break;
-						}
-					}
-				}
-			}
-		}
+		CompileRequest.UpdateSpawnDataInterfaces();
 
 		CompileRequest.RootObjects.Empty();
 
@@ -2907,7 +2767,7 @@ bool UNiagaraSystem::QueryCompileComplete(bool bWait, bool bDoPost, bool bDoNotA
 		}
 		else if (CompileRequest.bAllScriptsSynchronized == false)
 		{
-			UE_LOG(LogNiagara, Log, TEXT("Compiling System %s took %f sec, no shader worker used."), *GetFullName(), ElapsedWallTime);
+			UE_LOG(LogNiagara, Verbose, TEXT("Retrieving %s from DDC took %f sec."), *GetFullName(), ElapsedWallTime);
 		}
 
 		ActiveCompilations.RemoveAt(0);
@@ -3121,10 +2981,7 @@ bool UNiagaraSystem::RequestCompile(bool bForce, FNiagaraSystemUpdateContext* Op
 	for (int i = ActiveCompilations.Num() - 1; i >= 1; i--)
 	{
 		FNiagaraSystemCompileRequest& Request = ActiveCompilations[i];
-		for (auto& AsyncTask : Request.DDCTasks)
-		{
-			AsyncTask->AbortTask();
-		}
+		Request.Abort();
 		ActiveCompilations.RemoveAt(i);
 	}	
 
@@ -3138,7 +2995,13 @@ bool UNiagaraSystem::RequestCompile(bool bForce, FNiagaraSystemUpdateContext* Op
 	ActiveCompilation.StartTime = FPlatformTime::Seconds();
 
 	check(SystemSpawnScript->GetLatestSource() == SystemUpdateScript->GetLatestSource());
-	PrepareRapidIterationParametersForCompilation();
+
+	// when not compiling for edit mode we will defer preparing rapid iteration parameters till after we get the
+	// results from the DDC search (FNiagaraPrepareCompileTask::Tick())
+	if (ShouldUseRapidIterationParameters())
+	{
+		PrepareRapidIterationParametersForCompilation();
+	}
 
 	TArray<UNiagaraScript*> ScriptsNeedingCompile;
 	bool bAnyCompiled = false;
@@ -3147,6 +3010,7 @@ bool UNiagaraSystem::RequestCompile(bool bForce, FNiagaraSystemUpdateContext* Op
 		COOK_STAT(Timer.TrackCyclesOnly());
 
 		//Compile all emitters
+		TArray<FNiagaraSystemCompileRequest::FAsyncTaskPtr> Tasks;
 
 		// Pass one... determine if any need to be compiled.
 		{
@@ -3166,14 +3030,13 @@ bool UNiagaraSystem::RequestCompile(bool bForce, FNiagaraSystemUpdateContext* Op
 						Pair.VersionedEmitter = Handle.GetInstance();
 						Pair.CompiledScript = EmitterScript;
 
-						FNiagaraVMExecutableDataId NewID;
-						// we need to compute the vmID here to check later in the ddc task before doing the precompile if anything has changed in the meantime
-						Pair.CompiledScript->ComputeVMCompilationId(NewID, FGuid());
-						Pair.CompileId = NewID;
-
 						TSharedPtr<FNiagaraAsyncCompileTask, ESPMode::ThreadSafe> AsyncTask = MakeShared<FNiagaraAsyncCompileTask, ESPMode::ThreadSafe>(this, EmitterScript->GetPathName(), Pair);
 						if (EmitterScript->IsCompilable() && !EmitterScript->AreScriptAndSourceSynchronized())
 						{
+							// we need to compute the vmID here to check later in the ddc task before doing the precompile if anything has changed in the meantime.
+							// the compilation id was just calculated during AreScriptAndSourceSynchronized() so we reuse it here
+							AsyncTask->AssignInitialCompilationId(EmitterScript->GetComputedVMCompilationId());
+
 							ScriptsNeedingCompile.Add(EmitterScript);
 							bAnyUnsynchronized = true;
 						}
@@ -3182,7 +3045,7 @@ bool UNiagaraSystem::RequestCompile(bool bForce, FNiagaraSystemUpdateContext* Op
 							AsyncTask->CurrentState = ENiagaraCompilationState::Finished;
 						}
 						AsyncTask->UniqueEmitterName = Handle.GetInstance().Emitter->GetUniqueEmitterName();
-						ActiveCompilation.DDCTasks.Add(AsyncTask);
+						Tasks.Add(AsyncTask);
 					}
 				}
 			}
@@ -3195,14 +3058,13 @@ bool UNiagaraSystem::RequestCompile(bool bForce, FNiagaraSystemUpdateContext* Op
 				Pair.VersionedEmitter = FVersionedNiagaraEmitter();
 				Pair.CompiledScript = SystemSpawnScript;
 				
-				FNiagaraVMExecutableDataId NewID;
-				// we need to compute the vmID here to check later in the ddc task before doing the precompile if anything has changed in the meantime
-				Pair.CompiledScript->ComputeVMCompilationId(NewID, FGuid());
-				Pair.CompileId = NewID;
-
 				TSharedPtr<FNiagaraAsyncCompileTask, ESPMode::ThreadSafe> AsyncTask = MakeShared<FNiagaraAsyncCompileTask, ESPMode::ThreadSafe>(this, SystemSpawnScript->GetPathName(), Pair);
 				if (!SystemSpawnScript->AreScriptAndSourceSynchronized())
 				{
+					// we need to compute the vmID here to check later in the ddc task before doing the precompile if anything has changed in the meantime.
+					// the compilation id was just calculated during AreScriptAndSourceSynchronized() so we reuse it here
+					AsyncTask->AssignInitialCompilationId(SystemSpawnScript->GetComputedVMCompilationId());
+
 					ScriptsNeedingCompile.Add(SystemSpawnScript);
 					bAnyCompiled = true;
 				}
@@ -3210,7 +3072,7 @@ bool UNiagaraSystem::RequestCompile(bool bForce, FNiagaraSystemUpdateContext* Op
 				{
 					AsyncTask->CurrentState = ENiagaraCompilationState::Finished;
 				}
-				ActiveCompilation.DDCTasks.Add(AsyncTask);
+				Tasks.Add(AsyncTask);
 			}
 
 			{
@@ -3218,14 +3080,13 @@ bool UNiagaraSystem::RequestCompile(bool bForce, FNiagaraSystemUpdateContext* Op
 				Pair.VersionedEmitter = FVersionedNiagaraEmitter();
 				Pair.CompiledScript = SystemUpdateScript;
 				
-				FNiagaraVMExecutableDataId NewID;
-				// we need to compute the vmID here to check later in the ddc task before doing the precompile if anything has changed in the meantime
-				Pair.CompiledScript->ComputeVMCompilationId(NewID, FGuid());
-				Pair.CompileId = NewID;
-
 				TSharedPtr<FNiagaraAsyncCompileTask, ESPMode::ThreadSafe> AsyncTask = MakeShared<FNiagaraAsyncCompileTask, ESPMode::ThreadSafe>(this, SystemUpdateScript->GetPathName(), Pair);
 				if (!SystemUpdateScript->AreScriptAndSourceSynchronized())
 				{
+					// we need to compute the vmID here to check later in the ddc task before doing the precompile if anything has changed in the meantime.
+					// the compilation id was just calculated during AreScriptAndSourceSynchronized() so we reuse it here
+					AsyncTask->AssignInitialCompilationId(SystemUpdateScript->GetComputedVMCompilationId());
+
 					ScriptsNeedingCompile.Add(SystemUpdateScript);
 					bAnyCompiled = true;
 				}
@@ -3233,48 +3094,12 @@ bool UNiagaraSystem::RequestCompile(bool bForce, FNiagaraSystemUpdateContext* Op
 				{
 					AsyncTask->CurrentState = ENiagaraCompilationState::Finished;
 				}
-				ActiveCompilation.DDCTasks.Add(AsyncTask);
+				Tasks.Add(AsyncTask);
 			}
 		}
 
-
-		// prepare data for any precompile the ddc tasks need to do
-		TSharedPtr<FNiagaraLazyPrecompileReference, ESPMode::ThreadSafe> PrecompileReference = MakeShared<FNiagaraLazyPrecompileReference, ESPMode::ThreadSafe>();
-		PrecompileReference->System = this;
-		PrecompileReference->Scripts = ScriptsNeedingCompile;
-		
-		for (int32 i = 0; i < EmitterHandles.Num(); i++)
-		{
-			FNiagaraEmitterHandle Handle = EmitterHandles[i];
-			if (Handle.GetIsEnabled() && Handle.GetEmitterData())
-			{
-				TArray<UNiagaraScript*> EmitterScripts;
-				Handle.GetEmitterData()->GetScripts(EmitterScripts, false, true);
-				check(EmitterScripts.Num() > 0);
-				for (UNiagaraScript* EmitterScript : EmitterScripts)
-				{
-					PrecompileReference->EmitterScriptIndex.Add(EmitterScript, i);
-				}
-			}
-		}
-
-		// We have previously duplicated all that is needed for compilation, so let's now issue the compile requests!
-		for (UNiagaraScript* CompiledScript : ScriptsNeedingCompile)
 		{
-
-			const auto InPairs = [&CompiledScript](const TSharedPtr<FNiagaraAsyncCompileTask, ESPMode::ThreadSafe>& Other) -> bool
-			{
-				return CompiledScript == Other->GetScriptPair().CompiledScript;
-			};
-
-			TSharedPtr<FNiagaraAsyncCompileTask, ESPMode::ThreadSafe> AsyncTask = *ActiveCompilation.DDCTasks.FindByPredicate(InPairs);
-			AsyncTask->DDCKey = UNiagaraScript::BuildNiagaraDDCKeyString(AsyncTask->GetScriptPair().CompileId, AsyncTask->AssetPath);
-			AsyncTask->PrecompileReference = PrecompileReference;
-			AsyncTask->StartTaskTime = FPlatformTime::Seconds();
-
-			// fire off all the ddc tasks, which will trigger the compilation if the data is not in the ddc
-			UE_LOG(LogNiagara, Verbose, TEXT("Scheduling async get task for %s"), *AsyncTask->AssetPath);
-			AsyncTask->TaskHandle = GetDerivedDataCacheRef().GetAsynchronous(*AsyncTask->DDCKey, AsyncTask->AssetPath);
+			ActiveCompilation.Launch(this, ScriptsNeedingCompile, Tasks);
 		}
 	}
 
diff --git a/Engine/Plugins/FX/Niagara/Source/NiagaraEditor/Private/NiagaraNodeParameterMapGet.cpp b/Engine/Plugins/FX/Niagara/Source/NiagaraEditor/Private/NiagaraNodeParameterMapGet.cpp
--- a/Engine/Plugins/FX/Niagara/Source/NiagaraEditor/Private/NiagaraNodeParameterMapGet.cpp
+++ b/Engine/Plugins/FX/Niagara/Source/NiagaraEditor/Private/NiagaraNodeParameterMapGet.cpp
@@ -586,18 +586,27 @@ void UNiagaraNodeParameterMapGet::GatherExternalDependencyData(ENiagaraScriptUsa
 {
 	// If we are referencing any parameter collections, we need to register them here... might want to speeed this up in the future 
 	// by caching any parameter collections locally.
-	FPinCollectorArray OutputPins;
-	GetOutputPins(OutputPins);
-	const UEdGraphSchema_Niagara* Schema = CastChecked<UEdGraphSchema_Niagara>(GetSchema());
-
-	for (int32 i = 0; i < OutputPins.Num(); i++)
+	for (const UEdGraphPin* Pin : Pins)
 	{
-		if (IsAddPin(OutputPins[i]))
+		if (Pin->Direction == EGPD_Input)
 		{
 			continue;
 		}
 
-		FNiagaraVariable Var = CastChecked<UEdGraphSchema_Niagara>(GetSchema())->PinToNiagaraVariable(OutputPins[i]);
+		if (IsAddPin(Pin))
+		{
+			continue;
+		}
+
+		FNameBuilder PinName(Pin->PinName);
+		if (!PinName.ToView().StartsWith(FNiagaraConstants::ParameterCollectionNamespaceString))
+		{
+			continue;
+		}
+
+		const UEdGraphSchema_Niagara* Schema = CastChecked<UEdGraphSchema_Niagara>(GetSchema());
+
+		FNiagaraVariable Var = Schema->PinToNiagaraVariable(Pin);
 		UNiagaraParameterCollection* Collection = Schema->VariableIsFromParameterCollection(Var);
 		if (Collection)
 		{
diff --git a/Engine/Source/Runtime/Core/Private/UObject/FortniteMainBranchObjectVersion.cpp b/Engine/Source/Runtime/Core/Private/UObject/FortniteMainBranchObjectVersion.cpp
--- a/Engine/Source/Runtime/Core/Private/UObject/FortniteMainBranchObjectVersion.cpp
+++ b/Engine/Source/Runtime/Core/Private/UObject/FortniteMainBranchObjectVersion.cpp
@@ -11,7 +11,7 @@ TMap<FGuid, FGuid> FFortniteMainBranchObjectVersion::GetSystemGuids()
 	SystemGuids.Add(DevGuids.MATERIALSHADERMAP_DERIVEDDATA_VER, FGuid("CE7776E63E9E43FBB8C103EE1AB40B5D"));
 	SystemGuids.Add(DevGuids.NANITE_DERIVEDDATA_VER, FGuid("B62CDCAF66AA45289DDB82E276E7C3E7"));
 	SystemGuids.Add(DevGuids.NIAGARASHADERMAP_DERIVEDDATA_VER, FGuid("07FCB513221642B795CAC18A0D71FA01"));
-	SystemGuids.Add(DevGuids.Niagara_LatestScriptCompileVersion, FGuid("35DC6B707D5E4883AB13750DCD18F021"));
+	SystemGuids.Add(DevGuids.Niagara_LatestScriptCompileVersion, FGuid("C9B73C19EBDD498996662D2E727538F6"));
 	SystemGuids.Add(DevGuids.SkeletalMeshDerivedDataVersion, FGuid("ACF593EAAE354FCBB34CF44F5AA1BFD2"));
 	SystemGuids.Add(DevGuids.STATICMESH_DERIVEDDATA_VER, FGuid("543843B103794E0E9BA4BE80FB602F79"));
 	return SystemGuids;
