Small tweaks to curve data interfaces to reduce the cost of serialization
-CopyToInternal will just copy over the ShaderLUT (rather than regenerate it)
-a vectorized implementation has been added for 4 element curves
-fixed a typo that was regenerating the LUTs on load when any version difference was encountered

diff --git a/Engine/Plugins/FX/Niagara/Source/Niagara/Private/NiagaraDataInterfaceColorCurve.cpp b/Engine/Plugins/FX/Niagara/Source/Niagara/Private/NiagaraDataInterfaceColorCurve.cpp
--- a/Engine/Plugins/FX/Niagara/Source/Niagara/Private/NiagaraDataInterfaceColorCurve.cpp
+++ b/Engine/Plugins/FX/Niagara/Source/Niagara/Private/NiagaraDataInterfaceColorCurve.cpp
@@ -131,14 +131,7 @@ bool UNiagaraDataInterfaceColorCurve::CopyToInternal(UNiagaraDataInterface* Dest
 	DestinationColorCurve->GreenCurve = GreenCurve;
 	DestinationColorCurve->BlueCurve = BlueCurve;
 	DestinationColorCurve->AlphaCurve = AlphaCurve;
-#if WITH_EDITORONLY_DATA
-	DestinationColorCurve->UpdateLUT();
 
-	if (!CompareLUTS(DestinationColorCurve->ShaderLUT))
-	{
-		UE_LOG(LogNiagara, Log, TEXT("CopyToInternal LUT generation is out of sync. Please investigate. %s to %s"), *GetPathName(), *DestinationColorCurve->GetPathName());
-	}
-#endif
 	return true;
 }
 
diff --git a/Engine/Plugins/FX/Niagara/Source/Niagara/Private/NiagaraDataInterfaceCurve.cpp b/Engine/Plugins/FX/Niagara/Source/Niagara/Private/NiagaraDataInterfaceCurve.cpp
--- a/Engine/Plugins/FX/Niagara/Source/Niagara/Private/NiagaraDataInterfaceCurve.cpp
+++ b/Engine/Plugins/FX/Niagara/Source/Niagara/Private/NiagaraDataInterfaceCurve.cpp
@@ -71,13 +71,7 @@ bool UNiagaraDataInterfaceCurve::CopyToInternal(UNiagaraDataInterface* Destinati
 		return false;
 	}
 	CastChecked<UNiagaraDataInterfaceCurve>(Destination)->Curve = Curve;
-#if WITH_EDITORONLY_DATA
-	CastChecked<UNiagaraDataInterfaceCurve>(Destination)->UpdateLUT();
-	if (!CompareLUTS(CastChecked<UNiagaraDataInterfaceCurve>(Destination)->ShaderLUT))
-	{
-		UE_LOG(LogNiagara, Log, TEXT("Post CopyToInternal LUT generation is out of sync. Please investigate. %s"), *GetPathName());
-	}
-#endif
+
 	return true;
 }
 
diff --git a/Engine/Plugins/FX/Niagara/Source/Niagara/Private/NiagaraDataInterfaceCurveBase.cpp b/Engine/Plugins/FX/Niagara/Source/Niagara/Private/NiagaraDataInterfaceCurveBase.cpp
--- a/Engine/Plugins/FX/Niagara/Source/Niagara/Private/NiagaraDataInterfaceCurveBase.cpp
+++ b/Engine/Plugins/FX/Niagara/Source/Niagara/Private/NiagaraDataInterfaceCurveBase.cpp
@@ -46,6 +46,82 @@ END_SHADER_PARAMETER_STRUCT()
 
 static const TCHAR* NDICurve_TemplateShaderFile = TEXT("/Plugin/FX/Niagara/Private/NiagaraDataInterfaceCurveTemplate.ush");
 
+namespace NiagaraDataInterfaceCurveBaseImpl
+{
+
+static bool PassesErrorThreshold(TConstArrayView<float> ShaderLUT, TConstArrayView<float> ResampledLUT, int32 NumElements, float ErrorThreshold)
+{
+	const int32 CurrNumSamples = ShaderLUT.Num() / NumElements;
+	const int32 NewNumSamples = ResampledLUT.Num() / NumElements;
+
+	for (int iSample = 0; iSample < CurrNumSamples; ++iSample)
+	{
+		const float NormalizedSampleTime = float(iSample) / float(CurrNumSamples - 1);
+
+		const float LhsInterp = FMath::Frac(NormalizedSampleTime * CurrNumSamples);
+		const int LhsSampleA = FMath::Min(FMath::FloorToInt(NormalizedSampleTime * CurrNumSamples), CurrNumSamples - 1);
+		const int LhsSampleB = FMath::Min(LhsSampleA + 1, CurrNumSamples - 1);
+
+		const float RhsInterp = FMath::Frac(NormalizedSampleTime * NewNumSamples);
+		const int RhsSampleA = FMath::Min(FMath::FloorToInt(NormalizedSampleTime * NewNumSamples), NewNumSamples - 1);
+		const int RhsSampleB = FMath::Min(RhsSampleA + 1, NewNumSamples - 1);
+
+		for (int iElement = 0; iElement < NumElements; ++iElement)
+		{
+			const float LhsValue = FMath::Lerp(ShaderLUT[LhsSampleA * NumElements + iElement], ShaderLUT[LhsSampleB * NumElements + iElement], LhsInterp);
+			const float RhsValue = FMath::Lerp(ResampledLUT[RhsSampleA * NumElements + iElement], ResampledLUT[RhsSampleB * NumElements + iElement], RhsInterp);
+			const float Error = FMath::Abs(LhsValue - RhsValue);
+			if (Error > ErrorThreshold)
+			{
+				return false;
+			}
+		}
+	}
+
+	return true;
+}
+
+template<int32 ElementStride>
+static bool PassesErrorThresholdVectorized(TConstArrayView<float> ShaderLUT, TConstArrayView<float> ResampledLUT, float ErrorThreshold)
+{
+	const int32 CurrNumSamples = ShaderLUT.Num() / ElementStride;
+	const int32 NewNumSamples = ResampledLUT.Num() / ElementStride;
+
+	const float* ShaderLUTData = ShaderLUT.GetData();
+	const float* ResampledLUTData = ResampledLUT.GetData();
+
+	const VectorRegister4Float Threshold = VectorSetFloat1(ErrorThreshold);
+
+	for (int iSample = 0; iSample < CurrNumSamples; ++iSample)
+	{
+		const float NormalizedSampleTime = float(iSample) / float(CurrNumSamples - 1);
+
+		const int LhsSampleA = FMath::Min(FMath::FloorToInt(NormalizedSampleTime * CurrNumSamples), CurrNumSamples - 1);
+		const int LhsSampleB = FMath::Min(LhsSampleA + 1, CurrNumSamples - 1);
+
+		const int RhsSampleA = FMath::Min(FMath::FloorToInt(NormalizedSampleTime * NewNumSamples), NewNumSamples - 1);
+		const int RhsSampleB = FMath::Min(RhsSampleA + 1, NewNumSamples - 1);
+
+		VectorRegister4Float LhsInterp = VectorSetFloat1(NormalizedSampleTime * CurrNumSamples);
+		LhsInterp = VectorSubtract(LhsInterp, VectorTruncate(LhsInterp));
+
+		VectorRegister4Float RhsInterp = VectorSetFloat1(NormalizedSampleTime * NewNumSamples);
+		RhsInterp = VectorSubtract(RhsInterp, VectorTruncate(RhsInterp));
+
+		VectorRegister4Float LhsValue = VectorLerp(VectorLoad(ShaderLUTData + LhsSampleA * ElementStride), VectorLoad(ShaderLUTData + LhsSampleB * ElementStride), LhsInterp);
+		VectorRegister4Float RhsValue = VectorLerp(VectorLoad(ResampledLUTData + RhsSampleA * ElementStride), VectorLoad(ResampledLUTData + RhsSampleB * ElementStride), RhsInterp);
+		VectorRegister4Float Error = VectorAbs(VectorSubtract(LhsValue, RhsValue));
+		if (VectorAnyGreaterThan(Error, Threshold))
+		{
+			return false;
+		}
+	}
+
+	return true;
+}
+
+}; // NiagaraDataInterfaceCurveBaseImpl
+
 /** Base class for curve data proxy data. */
 struct FNiagaraDataInterfaceProxyCurveBase : public FNiagaraDataInterfaceProxy
 {
@@ -115,7 +191,7 @@ void UNiagaraDataInterfaceCurveBase::PostLoad()
 	{
 		const int32 NiagaraVer = GetLinkerCustomVersion(FNiagaraCustomVersion::GUID);
 
-		if (NiagaraVer < FNiagaraCustomVersion::LatestVersion)
+		if (NiagaraVer < FNiagaraCustomVersion::CurveLUTRegen)
 		{
 			UpdateLUT();
 		}
@@ -253,46 +329,31 @@ void UNiagaraDataInterfaceCurveBase::OptimizeLUT()
 
 	const int CurrNumSamples = ShaderLUT.Num() / NumElements;
 
-	for (int32 NewNumSamples = 1; NewNumSamples < CurrNumSamples; ++NewNumSamples)
+	if (NumElements == 4)
 	{
-		TArray<float> ResampledLUT = BuildLUT(NewNumSamples);
-
-		bool bCanUseLUT = true;
-		for (int iSample = 0; iSample < CurrNumSamples; ++iSample)
+		for (int32 NewNumSamples = 1; NewNumSamples < CurrNumSamples; ++NewNumSamples)
 		{
-			const float NormalizedSampleTime = float(iSample) / float(CurrNumSamples - 1);
+			TArray<float> ResampledLUT = BuildLUT(NewNumSamples);
 
-			const float LhsInterp = FMath::Frac(NormalizedSampleTime * CurrNumSamples);
-			const int LhsSampleA = FMath::Min(FMath::FloorToInt(NormalizedSampleTime * CurrNumSamples), CurrNumSamples - 1);
-			const int LhsSampleB = FMath::Min(LhsSampleA + 1, CurrNumSamples - 1);
-
-			const float RhsInterp = FMath::Frac(NormalizedSampleTime * NewNumSamples);
-			const int RhsSampleA = FMath::Min(FMath::FloorToInt(NormalizedSampleTime * NewNumSamples), NewNumSamples - 1);
-			const int RhsSampleB = FMath::Min(RhsSampleA + 1, NewNumSamples - 1);
-
-			for (int iElement = 0; iElement < NumElements; ++iElement)
+			if (NiagaraDataInterfaceCurveBaseImpl::PassesErrorThresholdVectorized<4>(ShaderLUT, ResampledLUT, ErrorThreshold))
 			{
-				const float LhsValue = FMath::Lerp(ShaderLUT[LhsSampleA * NumElements + iElement], ShaderLUT[LhsSampleB * NumElements + iElement], LhsInterp);
-				const float RhsValue = FMath::Lerp(ResampledLUT[RhsSampleA * NumElements + iElement], ResampledLUT[RhsSampleB * NumElements + iElement], RhsInterp);
-				const float Error = FMath::Abs(LhsValue - RhsValue);
-				if (Error > ErrorThreshold)
-				{
-					bCanUseLUT = false;
-					break;
-				}
+				ShaderLUT = MoveTemp(ResampledLUT);
+				break;
 			}
+		}
+	}
+	else
+	{
+		for (int32 NewNumSamples = 1; NewNumSamples < CurrNumSamples; ++NewNumSamples)
+		{
+			TArray<float> ResampledLUT = BuildLUT(NewNumSamples);
 
-			if (!bCanUseLUT)
+			if (NiagaraDataInterfaceCurveBaseImpl::PassesErrorThreshold(ShaderLUT, ResampledLUT, NumElements, ErrorThreshold))
 			{
+				ShaderLUT = MoveTemp(ResampledLUT);
 				break;
 			}
 		}
-
-		if (bCanUseLUT)
-		{
-			ShaderLUT = ResampledLUT;
-			break;
-		}
 	}
 }
 
diff --git a/Engine/Plugins/FX/Niagara/Source/Niagara/Private/NiagaraDataInterfaceVector2DCurve.cpp b/Engine/Plugins/FX/Niagara/Source/Niagara/Private/NiagaraDataInterfaceVector2DCurve.cpp
--- a/Engine/Plugins/FX/Niagara/Source/Niagara/Private/NiagaraDataInterfaceVector2DCurve.cpp
+++ b/Engine/Plugins/FX/Niagara/Source/Niagara/Private/NiagaraDataInterfaceVector2DCurve.cpp
@@ -115,13 +115,7 @@ bool UNiagaraDataInterfaceVector2DCurve::CopyToInternal(UNiagaraDataInterface* D
 	UNiagaraDataInterfaceVector2DCurve* DestinationVector2DCurve = CastChecked<UNiagaraDataInterfaceVector2DCurve>(Destination);
 	DestinationVector2DCurve->XCurve = XCurve;
 	DestinationVector2DCurve->YCurve = YCurve;
-#if WITH_EDITORONLY_DATA
-	DestinationVector2DCurve->UpdateLUT();
-	if (!CompareLUTS(DestinationVector2DCurve->ShaderLUT))
-	{
-		UE_LOG(LogNiagara, Log, TEXT("Post CopyToInternal LUT generation is out of sync. Please investigate. %s"), *GetPathName());
-	}
-#endif
+
 	return true;
 }
 
diff --git a/Engine/Plugins/FX/Niagara/Source/Niagara/Private/NiagaraDataInterfaceVector4Curve.cpp b/Engine/Plugins/FX/Niagara/Source/Niagara/Private/NiagaraDataInterfaceVector4Curve.cpp
--- a/Engine/Plugins/FX/Niagara/Source/Niagara/Private/NiagaraDataInterfaceVector4Curve.cpp
+++ b/Engine/Plugins/FX/Niagara/Source/Niagara/Private/NiagaraDataInterfaceVector4Curve.cpp
@@ -131,13 +131,6 @@ bool UNiagaraDataInterfaceVector4Curve::CopyToInternal(UNiagaraDataInterface* De
 	DestinationColorCurve->YCurve = YCurve;
 	DestinationColorCurve->ZCurve = ZCurve;
 	DestinationColorCurve->WCurve = WCurve;
-#if WITH_EDITORONLY_DATA
-	DestinationColorCurve->UpdateLUT();
-	if (!CompareLUTS(DestinationColorCurve->ShaderLUT))
-	{
-		UE_LOG(LogNiagara, Log, TEXT("Post CopyToInternal LUT generation is out of sync. Please investigate. %s"), *GetPathName());
-	}
-#endif
 
 	return true;
 }
diff --git a/Engine/Plugins/FX/Niagara/Source/Niagara/Private/NiagaraDataInterfaceVectorCurve.cpp b/Engine/Plugins/FX/Niagara/Source/Niagara/Private/NiagaraDataInterfaceVectorCurve.cpp
--- a/Engine/Plugins/FX/Niagara/Source/Niagara/Private/NiagaraDataInterfaceVectorCurve.cpp
+++ b/Engine/Plugins/FX/Niagara/Source/Niagara/Private/NiagaraDataInterfaceVectorCurve.cpp
@@ -124,14 +124,7 @@ bool UNiagaraDataInterfaceVectorCurve::CopyToInternal(UNiagaraDataInterface* Des
 	DestinationVectorCurve->XCurve = XCurve;
 	DestinationVectorCurve->YCurve = YCurve;
 	DestinationVectorCurve->ZCurve = ZCurve;
-#if WITH_EDITORONLY_DATA
-	DestinationVectorCurve->UpdateLUT();
 
-	if (!CompareLUTS(DestinationVectorCurve->ShaderLUT))
-	{
-		UE_LOG(LogNiagara, Log, TEXT("Post CopyToInternal LUT generation is out of sync. Please investigate. %s"), *GetPathName());
-	}
-#endif
 	return true;
 }
 
