Create Material Usage flag for Heterogeneous Volumes

diff --git a/Engine/Plugins/FX/Niagara/Source/Niagara/Private/NiagaraMeshRendererProperties.cpp b/Engine/Plugins/FX/Niagara/Source/Niagara/Private/NiagaraMeshRendererProperties.cpp
--- a/Engine/Plugins/FX/Niagara/Source/Niagara/Private/NiagaraMeshRendererProperties.cpp
+++ b/Engine/Plugins/FX/Niagara/Source/Niagara/Private/NiagaraMeshRendererProperties.cpp
@@ -10,6 +10,7 @@
 #include "NiagaraSystem.h"
 
 #include "MaterialDomain.h"
+#include "Materials/MaterialRenderProxy.h"
 #include "Engine/StaticMesh.h"
 #include "Materials/MaterialInstanceConstant.h"
 #include "Modules/ModuleManager.h"
@@ -1101,11 +1102,18 @@ void UNiagaraMeshRendererProperties::CheckMaterialUsage()
 			for (int32 SectionIndex = 0; SectionIndex < LODModel.Sections.Num(); SectionIndex++)
 			{
 				const FStaticMeshSection& Section = LODModel.Sections[SectionIndex];
-				UMaterialInterface *Material = MeshProperties.Mesh->GetMaterial(Section.MaterialIndex);
-				if (Material)
+				UMaterialInterface *MaterialInterface = MeshProperties.Mesh->GetMaterial(Section.MaterialIndex);
+				if (MaterialInterface)
 				{
-					FMaterialRenderProxy* MaterialProxy = Material->GetRenderProxy();
-					Material->CheckMaterialUsage(MATUSAGE_NiagaraMeshParticles);
+					const UMaterial* Material = MaterialInterface->GetMaterial();
+					if (UseHeterogeneousVolumes() && Material && Material->MaterialDomain == MD_Volume)
+					{
+						MaterialInterface->CheckMaterialUsage(MATUSAGE_HeterogeneousVolumes);
+					}
+					else
+					{
+						MaterialInterface->CheckMaterialUsage(MATUSAGE_NiagaraMeshParticles);
+					}
 				}
 			}
 		}
diff --git a/Engine/Plugins/FX/Niagara/Source/Niagara/Private/NiagaraRendererMeshes.cpp b/Engine/Plugins/FX/Niagara/Source/Niagara/Private/NiagaraRendererMeshes.cpp
--- a/Engine/Plugins/FX/Niagara/Source/Niagara/Private/NiagaraRendererMeshes.cpp
+++ b/Engine/Plugins/FX/Niagara/Source/Niagara/Private/NiagaraRendererMeshes.cpp
@@ -99,6 +99,7 @@ FNiagaraRendererMeshes::FNiagaraRendererMeshes(ERHIFeatureLevel::Type FeatureLev
 	DistanceCullRangeSquared = FVector2f(0, FLT_MAX);
 	RendererVisibility = Properties->RendererVisibility;
 	bAccurateMotionVectors = Properties->NeedsPreciseMotionVectors();
+	bIsHeterogeneousVolume = Properties->UseHeterogeneousVolumes();
 
 	if (Properties->bEnableCameraDistanceCulling)
 	{
@@ -1467,6 +1468,13 @@ void FNiagaraRendererMeshes::GetDynamicRayTracingInstances(FRayTracingMaterialGa
 				continue;
 			}
 
+			const FMaterialRenderProxy* ConstMaterialProxy = ParticleMeshRenderData.DynamicDataMesh->Materials[RemappedMaterialIndex];
+			const FMaterial& Material = ConstMaterialProxy->GetMaterialWithFallback(FeatureLevel, ConstMaterialProxy);
+			if (Material.IsUsedWithHeterogeneousVolumes())
+			{
+				continue;
+			}
+
 			FMeshBatch MeshBatch;
 			CreateMeshBatchForSection(
 				ParticleMeshRenderData,
@@ -1800,7 +1808,9 @@ FNiagaraDynamicDataBase* FNiagaraRendererMeshes::GenerateDynamicData(const FNiag
 		//In preparation for a material override feature, we pass our material(s) and relevance in via dynamic data.
 		//The renderer ensures we have the correct usage and relevance for materials in BaseMaterials_GT.
 		//Any override feature must also do the same for materials that are set.
-		check(Mat->CheckMaterialUsage_Concurrent(MATUSAGE_NiagaraMeshParticles));
+		check( (Properties->UseHeterogeneousVolumes() && Mat->GetMaterial() && (Mat->GetMaterial()->MaterialDomain == MD_Volume)) ?
+			Mat->CheckMaterialUsage_Concurrent(MATUSAGE_HeterogeneousVolumes) :
+			Mat->CheckMaterialUsage_Concurrent(MATUSAGE_NiagaraMeshParticles));
 		DynamicData->Materials.Add(Mat->GetRenderProxy());
 	}
 
@@ -1828,7 +1838,15 @@ int FNiagaraRendererMeshes::GetDynamicDataSize()const
 
 bool FNiagaraRendererMeshes::IsMaterialValid(const UMaterialInterface* Mat)const
 {
-	return Mat && Mat->CheckMaterialUsage_Concurrent(MATUSAGE_NiagaraMeshParticles);
+	bool bIsMaterialValid = false;
+	if (Mat)
+	{
+		bIsMaterialValid = (bIsHeterogeneousVolume && Mat->GetMaterial() && (Mat->GetMaterial()->MaterialDomain == MD_Volume)) ?
+			Mat->CheckMaterialUsage_Concurrent(MATUSAGE_HeterogeneousVolumes) :
+			Mat->CheckMaterialUsage_Concurrent(MATUSAGE_NiagaraMeshParticles);
+	}
+
+	return bIsMaterialValid;
 }
 
 
diff --git a/Engine/Plugins/FX/Niagara/Source/Niagara/Public/NiagaraRendererMeshes.h b/Engine/Plugins/FX/Niagara/Source/Niagara/Public/NiagaraRendererMeshes.h
--- a/Engine/Plugins/FX/Niagara/Source/Niagara/Public/NiagaraRendererMeshes.h
+++ b/Engine/Plugins/FX/Niagara/Source/Niagara/Public/NiagaraRendererMeshes.h
@@ -168,6 +168,7 @@ private:
 	uint32 bEnableCulling : 1;
 	uint32 bEnableFrustumCulling : 1;
 	uint32 bAccurateMotionVectors : 1;
+	uint32 bIsHeterogeneousVolume : 1;
 
 	uint32 bSubImageBlend : 1;
 	FVector2f SubImageSize = FVector2f::ZeroVector;
diff --git a/Engine/Source/Runtime/Engine/Classes/Materials/Material.h b/Engine/Source/Runtime/Engine/Classes/Materials/Material.h
--- a/Engine/Source/Runtime/Engine/Classes/Materials/Material.h
+++ b/Engine/Source/Runtime/Engine/Classes/Materials/Material.h
@@ -786,6 +786,13 @@ public:
 	UPROPERTY(EditAnywhere, BlueprintReadOnly, Category = Usage)
 	uint32 bUsedWithVolumetricCloud : 1;
 
+	/**
+	 * Indicates that the material and its instances with heterogeneous volumes. Without that flag, it can only be used on volumetric fog.
+	 * This will result in the shaders required to support Heterogeneous Volumes rendering being compiled which will increase shader compile time and memory usage.
+	 */
+	UPROPERTY(EditAnywhere, BlueprintReadOnly, Category = Usage)
+	uint32 bUsedWithHeterogeneousVolumes : 1;
+
 	/** 
 	 * Indicates that the material and its instances can be used with Slate UI and UMG
 	 * This will result in the shaders required to support UI materials being compiled which will increase shader compile time and memory usage.
diff --git a/Engine/Source/Runtime/Engine/Classes/Materials/MaterialInterface.h b/Engine/Source/Runtime/Engine/Classes/Materials/MaterialInterface.h
--- a/Engine/Source/Runtime/Engine/Classes/Materials/MaterialInterface.h
+++ b/Engine/Source/Runtime/Engine/Classes/Materials/MaterialInterface.h
@@ -92,6 +92,7 @@ enum EMaterialUsage : int
 	MATUSAGE_VirtualHeightfieldMesh,
 	MATUSAGE_Nanite,
 	MATUSAGE_VolumetricCloud,
+	MATUSAGE_HeterogeneousVolumes,
 
 	MATUSAGE_MAX,
 };
diff --git a/Engine/Source/Runtime/Engine/Private/Materials/Material.cpp b/Engine/Source/Runtime/Engine/Private/Materials/Material.cpp
--- a/Engine/Source/Runtime/Engine/Private/Materials/Material.cpp
+++ b/Engine/Source/Runtime/Engine/Private/Materials/Material.cpp
@@ -995,6 +995,7 @@ UMaterial::UMaterial(const FObjectInitializer& ObjectInitializer)
 	bIsBlendable = true;
 	bEnableStencilTest = false;
 	bUsedWithVolumetricCloud = false;
+	bUsedWithHeterogeneousVolumes = false;
 
 	bUseEmissiveForDynamicAreaLighting = false;
 	RefractionDepthBias = 0.0f;
@@ -1367,6 +1368,7 @@ bool UMaterial::GetUsageByFlag(EMaterialUsage Usage) const
 		case MATUSAGE_VirtualHeightfieldMesh: UsageValue = bUsedWithVirtualHeightfieldMesh; break;
 		case MATUSAGE_Nanite: UsageValue = bUsedWithNanite; break;
 		case MATUSAGE_VolumetricCloud: UsageValue = bUsedWithVolumetricCloud; break;
+		case MATUSAGE_HeterogeneousVolumes: UsageValue = bUsedWithHeterogeneousVolumes; break;
 		default: UE_LOG(LogMaterial, Fatal,TEXT("Unknown material usage: %u"), (int32)Usage);
 	};
 	return UsageValue;
@@ -1616,6 +1618,10 @@ void UMaterial::SetUsageByFlag(EMaterialUsage Usage, bool NewValue)
 		{
 			bUsedWithVolumetricCloud = NewValue; break;
 		}
+		case MATUSAGE_HeterogeneousVolumes:
+		{
+			bUsedWithHeterogeneousVolumes = NewValue; break;
+		}
 		default: UE_LOG(LogMaterial, Fatal, TEXT("Unknown material usage: %u"), (int32)Usage);
 	};
 #if WITH_EDITOR
@@ -1649,6 +1655,7 @@ FString UMaterial::GetUsageName(EMaterialUsage Usage) const
 		case MATUSAGE_VirtualHeightfieldMesh: UsageName = TEXT("bUsedWithVirtualHeightfieldMesh"); break;
 		case MATUSAGE_Nanite: UsageName = TEXT("bUsedWithNanite"); break;
 		case MATUSAGE_VolumetricCloud: UsageName = TEXT("bUsedWithVolumetricCloud"); break;
+		case MATUSAGE_HeterogeneousVolumes: UsageName = TEXT("bUsedWithHeterogeneousVolumes"); break;
 		default: UE_LOG(LogMaterial, Fatal,TEXT("Unknown material usage: %u"), (int32)Usage);
 	};
 	return UsageName;
diff --git a/Engine/Source/Runtime/Engine/Private/Materials/MaterialShared.cpp b/Engine/Source/Runtime/Engine/Private/Materials/MaterialShared.cpp
--- a/Engine/Source/Runtime/Engine/Private/Materials/MaterialShared.cpp
+++ b/Engine/Source/Runtime/Engine/Private/Materials/MaterialShared.cpp
@@ -1629,6 +1629,11 @@ bool FMaterialResource::IsUsedWithVolumetricCloud() const
 	return Material->bUsedWithVolumetricCloud;
 }
 
+bool FMaterialResource::IsUsedWithHeterogeneousVolumes() const
+{
+	return Material->bUsedWithHeterogeneousVolumes;
+}
+
 bool FMaterialResource::IsTranslucencyAfterDOFEnabled() const 
 { 
 	return Material->TranslucencyPass == MTP_AfterDOF
@@ -5029,6 +5034,7 @@ FMaterialShaderParameters::FMaterialShaderParameters(const FMaterial* InMaterial
 	bShouldDisableDepthTest = InMaterial->ShouldDisableDepthTest();
 	bHasRenderTracePhysicalMaterialOutput = InMaterial->HasRenderTracePhysicalMaterialOutputs();
 	bIsUsedWithVolumetricCloud = InMaterial->IsUsedWithVolumetricCloud();
+	bIsUsedWithHeterogeneousVolumes = InMaterial->IsUsedWithHeterogeneousVolumes();
 	bIsMobileSeparateTranslucencyEnabled = InMaterial->IsMobileSeparateTranslucencyEnabled();
 }
 
diff --git a/Engine/Source/Runtime/Engine/Public/MaterialShared.h b/Engine/Source/Runtime/Engine/Public/MaterialShared.h
--- a/Engine/Source/Runtime/Engine/Public/MaterialShared.h
+++ b/Engine/Source/Runtime/Engine/Public/MaterialShared.h
@@ -1849,6 +1849,7 @@ public:
 	virtual bool IsUsedWithAPEXCloth() const { return false; }
 	virtual bool IsUsedWithNanite() const { return false; }
 	virtual bool IsUsedWithVolumetricCloud() const { return false; }
+	virtual bool IsUsedWithHeterogeneousVolumes() const { return false; }
 	virtual bool IsUsedWithUI() const { return false; }
 	virtual bool IsUsedWithGeometryCache() const { return false; }
 	virtual bool IsUsedWithWater() const { return false; }
@@ -2562,6 +2563,7 @@ public:
 	ENGINE_API virtual bool IsUsedWithVirtualHeightfieldMesh() const override;
 	ENGINE_API virtual bool IsUsedWithNanite() const override;
 	ENGINE_API virtual bool IsUsedWithVolumetricCloud() const override;
+	ENGINE_API virtual bool IsUsedWithHeterogeneousVolumes() const override;
 	ENGINE_API virtual bool IsFullyRough() const override;
 	ENGINE_API virtual bool UseNormalCurvatureToRoughness() const override;
 	ENGINE_API virtual enum EMaterialFloatPrecisionMode GetMaterialFloatPrecisionMode() const override;
@@ -2979,6 +2981,7 @@ struct FMaterialShaderParameters
 			uint64 bShouldDisableDepthTest : 1;
 			uint64 bHasRenderTracePhysicalMaterialOutput : 1;
 			uint64 bIsUsedWithVolumetricCloud : 1;
+			uint64 bIsUsedWithHeterogeneousVolumes : 1;
 			uint64 bIsMobileSeparateTranslucencyEnabled : 1;
 		};
 	};
