Removing all internal /Private/ paths from NiagaraEditor.

diff --git a/Engine/Plugins/FX/Niagara/Source/NiagaraEditor/NiagaraEditor.Build.cs b/Engine/Plugins/FX/Niagara/Source/NiagaraEditor/NiagaraEditor.Build.cs
--- a/Engine/Plugins/FX/Niagara/Source/NiagaraEditor/NiagaraEditor.Build.cs
+++ b/Engine/Plugins/FX/Niagara/Source/NiagaraEditor/NiagaraEditor.Build.cs
@@ -8,12 +8,6 @@ public class NiagaraEditor : ModuleRules
 	public NiagaraEditor(ReadOnlyTargetRules Target) : base(Target)
 	{
 		PrivateIncludePaths.AddRange(new string[] {
-			"NiagaraEditor/Private",
-			"NiagaraEditor/Private/Toolkits",
-			"NiagaraEditor/Private/Widgets",
-			"NiagaraEditor/Private/Sequencer/NiagaraSequence",
-			"NiagaraEditor/Private/ViewModels",
-			"NiagaraEditor/Private/TypeEditorUtilities",
 			Path.Combine(GetModuleDirectory("PropertyEditor"), "Private"),
 		});
 
diff --git a/Engine/Plugins/FX/Niagara/Source/NiagaraEditor/Private/AssetDefinitions/AssetDefinition_NiagaraParameterCollection.cpp b/Engine/Plugins/FX/Niagara/Source/NiagaraEditor/Private/AssetDefinitions/AssetDefinition_NiagaraParameterCollection.cpp
--- a/Engine/Plugins/FX/Niagara/Source/NiagaraEditor/Private/AssetDefinitions/AssetDefinition_NiagaraParameterCollection.cpp
+++ b/Engine/Plugins/FX/Niagara/Source/NiagaraEditor/Private/AssetDefinitions/AssetDefinition_NiagaraParameterCollection.cpp
@@ -4,7 +4,7 @@
 
 #include "ContentBrowserMenuContexts.h"
 #include "NiagaraParameterCollection.h"
-#include "NiagaraParameterCollectionToolkit.h"
+#include "Toolkits/NiagaraParameterCollectionToolkit.h"
 #include "NiagaraParameterCollectionFactoryNew.h"
 #include "NiagaraEditorStyle.h"
 
diff --git a/Engine/Plugins/FX/Niagara/Source/NiagaraEditor/Private/AssetDefinitions/AssetDefinition_NiagaraParameterCollectionInstance.cpp b/Engine/Plugins/FX/Niagara/Source/NiagaraEditor/Private/AssetDefinitions/AssetDefinition_NiagaraParameterCollectionInstance.cpp
--- a/Engine/Plugins/FX/Niagara/Source/NiagaraEditor/Private/AssetDefinitions/AssetDefinition_NiagaraParameterCollectionInstance.cpp
+++ b/Engine/Plugins/FX/Niagara/Source/NiagaraEditor/Private/AssetDefinitions/AssetDefinition_NiagaraParameterCollectionInstance.cpp
@@ -2,7 +2,7 @@
 
 #include "AssetDefinition_NiagaraParameterCollectionInstance.h"
 #include "NiagaraParameterCollection.h"
-#include "NiagaraParameterCollectionToolkit.h"
+#include "Toolkits/NiagaraParameterCollectionToolkit.h"
 #include "NiagaraEditorStyle.h"
 
 #define LOCTEXT_NAMESPACE "UAssetDefinition_NiagaraParameterCollectionInstance"
diff --git a/Engine/Plugins/FX/Niagara/Source/NiagaraEditor/Private/AssetDefinitions/AssetDefinition_NiagaraScript.cpp b/Engine/Plugins/FX/Niagara/Source/NiagaraEditor/Private/AssetDefinitions/AssetDefinition_NiagaraScript.cpp
--- a/Engine/Plugins/FX/Niagara/Source/NiagaraEditor/Private/AssetDefinitions/AssetDefinition_NiagaraScript.cpp
+++ b/Engine/Plugins/FX/Niagara/Source/NiagaraEditor/Private/AssetDefinitions/AssetDefinition_NiagaraScript.cpp
@@ -1,7 +1,7 @@
 // Copyright Epic Games, Inc. All Rights Reserved.
 
 #include "AssetDefinition_NiagaraScript.h"
-#include "NiagaraScriptToolkit.h"
+#include "Toolkits/NiagaraScriptToolkit.h"
 #include "NiagaraEditorStyle.h"
 #include "AssetRegistry/AssetData.h"
 #include "NiagaraEditorUtilities.h"
diff --git a/Engine/Plugins/FX/Niagara/Source/NiagaraEditor/Private/Customizations/NiagaraComponentDetails.cpp b/Engine/Plugins/FX/Niagara/Source/NiagaraEditor/Private/Customizations/NiagaraComponentDetails.cpp
--- a/Engine/Plugins/FX/Niagara/Source/NiagaraEditor/Private/Customizations/NiagaraComponentDetails.cpp
+++ b/Engine/Plugins/FX/Niagara/Source/NiagaraEditor/Private/Customizations/NiagaraComponentDetails.cpp
@@ -42,7 +42,7 @@
 #include "ViewModels/HierarchyEditor/NiagaraUserParametersHierarchyViewModel.h"
 #include "Widgets/Input/SCheckBox.h"
 #include "Widgets/Input/SComboButton.h"
-#include "SNiagaraSystemUserParameters.h"
+#include "Widgets/SNiagaraSystemUserParameters.h"
 #include "SystemToolkitModes/NiagaraSystemToolkitModeBase.h"
 #include "Widgets/Input/SEditableTextBox.h"
 
diff --git a/Engine/Plugins/FX/Niagara/Source/NiagaraEditor/Private/Customizations/NiagaraParameterCollectionCustomNodeBuilder.h b/Engine/Plugins/FX/Niagara/Source/NiagaraEditor/Private/Customizations/NiagaraParameterCollectionCustomNodeBuilder.h
--- a/Engine/Plugins/FX/Niagara/Source/NiagaraEditor/Private/Customizations/NiagaraParameterCollectionCustomNodeBuilder.h
+++ b/Engine/Plugins/FX/Niagara/Source/NiagaraEditor/Private/Customizations/NiagaraParameterCollectionCustomNodeBuilder.h
@@ -6,11 +6,11 @@
 #include "IDetailPropertyRow.h"
 #include "DetailWidgetRow.h"
 #include "Widgets/Text/SInlineEditableTextBlock.h"
-#include "NiagaraParameterCollectionViewModel.h"
-#include "NiagaraScriptInputCollectionViewModel.h"
-#include "NiagaraScriptOutputCollectionViewModel.h"
+#include "ViewModels/NiagaraParameterCollectionViewModel.h"
+#include "ViewModels/NiagaraScriptInputCollectionViewModel.h"
+#include "ViewModels/NiagaraScriptOutputCollectionViewModel.h"
 #include "ViewModels/NiagaraScriptViewModel.h"
-#include "NiagaraParameterViewModel.h"
+#include "ViewModels/NiagaraParameterViewModel.h"
 #include "NiagaraEditorStyle.h"
 #include "IDetailChildrenBuilder.h"
 #include "Framework/MultiBox/MultiBoxBuilder.h"
diff --git a/Engine/Plugins/FX/Niagara/Source/NiagaraEditor/Private/Customizations/NiagaraScriptDetails.cpp b/Engine/Plugins/FX/Niagara/Source/NiagaraEditor/Private/Customizations/NiagaraScriptDetails.cpp
--- a/Engine/Plugins/FX/Niagara/Source/NiagaraEditor/Private/Customizations/NiagaraScriptDetails.cpp
+++ b/Engine/Plugins/FX/Niagara/Source/NiagaraEditor/Private/Customizations/NiagaraScriptDetails.cpp
@@ -11,10 +11,10 @@
 #include "NiagaraEditorStyle.h"
 #include "NiagaraMetaDataCustomNodeBuilder.h"
 #include "NiagaraParameterCollectionCustomNodeBuilder.h"
-#include "NiagaraParameterCollectionViewModel.h"
-#include "NiagaraParameterViewModel.h"
+#include "ViewModels/NiagaraParameterCollectionViewModel.h"
+#include "ViewModels/NiagaraParameterViewModel.h"
 #include "NiagaraScript.h"
-#include "NiagaraScriptGraphViewModel.h"
+#include "ViewModels/NiagaraScriptGraphViewModel.h"
 #include "PropertyCustomizationHelpers.h"
 #include "ScopedTransaction.h"
 #include "Framework/MultiBox/MultiBoxBuilder.h"
diff --git a/Engine/Plugins/FX/Niagara/Source/NiagaraEditor/Private/NiagaraEditorModule.cpp b/Engine/Plugins/FX/Niagara/Source/NiagaraEditor/Private/NiagaraEditorModule.cpp
--- a/Engine/Plugins/FX/Niagara/Source/NiagaraEditor/Private/NiagaraEditorModule.cpp
+++ b/Engine/Plugins/FX/Niagara/Source/NiagaraEditor/Private/NiagaraEditorModule.cpp
@@ -28,7 +28,7 @@
 #include "KismetPins/SGraphPinVector4.h"
 #include "NiagaraNodeAssignment.h"
 #include "SGraphPin.h"
-#include "SNiagaraGraphPinAdd.h"
+#include "Widgets/SNiagaraGraphPinAdd.h"
 #include "SNiagaraGraphPinNumeric.h"
 
 #include "TypeEditorUtilities/NiagaraFloatTypeEditorUtilities.h"
@@ -140,13 +140,13 @@
 #include "ViewModels/Stack/NiagaraStackObjectIssueGenerator.h"
 #include "NiagaraPlatformSet.h"
 #include "NiagaraEffectType.h"
-#include "SNiagaraSystemViewport.h"
+#include "Widgets/SNiagaraSystemViewport.h"
 
 #include "LevelEditor.h"
 #include "LevelEditorOutlinerSettings.h"
 #include "Filters/CustomClassFilterData.h"
 
-#include "SNiagaraDebugger.h"
+#include "Widgets/SNiagaraDebugger.h"
 
 #include "NiagaraDebugVis.h"
 #include "NiagaraPerfBaseline.h"
diff --git a/Engine/Plugins/FX/Niagara/Source/NiagaraEditor/Private/NiagaraEditorUtilities.cpp b/Engine/Plugins/FX/Niagara/Source/NiagaraEditor/Private/NiagaraEditorUtilities.cpp
--- a/Engine/Plugins/FX/Niagara/Source/NiagaraEditor/Private/NiagaraEditorUtilities.cpp
+++ b/Engine/Plugins/FX/Niagara/Source/NiagaraEditor/Private/NiagaraEditorUtilities.cpp
@@ -53,7 +53,7 @@
 #include "Interfaces/IPluginManager.h"
 #include "NiagaraSystem.h"
 #include "NiagaraSystemEditorData.h"
-#include "NiagaraSystemToolkit.h"
+#include "Toolkits/NiagaraSystemToolkit.h"
 #include "ScopedTransaction.h"
 #include "Selection.h"
 #include "UpgradeNiagaraScriptResults.h"
diff --git a/Engine/Plugins/FX/Niagara/Source/NiagaraEditor/Private/NiagaraMessageManager.cpp b/Engine/Plugins/FX/Niagara/Source/NiagaraEditor/Private/NiagaraMessageManager.cpp
--- a/Engine/Plugins/FX/Niagara/Source/NiagaraEditor/Private/NiagaraMessageManager.cpp
+++ b/Engine/Plugins/FX/Niagara/Source/NiagaraEditor/Private/NiagaraMessageManager.cpp
@@ -6,7 +6,7 @@
 #include "AssetRegistry/AssetRegistryModule.h"
 #include "AssetRegistry/IAssetRegistry.h"
 #include "NiagaraMessages.h"
-#include "NiagaraScriptToolkit.h"
+#include "Toolkits/NiagaraScriptToolkit.h"
 #include "NiagaraEditorModule.h"
 #include "Subsystems/AssetEditorSubsystem.h"
 #include "NiagaraEditorUtilities.h"
diff --git a/Engine/Plugins/FX/Niagara/Source/NiagaraEditor/Private/NiagaraNode.cpp b/Engine/Plugins/FX/Niagara/Source/NiagaraEditor/Private/NiagaraNode.cpp
--- a/Engine/Plugins/FX/Niagara/Source/NiagaraEditor/Private/NiagaraNode.cpp
+++ b/Engine/Plugins/FX/Niagara/Source/NiagaraEditor/Private/NiagaraNode.cpp
@@ -5,7 +5,7 @@
 #include "EdGraphSchema_Niagara.h"
 #include "NiagaraHlslTranslator.h"
 #include "GraphEditAction.h"
-#include "SNiagaraGraphNode.h"
+#include "Widgets/SNiagaraGraphNode.h"
 #include "Misc/SecureHash.h"
 #include "ToolMenuSection.h"
 #include "ToolMenu.h"
diff --git a/Engine/Plugins/FX/Niagara/Source/NiagaraEditor/Private/NiagaraNodeConvert.cpp b/Engine/Plugins/FX/Niagara/Source/NiagaraEditor/Private/NiagaraNodeConvert.cpp
--- a/Engine/Plugins/FX/Niagara/Source/NiagaraEditor/Private/NiagaraNodeConvert.cpp
+++ b/Engine/Plugins/FX/Niagara/Source/NiagaraEditor/Private/NiagaraNodeConvert.cpp
@@ -3,7 +3,7 @@
 #include "NiagaraNodeConvert.h"
 #include "EdGraphSchema_Niagara.h"
 #include "NiagaraEditorUtilities.h"
-#include "SNiagaraGraphNodeConvert.h"
+#include "Widgets/SNiagaraGraphNodeConvert.h"
 #include "NiagaraHlslTranslator.h"
 #include "UObject/UnrealType.h"
 #include "Kismet2/StructureEditorUtils.h"
diff --git a/Engine/Plugins/FX/Niagara/Source/NiagaraEditor/Private/NiagaraNodeCustomHlsl.cpp b/Engine/Plugins/FX/Niagara/Source/NiagaraEditor/Private/NiagaraNodeCustomHlsl.cpp
--- a/Engine/Plugins/FX/Niagara/Source/NiagaraEditor/Private/NiagaraNodeCustomHlsl.cpp
+++ b/Engine/Plugins/FX/Niagara/Source/NiagaraEditor/Private/NiagaraNodeCustomHlsl.cpp
@@ -1,7 +1,7 @@
 // Copyright Epic Games, Inc. All Rights Reserved.
 
 #include "NiagaraNodeCustomHlsl.h"
-#include "SNiagaraGraphNodeCustomHlsl.h"
+#include "Widgets/SNiagaraGraphNodeCustomHlsl.h"
 #include "EdGraphSchema_Niagara.h"
 #include "ScopedTransaction.h"
 #include "NiagaraGraph.h"
diff --git a/Engine/Plugins/FX/Niagara/Source/NiagaraEditor/Private/NiagaraNodeFunctionCall.cpp b/Engine/Plugins/FX/Niagara/Source/NiagaraEditor/Private/NiagaraNodeFunctionCall.cpp
--- a/Engine/Plugins/FX/Niagara/Source/NiagaraEditor/Private/NiagaraNodeFunctionCall.cpp
+++ b/Engine/Plugins/FX/Niagara/Source/NiagaraEditor/Private/NiagaraNodeFunctionCall.cpp
@@ -20,7 +20,7 @@
 #include "NiagaraScript.h"
 #include "NiagaraScriptSource.h"
 #include "NiagaraScriptVariable.h"
-#include "SNiagaraGraphNodeFunctionCallWithSpecifiers.h"
+#include "Widgets/SNiagaraGraphNodeFunctionCallWithSpecifiers.h"
 #include "Misc/SecureHash.h"
 #include "Modules/ModuleManager.h"
 #include "UObject/UnrealType.h"
diff --git a/Engine/Plugins/FX/Niagara/Source/NiagaraEditor/Private/NiagaraNodeInput.cpp b/Engine/Plugins/FX/Niagara/Source/NiagaraEditor/Private/NiagaraNodeInput.cpp
--- a/Engine/Plugins/FX/Niagara/Source/NiagaraEditor/Private/NiagaraNodeInput.cpp
+++ b/Engine/Plugins/FX/Niagara/Source/NiagaraEditor/Private/NiagaraNodeInput.cpp
@@ -11,7 +11,7 @@
 #include "EdGraphSchema_Niagara.h"
 #include "NiagaraConstants.h"
 #include "NiagaraDataInterface.h"
-#include "SNiagaraGraphNodeInput.h"
+#include "Widgets/SNiagaraGraphNodeInput.h"
 #include "NiagaraEditorUtilities.h"
 #include "NiagaraEditorModule.h"
 
diff --git a/Engine/Plugins/FX/Niagara/Source/NiagaraEditor/Private/NiagaraNodeOutputTag.cpp b/Engine/Plugins/FX/Niagara/Source/NiagaraEditor/Private/NiagaraNodeOutputTag.cpp
--- a/Engine/Plugins/FX/Niagara/Source/NiagaraEditor/Private/NiagaraNodeOutputTag.cpp
+++ b/Engine/Plugins/FX/Niagara/Source/NiagaraEditor/Private/NiagaraNodeOutputTag.cpp
@@ -2,7 +2,7 @@
 
 #include "NiagaraNodeOutputTag.h"
 #include "NiagaraCustomVersion.h"
-#include "SNiagaraGraphNodeCustomHlsl.h"
+#include "Widgets/SNiagaraGraphNodeCustomHlsl.h"
 #include "EdGraphSchema_Niagara.h"
 #include "ScopedTransaction.h"
 #include "EdGraphSchema_Niagara.h"
diff --git a/Engine/Plugins/FX/Niagara/Source/NiagaraEditor/Private/NiagaraNodeParameterMapBase.cpp b/Engine/Plugins/FX/Niagara/Source/NiagaraEditor/Private/NiagaraNodeParameterMapBase.cpp
--- a/Engine/Plugins/FX/Niagara/Source/NiagaraEditor/Private/NiagaraNodeParameterMapBase.cpp
+++ b/Engine/Plugins/FX/Niagara/Source/NiagaraEditor/Private/NiagaraNodeParameterMapBase.cpp
@@ -5,7 +5,7 @@
 #include "NiagaraEditorUtilities.h"
 #include "NiagaraNodeOutput.h"
 #include "ScopedTransaction.h"
-#include "SNiagaraGraphPinAdd.h"
+#include "Widgets/SNiagaraGraphPinAdd.h"
 #include "NiagaraGraph.h"
 #include "NiagaraComponent.h"
 #include "NiagaraScript.h"
diff --git a/Engine/Plugins/FX/Niagara/Source/NiagaraEditor/Private/NiagaraNodeParameterMapFor.cpp b/Engine/Plugins/FX/Niagara/Source/NiagaraEditor/Private/NiagaraNodeParameterMapFor.cpp
--- a/Engine/Plugins/FX/Niagara/Source/NiagaraEditor/Private/NiagaraNodeParameterMapFor.cpp
+++ b/Engine/Plugins/FX/Niagara/Source/NiagaraEditor/Private/NiagaraNodeParameterMapFor.cpp
@@ -3,7 +3,7 @@
 #include "NiagaraNodeParameterMapFor.h"
 #include "EdGraphSchema_Niagara.h"
 #include "NiagaraEditorUtilities.h"
-#include "SNiagaraGraphNodeConvert.h"
+#include "Widgets/SNiagaraGraphNodeConvert.h"
 #include "NiagaraHlslTranslator.h"
 #include "EdGraph/EdGraphNode.h"
 
diff --git a/Engine/Plugins/FX/Niagara/Source/NiagaraEditor/Private/NiagaraNodeParameterMapGet.cpp b/Engine/Plugins/FX/Niagara/Source/NiagaraEditor/Private/NiagaraNodeParameterMapGet.cpp
--- a/Engine/Plugins/FX/Niagara/Source/NiagaraEditor/Private/NiagaraNodeParameterMapGet.cpp
+++ b/Engine/Plugins/FX/Niagara/Source/NiagaraEditor/Private/NiagaraNodeParameterMapGet.cpp
@@ -6,7 +6,7 @@
 #include "NiagaraGraph.h"
 #include "NiagaraHlslTranslator.h"
 #include "ScopedTransaction.h"
-#include "SNiagaraGraphParameterMapGetNode.h"
+#include "Widgets/SNiagaraGraphParameterMapGetNode.h"
 #include "NiagaraEditorModule.h"
 #include "INiagaraEditorTypeUtilities.h"
 #include "Modules/ModuleManager.h"
diff --git a/Engine/Plugins/FX/Niagara/Source/NiagaraEditor/Private/NiagaraNodeParameterMapSet.cpp b/Engine/Plugins/FX/Niagara/Source/NiagaraEditor/Private/NiagaraNodeParameterMapSet.cpp
--- a/Engine/Plugins/FX/Niagara/Source/NiagaraEditor/Private/NiagaraNodeParameterMapSet.cpp
+++ b/Engine/Plugins/FX/Niagara/Source/NiagaraEditor/Private/NiagaraNodeParameterMapSet.cpp
@@ -11,7 +11,7 @@
 #include "NiagaraGraph.h"
 #include "NiagaraScriptVariable.h"
 #include "ScopedTransaction.h"
-#include "SNiagaraGraphParameterMapSetNode.h"
+#include "Widgets/SNiagaraGraphParameterMapSetNode.h"
 #include "Templates/SharedPointer.h"
 
 #include UE_INLINE_GENERATED_CPP_BY_NAME(NiagaraNodeParameterMapSet)
diff --git a/Engine/Plugins/FX/Niagara/Source/NiagaraEditor/Private/NiagaraNodeWriteDataSet.cpp b/Engine/Plugins/FX/Niagara/Source/NiagaraEditor/Private/NiagaraNodeWriteDataSet.cpp
--- a/Engine/Plugins/FX/Niagara/Source/NiagaraEditor/Private/NiagaraNodeWriteDataSet.cpp
+++ b/Engine/Plugins/FX/Niagara/Source/NiagaraEditor/Private/NiagaraNodeWriteDataSet.cpp
@@ -4,7 +4,7 @@
 #include "UObject/UnrealType.h"
 #include "NiagaraHlslTranslator.h"
 #include "NiagaraEvents.h"
-#include "SNiagaraGraphNodeWriteDataSet.h"
+#include "Widgets/SNiagaraGraphNodeWriteDataSet.h"
 #include "EdGraphSchema_Niagara.h"
 
 #include UE_INLINE_GENERATED_CPP_BY_NAME(NiagaraNodeWriteDataSet)
diff --git a/Engine/Plugins/FX/Niagara/Source/NiagaraEditor/Private/Sequencer/NiagaraSequence/Sections/MovieSceneNiagaraEmitterSection.h b/Engine/Plugins/FX/Niagara/Source/NiagaraEditor/Private/Sequencer/NiagaraSequence/Sections/MovieSceneNiagaraEmitterSection.h
--- a/Engine/Plugins/FX/Niagara/Source/NiagaraEditor/Private/Sequencer/NiagaraSequence/Sections/MovieSceneNiagaraEmitterSection.h
+++ b/Engine/Plugins/FX/Niagara/Source/NiagaraEditor/Private/Sequencer/NiagaraSequence/Sections/MovieSceneNiagaraEmitterSection.h
@@ -2,7 +2,7 @@
 
 #pragma once
 
-#include "MovieSceneNiagaraEmitterTrack.h"
+#include "Sequencer/NiagaraSequence/MovieSceneNiagaraEmitterTrack.h"
 #include "UObject/WeakObjectPtr.h"
 #include "UObject/WeakObjectPtrTemplates.h"
 #include "NiagaraStackFunctionInputBinder.h"
diff --git a/Engine/Plugins/FX/Niagara/Source/NiagaraEditor/Private/Toolkits/NiagaraParameterCollectionToolkit.cpp b/Engine/Plugins/FX/Niagara/Source/NiagaraEditor/Private/Toolkits/NiagaraParameterCollectionToolkit.cpp
--- a/Engine/Plugins/FX/Niagara/Source/NiagaraEditor/Private/Toolkits/NiagaraParameterCollectionToolkit.cpp
+++ b/Engine/Plugins/FX/Niagara/Source/NiagaraEditor/Private/Toolkits/NiagaraParameterCollectionToolkit.cpp
@@ -2,7 +2,7 @@
 
 #include "NiagaraParameterCollectionToolkit.h"
 #include "NiagaraEditorModule.h"
-#include "SNiagaraParameterCollection.h"
+#include "Widgets/SNiagaraParameterCollection.h"
 
 #include "NiagaraObjectSelection.h"
 #include "Widgets/SNiagaraSelectedObjectsDetails.h"
@@ -21,9 +21,9 @@
 #include "Widgets/SBoxPanel.h"
 #include "Widgets/Layout/SBox.h"
 #include "Widgets/Docking/SDockTab.h"
-#include "NiagaraCollectionParameterViewModel.h"
+#include "ViewModels/NiagaraCollectionParameterViewModel.h"
 #include "NiagaraParameterCollection.h"
-#include "NiagaraParameterCollectionAssetViewModel.h"
+#include "ViewModels/NiagaraParameterCollectionAssetViewModel.h"
 
 #include "PropertyEditorModule.h"
 #include "IDetailsView.h"
diff --git a/Engine/Plugins/FX/Niagara/Source/NiagaraEditor/Private/Toolkits/NiagaraScriptToolkit.cpp b/Engine/Plugins/FX/Niagara/Source/NiagaraEditor/Private/Toolkits/NiagaraScriptToolkit.cpp
--- a/Engine/Plugins/FX/Niagara/Source/NiagaraEditor/Private/Toolkits/NiagaraScriptToolkit.cpp
+++ b/Engine/Plugins/FX/Niagara/Source/NiagaraEditor/Private/Toolkits/NiagaraScriptToolkit.cpp
@@ -21,19 +21,19 @@
 #include "NiagaraEditorUtilities.h"
 #include "NiagaraEmitter.h"
 #include "NiagaraGraph.h"
-#include "NiagaraMessageLogViewModel.h"
+#include "ViewModels/NiagaraMessageLogViewModel.h"
 #include "NiagaraNode.h"
 #include "NiagaraNodeOutput.h"
 #include "NiagaraObjectSelection.h"
 #include "NiagaraScript.h"
-#include "NiagaraScriptGraphViewModel.h"
-#include "NiagaraScriptInputCollectionViewModel.h"
+#include "ViewModels/NiagaraScriptGraphViewModel.h"
+#include "ViewModels/NiagaraScriptInputCollectionViewModel.h"
 #include "NiagaraScriptSource.h"
-#include "NiagaraStandaloneScriptViewModel.h"
+#include "ViewModels/NiagaraStandaloneScriptViewModel.h"
 #include "NiagaraVersionMetaData.h"
 #include "PropertyEditorModule.h"
 #include "SGraphActionMenu.h"
-#include "SNiagaraScriptVersionWidget.h"
+#include "Widgets/SNiagaraScriptVersionWidget.h"
 #include "UObject/Linker.h"
 #include "UObject/Package.h"
 #include "ViewModels/NiagaraParameterDefinitionsPanelViewModel.h"
diff --git a/Engine/Plugins/FX/Niagara/Source/NiagaraEditor/Private/Toolkits/NiagaraSimCacheToolkit.cpp b/Engine/Plugins/FX/Niagara/Source/NiagaraEditor/Private/Toolkits/NiagaraSimCacheToolkit.cpp
--- a/Engine/Plugins/FX/Niagara/Source/NiagaraEditor/Private/Toolkits/NiagaraSimCacheToolkit.cpp
+++ b/Engine/Plugins/FX/Niagara/Source/NiagaraEditor/Private/Toolkits/NiagaraSimCacheToolkit.cpp
@@ -6,12 +6,12 @@
 // Widgets & ViewModels
 #include "AdvancedPreviewSceneModule.h"
 #include "Widgets/Docking/SDockTab.h"
-#include "NiagaraSimCacheViewModel.h"
-#include "SNiagaraSimCacheTreeView.h"
-#include "SNiagaraSimCacheOverview.h"
-#include "SNiagaraSimCacheViewport.h"
-#include "SNiagaraSimCacheViewTimeline.h"
-#include "SNiagaraSimCacheViewTransportControls.h"
+#include "ViewModels/NiagaraSimCacheViewModel.h"
+#include "Widgets/SNiagaraSimCacheTreeView.h"
+#include "Widgets/SNiagaraSimCacheOverview.h"
+#include "Widgets/SNiagaraSimCacheViewport.h"
+#include "Widgets/SNiagaraSimCacheViewTimeline.h"
+#include "Widgets/SNiagaraSimCacheViewTransportControls.h"
 
 #define LOCTEXT_NAMESPACE "NiagaraSimCacheToolkit"
 
diff --git a/Engine/Plugins/FX/Niagara/Source/NiagaraEditor/Private/Toolkits/NiagaraSimCacheToolkit.h b/Engine/Plugins/FX/Niagara/Source/NiagaraEditor/Private/Toolkits/NiagaraSimCacheToolkit.h
--- a/Engine/Plugins/FX/Niagara/Source/NiagaraEditor/Private/Toolkits/NiagaraSimCacheToolkit.h
+++ b/Engine/Plugins/FX/Niagara/Source/NiagaraEditor/Private/Toolkits/NiagaraSimCacheToolkit.h
@@ -3,8 +3,8 @@
 #pragma once
 
 #include "CoreMinimal.h"
-#include "SNiagaraSimCacheView.h"
-#include "SNiagaraSimCacheViewport.h"
+#include "Widgets/SNiagaraSimCacheView.h"
+#include "Widgets/SNiagaraSimCacheViewport.h"
 #include "Toolkits/IToolkitHost.h"
 #include "Toolkits/AssetEditorToolkit.h"
 
diff --git a/Engine/Plugins/FX/Niagara/Source/NiagaraEditor/Private/Toolkits/NiagaraSystemToolkit.cpp b/Engine/Plugins/FX/Niagara/Source/NiagaraEditor/Private/Toolkits/NiagaraSystemToolkit.cpp
--- a/Engine/Plugins/FX/Niagara/Source/NiagaraEditor/Private/Toolkits/NiagaraSystemToolkit.cpp
+++ b/Engine/Plugins/FX/Niagara/Source/NiagaraEditor/Private/Toolkits/NiagaraSystemToolkit.cpp
@@ -15,7 +15,7 @@
 #include "ViewModels/NiagaraScratchPadScriptViewModel.h"
 #include "ViewModels/NiagaraScratchPadViewModel.h"
 #include "ViewModels/NiagaraScriptGraphViewModel.h"
-#include "NiagaraSystemScriptViewModel.h"
+#include "ViewModels/NiagaraSystemScriptViewModel.h"
 #include "Widgets/SNiagaraSystemScript.h"
 #include "Widgets/SNiagaraSystemViewport.h"
 #include "Widgets/SNiagaraParameterPanel.h"
@@ -27,8 +27,8 @@
 #include "NiagaraSystemFactoryNew.h"
 #include "NiagaraComponent.h"
 #include "NiagaraSystemEditorData.h"
-#include "NiagaraScriptStatsViewModel.h"
-#include "NiagaraBakerViewModel.h"
+#include "ViewModels/NiagaraScriptStatsViewModel.h"
+#include "ViewModels/NiagaraBakerViewModel.h"
 #include "NiagaraToolkitCommon.h"
 #include "ViewModels/NiagaraSystemEditorDocumentsViewModel.h"
 
@@ -51,7 +51,7 @@
 #include "Misc/MessageDialog.h"
 #include "Misc/TransactionObjectEvent.h"
 #include "Modules/ModuleManager.h"
-#include "NiagaraMessageLogViewModel.h"
+#include "ViewModels/NiagaraMessageLogViewModel.h"
 #include "NiagaraVersionMetaData.h"
 #include "SNiagaraAssetPickerList.h"
 #include "SystemToolkitModes/NiagaraSystemToolkitModeBase.h"
diff --git a/Engine/Plugins/FX/Niagara/Source/NiagaraEditor/Private/Toolkits/SystemToolkitModes/NiagaraSystemToolkitModeBase.cpp b/Engine/Plugins/FX/Niagara/Source/NiagaraEditor/Private/Toolkits/SystemToolkitModes/NiagaraSystemToolkitModeBase.cpp
--- a/Engine/Plugins/FX/Niagara/Source/NiagaraEditor/Private/Toolkits/SystemToolkitModes/NiagaraSystemToolkitModeBase.cpp
+++ b/Engine/Plugins/FX/Niagara/Source/NiagaraEditor/Private/Toolkits/SystemToolkitModes/NiagaraSystemToolkitModeBase.cpp
@@ -7,10 +7,10 @@
 #include "NiagaraEditorModule.h"
 #include "NiagaraEditorSettings.h"
 #include "NiagaraEditorStyle.h"
-#include "NiagaraScriptGraphViewModel.h"
-#include "NiagaraSystemToolkit.h"
+#include "ViewModels/NiagaraScriptGraphViewModel.h"
+#include "Toolkits/NiagaraSystemToolkit.h"
 #include "NiagaraScriptSource.h"
-#include "NiagaraScriptStatsViewModel.h"
+#include "ViewModels/NiagaraScriptStatsViewModel.h"
 #include "ViewModels/NiagaraEmitterHandleViewModel.h"
 #include "ViewModels/NiagaraEmitterViewModel.h"
 #include "ViewModels/NiagaraScriptViewModel.h"
@@ -32,7 +32,7 @@
 #include "ViewModels/NiagaraScratchPadScriptViewModel.h"
 #include "Widgets/SNiagaraSelectedObjectsDetails.h"
 #include "NiagaraObjectSelection.h"
-#include "SNiagaraSystemUserParameters.h"
+#include "Widgets/SNiagaraSystemUserParameters.h"
 #include "Customizations/NiagaraComponentDetails.h"
 #include "Styling/StyleColors.h"
 #include "ViewModels/NiagaraSystemEditorDocumentsViewModel.h"
diff --git a/Engine/Plugins/FX/Niagara/Source/NiagaraEditor/Private/Toolkits/SystemToolkitModes/NiagaraSystemToolkitModeBase.h b/Engine/Plugins/FX/Niagara/Source/NiagaraEditor/Private/Toolkits/SystemToolkitModes/NiagaraSystemToolkitModeBase.h
--- a/Engine/Plugins/FX/Niagara/Source/NiagaraEditor/Private/Toolkits/SystemToolkitModes/NiagaraSystemToolkitModeBase.h
+++ b/Engine/Plugins/FX/Niagara/Source/NiagaraEditor/Private/Toolkits/SystemToolkitModes/NiagaraSystemToolkitModeBase.h
@@ -2,7 +2,7 @@
 
 #pragma once
 
-#include "NiagaraBakerViewModel.h"
+#include "ViewModels/NiagaraBakerViewModel.h"
 #include "WorkflowOrientedApp/ApplicationMode.h"
 #include "Widgets/Docking/SDockTab.h"
 #include "GraphEditAction.h"
diff --git a/Engine/Plugins/FX/Niagara/Source/NiagaraEditor/Private/Toolkits/SystemToolkitModes/NiagaraSystemToolkitMode_Default.cpp b/Engine/Plugins/FX/Niagara/Source/NiagaraEditor/Private/Toolkits/SystemToolkitModes/NiagaraSystemToolkitMode_Default.cpp
--- a/Engine/Plugins/FX/Niagara/Source/NiagaraEditor/Private/Toolkits/SystemToolkitModes/NiagaraSystemToolkitMode_Default.cpp
+++ b/Engine/Plugins/FX/Niagara/Source/NiagaraEditor/Private/Toolkits/SystemToolkitModes/NiagaraSystemToolkitMode_Default.cpp
@@ -3,7 +3,7 @@
 #include "NiagaraSystemToolkitMode_Default.h"
 
 #include "NiagaraEditorCommands.h"
-#include "NiagaraSystemToolkit.h"
+#include "Toolkits/NiagaraSystemToolkit.h"
 #include "NiagaraDebuggerCommon.h"
 #include "NiagaraEditorModule.h"
 #include "NiagaraEditorSettings.h"
diff --git a/Engine/Plugins/FX/Niagara/Source/NiagaraEditor/Private/Toolkits/SystemToolkitModes/NiagaraSystemToolkitMode_Scalability.cpp b/Engine/Plugins/FX/Niagara/Source/NiagaraEditor/Private/Toolkits/SystemToolkitModes/NiagaraSystemToolkitMode_Scalability.cpp
--- a/Engine/Plugins/FX/Niagara/Source/NiagaraEditor/Private/Toolkits/SystemToolkitModes/NiagaraSystemToolkitMode_Scalability.cpp
+++ b/Engine/Plugins/FX/Niagara/Source/NiagaraEditor/Private/Toolkits/SystemToolkitModes/NiagaraSystemToolkitMode_Scalability.cpp
@@ -5,7 +5,7 @@
 #include "IDetailTreeNode.h"
 #include "NiagaraEditorModule.h"
 #include "NiagaraSettings.h"
-#include "NiagaraSystemToolkit.h"
+#include "Toolkits/NiagaraSystemToolkit.h"
 #include "ViewModels/NiagaraSystemViewModel.h"
 #include "ViewModels/Stack/NiagaraStackRendererItem.h"
 #include "Widgets/Scalability/SNiagaraSystemResolvedScalabilitySettings.h"
diff --git a/Engine/Plugins/FX/Niagara/Source/NiagaraEditor/Private/ViewModels/NiagaraSystemScalabilityViewModel.cpp b/Engine/Plugins/FX/Niagara/Source/NiagaraEditor/Private/ViewModels/NiagaraSystemScalabilityViewModel.cpp
--- a/Engine/Plugins/FX/Niagara/Source/NiagaraEditor/Private/ViewModels/NiagaraSystemScalabilityViewModel.cpp
+++ b/Engine/Plugins/FX/Niagara/Source/NiagaraEditor/Private/ViewModels/NiagaraSystemScalabilityViewModel.cpp
@@ -5,7 +5,7 @@
 #include "ViewModels/NiagaraSystemViewModel.h"
 #include "NiagaraSystem.h"
 #include "ISequencerModule.h"
-#include "NiagaraSystemToolkit.h"
+#include "Toolkits/NiagaraSystemToolkit.h"
 #include "DeviceProfiles/DeviceProfile.h"
 #include "DeviceProfiles/DeviceProfileManager.h"
 #include "Modules/ModuleManager.h"
diff --git a/Engine/Plugins/FX/Niagara/Source/NiagaraEditor/Private/ViewModels/NiagaraSystemViewModel.cpp b/Engine/Plugins/FX/Niagara/Source/NiagaraEditor/Private/ViewModels/NiagaraSystemViewModel.cpp
--- a/Engine/Plugins/FX/Niagara/Source/NiagaraEditor/Private/ViewModels/NiagaraSystemViewModel.cpp
+++ b/Engine/Plugins/FX/Niagara/Source/NiagaraEditor/Private/ViewModels/NiagaraSystemViewModel.cpp
@@ -14,7 +14,7 @@
 #include "Modules/ModuleManager.h"
 #include "MovieScene.h"
 #include "MovieSceneFolder.h"
-#include "MovieSceneNiagaraEmitterTrack.h"
+#include "Sequencer/NiagaraSequence/MovieSceneNiagaraEmitterTrack.h"
 #include "NiagaraComponent.h"
 #include "NiagaraDataInterfaceCurveBase.h"
 #include "NiagaraEditorData.h"
@@ -41,7 +41,7 @@
 #include "NiagaraScriptGraphViewModel.h"
 #include "NiagaraScriptSource.h"
 #include "NiagaraScriptVariable.h"
-#include "NiagaraSequence.h"
+#include "Sequencer/NiagaraSequence/NiagaraSequence.h"
 #include "NiagaraSimulationStageBase.h"
 #include "NiagaraSystem.h"
 #include "NiagaraSystemEditorData.h"
diff --git a/Engine/Plugins/FX/Niagara/Source/NiagaraEditor/Private/ViewModels/Stack/NiagaraStackEventScriptItemGroup.cpp b/Engine/Plugins/FX/Niagara/Source/NiagaraEditor/Private/ViewModels/Stack/NiagaraStackEventScriptItemGroup.cpp
--- a/Engine/Plugins/FX/Niagara/Source/NiagaraEditor/Private/ViewModels/Stack/NiagaraStackEventScriptItemGroup.cpp
+++ b/Engine/Plugins/FX/Niagara/Source/NiagaraEditor/Private/ViewModels/Stack/NiagaraStackEventScriptItemGroup.cpp
@@ -5,7 +5,7 @@
 #include "ViewModels/Stack/NiagaraStackObject.h"
 #include "ViewModels/Stack/NiagaraStackModuleItem.h"
 #include "ViewModels/NiagaraScriptViewModel.h"
-#include "NiagaraScriptGraphViewModel.h"
+#include "ViewModels/NiagaraScriptGraphViewModel.h"
 #include "ViewModels/NiagaraEmitterViewModel.h"
 #include "NiagaraGraph.h"
 #include "NiagaraNodeOutput.h"
diff --git a/Engine/Plugins/FX/Niagara/Source/NiagaraEditor/Private/ViewModels/Stack/NiagaraStackFunctionInput.cpp b/Engine/Plugins/FX/Niagara/Source/NiagaraEditor/Private/ViewModels/Stack/NiagaraStackFunctionInput.cpp
--- a/Engine/Plugins/FX/Niagara/Source/NiagaraEditor/Private/ViewModels/Stack/NiagaraStackFunctionInput.cpp
+++ b/Engine/Plugins/FX/Niagara/Source/NiagaraEditor/Private/ViewModels/Stack/NiagaraStackFunctionInput.cpp
@@ -26,12 +26,12 @@
 #include "NiagaraNodeParameterMapSet.h"
 #include "NiagaraRendererProperties.h"
 #include "NiagaraScript.h"
-#include "NiagaraScriptGraphViewModel.h"
+#include "ViewModels/NiagaraScriptGraphViewModel.h"
 #include "NiagaraScriptMergeManager.h"
 #include "NiagaraScriptSource.h"
 #include "NiagaraScriptVariable.h"
 #include "NiagaraSettings.h"
-#include "NiagaraSystemScriptViewModel.h"
+#include "ViewModels/NiagaraSystemScriptViewModel.h"
 #include "ScopedTransaction.h"
 #include "EdGraph/EdGraphPin.h"
 #include "Framework/Notifications/NotificationManager.h"
diff --git a/Engine/Plugins/FX/Niagara/Source/NiagaraEditor/Private/ViewModels/Stack/NiagaraStackParameterStoreEntry.cpp b/Engine/Plugins/FX/Niagara/Source/NiagaraEditor/Private/ViewModels/Stack/NiagaraStackParameterStoreEntry.cpp
--- a/Engine/Plugins/FX/Niagara/Source/NiagaraEditor/Private/ViewModels/Stack/NiagaraStackParameterStoreEntry.cpp
+++ b/Engine/Plugins/FX/Niagara/Source/NiagaraEditor/Private/ViewModels/Stack/NiagaraStackParameterStoreEntry.cpp
@@ -7,10 +7,10 @@
 #include "NiagaraEditorUtilities.h"
 #include "ViewModels/NiagaraSystemViewModel.h"
 #include "NiagaraSystem.h"
-#include "NiagaraSystemScriptViewModel.h"
+#include "ViewModels/NiagaraSystemScriptViewModel.h"
 #include "ViewModels/NiagaraEmitterViewModel.h"
 #include "NiagaraEmitterHandle.h"
-#include "NiagaraScriptGraphViewModel.h"
+#include "ViewModels/NiagaraScriptGraphViewModel.h"
 #include "ViewModels/Stack/NiagaraStackObject.h"
 #include "NiagaraNodeParameterMapGet.h"
 #include "ViewModels/Stack/NiagaraStackGraphUtilities.h"
diff --git a/Engine/Plugins/FX/Niagara/Source/NiagaraEditor/Private/ViewModels/Stack/NiagaraStackRendererItem.cpp b/Engine/Plugins/FX/Niagara/Source/NiagaraEditor/Private/ViewModels/Stack/NiagaraStackRendererItem.cpp
--- a/Engine/Plugins/FX/Niagara/Source/NiagaraEditor/Private/ViewModels/Stack/NiagaraStackRendererItem.cpp
+++ b/Engine/Plugins/FX/Niagara/Source/NiagaraEditor/Private/ViewModels/Stack/NiagaraStackRendererItem.cpp
@@ -12,7 +12,7 @@
 #include "NiagaraNodeAssignment.h"
 #include "NiagaraNodeOutput.h"
 #include "NiagaraConstants.h"
-#include "NiagaraScriptGraphViewModel.h"
+#include "ViewModels/NiagaraScriptGraphViewModel.h"
 #include "NiagaraGraph.h"
 #include "ScopedTransaction.h"
 #include "NiagaraScriptSource.h"
diff --git a/Engine/Plugins/FX/Niagara/Source/NiagaraEditor/Private/ViewModels/Stack/NiagaraStackScriptItemGroup.cpp b/Engine/Plugins/FX/Niagara/Source/NiagaraEditor/Private/ViewModels/Stack/NiagaraStackScriptItemGroup.cpp
--- a/Engine/Plugins/FX/Niagara/Source/NiagaraEditor/Private/ViewModels/Stack/NiagaraStackScriptItemGroup.cpp
+++ b/Engine/Plugins/FX/Niagara/Source/NiagaraEditor/Private/ViewModels/Stack/NiagaraStackScriptItemGroup.cpp
@@ -12,7 +12,7 @@
 #include "NiagaraNodeAssignment.h"
 #include "NiagaraNodeFunctionCall.h"
 #include "NiagaraNodeOutput.h"
-#include "NiagaraScriptGraphViewModel.h"
+#include "ViewModels/NiagaraScriptGraphViewModel.h"
 #include "NiagaraStackEditorData.h"
 #include "ScopedTransaction.h"
 #include "DragAndDrop/AssetDragDropOp.h"
diff --git a/Engine/Plugins/FX/Niagara/Source/NiagaraEditor/Private/ViewModels/Stack/NiagaraStackSimulationStageGroup.cpp b/Engine/Plugins/FX/Niagara/Source/NiagaraEditor/Private/ViewModels/Stack/NiagaraStackSimulationStageGroup.cpp
--- a/Engine/Plugins/FX/Niagara/Source/NiagaraEditor/Private/ViewModels/Stack/NiagaraStackSimulationStageGroup.cpp
+++ b/Engine/Plugins/FX/Niagara/Source/NiagaraEditor/Private/ViewModels/Stack/NiagaraStackSimulationStageGroup.cpp
@@ -6,7 +6,7 @@
 #include "ViewModels/Stack/NiagaraStackObject.h"
 #include "ViewModels/Stack/NiagaraStackViewModel.h"
 #include "ViewModels/NiagaraScriptViewModel.h"
-#include "NiagaraScriptGraphViewModel.h"
+#include "ViewModels/NiagaraScriptGraphViewModel.h"
 #include "ViewModels/NiagaraEmitterViewModel.h"
 #include "NiagaraGraph.h"
 #include "NiagaraNodeOutput.h"
diff --git a/Engine/Plugins/FX/Niagara/Source/NiagaraEditor/Private/ViewModels/Stack/NiagaraStackSystemSettingsGroup.cpp b/Engine/Plugins/FX/Niagara/Source/NiagaraEditor/Private/ViewModels/Stack/NiagaraStackSystemSettingsGroup.cpp
--- a/Engine/Plugins/FX/Niagara/Source/NiagaraEditor/Private/ViewModels/Stack/NiagaraStackSystemSettingsGroup.cpp
+++ b/Engine/Plugins/FX/Niagara/Source/NiagaraEditor/Private/ViewModels/Stack/NiagaraStackSystemSettingsGroup.cpp
@@ -2,7 +2,7 @@
 
 #include "ViewModels/Stack/NiagaraStackSystemSettingsGroup.h"
 #include "ViewModels/Stack/NiagaraStackModuleItem.h"
-#include "NiagaraScriptGraphViewModel.h"
+#include "ViewModels/NiagaraScriptGraphViewModel.h"
 #include "NiagaraNodeFunctionCall.h"
 #include "ViewModels/Stack/NiagaraStackErrorItem.h"
 #include "ViewModels/Stack/NiagaraStackParameterStoreEntry.h"
diff --git a/Engine/Plugins/FX/Niagara/Source/NiagaraEditor/Private/ViewModels/Stack/NiagaraStackViewModel.cpp b/Engine/Plugins/FX/Niagara/Source/NiagaraEditor/Private/ViewModels/Stack/NiagaraStackViewModel.cpp
--- a/Engine/Plugins/FX/Niagara/Source/NiagaraEditor/Private/ViewModels/Stack/NiagaraStackViewModel.cpp
+++ b/Engine/Plugins/FX/Niagara/Source/NiagaraEditor/Private/ViewModels/Stack/NiagaraStackViewModel.cpp
@@ -3,7 +3,7 @@
 #include "ViewModels/Stack/NiagaraStackViewModel.h"
 
 #include "NiagaraEmitterEditorData.h"
-#include "NiagaraScriptGraphViewModel.h"
+#include "ViewModels/NiagaraScriptGraphViewModel.h"
 #include "NiagaraStackEditorData.h"
 #include "NiagaraSystem.h"
 #include "NiagaraSystemEditorData.h"
diff --git a/Engine/Plugins/FX/Niagara/Source/NiagaraEditor/Private/Widgets/SNiagaraConvertPinSocket.cpp b/Engine/Plugins/FX/Niagara/Source/NiagaraEditor/Private/Widgets/SNiagaraConvertPinSocket.cpp
--- a/Engine/Plugins/FX/Niagara/Source/NiagaraEditor/Private/Widgets/SNiagaraConvertPinSocket.cpp
+++ b/Engine/Plugins/FX/Niagara/Source/NiagaraEditor/Private/Widgets/SNiagaraConvertPinSocket.cpp
@@ -1,9 +1,9 @@
 // Copyright Epic Games, Inc. All Rights Reserved.
 
 #include "SNiagaraConvertPinSocket.h"
-#include "NiagaraConvertNodeViewModel.h"
-#include "NiagaraConvertPinViewModel.h"
-#include "NiagaraConvertPinSocketViewModel.h"
+#include "ViewModels/NiagaraConvertNodeViewModel.h"
+#include "ViewModels/NiagaraConvertPinViewModel.h"
+#include "ViewModels/NiagaraConvertPinSocketViewModel.h"
 #include "ConnectionDrawingPolicy.h"
 #include "Styling/CoreStyle.h"
 #include "Widgets/Input/SButton.h"
diff --git a/Engine/Plugins/FX/Niagara/Source/NiagaraEditor/Private/Widgets/SNiagaraGeneratedCodeView.cpp b/Engine/Plugins/FX/Niagara/Source/NiagaraEditor/Private/Widgets/SNiagaraGeneratedCodeView.cpp
--- a/Engine/Plugins/FX/Niagara/Source/NiagaraEditor/Private/Widgets/SNiagaraGeneratedCodeView.cpp
+++ b/Engine/Plugins/FX/Niagara/Source/NiagaraEditor/Private/Widgets/SNiagaraGeneratedCodeView.cpp
@@ -13,7 +13,7 @@
 #include "NiagaraEmitterHandle.h"
 #include "NiagaraEmitter.h"
 #include "NiagaraScript.h"
-#include "NiagaraSystemScriptViewModel.h"
+#include "ViewModels/NiagaraSystemScriptViewModel.h"
 #include "Styling/AppStyle.h"
 #include "Widgets/Layout/SScrollBox.h"
 #include "UObject/Class.h"
diff --git a/Engine/Plugins/FX/Niagara/Source/NiagaraEditor/Private/Widgets/SNiagaraGraphNodeConvert.cpp b/Engine/Plugins/FX/Niagara/Source/NiagaraEditor/Private/Widgets/SNiagaraGraphNodeConvert.cpp
--- a/Engine/Plugins/FX/Niagara/Source/NiagaraEditor/Private/Widgets/SNiagaraGraphNodeConvert.cpp
+++ b/Engine/Plugins/FX/Niagara/Source/NiagaraEditor/Private/Widgets/SNiagaraGraphNodeConvert.cpp
@@ -2,9 +2,9 @@
 
 #include "SNiagaraGraphNodeConvert.h"
 #include "NiagaraNodeConvert.h"
-#include "NiagaraConvertNodeViewModel.h"
-#include "NiagaraConvertPinViewModel.h"
-#include "NiagaraConvertPinSocketViewModel.h"
+#include "ViewModels/NiagaraConvertNodeViewModel.h"
+#include "ViewModels/NiagaraConvertPinViewModel.h"
+#include "ViewModels/NiagaraConvertPinSocketViewModel.h"
 #include "SNiagaraConvertPinSocket.h"
 #include "GraphEditorSettings.h"
 #include "Rendering/DrawElements.h"
diff --git a/Engine/Plugins/FX/Niagara/Source/NiagaraEditor/Private/Widgets/SNiagaraParameterCollection.cpp b/Engine/Plugins/FX/Niagara/Source/NiagaraEditor/Private/Widgets/SNiagaraParameterCollection.cpp
--- a/Engine/Plugins/FX/Niagara/Source/NiagaraEditor/Private/Widgets/SNiagaraParameterCollection.cpp
+++ b/Engine/Plugins/FX/Niagara/Source/NiagaraEditor/Private/Widgets/SNiagaraParameterCollection.cpp
@@ -2,8 +2,8 @@
 
 #include "SNiagaraParameterCollection.h"
 #include "NiagaraEditorModule.h"
-#include "NiagaraParameterCollectionViewModel.h"
-#include "NiagaraParameterViewModel.h"
+#include "ViewModels/NiagaraParameterCollectionViewModel.h"
+#include "ViewModels/NiagaraParameterViewModel.h"
 #include "NiagaraTypes.h"
 #include "INiagaraEditorTypeUtilities.h"
 #include "SNiagaraParameterEditor.h"
diff --git a/Engine/Plugins/FX/Niagara/Source/NiagaraEditor/Private/Widgets/SNiagaraScriptGraph.cpp b/Engine/Plugins/FX/Niagara/Source/NiagaraEditor/Private/Widgets/SNiagaraScriptGraph.cpp
--- a/Engine/Plugins/FX/Niagara/Source/NiagaraEditor/Private/Widgets/SNiagaraScriptGraph.cpp
+++ b/Engine/Plugins/FX/Niagara/Source/NiagaraEditor/Private/Widgets/SNiagaraScriptGraph.cpp
@@ -1,7 +1,7 @@
 // Copyright Epic Games, Inc. All Rights Reserved.
 
 #include "Widgets/SNiagaraScriptGraph.h"
-#include "NiagaraScriptGraphViewModel.h"
+#include "ViewModels/NiagaraScriptGraphViewModel.h"
 #include "NiagaraEditorUtilities.h"
 #include "NiagaraGraph.h"
 #include "NiagaraNode.h"
diff --git a/Engine/Plugins/FX/Niagara/Source/NiagaraEditor/Private/Widgets/SNiagaraSimCacheViewTransportControls.cpp b/Engine/Plugins/FX/Niagara/Source/NiagaraEditor/Private/Widgets/SNiagaraSimCacheViewTransportControls.cpp
--- a/Engine/Plugins/FX/Niagara/Source/NiagaraEditor/Private/Widgets/SNiagaraSimCacheViewTransportControls.cpp
+++ b/Engine/Plugins/FX/Niagara/Source/NiagaraEditor/Private/Widgets/SNiagaraSimCacheViewTransportControls.cpp
@@ -7,7 +7,7 @@
 #include "EditorWidgetsModule.h"
 #include "SlateOptMacros.h"
 #include "ITransportControl.h"
-#include "NiagaraSimCacheViewModel.h"
+#include "ViewModels/NiagaraSimCacheViewModel.h"
 #include "Widgets/Input/SNumericEntryBox.h"
 
 void SNiagaraSimCacheViewTransportControls::Construct(const FArguments& InArgs)
diff --git a/Engine/Plugins/FX/Niagara/Source/NiagaraEditor/Private/Widgets/SNiagaraSystemScript.cpp b/Engine/Plugins/FX/Niagara/Source/NiagaraEditor/Private/Widgets/SNiagaraSystemScript.cpp
--- a/Engine/Plugins/FX/Niagara/Source/NiagaraEditor/Private/Widgets/SNiagaraSystemScript.cpp
+++ b/Engine/Plugins/FX/Niagara/Source/NiagaraEditor/Private/Widgets/SNiagaraSystemScript.cpp
@@ -2,9 +2,9 @@
 
 #include "SNiagaraSystemScript.h"
 #include "ViewModels/NiagaraSystemViewModel.h"
-#include "NiagaraSystemScriptViewModel.h"
+#include "ViewModels/NiagaraSystemScriptViewModel.h"
 #include "ViewModels/NiagaraScriptViewModel.h"
-#include "NiagaraScriptInputCollectionViewModel.h"
+#include "ViewModels/NiagaraScriptInputCollectionViewModel.h"
 #include "SNiagaraParameterCollection.h"
 #include "Widgets/SNiagaraScriptGraph.h"
 
diff --git a/Engine/Plugins/FX/Niagara/Source/NiagaraEditor/Private/Widgets/SNiagaraSystemUserParameters.cpp b/Engine/Plugins/FX/Niagara/Source/NiagaraEditor/Private/Widgets/SNiagaraSystemUserParameters.cpp
--- a/Engine/Plugins/FX/Niagara/Source/NiagaraEditor/Private/Widgets/SNiagaraSystemUserParameters.cpp
+++ b/Engine/Plugins/FX/Niagara/Source/NiagaraEditor/Private/Widgets/SNiagaraSystemUserParameters.cpp
@@ -7,7 +7,7 @@
 #include "NiagaraEditorUtilities.h"
 #include "Customizations/NiagaraComponentDetails.h"
 #include "Styling/StyleColors.h"
-#include "SystemToolkitModes/NiagaraSystemToolkitModeBase.h"
+#include "Toolkits/SystemToolkitModes/NiagaraSystemToolkitModeBase.h"
 #include "Modules/ModuleManager.h"
 #include "PropertyEditorModule.h"
 #include "ViewModels/NiagaraParameterPanelViewModel.h"
diff --git a/Engine/Plugins/FX/Niagara/Source/NiagaraEditor/Public/Widgets/SNiagaraGraphActionMenu.h b/Engine/Plugins/FX/Niagara/Source/NiagaraEditor/Public/Widgets/SNiagaraGraphActionMenu.h
--- a/Engine/Plugins/FX/Niagara/Source/NiagaraEditor/Public/Widgets/SNiagaraGraphActionMenu.h
+++ b/Engine/Plugins/FX/Niagara/Source/NiagaraEditor/Public/Widgets/SNiagaraGraphActionMenu.h
@@ -6,7 +6,7 @@
 #include "Widgets/DeclarativeSyntaxSupport.h"
 #include "GraphEditor.h"
 #include "NiagaraActions.h"
-#include "NiagaraScriptGraphViewModel.h"
+#include "ViewModels/NiagaraScriptGraphViewModel.h"
 #include "SItemSelector.h"
 #include "SNiagaraFilterBox.h"
 #include "Widgets/Layout/SBorder.h"
