Fixing compile errors in NiagaraEditor after remove include paths.

diff --git a/Engine/Plugins/FX/Niagara/Source/NiagaraEditor/Private/AssetDefinitions/AssetDefinition_NiagaraEmitter.cpp b/Engine/Plugins/FX/Niagara/Source/NiagaraEditor/Private/AssetDefinitions/AssetDefinition_NiagaraEmitter.cpp
--- a/Engine/Plugins/FX/Niagara/Source/NiagaraEditor/Private/AssetDefinitions/AssetDefinition_NiagaraEmitter.cpp
+++ b/Engine/Plugins/FX/Niagara/Source/NiagaraEditor/Private/AssetDefinitions/AssetDefinition_NiagaraEmitter.cpp
@@ -7,7 +7,7 @@
 #include "NiagaraEmitter.h"
 #include "NiagaraScriptSourceBase.h"
 #include "NiagaraSystem.h"
-#include "NiagaraSystemToolkit.h"
+#include "Toolkits/NiagaraSystemToolkit.h"
 #include "NiagaraEditorStyle.h"
 #include "NiagaraSystemFactoryNew.h"
 #include "ViewModels/NiagaraSystemViewModel.h"
diff --git a/Engine/Plugins/FX/Niagara/Source/NiagaraEditor/Private/AssetDefinitions/AssetDefinition_NiagaraSimCache.cpp b/Engine/Plugins/FX/Niagara/Source/NiagaraEditor/Private/AssetDefinitions/AssetDefinition_NiagaraSimCache.cpp
--- a/Engine/Plugins/FX/Niagara/Source/NiagaraEditor/Private/AssetDefinitions/AssetDefinition_NiagaraSimCache.cpp
+++ b/Engine/Plugins/FX/Niagara/Source/NiagaraEditor/Private/AssetDefinitions/AssetDefinition_NiagaraSimCache.cpp
@@ -3,7 +3,7 @@
 #include "AssetDefinition_NiagaraSimCache.h"
 
 #include "NiagaraEditorStyle.h"
-#include "NiagaraSimCacheToolkit.h"
+#include "Toolkits/NiagaraSimCacheToolkit.h"
 
 FLinearColor UAssetDefinition_NiagaraSimCache::GetAssetColor() const
 {
diff --git a/Engine/Plugins/FX/Niagara/Source/NiagaraEditor/Private/AssetDefinitions/AssetDefinition_NiagaraSystem.cpp b/Engine/Plugins/FX/Niagara/Source/NiagaraEditor/Private/AssetDefinitions/AssetDefinition_NiagaraSystem.cpp
--- a/Engine/Plugins/FX/Niagara/Source/NiagaraEditor/Private/AssetDefinitions/AssetDefinition_NiagaraSystem.cpp
+++ b/Engine/Plugins/FX/Niagara/Source/NiagaraEditor/Private/AssetDefinitions/AssetDefinition_NiagaraSystem.cpp
@@ -2,7 +2,7 @@
 
 #include "AssetDefinition_NiagaraSystem.h"
 #include "NiagaraSystem.h"
-#include "NiagaraSystemToolkit.h"
+#include "Toolkits/NiagaraSystemToolkit.h"
 #include "NiagaraEditorStyle.h"
 
 FLinearColor UAssetDefinition_NiagaraSystem::GetAssetColor() const
diff --git a/Engine/Plugins/FX/Niagara/Source/NiagaraEditor/Private/Customizations/NiagaraComponentDetails.cpp b/Engine/Plugins/FX/Niagara/Source/NiagaraEditor/Private/Customizations/NiagaraComponentDetails.cpp
--- a/Engine/Plugins/FX/Niagara/Source/NiagaraEditor/Private/Customizations/NiagaraComponentDetails.cpp
+++ b/Engine/Plugins/FX/Niagara/Source/NiagaraEditor/Private/Customizations/NiagaraComponentDetails.cpp
@@ -43,7 +43,7 @@
 #include "Widgets/Input/SCheckBox.h"
 #include "Widgets/Input/SComboButton.h"
 #include "Widgets/SNiagaraSystemUserParameters.h"
-#include "SystemToolkitModes/NiagaraSystemToolkitModeBase.h"
+#include "Toolkits/SystemToolkitModes/NiagaraSystemToolkitModeBase.h"
 #include "Widgets/Input/SEditableTextBox.h"
 
 #define LOCTEXT_NAMESPACE "NiagaraComponentDetails"
diff --git a/Engine/Plugins/FX/Niagara/Source/NiagaraEditor/Private/NiagaraEditorUtilities.cpp b/Engine/Plugins/FX/Niagara/Source/NiagaraEditor/Private/NiagaraEditorUtilities.cpp
--- a/Engine/Plugins/FX/Niagara/Source/NiagaraEditor/Private/NiagaraEditorUtilities.cpp
+++ b/Engine/Plugins/FX/Niagara/Source/NiagaraEditor/Private/NiagaraEditorUtilities.cpp
@@ -29,7 +29,7 @@
 #include "NiagaraParameterMapHistory.h"
 #include "NiagaraParameterDefinitions.h"
 #include "NiagaraScript.h"
-#include "NiagaraScriptGraphViewModel.h"
+#include "ViewModels/NiagaraScriptGraphViewModel.h"
 #include "NiagaraScriptMergeManager.h"
 #include "NiagaraScriptSource.h"
 #include "NiagaraSettings.h"
diff --git a/Engine/Plugins/FX/Niagara/Source/NiagaraEditor/Private/ViewModels/NiagaraSimCacheViewModel.cpp b/Engine/Plugins/FX/Niagara/Source/NiagaraEditor/Private/ViewModels/NiagaraSimCacheViewModel.cpp
--- a/Engine/Plugins/FX/Niagara/Source/NiagaraEditor/Private/ViewModels/NiagaraSimCacheViewModel.cpp
+++ b/Engine/Plugins/FX/Niagara/Source/NiagaraEditor/Private/ViewModels/NiagaraSimCacheViewModel.cpp
@@ -6,7 +6,7 @@
 #include "NiagaraEditorCommon.h"
 #include "NiagaraDebuggerCommon.h"
 #include "Serialization/MemoryReader.h"
-#include "SNiagaraSimCacheTreeView.h"
+#include "Widgets/SNiagaraSimCacheTreeView.h"
 #include "Serialization/ObjectAndNameAsStringProxyArchive.h"
 #include "UObject/Package.h"
 
diff --git a/Engine/Plugins/FX/Niagara/Source/NiagaraEditor/Private/ViewModels/NiagaraSimCacheViewModel.h b/Engine/Plugins/FX/Niagara/Source/NiagaraEditor/Private/ViewModels/NiagaraSimCacheViewModel.h
--- a/Engine/Plugins/FX/Niagara/Source/NiagaraEditor/Private/ViewModels/NiagaraSimCacheViewModel.h
+++ b/Engine/Plugins/FX/Niagara/Source/NiagaraEditor/Private/ViewModels/NiagaraSimCacheViewModel.h
@@ -4,7 +4,7 @@
 
 #include "CoreMinimal.h"
 #include "NiagaraSimCache.h"
-#include "SNiagaraSimCacheOverview.h"
+#include "Widgets/SNiagaraSimCacheOverview.h"
 
 class SNiagaraSimCacheTreeView;
 struct FNiagaraSimCacheComponentTreeItem;
diff --git a/Engine/Plugins/FX/Niagara/Source/NiagaraEditor/Private/ViewModels/Stack/NiagaraStackGraphUtilities.cpp b/Engine/Plugins/FX/Niagara/Source/NiagaraEditor/Private/ViewModels/Stack/NiagaraStackGraphUtilities.cpp
--- a/Engine/Plugins/FX/Niagara/Source/NiagaraEditor/Private/ViewModels/Stack/NiagaraStackGraphUtilities.cpp
+++ b/Engine/Plugins/FX/Niagara/Source/NiagaraEditor/Private/ViewModels/Stack/NiagaraStackGraphUtilities.cpp
@@ -42,7 +42,7 @@
 #include "INiagaraEditorTypeUtilities.h"
 #include "NiagaraEditorData.h"
 #include "NiagaraGraphDataCache.h"
-#include "NiagaraScriptGraphViewModel.h"
+#include "ViewModels/NiagaraScriptGraphViewModel.h"
 #include "NiagaraSettings.h"
 #include "ViewModels/HierarchyEditor/NiagaraHierarchyViewModelBase.h"
 
diff --git a/Engine/Plugins/FX/Niagara/Source/NiagaraEditor/Private/Widgets/SNiagaraSimCacheTreeView.cpp b/Engine/Plugins/FX/Niagara/Source/NiagaraEditor/Private/Widgets/SNiagaraSimCacheTreeView.cpp
--- a/Engine/Plugins/FX/Niagara/Source/NiagaraEditor/Private/Widgets/SNiagaraSimCacheTreeView.cpp
+++ b/Engine/Plugins/FX/Niagara/Source/NiagaraEditor/Private/Widgets/SNiagaraSimCacheTreeView.cpp
@@ -2,7 +2,7 @@
 
 #include "SNiagaraSimCacheTreeView.h"
 
-#include "NiagaraSimCacheViewModel.h"
+#include "ViewModels/NiagaraSimCacheViewModel.h"
 #include "Widgets/Views/STreeView.h"
 #include "Widgets/Input/SButton.h"
 #include "Widgets/Input/SCheckBox.h"
diff --git a/Engine/Plugins/FX/Niagara/Source/NiagaraEditor/Private/Widgets/SNiagaraSimCacheTreeView.h b/Engine/Plugins/FX/Niagara/Source/NiagaraEditor/Private/Widgets/SNiagaraSimCacheTreeView.h
--- a/Engine/Plugins/FX/Niagara/Source/NiagaraEditor/Private/Widgets/SNiagaraSimCacheTreeView.h
+++ b/Engine/Plugins/FX/Niagara/Source/NiagaraEditor/Private/Widgets/SNiagaraSimCacheTreeView.h
@@ -2,7 +2,7 @@
 
 #pragma once
 
-#include "NiagaraSimCacheViewModel.h"
+#include "ViewModels/NiagaraSimCacheViewModel.h"
 
 #include "NiagaraEditorStyle.h"
 #include "Widgets/SNiagaraPinTypeSelector.h"
diff --git a/Engine/Plugins/FX/Niagara/Source/NiagaraEditor/Private/Widgets/SNiagaraSimCacheViewport.cpp b/Engine/Plugins/FX/Niagara/Source/NiagaraEditor/Private/Widgets/SNiagaraSimCacheViewport.cpp
--- a/Engine/Plugins/FX/Niagara/Source/NiagaraEditor/Private/Widgets/SNiagaraSimCacheViewport.cpp
+++ b/Engine/Plugins/FX/Niagara/Source/NiagaraEditor/Private/Widgets/SNiagaraSimCacheViewport.cpp
@@ -3,7 +3,7 @@
 #include "SNiagaraSimCacheViewport.h"
 
 #include "AdvancedPreviewScene.h"
-#include "NiagaraSimCacheViewModel.h"
+#include "ViewModels/NiagaraSimCacheViewModel.h"
 #include "SEditorViewportToolBarMenu.h"
 #include "SNiagaraSimCacheViewportToolbar.h"
 #include "NiagaraComponent.h"
