- Add Engine.System.TickCount to debug hud engine variables & clean up handling

#rnx

diff --git a/Engine/Plugins/FX/Niagara/Source/Niagara/Private/NiagaraDebugHud.cpp b/Engine/Plugins/FX/Niagara/Source/Niagara/Private/NiagaraDebugHud.cpp
--- a/Engine/Plugins/FX/Niagara/Source/Niagara/Private/NiagaraDebugHud.cpp
+++ b/Engine/Plugins/FX/Niagara/Source/Niagara/Private/NiagaraDebugHud.cpp
@@ -29,19 +29,6 @@ namespace NiagaraDebugLocal
 {
 	FCriticalSection RTFramesGuard;
 
-	enum class EEngineVariables : uint8
-	{
-		LODDistance,
-		LODFraction,
-		Num
-	};
-
-	static FString GEngineVariableStrings[(int)EEngineVariables::Num] =
-	{
-		TEXT("Engine.LODDistance"),
-		TEXT("Engine.LODFraction"),
-	};
-
 	enum class EParameterStoreLocation
 	{
 		 UserOverride,
@@ -76,8 +63,7 @@ namespace NiagaraDebugLocal
 		FDelegateHandle CompiledDelegate;
 #endif
 
-		bool bShowEngineVariable[(int)EEngineVariables::Num] = {};				// Engine variables that are not contained within the store
-
+		TArray<int32>										EngineVariables;			// Engine varibles that are visible on the HUD, these are special because they are not in the buffers
 		TArray<FNiagaraDataSetDebugAccessor>				SystemVariables;			// System & Emitter variables since both are inside the same DataBuffer
 		TArray<FParameterStoreVariable>						ParameterStoreVariables;	// Variables Stored inside Parameter Stores
 	#if WITH_EDITORONLY_DATA
@@ -88,6 +74,13 @@ namespace NiagaraDebugLocal
 		TArray<FNiagaraDataSetAccessor<FNiagaraPosition>>	ParticlePositionAccessors;	// Only valid if we have particle attributes
 	};
 
+	static const TPair<FString, TFunction<FString(FNiagaraSystemInstance*)>> GEngineVariables[] =
+	{
+		{TEXT("Engine.LODDistance"),		[](FNiagaraSystemInstance* SystemInstance) -> FString { return FString::Printf(TEXT("%6.2f"), SystemInstance->GetLODDistance()); }},
+		{TEXT("Engine.LODFraction"),		[](FNiagaraSystemInstance* SystemInstance) -> FString { return FString::Printf(TEXT("%6.2f"), SystemInstance->GetLODDistance() / SystemInstance->GetMaxLODDistance()); }},
+		{TEXT("Engine.System.TickCount"),	[](FNiagaraSystemInstance* SystemInstance) -> FString { return FString::Printf(TEXT("%d"), SystemInstance->GetTickCount()); }},
+	};
+
 	static TMap<TWeakObjectPtr<UNiagaraSystem>, FCachedVariables> GCachedSystemVariables;
 
 	FNiagaraDebugHUDSettingsData Settings;
@@ -454,14 +447,14 @@ namespace NiagaraDebugLocal
 			if (Settings.bShowSystemVariables && Settings.SystemVariables.Num() > 0)
 			{
 				FindSystemVariablesByWildcard(NiagaraSystem, Settings.SystemVariables, CachedVariables);
-
-				for (int32 iVariable = 0; iVariable < (int32)EEngineVariables::Num; ++iVariable)
+				
+				for (int32 iVariable=0; iVariable < UE_ARRAY_COUNT(GEngineVariables); ++iVariable)
 				{
 					for (const FNiagaraDebugHUDVariable& DebugVariable : Settings.SystemVariables)
 					{
-						if (DebugVariable.bEnabled && GEngineVariableStrings[iVariable].MatchesWildcard(DebugVariable.Name))
+						if (DebugVariable.bEnabled && GEngineVariables[iVariable].Key.MatchesWildcard(DebugVariable.Name))
 						{
-							CachedVariables->bShowEngineVariable[iVariable] = true;
+							CachedVariables->EngineVariables.Add(iVariable);
 							break;
 						}
 					}
@@ -2855,13 +2848,9 @@ void FNiagaraDebugHud::DrawComponents(FNiagaraWorldManager* WorldManager, UCanva
 						const FCachedVariables& CachedVariables = GetCachedVariables(NiagaraSystem);
 
 						// Engine Variables
-						if (CachedVariables.bShowEngineVariable[(int)EEngineVariables::LODDistance])
-						{
-							StringBuilder.Appendf(TEXT("%s = %.2f\n"), *GEngineVariableStrings[(int)EEngineVariables::LODDistance], SystemInstance->GetLODDistance());
-						}
-						if (CachedVariables.bShowEngineVariable[(int)EEngineVariables::LODFraction])
+						for (int32 EngineVariable : CachedVariables.EngineVariables)
 						{
-							StringBuilder.Appendf(TEXT("%s = %.2f\n"), *GEngineVariableStrings[(int)EEngineVariables::LODFraction], SystemInstance->GetLODDistance() / SystemInstance->GetMaxLODDistance());
+							StringBuilder.Appendf(TEXT("%s = %s\n"), *GEngineVariables[EngineVariable].Key, *GEngineVariables[EngineVariable].Value(SystemInstance));
 						}
 
 						// System variables
