- Add UpdateLUT back into CopyToInternal until we figure out why this needs to be called

#rnx

diff --git a/Engine/Plugins/FX/Niagara/Source/Niagara/Private/NiagaraDataInterfaceColorCurve.cpp b/Engine/Plugins/FX/Niagara/Source/Niagara/Private/NiagaraDataInterfaceColorCurve.cpp
--- a/Engine/Plugins/FX/Niagara/Source/Niagara/Private/NiagaraDataInterfaceColorCurve.cpp
+++ b/Engine/Plugins/FX/Niagara/Source/Niagara/Private/NiagaraDataInterfaceColorCurve.cpp
@@ -131,7 +131,14 @@ bool UNiagaraDataInterfaceColorCurve::CopyToInternal(UNiagaraDataInterface* Dest
 	DestinationColorCurve->GreenCurve = GreenCurve;
 	DestinationColorCurve->BlueCurve = BlueCurve;
 	DestinationColorCurve->AlphaCurve = AlphaCurve;
+#if WITH_EDITORONLY_DATA
+	DestinationColorCurve->UpdateLUT();
 
+	if (!CompareLUTS(DestinationColorCurve->ShaderLUT))
+	{
+		UE_LOG(LogNiagara, Log, TEXT("CopyToInternal LUT generation is out of sync. Please investigate. %s to %s"), *GetPathName(), *DestinationColorCurve->GetPathName());
+	}
+#endif
 	return true;
 }
 
diff --git a/Engine/Plugins/FX/Niagara/Source/Niagara/Private/NiagaraDataInterfaceCurve.cpp b/Engine/Plugins/FX/Niagara/Source/Niagara/Private/NiagaraDataInterfaceCurve.cpp
--- a/Engine/Plugins/FX/Niagara/Source/Niagara/Private/NiagaraDataInterfaceCurve.cpp
+++ b/Engine/Plugins/FX/Niagara/Source/Niagara/Private/NiagaraDataInterfaceCurve.cpp
@@ -71,7 +71,13 @@ bool UNiagaraDataInterfaceCurve::CopyToInternal(UNiagaraDataInterface* Destinati
 		return false;
 	}
 	CastChecked<UNiagaraDataInterfaceCurve>(Destination)->Curve = Curve;
-
+#if WITH_EDITORONLY_DATA
+	CastChecked<UNiagaraDataInterfaceCurve>(Destination)->UpdateLUT();
+	if (!CompareLUTS(CastChecked<UNiagaraDataInterfaceCurve>(Destination)->ShaderLUT))
+	{
+		UE_LOG(LogNiagara, Log, TEXT("Post CopyToInternal LUT generation is out of sync. Please investigate. %s"), *GetPathName());
+	}
+#endif
 	return true;
 }
 
diff --git a/Engine/Plugins/FX/Niagara/Source/Niagara/Private/NiagaraDataInterfaceVector2DCurve.cpp b/Engine/Plugins/FX/Niagara/Source/Niagara/Private/NiagaraDataInterfaceVector2DCurve.cpp
--- a/Engine/Plugins/FX/Niagara/Source/Niagara/Private/NiagaraDataInterfaceVector2DCurve.cpp
+++ b/Engine/Plugins/FX/Niagara/Source/Niagara/Private/NiagaraDataInterfaceVector2DCurve.cpp
@@ -115,7 +115,13 @@ bool UNiagaraDataInterfaceVector2DCurve::CopyToInternal(UNiagaraDataInterface* D
 	UNiagaraDataInterfaceVector2DCurve* DestinationVector2DCurve = CastChecked<UNiagaraDataInterfaceVector2DCurve>(Destination);
 	DestinationVector2DCurve->XCurve = XCurve;
 	DestinationVector2DCurve->YCurve = YCurve;
-
+#if WITH_EDITORONLY_DATA
+	DestinationVector2DCurve->UpdateLUT();
+	if (!CompareLUTS(DestinationVector2DCurve->ShaderLUT))
+	{
+		UE_LOG(LogNiagara, Log, TEXT("Post CopyToInternal LUT generation is out of sync. Please investigate. %s"), *GetPathName());
+	}
+#endif
 	return true;
 }
 
diff --git a/Engine/Plugins/FX/Niagara/Source/Niagara/Private/NiagaraDataInterfaceVector4Curve.cpp b/Engine/Plugins/FX/Niagara/Source/Niagara/Private/NiagaraDataInterfaceVector4Curve.cpp
--- a/Engine/Plugins/FX/Niagara/Source/Niagara/Private/NiagaraDataInterfaceVector4Curve.cpp
+++ b/Engine/Plugins/FX/Niagara/Source/Niagara/Private/NiagaraDataInterfaceVector4Curve.cpp
@@ -131,6 +131,13 @@ bool UNiagaraDataInterfaceVector4Curve::CopyToInternal(UNiagaraDataInterface* De
 	DestinationColorCurve->YCurve = YCurve;
 	DestinationColorCurve->ZCurve = ZCurve;
 	DestinationColorCurve->WCurve = WCurve;
+#if WITH_EDITORONLY_DATA
+	DestinationColorCurve->UpdateLUT();
+	if (!CompareLUTS(DestinationColorCurve->ShaderLUT))
+	{
+		UE_LOG(LogNiagara, Log, TEXT("Post CopyToInternal LUT generation is out of sync. Please investigate. %s"), *GetPathName());
+	}
+#endif
 
 	return true;
 }
diff --git a/Engine/Plugins/FX/Niagara/Source/Niagara/Private/NiagaraDataInterfaceVectorCurve.cpp b/Engine/Plugins/FX/Niagara/Source/Niagara/Private/NiagaraDataInterfaceVectorCurve.cpp
--- a/Engine/Plugins/FX/Niagara/Source/Niagara/Private/NiagaraDataInterfaceVectorCurve.cpp
+++ b/Engine/Plugins/FX/Niagara/Source/Niagara/Private/NiagaraDataInterfaceVectorCurve.cpp
@@ -124,7 +124,14 @@ bool UNiagaraDataInterfaceVectorCurve::CopyToInternal(UNiagaraDataInterface* Des
 	DestinationVectorCurve->XCurve = XCurve;
 	DestinationVectorCurve->YCurve = YCurve;
 	DestinationVectorCurve->ZCurve = ZCurve;
+#if WITH_EDITORONLY_DATA
+	DestinationVectorCurve->UpdateLUT();
 
+	if (!CompareLUTS(DestinationVectorCurve->ShaderLUT))
+	{
+		UE_LOG(LogNiagara, Log, TEXT("Post CopyToInternal LUT generation is out of sync. Please investigate. %s"), *GetPathName());
+	}
+#endif
 	return true;
 }
 
