- Fix validation warning with NiagaraCulledGPUInstanceCounts
- We now track the state as we updated culled counts in two locations and using bAcquiredCulledCounts is not reliable

#rnx

diff --git a/Engine/Plugins/FX/Niagara/Source/Niagara/Classes/NiagaraGPUInstanceCountManager.h b/Engine/Plugins/FX/Niagara/Source/Niagara/Classes/NiagaraGPUInstanceCountManager.h
--- a/Engine/Plugins/FX/Niagara/Source/Niagara/Classes/NiagaraGPUInstanceCountManager.h
+++ b/Engine/Plugins/FX/Niagara/Source/Niagara/Classes/NiagaraGPUInstanceCountManager.h
@@ -140,6 +140,8 @@ protected:
 	int32 AllocatedCulledCounts = 0;
 	/** Whether or not the culled counts were acquired this frame */
 	bool bAcquiredCulledCounts = false;
+	/** Used to track RHI access for the culled counts buffer. */
+	ERHIAccess CulledCountsRHIAccess = ERHIAccess::Unknown;
 
 	/** A buffer holding the each emitter particle count after a simulation tick. */
 	FRWBuffer CountBuffer;
diff --git a/Engine/Plugins/FX/Niagara/Source/Niagara/Private/NiagaraGPUInstanceCountManager.cpp b/Engine/Plugins/FX/Niagara/Source/Niagara/Private/NiagaraGPUInstanceCountManager.cpp
--- a/Engine/Plugins/FX/Niagara/Source/Niagara/Private/NiagaraGPUInstanceCountManager.cpp
+++ b/Engine/Plugins/FX/Niagara/Source/Niagara/Private/NiagaraGPUInstanceCountManager.cpp
@@ -232,16 +232,26 @@ FRWBuffer* FNiagaraGPUInstanceCountManager::AcquireCulledCountsBuffer(FRHIComman
 				CulledCountBuffer.Release();
 
 				AllocatedCulledCounts = RecommendedCulledCounts;
-				CulledCountBuffer.Initialize(TEXT("NiagaraCulledGPUInstanceCounts"), sizeof(uint32), AllocatedCulledCounts, EPixelFormat::PF_R32_UINT, ERHIAccess::SRVCompute);
+				CulledCountBuffer.Initialize(TEXT("NiagaraCulledGPUInstanceCounts"), sizeof(uint32), AllocatedCulledCounts, EPixelFormat::PF_R32_UINT, ERHIAccess::UAVCompute);
+				CulledCountsRHIAccess = ERHIAccess::UAVCompute;
+			}
+			else if (CulledCountsRHIAccess != ERHIAccess::UAVCompute)
+			{
+				RHICmdList.Transition(FRHITransitionInfo(CulledCountBuffer.UAV, CulledCountsRHIAccess, ERHIAccess::UAVCompute));
+				CulledCountsRHIAccess = ERHIAccess::UAVCompute;
 			}
 
 			// Initialize the buffer by clearing it to zero then transition it to be ready to write to
-			RHICmdList.Transition(FRHITransitionInfo(CulledCountBuffer.UAV, ERHIAccess::SRVCompute, ERHIAccess::UAVCompute));
 			RHICmdList.ClearUAVUint(CulledCountBuffer.UAV, FUintVector4(EForceInit::ForceInitToZero));
 			RHICmdList.Transition(FRHITransitionInfo(CulledCountBuffer.UAV, ERHIAccess::UAVCompute, ERHIAccess::UAVCompute));
 
 			bAcquiredCulledCounts = true;
 		}
+		else if (CulledCountsRHIAccess != ERHIAccess::UAVCompute)
+		{
+			RHICmdList.Transition(FRHITransitionInfo(CulledCountBuffer.UAV, CulledCountsRHIAccess, ERHIAccess::UAVCompute));
+			CulledCountsRHIAccess = ERHIAccess::UAVCompute;
+		}
 
 		return &CulledCountBuffer;
 	}
@@ -430,9 +440,10 @@ void FNiagaraGPUInstanceCountManager::UpdateDrawIndirectBuffers(FNiagaraGpuCompu
 		FShaderResourceViewRHIRef CulledCountsSRV = nullptr;
 		if (CulledCountBuffer.SRV.IsValid())
 		{
-			if (bAcquiredCulledCounts)
+			if (CulledCountsRHIAccess != ERHIAccess::SRVCompute)
 			{
-				Transitions.Emplace(CulledCountBuffer.UAV, ERHIAccess::UAVCompute, ERHIAccess::SRVCompute);
+				Transitions.Emplace(CulledCountBuffer.UAV, CulledCountsRHIAccess, ERHIAccess::SRVCompute);
+				CulledCountsRHIAccess = ERHIAccess::SRVCompute;
 			}
 			CulledCountsSRV = CulledCountBuffer.SRV.GetReference();
 		}
