Adding include before running IWYU

diff --git a/Engine/Plugins/FX/Niagara/Source/Niagara/Private/NiagaraAsyncGpuTraceHelper.h b/Engine/Plugins/FX/Niagara/Source/Niagara/Private/NiagaraAsyncGpuTraceHelper.h
--- a/Engine/Plugins/FX/Niagara/Source/Niagara/Private/NiagaraAsyncGpuTraceHelper.h
+++ b/Engine/Plugins/FX/Niagara/Source/Niagara/Private/NiagaraAsyncGpuTraceHelper.h
@@ -5,6 +5,7 @@
 
 #include "NiagaraAsyncGpuTraceProvider.h"
 #include "NiagaraGpuScratchPad.h"
+#include "NiagaraSettings.h"
 
 #include "RHIDefinitions.h"
 
@@ -156,4 +157,4 @@ private:
 	int32 NumGPURayTracedCollisionGroups = 0;
 	bool bCollisionGroupMapDirty = true;
 #endif
-};
\ No newline at end of file
+};
