Ran IWYU on NiagaraShader Pirvate and Public.

diff --git a/Engine/Plugins/FX/Niagara/Source/NiagaraShader/Private/NiagaraAsyncGpuTraceProviderGsdf.cpp b/Engine/Plugins/FX/Niagara/Source/NiagaraShader/Private/NiagaraAsyncGpuTraceProviderGsdf.cpp
--- a/Engine/Plugins/FX/Niagara/Source/NiagaraShader/Private/NiagaraAsyncGpuTraceProviderGsdf.cpp
+++ b/Engine/Plugins/FX/Niagara/Source/NiagaraShader/Private/NiagaraAsyncGpuTraceProviderGsdf.cpp
@@ -11,6 +11,7 @@
 #include "RenderGraphUtils.h"
 #include "SceneInterface.h"
 #include "SceneManagement.h"
+#include "SceneView.h"
 #include "ShaderParameterStruct.h"
 
 static int GNiagaraAsyncGpuTraceGsdfEnabled = 1;
diff --git a/Engine/Plugins/FX/Niagara/Source/NiagaraShader/Private/NiagaraAsyncGpuTraceProviderGsdf.h b/Engine/Plugins/FX/Niagara/Source/NiagaraShader/Private/NiagaraAsyncGpuTraceProviderGsdf.h
--- a/Engine/Plugins/FX/Niagara/Source/NiagaraShader/Private/NiagaraAsyncGpuTraceProviderGsdf.h
+++ b/Engine/Plugins/FX/Niagara/Source/NiagaraShader/Private/NiagaraAsyncGpuTraceProviderGsdf.h
@@ -2,11 +2,11 @@
 
 #pragma once
 
-#include "CoreMinimal.h"
-
 #include "GlobalDistanceFieldParameters.h"
 #include "NiagaraAsyncGpuTraceProvider.h"
 
+class FViewUniformShaderParameters;
+
 class FNiagaraAsyncGpuTraceProviderGsdf : public FNiagaraAsyncGpuTraceProvider
 {
 public:
@@ -24,4 +24,4 @@ public:
 private:
 	FGlobalDistanceFieldParameterData m_DistanceFieldData;
 	TUniformBufferRef<FViewUniformShaderParameters> m_ViewUniformBuffer;
-};
\ No newline at end of file
+};
diff --git a/Engine/Plugins/FX/Niagara/Source/NiagaraShader/Private/NiagaraAsyncGpuTraceProviderHwrt.cpp b/Engine/Plugins/FX/Niagara/Source/NiagaraShader/Private/NiagaraAsyncGpuTraceProviderHwrt.cpp
--- a/Engine/Plugins/FX/Niagara/Source/NiagaraShader/Private/NiagaraAsyncGpuTraceProviderHwrt.cpp
+++ b/Engine/Plugins/FX/Niagara/Source/NiagaraShader/Private/NiagaraAsyncGpuTraceProviderHwrt.cpp
@@ -5,12 +5,12 @@
 
 #if RHI_RAYTRACING
 
+#include "Containers/StridedView.h"
 #include "FXRenderingUtils.h"
 #include "GlobalShader.h"
 #include "MeshPassProcessor.h"
-#include "NiagaraRenderer.h"
 #include "NiagaraSettings.h"
-#include "SceneInterface.h"
+#include "PipelineStateCache.h"
 #include "ShaderParameterMacros.h"
 #include "ShaderParameterStruct.h"
 #include "Templates/Function.h"
diff --git a/Engine/Plugins/FX/Niagara/Source/NiagaraShader/Private/NiagaraAsyncGpuTraceProviderHwrt.h b/Engine/Plugins/FX/Niagara/Source/NiagaraShader/Private/NiagaraAsyncGpuTraceProviderHwrt.h
--- a/Engine/Plugins/FX/Niagara/Source/NiagaraShader/Private/NiagaraAsyncGpuTraceProviderHwrt.h
+++ b/Engine/Plugins/FX/Niagara/Source/NiagaraShader/Private/NiagaraAsyncGpuTraceProviderHwrt.h
@@ -2,9 +2,10 @@
 
 #pragma once
 
-#include "CoreMinimal.h"
-
 #include "NiagaraAsyncGpuTraceProvider.h"
+#include "ShaderParameterMacros.h"
+
+class FViewUniformShaderParameters;
 
 #if RHI_RAYTRACING
 
@@ -29,4 +30,4 @@ private:
 	TUniformBufferRef<FViewUniformShaderParameters> ViewUniformBuffer;
 };
 
-#endif
\ No newline at end of file
+#endif
diff --git a/Engine/Plugins/FX/Niagara/Source/NiagaraShader/Private/NiagaraClearCounts.cpp b/Engine/Plugins/FX/Niagara/Source/NiagaraShader/Private/NiagaraClearCounts.cpp
--- a/Engine/Plugins/FX/Niagara/Source/NiagaraShader/Private/NiagaraClearCounts.cpp
+++ b/Engine/Plugins/FX/Niagara/Source/NiagaraShader/Private/NiagaraClearCounts.cpp
@@ -2,14 +2,11 @@
 
 #include "NiagaraClearCounts.h"
 
+#include "GlobalShader.h"
 #include "RHI.h"
-#include "ScreenRendering.h"
-#include "CommonRenderResources.h"
-#include "Modules/ModuleManager.h"
-
+#include "RHIUtilities.h"
 #include "RenderGraphBuilder.h"
 #include "RenderGraphUtils.h"
-#include "ScreenPass.h"
 
 //////////////////////////////////////////////////////////////////////////
 
diff --git a/Engine/Plugins/FX/Niagara/Source/NiagaraShader/Private/NiagaraGPUSceneUtils.cpp b/Engine/Plugins/FX/Niagara/Source/NiagaraShader/Private/NiagaraGPUSceneUtils.cpp
--- a/Engine/Plugins/FX/Niagara/Source/NiagaraShader/Private/NiagaraGPUSceneUtils.cpp
+++ b/Engine/Plugins/FX/Niagara/Source/NiagaraShader/Private/NiagaraGPUSceneUtils.cpp
@@ -4,6 +4,7 @@
 #include "GlobalShader.h"
 #include "RenderGraphUtils.h"
 #include "DataDrivenShaderPlatformInfo.h"
+#include "RenderUtils.h"
 
 class FNiagaraUpdateMeshGPUSceneInstancesCS : public FGlobalShader
 {
@@ -53,4 +54,4 @@ void FNiagaraGPUSceneUtils::AddUpdateMeshParticleInstancesPass(
 		&Params,
 		FComputeShaderUtils::GetGroupCountWrapped(Params.ParticleCount, FNiagaraUpdateMeshGPUSceneInstancesCS::ThreadGroupSize)
 	);
-}
\ No newline at end of file
+}
diff --git a/Engine/Plugins/FX/Niagara/Source/NiagaraShader/Private/NiagaraGenerateMips.cpp b/Engine/Plugins/FX/Niagara/Source/NiagaraShader/Private/NiagaraGenerateMips.cpp
--- a/Engine/Plugins/FX/Niagara/Source/NiagaraShader/Private/NiagaraGenerateMips.cpp
+++ b/Engine/Plugins/FX/Niagara/Source/NiagaraShader/Private/NiagaraGenerateMips.cpp
@@ -4,7 +4,6 @@
 #include "GlobalShader.h"
 #include "RenderGraphUtils.h"
 #include "RHIStaticStates.h"
-#include "DataDrivenShaderPlatformInfo.h"
 
 #include UE_INLINE_GENERATED_CPP_BY_NAME(NiagaraGenerateMips)
 
diff --git a/Engine/Plugins/FX/Niagara/Source/NiagaraShader/Private/NiagaraShader.cpp b/Engine/Plugins/FX/Niagara/Source/NiagaraShader/Private/NiagaraShader.cpp
--- a/Engine/Plugins/FX/Niagara/Source/NiagaraShader/Private/NiagaraShader.cpp
+++ b/Engine/Plugins/FX/Niagara/Source/NiagaraShader/Private/NiagaraShader.cpp
@@ -2,23 +2,19 @@
 
 #include "NiagaraShader.h"
 #include "NiagaraShared.h"
-#include "NiagaraShaderMap.h"
 #include "NiagaraScriptBase.h"
+#include "RenderingThread.h"
 #include "Stats/StatsMisc.h"
 #include "Serialization/MemoryWriter.h"
 #include "Serialization/MemoryReader.h"
-#include "ProfilingDebugging/DiagnosticTable.h"
 #include "ShaderCompiler.h"
 #include "DataDrivenShaderPlatformInfo.h"
 #include "UObject/DevObjectVersion.h"
 #include "NiagaraShaderCompilationManager.h"
-#include "UObject/CoreRedirects.h"
 #if WITH_EDITOR
 	#include "Interfaces/ITargetPlatformManagerModule.h"
-	#include "TickableEditorObject.h"
 	#include "DerivedDataCacheInterface.h"
 	#include "Interfaces/ITargetPlatformManagerModule.h"
-	#include "Interfaces/ITargetPlatform.h"
 #endif
 #include "ProfilingDebugging/CookStats.h"
 
diff --git a/Engine/Plugins/FX/Niagara/Source/NiagaraShader/Private/NiagaraShaderCompilationManager.cpp b/Engine/Plugins/FX/Niagara/Source/NiagaraShader/Private/NiagaraShaderCompilationManager.cpp
--- a/Engine/Plugins/FX/Niagara/Source/NiagaraShader/Private/NiagaraShaderCompilationManager.cpp
+++ b/Engine/Plugins/FX/Niagara/Source/NiagaraShader/Private/NiagaraShaderCompilationManager.cpp
@@ -2,14 +2,9 @@
 
 #include "NiagaraShaderCompilationManager.h"
 
-#include "GlobalShader.h"
-#include "HAL/FileManager.h"
-#if WITH_EDITOR
-#include "Interfaces/IShaderFormat.h"
-#include "Interfaces/ITargetPlatformManagerModule.h"
-#endif
 #include "Misc/Paths.h"
 #include "NiagaraShared.h"
+#include "RenderingThread.h"
 #include "ShaderCompiler.h"
 #include "Tickable.h"
 #include "UObject/UObjectThreadContext.h"
diff --git a/Engine/Plugins/FX/Niagara/Source/NiagaraShader/Private/NiagaraShaderParticleID.cpp b/Engine/Plugins/FX/Niagara/Source/NiagaraShader/Private/NiagaraShaderParticleID.cpp
--- a/Engine/Plugins/FX/Niagara/Source/NiagaraShader/Private/NiagaraShaderParticleID.cpp
+++ b/Engine/Plugins/FX/Niagara/Source/NiagaraShader/Private/NiagaraShaderParticleID.cpp
@@ -3,7 +3,6 @@
 #include "NiagaraShaderParticleID.h"
 #include "GlobalShader.h"
 #include "RenderGraphUtils.h"
-#include "ShaderParameterUtils.h"
 #include "PipelineStateCache.h"
 #include "DataDrivenShaderPlatformInfo.h"
 
diff --git a/Engine/Plugins/FX/Niagara/Source/NiagaraShader/Private/NiagaraShared.cpp b/Engine/Plugins/FX/Niagara/Source/NiagaraShader/Private/NiagaraShared.cpp
--- a/Engine/Plugins/FX/Niagara/Source/NiagaraShader/Private/NiagaraShared.cpp
+++ b/Engine/Plugins/FX/Niagara/Source/NiagaraShader/Private/NiagaraShared.cpp
@@ -11,16 +11,14 @@
 #include "NiagaraShader.h"
 #include "NiagaraScriptBase.h"
 #include "NiagaraScript.h"		//-TODO: This should be fixed so we are not reading structures from modules we do not depend on
+#include "PipelineStateCache.h"
 #include "Stats/StatsMisc.h"
-#include "UObject/CoreObjectVersion.h"
 #include "Misc/App.h"
 #include "UObject/Package.h"
 #include "UObject/UObjectHash.h"
-#include "UObject/UObjectIterator.h"
 #include "PSOPrecache.h"
 #include "ShaderCompiler.h"
 #include "NiagaraShaderCompilationManager.h"
-#include "RendererInterface.h"
 #include "Modules/ModuleManager.h"
 #include "NiagaraCustomVersion.h"
 #if WITH_EDITOR
diff --git a/Engine/Plugins/FX/Niagara/Source/NiagaraShader/Public/NiagaraAsyncGpuTraceProvider.h b/Engine/Plugins/FX/Niagara/Source/NiagaraShader/Public/NiagaraAsyncGpuTraceProvider.h
--- a/Engine/Plugins/FX/Niagara/Source/NiagaraShader/Public/NiagaraAsyncGpuTraceProvider.h
+++ b/Engine/Plugins/FX/Niagara/Source/NiagaraShader/Public/NiagaraAsyncGpuTraceProvider.h
@@ -1,10 +1,12 @@
 // Copyright Epic Games, Inc. All Rights Reserved.
 #pragma once
 
-#include "CoreMinimal.h"
-
-#include "NiagaraSettings.h"
 #include "RHIDefinitions.h"
+#include "RHIUtilities.h"
+
+class FPrimitiveComponentId;
+class FSceneInterface;
+namespace ENDICollisionQuery_AsyncGpuTraceProvider { enum Type : int; }
 
 class FNiagaraGpuComputeDispatchInterface;
 class FScene;
@@ -92,4 +94,9 @@ public:
 
 private:
 	UE_NONCOPYABLE(FNiagaraAsyncGpuTraceProvider);
-};
\ No newline at end of file
+};
+
+#if UE_ENABLE_INCLUDE_ORDER_DEPRECATED_IN_5_3
+#include "CoreMinimal.h"
+#include "NiagaraSettings.h"
+#endif
diff --git a/Engine/Plugins/FX/Niagara/Source/NiagaraShader/Public/NiagaraBatchedElements.h b/Engine/Plugins/FX/Niagara/Source/NiagaraShader/Public/NiagaraBatchedElements.h
--- a/Engine/Plugins/FX/Niagara/Source/NiagaraShader/Public/NiagaraBatchedElements.h
+++ b/Engine/Plugins/FX/Niagara/Source/NiagaraShader/Public/NiagaraBatchedElements.h
@@ -4,6 +4,8 @@
 
 #include "BatchedElements.h"
 
+class FTexture;
+
 /**
  * Batched element parameters for gathering attributes from different slices into a single color
  */
diff --git a/Engine/Plugins/FX/Niagara/Source/NiagaraShader/Public/NiagaraClearCounts.h b/Engine/Plugins/FX/Niagara/Source/NiagaraShader/Public/NiagaraClearCounts.h
--- a/Engine/Plugins/FX/Niagara/Source/NiagaraShader/Public/NiagaraClearCounts.h
+++ b/Engine/Plugins/FX/Niagara/Source/NiagaraShader/Public/NiagaraClearCounts.h
@@ -2,9 +2,12 @@
 
 #pragma once
 
-#include "CommonRenderResources.h"
-#include "RenderGraphResources.h"
-#include "RHI.h"
+#include "Containers/Map.h"
+#include "RenderGraphFwd.h"
+
+class FRDGBuilder;
+class FRHICommandList;
+class FRHIUnorderedAccessView;
 
 namespace NiagaraClearCounts
 {
@@ -13,3 +16,10 @@ namespace NiagaraClearCounts
 	//-TODO:RDG: Deprecated
 	NIAGARASHADER_API void ClearCountsInt(FRHICommandList& RHICmdList, FRHIUnorderedAccessView* UAV, TConstArrayView<TPair<uint32, int32>> IndexAndValueArray);
 }
+
+#if UE_ENABLE_INCLUDE_ORDER_DEPRECATED_IN_5_3
+#include "RHIFwd.h"
+#include "CommonRenderResources.h"
+#include "RHI.h"
+#include "RenderGraphResources.h"
+#endif
diff --git a/Engine/Plugins/FX/Niagara/Source/NiagaraShader/Public/NiagaraDebugShaders.h b/Engine/Plugins/FX/Niagara/Source/NiagaraShader/Public/NiagaraDebugShaders.h
--- a/Engine/Plugins/FX/Niagara/Source/NiagaraShader/Public/NiagaraDebugShaders.h
+++ b/Engine/Plugins/FX/Niagara/Source/NiagaraShader/Public/NiagaraDebugShaders.h
@@ -2,9 +2,12 @@
 
 #pragma once
 
-#include "RHI.h"
-#include "ScreenRendering.h"
-#include "CommonRenderResources.h"
+#include "Math/Vector2D.h"
+#include "RenderGraphFwd.h"
+
+class FRDGBuilder;
+
+class FSceneView;
 
 struct FScreenPassRenderTarget;
 namespace NiagaraDebugShaders
@@ -28,3 +31,9 @@ namespace NiagaraDebugShaders
 		const FVector2D& PreviewDisplayRange = FVector2D::ZeroVector
 	);
 }
+
+#if UE_ENABLE_INCLUDE_ORDER_DEPRECATED_IN_5_3
+#include "CommonRenderResources.h"
+#include "RHI.h"
+#include "ScreenRendering.h"
+#endif
diff --git a/Engine/Plugins/FX/Niagara/Source/NiagaraShader/Public/NiagaraGPUSceneUtils.h b/Engine/Plugins/FX/Niagara/Source/NiagaraShader/Public/NiagaraGPUSceneUtils.h
--- a/Engine/Plugins/FX/Niagara/Source/NiagaraShader/Public/NiagaraGPUSceneUtils.h
+++ b/Engine/Plugins/FX/Niagara/Source/NiagaraShader/Public/NiagaraGPUSceneUtils.h
@@ -2,11 +2,9 @@
 
 #pragma once
 
-#include "CoreMinimal.h"
+#include "GPUSceneWriter.h"
 #include "RHI.h"
 #include "RHIDefinitions.h"
-#include "ShaderParameterStruct.h"
-#include "ShaderParameterUtils.h"
 #include "NiagaraMeshVertexFactory.h"
 #include "SceneView.h"
 
@@ -43,4 +41,10 @@ public:
 		ERHIFeatureLevel::Type FeatureLevel,
 		bool bPreciseMotionVectors
 	);
-};
\ No newline at end of file
+};
+
+#if UE_ENABLE_INCLUDE_ORDER_DEPRECATED_IN_5_3
+#include "CoreMinimal.h"
+#include "ShaderParameterStruct.h"
+#include "ShaderParameterUtils.h"
+#endif
diff --git a/Engine/Plugins/FX/Niagara/Source/NiagaraShader/Public/NiagaraGenerateMips.h b/Engine/Plugins/FX/Niagara/Source/NiagaraShader/Public/NiagaraGenerateMips.h
--- a/Engine/Plugins/FX/Niagara/Source/NiagaraShader/Public/NiagaraGenerateMips.h
+++ b/Engine/Plugins/FX/Niagara/Source/NiagaraShader/Public/NiagaraGenerateMips.h
@@ -2,9 +2,11 @@
 
 #pragma once
 
-#include "RenderGraphDefinitions.h"
+#include "RenderGraphFwd.h"
 #include "NiagaraGenerateMips.generated.h"
 
+class FRDGBuilder;
+
 UENUM()
 enum class ENiagaraMipMapGenerationType : uint8
 {
@@ -26,3 +28,7 @@ namespace NiagaraGenerateMips
 {
 	NIAGARASHADER_API void GenerateMips(FRDGBuilder& GraphBuilder, FRDGTextureRef RDGTexture, ENiagaraMipMapGenerationType GenType);
 }
+
+#if UE_ENABLE_INCLUDE_ORDER_DEPRECATED_IN_5_3
+#include "RenderGraphDefinitions.h"
+#endif
diff --git a/Engine/Plugins/FX/Niagara/Source/NiagaraShader/Public/NiagaraRibbonCompute.h b/Engine/Plugins/FX/Niagara/Source/NiagaraShader/Public/NiagaraRibbonCompute.h
--- a/Engine/Plugins/FX/Niagara/Source/NiagaraShader/Public/NiagaraRibbonCompute.h
+++ b/Engine/Plugins/FX/Niagara/Source/NiagaraShader/Public/NiagaraRibbonCompute.h
@@ -6,7 +6,6 @@ NiagaraRibbonCompute.h: Niagara ribbon compute shaders for initialization of rib
 
 #pragma once
 
-#include "CoreMinimal.h"
 #include "GlobalShader.h"
 #include "ShaderPermutation.h"
 #include "ShaderParameterStruct.h"
@@ -303,3 +302,7 @@ class NIAGARASHADER_API FNiagaraRibbonCreateIndexBufferCS : public FGlobalShader
 	static void ModifyCompilationEnvironment(const FGlobalShaderPermutationParameters& Parameters, FShaderCompilerEnvironment& OutEnvironment);
 };
 
+
+#if UE_ENABLE_INCLUDE_ORDER_DEPRECATED_IN_5_3
+#include "CoreMinimal.h"
+#endif
diff --git a/Engine/Plugins/FX/Niagara/Source/NiagaraShader/Public/NiagaraScriptBase.h b/Engine/Plugins/FX/Niagara/Source/NiagaraShader/Public/NiagaraScriptBase.h
--- a/Engine/Plugins/FX/Niagara/Source/NiagaraShader/Public/NiagaraScriptBase.h
+++ b/Engine/Plugins/FX/Niagara/Source/NiagaraShader/Public/NiagaraScriptBase.h
@@ -2,12 +2,7 @@
 
 #pragma once
 
-#include "CoreMinimal.h"
-#if UE_ENABLE_INCLUDE_ORDER_DEPRECATED_IN_5_2
-#include "RHIDefinitions.h"
-#endif
 #include "NiagaraCore.h"
-
 #include "NiagaraScriptBase.generated.h"
 
 enum EShaderPlatform : uint16;
@@ -171,3 +166,11 @@ class UNiagaraScriptBase : public UObject
 	virtual bool ShouldCompile(EShaderPlatform Platform) const PURE_VIRTUAL(UNiagaraScriptBase::ShouldCompile, return true; );
 	virtual TConstArrayView<FSimulationStageMetaData> GetSimulationStageMetaData() const PURE_VIRTUAL(UNiagaraScriptBase::GetSimulationStageMetaData(), return MakeArrayView<FSimulationStageMetaData>(nullptr, 0); )
 };
+
+#if UE_ENABLE_INCLUDE_ORDER_DEPRECATED_IN_5_2
+#include "RHIDefinitions.h"
+#endif
+
+#if UE_ENABLE_INCLUDE_ORDER_DEPRECATED_IN_5_3
+#include "CoreMinimal.h"
+#endif
diff --git a/Engine/Plugins/FX/Niagara/Source/NiagaraShader/Public/NiagaraShaderCompilationManager.h b/Engine/Plugins/FX/Niagara/Source/NiagaraShader/Public/NiagaraShaderCompilationManager.h
--- a/Engine/Plugins/FX/Niagara/Source/NiagaraShader/Public/NiagaraShaderCompilationManager.h
+++ b/Engine/Plugins/FX/Niagara/Source/NiagaraShader/Public/NiagaraShaderCompilationManager.h
@@ -4,7 +4,6 @@
 
 #if WITH_EDITOR
 
-#include "NiagaraShared.h"
 #include "ShaderCompiler.h"
 
 /** Results for a single compiled shader map. */
@@ -59,4 +58,8 @@ private:
 
 extern NIAGARASHADER_API FNiagaraShaderCompilationManager GNiagaraShaderCompilationManager;
 
+#if UE_ENABLE_INCLUDE_ORDER_DEPRECATED_IN_5_3
+#include "NiagaraShared.h"
+#endif
+
 #endif //WITH_EDITOR
diff --git a/Engine/Plugins/FX/Niagara/Source/NiagaraShader/Public/NiagaraShaderModule.h b/Engine/Plugins/FX/Niagara/Source/NiagaraShader/Public/NiagaraShaderModule.h
--- a/Engine/Plugins/FX/Niagara/Source/NiagaraShader/Public/NiagaraShaderModule.h
+++ b/Engine/Plugins/FX/Niagara/Source/NiagaraShader/Public/NiagaraShaderModule.h
@@ -2,9 +2,10 @@
 
 #pragma once
 
-#include "CoreMinimal.h"
+#include "Delegates/Delegate.h"
 #include "Modules/ModuleInterface.h"
 
+class FString;
 class UNiagaraDataInterfaceBase;
 
 /**
@@ -44,3 +45,7 @@ private:
 	static INiagaraShaderModule* Singleton;
 };
 
+
+#if UE_ENABLE_INCLUDE_ORDER_DEPRECATED_IN_5_3
+#include "CoreMinimal.h"
+#endif
diff --git a/Engine/Plugins/FX/Niagara/Source/NiagaraShader/Public/NiagaraShaderParticleID.h b/Engine/Plugins/FX/Niagara/Source/NiagaraShader/Public/NiagaraShaderParticleID.h
--- a/Engine/Plugins/FX/Niagara/Source/NiagaraShader/Public/NiagaraShaderParticleID.h
+++ b/Engine/Plugins/FX/Niagara/Source/NiagaraShader/Public/NiagaraShaderParticleID.h
@@ -2,8 +2,18 @@
 
 #pragma once
 
-#include "RHI.h"
+#include "HAL/Platform.h"
+
+class FRHICommandList;
+class FRHIShaderResourceView;
+class FRHIUnorderedAccessView;
+namespace ERHIFeatureLevel { enum Type : int; }
 
 extern NIAGARASHADER_API void NiagaraInitGPUFreeIDList(FRHICommandList& RHICmdList, ERHIFeatureLevel::Type FeatureLevel, FRHIUnorderedAccessView* NewBufferUAV, uint32 NewBufferNumElements, FRHIShaderResourceView* ExistingBufferSRV, uint32 ExistingBufferNumElements);
 extern NIAGARASHADER_API void NiagaraComputeGPUFreeIDs(FRHICommandList& RHICmdList, ERHIFeatureLevel::Type FeatureLevel, FRHIShaderResourceView* IDToIndexTableSRV, uint32 NumIDs, FRHIUnorderedAccessView* FreeIDUAV, FRHIUnorderedAccessView* FreeIDListSizesUAV, uint32 FreeIDListIndex);
 extern NIAGARASHADER_API void NiagaraFillGPUIntBuffer(FRHICommandList& RHICmdList, ERHIFeatureLevel::Type FeatureLevel, FRHIUnorderedAccessView* BufferUAV, uint32 NumElements, int32 FillValue);
+
+#if UE_ENABLE_INCLUDE_ORDER_DEPRECATED_IN_5_3
+#include "RHIFwd.h"
+#include "RHI.h"
+#endif
diff --git a/Engine/Plugins/FX/Niagara/Source/NiagaraShader/Public/NiagaraShaderType.h b/Engine/Plugins/FX/Niagara/Source/NiagaraShader/Public/NiagaraShaderType.h
--- a/Engine/Plugins/FX/Niagara/Source/NiagaraShader/Public/NiagaraShaderType.h
+++ b/Engine/Plugins/FX/Niagara/Source/NiagaraShader/Public/NiagaraShaderType.h
@@ -6,12 +6,12 @@
 
 #pragma once
 
-#include "CoreMinimal.h"
 #include "Shader.h"
-#include "GlobalShader.h"
-#include "Engine/EngineTypes.h"
 
+struct FCachedUniformBufferDeclaration;
 struct FNiagaraShaderScriptParametersMetadata;
+struct FShaderTarget;
+struct FSharedShaderCompilerEnvironment;
 
 /** A macro to implement Niagara shaders. */
 #define IMPLEMENT_NIAGARA_SHADER_TYPE(TemplatePrefix,ShaderClass,SourceFilename,FunctionName,Frequency) \
@@ -180,3 +180,9 @@ protected:
 	}
 #endif // WITH_EDITOR
 };
+
+#if UE_ENABLE_INCLUDE_ORDER_DEPRECATED_IN_5_3
+#include "CoreMinimal.h"
+#include "Engine/EngineTypes.h"
+#include "GlobalShader.h"
+#endif
diff --git a/Engine/Plugins/FX/Niagara/Source/NiagaraShader/Public/NiagaraShared.h b/Engine/Plugins/FX/Niagara/Source/NiagaraShader/Public/NiagaraShared.h
--- a/Engine/Plugins/FX/Niagara/Source/NiagaraShader/Public/NiagaraShared.h
+++ b/Engine/Plugins/FX/Niagara/Source/NiagaraShader/Public/NiagaraShared.h
@@ -6,7 +6,6 @@
 
 #pragma once
 
-#include "CoreMinimal.h"
 #include "Containers/IndirectArray.h"
 #include "Misc/Guid.h"
 #include "Engine/EngineTypes.h"
@@ -15,11 +14,7 @@
 #include "MaterialShared.h"
 #include "RHI.h"
 #include "RenderResource.h"
-#if UE_ENABLE_INCLUDE_ORDER_DEPRECATED_IN_5_2
-#include "RenderingThread.h"
-#endif
 #include "RenderDeferredCleanup.h"
-#include "UniformBuffer.h"
 #include "Shader.h"
 #include "VertexFactory.h"
 #include "SceneTypes.h"
@@ -34,8 +29,10 @@ class FNiagaraShaderScript;
 class FNiagaraShaderMap;
 class FNiagaraShader;
 class FNiagaraShaderMapId;
+class FThreadSafeBool;
 class UNiagaraScriptBase;
 struct FNiagaraVMExecutableDataId;
+struct FSharedShaderCompilerEnvironment;
 
 /** Defines the compile event types for translation/compilation.*/
 UENUM()
@@ -1031,3 +1028,12 @@ namespace FNiagaraCVarUtilities
 	NIAGARASHADER_API FNiagaraCompileEventSeverity GetCompileEventSeverityForFailIfNotSet();
 	NIAGARASHADER_API bool GetShouldEmitMessagesForFailIfNotSet();
 };
+
+#if UE_ENABLE_INCLUDE_ORDER_DEPRECATED_IN_5_2
+#include "RenderingThread.h"
+#endif
+
+#if UE_ENABLE_INCLUDE_ORDER_DEPRECATED_IN_5_3
+#include "CoreMinimal.h"
+#include "UniformBuffer.h"
+#endif
