Ran IWYU on NiagaraAnimNotifies Private and Public

diff --git a/Engine/Plugins/FX/Niagara/Source/NiagaraAnimNotifies/Private/AnimNotifyState_TimedNiagaraEffect.cpp b/Engine/Plugins/FX/Niagara/Source/NiagaraAnimNotifies/Private/AnimNotifyState_TimedNiagaraEffect.cpp
--- a/Engine/Plugins/FX/Niagara/Source/NiagaraAnimNotifies/Private/AnimNotifyState_TimedNiagaraEffect.cpp
+++ b/Engine/Plugins/FX/Niagara/Source/NiagaraAnimNotifies/Private/AnimNotifyState_TimedNiagaraEffect.cpp
@@ -2,8 +2,8 @@
 
 #include "AnimNotifyState_TimedNiagaraEffect.h"
 
+#include "Animation/AnimInstance.h"
 #include "Components/SkeletalMeshComponent.h"
-#include "Kismet/GameplayStatics.h"
 
 #include "NiagaraComponent.h"
 #include "NiagaraSystem.h"
diff --git a/Engine/Plugins/FX/Niagara/Source/NiagaraAnimNotifies/Private/AnimNotify_PlayNiagaraEffect.cpp b/Engine/Plugins/FX/Niagara/Source/NiagaraAnimNotifies/Private/AnimNotify_PlayNiagaraEffect.cpp
--- a/Engine/Plugins/FX/Niagara/Source/NiagaraAnimNotifies/Private/AnimNotify_PlayNiagaraEffect.cpp
+++ b/Engine/Plugins/FX/Niagara/Source/NiagaraAnimNotifies/Private/AnimNotify_PlayNiagaraEffect.cpp
@@ -4,8 +4,6 @@
 
 #include "Components/SkeletalMeshComponent.h"
 #include "Engine/World.h"
-#include "Kismet/GameplayStatics.h"
-#include "Animation/AnimSequenceBase.h"
 
 #if WITH_EDITOR
 #include "Logging/MessageLog.h"
diff --git a/Engine/Plugins/FX/Niagara/Source/NiagaraAnimNotifies/Public/AnimNotifyState_TimedNiagaraEffect.h b/Engine/Plugins/FX/Niagara/Source/NiagaraAnimNotifies/Public/AnimNotifyState_TimedNiagaraEffect.h
--- a/Engine/Plugins/FX/Niagara/Source/NiagaraAnimNotifies/Public/AnimNotifyState_TimedNiagaraEffect.h
+++ b/Engine/Plugins/FX/Niagara/Source/NiagaraAnimNotifies/Public/AnimNotifyState_TimedNiagaraEffect.h
@@ -2,12 +2,11 @@
 
 #pragma once
 
-#include "CoreMinimal.h"
 #include "UObject/ObjectMacros.h"
 #include "Animation/AnimNotifies/AnimNotifyState.h"
-#include "Animation/AnimInstance.h"
 #include "AnimNotifyState_TimedNiagaraEffect.generated.h"
 
+class UMeshComponent;
 class UNiagaraSystem;
 class USkeletalMeshComponent;
 class UFXSystemAsset;
@@ -122,4 +121,9 @@ protected:
 		float Elapsed = 0.0f;
 	};
 	TMap<UMeshComponent*, FInstanceProgressInfo> ProgressInfoMap;
-};
\ No newline at end of file
+};
+
+#if UE_ENABLE_INCLUDE_ORDER_DEPRECATED_IN_5_3
+#include "Animation/AnimInstance.h"
+#include "CoreMinimal.h"
+#endif
diff --git a/Engine/Plugins/FX/Niagara/Source/NiagaraAnimNotifies/Public/AnimNotify_PlayNiagaraEffect.h b/Engine/Plugins/FX/Niagara/Source/NiagaraAnimNotifies/Public/AnimNotify_PlayNiagaraEffect.h
--- a/Engine/Plugins/FX/Niagara/Source/NiagaraAnimNotifies/Public/AnimNotify_PlayNiagaraEffect.h
+++ b/Engine/Plugins/FX/Niagara/Source/NiagaraAnimNotifies/Public/AnimNotify_PlayNiagaraEffect.h
@@ -2,7 +2,6 @@
 
 #pragma once
 
-#include "CoreMinimal.h"
 #include "UObject/ObjectMacros.h"
 #include "Animation/AnimNotifies/AnimNotify.h"
 
@@ -93,3 +92,7 @@ public:
 
 
 
+
+#if UE_ENABLE_INCLUDE_ORDER_DEPRECATED_IN_5_3
+#include "CoreMinimal.h"
+#endif
diff --git a/Engine/Plugins/FX/Niagara/Source/NiagaraAnimNotifies/Public/NiagaraAnimNotifiesModule.h b/Engine/Plugins/FX/Niagara/Source/NiagaraAnimNotifies/Public/NiagaraAnimNotifiesModule.h
--- a/Engine/Plugins/FX/Niagara/Source/NiagaraAnimNotifies/Public/NiagaraAnimNotifiesModule.h
+++ b/Engine/Plugins/FX/Niagara/Source/NiagaraAnimNotifies/Public/NiagaraAnimNotifiesModule.h
@@ -2,7 +2,6 @@
 
 #pragma once
 
-#include "CoreMinimal.h"
 #include "Modules/ModuleInterface.h"
 
 /**
@@ -12,3 +11,7 @@ class INiagaraAnimNotifiesModule : public IModuleInterface
 {
 };
 
+
+#if UE_ENABLE_INCLUDE_ORDER_DEPRECATED_IN_5_3
+#include "CoreMinimal.h"
+#endif
