Clearing out empty scalability managers and any for garbage EffectTypes on preGC to avoid GC warnings.

diff --git a/Engine/Plugins/FX/Niagara/Source/Niagara/Private/NiagaraSystemSimulation.cpp b/Engine/Plugins/FX/Niagara/Source/Niagara/Private/NiagaraSystemSimulation.cpp
--- a/Engine/Plugins/FX/Niagara/Source/Niagara/Private/NiagaraSystemSimulation.cpp
+++ b/Engine/Plugins/FX/Niagara/Source/Niagara/Private/NiagaraSystemSimulation.cpp
@@ -611,8 +611,6 @@ FNiagaraSystemSimulationTickContext::FNiagaraSystemSimulationTickContext(class F
 
 void FNiagaraSystemSimulation::AddReferencedObjects(FReferenceCollector& Collector)
 {
-	//We keep a hard ref to the system.
-	Collector.AddReferencedObject(EffectType);
 }
 
 FNiagaraSystemSimulation::FNiagaraSystemSimulation()
@@ -2148,9 +2146,9 @@ void FNiagaraSystemSimulation::RemoveInstance(FNiagaraSystemInstance* Instance)
 	}
 
 	check(IsInGameThread());
-	if (EffectType)
+	if (UNiagaraEffectType* FXType = EffectType.Get())
 	{
-		--EffectType->NumInstances;
+		--FXType->NumInstances;
 	}
 
 	// Remove from pending promotions list
@@ -2198,10 +2196,10 @@ void FNiagaraSystemSimulation::AddInstance(FNiagaraSystemInstance* Instance)
 		System->RegisterActiveInstance();
 	}
 
-	if (EffectType)
+	if (UNiagaraEffectType* FXType = EffectType.Get())
 	{
-		++EffectType->NumInstances;
-		EffectType->bNewSystemsSinceLastScalabilityUpdate = true;
+		++FXType->NumInstances;
+		FXType->bNewSystemsSinceLastScalabilityUpdate = true;
 	}
 
 	check(GetSystemInstances(ENiagaraSystemInstanceState::Running).Num() == MainDataSet.GetCurrentDataChecked().GetNumInstances());
diff --git a/Engine/Plugins/FX/Niagara/Source/Niagara/Private/NiagaraWorldManager.cpp b/Engine/Plugins/FX/Niagara/Source/Niagara/Private/NiagaraWorldManager.cpp
--- a/Engine/Plugins/FX/Niagara/Source/Niagara/Private/NiagaraWorldManager.cpp
+++ b/Engine/Plugins/FX/Niagara/Source/Niagara/Private/NiagaraWorldManager.cpp
@@ -668,6 +668,17 @@ void FNiagaraWorldManager::PreGarbageCollect()
 			}
 		}
 	}
+	
+	for (auto It = ScalabilityManagers.CreateIterator(); It; ++It)
+	{
+		UNiagaraEffectType* EffectType = It.Key();
+		FNiagaraScalabilityManager& ScalabilityMan = It.Value();
+
+		if(!IsValid(EffectType) || ScalabilityMan.NumRegisteredComponents() == 0)
+		{
+			It.RemoveCurrent();
+		}
+	}
 }
 
 void FNiagaraWorldManager::PostReachabilityAnalysis()
diff --git a/Engine/Plugins/FX/Niagara/Source/Niagara/Public/NiagaraScalabilityManager.h b/Engine/Plugins/FX/Niagara/Source/Niagara/Public/NiagaraScalabilityManager.h
--- a/Engine/Plugins/FX/Niagara/Source/Niagara/Public/NiagaraScalabilityManager.h
+++ b/Engine/Plugins/FX/Niagara/Source/Niagara/Public/NiagaraScalabilityManager.h
@@ -75,6 +75,8 @@ struct FNiagaraScalabilityManager
 
 	void OnRefreshOwnerAllowsScalability();
 
+	int32 NumRegisteredComponents()const { return ManagedComponents.Num(); }
+
 #if DEBUG_SCALABILITY_STATE
 	void Dump();
 #endif
diff --git a/Engine/Plugins/FX/Niagara/Source/Niagara/Public/NiagaraSystemSimulation.h b/Engine/Plugins/FX/Niagara/Source/Niagara/Public/NiagaraSystemSimulation.h
--- a/Engine/Plugins/FX/Niagara/Source/Niagara/Public/NiagaraSystemSimulation.h
+++ b/Engine/Plugins/FX/Niagara/Source/Niagara/Public/NiagaraSystemSimulation.h
@@ -360,7 +360,7 @@ protected:
 	UNiagaraSystem* System;
 
 	/** We cache off the effect type in the unlikely even that someone GCs the System from under us so that we can keep the effect types instance count etc accurate. */
-	UNiagaraEffectType* EffectType;
+	TWeakObjectPtr<UNiagaraEffectType> EffectType;
 
 	/** Which tick group we are in, only valid when not in Solo mode. */
 	ETickingGroup SystemTickGroup = TG_MAX;
