Ran IWYU on NiagaraCore Private and Public

diff --git a/Engine/Plugins/FX/Niagara/Source/NiagaraCore/Public/NiagaraCoreModule.h b/Engine/Plugins/FX/Niagara/Source/NiagaraCore/Public/NiagaraCoreModule.h
--- a/Engine/Plugins/FX/Niagara/Source/NiagaraCore/Public/NiagaraCoreModule.h
+++ b/Engine/Plugins/FX/Niagara/Source/NiagaraCore/Public/NiagaraCoreModule.h
@@ -2,7 +2,6 @@
 
 #pragma once
 
-#include "CoreMinimal.h"
 #include "Modules/ModuleInterface.h"
 
 /**
@@ -12,3 +11,7 @@ class INiagaraCoreModule : public IModuleInterface
 {
 };
 
+
+#if UE_ENABLE_INCLUDE_ORDER_DEPRECATED_IN_5_3
+#include "CoreMinimal.h"
+#endif
diff --git a/Engine/Plugins/FX/Niagara/Source/NiagaraCore/Public/NiagaraCustomVersion.h b/Engine/Plugins/FX/Niagara/Source/NiagaraCore/Public/NiagaraCustomVersion.h
--- a/Engine/Plugins/FX/Niagara/Source/NiagaraCore/Public/NiagaraCustomVersion.h
+++ b/Engine/Plugins/FX/Niagara/Source/NiagaraCore/Public/NiagaraCustomVersion.h
@@ -2,8 +2,9 @@
 
 #pragma once
 
-#include "CoreMinimal.h"
-#include "Misc/Guid.h"
+#include "HAL/Platform.h"
+
+struct FGuid;
 
 // Custom serialization version for all packages containing Niagara asset types
 struct FNiagaraCustomVersion
@@ -207,3 +208,8 @@ private:
 
 
 
+
+#if UE_ENABLE_INCLUDE_ORDER_DEPRECATED_IN_5_3
+#include "CoreMinimal.h"
+#include "Misc/Guid.h"
+#endif
