Running IWYU on Niagara/Private

diff --git a/Engine/Plugins/FX/Niagara/Source/Niagara/Private/DataInterface/NiagaraDataInterfaceActorComponent.cpp b/Engine/Plugins/FX/Niagara/Source/Niagara/Private/DataInterface/NiagaraDataInterfaceActorComponent.cpp
--- a/Engine/Plugins/FX/Niagara/Source/Niagara/Private/DataInterface/NiagaraDataInterfaceActorComponent.cpp
+++ b/Engine/Plugins/FX/Niagara/Source/Niagara/Private/DataInterface/NiagaraDataInterfaceActorComponent.cpp
@@ -1,9 +1,10 @@
 // Copyright Epic Games, Inc. All Rights Reserved.
 
 #include "DataInterface/NiagaraDataInterfaceActorComponent.h"
+#include "Components/PrimitiveComponent.h"
 #include "GameFramework/Pawn.h"
+#include "NiagaraCompileHashVisitor.h"
 #include "NiagaraTypes.h"
-#include "NiagaraCustomVersion.h"
 #include "NiagaraShaderParametersBuilder.h"
 #include "NiagaraSystemInstance.h"
 #include "NiagaraWorldManager.h"
@@ -13,7 +14,6 @@
 #include "GameFramework/Actor.h"
 #include "GameFramework/PlayerController.h"
 #include "ShaderCompilerCore.h"
-#include "ShaderParameterUtils.h"
 #include "Misc/LargeWorldRenderPosition.h"
 
 #include UE_INLINE_GENERATED_CPP_BY_NAME(NiagaraDataInterfaceActorComponent)
diff --git a/Engine/Plugins/FX/Niagara/Source/Niagara/Private/DataInterface/NiagaraDataInterfaceAsyncGpuTrace.cpp b/Engine/Plugins/FX/Niagara/Source/Niagara/Private/DataInterface/NiagaraDataInterfaceAsyncGpuTrace.cpp
--- a/Engine/Plugins/FX/Niagara/Source/Niagara/Private/DataInterface/NiagaraDataInterfaceAsyncGpuTrace.cpp
+++ b/Engine/Plugins/FX/Niagara/Source/Niagara/Private/DataInterface/NiagaraDataInterfaceAsyncGpuTrace.cpp
@@ -2,22 +2,20 @@
 
 #include "NiagaraDataInterfaceAsyncGpuTrace.h"
 
-#include "GlobalDistanceFieldParameters.h"
 #include "NiagaraAsyncGpuTraceHelper.h"
-#include "NiagaraComponent.h"
+#include "NiagaraCompileHashVisitor.h"
 #include "NiagaraDataInterfaceUtilities.h"
 #include "NiagaraGpuComputeDispatchInterface.h"
-#include "NiagaraGpuComputeDispatch.h"
 #include "NiagaraSimStageData.h"
 #include "NiagaraShaderParametersBuilder.h"
-#include "NiagaraStats.h"
+#include "NiagaraSystem.h"
 #include "NiagaraTypes.h"
-#include "NiagaraWorldManager.h"
+#include "NiagaraSystemInstance.h"
 #include "RenderResource.h"
+#include "RenderGraphBuilder.h"
 #include "Shader.h"
 #include "ShaderCore.h"
 #include "ShaderCompilerCore.h"
-#include "ShaderParameterUtils.h"
 
 #include UE_INLINE_GENERATED_CPP_BY_NAME(NiagaraDataInterfaceAsyncGpuTrace)
 
diff --git a/Engine/Plugins/FX/Niagara/Source/Niagara/Private/DataInterface/NiagaraDataInterfaceDataChannelCommon.cpp b/Engine/Plugins/FX/Niagara/Source/Niagara/Private/DataInterface/NiagaraDataInterfaceDataChannelCommon.cpp
--- a/Engine/Plugins/FX/Niagara/Source/Niagara/Private/DataInterface/NiagaraDataInterfaceDataChannelCommon.cpp
+++ b/Engine/Plugins/FX/Niagara/Source/Niagara/Private/DataInterface/NiagaraDataInterfaceDataChannelCommon.cpp
@@ -1,6 +1,7 @@
 // Copyright Epic Games, Inc. All Rights Reserved.
 
 #include "DataInterface/NiagaraDataInterfaceDataChannelCommon.h"
+#include "Misc/LazySingleton.h"
 #include "NiagaraDataInterfaceUtilities.h"
 #include "NiagaraScript.h"
 
@@ -283,4 +284,4 @@ namespace NDIDataChannelUtilities
 	}
 }
 
-#undef LOCTEXT_NAMESPACE
\ No newline at end of file
+#undef LOCTEXT_NAMESPACE
diff --git a/Engine/Plugins/FX/Niagara/Source/Niagara/Private/DataInterface/NiagaraDataInterfaceDataChannelRead.cpp b/Engine/Plugins/FX/Niagara/Source/Niagara/Private/DataInterface/NiagaraDataInterfaceDataChannelRead.cpp
--- a/Engine/Plugins/FX/Niagara/Source/Niagara/Private/DataInterface/NiagaraDataInterfaceDataChannelRead.cpp
+++ b/Engine/Plugins/FX/Niagara/Source/Niagara/Private/DataInterface/NiagaraDataInterfaceDataChannelRead.cpp
@@ -2,12 +2,13 @@
 
 #include "DataInterface/NiagaraDataInterfaceDataChannelRead.h"
 
+#include "NiagaraCompileHashVisitor.h"
 #include "ShaderCompilerCore.h"
-#include "ShaderParameterUtils.h"
+#include "NiagaraModule.h"
 #include "NiagaraShaderParametersBuilder.h"
 
-#include "CoreMinimal.h"
 #include "NiagaraCommon.h"
+#include "NiagaraSystem.h"
 #include "NiagaraWorldManager.h"
 
 #include "NiagaraDataChannel.h"
@@ -1353,4 +1354,4 @@ int32 FNiagaraDataInterfaceProxy_DataChannelRead::PerInstanceDataPassedToRenderT
 	return sizeof(FNDIDataChannelReadInstanceData_RT);
 }
 
-#undef LOCTEXT_NAMESPACE
\ No newline at end of file
+#undef LOCTEXT_NAMESPACE
diff --git a/Engine/Plugins/FX/Niagara/Source/Niagara/Private/DataInterface/NiagaraDataInterfaceDataChannelSpawn.cpp b/Engine/Plugins/FX/Niagara/Source/Niagara/Private/DataInterface/NiagaraDataInterfaceDataChannelSpawn.cpp
--- a/Engine/Plugins/FX/Niagara/Source/Niagara/Private/DataInterface/NiagaraDataInterfaceDataChannelSpawn.cpp
+++ b/Engine/Plugins/FX/Niagara/Source/Niagara/Private/DataInterface/NiagaraDataInterfaceDataChannelSpawn.cpp
@@ -7,15 +7,14 @@ A specialized version of the Read DI that allows for direct spawning into emitte
 
 #include "NiagaraDataInterfaceDataChannelSpawn.h"
 
-#include "CoreMinimal.h"
 
 #include "NiagaraCommon.h"
 #include "NiagaraDataChannel.h"
-#include "NiagaraDataChannelHandler.h"
-#include "NiagaraDataChannelDefinitions.h"
 
 #include "NiagaraDataSetAccessor.h"
+#include "NiagaraEmitterInstance.h"
 #include "NiagaraSystemInstance.h"
+#include "NiagaraModule.h"
 
 
 #define LOCTEXT_NAMESPACE "NiagaraDataInterfaceDataChannelSpawn"
@@ -282,4 +281,4 @@ bool UNiagaraDataInterfaceDataChannelSpawn::CopyToInternal(UNiagaraDataInterface
 }
 
 
-#undef LOCTEXT_NAMESPACE
\ No newline at end of file
+#undef LOCTEXT_NAMESPACE
diff --git a/Engine/Plugins/FX/Niagara/Source/Niagara/Private/DataInterface/NiagaraDataInterfaceDataChannelSpawn.h b/Engine/Plugins/FX/Niagara/Source/Niagara/Private/DataInterface/NiagaraDataInterfaceDataChannelSpawn.h
--- a/Engine/Plugins/FX/Niagara/Source/Niagara/Private/DataInterface/NiagaraDataInterfaceDataChannelSpawn.h
+++ b/Engine/Plugins/FX/Niagara/Source/Niagara/Private/DataInterface/NiagaraDataInterfaceDataChannelSpawn.h
@@ -10,6 +10,10 @@ Specialization of the Read DI that allows Emitters to Spawn directly from Data C
 #include "NiagaraDataInterfaceEmitterBinding.h"
 #include "NiagaraDataInterfaceDataChannelSpawn.generated.h"
 
+enum class ENiagaraSimTarget : uint8;
+struct FNiagaraFunctionSignature;
+struct FVMExternalFunctionBindingInfo;
+
 
 UCLASS(Experimental, EditInlineNew, Category = "Data Channels", meta = (DisplayName = "Spawn From Data Channel"))
 class NIAGARA_API UNiagaraDataInterfaceDataChannelSpawn : public UNiagaraDataInterfaceDataChannelRead
@@ -44,4 +48,4 @@ class NIAGARA_API UNiagaraDataInterfaceDataChannelSpawn : public UNiagaraDataInt
 protected:
 	virtual bool CopyToInternal(UNiagaraDataInterface* Destination) const override;
 
-};
\ No newline at end of file
+};
diff --git a/Engine/Plugins/FX/Niagara/Source/Niagara/Private/DataInterface/NiagaraDataInterfaceDataChannelWrite.cpp b/Engine/Plugins/FX/Niagara/Source/Niagara/Private/DataInterface/NiagaraDataInterfaceDataChannelWrite.cpp
--- a/Engine/Plugins/FX/Niagara/Source/Niagara/Private/DataInterface/NiagaraDataInterfaceDataChannelWrite.cpp
+++ b/Engine/Plugins/FX/Niagara/Source/Niagara/Private/DataInterface/NiagaraDataInterfaceDataChannelWrite.cpp
@@ -2,15 +2,15 @@
 
 #include "DataInterface/NiagaraDataInterfaceDataChannelWrite.h"
 
-#include "CoreMinimal.h"
+#include "NiagaraModule.h"
 #include "Stats/Stats.h"
 #include "NiagaraCommon.h"
 #include "NiagaraShared.h"
 
+#include "NiagaraSystem.h"
 #include "NiagaraWorldManager.h"
 #include "NiagaraSystemInstance.h"
 
-#include "NiagaraSystemImpl.h"
 
 #include "NiagaraDataChannel.h"
 #include "NiagaraDataChannelHandler.h"
@@ -741,4 +741,4 @@ void UNiagaraDataInterfaceDataChannelWrite::SetShaderParameters(const FNiagaraDa
 {
 }
 
-#undef LOCTEXT_NAMESPACE
\ No newline at end of file
+#undef LOCTEXT_NAMESPACE
diff --git a/Engine/Plugins/FX/Niagara/Source/Niagara/Private/DataInterface/NiagaraDataInterfaceDebugDraw.cpp b/Engine/Plugins/FX/Niagara/Source/Niagara/Private/DataInterface/NiagaraDataInterfaceDebugDraw.cpp
--- a/Engine/Plugins/FX/Niagara/Source/Niagara/Private/DataInterface/NiagaraDataInterfaceDebugDraw.cpp
+++ b/Engine/Plugins/FX/Niagara/Source/Niagara/Private/DataInterface/NiagaraDataInterfaceDebugDraw.cpp
@@ -1,16 +1,15 @@
 // Copyright Epic Games, Inc. All Rights Reserved.
 
 #include "NiagaraDataInterfaceDebugDraw.h"
+#include "NiagaraCompileHashVisitor.h"
 #include "NiagaraTypes.h"
+#include "NiagaraEmitter.h"
 #include "NiagaraGpuComputeDispatchInterface.h"
 #include "NiagaraGpuComputeDebug.h"
-#include "NiagaraWorldManager.h"
 #include "NiagaraShaderParametersBuilder.h"
+#include "NiagaraSystem.h"
 #include "NiagaraSystemInstance.h"
 
-#include "Async/Async.h"
-#include "DrawDebugHelpers.h"
-#include "ShaderParameterUtils.h"
 
 #include UE_INLINE_GENERATED_CPP_BY_NAME(NiagaraDataInterfaceDebugDraw)
 
diff --git a/Engine/Plugins/FX/Niagara/Source/Niagara/Private/DataInterface/NiagaraDataInterfaceDebugDraw.h b/Engine/Plugins/FX/Niagara/Source/Niagara/Private/DataInterface/NiagaraDataInterfaceDebugDraw.h
--- a/Engine/Plugins/FX/Niagara/Source/Niagara/Private/DataInterface/NiagaraDataInterfaceDebugDraw.h
+++ b/Engine/Plugins/FX/Niagara/Source/Niagara/Private/DataInterface/NiagaraDataInterfaceDebugDraw.h
@@ -6,6 +6,12 @@
 #include "NiagaraDataInterface.h"
 #include "NiagaraDataInterfaceDebugDraw.generated.h"
 
+enum class ENiagaraSimTarget : uint8;
+struct FNiagaraDataInterfaceGeneratedFunction;
+struct FNiagaraFunctionSignature;
+struct FNiagaraScriptDataInterfaceCompileInfo;
+struct FVMExternalFunctionBindingInfo;
+
 UCLASS(EditInlineNew, Category = "Debug", meta = (DisplayName = "DebugDraw"))
 class NIAGARA_API UNiagaraDataInterfaceDebugDraw : public UNiagaraDataInterface
 {
diff --git a/Engine/Plugins/FX/Niagara/Source/Niagara/Private/DataInterface/NiagaraDataInterfaceEmitterProperties.cpp b/Engine/Plugins/FX/Niagara/Source/Niagara/Private/DataInterface/NiagaraDataInterfaceEmitterProperties.cpp
--- a/Engine/Plugins/FX/Niagara/Source/Niagara/Private/DataInterface/NiagaraDataInterfaceEmitterProperties.cpp
+++ b/Engine/Plugins/FX/Niagara/Source/Niagara/Private/DataInterface/NiagaraDataInterfaceEmitterProperties.cpp
@@ -1,14 +1,13 @@
 // Copyright Epic Games, Inc. All Rights Reserved.
 
 #include "NiagaraDataInterfaceEmitterProperties.h"
+#include "NiagaraCompileHashVisitor.h"
 #include "NiagaraEmitterInstance.h"
-#include "NiagaraComputeExecutionContext.h"
 #include "NiagaraShaderParametersBuilder.h"
 #include "NiagaraSystemInstance.h"
 
 #include "Internationalization/Internationalization.h"
 #include "ShaderCompilerCore.h"
-#include "ShaderParameterUtils.h"
 
 #include UE_INLINE_GENERATED_CPP_BY_NAME(NiagaraDataInterfaceEmitterProperties)
 
diff --git a/Engine/Plugins/FX/Niagara/Source/Niagara/Private/DataInterface/NiagaraDataInterfaceEmitterProperties.h b/Engine/Plugins/FX/Niagara/Source/Niagara/Private/DataInterface/NiagaraDataInterfaceEmitterProperties.h
--- a/Engine/Plugins/FX/Niagara/Source/Niagara/Private/DataInterface/NiagaraDataInterfaceEmitterProperties.h
+++ b/Engine/Plugins/FX/Niagara/Source/Niagara/Private/DataInterface/NiagaraDataInterfaceEmitterProperties.h
@@ -5,9 +5,13 @@
 #include "NiagaraShared.h"
 #include "NiagaraDataInterface.h"
 #include "NiagaraDataInterfaceEmitterBinding.h"
-#include "NiagaraParameterStore.h"
 #include "NiagaraDataInterfaceEmitterProperties.generated.h"
 
+enum class ENiagaraSimTarget : uint8;
+struct FNiagaraDataInterfaceGeneratedFunction;
+struct FNiagaraFunctionSignature;
+struct FVMExternalFunctionBindingInfo;
+
 /**
 Allows access to various emitter properties that are not part of the simulation data.
 */
diff --git a/Engine/Plugins/FX/Niagara/Source/Niagara/Private/DataInterface/NiagaraDataInterfaceGBuffer.cpp b/Engine/Plugins/FX/Niagara/Source/Niagara/Private/DataInterface/NiagaraDataInterfaceGBuffer.cpp
--- a/Engine/Plugins/FX/Niagara/Source/Niagara/Private/DataInterface/NiagaraDataInterfaceGBuffer.cpp
+++ b/Engine/Plugins/FX/Niagara/Source/Niagara/Private/DataInterface/NiagaraDataInterfaceGBuffer.cpp
@@ -1,16 +1,17 @@
 // Copyright Epic Games, Inc. All Rights Reserved.
 
 #include "NiagaraDataInterfaceGBuffer.h"
+#include "Containers/StridedView.h"
 #include "NiagaraGpuComputeDispatchInterface.h"
-#include "NiagaraGpuComputeDispatch.h"
+#include "NiagaraCompileHashVisitor.h"
 #include "NiagaraTypes.h"
 #include "NiagaraShaderParametersBuilder.h"
-#include "NiagaraSystemInstance.h"
-#include "NiagaraWorldManager.h"
 
 #include "Internationalization/Internationalization.h"
-#include "ShaderParameterUtils.h"
 #include "FXRenderingUtils.h"
+#include "RHIStaticStates.h"
+#include "RenderGraphResources.h"
+#include "SceneView.h"
 
 #include UE_INLINE_GENERATED_CPP_BY_NAME(NiagaraDataInterfaceGBuffer)
 
diff --git a/Engine/Plugins/FX/Niagara/Source/Niagara/Private/DataInterface/NiagaraDataInterfaceGBuffer.h b/Engine/Plugins/FX/Niagara/Source/Niagara/Private/DataInterface/NiagaraDataInterfaceGBuffer.h
--- a/Engine/Plugins/FX/Niagara/Source/Niagara/Private/DataInterface/NiagaraDataInterfaceGBuffer.h
+++ b/Engine/Plugins/FX/Niagara/Source/Niagara/Private/DataInterface/NiagaraDataInterfaceGBuffer.h
@@ -6,6 +6,8 @@
 #include "NiagaraDataInterface.h"
 #include "NiagaraDataInterfaceGBuffer.generated.h"
 
+struct FNiagaraDataInterfaceGeneratedFunction;
+
 UCLASS(EditInlineNew, Category = "Rendering", meta = (DisplayName = "GBuffer"))
 class NIAGARA_API UNiagaraDataInterfaceGBuffer : public UNiagaraDataInterface
 {
diff --git a/Engine/Plugins/FX/Niagara/Source/Niagara/Private/DataInterface/NiagaraDataInterfacePhysicsAsset.cpp b/Engine/Plugins/FX/Niagara/Source/Niagara/Private/DataInterface/NiagaraDataInterfacePhysicsAsset.cpp
--- a/Engine/Plugins/FX/Niagara/Source/Niagara/Private/DataInterface/NiagaraDataInterfacePhysicsAsset.cpp
+++ b/Engine/Plugins/FX/Niagara/Source/Niagara/Private/DataInterface/NiagaraDataInterfacePhysicsAsset.cpp
@@ -3,16 +3,15 @@
 #include "NiagaraDataInterfacePhysicsAsset.h"
 #include "Animation/SkeletalMeshActor.h"
 #include "Components/SkeletalMeshComponent.h"
-#include "SkeletalRenderPublic.h"
-#include "SkeletalMeshTypes.h"
 #include "AnimationRuntime.h"
-#include "NiagaraShader.h"
-#include "NiagaraComponent.h"
+#include "Engine/SkeletalMesh.h"
 #include "NiagaraRenderer.h"
+#include "NiagaraCompileHashVisitor.h"
 #include "NiagaraSimStageData.h"
 #include "NiagaraShaderParametersBuilder.h"
+#include "NiagaraSystem.h"
 #include "NiagaraSystemInstance.h"
-#include "ShaderParameterUtils.h"
+#include "PhysicsEngine/PhysicsAsset.h"
 
 #include UE_INLINE_GENERATED_CPP_BY_NAME(NiagaraDataInterfacePhysicsAsset)
 
diff --git a/Engine/Plugins/FX/Niagara/Source/Niagara/Private/DataInterface/NiagaraDataInterfaceSimCacheReader.cpp b/Engine/Plugins/FX/Niagara/Source/Niagara/Private/DataInterface/NiagaraDataInterfaceSimCacheReader.cpp
--- a/Engine/Plugins/FX/Niagara/Source/Niagara/Private/DataInterface/NiagaraDataInterfaceSimCacheReader.cpp
+++ b/Engine/Plugins/FX/Niagara/Source/Niagara/Private/DataInterface/NiagaraDataInterfaceSimCacheReader.cpp
@@ -1,6 +1,7 @@
 // Copyright Epic Games, Inc. All Rights Reserved.
 
 #include "NiagaraDataInterfaceSimCacheReader.h"
+#include "NiagaraCompileHashVisitor.h"
 #include "NiagaraDataInterfaceUtilities.h"
 #include "NiagaraRenderer.h"
 #include "NiagaraSimCache.h"
@@ -10,7 +11,6 @@
 #include "NiagaraSystemInstance.h"
 
 #include "Internationalization/Internationalization.h"
-#include "ShaderParameterUtils.h"
 
 #include UE_INLINE_GENERATED_CPP_BY_NAME(NiagaraDataInterfaceSimCacheReader)
 
diff --git a/Engine/Plugins/FX/Niagara/Source/Niagara/Private/DataInterface/NiagaraDataInterfaceSimCacheReader.h b/Engine/Plugins/FX/Niagara/Source/Niagara/Private/DataInterface/NiagaraDataInterfaceSimCacheReader.h
--- a/Engine/Plugins/FX/Niagara/Source/Niagara/Private/DataInterface/NiagaraDataInterfaceSimCacheReader.h
+++ b/Engine/Plugins/FX/Niagara/Source/Niagara/Private/DataInterface/NiagaraDataInterfaceSimCacheReader.h
@@ -6,6 +6,8 @@
 #include "NiagaraDataInterface.h"
 #include "NiagaraDataInterfaceSimCacheReader.generated.h"
 
+struct FNiagaraDataInterfaceGeneratedFunction;
+
 class UNiagaraSimCache;
 
 /**
diff --git a/Engine/Plugins/FX/Niagara/Source/Niagara/Private/DataInterface/NiagaraDataInterfaceSimpleCounter.cpp b/Engine/Plugins/FX/Niagara/Source/Niagara/Private/DataInterface/NiagaraDataInterfaceSimpleCounter.cpp
--- a/Engine/Plugins/FX/Niagara/Source/Niagara/Private/DataInterface/NiagaraDataInterfaceSimpleCounter.cpp
+++ b/Engine/Plugins/FX/Niagara/Source/Niagara/Private/DataInterface/NiagaraDataInterfaceSimpleCounter.cpp
@@ -2,6 +2,7 @@
 
 #include "NiagaraDataInterfaceSimpleCounter.h"
 #include "NiagaraClearCounts.h"
+#include "NiagaraCompileHashVisitor.h"
 #include "NiagaraGpuComputeDispatchInterface.h"
 #include "NiagaraGpuReadbackManager.h"
 #include "NiagaraShaderParametersBuilder.h"
@@ -10,7 +11,6 @@
 #include "Internationalization/Internationalization.h"
 #include "Async/Async.h"
 #include "RenderGraphUtils.h"
-#include "ShaderParameterUtils.h"
 #include "ShaderCompilerCore.h"
 
 #include UE_INLINE_GENERATED_CPP_BY_NAME(NiagaraDataInterfaceSimpleCounter)
diff --git a/Engine/Plugins/FX/Niagara/Source/Niagara/Private/DataInterface/NiagaraDataInterfaceStaticMesh.cpp b/Engine/Plugins/FX/Niagara/Source/Niagara/Private/DataInterface/NiagaraDataInterfaceStaticMesh.cpp
--- a/Engine/Plugins/FX/Niagara/Source/Niagara/Private/DataInterface/NiagaraDataInterfaceStaticMesh.cpp
+++ b/Engine/Plugins/FX/Niagara/Source/Niagara/Private/DataInterface/NiagaraDataInterfaceStaticMesh.cpp
@@ -1,25 +1,29 @@
 // Copyright Epic Games, Inc. All Rights Reserved.
 
 #include "DataInterface/NiagaraDataInterfaceStaticMesh.h"
-#include "NiagaraEmitterInstance.h"
-#include "NiagaraComponent.h"
-#include "NiagaraDistanceFieldHelper.h"
+#include "Components/StaticMeshComponent.h"
 #include "NiagaraStats.h"
+#include "Containers/StridedView.h"
 #include "NiagaraSystemInstance.h"
+#include "Engine/StaticMesh.h"
 #include "NiagaraRenderer.h"
+#include "Misc/LargeWorldRenderPosition.h"
 #include "NiagaraSettings.h"
-#include "NiagaraScript.h"
+#include "NiagaraCompileHashVisitor.h"
 #include "NiagaraShaderParametersBuilder.h"
+#include "NiagaraDataInterfaceMeshCommon.h"
 #include "NiagaraWorldManager.h"
 #include "NiagaraDataInterfaceUtilities.h"
+#include "NiagaraSystem.h"
 #include "SceneInterface.h"
-#include "ShaderParameterUtils.h"
+#include "SceneView.h"
 #include "StaticMeshComponentLODInfo.h"
 #include "NiagaraStats.h"
 #include "PrimitiveSceneInfo.h"
 
 #if WITH_EDITOR
 #include "Editor.h"
+#include "StaticMeshResources.h"
 #include "Subsystems/ImportSubsystem.h"
 #endif
 
@@ -31,7 +35,6 @@
 #include "ShaderCompilerCore.h"
 
 #include "NiagaraGpuComputeDispatchInterface.h"
-#include "NiagaraGpuComputeDispatch.h"
 #include "FXRenderingUtils.h"
 
 #include UE_INLINE_GENERATED_CPP_BY_NAME(NiagaraDataInterfaceStaticMesh)
diff --git a/Engine/Plugins/FX/Niagara/Source/Niagara/Private/DataInterface/NiagaraDataInterfaceUObjectPropertyReader.cpp b/Engine/Plugins/FX/Niagara/Source/Niagara/Private/DataInterface/NiagaraDataInterfaceUObjectPropertyReader.cpp
--- a/Engine/Plugins/FX/Niagara/Source/Niagara/Private/DataInterface/NiagaraDataInterfaceUObjectPropertyReader.cpp
+++ b/Engine/Plugins/FX/Niagara/Source/Niagara/Private/DataInterface/NiagaraDataInterfaceUObjectPropertyReader.cpp
@@ -1,13 +1,16 @@
 // Copyright Epic Games, Inc. All Rights Reserved.
 
 #include "NiagaraDataInterfaceUObjectPropertyReader.h"
+#include "NiagaraCompileHashVisitor.h"
 #include "NiagaraComputeExecutionContext.h"
+#include "NiagaraEmitterInstance.h"
 #include "NiagaraFunctionLibrary.h"
 #include "NiagaraShaderParametersBuilder.h"
 #include "NiagaraSystemInstance.h"
 
 #include "Internationalization/Internationalization.h"
-#include "ShaderParameterUtils.h"
+#include "UObject/EnumProperty.h"
+#include "UObject/WeakFieldPtr.h"
 
 #include UE_INLINE_GENERATED_CPP_BY_NAME(NiagaraDataInterfaceUObjectPropertyReader)
 
diff --git a/Engine/Plugins/FX/Niagara/Source/Niagara/Private/DataInterface/NiagaraDataInterfaceUObjectPropertyReader.h b/Engine/Plugins/FX/Niagara/Source/Niagara/Private/DataInterface/NiagaraDataInterfaceUObjectPropertyReader.h
--- a/Engine/Plugins/FX/Niagara/Source/Niagara/Private/DataInterface/NiagaraDataInterfaceUObjectPropertyReader.h
+++ b/Engine/Plugins/FX/Niagara/Source/Niagara/Private/DataInterface/NiagaraDataInterfaceUObjectPropertyReader.h
@@ -4,9 +4,10 @@
 #include "NiagaraCommon.h"
 #include "NiagaraShared.h"
 #include "NiagaraDataInterface.h"
-#include "NiagaraParameterStore.h"
 #include "NiagaraDataInterfaceUObjectPropertyReader.generated.h"
 
+struct FNiagaraDataInterfaceGeneratedFunction;
+
 USTRUCT()
 struct FNiagaraUObjectPropertyReaderRemap
 {
diff --git a/Engine/Plugins/FX/Niagara/Source/Niagara/Private/DataInterface/NiagaraDataInterfaceUtilities.cpp b/Engine/Plugins/FX/Niagara/Source/Niagara/Private/DataInterface/NiagaraDataInterfaceUtilities.cpp
--- a/Engine/Plugins/FX/Niagara/Source/Niagara/Private/DataInterface/NiagaraDataInterfaceUtilities.cpp
+++ b/Engine/Plugins/FX/Niagara/Source/Niagara/Private/DataInterface/NiagaraDataInterfaceUtilities.cpp
@@ -5,6 +5,7 @@
 #include "NiagaraComputeExecutionContext.h"
 #include "NiagaraConstants.h"
 #include "NiagaraDataInterface.h"
+#include "NiagaraEmitterInstance.h"
 #include "NiagaraSystemImpl.h"
 #include "NiagaraSystemInstance.h"
 #include "NiagaraSystemSimulation.h"
diff --git a/Engine/Plugins/FX/Niagara/Source/Niagara/Private/DataInterface/NiagaraDataInterfaceVirtualTexture.cpp b/Engine/Plugins/FX/Niagara/Source/Niagara/Private/DataInterface/NiagaraDataInterfaceVirtualTexture.cpp
--- a/Engine/Plugins/FX/Niagara/Source/Niagara/Private/DataInterface/NiagaraDataInterfaceVirtualTexture.cpp
+++ b/Engine/Plugins/FX/Niagara/Source/Niagara/Private/DataInterface/NiagaraDataInterfaceVirtualTexture.cpp
@@ -1,10 +1,8 @@
 // Copyright Epic Games, Inc. All Rights Reserved.
 
 #include "NiagaraDataInterfaceVirtualTexture.h"
-#include "NiagaraComputeExecutionContext.h"
+#include "NiagaraCompileHashVisitor.h"
 #include "NiagaraCustomVersion.h"
-#include "NiagaraGpuComputeDispatchInterface.h"
-#include "NiagaraShader.h"
 #include "NiagaraShaderParametersBuilder.h"
 #include "NiagaraSystemInstance.h"
 
diff --git a/Engine/Plugins/FX/Niagara/Source/Niagara/Private/Experimental/NiagaraDataInterfaceSkeletalMeshUvMapping.cpp b/Engine/Plugins/FX/Niagara/Source/Niagara/Private/Experimental/NiagaraDataInterfaceSkeletalMeshUvMapping.cpp
--- a/Engine/Plugins/FX/Niagara/Source/Niagara/Private/Experimental/NiagaraDataInterfaceSkeletalMeshUvMapping.cpp
+++ b/Engine/Plugins/FX/Niagara/Source/Niagara/Private/Experimental/NiagaraDataInterfaceSkeletalMeshUvMapping.cpp
@@ -4,8 +4,7 @@
 
 #include "Experimental/NiagaraMeshUvMapping.h"
 #include "NDISkeletalMeshCommon.h"
-#include "NiagaraResourceArrayWriter.h"
-#include "NiagaraStats.h"
+#include "Rendering/SkeletalMeshRenderData.h"
 
 template<bool UseFullPrecisionUv>
 struct FSkelMeshVertexAccessorHelper
@@ -182,4 +181,4 @@ void FSkeletalMeshUvMapping::BuildQuadTree()
 			BuildHelper.Build();
 		}
 	}
-}
\ No newline at end of file
+}
diff --git a/Engine/Plugins/FX/Niagara/Source/Niagara/Private/Experimental/NiagaraDataInterfaceSkeletalMeshUvMapping.h b/Engine/Plugins/FX/Niagara/Source/Niagara/Private/Experimental/NiagaraDataInterfaceSkeletalMeshUvMapping.h
--- a/Engine/Plugins/FX/Niagara/Source/Niagara/Private/Experimental/NiagaraDataInterfaceSkeletalMeshUvMapping.h
+++ b/Engine/Plugins/FX/Niagara/Source/Niagara/Private/Experimental/NiagaraDataInterfaceSkeletalMeshUvMapping.h
@@ -2,11 +2,14 @@
 #pragma once
 
 #include "Experimental/NiagaraMeshUvMapping.h"
-#include "NiagaraDataInterfaceSkeletalMesh.h"
 #include "NiagaraUvQuadTree.h"
 #include "RenderResource.h"
+#include "UObject/WeakObjectPtr.h"
 #include "UObject/WeakObjectPtrTemplates.h"
 
+class FSkeletalMeshLODRenderData;
+class USkeletalMesh;
+
 struct FSkeletalMeshUvMapping : public FMeshUvMapping
 {
 	FSkeletalMeshUvMapping(TWeakObjectPtr<USkeletalMesh> InMeshObject, int32 InLodIndex, int32 InUvSetIndex);
@@ -27,4 +30,4 @@ private:
 	virtual void BuildQuadTree() override;
 
 	TWeakObjectPtr<USkeletalMesh> MeshObject;
-};
\ No newline at end of file
+};
diff --git a/Engine/Plugins/FX/Niagara/Source/Niagara/Private/Experimental/NiagaraMeshUvMapping.cpp b/Engine/Plugins/FX/Niagara/Source/Niagara/Private/Experimental/NiagaraMeshUvMapping.cpp
--- a/Engine/Plugins/FX/Niagara/Source/Niagara/Private/Experimental/NiagaraMeshUvMapping.cpp
+++ b/Engine/Plugins/FX/Niagara/Source/Niagara/Private/Experimental/NiagaraMeshUvMapping.cpp
@@ -5,6 +5,7 @@
 #include "Experimental/NiagaraMeshUvMappingHandle.h"
 #include "NiagaraResourceArrayWriter.h"
 #include "NiagaraStats.h"
+#include "RenderingThread.h"
 
 DECLARE_CYCLE_STAT(TEXT("Niagara - UvQuadTree Cpu"), STAT_Niagara_UvMapping_Cpu, STATGROUP_Niagara);
 DECLARE_CYCLE_STAT(TEXT("Niagara - UvQuadTree Gpu"), STAT_Niagara_UvMapping_Gpu, STATGROUP_Niagara);
@@ -220,4 +221,4 @@ int32 FMeshUvMappingHandleBase::GetLodIndex() const
 void FMeshUvMappingHandleBase::PinAndInvalidateHandle()
 {
 	UvMappingData.Reset();
-}
\ No newline at end of file
+}
diff --git a/Engine/Plugins/FX/Niagara/Source/Niagara/Private/Experimental/NiagaraMeshUvMapping.h b/Engine/Plugins/FX/Niagara/Source/Niagara/Private/Experimental/NiagaraMeshUvMapping.h
--- a/Engine/Plugins/FX/Niagara/Source/Niagara/Private/Experimental/NiagaraMeshUvMapping.h
+++ b/Engine/Plugins/FX/Niagara/Source/Niagara/Private/Experimental/NiagaraMeshUvMapping.h
@@ -2,13 +2,13 @@
 
 #pragma once
 
-#include "CoreMinimal.h"
-
-#include "Experimental/NiagaraMeshUvMappingHandle.h"
+#include "Containers/DynamicRHIResourceArray.h"
+#include "Math/Vector.h"
 #include "NiagaraUvQuadTree.h"
 #include "RenderResource.h"
-#include "RHI.h"
-#include "Containers/DynamicRHIResourceArray.h"
+#include "RHIResources.h"
+
+struct FMeshUvMappingUsage;
 
 struct FMeshUvMapping;
 class FNiagaraUvQuadTree;
@@ -227,4 +227,4 @@ protected:
 
 	std::atomic<int32> CpuQuadTreeUserCount = {0};
 	std::atomic<int32> GpuQuadTreeUserCount = {0};
-};
\ No newline at end of file
+};
diff --git a/Engine/Plugins/FX/Niagara/Source/Niagara/Private/MovieScene/MovieSceneNiagaraSystemSpawnSection.cpp b/Engine/Plugins/FX/Niagara/Source/Niagara/Private/MovieScene/MovieSceneNiagaraSystemSpawnSection.cpp
--- a/Engine/Plugins/FX/Niagara/Source/Niagara/Private/MovieScene/MovieSceneNiagaraSystemSpawnSection.cpp
+++ b/Engine/Plugins/FX/Niagara/Source/Niagara/Private/MovieScene/MovieSceneNiagaraSystemSpawnSection.cpp
@@ -1,6 +1,7 @@
 // Copyright Epic Games, Inc. All Rights Reserved.
 
 #include "MovieScene/MovieSceneNiagaraSystemSpawnSection.h"
+#include "NiagaraCommon.h"
 
 #include UE_INLINE_GENERATED_CPP_BY_NAME(MovieSceneNiagaraSystemSpawnSection)
 
diff --git a/Engine/Plugins/FX/Niagara/Source/Niagara/Private/MovieScene/MovieSceneNiagaraSystemTrackTemplate.cpp b/Engine/Plugins/FX/Niagara/Source/Niagara/Private/MovieScene/MovieSceneNiagaraSystemTrackTemplate.cpp
--- a/Engine/Plugins/FX/Niagara/Source/Niagara/Private/MovieScene/MovieSceneNiagaraSystemTrackTemplate.cpp
+++ b/Engine/Plugins/FX/Niagara/Source/Niagara/Private/MovieScene/MovieSceneNiagaraSystemTrackTemplate.cpp
@@ -1,7 +1,9 @@
 // Copyright Epic Games, Inc. All Rights Reserved.
 
 #include "MovieSceneNiagaraSystemTrackTemplate.h"
+#include "Evaluation/MovieSceneExecutionTokens.h"
 #include "MovieSceneExecutionToken.h"
+#include "MovieScene/MovieSceneNiagaraSystemSpawnSection.h"
 #include "NiagaraComponent.h"
 #include "NiagaraSystemInstanceController.h"
 #include "IMovieScenePlayer.h"
diff --git a/Engine/Plugins/FX/Niagara/Source/Niagara/Private/MovieScene/MovieSceneNiagaraSystemTrackTemplate.h b/Engine/Plugins/FX/Niagara/Source/Niagara/Private/MovieScene/MovieSceneNiagaraSystemTrackTemplate.h
--- a/Engine/Plugins/FX/Niagara/Source/Niagara/Private/MovieScene/MovieSceneNiagaraSystemTrackTemplate.h
+++ b/Engine/Plugins/FX/Niagara/Source/Niagara/Private/MovieScene/MovieSceneNiagaraSystemTrackTemplate.h
@@ -2,12 +2,15 @@
 
 #pragma once
 
-#include "NiagaraCommon.h"
-#include "MovieScene/MovieSceneNiagaraSystemSpawnSection.h"
 #include "Evaluation/MovieSceneEvalTemplate.h"
 #include "Evaluation/MovieSceneTrackImplementation.h"
 #include "MovieSceneNiagaraSystemTrackTemplate.generated.h"
 
+enum class ENiagaraAgeUpdateMode : uint8;
+enum class ENiagaraSystemSpawnSectionEndBehavior;
+enum class ENiagaraSystemSpawnSectionEvaluateBehavior;
+enum class ENiagaraSystemSpawnSectionStartBehavior;
+
 USTRUCT()
 struct FMovieSceneNiagaraSystemTrackTemplate : public FMovieSceneEvalTemplate
 {
@@ -57,4 +60,4 @@ private:
 	UPROPERTY()
 	bool bAllowScalability;
 
-};
\ No newline at end of file
+};
diff --git a/Engine/Plugins/FX/Niagara/Source/Niagara/Private/MovieScene/Parameters/MovieSceneNiagaraBoolParameterSectionTemplate.cpp b/Engine/Plugins/FX/Niagara/Source/Niagara/Private/MovieScene/Parameters/MovieSceneNiagaraBoolParameterSectionTemplate.cpp
--- a/Engine/Plugins/FX/Niagara/Source/Niagara/Private/MovieScene/Parameters/MovieSceneNiagaraBoolParameterSectionTemplate.cpp
+++ b/Engine/Plugins/FX/Niagara/Source/Niagara/Private/MovieScene/Parameters/MovieSceneNiagaraBoolParameterSectionTemplate.cpp
@@ -1,9 +1,7 @@
 // Copyright Epic Games, Inc. All Rights Reserved.
 
 #include "MovieSceneNiagaraBoolParameterSectionTemplate.h"
-#include "NiagaraComponent.h"
 #include "NiagaraTypes.h"
-#include "IMovieScenePlayer.h"
 
 #include UE_INLINE_GENERATED_CPP_BY_NAME(MovieSceneNiagaraBoolParameterSectionTemplate)
 
diff --git a/Engine/Plugins/FX/Niagara/Source/Niagara/Private/MovieScene/Parameters/MovieSceneNiagaraBoolParameterSectionTemplate.h b/Engine/Plugins/FX/Niagara/Source/Niagara/Private/MovieScene/Parameters/MovieSceneNiagaraBoolParameterSectionTemplate.h
--- a/Engine/Plugins/FX/Niagara/Source/Niagara/Private/MovieScene/Parameters/MovieSceneNiagaraBoolParameterSectionTemplate.h
+++ b/Engine/Plugins/FX/Niagara/Source/Niagara/Private/MovieScene/Parameters/MovieSceneNiagaraBoolParameterSectionTemplate.h
@@ -6,6 +6,9 @@
 #include "Channels/MovieSceneBoolChannel.h"
 #include "MovieSceneNiagaraBoolParameterSectionTemplate.generated.h"
 
+struct FNiagaraVariable;
+struct FNiagaraVariableBase;
+
 USTRUCT()
 struct FMovieSceneNiagaraBoolParameterSectionTemplate : public FMovieSceneNiagaraParameterSectionTemplate
 {
@@ -25,4 +28,4 @@ protected:
 private:
 	UPROPERTY()
 	FMovieSceneBoolChannel BoolChannel;
-};
\ No newline at end of file
+};
diff --git a/Engine/Plugins/FX/Niagara/Source/Niagara/Private/MovieScene/Parameters/MovieSceneNiagaraColorParameterSectionTemplate.cpp b/Engine/Plugins/FX/Niagara/Source/Niagara/Private/MovieScene/Parameters/MovieSceneNiagaraColorParameterSectionTemplate.cpp
--- a/Engine/Plugins/FX/Niagara/Source/Niagara/Private/MovieScene/Parameters/MovieSceneNiagaraColorParameterSectionTemplate.cpp
+++ b/Engine/Plugins/FX/Niagara/Source/Niagara/Private/MovieScene/Parameters/MovieSceneNiagaraColorParameterSectionTemplate.cpp
@@ -1,9 +1,7 @@
 // Copyright Epic Games, Inc. All Rights Reserved.
 
 #include "MovieSceneNiagaraColorParameterSectionTemplate.h"
-#include "NiagaraComponent.h"
 #include "NiagaraTypes.h"
-#include "IMovieScenePlayer.h"
 
 #include UE_INLINE_GENERATED_CPP_BY_NAME(MovieSceneNiagaraColorParameterSectionTemplate)
 
diff --git a/Engine/Plugins/FX/Niagara/Source/Niagara/Private/MovieScene/Parameters/MovieSceneNiagaraColorParameterSectionTemplate.h b/Engine/Plugins/FX/Niagara/Source/Niagara/Private/MovieScene/Parameters/MovieSceneNiagaraColorParameterSectionTemplate.h
--- a/Engine/Plugins/FX/Niagara/Source/Niagara/Private/MovieScene/Parameters/MovieSceneNiagaraColorParameterSectionTemplate.h
+++ b/Engine/Plugins/FX/Niagara/Source/Niagara/Private/MovieScene/Parameters/MovieSceneNiagaraColorParameterSectionTemplate.h
@@ -6,6 +6,9 @@
 #include "Channels/MovieSceneFloatChannel.h"
 #include "MovieSceneNiagaraColorParameterSectionTemplate.generated.h"
 
+struct FNiagaraVariable;
+struct FNiagaraVariableBase;
+
 USTRUCT()
 struct FMovieSceneNiagaraColorParameterSectionTemplate : public FMovieSceneNiagaraParameterSectionTemplate
 {
@@ -34,4 +37,4 @@ private:
 
 	UPROPERTY()
 	FMovieSceneFloatChannel AlphaChannel;
-};
\ No newline at end of file
+};
diff --git a/Engine/Plugins/FX/Niagara/Source/Niagara/Private/MovieScene/Parameters/MovieSceneNiagaraParameterSectionTemplate.cpp b/Engine/Plugins/FX/Niagara/Source/Niagara/Private/MovieScene/Parameters/MovieSceneNiagaraParameterSectionTemplate.cpp
--- a/Engine/Plugins/FX/Niagara/Source/Niagara/Private/MovieScene/Parameters/MovieSceneNiagaraParameterSectionTemplate.cpp
+++ b/Engine/Plugins/FX/Niagara/Source/Niagara/Private/MovieScene/Parameters/MovieSceneNiagaraParameterSectionTemplate.cpp
@@ -1,8 +1,8 @@
 // Copyright Epic Games, Inc. All Rights Reserved.
 
 #include "MovieSceneNiagaraParameterSectionTemplate.h"
+#include "Evaluation/MovieSceneExecutionTokens.h"
 #include "NiagaraComponent.h"
-#include "NiagaraEffectType.h"
 #include "NiagaraSystem.h"
 #include "NiagaraTypes.h"
 #include "IMovieScenePlayer.h"
diff --git a/Engine/Plugins/FX/Niagara/Source/Niagara/Private/NDISkeletalMesh_BoneSampling.cpp b/Engine/Plugins/FX/Niagara/Source/Niagara/Private/NDISkeletalMesh_BoneSampling.cpp
--- a/Engine/Plugins/FX/Niagara/Source/Niagara/Private/NDISkeletalMesh_BoneSampling.cpp
+++ b/Engine/Plugins/FX/Niagara/Source/Niagara/Private/NDISkeletalMesh_BoneSampling.cpp
@@ -1,11 +1,11 @@
 // Copyright Epic Games, Inc. All Rights Reserved.
 
 #include "NiagaraDataInterfaceSkeletalMesh.h"
-#include "NiagaraComponent.h"
 #include "Components/SkeletalMeshComponent.h"
 #include "SkeletalMeshTypes.h"
 #include "NiagaraStats.h"
 #include "NDISkeletalMeshCommon.h"
+#include "Templates/IntegralConstant.h"
 
 #define LOCTEXT_NAMESPACE "NiagaraDataInterfaceSkeletalMesh_BoneSampling"
 
diff --git a/Engine/Plugins/FX/Niagara/Source/Niagara/Private/NDISkeletalMesh_TriangleSampling.cpp b/Engine/Plugins/FX/Niagara/Source/Niagara/Private/NDISkeletalMesh_TriangleSampling.cpp
--- a/Engine/Plugins/FX/Niagara/Source/Niagara/Private/NDISkeletalMesh_TriangleSampling.cpp
+++ b/Engine/Plugins/FX/Niagara/Source/Niagara/Private/NDISkeletalMesh_TriangleSampling.cpp
@@ -2,15 +2,14 @@
 
 #include "NiagaraDataInterfaceSkeletalMesh.h"
 
-#include "Animation/SkeletalMeshActor.h"
 #include "Components/SkeletalMeshComponent.h"
 #include "NDISkeletalMeshCommon.h"
-#include "NiagaraComponent.h"
-#include "NiagaraDataInterfaceArrayInt.h"
 #include "Experimental/NiagaraDataInterfaceSkeletalMeshUvMapping.h"
+#include "NiagaraDataInterfaceMeshCommon.h"
 #include "NiagaraStats.h"
 #include "SkeletalMeshTypes.h"
 #include "Templates/AlignmentTemplates.h"
+#include "Templates/IntegralConstant.h"
 
 #define LOCTEXT_NAMESPACE "NiagaraDataInterfaceSkeletalMesh_TriangleSampling"
 
diff --git a/Engine/Plugins/FX/Niagara/Source/Niagara/Private/NDISkeletalMesh_VertexSampling.cpp b/Engine/Plugins/FX/Niagara/Source/Niagara/Private/NDISkeletalMesh_VertexSampling.cpp
--- a/Engine/Plugins/FX/Niagara/Source/Niagara/Private/NDISkeletalMesh_VertexSampling.cpp
+++ b/Engine/Plugins/FX/Niagara/Source/Niagara/Private/NDISkeletalMesh_VertexSampling.cpp
@@ -1,8 +1,6 @@
 // Copyright Epic Games, Inc. All Rights Reserved.
 
 #include "NiagaraDataInterfaceSkeletalMesh.h"
-#include "NiagaraComponent.h"
-#include "Animation/SkeletalMeshActor.h"
 #include "Components/SkeletalMeshComponent.h"
 #include "SkeletalMeshTypes.h"
 #include "NiagaraStats.h"
diff --git a/Engine/Plugins/FX/Niagara/Source/Niagara/Private/NiagaraActor.cpp b/Engine/Plugins/FX/Niagara/Source/Niagara/Private/NiagaraActor.cpp
--- a/Engine/Plugins/FX/Niagara/Source/Niagara/Private/NiagaraActor.cpp
+++ b/Engine/Plugins/FX/Niagara/Source/Niagara/Private/NiagaraActor.cpp
@@ -8,7 +8,6 @@
 #include "Components/BillboardComponent.h"
 #include "NiagaraComponent.h"
 #include "NiagaraSystem.h"
-#include "NiagaraEffectType.h"
 
 #include UE_INLINE_GENERATED_CPP_BY_NAME(NiagaraActor)
 
diff --git a/Engine/Plugins/FX/Niagara/Source/Niagara/Private/NiagaraAsyncCompile.cpp b/Engine/Plugins/FX/Niagara/Source/Niagara/Private/NiagaraAsyncCompile.cpp
--- a/Engine/Plugins/FX/Niagara/Source/Niagara/Private/NiagaraAsyncCompile.cpp
+++ b/Engine/Plugins/FX/Niagara/Source/Niagara/Private/NiagaraAsyncCompile.cpp
@@ -10,6 +10,7 @@
 #include "NiagaraPrecompileContainer.h"
 #include "NiagaraScriptSourceBase.h"
 #include "NiagaraStats.h"
+#include "NiagaraSystem.h"
 #include "UObject/Package.h"
 #include "VectorVM.h"
 
diff --git a/Engine/Plugins/FX/Niagara/Source/Niagara/Private/NiagaraAsyncCompile.h b/Engine/Plugins/FX/Niagara/Source/Niagara/Private/NiagaraAsyncCompile.h
--- a/Engine/Plugins/FX/Niagara/Source/Niagara/Private/NiagaraAsyncCompile.h
+++ b/Engine/Plugins/FX/Niagara/Source/Niagara/Private/NiagaraAsyncCompile.h
@@ -2,8 +2,6 @@
 
 #pragma once
 
-#include "CoreMinimal.h"
-
 #include "NiagaraScript.h"
 
 #include "NiagaraAsyncCompile.generated.h"
diff --git a/Engine/Plugins/FX/Niagara/Source/Niagara/Private/NiagaraAsyncGpuTraceHelper.cpp b/Engine/Plugins/FX/Niagara/Source/Niagara/Private/NiagaraAsyncGpuTraceHelper.cpp
--- a/Engine/Plugins/FX/Niagara/Source/Niagara/Private/NiagaraAsyncGpuTraceHelper.cpp
+++ b/Engine/Plugins/FX/Niagara/Source/Niagara/Private/NiagaraAsyncGpuTraceHelper.cpp
@@ -3,9 +3,12 @@
 #include "NiagaraAsyncGpuTraceHelper.h"
 
 #include "Components/PrimitiveComponent.h"
+#include "NiagaraDataInterface.h"
 #include "NiagaraGpuComputeDispatchInterface.h"
 #include "NiagaraStats.h"
+#include "PrimitiveSceneInfo.h"
 #include "PrimitiveSceneProxy.h"
+#include "RenderingThread.h"
 
 int32 GNiagaraAsyncTraceScratchBucketSize = 1024;
 static FAutoConsoleVariableRef CVarNiagaraAsyncTraceScratchBucketSize(
@@ -349,4 +352,4 @@ void FNiagaraAsyncGpuTraceHelper::OnPrimitiveGPUSceneInstancesDirtied()
 {
 	bCollisionGroupMapDirty = true;
 }
-#endif
\ No newline at end of file
+#endif
diff --git a/Engine/Plugins/FX/Niagara/Source/Niagara/Private/NiagaraAsyncGpuTraceHelper.h b/Engine/Plugins/FX/Niagara/Source/Niagara/Private/NiagaraAsyncGpuTraceHelper.h
--- a/Engine/Plugins/FX/Niagara/Source/Niagara/Private/NiagaraAsyncGpuTraceHelper.h
+++ b/Engine/Plugins/FX/Niagara/Source/Niagara/Private/NiagaraAsyncGpuTraceHelper.h
@@ -1,20 +1,16 @@
 // Copyright Epic Games, Inc. All Rights Reserved.
 #pragma once
 
-#include "CoreMinimal.h"
-
 #include "NiagaraAsyncGpuTraceProvider.h"
 #include "NiagaraGpuScratchPad.h"
 #include "NiagaraSettings.h"
-
+#include "PrimitiveComponentId.h"
 #include "RHIDefinitions.h"
-
-#include "PrimitiveSceneInfo.h"
-#include "RHI.h"
 #include "RHIUtilities.h"
 
 struct FNiagaraDataInterfaceProxy;
 class FNiagaraGpuComputeDispatchInterface;
+class FPrimitiveSceneInfo;
 class FRayTracingPipelineState;
 class FRHICommandList;
 class FRHIRayTracingScene;
diff --git a/Engine/Plugins/FX/Niagara/Source/Niagara/Private/NiagaraBakerOutput.cpp b/Engine/Plugins/FX/Niagara/Source/Niagara/Private/NiagaraBakerOutput.cpp
--- a/Engine/Plugins/FX/Niagara/Source/Niagara/Private/NiagaraBakerOutput.cpp
+++ b/Engine/Plugins/FX/Niagara/Source/Niagara/Private/NiagaraBakerOutput.cpp
@@ -2,14 +2,15 @@
 
 #include "NiagaraBakerOutput.h"
 #include "Misc/Paths.h"
-#include "NiagaraEffectType.h"
 #include "NiagaraSystem.h"
 
 #include UE_INLINE_GENERATED_CPP_BY_NAME(NiagaraBakerOutput)
 
 #if WITH_EDITOR
+#include "AssetRegistry/AssetData.h"
 #include "AssetRegistry/IAssetRegistry.h"
 #include "Factories/Factory.h"
+#include "IAssetTools.h"
 #include "Misc/MessageDialog.h"
 #include "AssetToolsModule.h"
 #endif //WITH_EDITOR
diff --git a/Engine/Plugins/FX/Niagara/Source/Niagara/Private/NiagaraBakerOutputSimCache.cpp b/Engine/Plugins/FX/Niagara/Source/Niagara/Private/NiagaraBakerOutputSimCache.cpp
--- a/Engine/Plugins/FX/Niagara/Source/Niagara/Private/NiagaraBakerOutputSimCache.cpp
+++ b/Engine/Plugins/FX/Niagara/Source/Niagara/Private/NiagaraBakerOutputSimCache.cpp
@@ -1,7 +1,6 @@
 // Copyright Epic Games, Inc. All Rights Reserved.
 
 #include "NiagaraBakerOutputSimCache.h"
-#include "NiagaraBakerSettings.h"
 
 #include UE_INLINE_GENERATED_CPP_BY_NAME(NiagaraBakerOutputSimCache)
 
diff --git a/Engine/Plugins/FX/Niagara/Source/Niagara/Private/NiagaraBakerOutputVolumeTexture.cpp b/Engine/Plugins/FX/Niagara/Source/Niagara/Private/NiagaraBakerOutputVolumeTexture.cpp
--- a/Engine/Plugins/FX/Niagara/Source/Niagara/Private/NiagaraBakerOutputVolumeTexture.cpp
+++ b/Engine/Plugins/FX/Niagara/Source/Niagara/Private/NiagaraBakerOutputVolumeTexture.cpp
@@ -1,7 +1,6 @@
 // Copyright Epic Games, Inc. All Rights Reserved.
 
 #include "NiagaraBakerOutputVolumeTexture.h"
-#include "NiagaraBakerSettings.h"
 
 #include UE_INLINE_GENERATED_CPP_BY_NAME(NiagaraBakerOutputVolumeTexture)
 
diff --git a/Engine/Plugins/FX/Niagara/Source/Niagara/Private/NiagaraBakerSettings.cpp b/Engine/Plugins/FX/Niagara/Source/Niagara/Private/NiagaraBakerSettings.cpp
--- a/Engine/Plugins/FX/Niagara/Source/Niagara/Private/NiagaraBakerSettings.cpp
+++ b/Engine/Plugins/FX/Niagara/Source/Niagara/Private/NiagaraBakerSettings.cpp
@@ -1,13 +1,14 @@
 // Copyright Epic Games, Inc. All Rights Reserved.
 
 #include "NiagaraBakerSettings.h"
-#include "Framework/Application/SlateApplication.h"
+#include "EngineDefines.h"
 #include "NiagaraBakerOutputTexture2D.h"
+#include "Math/InverseRotationMatrix.h"
 #include "NiagaraCustomVersion.h"
-#include "NiagaraEffectType.h"
-#include "NiagaraSystem.h"
 #include "Engine/Texture2D.h"
+#include "Math/OrthoMatrix.h"
 #include "Misc/PathViews.h"
+#include "Math/PerspectiveMatrix.h"
 #include "UObject/Package.h"
 
 #include UE_INLINE_GENERATED_CPP_BY_NAME(NiagaraBakerSettings)
diff --git a/Engine/Plugins/FX/Niagara/Source/Niagara/Private/NiagaraCollision.cpp b/Engine/Plugins/FX/Niagara/Source/Niagara/Private/NiagaraCollision.cpp
--- a/Engine/Plugins/FX/Niagara/Source/Niagara/Private/NiagaraCollision.cpp
+++ b/Engine/Plugins/FX/Niagara/Source/Niagara/Private/NiagaraCollision.cpp
@@ -2,10 +2,9 @@
 
 #include "NiagaraCollision.h"
 #include "Engine/World.h"
-#include "NiagaraEmitterInstance.h"
-#include "NiagaraSystemInstance.h"
-#include "NiagaraComponent.h"
+#include "HAL/IConsoleManager.h"
 #include "NiagaraStats.h"
+#include "Misc/ScopeRWLock.h"
 #include "PhysicalMaterials/PhysicalMaterial.h"
 
 DECLARE_CYCLE_STAT(TEXT("Event Emission"), STAT_NiagaraEventWrite, STATGROUP_Niagara);
@@ -190,4 +189,4 @@ bool FNiagaraDICollisionQueryBatch::GetQueryResult(uint32 InTraceID, FNiagaraDIC
 	}
 
 	return false;
-}
\ No newline at end of file
+}
diff --git a/Engine/Plugins/FX/Niagara/Source/Niagara/Private/NiagaraCommon.cpp b/Engine/Plugins/FX/Niagara/Source/Niagara/Private/NiagaraCommon.cpp
--- a/Engine/Plugins/FX/Niagara/Source/Niagara/Private/NiagaraCommon.cpp
+++ b/Engine/Plugins/FX/Niagara/Source/Niagara/Private/NiagaraCommon.cpp
@@ -8,7 +8,9 @@
 #include "NiagaraComponent.h"
 #include "NiagaraConstants.h"
 #include "NiagaraCustomVersion.h"
+#include "NiagaraEmitter.h"
 #include "NiagaraParameterCollection.h"
+#include "NiagaraParameterDefinitionsBase.h"
 #include "NiagaraScriptSourceBase.h"
 #include "NiagaraSettings.h"
 #include "NiagaraStats.h"
@@ -17,6 +19,7 @@
 #include "NiagaraWorldManager.h"
 #include "String/ParseTokens.h"
 #include "UObject/Class.h"
+#include "UObject/UObjectIterator.h"
 
 DECLARE_CYCLE_STAT(TEXT("Niagara - Utilities - PrepareRapidIterationParameters"), STAT_Niagara_Utilities_PrepareRapidIterationParameters, STATGROUP_Niagara);
 
@@ -1622,4 +1625,4 @@ void FNiagaraFunctionSignature::GetVariadicOutputs(TArray<FNiagaraVariableBase>&
 
 		OutVariadicOutputs.Emplace(Param);
 	}
-}
\ No newline at end of file
+}
diff --git a/Engine/Plugins/FX/Niagara/Source/Niagara/Private/NiagaraComponent.cpp b/Engine/Plugins/FX/Niagara/Source/Niagara/Private/NiagaraComponent.cpp
--- a/Engine/Plugins/FX/Niagara/Source/Niagara/Private/NiagaraComponent.cpp
+++ b/Engine/Plugins/FX/Niagara/Source/Niagara/Private/NiagaraComponent.cpp
@@ -2,10 +2,10 @@
 
 #include "NiagaraComponent.h"
 #include "Engine/CollisionProfile.h"
+#include "Engine/Level.h"
 #include "Engine/TextureRenderTarget.h"
-#include "EngineUtils.h"
-#include "Materials/MaterialInstanceDynamic.h"
 #include "NiagaraCommon.h"
+#include "NiagaraComponentPool.h"
 #include "NiagaraComponentSettings.h"
 #include "NiagaraConstants.h"
 #include "NiagaraCustomVersion.h"
@@ -14,16 +14,20 @@
 #include "NiagaraGpuComputeDispatchInterface.h"
 #include "NiagaraRenderer.h"
 #include "NiagaraSceneProxy.h"
+#include "NiagaraSimCache.h"
 #include "NiagaraStats.h"
 #include "NiagaraSystem.h"
 #include "NiagaraSystemInstanceController.h"
 #include "NiagaraWorldManager.h"
 #include "PrimitiveSceneInfo.h"
+#include "PrimitiveViewRelevance.h"
 #include "ProfilingDebugging/CsvProfiler.h"
+#include "SceneManagement.h"
 #include "UObject/NameTypes.h"
 #include "Engine/StaticMesh.h"
 #include "NiagaraCullProxyComponent.h"
 #include "SceneInterface.h"
+#include "UObject/UObjectIterator.h"
 
 #include UE_INLINE_GENERATED_CPP_BY_NAME(NiagaraComponent)
 
diff --git a/Engine/Plugins/FX/Niagara/Source/Niagara/Private/NiagaraComponentRendererProperties.cpp b/Engine/Plugins/FX/Niagara/Source/Niagara/Private/NiagaraComponentRendererProperties.cpp
--- a/Engine/Plugins/FX/Niagara/Source/Niagara/Private/NiagaraComponentRendererProperties.cpp
+++ b/Engine/Plugins/FX/Niagara/Source/Niagara/Private/NiagaraComponentRendererProperties.cpp
@@ -1,7 +1,9 @@
 // Copyright Epic Games, Inc. All Rights Reserved.
 
 #include "NiagaraComponentRendererProperties.h"
+#include "Misc/PackageName.h"
 #include "NiagaraEmitterInstance.h"
+#include "NiagaraModule.h"
 #include "NiagaraRenderer.h"
 #include "NiagaraConstants.h"
 #include "NiagaraRendererComponents.h"
@@ -12,7 +14,6 @@
 #include "Widgets/Images/SImage.h"
 #include "Styling/SlateIconFinder.h"
 #include "Widgets/SWidget.h"
-#include "AssetThumbnail.h"
 #include "Widgets/Text/STextBlock.h"
 #endif
 #include "NiagaraCustomVersion.h"
diff --git a/Engine/Plugins/FX/Niagara/Source/Niagara/Private/NiagaraComputeExecutionContext.cpp b/Engine/Plugins/FX/Niagara/Source/Niagara/Private/NiagaraComputeExecutionContext.cpp
--- a/Engine/Plugins/FX/Niagara/Source/Niagara/Private/NiagaraComputeExecutionContext.cpp
+++ b/Engine/Plugins/FX/Niagara/Source/Niagara/Private/NiagaraComputeExecutionContext.cpp
@@ -1,18 +1,15 @@
 // Copyright Epic Games, Inc. All Rights Reserved.
 
 #include "NiagaraComputeExecutionContext.h"
-#include "NiagaraStats.h"
 #include "NiagaraGpuComputeDispatchInterface.h"
 #include "NiagaraGpuComputeDispatch.h"
 #include "NiagaraDataInterface.h"
+#include "NiagaraScript.h"
 #include "NiagaraSystemInstance.h"
-#include "NiagaraSystemGpuComputeProxy.h"
-#include "NiagaraWorldManager.h"
-#include "NiagaraEmitterInstance.h"
-#include "NiagaraFunctionLibrary.h"
 #include "NiagaraGPUInstanceCountManager.h"
 #include "NiagaraGPUSystemTick.h"
 #include "NiagaraDataInterfaceRW.h"
+#include "NiagaraShader.h"
 
 FNiagaraComputeExecutionContext::FNiagaraComputeExecutionContext()
 	: MainDataSet(nullptr)
diff --git a/Engine/Plugins/FX/Niagara/Source/Niagara/Private/NiagaraConstants.cpp b/Engine/Plugins/FX/Niagara/Source/Niagara/Private/NiagaraConstants.cpp
--- a/Engine/Plugins/FX/Niagara/Source/Niagara/Private/NiagaraConstants.cpp
+++ b/Engine/Plugins/FX/Niagara/Source/Niagara/Private/NiagaraConstants.cpp
@@ -1,6 +1,7 @@
 // Copyright Epic Games, Inc. All Rights Reserved.
 
 #include "NiagaraConstants.h"
+#include "NiagaraCommon.h"
 #include "NiagaraModule.h"
 
 #include "Internationalization/Internationalization.h"
diff --git a/Engine/Plugins/FX/Niagara/Source/Niagara/Private/NiagaraCrashReporterHandler.cpp b/Engine/Plugins/FX/Niagara/Source/Niagara/Private/NiagaraCrashReporterHandler.cpp
--- a/Engine/Plugins/FX/Niagara/Source/Niagara/Private/NiagaraCrashReporterHandler.cpp
+++ b/Engine/Plugins/FX/Niagara/Source/Niagara/Private/NiagaraCrashReporterHandler.cpp
@@ -2,13 +2,10 @@
 
 #include "NiagaraCrashReporterHandler.h"
 
-#include "CoreMinimal.h"
-#include "HAL/ThreadManager.h"
 #include "GenericPlatform/GenericPlatformCrashContext.h"
 #include "Misc/ScopeLock.h"
 #include "NiagaraSystemInstance.h"
 #include "NiagaraSystemSimulation.h"
-#include "NiagaraComponent.h"
 #include "NiagaraSystem.h"
 
 #if WITH_NIAGARA_CRASHREPORTER
diff --git a/Engine/Plugins/FX/Niagara/Source/Niagara/Private/NiagaraDataChannel.cpp b/Engine/Plugins/FX/Niagara/Source/Niagara/Private/NiagaraDataChannel.cpp
--- a/Engine/Plugins/FX/Niagara/Source/Niagara/Private/NiagaraDataChannel.cpp
+++ b/Engine/Plugins/FX/Niagara/Source/Niagara/Private/NiagaraDataChannel.cpp
@@ -2,7 +2,6 @@
 
 #include "NiagaraDataChannel.h"
 
-#include "CoreMinimal.h"
 #include "NiagaraWorldManager.h"
 #include "NiagaraDataChannelHandler.h"
 #include "NiagaraDataChannelManager.h"
@@ -606,4 +605,4 @@ UNiagaraDataChannelHandler* UNiagaraDataChannelLibrary::GetNiagaraDataChannel(co
 }
 
 
-#undef LOCTEXT_NAMESPACE
\ No newline at end of file
+#undef LOCTEXT_NAMESPACE
diff --git a/Engine/Plugins/FX/Niagara/Source/Niagara/Private/NiagaraDataChannelAccessor.h b/Engine/Plugins/FX/Niagara/Source/Niagara/Private/NiagaraDataChannelAccessor.h
--- a/Engine/Plugins/FX/Niagara/Source/Niagara/Private/NiagaraDataChannelAccessor.h
+++ b/Engine/Plugins/FX/Niagara/Source/Niagara/Private/NiagaraDataChannelAccessor.h
@@ -5,6 +5,9 @@
 #include "NiagaraDataChannelCommon.h"
 #include "NiagaraDataChannelAccessor.generated.h"
 
+class UNiagaraDataChannelHandler;
+struct FNiagaraSpawnInfo;
+struct FNiagaraVariableBase;
 
 /** 
 Initial simple API for reading and writing data in a data channel from game code / BP. 
@@ -109,4 +112,4 @@ public:
 
 	UFUNCTION(BlueprintCallable, Category = NiagaraDataChannel, meta = (Keywords = "niagara DataChannel"))
 	void WriteSpawnInfo(FName VarName, int32 Index, FNiagaraSpawnInfo InData);
-};
\ No newline at end of file
+};
diff --git a/Engine/Plugins/FX/Niagara/Source/Niagara/Private/NiagaraDataChannelDefinitions.cpp b/Engine/Plugins/FX/Niagara/Source/Niagara/Private/NiagaraDataChannelDefinitions.cpp
--- a/Engine/Plugins/FX/Niagara/Source/Niagara/Private/NiagaraDataChannelDefinitions.cpp
+++ b/Engine/Plugins/FX/Niagara/Source/Niagara/Private/NiagaraDataChannelDefinitions.cpp
@@ -1,6 +1,7 @@
 // Copyright Epic Games, Inc. All Rights Reserved.
 
 #include "NiagaraDataChannelDefinitions.h"
+#include "Modules/ModuleManager.h"
 #include "NiagaraWorldManager.h"
 #include "NiagaraDataChannel.h"
 #include "NiagaraDataChannelManager.h"
@@ -13,6 +14,8 @@
 #include "AssetRegistry/AssetData.h"
 
 #include "Misc/MessageDialog.h"
+#include "NiagaraSystem.h"
+#include "UObject/UObjectIterator.h"
 
 #define LOCTEXT_NAMESPACE "NiagaraDataChannelDefinitions"
 
@@ -259,4 +262,4 @@ const UNiagaraDataChannel* UNiagaraDataChannelDefinitions::FindDataChannel(FName
 	return nullptr;
 }
 
-#undef LOCTEXT_NAMESPACE
\ No newline at end of file
+#undef LOCTEXT_NAMESPACE
diff --git a/Engine/Plugins/FX/Niagara/Source/Niagara/Private/NiagaraDataChannelManager.cpp b/Engine/Plugins/FX/Niagara/Source/Niagara/Private/NiagaraDataChannelManager.cpp
--- a/Engine/Plugins/FX/Niagara/Source/Niagara/Private/NiagaraDataChannelManager.cpp
+++ b/Engine/Plugins/FX/Niagara/Source/Niagara/Private/NiagaraDataChannelManager.cpp
@@ -2,6 +2,7 @@
 
 #include "NiagaraDataChannelManager.h"
 
+#include "NiagaraModule.h"
 #include "NiagaraWorldManager.h"
 #include "NiagaraDataChannel.h"
 #include "NiagaraDataChannelHandler.h"
diff --git a/Engine/Plugins/FX/Niagara/Source/Niagara/Private/NiagaraDataChannelManager.h b/Engine/Plugins/FX/Niagara/Source/Niagara/Private/NiagaraDataChannelManager.h
--- a/Engine/Plugins/FX/Niagara/Source/Niagara/Private/NiagaraDataChannelManager.h
+++ b/Engine/Plugins/FX/Niagara/Source/Niagara/Private/NiagaraDataChannelManager.h
@@ -2,8 +2,12 @@
 
 #pragma once
 
-#include "NiagaraDataChannelCommon.h"
+#include "UObject/ObjectPtr.h"
 
+class FNiagaraWorldManager;
+class UNiagaraDataChannel;
+class UNiagaraDataChannelHandler;
+enum ETickingGroup : int;
 
 /**
 Manager class for all Niagara DataChannels at the System/World level.
@@ -36,4 +40,4 @@ private:
 
 	/** Runtime handlers for each DataChannel channel. */
 	TMap<FName, TObjectPtr<UNiagaraDataChannelHandler>> Channels;
-};
\ No newline at end of file
+};
diff --git a/Engine/Plugins/FX/Niagara/Source/Niagara/Private/NiagaraDataChannel_Global.cpp b/Engine/Plugins/FX/Niagara/Source/Niagara/Private/NiagaraDataChannel_Global.cpp
--- a/Engine/Plugins/FX/Niagara/Source/Niagara/Private/NiagaraDataChannel_Global.cpp
+++ b/Engine/Plugins/FX/Niagara/Source/Niagara/Private/NiagaraDataChannel_Global.cpp
@@ -3,8 +3,6 @@
 #include "NiagaraDataChannel_Global.h"
 
 #include "NiagaraSystemInstance.h"
-#include "NiagaraDataChannelManager.h"
-#include "DataInterface/NiagaraDataInterfaceDataChannelWrite.h"
 #include "NiagaraGpuComputeDispatchInterface.h"
 #include "NiagaraWorldManager.h"
 
diff --git a/Engine/Plugins/FX/Niagara/Source/Niagara/Private/NiagaraDataInterface.cpp b/Engine/Plugins/FX/Niagara/Source/Niagara/Private/NiagaraDataInterface.cpp
--- a/Engine/Plugins/FX/Niagara/Source/Niagara/Private/NiagaraDataInterface.cpp
+++ b/Engine/Plugins/FX/Niagara/Source/Niagara/Private/NiagaraDataInterface.cpp
@@ -1,16 +1,15 @@
 // Copyright Epic Games, Inc. All Rights Reserved.
 
 #include "NiagaraDataInterface.h"
-#include "Curves/CurveVector.h"
-#include "Curves/CurveLinearColor.h"
 #include "DataDrivenShaderPlatformInfo.h"
+#include "NiagaraCompileHashVisitor.h"
 #include "NiagaraTypes.h"
-#include "ShaderParameterUtils.h"
 #include "NiagaraGPUSystemTick.h"
 #include "NiagaraGpuComputeDispatch.h"
 #include "NiagaraShader.h"
 #include "NiagaraSystem.h"
 #include "NiagaraComponent.h"
+#include "NiagaraSystemGpuComputeProxy.h"
 #include "ShaderCompilerCore.h"
 #include "NiagaraDataInterfaceUtilities.h"
 
diff --git a/Engine/Plugins/FX/Niagara/Source/Niagara/Private/NiagaraDataInterface2DArrayTexture.cpp b/Engine/Plugins/FX/Niagara/Source/Niagara/Private/NiagaraDataInterface2DArrayTexture.cpp
--- a/Engine/Plugins/FX/Niagara/Source/Niagara/Private/NiagaraDataInterface2DArrayTexture.cpp
+++ b/Engine/Plugins/FX/Niagara/Source/Niagara/Private/NiagaraDataInterface2DArrayTexture.cpp
@@ -1,9 +1,8 @@
 // Copyright Epic Games, Inc. All Rights Reserved.
 
 #include "NiagaraDataInterface2DArrayTexture.h"
-#include "NiagaraComputeExecutionContext.h"
+#include "NiagaraCompileHashVisitor.h"
 #include "NiagaraGpuComputeDispatchInterface.h"
-#include "NiagaraShader.h"
 #include "NiagaraShaderParametersBuilder.h"
 #include "NiagaraSystemInstance.h"
 
diff --git a/Engine/Plugins/FX/Niagara/Source/Niagara/Private/NiagaraDataInterfaceArrayFloat.cpp b/Engine/Plugins/FX/Niagara/Source/Niagara/Private/NiagaraDataInterfaceArrayFloat.cpp
--- a/Engine/Plugins/FX/Niagara/Source/Niagara/Private/NiagaraDataInterfaceArrayFloat.cpp
+++ b/Engine/Plugins/FX/Niagara/Source/Niagara/Private/NiagaraDataInterfaceArrayFloat.cpp
@@ -1,6 +1,5 @@
 // Copyright Epic Games, Inc. All Rights Reserved.
 
 #include "NiagaraDataInterfaceArrayFloat.h"
-#include "NiagaraRenderer.h"
 
 #include UE_INLINE_GENERATED_CPP_BY_NAME(NiagaraDataInterfaceArrayFloat)
diff --git a/Engine/Plugins/FX/Niagara/Source/Niagara/Private/NiagaraDataInterfaceArrayInt.cpp b/Engine/Plugins/FX/Niagara/Source/Niagara/Private/NiagaraDataInterfaceArrayInt.cpp
--- a/Engine/Plugins/FX/Niagara/Source/Niagara/Private/NiagaraDataInterfaceArrayInt.cpp
+++ b/Engine/Plugins/FX/Niagara/Source/Niagara/Private/NiagaraDataInterfaceArrayInt.cpp
@@ -1,7 +1,6 @@
 // Copyright Epic Games, Inc. All Rights Reserved.
 
 #include "NiagaraDataInterfaceArrayInt.h"
-#include "NiagaraRenderer.h"
 
 #include UE_INLINE_GENERATED_CPP_BY_NAME(NiagaraDataInterfaceArrayInt)
 
diff --git a/Engine/Plugins/FX/Niagara/Source/Niagara/Private/NiagaraDataInterfaceAudio.cpp b/Engine/Plugins/FX/Niagara/Source/Niagara/Private/NiagaraDataInterfaceAudio.cpp
--- a/Engine/Plugins/FX/Niagara/Source/Niagara/Private/NiagaraDataInterfaceAudio.cpp
+++ b/Engine/Plugins/FX/Niagara/Source/Niagara/Private/NiagaraDataInterfaceAudio.cpp
@@ -3,6 +3,7 @@
 #include "NiagaraDataInterfaceAudio.h"
 
 #include "AudioDevice.h"
+#include "AudioDeviceManager.h"
 
 #include UE_INLINE_GENERATED_CPP_BY_NAME(NiagaraDataInterfaceAudio)
 
diff --git a/Engine/Plugins/FX/Niagara/Source/Niagara/Private/NiagaraDataInterfaceAudioOscilloscope.cpp b/Engine/Plugins/FX/Niagara/Source/Niagara/Private/NiagaraDataInterfaceAudioOscilloscope.cpp
--- a/Engine/Plugins/FX/Niagara/Source/Niagara/Private/NiagaraDataInterfaceAudioOscilloscope.cpp
+++ b/Engine/Plugins/FX/Niagara/Source/Niagara/Private/NiagaraDataInterfaceAudioOscilloscope.cpp
@@ -1,20 +1,15 @@
 // Copyright Epic Games, Inc. All Rights Reserved.
 
 #include "NiagaraDataInterfaceAudioOscilloscope.h"
+#include "AudioDevice.h"
 #include "NiagaraTypes.h"
-#include "NiagaraCustomVersion.h"
+#include "AudioDeviceManager.h"
 #include "AudioResampler.h"
-#include "NiagaraShader.h"
-#include "ShaderParameterUtils.h"
-#include "ClearQuad.h"
-#include "TextureResource.h"
-#include "Engine/Texture2D.h"
+#include "NiagaraCompileHashVisitor.h"
 #include "NiagaraRenderer.h"
+#include "NiagaraDataInterfaceAudio.h"
 #include "NiagaraShaderParametersBuilder.h"
-#include "NiagaraSystemInstance.h"
-#include "Engine/TextureRenderTarget2D.h"
-#include "Engine/Engine.h"
-#include "NiagaraComponent.h"
+#include "RenderingThread.h"
 
 #include UE_INLINE_GENERATED_CPP_BY_NAME(NiagaraDataInterfaceAudioOscilloscope)
 
diff --git a/Engine/Plugins/FX/Niagara/Source/Niagara/Private/NiagaraDataInterfaceAudioPlayer.cpp b/Engine/Plugins/FX/Niagara/Source/Niagara/Private/NiagaraDataInterfaceAudioPlayer.cpp
--- a/Engine/Plugins/FX/Niagara/Source/Niagara/Private/NiagaraDataInterfaceAudioPlayer.cpp
+++ b/Engine/Plugins/FX/Niagara/Source/Niagara/Private/NiagaraDataInterfaceAudioPlayer.cpp
@@ -3,7 +3,6 @@
 #include "NiagaraDataInterfaceAudioPlayer.h"
 
 #include "NiagaraTypes.h"
-#include "NiagaraCustomVersion.h"
 #include "NiagaraStats.h"
 #include "Internationalization/Internationalization.h"
 #include "NiagaraSystemInstance.h"
diff --git a/Engine/Plugins/FX/Niagara/Source/Niagara/Private/NiagaraDataInterfaceAudioSpectrum.cpp b/Engine/Plugins/FX/Niagara/Source/Niagara/Private/NiagaraDataInterfaceAudioSpectrum.cpp
--- a/Engine/Plugins/FX/Niagara/Source/Niagara/Private/NiagaraDataInterfaceAudioSpectrum.cpp
+++ b/Engine/Plugins/FX/Niagara/Source/Niagara/Private/NiagaraDataInterfaceAudioSpectrum.cpp
@@ -1,20 +1,18 @@
 // Copyright Epic Games, Inc. All Rights Reserved.
 
 #include "NiagaraDataInterfaceAudioSpectrum.h"
+#include "DSP/AudioFFT.h"
 #include "NiagaraTypes.h"
-#include "NiagaraCustomVersion.h"
+#include "NiagaraCompileHashVisitor.h"
 #include "NiagaraRenderer.h"
-#include "NiagaraShader.h"
 #include "NiagaraShaderParametersBuilder.h"
-#include "ShaderParameterUtils.h"
-#include "Engine/Engine.h"
-#include "NiagaraComponent.h"
 
 #include "DSP/DeinterleaveView.h"
 #include "DSP/SlidingWindow.h"
 #include "DSP/ConstantQ.h"
 #include "DSP/FFTAlgorithm.h"
 #include "DSP/FloatArrayMath.h"
+#include "RenderingThread.h"
 
 #include UE_INLINE_GENERATED_CPP_BY_NAME(NiagaraDataInterfaceAudioSpectrum)
 
diff --git a/Engine/Plugins/FX/Niagara/Source/Niagara/Private/NiagaraDataInterfaceBindingInstance.cpp b/Engine/Plugins/FX/Niagara/Source/Niagara/Private/NiagaraDataInterfaceBindingInstance.cpp
--- a/Engine/Plugins/FX/Niagara/Source/Niagara/Private/NiagaraDataInterfaceBindingInstance.cpp
+++ b/Engine/Plugins/FX/Niagara/Source/Niagara/Private/NiagaraDataInterfaceBindingInstance.cpp
@@ -1,6 +1,3 @@
 // Copyright Epic Games, Inc. All Rights Reserved.
 
 #include "NiagaraDataInterfaceBindingInstance.h"
-#include "NiagaraTypes.h"
-#include "NiagaraCommon.h"
-#include "NiagaraDataInterface.h"
diff --git a/Engine/Plugins/FX/Niagara/Source/Niagara/Private/NiagaraDataInterfaceCamera.cpp b/Engine/Plugins/FX/Niagara/Source/Niagara/Private/NiagaraDataInterfaceCamera.cpp
--- a/Engine/Plugins/FX/Niagara/Source/Niagara/Private/NiagaraDataInterfaceCamera.cpp
+++ b/Engine/Plugins/FX/Niagara/Source/Niagara/Private/NiagaraDataInterfaceCamera.cpp
@@ -3,15 +3,13 @@
 #include "NiagaraDataInterfaceCamera.h"
 
 #include "Engine/LocalPlayer.h"
+#include "NiagaraCompileHashVisitor.h"
 #include "NiagaraTypes.h"
 #include "NiagaraWorldManager.h"
 #include "Internationalization/Internationalization.h"
-#include "NiagaraGPUSystemTick.h"
 #include "NiagaraShaderParametersBuilder.h"
-#include "NiagaraSystemGpuComputeProxy.h"
 #include "NiagaraSystemInstance.h"
 #include "GameFramework/PlayerController.h"
-#include "ShaderParameterUtils.h"
 
 #include UE_INLINE_GENERATED_CPP_BY_NAME(NiagaraDataInterfaceCamera)
 
diff --git a/Engine/Plugins/FX/Niagara/Source/Niagara/Private/NiagaraDataInterfaceCollisionQuery.cpp b/Engine/Plugins/FX/Niagara/Source/Niagara/Private/NiagaraDataInterfaceCollisionQuery.cpp
--- a/Engine/Plugins/FX/Niagara/Source/Niagara/Private/NiagaraDataInterfaceCollisionQuery.cpp
+++ b/Engine/Plugins/FX/Niagara/Source/Niagara/Private/NiagaraDataInterfaceCollisionQuery.cpp
@@ -4,19 +4,17 @@
 
 #include "FXRenderingUtils.h"
 #include "GlobalDistanceFieldParameters.h"
-#include "NiagaraAsyncGpuTraceHelper.h"
+#include "NiagaraCompileHashVisitor.h"
 #include "NiagaraDistanceFieldHelper.h"
-#include "NiagaraComponent.h"
+#include "NiagaraEmitter.h"
 #include "NiagaraGpuComputeDispatchInterface.h"
-#include "NiagaraGpuComputeDispatch.h"
 #include "NiagaraShaderParametersBuilder.h"
-#include "NiagaraStats.h"
+#include "NiagaraSystem.h"
 #include "NiagaraTypes.h"
-#include "NiagaraWorldManager.h"
+#include "NiagaraSystemInstance.h"
 #include "Shader.h"
 #include "ShaderCore.h"
 #include "ShaderCompilerCore.h"
-#include "ShaderParameterUtils.h"
 
 #include UE_INLINE_GENERATED_CPP_BY_NAME(NiagaraDataInterfaceCollisionQuery)
 
diff --git a/Engine/Plugins/FX/Niagara/Source/Niagara/Private/NiagaraDataInterfaceColorCurve.cpp b/Engine/Plugins/FX/Niagara/Source/Niagara/Private/NiagaraDataInterfaceColorCurve.cpp
--- a/Engine/Plugins/FX/Niagara/Source/Niagara/Private/NiagaraDataInterfaceColorCurve.cpp
+++ b/Engine/Plugins/FX/Niagara/Source/Niagara/Private/NiagaraDataInterfaceColorCurve.cpp
@@ -1,9 +1,6 @@
 // Copyright Epic Games, Inc. All Rights Reserved.
 
 #include "NiagaraDataInterfaceColorCurve.h"
-#include "Curves/CurveVector.h"
-#include "Curves/CurveLinearColor.h"
-#include "Curves/CurveFloat.h"
 #include "NiagaraTypes.h"
 #include "UObject/Package.h"
 
diff --git a/Engine/Plugins/FX/Niagara/Source/Niagara/Private/NiagaraDataInterfaceCubeTexture.cpp b/Engine/Plugins/FX/Niagara/Source/Niagara/Private/NiagaraDataInterfaceCubeTexture.cpp
--- a/Engine/Plugins/FX/Niagara/Source/Niagara/Private/NiagaraDataInterfaceCubeTexture.cpp
+++ b/Engine/Plugins/FX/Niagara/Source/Niagara/Private/NiagaraDataInterfaceCubeTexture.cpp
@@ -1,9 +1,8 @@
 // Copyright Epic Games, Inc. All Rights Reserved.
 
 #include "NiagaraDataInterfaceCubeTexture.h"
-#include "NiagaraComputeExecutionContext.h"
+#include "NiagaraCompileHashVisitor.h"
 #include "NiagaraGpuComputeDispatchInterface.h"
-#include "NiagaraShader.h"
 #include "NiagaraShaderParametersBuilder.h"
 #include "NiagaraSystemInstance.h"
 
diff --git a/Engine/Plugins/FX/Niagara/Source/Niagara/Private/NiagaraDataInterfaceCurlNoise.cpp b/Engine/Plugins/FX/Niagara/Source/Niagara/Private/NiagaraDataInterfaceCurlNoise.cpp
--- a/Engine/Plugins/FX/Niagara/Source/Niagara/Private/NiagaraDataInterfaceCurlNoise.cpp
+++ b/Engine/Plugins/FX/Niagara/Source/Niagara/Private/NiagaraDataInterfaceCurlNoise.cpp
@@ -2,10 +2,9 @@
 
 #include "NiagaraDataInterfaceCurlNoise.h"
 #include "NiagaraCompileHashVisitor.h"
-#include "NiagaraShader.h"
 #include "NiagaraShaderParametersBuilder.h"
-#include "ShaderParameterUtils.h"
 #include "Math/IntVector.h"
+#include "RenderingThread.h"
 
 #include UE_INLINE_GENERATED_CPP_BY_NAME(NiagaraDataInterfaceCurlNoise)
 
diff --git a/Engine/Plugins/FX/Niagara/Source/Niagara/Private/NiagaraDataInterfaceCurve.cpp b/Engine/Plugins/FX/Niagara/Source/Niagara/Private/NiagaraDataInterfaceCurve.cpp
--- a/Engine/Plugins/FX/Niagara/Source/Niagara/Private/NiagaraDataInterfaceCurve.cpp
+++ b/Engine/Plugins/FX/Niagara/Source/Niagara/Private/NiagaraDataInterfaceCurve.cpp
@@ -1,9 +1,6 @@
 // Copyright Epic Games, Inc. All Rights Reserved.
 
 #include "NiagaraDataInterfaceCurve.h"
-#include "Curves/CurveVector.h"
-#include "Curves/CurveLinearColor.h"
-#include "Curves/CurveFloat.h"
 #include "NiagaraTypes.h"
 #include "UObject/Package.h"
 
diff --git a/Engine/Plugins/FX/Niagara/Source/Niagara/Private/NiagaraDataInterfaceCurveBase.cpp b/Engine/Plugins/FX/Niagara/Source/Niagara/Private/NiagaraDataInterfaceCurveBase.cpp
--- a/Engine/Plugins/FX/Niagara/Source/Niagara/Private/NiagaraDataInterfaceCurveBase.cpp
+++ b/Engine/Plugins/FX/Niagara/Source/Niagara/Private/NiagaraDataInterfaceCurveBase.cpp
@@ -1,15 +1,13 @@
 // Copyright Epic Games, Inc. All Rights Reserved.
 
 #include "NiagaraDataInterfaceCurveBase.h"
-#include "Curves/CurveVector.h"
 #include "Engine/Texture2D.h"
 #include "Internationalization/Internationalization.h"
-#include "ShaderParameterUtils.h"
 
+#include "NiagaraCompileHashVisitor.h"
 #include "NiagaraCustomVersion.h"
 #include "NiagaraEmitter.h"
 #include "NiagaraRenderer.h"
-#include "NiagaraShader.h"
 #include "NiagaraShaderParametersBuilder.h"
 #include "NiagaraStats.h"
 #include "NiagaraSystemStaticBuffers.h"
diff --git a/Engine/Plugins/FX/Niagara/Source/Niagara/Private/NiagaraDataInterfaceExport.cpp b/Engine/Plugins/FX/Niagara/Source/Niagara/Private/NiagaraDataInterfaceExport.cpp
--- a/Engine/Plugins/FX/Niagara/Source/Niagara/Private/NiagaraDataInterfaceExport.cpp
+++ b/Engine/Plugins/FX/Niagara/Source/Niagara/Private/NiagaraDataInterfaceExport.cpp
@@ -1,20 +1,20 @@
 // Copyright Epic Games, Inc. All Rights Reserved.
 
 #include "NiagaraDataInterfaceExport.h"
+#include "NiagaraCompileHashVisitor.h"
 #include "NiagaraTypes.h"
-#include "NiagaraCustomVersion.h"
 #include "NiagaraComputeExecutionContext.h"
 #include "NiagaraGpuComputeDispatchInterface.h"
 #include "NiagaraGpuReadbackManager.h"
 #include "NiagaraGPUSystemTick.h"
 #include "NiagaraShaderParametersBuilder.h"
+#include "NiagaraSystem.h"
 #include "NiagaraSystemInstance.h"
 
 #include "Internationalization/Internationalization.h"
 #include "RenderGraphBuilder.h"
 #include "RenderGraphUtils.h"
 #include "ShaderCompilerCore.h"
-#include "ShaderParameterUtils.h"
 #include "Misc/LargeWorldRenderPosition.h"
 
 #include UE_INLINE_GENERATED_CPP_BY_NAME(NiagaraDataInterfaceExport)
diff --git a/Engine/Plugins/FX/Niagara/Source/Niagara/Private/NiagaraDataInterfaceGrid2DCollection.cpp b/Engine/Plugins/FX/Niagara/Source/Niagara/Private/NiagaraDataInterfaceGrid2DCollection.cpp
--- a/Engine/Plugins/FX/Niagara/Source/Niagara/Private/NiagaraDataInterfaceGrid2DCollection.cpp
+++ b/Engine/Plugins/FX/Niagara/Source/Niagara/Private/NiagaraDataInterfaceGrid2DCollection.cpp
@@ -1,26 +1,29 @@
 // Copyright Epic Games, Inc. All Rights Reserved.
 #include "NiagaraDataInterfaceGrid2DCollection.h"
+#include "CanvasTypes.h"
 #include "NiagaraBatchedElements.h"
+#include "GlobalRenderResources.h"
 #include "NiagaraConstants.h"
+#include "NiagaraCompileHashVisitor.h"
 #include "NiagaraGpuComputeDebugInterface.h"
+#include "NiagaraComponent.h"
 #include "NiagaraGpuComputeDispatchInterface.h"
+#include "NiagaraEmitter.h"
 #include "NiagaraRenderer.h"
+#include "NiagaraRenderGraphUtils.h"
 #include "NiagaraSettings.h"
-#include "NiagaraShader.h"
 #include "NiagaraShaderParametersBuilder.h"
 #include "NiagaraSystemInstance.h"
 #include "NiagaraSimStageData.h"
 #include "NiagaraSystemInstanceController.h"
 
-#include "Engine/Canvas.h"
 #include "Engine/TextureRenderTarget2D.h"
 #include "Engine/TextureRenderTarget2DArray.h"
-#include "ClearQuad.h"
 #include "CanvasItem.h"
+#include "RHIStaticStates.h"
 #include "RenderGraphBuilder.h"
 #include "RenderGraphUtils.h"
 #include "ShaderParameterMetadataBuilder.h"
-#include "ShaderParameterUtils.h"
 #include "TextureResource.h"
 
 #include "NiagaraDataInterfaceGrid2DCollectionReader.h"
diff --git a/Engine/Plugins/FX/Niagara/Source/Niagara/Private/NiagaraDataInterfaceGrid2DCollectionReader.cpp b/Engine/Plugins/FX/Niagara/Source/Niagara/Private/NiagaraDataInterfaceGrid2DCollectionReader.cpp
--- a/Engine/Plugins/FX/Niagara/Source/Niagara/Private/NiagaraDataInterfaceGrid2DCollectionReader.cpp
+++ b/Engine/Plugins/FX/Niagara/Source/Niagara/Private/NiagaraDataInterfaceGrid2DCollectionReader.cpp
@@ -1,15 +1,10 @@
 // Copyright Epic Games, Inc. All Rights Reserved.
 #include "NiagaraDataInterfaceGrid2DCollectionReader.h"
-#include "NiagaraShader.h"
-#include "ShaderParameterUtils.h"
-#include "ClearQuad.h"
+#include "Engine/TextureRenderTarget.h"
 #include "NiagaraSystemInstance.h"
-#include "NiagaraRenderer.h"
-#include "Engine/VolumeTexture.h"
-#include "Engine/TextureRenderTargetVolume.h"
-#include "NiagaraSettings.h"
-#include "NiagaraConstants.h"
 #include "NiagaraComputeExecutionContext.h"
+#include "NiagaraEmitterInstance.h"
+#include "NiagaraSystem.h"
 
 #include UE_INLINE_GENERATED_CPP_BY_NAME(NiagaraDataInterfaceGrid2DCollectionReader)
 
diff --git a/Engine/Plugins/FX/Niagara/Source/Niagara/Private/NiagaraDataInterfaceGrid3DCollection.cpp b/Engine/Plugins/FX/Niagara/Source/Niagara/Private/NiagaraDataInterfaceGrid3DCollection.cpp
--- a/Engine/Plugins/FX/Niagara/Source/Niagara/Private/NiagaraDataInterfaceGrid3DCollection.cpp
+++ b/Engine/Plugins/FX/Niagara/Source/Niagara/Private/NiagaraDataInterfaceGrid3DCollection.cpp
@@ -1,12 +1,14 @@
 // Copyright Epic Games, Inc. All Rights Reserved.
 #include "NiagaraDataInterfaceGrid3DCollection.h"
 #include "NiagaraCompileHashVisitor.h"
+#include "NiagaraComponent.h"
 #include "NiagaraConstants.h"
+#include "NiagaraEmitter.h"
 #include "NiagaraGpuComputeDispatchInterface.h"
+#include "NiagaraRenderGraphUtils.h"
 #include "NiagaraSimStageData.h"
 #include "NiagaraRenderer.h"
 #include "NiagaraSettings.h"
-#include "NiagaraShader.h"
 #include "NiagaraShaderParametersBuilder.h"
 #include "NiagaraSystemInstance.h"
 #include "NiagaraSystemInstanceController.h"
@@ -15,10 +17,8 @@
 
 #include "DataDrivenShaderPlatformInfo.h"
 #include "Engine/TextureRenderTargetVolume.h"
-#include "Engine/VolumeTexture.h"
-#include "ClearQuad.h"
+#include "RHIStaticStates.h"
 #include "RenderGraphUtils.h"
-#include "ShaderParameterUtils.h"
 #include "TextureResource.h"
 
 #include "NiagaraDataInterfaceGrid3DCollectionReader.h"
diff --git a/Engine/Plugins/FX/Niagara/Source/Niagara/Private/NiagaraDataInterfaceGrid3DCollectionReader.cpp b/Engine/Plugins/FX/Niagara/Source/Niagara/Private/NiagaraDataInterfaceGrid3DCollectionReader.cpp
--- a/Engine/Plugins/FX/Niagara/Source/Niagara/Private/NiagaraDataInterfaceGrid3DCollectionReader.cpp
+++ b/Engine/Plugins/FX/Niagara/Source/Niagara/Private/NiagaraDataInterfaceGrid3DCollectionReader.cpp
@@ -1,14 +1,11 @@
 // Copyright Epic Games, Inc. All Rights Reserved.
 #include "NiagaraDataInterfaceGrid3DCollectionReader.h"
-#include "NiagaraShader.h"
-#include "ShaderParameterUtils.h"
+#include "Engine/TextureRenderTarget.h"
 #include "NiagaraSystemInstance.h"
-#include "Engine/VolumeTexture.h"
-#include "Engine/TextureRenderTargetVolume.h"
-#include "NiagaraSettings.h"
-#include "NiagaraConstants.h"
 #include "NiagaraComputeExecutionContext.h"
 #include "NiagaraDataInterfaceGrid3DCollectionUtils.h"
+#include "NiagaraEmitterInstance.h"
+#include "NiagaraSystem.h"
 
 #include UE_INLINE_GENERATED_CPP_BY_NAME(NiagaraDataInterfaceGrid3DCollectionReader)
 
diff --git a/Engine/Plugins/FX/Niagara/Source/Niagara/Private/NiagaraDataInterfaceIntRenderTarget2D.cpp b/Engine/Plugins/FX/Niagara/Source/Niagara/Private/NiagaraDataInterfaceIntRenderTarget2D.cpp
--- a/Engine/Plugins/FX/Niagara/Source/Niagara/Private/NiagaraDataInterfaceIntRenderTarget2D.cpp
+++ b/Engine/Plugins/FX/Niagara/Source/Niagara/Private/NiagaraDataInterfaceIntRenderTarget2D.cpp
@@ -1,25 +1,17 @@
 // Copyright Epic Games, Inc. All Rights Reserved.
 #include "NiagaraDataInterfaceIntRenderTarget2D.h"
+#include "NiagaraCompileHashVisitor.h"
 #include "NiagaraDataInterfaceRenderTargetCommon.h"
-#include "NiagaraShader.h"
-#include "NiagaraSettings.h"
 #include "NiagaraSystemInstance.h"
-#include "NiagaraRenderer.h"
 #include "NiagaraGpuComputeDebugInterface.h"
 #include "NiagaraGpuComputeDispatchInterface.h"
-#include "NiagaraShader.h"
 #include "NiagaraShaderParametersBuilder.h"
 #include "NiagaraStats.h"
 
-#include "ShaderParameterUtils.h"
-#include "ClearQuad.h"
 #include "TextureResource.h"
-#include "GenerateMips.h"
-#include "ShaderParameterUtils.h"
 #include "ShaderCompilerCore.h"
 #include "CanvasTypes.h"
 #include "Engine/TextureRenderTarget2D.h"
-#include "Engine/Texture2D.h"
 #include "RenderGraphBuilder.h"
 #include "RenderGraphUtils.h"
 
diff --git a/Engine/Plugins/FX/Niagara/Source/Niagara/Private/NiagaraDataInterfaceLandscape.cpp b/Engine/Plugins/FX/Niagara/Source/Niagara/Private/NiagaraDataInterfaceLandscape.cpp
--- a/Engine/Plugins/FX/Niagara/Source/Niagara/Private/NiagaraDataInterfaceLandscape.cpp
+++ b/Engine/Plugins/FX/Niagara/Source/Niagara/Private/NiagaraDataInterfaceLandscape.cpp
@@ -2,22 +2,23 @@
 
 #include "NiagaraDataInterfaceLandscape.h"
 
-#include "NiagaraComponent.h"
-#include "NiagaraShader.h"
+#include "Containers/ResourceArray.h"
 #include "NiagaraShaderParametersBuilder.h"
+#include "GlobalRenderResources.h"
 #include "NiagaraStats.h"
+#include "NiagaraCompileHashVisitor.h"
 #include "NiagaraWorldManager.h"
 
 #include "Algo/RemoveIf.h"
+#include "RHIStaticStates.h"
 #include "VT/RuntimeVirtualTexture.h"
-#include "EngineModule.h"
 #include "EngineUtils.h"
 #include "Landscape.h"
 #include "LandscapeHeightfieldCollisionComponent.h"
 #include "LandscapeInfo.h"
 #include "LandscapeProxy.h"
+#include "RenderUtils.h"
 #include "ShaderCompilerCore.h"
-#include "ShaderParameterUtils.h"
 #include "PhysicalMaterials/PhysicalMaterial.h"
 
 #include UE_INLINE_GENERATED_CPP_BY_NAME(NiagaraDataInterfaceLandscape)
diff --git a/Engine/Plugins/FX/Niagara/Source/Niagara/Private/NiagaraDataInterfaceMeshCommon.cpp b/Engine/Plugins/FX/Niagara/Source/Niagara/Private/NiagaraDataInterfaceMeshCommon.cpp
--- a/Engine/Plugins/FX/Niagara/Source/Niagara/Private/NiagaraDataInterfaceMeshCommon.cpp
+++ b/Engine/Plugins/FX/Niagara/Source/Niagara/Private/NiagaraDataInterfaceMeshCommon.cpp
@@ -1,11 +1,7 @@
 // Copyright Epic Games, Inc. All Rights Reserved.
 
 #include "NiagaraDataInterfaceMeshCommon.h"
-#include "NiagaraEmitterInstance.h"
-#include "NiagaraComponent.h"
-#include "NiagaraSystemInstance.h"
 #include "Internationalization/Internationalization.h"
-#include "NiagaraScript.h"
 
 #include UE_INLINE_GENERATED_CPP_BY_NAME(NiagaraDataInterfaceMeshCommon)
 
diff --git a/Engine/Plugins/FX/Niagara/Source/Niagara/Private/NiagaraDataInterfaceMeshRendererInfo.cpp b/Engine/Plugins/FX/Niagara/Source/Niagara/Private/NiagaraDataInterfaceMeshRendererInfo.cpp
--- a/Engine/Plugins/FX/Niagara/Source/Niagara/Private/NiagaraDataInterfaceMeshRendererInfo.cpp
+++ b/Engine/Plugins/FX/Niagara/Source/Niagara/Private/NiagaraDataInterfaceMeshRendererInfo.cpp
@@ -1,13 +1,11 @@
 // Copyright Epic Games, Inc. All Rights Reserved.
 #include "NiagaraDataInterfaceMeshRendererInfo.h"
-#include "NiagaraSystem.h"
-#include "NiagaraEmitter.h"
+#include "NiagaraCompileHashVisitor.h"
 #include "NiagaraMeshRendererProperties.h"
-#include "ShaderParameterUtils.h"
-#include "NiagaraStats.h"
 #include "NiagaraRenderer.h"
 #include "NiagaraShaderParametersBuilder.h"
 #include "Engine/StaticMesh.h"
+#include "RenderingThread.h"
 
 #include UE_INLINE_GENERATED_CPP_BY_NAME(NiagaraDataInterfaceMeshRendererInfo)
 
diff --git a/Engine/Plugins/FX/Niagara/Source/Niagara/Private/NiagaraDataInterfaceNeighborGrid3D.cpp b/Engine/Plugins/FX/Niagara/Source/Niagara/Private/NiagaraDataInterfaceNeighborGrid3D.cpp
--- a/Engine/Plugins/FX/Niagara/Source/Niagara/Private/NiagaraDataInterfaceNeighborGrid3D.cpp
+++ b/Engine/Plugins/FX/Niagara/Source/Niagara/Private/NiagaraDataInterfaceNeighborGrid3D.cpp
@@ -1,12 +1,9 @@
 // Copyright Epic Games, Inc. All Rights Reserved.
 #include "NiagaraDataInterfaceNeighborGrid3D.h"
-#include "NiagaraShader.h"
-#include "ShaderParameterUtils.h"
+#include "NiagaraCompileHashVisitor.h"
 #include "NiagaraGpuComputeDispatchInterface.h"
 #include "NiagaraShaderParametersBuilder.h"
 #include "NiagaraSystemInstance.h"
-#include "NiagaraRenderer.h"
-#include "NiagaraShaderParticleID.h"
 #include "DataDrivenShaderPlatformInfo.h"
 #include "RenderGraphUtils.h"
 
diff --git a/Engine/Plugins/FX/Niagara/Source/Niagara/Private/NiagaraDataInterfaceOcclusion.cpp b/Engine/Plugins/FX/Niagara/Source/Niagara/Private/NiagaraDataInterfaceOcclusion.cpp
--- a/Engine/Plugins/FX/Niagara/Source/Niagara/Private/NiagaraDataInterfaceOcclusion.cpp
+++ b/Engine/Plugins/FX/Niagara/Source/Niagara/Private/NiagaraDataInterfaceOcclusion.cpp
@@ -1,17 +1,16 @@
 // Copyright Epic Games, Inc. All Rights Reserved.
 
 #include "NiagaraDataInterfaceOcclusion.h"
+#include "Containers/StridedView.h"
 #include "NiagaraGpuComputeDispatchInterface.h"
+#include "NiagaraCompileHashVisitor.h"
 #include "NiagaraTypes.h"
-#include "NiagaraWorldManager.h"
-#include "NiagaraGPUSystemTick.h"
 #include "NiagaraShaderParametersBuilder.h"
-#include "NiagaraSystemGpuComputeProxy.h"
-#include "NiagaraSystemInstance.h"
 
+#include "RHIStaticStates.h"
 #include "SceneView.h"
-#include "ShaderParameterUtils.h"
 #include "Internationalization/Internationalization.h"
+#include "SceneManagement.h"
 
 #include UE_INLINE_GENERATED_CPP_BY_NAME(NiagaraDataInterfaceOcclusion)
 
diff --git a/Engine/Plugins/FX/Niagara/Source/Niagara/Private/NiagaraDataInterfaceParticleRead.cpp b/Engine/Plugins/FX/Niagara/Source/Niagara/Private/NiagaraDataInterfaceParticleRead.cpp
--- a/Engine/Plugins/FX/Niagara/Source/Niagara/Private/NiagaraDataInterfaceParticleRead.cpp
+++ b/Engine/Plugins/FX/Niagara/Source/Niagara/Private/NiagaraDataInterfaceParticleRead.cpp
@@ -3,14 +3,15 @@
 #include "NiagaraDataInterfaceParticleRead.h"
 #include "NiagaraCompileHashVisitor.h"
 #include "NiagaraComputeExecutionContext.h"
-#include "NiagaraConstants.h"
 #include "NiagaraDataSetAccessor.h"
+#include "NiagaraEmitterInstance.h"
 #include "NiagaraGPUSystemTick.h"
 #include "NiagaraRenderer.h"
+#include "NiagaraSimStageData.h"
 #include "NiagaraSystemInstance.h"
 #include "NiagaraShaderParametersBuilder.h"
 
-#include "ShaderParameterUtils.h"
+#include "NiagaraSystem.h"
 
 #include UE_INLINE_GENERATED_CPP_BY_NAME(NiagaraDataInterfaceParticleRead)
 
diff --git a/Engine/Plugins/FX/Niagara/Source/Niagara/Private/NiagaraDataInterfaceRW.cpp b/Engine/Plugins/FX/Niagara/Source/Niagara/Private/NiagaraDataInterfaceRW.cpp
--- a/Engine/Plugins/FX/Niagara/Source/Niagara/Private/NiagaraDataInterfaceRW.cpp
+++ b/Engine/Plugins/FX/Niagara/Source/Niagara/Private/NiagaraDataInterfaceRW.cpp
@@ -1,11 +1,8 @@
 // Copyright Epic Games, Inc. All Rights Reserved.
 #include "NiagaraDataInterfaceRW.h"
 #include "NiagaraGpuComputeDispatchInterface.h"
-#include "NiagaraShader.h"
 #include "NiagaraSimStageData.h"
 #include "RenderGraph.h"
-#include "ShaderParameterUtils.h"
-#include "ClearQuad.h"
 
 #include "UObject/UE5MainStreamObjectVersion.h"
 
diff --git a/Engine/Plugins/FX/Niagara/Source/Niagara/Private/NiagaraDataInterfaceRasterizationGrid3D.cpp b/Engine/Plugins/FX/Niagara/Source/Niagara/Private/NiagaraDataInterfaceRasterizationGrid3D.cpp
--- a/Engine/Plugins/FX/Niagara/Source/Niagara/Private/NiagaraDataInterfaceRasterizationGrid3D.cpp
+++ b/Engine/Plugins/FX/Niagara/Source/Niagara/Private/NiagaraDataInterfaceRasterizationGrid3D.cpp
@@ -1,12 +1,10 @@
 // Copyright Epic Games, Inc. All Rights Reserved.
 #include "NiagaraDataInterfaceRasterizationGrid3D.h"
-#include "NiagaraShader.h"
-#include "ShaderParameterUtils.h"
+#include "NiagaraCompileHashVisitor.h"
 #include "NiagaraGpuComputeDispatchInterface.h"
 #include "NiagaraShaderParametersBuilder.h"
 #include "NiagaraSystemInstance.h"
 #include "NiagaraRenderer.h"
-#include "NiagaraShaderParticleID.h"
 #include "DataDrivenShaderPlatformInfo.h"
 #include "RenderGraphUtils.h"
 
diff --git a/Engine/Plugins/FX/Niagara/Source/Niagara/Private/NiagaraDataInterfaceRenderTarget2D.cpp b/Engine/Plugins/FX/Niagara/Source/Niagara/Private/NiagaraDataInterfaceRenderTarget2D.cpp
--- a/Engine/Plugins/FX/Niagara/Source/Niagara/Private/NiagaraDataInterfaceRenderTarget2D.cpp
+++ b/Engine/Plugins/FX/Niagara/Source/Niagara/Private/NiagaraDataInterfaceRenderTarget2D.cpp
@@ -1,21 +1,14 @@
 // Copyright Epic Games, Inc. All Rights Reserved.
 #include "NiagaraDataInterfaceRenderTarget2D.h"
-#include "ShaderParameterUtils.h"
-#include "ClearQuad.h"
+#include "NiagaraCompileHashVisitor.h"
 #include "TextureResource.h"
-#include "GenerateMips.h"
 #include "CanvasTypes.h"
-#include "Engine/Texture2D.h"
 #include "Engine/TextureRenderTarget2D.h"
 
 #include "NiagaraDataInterfaceRenderTargetCommon.h"
-#include "NiagaraSettings.h"
 #include "NiagaraSystemInstance.h"
-#include "NiagaraRenderer.h"
 #include "NiagaraGpuComputeDebugInterface.h"
 #include "NiagaraGpuComputeDispatchInterface.h"
-#include "NiagaraGPUProfilerInterface.h"
-#include "NiagaraShader.h"
 #include "NiagaraShaderParametersBuilder.h"
 #include "NiagaraStats.h"
 #include "RenderGraphBuilder.h"
diff --git a/Engine/Plugins/FX/Niagara/Source/Niagara/Private/NiagaraDataInterfaceRenderTarget2DArray.cpp b/Engine/Plugins/FX/Niagara/Source/Niagara/Private/NiagaraDataInterfaceRenderTarget2DArray.cpp
--- a/Engine/Plugins/FX/Niagara/Source/Niagara/Private/NiagaraDataInterfaceRenderTarget2DArray.cpp
+++ b/Engine/Plugins/FX/Niagara/Source/Niagara/Private/NiagaraDataInterfaceRenderTarget2DArray.cpp
@@ -1,18 +1,14 @@
 // Copyright Epic Games, Inc. All Rights Reserved.
 #include "NiagaraDataInterfaceRenderTarget2DArray.h"
-#include "ShaderParameterUtils.h"
-#include "ClearQuad.h"
+#include "Engine/TextureRenderTarget2D.h"
 #include "TextureResource.h"
+#include "NiagaraCompileHashVisitor.h"
 #include "TextureRenderTarget2DArrayResource.h"
-#include "Engine/Texture2DArray.h"
 #include "Engine/TextureRenderTarget2DArray.h"
 
 #include "NiagaraDataInterfaceRenderTargetCommon.h"
 #include "NiagaraSystemInstance.h"
 #include "NiagaraStats.h"
-#include "NiagaraRenderer.h"
-#include "NiagaraSettings.h"
-#include "NiagaraShader.h"
 #include "NiagaraShaderParametersBuilder.h"
 #include "NiagaraGpuComputeDebugInterface.h"
 #include "NiagaraGpuComputeDispatchInterface.h"
diff --git a/Engine/Plugins/FX/Niagara/Source/Niagara/Private/NiagaraDataInterfaceRenderTargetCommon.h b/Engine/Plugins/FX/Niagara/Source/Niagara/Private/NiagaraDataInterfaceRenderTargetCommon.h
--- a/Engine/Plugins/FX/Niagara/Source/Niagara/Private/NiagaraDataInterfaceRenderTargetCommon.h
+++ b/Engine/Plugins/FX/Niagara/Source/Niagara/Private/NiagaraDataInterfaceRenderTargetCommon.h
@@ -2,7 +2,9 @@
 
 #pragma once
 
-#include "Engine/TextureRenderTarget2D.h"
+#include "HAL/Platform.h"
+
+enum ETextureRenderTargetFormat : int;
 
 namespace NiagaraDataInterfaceRenderTargetCommon
 {
diff --git a/Engine/Plugins/FX/Niagara/Source/Niagara/Private/NiagaraDataInterfaceRenderTargetCube.cpp b/Engine/Plugins/FX/Niagara/Source/Niagara/Private/NiagaraDataInterfaceRenderTargetCube.cpp
--- a/Engine/Plugins/FX/Niagara/Source/Niagara/Private/NiagaraDataInterfaceRenderTargetCube.cpp
+++ b/Engine/Plugins/FX/Niagara/Source/Niagara/Private/NiagaraDataInterfaceRenderTargetCube.cpp
@@ -1,16 +1,13 @@
 // Copyright Epic Games, Inc. All Rights Reserved.
 #include "NiagaraDataInterfaceRenderTargetCube.h"
-#include "ShaderParameterUtils.h"
-#include "ClearQuad.h"
+#include "Engine/TextureRenderTarget2D.h"
 #include "TextureResource.h"
 #include "Engine/TextureRenderTargetCube.h"
 
+#include "NiagaraCompileHashVisitor.h"
 #include "NiagaraDataInterfaceRenderTargetCommon.h"
 #include "NiagaraSystemInstance.h"
 #include "NiagaraStats.h"
-#include "NiagaraRenderer.h"
-#include "NiagaraSettings.h"
-#include "NiagaraShader.h"
 #include "NiagaraShaderParametersBuilder.h"
 #include "NiagaraGpuComputeDebugInterface.h"
 #include "NiagaraGpuComputeDispatchInterface.h"
diff --git a/Engine/Plugins/FX/Niagara/Source/Niagara/Private/NiagaraDataInterfaceRenderTargetVolume.cpp b/Engine/Plugins/FX/Niagara/Source/Niagara/Private/NiagaraDataInterfaceRenderTargetVolume.cpp
--- a/Engine/Plugins/FX/Niagara/Source/Niagara/Private/NiagaraDataInterfaceRenderTargetVolume.cpp
+++ b/Engine/Plugins/FX/Niagara/Source/Niagara/Private/NiagaraDataInterfaceRenderTargetVolume.cpp
@@ -1,8 +1,8 @@
 // Copyright Epic Games, Inc. All Rights Reserved.
 #include "NiagaraDataInterfaceRenderTargetVolume.h"
-#include "ShaderParameterUtils.h"
-#include "ClearQuad.h"
+#include "Engine/TextureRenderTarget2D.h"
 #include "TextureResource.h"
+#include "NiagaraCompileHashVisitor.h"
 #include "TextureRenderTargetVolumeResource.h"
 #include "Engine/TextureRenderTargetVolume.h"
 #include "HAL/PlatformFileManager.h"
@@ -12,9 +12,6 @@
 #include "NiagaraGpuComputeDispatchInterface.h"
 #include "NiagaraSystemInstance.h"
 #include "NiagaraStats.h"
-#include "NiagaraRenderer.h"
-#include "NiagaraSettings.h"
-#include "NiagaraShader.h"
 #include "NiagaraShaderParametersBuilder.h"
 #include "NiagaraGpuComputeDebugInterface.h"
 #include "VolumeCache.h"
diff --git a/Engine/Plugins/FX/Niagara/Source/Niagara/Private/NiagaraDataInterfaceRigidMeshCollisionQuery.cpp b/Engine/Plugins/FX/Niagara/Source/Niagara/Private/NiagaraDataInterfaceRigidMeshCollisionQuery.cpp
--- a/Engine/Plugins/FX/Niagara/Source/Niagara/Private/NiagaraDataInterfaceRigidMeshCollisionQuery.cpp
+++ b/Engine/Plugins/FX/Niagara/Source/Niagara/Private/NiagaraDataInterfaceRigidMeshCollisionQuery.cpp
@@ -4,24 +4,21 @@
 #include "Algo/Accumulate.h"
 #include "Algo/ForEach.h"
 #include "Algo/RemoveIf.h"
-#include "Animation/SkeletalMeshActor.h"
-#include "AnimationRuntime.h"
 #include "Components/SkeletalMeshComponent.h"
 #include "Components/StaticMeshComponent.h"
+#include "Containers/StridedView.h"
 #include "FXRenderingUtils.h"
 #include "Engine/Canvas.h"
 #include "Engine/Engine.h"
+#include "Engine/SkeletalMesh.h"
 #include "EngineUtils.h"
-#include "SkeletalRenderPublic.h"
 
-#include "NiagaraComponent.h"
+#include "Misc/LargeWorldRenderPosition.h"
 #include "NiagaraDataInterfaceUtilities.h"
-#include "NiagaraDistanceFieldHelper.h"
+#include "NiagaraCompileHashVisitor.h"
 #include "NiagaraFunctionLibrary.h"
-#include "NiagaraGpuComputeDispatch.h"
 #include "NiagaraGpuComputeDispatchInterface.h"
 #include "NiagaraRenderer.h"
-#include "NiagaraShader.h"
 #include "NiagaraShaderParametersBuilder.h"
 #include "NiagaraSimStageData.h"
 #include "NiagaraSystemImpl.h"
@@ -30,7 +27,7 @@
 #include "PhysicsEngine/PhysicsAsset.h"
 #include "PrimitiveSceneInfo.h"
 #include "SceneInterface.h"
-#include "ShaderParameterUtils.h"
+#include "SceneView.h"
 
 #include UE_INLINE_GENERATED_CPP_BY_NAME(NiagaraDataInterfaceRigidMeshCollisionQuery)
 
diff --git a/Engine/Plugins/FX/Niagara/Source/Niagara/Private/NiagaraDataInterfaceSkeletalMesh.cpp b/Engine/Plugins/FX/Niagara/Source/Niagara/Private/NiagaraDataInterfaceSkeletalMesh.cpp
--- a/Engine/Plugins/FX/Niagara/Source/Niagara/Private/NiagaraDataInterfaceSkeletalMesh.cpp
+++ b/Engine/Plugins/FX/Niagara/Source/Niagara/Private/NiagaraDataInterfaceSkeletalMesh.cpp
@@ -7,23 +7,25 @@
 #include "Async/ParallelFor.h"
 #include "Components/SkeletalMeshComponent.h"
 #include "Engine/Canvas.h"
-#include "Engine/SkeletalMeshSocket.h"
 #include "Internationalization/Internationalization.h"
+#include "Misc/LargeWorldRenderPosition.h"
 #include "NDISkeletalMeshCommon.h"
 #include "NiagaraCompileHashVisitor.h"
-#include "NiagaraEmitterInstance.h"
 #include "NiagaraComponent.h"
+#include "NiagaraDataInterfaceMeshCommon.h"
 #include "NiagaraDataInterfaceSkeletalMeshConnectivity.h"
 #include "Experimental/NiagaraDataInterfaceSkeletalMeshUvMapping.h"
+#include "NiagaraEmitter.h"
 #include "NiagaraSettings.h"
 #include "NiagaraStats.h"
+#include "NiagaraSystem.h"
 #include "NiagaraSystemInstance.h"
 #include "NiagaraRenderer.h"
 #include "NiagaraScript.h"
 #include "NiagaraShaderParametersBuilder.h"
 #include "NiagaraWorldManager.h"
+#include "Rendering/SkeletalMeshRenderData.h"
 #include "Templates/AlignmentTemplates.h"
-#include "ShaderParameterUtils.h"
 #include "ShaderCore.h"
 
 #include UE_INLINE_GENERATED_CPP_BY_NAME(NiagaraDataInterfaceSkeletalMesh)
diff --git a/Engine/Plugins/FX/Niagara/Source/Niagara/Private/NiagaraDataInterfaceSkeletalMeshConnectivity.cpp b/Engine/Plugins/FX/Niagara/Source/Niagara/Private/NiagaraDataInterfaceSkeletalMeshConnectivity.cpp
--- a/Engine/Plugins/FX/Niagara/Source/Niagara/Private/NiagaraDataInterfaceSkeletalMeshConnectivity.cpp
+++ b/Engine/Plugins/FX/Niagara/Source/Niagara/Private/NiagaraDataInterfaceSkeletalMeshConnectivity.cpp
@@ -2,13 +2,15 @@
 
 #include "NiagaraDataInterfaceSkeletalMeshConnectivity.h"
 
-#include "Algo/StableSort.h"
+#include "Engine/SkeletalMesh.h"
 #include "Engine/SkinnedAssetCommon.h"
-#include "NiagaraResourceArrayWriter.h"
+#include "NiagaraDataInterfaceSkeletalMesh.h"
 #include "NiagaraSettings.h"
 #include "NiagaraStats.h"
 
+#include "Rendering/SkeletalMeshRenderData.h"
 #include <limits>
+#include "RenderingThread.h"
 
 DECLARE_CYCLE_STAT(TEXT("Niagara - SkelMesh - BuildAdjacencyBuffer"), STAT_NiagaraSkel_Connectivity_Adjacency, STATGROUP_Niagara);
 
@@ -425,4 +427,4 @@ FSkeletalMeshConnectivityProxy::ReleaseRHI()
 #if STATS
 	GpuMemoryUsage = 0;
 #endif
-}
\ No newline at end of file
+}
diff --git a/Engine/Plugins/FX/Niagara/Source/Niagara/Private/NiagaraDataInterfaceSkeletalMeshConnectivity.h b/Engine/Plugins/FX/Niagara/Source/Niagara/Private/NiagaraDataInterfaceSkeletalMeshConnectivity.h
--- a/Engine/Plugins/FX/Niagara/Source/Niagara/Private/NiagaraDataInterfaceSkeletalMeshConnectivity.h
+++ b/Engine/Plugins/FX/Niagara/Source/Niagara/Private/NiagaraDataInterfaceSkeletalMeshConnectivity.h
@@ -1,9 +1,14 @@
 // Copyright Epic Games, Inc. All Rights Reserved.
 #pragma once
 
-#include "NiagaraDataInterfaceSkeletalMesh.h"
-#include "NiagaraUvQuadTree.h"
+#include "Containers/DynamicRHIResourceArray.h"
 #include "RenderResource.h"
+#include "RHIResources.h"
+#include "UObject/WeakObjectPtr.h"
+
+class FSkeletalMeshLODRenderData;
+class USkeletalMesh;
+struct FSkeletalMeshConnectivityUsage;
 
 struct FSkeletalMeshConnectivity;
 
@@ -63,4 +68,4 @@ private:
 	TUniquePtr<FSkeletalMeshConnectivityProxy> Proxy;
 
 	std::atomic<int32> GpuUserCount;
-};
\ No newline at end of file
+};
diff --git a/Engine/Plugins/FX/Niagara/Source/Niagara/Private/NiagaraDataInterfaceSpline.cpp b/Engine/Plugins/FX/Niagara/Source/Niagara/Private/NiagaraDataInterfaceSpline.cpp
--- a/Engine/Plugins/FX/Niagara/Source/Niagara/Private/NiagaraDataInterfaceSpline.cpp
+++ b/Engine/Plugins/FX/Niagara/Source/Niagara/Private/NiagaraDataInterfaceSpline.cpp
@@ -1,14 +1,15 @@
 // Copyright Epic Games, Inc. All Rights Reserved.
 
 #include "NiagaraDataInterfaceSpline.h"
-#include "NiagaraEmitterInstance.h"
+#include "Misc/LargeWorldRenderPosition.h"
 #include "NiagaraComponent.h"
+#include "NiagaraCompileHashVisitor.h"
 #include "NiagaraRenderer.h"
 #include "NiagaraShaderParametersBuilder.h"
 #include "NiagaraStats.h"
+#include "NiagaraSystem.h"
 #include "NiagaraSystemInstance.h"
 #include "Internationalization/Internationalization.h"
-#include "ShaderParameterUtils.h"
 #include "ShaderCompilerCore.h"
 
 #include UE_INLINE_GENERATED_CPP_BY_NAME(NiagaraDataInterfaceSpline)
diff --git a/Engine/Plugins/FX/Niagara/Source/Niagara/Private/NiagaraDataInterfaceSpriteRendererInfo.cpp b/Engine/Plugins/FX/Niagara/Source/Niagara/Private/NiagaraDataInterfaceSpriteRendererInfo.cpp
--- a/Engine/Plugins/FX/Niagara/Source/Niagara/Private/NiagaraDataInterfaceSpriteRendererInfo.cpp
+++ b/Engine/Plugins/FX/Niagara/Source/Niagara/Private/NiagaraDataInterfaceSpriteRendererInfo.cpp
@@ -1,12 +1,9 @@
 // Copyright Epic Games, Inc. All Rights Reserved.
 #include "NiagaraDataInterfaceSpriteRendererInfo.h"
-#include "NiagaraSystem.h"
-#include "NiagaraEmitter.h"
+#include "NiagaraCompileHashVisitor.h"
 #include "NiagaraSpriteRendererProperties.h"
-#include "ShaderParameterUtils.h"
-#include "NiagaraStats.h"
-#include "NiagaraRenderer.h"
 #include "NiagaraShaderParametersBuilder.h"
+#include "RenderingThread.h"
 
 #include UE_INLINE_GENERATED_CPP_BY_NAME(NiagaraDataInterfaceSpriteRendererInfo)
 
diff --git a/Engine/Plugins/FX/Niagara/Source/Niagara/Private/NiagaraDataInterfaceStaticMeshUvMapping.cpp b/Engine/Plugins/FX/Niagara/Source/Niagara/Private/NiagaraDataInterfaceStaticMeshUvMapping.cpp
--- a/Engine/Plugins/FX/Niagara/Source/Niagara/Private/NiagaraDataInterfaceStaticMeshUvMapping.cpp
+++ b/Engine/Plugins/FX/Niagara/Source/Niagara/Private/NiagaraDataInterfaceStaticMeshUvMapping.cpp
@@ -3,8 +3,6 @@
 #include "NiagaraDataInterfaceStaticMeshUvMapping.h"
 
 #include "Experimental/NiagaraMeshUvMapping.h"
-#include "NiagaraResourceArrayWriter.h"
-#include "NiagaraStats.h"
 #include "RawIndexBuffer.h"
 #include "Engine/StaticMesh.h"
 #include "StaticMeshResources.h"
@@ -133,4 +131,4 @@ void FStaticMeshUvMapping::BuildQuadTree()
 		FQuadTreeBuildHelper BuildHelper(VertexHelper, UvSetIndex, TriangleIndexQuadTree);
 		BuildHelper.Build();
 	}
-}
\ No newline at end of file
+}
diff --git a/Engine/Plugins/FX/Niagara/Source/Niagara/Private/NiagaraDataInterfaceStaticMeshUvMapping.h b/Engine/Plugins/FX/Niagara/Source/Niagara/Private/NiagaraDataInterfaceStaticMeshUvMapping.h
--- a/Engine/Plugins/FX/Niagara/Source/Niagara/Private/NiagaraDataInterfaceStaticMeshUvMapping.h
+++ b/Engine/Plugins/FX/Niagara/Source/Niagara/Private/NiagaraDataInterfaceStaticMeshUvMapping.h
@@ -1,10 +1,10 @@
 // Copyright Epic Games, Inc. All Rights Reserved.
 #pragma once
 
-#include "Engine/StaticMesh.h"
 #include "Experimental/NiagaraMeshUvMapping.h"
 #include "NiagaraUvQuadTree.h"
 #include "RenderResource.h"
+#include "UObject/WeakObjectPtr.h"
 #include "UObject/WeakObjectPtrTemplates.h"
 
 class UStaticMesh;
@@ -30,4 +30,4 @@ private:
 	virtual void BuildQuadTree() override;
 
 	TWeakObjectPtr<UStaticMesh> MeshObject;
-};
\ No newline at end of file
+};
diff --git a/Engine/Plugins/FX/Niagara/Source/Niagara/Private/NiagaraDataInterfaceVectorField.cpp b/Engine/Plugins/FX/Niagara/Source/Niagara/Private/NiagaraDataInterfaceVectorField.cpp
--- a/Engine/Plugins/FX/Niagara/Source/Niagara/Private/NiagaraDataInterfaceVectorField.cpp
+++ b/Engine/Plugins/FX/Niagara/Source/Niagara/Private/NiagaraDataInterfaceVectorField.cpp
@@ -1,11 +1,10 @@
 // Copyright Epic Games, Inc. All Rights Reserved.
 
 #include "NiagaraDataInterfaceVectorField.h"
+#include "NiagaraCompileHashVisitor.h"
 #include "NiagaraEmitter.h"
 #include "VectorField/VectorFieldStatic.h"
 #include "VectorField/VectorFieldAnimated.h"
-#include "NiagaraShader.h"
-#include "ShaderParameterUtils.h"
 #include "NiagaraSystem.h"
 #include "NiagaraEmitterHandle.h"
 #include "NiagaraComponent.h"
diff --git a/Engine/Plugins/FX/Niagara/Source/Niagara/Private/NiagaraDataInterfaceVolumeCache.cpp b/Engine/Plugins/FX/Niagara/Source/Niagara/Private/NiagaraDataInterfaceVolumeCache.cpp
--- a/Engine/Plugins/FX/Niagara/Source/Niagara/Private/NiagaraDataInterfaceVolumeCache.cpp
+++ b/Engine/Plugins/FX/Niagara/Source/Niagara/Private/NiagaraDataInterfaceVolumeCache.cpp
@@ -1,19 +1,12 @@
 // Copyright Epic Games, Inc. All Rights Reserved.
 #include "NiagaraDataInterfaceVolumeCache.h"
-#include "NiagaraShader.h"
-#include "ShaderParameterUtils.h"
-#include "ClearQuad.h"
-#include "TextureResource.h"
+#include "NiagaraCompileHashVisitor.h"
 #include "NiagaraSystemInstance.h"
-#include "NiagaraRenderer.h"
-#include "Engine/VolumeTexture.h"
-#include "Engine/TextureRenderTargetVolume.h"
-#include "NiagaraSettings.h"
-#include "NiagaraConstants.h"
-#include "NiagaraComputeExecutionContext.h"
 #include "NiagaraShaderParametersBuilder.h"
+#include "NiagaraSystem.h"
 #include "RHIStaticStates.h"
 #include "GlobalRenderResources.h"
+#include "VolumeCache.h"
 
 #include UE_INLINE_GENERATED_CPP_BY_NAME(NiagaraDataInterfaceVolumeCache)
 
diff --git a/Engine/Plugins/FX/Niagara/Source/Niagara/Private/NiagaraDataSet.cpp b/Engine/Plugins/FX/Niagara/Source/Niagara/Private/NiagaraDataSet.cpp
--- a/Engine/Plugins/FX/Niagara/Source/Niagara/Private/NiagaraDataSet.cpp
+++ b/Engine/Plugins/FX/Niagara/Source/Niagara/Private/NiagaraDataSet.cpp
@@ -2,10 +2,7 @@
 
 #include "NiagaraDataSet.h"
 #include "NiagaraCommon.h"
-#include "NiagaraShader.h"
 #include "GlobalShader.h"
-#include "UpdateTextureShaders.h"
-#include "ShaderParameterUtils.h"
 #include "NiagaraStats.h"
 #include "NiagaraRenderer.h"
 #include "NiagaraGPUInstanceCountManager.h"
diff --git a/Engine/Plugins/FX/Niagara/Source/Niagara/Private/NiagaraDataSetDebugAccessor.cpp b/Engine/Plugins/FX/Niagara/Source/Niagara/Private/NiagaraDataSetDebugAccessor.cpp
--- a/Engine/Plugins/FX/Niagara/Source/Niagara/Private/NiagaraDataSetDebugAccessor.cpp
+++ b/Engine/Plugins/FX/Niagara/Source/Niagara/Private/NiagaraDataSetDebugAccessor.cpp
@@ -1,5 +1,6 @@
 // Copyright Epic Games, Inc. All Rights Reserved.
 #include "NiagaraDataSetDebugAccessor.h"
+#include "NiagaraDataSet.h"
 
 bool FNiagaraDataSetDebugAccessor::Init(const FNiagaraDataSetCompiledData& CompiledData, FName InVariableName)
 {
diff --git a/Engine/Plugins/FX/Niagara/Source/Niagara/Private/NiagaraDataSetReadback.cpp b/Engine/Plugins/FX/Niagara/Source/Niagara/Private/NiagaraDataSetReadback.cpp
--- a/Engine/Plugins/FX/Niagara/Source/Niagara/Private/NiagaraDataSetReadback.cpp
+++ b/Engine/Plugins/FX/Niagara/Source/Niagara/Private/NiagaraDataSetReadback.cpp
@@ -2,6 +2,7 @@
 
 #include "NiagaraDataSetReadback.h"
 #include "NiagaraComputeExecutionContext.h"
+#include "NiagaraEmitterHandle.h"
 #include "NiagaraEmitterInstance.h"
 #include "NiagaraSystemInstance.h"
 #include "NiagaraGpuComputeDispatchInterface.h"
diff --git a/Engine/Plugins/FX/Niagara/Source/Niagara/Private/NiagaraDebugHud.cpp b/Engine/Plugins/FX/Niagara/Source/Niagara/Private/NiagaraDebugHud.cpp
--- a/Engine/Plugins/FX/Niagara/Source/Niagara/Private/NiagaraDebugHud.cpp
+++ b/Engine/Plugins/FX/Niagara/Source/Niagara/Private/NiagaraDebugHud.cpp
@@ -1,18 +1,21 @@
 // Copyright Epic Games, Inc. All Rights Reserved.
 
 #include "NiagaraDebugHud.h"
+#include "BatchedElements.h"
 #include "Engine/Engine.h"
+#include "DrawDebugHelpers.h"
 #include "GameFramework/Pawn.h"
 #include "NiagaraComponent.h"
-#include "NiagaraComputeExecutionContext.h"
 #include "NiagaraDataSetDebugAccessor.h"
 #include "NiagaraDataSetReadback.h"
+#include "NiagaraEmitterInstance.h"
 #include "NiagaraFunctionLibrary.h"
 #include "NiagaraGpuComputeDispatchInterface.h"
 #include "NiagaraMeshRendererProperties.h"
 #include "NiagaraScript.h"
 #include "NiagaraSpriteRendererProperties.h"
 #include "NiagaraSystem.h"
+#include "NiagaraSystemGpuComputeProxy.h"
 #include "NiagaraSystemInstanceController.h"
 #include "NiagaraWorldManager.h"
 
@@ -22,6 +25,8 @@
 #include "GameFramework/PlayerController.h"
 #include "Particles/FXBudget.h"
 #include "SceneInterface.h"
+#include "SceneView.h"
+#include "UObject/UObjectIterator.h"
 
 #if WITH_NIAGARA_DEBUGGER
 
