Running IWYU on Niagara/Private

diff --git a/Engine/Plugins/FX/Niagara/Source/Niagara/Private/NiagaraDebugHud.h b/Engine/Plugins/FX/Niagara/Source/Niagara/Private/NiagaraDebugHud.h
--- a/Engine/Plugins/FX/Niagara/Source/Niagara/Private/NiagaraDebugHud.h
+++ b/Engine/Plugins/FX/Niagara/Source/Niagara/Private/NiagaraDebugHud.h
@@ -5,10 +5,9 @@ Class used help debugging Niagara simulations
 ==============================================================================*/
 #pragma once
 
-#include "CoreMinimal.h"
 #include "NiagaraCommon.h"
+#include "NiagaraDefines.h"
 #include "NiagaraTypes.h"
-#include "RHI.h"
 #include "RHICommandList.h"
 #include "NiagaraDebuggerCommon.h"
 #include "NiagaraGPUProfilerInterface.h"
diff --git a/Engine/Plugins/FX/Niagara/Source/Niagara/Private/NiagaraDebuggerClient.cpp b/Engine/Plugins/FX/Niagara/Source/Niagara/Private/NiagaraDebuggerClient.cpp
--- a/Engine/Plugins/FX/Niagara/Source/Niagara/Private/NiagaraDebuggerClient.cpp
+++ b/Engine/Plugins/FX/Niagara/Source/Niagara/Private/NiagaraDebuggerClient.cpp
@@ -6,14 +6,18 @@
 #include "MessageEndpoint.h"
 #include "MessageEndpointBuilder.h"
 #include "IMessageContext.h"
+#include "NiagaraEmitterInstance.h"
 #include "NiagaraWorldManager.h"
 #include "NiagaraComponent.h"
+#include "NiagaraModule.h"
 #include "NiagaraSystemInstanceController.h"
 #include "NiagaraDebugHud.h"
 #include "Containers/Ticker.h"
 #include "NiagaraSimCache.h"
+#include "NiagaraSystemGpuComputeProxy.h"
 #include "Serialization/ObjectAndNameAsStringProxyArchive.h"
 #include "UObject/Package.h"
+#include "UObject/UObjectIterator.h"
 
 #if WITH_NIAGARA_DEBUGGER
 
diff --git a/Engine/Plugins/FX/Niagara/Source/Niagara/Private/NiagaraDebuggerClient.h b/Engine/Plugins/FX/Niagara/Source/Niagara/Private/NiagaraDebuggerClient.h
--- a/Engine/Plugins/FX/Niagara/Source/Niagara/Private/NiagaraDebuggerClient.h
+++ b/Engine/Plugins/FX/Niagara/Source/Niagara/Private/NiagaraDebuggerClient.h
@@ -7,14 +7,15 @@ Controller class running on game clients that handles the passing of messages to
 #pragma once
 
 #include "Containers/Ticker.h"
-#include "CoreMinimal.h"
 #include "NiagaraCommon.h"
 #include "NiagaraDebuggerCommon.h"
 #include "IMessageContext.h"
+#include "NiagaraDefines.h"
 #include "Particles/ParticlePerfStatsManager.h"
-#include "NiagaraSimCache.h"
 #include "UObject/StrongObjectPtr.h"
 
+class UNiagaraSimCache;
+
 class FMessageEndpoint;
 
 #if WITH_PARTICLE_PERF_STATS
@@ -107,4 +108,4 @@ private:
 	TArray<FNiagaraSimCacheCaptureInfo> SimCacheCaptures;
 };
 
-#endif
\ No newline at end of file
+#endif
diff --git a/Engine/Plugins/FX/Niagara/Source/Niagara/Private/NiagaraEffectType.cpp b/Engine/Plugins/FX/Niagara/Source/Niagara/Private/NiagaraEffectType.cpp
--- a/Engine/Plugins/FX/Niagara/Source/Niagara/Private/NiagaraEffectType.cpp
+++ b/Engine/Plugins/FX/Niagara/Source/Niagara/Private/NiagaraEffectType.cpp
@@ -1,11 +1,14 @@
 // Copyright Epic Games, Inc. All Rights Reserved.
 #include "NiagaraEffectType.h"
 #include "AssetRegistry/AssetData.h"
+#include "Engine/World.h"
 #include "NiagaraCommon.h"
+#include "Modules/ModuleManager.h"
 #include "NiagaraComponent.h"
 #include "NiagaraCustomVersion.h"
 #include "NiagaraSystem.h"
 #include "NiagaraSystemInstanceController.h"
+#include "UObject/UObjectIterator.h"
 
 //In an effort to cut the impact of runtime perf tracking, I limit the number of fames we actually sample on.
 int32 GNumFramesBetweenRuntimePerfSamples = 5;
diff --git a/Engine/Plugins/FX/Niagara/Source/Niagara/Private/NiagaraEmitter.cpp b/Engine/Plugins/FX/Niagara/Source/Niagara/Private/NiagaraEmitter.cpp
--- a/Engine/Plugins/FX/Niagara/Source/Niagara/Private/NiagaraEmitter.cpp
+++ b/Engine/Plugins/FX/Niagara/Source/Niagara/Private/NiagaraEmitter.cpp
@@ -4,15 +4,18 @@
 #include "NiagaraEmitter.h"
 
 #include "INiagaraEditorOnlyDataUtlities.h"
+#include "NiagaraBoundsCalculator.h"
 #include "NiagaraCustomVersion.h"
+#include "NiagaraMessageDataBase.h"
 #include "UObject/UE5MainStreamObjectVersion.h"
 #include "NiagaraEditorDataBase.h"
 #include "NiagaraModule.h"
-#include "NiagaraRenderer.h"
 #include "NiagaraRendererProperties.h"
+#include "NiagaraScratchPadContainer.h"
 #include "NiagaraScript.h"
 #include "NiagaraScriptSourceBase.h"
 #include "NiagaraSettings.h"
+#include "NiagaraShader.h"
 #include "NiagaraSimulationStageBase.h"
 #include "NiagaraStats.h"
 #include "NiagaraSystem.h"
@@ -20,10 +23,12 @@
 #include "HAL/LowLevelMemTracker.h"
 #include "Interfaces/ITargetPlatform.h"
 #include "Modules/ModuleManager.h"
+#include "Templates/Greater.h"
 #include "UObject/LinkerLoad.h"
 #include "UObject/ObjectSaveContext.h"
 #include "UObject/Package.h"
 #include "UObject/UObjectIterator.h"
+#include "WorldCollision.h"
 
 #include UE_INLINE_GENERATED_CPP_BY_NAME(NiagaraEmitter)
 
diff --git a/Engine/Plugins/FX/Niagara/Source/Niagara/Private/NiagaraEmitterInstance.cpp b/Engine/Plugins/FX/Niagara/Source/Niagara/Private/NiagaraEmitterInstance.cpp
--- a/Engine/Plugins/FX/Niagara/Source/Niagara/Private/NiagaraEmitterInstance.cpp
+++ b/Engine/Plugins/FX/Niagara/Source/Niagara/Private/NiagaraEmitterInstance.cpp
@@ -2,18 +2,19 @@
 
 #include "NiagaraEmitterInstance.h"
 #include "Engine/Engine.h"
+#include "NiagaraBoundsCalculator.h"
 #include "NiagaraStats.h"
 #include "NiagaraConstants.h"
-#include "NiagaraRenderer.h"
+#include "NiagaraModule.h"
 #include "NiagaraSystemInstance.h"
 #include "NiagaraComputeExecutionContext.h"
 #include "NiagaraDataInterface.h"
 #include "NiagaraGpuComputeDispatchInterface.h"
 #include "NiagaraScriptExecutionContext.h"
 #include "NiagaraParameterCollection.h"
-#include "NiagaraWorldManager.h"
 #include "NiagaraSimulationStageBase.h"
 #include "NiagaraComponentSettings.h"
+#include "NiagaraSystemSimulation.h"
 
 DECLARE_DWORD_COUNTER_STAT(TEXT("Num Custom Events"), STAT_NiagaraNumCustomEvents, STATGROUP_Niagara);
 
diff --git a/Engine/Plugins/FX/Niagara/Source/Niagara/Private/NiagaraEmptyUAVPool.cpp b/Engine/Plugins/FX/Niagara/Source/Niagara/Private/NiagaraEmptyUAVPool.cpp
--- a/Engine/Plugins/FX/Niagara/Source/Niagara/Private/NiagaraEmptyUAVPool.cpp
+++ b/Engine/Plugins/FX/Niagara/Source/Niagara/Private/NiagaraEmptyUAVPool.cpp
@@ -2,6 +2,8 @@
 
 #include "NiagaraEmptyUAVPool.h"
 #include "NiagaraStats.h"
+#include "NiagaraTypes.h"
+#include "RHICommandList.h"
 
 DECLARE_DWORD_ACCUMULATOR_STAT(TEXT("# EmptyUAVs"), STAT_NiagaraEmptyUAVPool, STATGROUP_Niagara);
 
diff --git a/Engine/Plugins/FX/Niagara/Source/Niagara/Private/NiagaraEvents.cpp b/Engine/Plugins/FX/Niagara/Source/Niagara/Private/NiagaraEvents.cpp
--- a/Engine/Plugins/FX/Niagara/Source/Niagara/Private/NiagaraEvents.cpp
+++ b/Engine/Plugins/FX/Niagara/Source/Niagara/Private/NiagaraEvents.cpp
@@ -1,7 +1,6 @@
 // Copyright Epic Games, Inc. All Rights Reserved.
 
 #include "NiagaraEvents.h"
-#include "NiagaraSystemInstance.h"
 #include "NiagaraStats.h"
 
 #include UE_INLINE_GENERATED_CPP_BY_NAME(NiagaraEvents)
diff --git a/Engine/Plugins/FX/Niagara/Source/Niagara/Private/NiagaraFunctionLibrary.cpp b/Engine/Plugins/FX/Niagara/Source/Niagara/Private/NiagaraFunctionLibrary.cpp
--- a/Engine/Plugins/FX/Niagara/Source/Niagara/Private/NiagaraFunctionLibrary.cpp
+++ b/Engine/Plugins/FX/Niagara/Source/Niagara/Private/NiagaraFunctionLibrary.cpp
@@ -1,19 +1,22 @@
 // Copyright Epic Games, Inc. All Rights Reserved.
 
 #include "NiagaraFunctionLibrary.h"
-#include "EngineGlobals.h"
+#include "Components/StaticMeshComponent.h"
 #include "GameFramework/WorldSettings.h"
 #include "Engine/Engine.h"
+#include "Engine/SkeletalMesh.h"
 #include "NiagaraComponent.h"
-#include "NiagaraComponentSettings.h"
+#include "Engine/StaticMesh.h"
 #include "NiagaraSystem.h"
 #include "ContentStreaming.h"
+#include "Engine/Texture2D.h"
 #include "Internationalization/Internationalization.h"
 
 #include "NiagaraAsyncGpuTraceHelper.h"
+#include "NiagaraComponentPool.h"
 #include "NiagaraWorldManager.h"
 #include "DataInterface/NiagaraDataInterfaceStaticMesh.h"
-#include "NiagaraStats.h"
+#include "NiagaraDataInterfaceSkeletalMesh.h"
 #include "NiagaraDataInterfaceTexture.h"
 #include "NiagaraDataInterface2DArrayTexture.h"
 #include "NiagaraDataInterfaceVolumeTexture.h"
diff --git a/Engine/Plugins/FX/Niagara/Source/Niagara/Private/NiagaraGPUInstanceCountManager.cpp b/Engine/Plugins/FX/Niagara/Source/Niagara/Private/NiagaraGPUInstanceCountManager.cpp
--- a/Engine/Plugins/FX/Niagara/Source/Niagara/Private/NiagaraGPUInstanceCountManager.cpp
+++ b/Engine/Plugins/FX/Niagara/Source/Niagara/Private/NiagaraGPUInstanceCountManager.cpp
@@ -2,15 +2,15 @@
 
 #include "NiagaraGPUInstanceCountManager.h"
 #include "NiagaraEmptyUAVPool.h"
-#include "NiagaraGpuComputeDispatch.h"
+#include "NiagaraGpuComputeDispatchInterface.h"
 #include "NiagaraRenderer.h"
 #include "NiagaraStats.h"
 
 #include "Containers/DynamicRHIResourceArray.h"
 #include "GPUSortManager.h" // CopyUIntBufferToTargets
 #include "ProfilingDebugging/RealtimeGPUProfiler.h"
-#include "ClearQuad.h"
 #include "PipelineStateCache.h"
+#include "RHIGPUReadback.h"
 
 int32 GNiagaraMinGPUInstanceCount = 2048;
 static FAutoConsoleVariableRef CVarNiagaraMinGPUInstanceCount(
diff --git a/Engine/Plugins/FX/Niagara/Source/Niagara/Private/NiagaraGPUProfiler.cpp b/Engine/Plugins/FX/Niagara/Source/Niagara/Private/NiagaraGPUProfiler.cpp
--- a/Engine/Plugins/FX/Niagara/Source/Niagara/Private/NiagaraGPUProfiler.cpp
+++ b/Engine/Plugins/FX/Niagara/Source/Niagara/Private/NiagaraGPUProfiler.cpp
@@ -1,9 +1,8 @@
 ﻿// Copyright Epic Games, Inc. All Rights Reserved.
 
 #include "NiagaraGPUProfiler.h"
-#include "NiagaraGPUSystemTick.h"
-#include "NiagaraSimStageData.h"
 #include "HAL/IConsoleManager.h"
+#include "RHICommandList.h"
 
 bool GNiagaraGpuProfilingEnabled = true;
 static FAutoConsoleVariableRef CVarNiagaraGpuProfilingEnabled(
diff --git a/Engine/Plugins/FX/Niagara/Source/Niagara/Private/NiagaraGPUProfiler.h b/Engine/Plugins/FX/Niagara/Source/Niagara/Private/NiagaraGPUProfiler.h
--- a/Engine/Plugins/FX/Niagara/Source/Niagara/Private/NiagaraGPUProfiler.h
+++ b/Engine/Plugins/FX/Niagara/Source/Niagara/Private/NiagaraGPUProfiler.h
@@ -3,9 +3,6 @@
 #pragma once
 
 #include "NiagaraCommon.h"
-#include "NiagaraEmitter.h"
-#include "RHI.h"
-#include "RHICommandList.h"
 #include "RHIResources.h"
 #include "HAL/Platform.h"
 #include "Misc/Build.h"
diff --git a/Engine/Plugins/FX/Niagara/Source/Niagara/Private/NiagaraGPUProfilerInterface.cpp b/Engine/Plugins/FX/Niagara/Source/Niagara/Private/NiagaraGPUProfilerInterface.cpp
--- a/Engine/Plugins/FX/Niagara/Source/Niagara/Private/NiagaraGPUProfilerInterface.cpp
+++ b/Engine/Plugins/FX/Niagara/Source/Niagara/Private/NiagaraGPUProfilerInterface.cpp
@@ -9,6 +9,7 @@
 #include "NiagaraEmitter.h"
 
 #include "Async/Async.h"
+#include "NiagaraSimStageData.h"
 
 #if WITH_NIAGARA_GPU_PROFILER
 
diff --git a/Engine/Plugins/FX/Niagara/Source/Niagara/Private/NiagaraGPUSystemTick.cpp b/Engine/Plugins/FX/Niagara/Source/Niagara/Private/NiagaraGPUSystemTick.cpp
--- a/Engine/Plugins/FX/Niagara/Source/Niagara/Private/NiagaraGPUSystemTick.cpp
+++ b/Engine/Plugins/FX/Niagara/Source/Niagara/Private/NiagaraGPUSystemTick.cpp
@@ -1,9 +1,9 @@
 // Copyright Epic Games, Inc. All Rights Reserved.
 
 #include "NiagaraGPUSystemTick.h"
-#include "NiagaraSystemGpuComputeProxy.h"
+#include "NiagaraEmitterInstance.h"
 #include "NiagaraSystemInstance.h"
-#include "NiagaraSystemSimulation.h"
+#include "NiagaraSystem.h"
 
 void FNiagaraGPUSystemTick::Init(FNiagaraSystemInstance* InSystemInstance)
 {
diff --git a/Engine/Plugins/FX/Niagara/Source/Niagara/Private/NiagaraGpuComputeDebug.cpp b/Engine/Plugins/FX/Niagara/Source/Niagara/Private/NiagaraGpuComputeDebug.cpp
--- a/Engine/Plugins/FX/Niagara/Source/Niagara/Private/NiagaraGpuComputeDebug.cpp
+++ b/Engine/Plugins/FX/Niagara/Source/Niagara/Private/NiagaraGpuComputeDebug.cpp
@@ -7,12 +7,10 @@
 #include "CanvasTypes.h"
 #include "CommonRenderResources.h"
 #include "Engine/Font.h"
-#include "Modules/ModuleManager.h"
 #include "RHI.h"
 #include "RenderGraphBuilder.h"
 #include "RenderGraphUtils.h"
 #include "ScreenPass.h"
-#include "ScreenRendering.h"
 
 int32 GNiagaraGpuComputeDebug_MinTextureHeight = 128;
 static FAutoConsoleVariableRef CVarNiagaraGpuComputeDebug_MinTextureHeight(
diff --git a/Engine/Plugins/FX/Niagara/Source/Niagara/Private/NiagaraGpuComputeDispatch.cpp b/Engine/Plugins/FX/Niagara/Source/Niagara/Private/NiagaraGpuComputeDispatch.cpp
--- a/Engine/Plugins/FX/Niagara/Source/Niagara/Private/NiagaraGpuComputeDispatch.cpp
+++ b/Engine/Plugins/FX/Niagara/Source/Niagara/Private/NiagaraGpuComputeDispatch.cpp
@@ -4,15 +4,13 @@
 
 #include "Async/Async.h"
 #include "CanvasTypes.h"
-#include "ClearQuad.h"
 #include "EngineModule.h"
-#include "Engine/Canvas.h"
 #include "Engine/Engine.h"
-#include "Engine/GameViewportClient.h"
 #include "Engine/World.h"
 #include "EngineModule.h"
 #include "Materials/MaterialRenderProxy.h"
 #include "Misc/ScopeExit.h"
+#include "NiagaraGPUSortInfo.h"
 #include "ProfilingDebugging/CsvProfiler.h"
 
 #include "NiagaraAsyncGpuTraceHelper.h"
@@ -23,6 +21,7 @@
 #include "NiagaraGPUProfilerInterface.h"
 #include "NiagaraGpuReadbackManager.h"
 #include "NiagaraRenderer.h"
+#include "NiagaraScript.h"
 #include "NiagaraShader.h"
 #include "NiagaraShaderParticleID.h"
 #include "NiagaraSortingGPU.h"
@@ -33,7 +32,6 @@
 #include "RHI.h"
 #include "SceneInterface.h"
 #include "SceneRenderTargetParameters.h"
-#include "ShaderParameterUtils.h"
 #include "TextureResource.h"
 #include "FXRenderingUtils.h"
 
diff --git a/Engine/Plugins/FX/Niagara/Source/Niagara/Private/NiagaraGpuComputeDispatch.h b/Engine/Plugins/FX/Niagara/Source/Niagara/Private/NiagaraGpuComputeDispatch.h
--- a/Engine/Plugins/FX/Niagara/Source/Niagara/Private/NiagaraGpuComputeDispatch.h
+++ b/Engine/Plugins/FX/Niagara/Source/Niagara/Private/NiagaraGpuComputeDispatch.h
@@ -7,24 +7,19 @@ the same VectorVM byte code / compute shader code
 ==============================================================================*/
 #pragma once
 
-#include "CoreMinimal.h"
 #include "FXSystem.h"
-#include "Modules/ModuleManager.h"
 #include "NiagaraCommon.h"
 #include "NiagaraComputeExecutionContext.h"
-#include "NiagaraEmitter.h"
 #include "NiagaraGPUProfiler.h"
-#include "NiagaraGPUSortInfo.h"
 #include "NiagaraGPUSystemTick.h"
 #include "NiagaraGpuComputeDispatchInterface.h"
-#include "NiagaraParameters.h"
-#include "NiagaraRendererProperties.h"
-#include "NiagaraSystemGpuComputeProxy.h"
-#include "ParticleResources.h"
+#include "NiagaraSimStageData.h"
 #include "RenderGraphUtils.h"
 #include "RHIResources.h"
 #include "RendererInterface.h"
-#include "Tickable.h"
+
+enum class EGPUSortFlags : uint32;
+struct FNiagaraGPUSortInfo;
 
 class FGPUSortManager;
 class FNiagaraAsyncGpuTraceHelper;
diff --git a/Engine/Plugins/FX/Niagara/Source/Niagara/Private/NiagaraGpuReadbackManager.cpp b/Engine/Plugins/FX/Niagara/Source/Niagara/Private/NiagaraGpuReadbackManager.cpp
--- a/Engine/Plugins/FX/Niagara/Source/Niagara/Private/NiagaraGpuReadbackManager.cpp
+++ b/Engine/Plugins/FX/Niagara/Source/Niagara/Private/NiagaraGpuReadbackManager.cpp
@@ -1,9 +1,7 @@
 // Copyright Epic Games, Inc. All Rights Reserved.
 
 #include "NiagaraGpuReadbackManager.h"
-#include "NiagaraCommon.h"
 
-#include "CoreMinimal.h"
 #include "RenderGraphUtils.h"
 #include "RHI.h"
 #include "RHICommandList.h"
diff --git a/Engine/Plugins/FX/Niagara/Source/Niagara/Private/NiagaraGpuScratchPad.h b/Engine/Plugins/FX/Niagara/Source/Niagara/Private/NiagaraGpuScratchPad.h
--- a/Engine/Plugins/FX/Niagara/Source/Niagara/Private/NiagaraGpuScratchPad.h
+++ b/Engine/Plugins/FX/Niagara/Source/Niagara/Private/NiagaraGpuScratchPad.h
@@ -1,12 +1,10 @@
 // Copyright Epic Games, Inc. All Rights Reserved.
 #pragma once
 
-#include "CoreMinimal.h"
 
 #include "RHIDefinitions.h"
-#include "PrimitiveSceneInfo.h"
-#include "RHI.h"
 #include "RHIUtilities.h"
+#include "RendererInterface.h"
 
 class FScene;
 class FSceneView;
diff --git a/Engine/Plugins/FX/Niagara/Source/Niagara/Private/NiagaraLensEffectBase.cpp b/Engine/Plugins/FX/Niagara/Source/Niagara/Private/NiagaraLensEffectBase.cpp
--- a/Engine/Plugins/FX/Niagara/Source/Niagara/Private/NiagaraLensEffectBase.cpp
+++ b/Engine/Plugins/FX/Niagara/Source/Niagara/Private/NiagaraLensEffectBase.cpp
@@ -3,7 +3,6 @@
 #include "NiagaraLensEffectBase.h"
 
 #include "NiagaraComponent.h"
-#include "NiagaraEffectType.h"
 #include "NiagaraSystem.h"
 #include "Camera/PlayerCameraManager.h"
 
diff --git a/Engine/Plugins/FX/Niagara/Source/Niagara/Private/NiagaraLightRendererProperties.cpp b/Engine/Plugins/FX/Niagara/Source/Niagara/Private/NiagaraLightRendererProperties.cpp
--- a/Engine/Plugins/FX/Niagara/Source/Niagara/Private/NiagaraLightRendererProperties.cpp
+++ b/Engine/Plugins/FX/Niagara/Source/Niagara/Private/NiagaraLightRendererProperties.cpp
@@ -1,8 +1,10 @@
 // Copyright Epic Games, Inc. All Rights Reserved.
 
 #include "NiagaraLightRendererProperties.h"
+#include "NiagaraBoundsCalculator.h"
 #include "NiagaraRenderer.h"
 #include "NiagaraConstants.h"
+#include "NiagaraModule.h"
 #include "NiagaraRendererLights.h"
 #include "Modules/ModuleManager.h"
 
@@ -13,7 +15,6 @@
 #include "Styling/SlateIconFinder.h"
 #include "Widgets/SWidget.h"
 #include "Styling/SlateBrush.h"
-#include "AssetThumbnail.h"
 #include "Widgets/Text/STextBlock.h"
 #endif
 
diff --git a/Engine/Plugins/FX/Niagara/Source/Niagara/Private/NiagaraMeshRendererProperties.cpp b/Engine/Plugins/FX/Niagara/Source/Niagara/Private/NiagaraMeshRendererProperties.cpp
--- a/Engine/Plugins/FX/Niagara/Source/Niagara/Private/NiagaraMeshRendererProperties.cpp
+++ b/Engine/Plugins/FX/Niagara/Source/Niagara/Private/NiagaraMeshRendererProperties.cpp
@@ -1,6 +1,7 @@
 // Copyright Epic Games, Inc. All Rights Reserved.
 
 #include "NiagaraMeshRendererProperties.h"
+#include "Materials/Material.h"
 #include "NiagaraRendererMeshes.h"
 #include "NiagaraConstants.h"
 #include "NiagaraBoundsCalculatorHelper.h"
@@ -18,17 +19,17 @@
 #include UE_INLINE_GENERATED_CPP_BY_NAME(NiagaraMeshRendererProperties)
 
 #if WITH_EDITOR
-#include "Editor.h"
 #include "AssetThumbnail.h"
+#include "NiagaraModule.h"
 #include "Styling/SlateIconFinder.h"
 #include "Internationalization/Regex.h"
 #include "Dialogs/Dialogs.h"
 #include "Framework/Notifications/NotificationManager.h"
+#include "StaticMeshResources.h"
 #include "Widgets/Images/SImage.h"
 #include "Widgets/SWidget.h"
 #include "Widgets/Text/STextBlock.h"
 #include "Widgets/Notifications/SNotificationList.h"
-#include "Subsystems/ImportSubsystem.h"
 #endif
 
 #define LOCTEXT_NAMESPACE "UNiagaraMeshRendererProperties"
diff --git a/Engine/Plugins/FX/Niagara/Source/Niagara/Private/NiagaraModule.cpp b/Engine/Plugins/FX/Niagara/Source/Niagara/Private/NiagaraModule.cpp
--- a/Engine/Plugins/FX/Niagara/Source/Niagara/Private/NiagaraModule.cpp
+++ b/Engine/Plugins/FX/Niagara/Source/Niagara/Private/NiagaraModule.cpp
@@ -1,21 +1,26 @@
 // Copyright Epic Games, Inc. All Rights Reserved.
 
 #include "NiagaraModule.h"
+#include "GPUSortManager.h"
 #include "Misc/LazySingleton.h"
+#include "Materials/MaterialInterface.h"
 #include "Modules/ModuleManager.h"
 #include "NiagaraComponentSettings.h"
 #include "NiagaraCompileHashVisitor.h"
+#include "NiagaraEmitter.h"
 #include "NiagaraTypes.h"
-#include "NiagaraDebugVis.h"
+#include "NiagaraDebugVis.h" // IWYU pragma: keep
 #include "NiagaraEvents.h"
 #include "NiagaraSettings.h"
 #include "NiagaraDataChannelManager.h"
-#include "NiagaraDataInterfaceCurlNoise.h"
+#include "NiagaraSimCache.h"
 #include "NiagaraUseOpenVDB.h"
+#include "NiagaraSystem.h"
 #include "UObject/Class.h"
 #include "UObject/ObjectRedirector.h"
 #include "UObject/Package.h"
 #include "NiagaraWorldManager.h"
+#include "UObject/UObjectIterator.h"
 #include "VectorVM.h"
 #include "NiagaraConstants.h"
 #include "NiagaraDecalRendererProperties.h"
@@ -25,7 +30,6 @@
 #include "NiagaraRibbonRendererProperties.h"
 #include "NiagaraComponentRendererProperties.h"
 #include "NiagaraCustomVersion.h"
-#include "NiagaraRenderer.h"
 #include "NiagaraShaderModule.h"
 #include "UObject/CoreRedirects.h"
 #include "NiagaraGpuComputeDispatch.h"
diff --git a/Engine/Plugins/FX/Niagara/Source/Niagara/Private/NiagaraParameterCollection.cpp b/Engine/Plugins/FX/Niagara/Source/Niagara/Private/NiagaraParameterCollection.cpp
--- a/Engine/Plugins/FX/Niagara/Source/Niagara/Private/NiagaraParameterCollection.cpp
+++ b/Engine/Plugins/FX/Niagara/Source/Niagara/Private/NiagaraParameterCollection.cpp
@@ -1,13 +1,14 @@
 // Copyright Epic Games, Inc. All Rights Reserved.
 
 #include "NiagaraParameterCollection.h"
+#include "AssetRegistry/AssetData.h"
 #include "Engine/World.h"
+#include "Modules/ModuleManager.h"
 #include "NiagaraDataInterface.h"
 #include "Materials/MaterialParameterCollection.h"
 #include "Materials/MaterialParameterCollectionInstance.h"
 #include "Misc/SecureHash.h"
 #if WITH_EDITORONLY_DATA
-	#include "IAssetTools.h"
 #endif
 #include "AssetRegistry/AssetRegistryModule.h"
 
diff --git a/Engine/Plugins/FX/Niagara/Source/Niagara/Private/NiagaraParameterStore.cpp b/Engine/Plugins/FX/Niagara/Source/Niagara/Private/NiagaraParameterStore.cpp
--- a/Engine/Plugins/FX/Niagara/Source/Niagara/Private/NiagaraParameterStore.cpp
+++ b/Engine/Plugins/FX/Niagara/Source/Niagara/Private/NiagaraParameterStore.cpp
@@ -2,11 +2,8 @@
 
 #include "NiagaraParameterStore.h"
 #include "Algo/Find.h"
-#include "NiagaraDataSet.h"
-#include "NiagaraComponent.h"
 #include "NiagaraCustomVersion.h"
 #include "NiagaraDataInterface.h"
-#include "NiagaraSystemInstance.h"
 #include "NiagaraStats.h"
 
 #include UE_INLINE_GENERATED_CPP_BY_NAME(NiagaraParameterStore)
diff --git a/Engine/Plugins/FX/Niagara/Source/Niagara/Private/NiagaraParameters.cpp b/Engine/Plugins/FX/Niagara/Source/Niagara/Private/NiagaraParameters.cpp
--- a/Engine/Plugins/FX/Niagara/Source/Niagara/Private/NiagaraParameters.cpp
+++ b/Engine/Plugins/FX/Niagara/Source/Niagara/Private/NiagaraParameters.cpp
@@ -1,7 +1,6 @@
 // Copyright Epic Games, Inc. All Rights Reserved.
 
 #include "NiagaraParameters.h"
-#include "NiagaraEmitter.h"
 
 #include UE_INLINE_GENERATED_CPP_BY_NAME(NiagaraParameters)
 
diff --git a/Engine/Plugins/FX/Niagara/Source/Niagara/Private/NiagaraPerfBaseline.cpp b/Engine/Plugins/FX/Niagara/Source/Niagara/Private/NiagaraPerfBaseline.cpp
--- a/Engine/Plugins/FX/Niagara/Source/Niagara/Private/NiagaraPerfBaseline.cpp
+++ b/Engine/Plugins/FX/Niagara/Source/Niagara/Private/NiagaraPerfBaseline.cpp
@@ -8,14 +8,13 @@
 #include "NiagaraFunctionLibrary.h"
 #include "Components/TextRenderComponent.h"
 #include "Misc/OutputDeviceArchiveWrapper.h"
+#include "NiagaraSystem.h"
 #include "RenderingThread.h"
 #include "HAL/IConsoleManager.h"
 #include "HAL/FileManager.h"
-#include "Slate/SceneViewport.h"
-#include "HighResScreenshot.h"
 #include "CanvasTypes.h"
 #include "HAL/FileManager.h"
-#include "UObject/UObjectIterator.h"
+#include "UnrealClient.h"
 
 #include UE_INLINE_GENERATED_CPP_BY_NAME(NiagaraPerfBaseline)
 
diff --git a/Engine/Plugins/FX/Niagara/Source/Niagara/Private/NiagaraPlatformSet.cpp b/Engine/Plugins/FX/Niagara/Source/Niagara/Private/NiagaraPlatformSet.cpp
--- a/Engine/Plugins/FX/Niagara/Source/Niagara/Private/NiagaraPlatformSet.cpp
+++ b/Engine/Plugins/FX/Niagara/Source/Niagara/Private/NiagaraPlatformSet.cpp
@@ -3,19 +3,17 @@
 #include "NiagaraPlatformSet.h"
 #include "DeviceProfiles/DeviceProfile.h"
 #include "DeviceProfiles/DeviceProfileManager.h"
+#include "NiagaraEmitter.h"
 #include "Scalability.h"
 #include "Misc/ConfigCacheIni.h"
 #include "Misc/ConfigUtilities.h"
-#include "Interfaces/ITargetPlatform.h"
 #include "NiagaraSystemImpl.h"
 #include "NiagaraSettings.h"
-#include "SystemSettings.h"
 #include "UObject/UObjectIterator.h"
 
 #include UE_INLINE_GENERATED_CPP_BY_NAME(NiagaraPlatformSet)
 
 #if WITH_EDITOR
-#include "PlatformInfo.h"
 #endif
 
 #define LOCTEXT_NAMESPACE "NiagaraPlatformSet"
diff --git a/Engine/Plugins/FX/Niagara/Source/Niagara/Private/NiagaraPreviewGrid.cpp b/Engine/Plugins/FX/Niagara/Source/Niagara/Private/NiagaraPreviewGrid.cpp
--- a/Engine/Plugins/FX/Niagara/Source/Niagara/Private/NiagaraPreviewGrid.cpp
+++ b/Engine/Plugins/FX/Niagara/Source/Niagara/Private/NiagaraPreviewGrid.cpp
@@ -2,14 +2,10 @@
 
 #include "NiagaraPreviewGrid.h"
 #include "UObject/ConstructorHelpers.h"
-#include "Engine/BlueprintGeneratedClass.h"
 #include "Engine/Texture2D.h"
 #include "NiagaraComponent.h"
-#include "NiagaraEffectType.h"
 #include "NiagaraSystem.h"
 
-#include "Components/TextRenderComponent.h"
-#include "Components/StaticMeshComponent.h"
 
 #include UE_INLINE_GENERATED_CPP_BY_NAME(NiagaraPreviewGrid)
 
diff --git a/Engine/Plugins/FX/Niagara/Source/Niagara/Private/NiagaraRenderGraphUtils.cpp b/Engine/Plugins/FX/Niagara/Source/Niagara/Private/NiagaraRenderGraphUtils.cpp
--- a/Engine/Plugins/FX/Niagara/Source/Niagara/Private/NiagaraRenderGraphUtils.cpp
+++ b/Engine/Plugins/FX/Niagara/Source/Niagara/Private/NiagaraRenderGraphUtils.cpp
@@ -2,7 +2,6 @@
 #include "NiagaraRenderGraphUtils.h"
 #include "NiagaraStats.h"
 
-#include "RenderGraph.h"
 #include "RenderGraphUtils.h"
 
 //////////////////////////////////////////////////////////////////////////
diff --git a/Engine/Plugins/FX/Niagara/Source/Niagara/Private/NiagaraRenderer.cpp b/Engine/Plugins/FX/Niagara/Source/Niagara/Private/NiagaraRenderer.cpp
--- a/Engine/Plugins/FX/Niagara/Source/Niagara/Private/NiagaraRenderer.cpp
+++ b/Engine/Plugins/FX/Niagara/Source/Niagara/Private/NiagaraRenderer.cpp
@@ -2,18 +2,18 @@
 
 #include "NiagaraRenderer.h"
 #include "Engine/Texture.h"
-#include "ParticleResources.h"
 #include "MaterialDomain.h"
+#include "Materials/Material.h"
 #include "NiagaraDataSet.h"
 #include "NiagaraSceneProxy.h"
 #include "NiagaraStats.h"
-#include "NiagaraComponent.h"
 #include "DynamicBufferAllocator.h"
 #include "NiagaraComputeExecutionContext.h"
 #include "NiagaraGPUSortInfo.h"
 #include "NiagaraEmitterInstance.h"
 #include "NiagaraSystemInstanceController.h"
 #include "Materials/MaterialInstanceDynamic.h"
+#include "PrimitiveViewRelevance.h"
 
 DECLARE_CYCLE_STAT(TEXT("Sort Particles"), STAT_NiagaraSortParticles, STATGROUP_Niagara);
 DECLARE_CYCLE_STAT(TEXT("Global Float Alloc - All"), STAT_NiagaraAllocateGlobalFloatAll, STATGROUP_Niagara);
diff --git a/Engine/Plugins/FX/Niagara/Source/Niagara/Private/NiagaraRendererComponents.cpp b/Engine/Plugins/FX/Niagara/Source/Niagara/Private/NiagaraRendererComponents.cpp
--- a/Engine/Plugins/FX/Niagara/Source/Niagara/Private/NiagaraRendererComponents.cpp
+++ b/Engine/Plugins/FX/Niagara/Source/Niagara/Private/NiagaraRendererComponents.cpp
@@ -2,10 +2,11 @@
 
 #include "NiagaraRendererComponents.h"
 #include "Engine/World.h"
-#include "NiagaraConstants.h"
 #include "NiagaraDataSet.h"
+#include "NiagaraEmitterInstance.h"
 #include "NiagaraStats.h"
 #include "NiagaraComponentRendererProperties.h"
+#include "NiagaraSystem.h"
 #include "NiagaraSystemInstance.h"
 #include "Async/Async.h"
 
diff --git a/Engine/Plugins/FX/Niagara/Source/Niagara/Private/NiagaraRendererLights.cpp b/Engine/Plugins/FX/Niagara/Source/Niagara/Private/NiagaraRendererLights.cpp
--- a/Engine/Plugins/FX/Niagara/Source/Niagara/Private/NiagaraRendererLights.cpp
+++ b/Engine/Plugins/FX/Niagara/Source/Niagara/Private/NiagaraRendererLights.cpp
@@ -1,7 +1,6 @@
 // Copyright Epic Games, Inc. All Rights Reserved.
 
 #include "NiagaraRendererLights.h"
-#include "ParticleResources.h"
 #include "NiagaraDataSet.h"
 #include "NiagaraDataSetAccessor.h"
 #include "NiagaraEmitterInstance.h"
@@ -9,13 +8,12 @@
 #include "NiagaraSettings.h"
 #include "NiagaraStats.h"
 #include "NiagaraSystemInstance.h"
-#include "NiagaraVertexFactory.h"
 #include "NiagaraComponent.h"
-#include "Engine/Engine.h"
 
 #include "NiagaraLightRendererProperties.h"
 #include "NiagaraRendererLights.h"
 #include "NiagaraCullProxyComponent.h"
+#include "PrimitiveViewRelevance.h"
 #include "SceneInterface.h"
 
 
diff --git a/Engine/Plugins/FX/Niagara/Source/Niagara/Private/NiagaraRendererMeshes.cpp b/Engine/Plugins/FX/Niagara/Source/Niagara/Private/NiagaraRendererMeshes.cpp
--- a/Engine/Plugins/FX/Niagara/Source/Niagara/Private/NiagaraRendererMeshes.cpp
+++ b/Engine/Plugins/FX/Niagara/Source/Niagara/Private/NiagaraRendererMeshes.cpp
@@ -1,12 +1,11 @@
 // Copyright Epic Games, Inc. All Rights Reserved.
 
 #include "NiagaraRendererMeshes.h"
-#include "Async/ParallelFor.h"
 #include "DataDrivenShaderPlatformInfo.h"
 #include "Engine/StaticMesh.h"
 #include "FXRenderingUtils.h"
-#include "IXRTrackingSystem.h"
 #include "MaterialDomain.h"
+#include "Materials/Material.h"
 #include "Materials/MaterialInstanceDynamic.h"
 #include "Materials/MaterialRenderProxy.h"
 #include "NiagaraComponent.h"
@@ -14,21 +13,22 @@
 #include "NiagaraDataSet.h"
 #include "NiagaraEmitterInstance.h"
 #include "NiagaraGPURayTracingTransformsShader.h"
+#include "NiagaraGPUSortInfo.h"
 #include "NiagaraGpuComputeDispatchInterface.h"
 #include "NiagaraRenderableMeshInterface.h"
 #include "NiagaraSceneProxy.h"
 #include "NiagaraSortingGPU.h"
 #include "NiagaraStats.h"
 #include "NiagaraSystemInstanceController.h"
-#include "ParticleResources.h"
+#include "NiagaraSystemInstance.h"
 #include "PipelineStateCache.h"
-#include "RayTracingDefinitions.h"
-#include "RayTracingDynamicGeometryCollection.h"
 #include "RayTracingInstance.h"
 #include "RenderGraphBuilder.h"
 #include "RenderGraphUtils.h"
+#include "RenderUtils.h"
 #include "SceneInterface.h"
 #include "SceneRendererInterface.h"
+#include "StaticMeshResources.h"
 
 DECLARE_CYCLE_STAT(TEXT("Generate Mesh Vertex Data [GT]"), STAT_NiagaraGenMeshVertexData, STATGROUP_Niagara);
 
diff --git a/Engine/Plugins/FX/Niagara/Source/Niagara/Private/NiagaraRendererRibbons.cpp b/Engine/Plugins/FX/Niagara/Source/Niagara/Private/NiagaraRendererRibbons.cpp
--- a/Engine/Plugins/FX/Niagara/Source/Niagara/Private/NiagaraRendererRibbons.cpp
+++ b/Engine/Plugins/FX/Niagara/Source/Niagara/Private/NiagaraRendererRibbons.cpp
@@ -4,16 +4,16 @@
 
 #include "GlobalRenderResources.h"
 #include "GPUSortManager.h"
-#include "ParticleResources.h"
+#include "Materials/Material.h"
 #include "NiagaraRibbonVertexFactory.h"
 #include "NiagaraDataSet.h"
 #include "NiagaraDataSetAccessor.h"
+#include "NiagaraEmitterInstance.h"
 #include "NiagaraSceneProxy.h"
+#include "NiagaraGpuComputeDataManager.h"
 #include "NiagaraStats.h"
 #include "NiagaraSystemInstance.h"
 #include "NiagaraComponent.h"
-#include "RayTracingDefinitions.h"
-#include "RayTracingDynamicGeometryCollection.h"
 #include "RayTracingInstance.h"
 #include "Materials/MaterialRenderProxy.h"
 #include "MaterialDomain.h"
@@ -21,8 +21,9 @@
 #include "NiagaraCullProxyComponent.h"
 #include "NiagaraGpuComputeDispatchInterface.h"
 #include "NiagaraRibbonCompute.h"
-#include "Misc/LazySingleton.h"
 #include "RenderGraphUtils.h"
+#include "RenderUtils.h"
+#include "SceneManagement.h"
 
 DECLARE_CYCLE_STAT(TEXT("Generate Ribbon Vertex Data [GT]"), STAT_NiagaraGenRibbonVertexData, STATGROUP_Niagara);
 DECLARE_CYCLE_STAT(TEXT("Render Ribbons [RT]"), STAT_NiagaraRenderRibbons, STATGROUP_Niagara);
diff --git a/Engine/Plugins/FX/Niagara/Source/Niagara/Private/NiagaraRendererSprites.cpp b/Engine/Plugins/FX/Niagara/Source/Niagara/Private/NiagaraRendererSprites.cpp
--- a/Engine/Plugins/FX/Niagara/Source/Niagara/Private/NiagaraRendererSprites.cpp
+++ b/Engine/Plugins/FX/Niagara/Source/Niagara/Private/NiagaraRendererSprites.cpp
@@ -2,21 +2,21 @@
 
 #include "NiagaraRendererSprites.h"
 #include "MaterialDomain.h"
-#include "Materials/MaterialInstanceDynamic.h"
+#include "Materials/Material.h"
 #include "Materials/MaterialRenderProxy.h"
 #include "NiagaraComponent.h"
 #include "NiagaraCullProxyComponent.h"
 #include "NiagaraCutoutVertexBuffer.h"
 #include "NiagaraDataSet.h"
+#include "NiagaraEmitterInstance.h"
 #include "NiagaraGpuComputeDispatchInterface.h"
+#include "NiagaraGPUSortInfo.h"
 #include "NiagaraSceneProxy.h"
 #include "NiagaraSettings.h"
 #include "NiagaraSortingGPU.h"
 #include "NiagaraStats.h"
 #include "NiagaraSystemInstance.h"
 #include "ParticleResources.h"
-#include "RayTracingDefinitions.h"
-#include "RayTracingDynamicGeometryCollection.h"
 #include "RayTracingInstance.h"
 
 DECLARE_DWORD_COUNTER_STAT(TEXT("NumSprites"), STAT_NiagaraNumSprites, STATGROUP_Niagara);
diff --git a/Engine/Plugins/FX/Niagara/Source/Niagara/Private/NiagaraRibbonRendererProperties.cpp b/Engine/Plugins/FX/Niagara/Source/Niagara/Private/NiagaraRibbonRendererProperties.cpp
--- a/Engine/Plugins/FX/Niagara/Source/Niagara/Private/NiagaraRibbonRendererProperties.cpp
+++ b/Engine/Plugins/FX/Niagara/Source/Niagara/Private/NiagaraRibbonRendererProperties.cpp
@@ -1,6 +1,7 @@
 // Copyright Epic Games, Inc. All Rights Reserved.
 
 #include "NiagaraRibbonRendererProperties.h"
+#include "NiagaraModule.h"
 #include "NiagaraRendererRibbons.h"
 #include "NiagaraConstants.h"
 #include "NiagaraBoundsCalculatorHelper.h"
diff --git a/Engine/Plugins/FX/Niagara/Source/Niagara/Private/NiagaraScratchPadContainer.cpp b/Engine/Plugins/FX/Niagara/Source/Niagara/Private/NiagaraScratchPadContainer.cpp
--- a/Engine/Plugins/FX/Niagara/Source/Niagara/Private/NiagaraScratchPadContainer.cpp
+++ b/Engine/Plugins/FX/Niagara/Source/Niagara/Private/NiagaraScratchPadContainer.cpp
@@ -2,6 +2,7 @@
 
 
 #include "NiagaraScratchPadContainer.h"
+#include "NiagaraScript.h"
 
 #include UE_INLINE_GENERATED_CPP_BY_NAME(NiagaraScratchPadContainer)
 
diff --git a/Engine/Plugins/FX/Niagara/Source/Niagara/Private/NiagaraScript.cpp b/Engine/Plugins/FX/Niagara/Source/Niagara/Private/NiagaraScript.cpp
--- a/Engine/Plugins/FX/Niagara/Source/Niagara/Private/NiagaraScript.cpp
+++ b/Engine/Plugins/FX/Niagara/Source/Niagara/Private/NiagaraScript.cpp
@@ -5,20 +5,19 @@
 #include "Algo/RemoveIf.h"
 #include "Algo/Unique.h"
 #include "DataDrivenShaderPlatformInfo.h"
+#include "Misc/Compression.h"
 #include "Modules/ModuleManager.h"
+#include "NiagaraCompileHashVisitor.h"
 #include "NiagaraScriptSourceBase.h"
-#include "NiagaraComponent.h"
 #include "NiagaraEmitter.h"
+#include "NiagaraShader.h"
 #include "UObject/Package.h"
 #include "NiagaraModule.h"
 #include "NiagaraCustomVersion.h"
-#include "NiagaraShaderCompilationManager.h"
-#include "NiagaraShaderModule.h"
-#include "NiagaraShaderParametersBuilder.h"
+#include "NiagaraStackSection.h"
 #include "NiagaraSystem.h"
 #include "Serialization/MemoryReader.h"
 #include "Misc/SecureHash.h"
-#include "ShaderParameterMetadataBuilder.h"
 
 #include "ProfilingDebugging/CookStats.h"
 #include "Stats/Stats.h"
@@ -46,12 +45,12 @@
 
 #include "UObject/FortniteMainBranchObjectVersion.h"
 #include "UObject/RenderingObjectVersion.h"
+#include "UObject/UObjectIterator.h"
 #include "UObject/UObjectThreadContext.h"
 #include "Serialization/ObjectAndNameAsStringProxyArchive.h"
 
 #include "VectorVM.h"
 #include "NiagaraSimulationStageBase.h"
-#include "Async/Async.h"
 
 #include "Algo/Sort.h"
 #include "Misc/CoreMiscDefines.h"
diff --git a/Engine/Plugins/FX/Niagara/Source/Niagara/Private/NiagaraScriptExecutionContext.cpp b/Engine/Plugins/FX/Niagara/Source/Niagara/Private/NiagaraScriptExecutionContext.cpp
--- a/Engine/Plugins/FX/Niagara/Source/Niagara/Private/NiagaraScriptExecutionContext.cpp
+++ b/Engine/Plugins/FX/Niagara/Source/Niagara/Private/NiagaraScriptExecutionContext.cpp
@@ -1,15 +1,13 @@
 // Copyright Epic Games, Inc. All Rights Reserved.
 
 #include "NiagaraScriptExecutionContext.h"
+#include "NiagaraScript.h"
 #include "NiagaraSettings.h"
 #include "NiagaraStats.h"
 #include "NiagaraDataInterface.h"
+#include "NiagaraSystem.h"
 #include "NiagaraSystemInstance.h"
-#include "NiagaraWorldManager.h"
 #include "NiagaraFunctionLibrary.h"
-#include "NiagaraGPUInstanceCountManager.h"
-#include "Interfaces/ITargetPlatform.h"
-#include "HAL/PlatformFileManager.h"
 
 DECLARE_CYCLE_STAT(TEXT("Register Setup"), STAT_NiagaraSimRegisterSetup, STATGROUP_Niagara);
 DECLARE_CYCLE_STAT(TEXT("Context Ticking"), STAT_NiagaraScriptExecContextTick, STATGROUP_Niagara);
diff --git a/Engine/Plugins/FX/Niagara/Source/Niagara/Private/NiagaraScriptExecutionParameterStore.cpp b/Engine/Plugins/FX/Niagara/Source/Niagara/Private/NiagaraScriptExecutionParameterStore.cpp
--- a/Engine/Plugins/FX/Niagara/Source/Niagara/Private/NiagaraScriptExecutionParameterStore.cpp
+++ b/Engine/Plugins/FX/Niagara/Source/Niagara/Private/NiagaraScriptExecutionParameterStore.cpp
@@ -2,9 +2,8 @@
 
 #include "NiagaraScriptExecutionParameterStore.h"
 #include "NiagaraConstants.h"
-#include "NiagaraStats.h"
 #include "NiagaraDataInterface.h"
-#include "NiagaraSystemInstance.h"
+#include "NiagaraScript.h"
 
 #include UE_INLINE_GENERATED_CPP_BY_NAME(NiagaraScriptExecutionParameterStore)
 
diff --git a/Engine/Plugins/FX/Niagara/Source/Niagara/Private/NiagaraSettings.cpp b/Engine/Plugins/FX/Niagara/Source/Niagara/Private/NiagaraSettings.cpp
--- a/Engine/Plugins/FX/Niagara/Source/Niagara/Private/NiagaraSettings.cpp
+++ b/Engine/Plugins/FX/Niagara/Source/Niagara/Private/NiagaraSettings.cpp
@@ -1,6 +1,7 @@
 // Copyright Epic Games, Inc. All Rights Reserved.
 
 #include "NiagaraSettings.h"
+#include "NiagaraEffectType.h"
 
 UNiagaraSettings::UNiagaraSettings(const FObjectInitializer& ObjectInitlaizer)
 	: Super(ObjectInitlaizer)
diff --git a/Engine/Plugins/FX/Niagara/Source/Niagara/Private/NiagaraSimCache.cpp b/Engine/Plugins/FX/Niagara/Source/Niagara/Private/NiagaraSimCache.cpp
--- a/Engine/Plugins/FX/Niagara/Source/Niagara/Private/NiagaraSimCache.cpp
+++ b/Engine/Plugins/FX/Niagara/Source/Niagara/Private/NiagaraSimCache.cpp
@@ -2,6 +2,8 @@
 
 #include "NiagaraSimCache.h"
 
+#include "Engine/World.h"
+#include "Misc/LargeWorldRenderPosition.h"
 #include "NiagaraConstants.h"
 #include "NiagaraComponent.h"
 #include "NiagaraDataInterfaceUtilities.h"
diff --git a/Engine/Plugins/FX/Niagara/Source/Niagara/Private/NiagaraSimCacheAttributeReaderHelper.h b/Engine/Plugins/FX/Niagara/Source/Niagara/Private/NiagaraSimCacheAttributeReaderHelper.h
--- a/Engine/Plugins/FX/Niagara/Source/Niagara/Private/NiagaraSimCacheAttributeReaderHelper.h
+++ b/Engine/Plugins/FX/Niagara/Source/Niagara/Private/NiagaraSimCacheAttributeReaderHelper.h
@@ -2,6 +2,7 @@
 
 #pragma once
 
+#include "Math/Float16.h"
 #include "NiagaraSimCache.h"
 
 struct FNiagaraSimCacheAttributeReaderHelper
diff --git a/Engine/Plugins/FX/Niagara/Source/Niagara/Private/NiagaraSimCacheGpuResource.cpp b/Engine/Plugins/FX/Niagara/Source/Niagara/Private/NiagaraSimCacheGpuResource.cpp
--- a/Engine/Plugins/FX/Niagara/Source/Niagara/Private/NiagaraSimCacheGpuResource.cpp
+++ b/Engine/Plugins/FX/Niagara/Source/Niagara/Private/NiagaraSimCacheGpuResource.cpp
@@ -1,6 +1,8 @@
 // Copyright Epic Games, Inc. All Rights Reserved.
 
 #include "NiagaraSimCacheGpuResource.h"
+#include "NiagaraSimCache.h"
+#include "RenderingThread.h"
 
 void FNiagaraSimCacheGpuResource::FDeletor::operator()(FNiagaraSimCacheGpuResource* Ptr)
 {
diff --git a/Engine/Plugins/FX/Niagara/Source/Niagara/Private/NiagaraSimCacheGpuResource.h b/Engine/Plugins/FX/Niagara/Source/Niagara/Private/NiagaraSimCacheGpuResource.h
--- a/Engine/Plugins/FX/Niagara/Source/Niagara/Private/NiagaraSimCacheGpuResource.h
+++ b/Engine/Plugins/FX/Niagara/Source/Niagara/Private/NiagaraSimCacheGpuResource.h
@@ -4,9 +4,10 @@
 
 #include "Templates/SharedPointer.h"
 
-#include "NiagaraSimCache.h"
 #include "NiagaraRenderer.h"
 
+class UNiagaraSimCache;
+
 struct FNiagaraSimCacheGpuResource
 {
 	UE_NONCOPYABLE(FNiagaraSimCacheGpuResource)
diff --git a/Engine/Plugins/FX/Niagara/Source/Niagara/Private/NiagaraSimCacheHelper.h b/Engine/Plugins/FX/Niagara/Source/Niagara/Private/NiagaraSimCacheHelper.h
--- a/Engine/Plugins/FX/Niagara/Source/Niagara/Private/NiagaraSimCacheHelper.h
+++ b/Engine/Plugins/FX/Niagara/Source/Niagara/Private/NiagaraSimCacheHelper.h
@@ -7,6 +7,7 @@
 #include "NiagaraComputeExecutionContext.h"
 #include "NiagaraConstants.h"
 #include "NiagaraDataSetReadback.h"
+#include "NiagaraEmitterInstance.h"
 #include "NiagaraGpuComputeDispatchInterface.h"
 #include "NiagaraSimCache.h"
 #include "NiagaraSystemInstanceController.h"
diff --git a/Engine/Plugins/FX/Niagara/Source/Niagara/Private/NiagaraSimulationStageBase.cpp b/Engine/Plugins/FX/Niagara/Source/Niagara/Private/NiagaraSimulationStageBase.cpp
--- a/Engine/Plugins/FX/Niagara/Source/Niagara/Private/NiagaraSimulationStageBase.cpp
+++ b/Engine/Plugins/FX/Niagara/Source/Niagara/Private/NiagaraSimulationStageBase.cpp
@@ -3,6 +3,7 @@
 #include "NiagaraSimulationStageBase.h"
 #include "NiagaraCompileHashVisitor.h"
 #include "NiagaraEmitter.h"
+#include "NiagaraSimulationStageCompileData.h"
 #include "NiagaraSystem.h"
 #include "NiagaraScriptSourceBase.h"
 
diff --git a/Engine/Plugins/FX/Niagara/Source/Niagara/Private/NiagaraSpriteRendererProperties.cpp b/Engine/Plugins/FX/Niagara/Source/Niagara/Private/NiagaraSpriteRendererProperties.cpp
--- a/Engine/Plugins/FX/Niagara/Source/Niagara/Private/NiagaraSpriteRendererProperties.cpp
+++ b/Engine/Plugins/FX/Niagara/Source/Niagara/Private/NiagaraSpriteRendererProperties.cpp
@@ -1,14 +1,15 @@
 // Copyright Epic Games, Inc. All Rights Reserved.
 
 #include "NiagaraSpriteRendererProperties.h"
+#include "Materials/Material.h"
 #include "NiagaraRenderer.h"
 #include "NiagaraConstants.h"
+#include "NiagaraGPUSortInfo.h"
 #include "NiagaraRendererSprites.h"
 #include "NiagaraBoundsCalculatorHelper.h"
 #include "NiagaraCustomVersion.h"
 #include "NiagaraEmitterInstance.h"
 #include "NiagaraEmitter.h"
-#include "NiagaraScriptSourceBase.h"
 #include "NiagaraSystem.h"
 
 #include "Engine/Texture2D.h"
@@ -22,6 +23,7 @@
 #include "DerivedDataCacheInterface.h"
 #include "Materials/MaterialExpressionTextureSample.h"
 #include "Materials/MaterialExpressionTextureSampleParameter.h"
+#include "NiagaraModule.h"
 #include "Widgets/Images/SImage.h"
 #include "Styling/SlateIconFinder.h"
 #include "Widgets/SWidget.h"
diff --git a/Engine/Plugins/FX/Niagara/Source/Niagara/Private/NiagaraStats.h b/Engine/Plugins/FX/Niagara/Source/Niagara/Private/NiagaraStats.h
--- a/Engine/Plugins/FX/Niagara/Source/Niagara/Private/NiagaraStats.h
+++ b/Engine/Plugins/FX/Niagara/Source/Niagara/Private/NiagaraStats.h
@@ -3,9 +3,7 @@
 
 #pragma once
 
-#include "CoreMinimal.h"
 #include "Stats/Stats.h"
-#include "NiagaraModule.h"
 
 DECLARE_STATS_GROUP(TEXT("Niagara"), STATGROUP_Niagara, STATCAT_Advanced);
 DECLARE_DWORD_COUNTER_STAT(TEXT("# CPU Particles"), STAT_NiagaraNumParticles, STATGROUP_Niagara);
diff --git a/Engine/Plugins/FX/Niagara/Source/Niagara/Private/NiagaraSystem.cpp b/Engine/Plugins/FX/Niagara/Source/Niagara/Private/NiagaraSystem.cpp
--- a/Engine/Plugins/FX/Niagara/Source/Niagara/Private/NiagaraSystem.cpp
+++ b/Engine/Plugins/FX/Niagara/Source/Niagara/Private/NiagaraSystem.cpp
@@ -2,19 +2,21 @@
 
 
 #include "NiagaraSystem.h"
+#include "Components/PrimitiveComponent.h"
 #include "NiagaraSystemImpl.h"
 
 #include "AssetRegistry/IAssetRegistry.h"
 #include "INiagaraEditorOnlyDataUtlities.h"
+#include "Materials/MaterialInterface.h"
 #include "NiagaraAsyncCompile.h"
 #include "NiagaraConstants.h"
 #include "NiagaraBakerSettings.h"
 #include "NiagaraCustomVersion.h"
+#include "NiagaraDataInterfaceUtilities.h"
 #include "NiagaraEditorDataBase.h"
 #include "NiagaraEmitter.h"
 #include "NiagaraEmitterHandle.h"
 #include "NiagaraModule.h"
-#include "NiagaraPrecompileContainer.h"
 #include "NiagaraRendererProperties.h"
 #include "NiagaraScratchPadContainer.h"
 #include "NiagaraScriptSourceBase.h"
@@ -27,27 +29,23 @@
 #include "NiagaraTypes.h"
 #include "NiagaraWorldManager.h"
 #include "Algo/RemoveIf.h"
-#include "Async/Async.h"
 #include "HAL/LowLevelMemTracker.h"
-#include "Interfaces/ITargetPlatform.h"
 #include "Misc/ScopedSlowTask.h"
 #include "Misc/ScopeExit.h"
 #include "Modules/ModuleManager.h"
 #include "ProfilingDebugging/CookStats.h"
 #include "UObject/ObjectSaveContext.h"
 #include "UObject/Package.h"
-#include "UObject/UObjectThreadContext.h"
-#include "ShaderCompiler.h"
 #include "PipelineStateCache.h"
 #include "NiagaraDataChannel.h"
 #include "NiagaraDataChannelDefinitions.h"
+#include "UObject/UObjectIterator.h"
 
 #include UE_INLINE_GENERATED_CPP_BY_NAME(NiagaraSystem)
 
 #define LOCTEXT_NAMESPACE "NiagaraSystem"
 
 #if WITH_EDITOR
-#include "DerivedDataCacheInterface.h"
 #endif
 
 DECLARE_CYCLE_STAT(TEXT("Niagara - System - CompileScript"), STAT_Niagara_System_CompileScript, STATGROUP_Niagara);
diff --git a/Engine/Plugins/FX/Niagara/Source/Niagara/Private/NiagaraSystemInstance.cpp b/Engine/Plugins/FX/Niagara/Source/Niagara/Private/NiagaraSystemInstance.cpp
--- a/Engine/Plugins/FX/Niagara/Source/Niagara/Private/NiagaraSystemInstance.cpp
+++ b/Engine/Plugins/FX/Niagara/Source/Niagara/Private/NiagaraSystemInstance.cpp
@@ -1,24 +1,24 @@
 // Copyright Epic Games, Inc. All Rights Reserved.
 
 #include "NiagaraSystemInstance.h"
+#include "Misc/LargeWorldRenderPosition.h"
 #include "NiagaraSystemGpuComputeProxy.h"
-#include "NiagaraConstants.h"
 #include "NiagaraCommon.h"
 #include "NiagaraComputeExecutionContext.h"
 #include "NiagaraDataInterface.h"
+#include "NiagaraEmitterInstance.h"
 #include "NiagaraGPUSystemTick.h"
 #include "NiagaraStats.h"
 #include "NiagaraParameterCollection.h"
+#include "NiagaraSystemSimulation.h"
 #include "NiagaraWorldManager.h"
 #include "NiagaraComponent.h"
-#include "NiagaraRenderer.h"
 #include "NiagaraSimCache.h"
 #include "NiagaraGpuComputeDebug.h"
 #include "NiagaraGpuComputeDispatchInterface.h"
 #include "NiagaraCrashReporterHandler.h"
 #include "NiagaraSystemImpl.h"
 
-#include "Async/Async.h"
 #include "Components/SkeletalMeshComponent.h"
 #include "GameFramework/PlayerController.h"
 #include "Templates/AlignmentTemplates.h"
diff --git a/Engine/Plugins/FX/Niagara/Source/Niagara/Private/NiagaraSystemInstanceController.cpp b/Engine/Plugins/FX/Niagara/Source/Niagara/Private/NiagaraSystemInstanceController.cpp
--- a/Engine/Plugins/FX/Niagara/Source/Niagara/Private/NiagaraSystemInstanceController.cpp
+++ b/Engine/Plugins/FX/Niagara/Source/Niagara/Private/NiagaraSystemInstanceController.cpp
@@ -4,6 +4,8 @@
 #include "Materials/MaterialInstanceDynamic.h"
 #include "Engine/StaticMesh.h"
 #include "Engine/TextureRenderTarget.h"
+#include "NiagaraEmitterInstance.h"
+#include "NiagaraWorldManager.h"
 
 FNiagaraSystemInstanceController::FNiagaraSystemInstanceController()
 	: bNeedsRendererRecache(false)
@@ -473,4 +475,4 @@ void FNiagaraSystemInstanceController::SetVariable(FName InVariableName, TWeakOb
 		const FNiagaraVariable VariableDesc(FNiagaraTypeDefinition::GetUTextureRenderTargetDef(), InVariableName);
 		OverrideParameters->SetUObject(TextureRenderTarget.Get(), VariableDesc);
 	}
-}
\ No newline at end of file
+}
diff --git a/Engine/Plugins/FX/Niagara/Source/Niagara/Private/NiagaraSystemRenderData.cpp b/Engine/Plugins/FX/Niagara/Source/Niagara/Private/NiagaraSystemRenderData.cpp
--- a/Engine/Plugins/FX/Niagara/Source/Niagara/Private/NiagaraSystemRenderData.cpp
+++ b/Engine/Plugins/FX/Niagara/Source/Niagara/Private/NiagaraSystemRenderData.cpp
@@ -1,8 +1,12 @@
 // Copyright Epic Games, Inc. All Rights Reserved.
 
 #include "NiagaraSystemRenderData.h"
+#include "NiagaraEmitterInstance.h"
 #include "NiagaraSystemInstance.h"
 #include "NiagaraCrashReporterHandler.h"
+#include "NiagaraSystem.h"
+#include "PrimitiveViewRelevance.h"
+#include "SceneView.h"
 
 void FNiagaraSystemRenderData::ExecuteDynamicDataCommands_RenderThread(const FSetDynamicDataCommandList& Commands)
 {
@@ -364,4 +368,4 @@ void FNiagaraSystemRenderData::OnSystemComplete(const FNiagaraSystemInstance& Sy
 			}
 		}
 	}
-}
\ No newline at end of file
+}
diff --git a/Engine/Plugins/FX/Niagara/Source/Niagara/Private/NiagaraSystemSimulation.cpp b/Engine/Plugins/FX/Niagara/Source/Niagara/Private/NiagaraSystemSimulation.cpp
--- a/Engine/Plugins/FX/Niagara/Source/Niagara/Private/NiagaraSystemSimulation.cpp
+++ b/Engine/Plugins/FX/Niagara/Source/Niagara/Private/NiagaraSystemSimulation.cpp
@@ -1,10 +1,9 @@
 // Copyright Epic Games, Inc. All Rights Reserved.
 
 #include "NiagaraSystemSimulation.h"
+#include "NiagaraEmitterInstance.h"
 #include "NiagaraModule.h"
 #include "NiagaraTypes.h"
-#include "NiagaraEvents.h"
-#include "NiagaraSettings.h"
 #include "NiagaraParameterCollection.h"
 #include "NiagaraSystemInstance.h"
 #include "NiagaraSystemGpuComputeProxy.h"
diff --git a/Engine/Plugins/FX/Niagara/Source/Niagara/Private/NiagaraSystemStaticBuffers.cpp b/Engine/Plugins/FX/Niagara/Source/Niagara/Private/NiagaraSystemStaticBuffers.cpp
--- a/Engine/Plugins/FX/Niagara/Source/Niagara/Private/NiagaraSystemStaticBuffers.cpp
+++ b/Engine/Plugins/FX/Niagara/Source/Niagara/Private/NiagaraSystemStaticBuffers.cpp
@@ -1,6 +1,7 @@
 // Copyright Epic Games, Inc. All Rights Reserved.
 
 #include "NiagaraSystemStaticBuffers.h"
+#include "RenderingThread.h"
 
 FNiagaraSystemStaticBuffers::~FNiagaraSystemStaticBuffers()
 {
diff --git a/Engine/Plugins/FX/Niagara/Source/Niagara/Private/NiagaraTrace.cpp b/Engine/Plugins/FX/Niagara/Source/Niagara/Private/NiagaraTrace.cpp
--- a/Engine/Plugins/FX/Niagara/Source/Niagara/Private/NiagaraTrace.cpp
+++ b/Engine/Plugins/FX/Niagara/Source/Niagara/Private/NiagaraTrace.cpp
@@ -2,6 +2,6 @@
 
 #include "NiagaraTrace.h"
 
-#include "Trace/Trace.inl"
+#include "Trace/Detail/Channel.h"
 
 UE_TRACE_CHANNEL_DEFINE(NiagaraChannel);
diff --git a/Engine/Plugins/FX/Niagara/Source/Niagara/Private/NiagaraTypes.cpp b/Engine/Plugins/FX/Niagara/Source/Niagara/Private/NiagaraTypes.cpp
--- a/Engine/Plugins/FX/Niagara/Source/Niagara/Private/NiagaraTypes.cpp
+++ b/Engine/Plugins/FX/Niagara/Source/Niagara/Private/NiagaraTypes.cpp
@@ -2,12 +2,11 @@
 
 #include "NiagaraTypes.h"
 
-#include "Engine/Texture2D.h"
 #include "NiagaraDataInterface.h"
-#include "NiagaraStats.h"
 #include "NiagaraEmitter.h"
 #include "Misc/StringBuilder.h"
 #include "UObject/Class.h"
+#include "UObject/EnumProperty.h"
 #include "UObject/Package.h"
 
 #include UE_INLINE_GENERATED_CPP_BY_NAME(NiagaraTypes)
diff --git a/Engine/Plugins/FX/Niagara/Source/Niagara/Private/NiagaraUserRedirectionParameterStore.cpp b/Engine/Plugins/FX/Niagara/Source/Niagara/Private/NiagaraUserRedirectionParameterStore.cpp
--- a/Engine/Plugins/FX/Niagara/Source/Niagara/Private/NiagaraUserRedirectionParameterStore.cpp
+++ b/Engine/Plugins/FX/Niagara/Source/Niagara/Private/NiagaraUserRedirectionParameterStore.cpp
@@ -1,7 +1,6 @@
 // Copyright Epic Games, Inc. All Rights Reserved.
 
 #include "NiagaraUserRedirectionParameterStore.h"
-#include "NiagaraSystemInstance.h"
 
 #include UE_INLINE_GENERATED_CPP_BY_NAME(NiagaraUserRedirectionParameterStore)
 
diff --git a/Engine/Plugins/FX/Niagara/Source/Niagara/Private/NiagaraWorldManager.cpp b/Engine/Plugins/FX/Niagara/Source/Niagara/Private/NiagaraWorldManager.cpp
--- a/Engine/Plugins/FX/Niagara/Source/Niagara/Private/NiagaraWorldManager.cpp
+++ b/Engine/Plugins/FX/Niagara/Source/Niagara/Private/NiagaraWorldManager.cpp
@@ -1,21 +1,15 @@
 // Copyright Epic Games, Inc. All Rights Reserved.
 
 #include "NiagaraWorldManager.h"
+#include "Engine/Level.h"
 #include "GameFramework/Pawn.h"
 #include "NiagaraModule.h"
-#include "Modules/ModuleManager.h"
+#include "NiagaraSystemSimulation.h"
 #include "NiagaraTypes.h"
-#include "NiagaraDataChannel.h"
-#include "NiagaraSettings.h"
 #include "NiagaraParameterCollection.h"
 #include "NiagaraSystemInstance.h"
-#include "Scalability.h"
-#include "Misc/ConfigCacheIni.h"
-#include "Misc/CoreDelegates.h"
-#include "NiagaraDataInterfaceSkeletalMesh.h"
 #include "GameFramework/PlayerController.h"
 #include "Engine/GameViewportClient.h"
-#include "EngineModule.h"
 #include "NiagaraStats.h"
 #include "NiagaraComponentPool.h"
 #include "NiagaraComponent.h"
@@ -23,13 +17,12 @@
 #include "NiagaraEffectType.h"
 #include "NiagaraDebugHud.h"
 #include "NiagaraGpuComputeDispatchInterface.h"
-#include "HAL/PlatformApplicationMisc.h"
 #include "Particles/FXBudget.h"
-#include "NiagaraFunctionLibrary.h"
 #include "NiagaraCullProxyComponent.h"
 #include "GameDelegates.h"
-#include "DrawDebugHelpers.h"
 #include "Engine/LocalPlayer.h"
+#include "SceneView.h"
+#include "UObject/UObjectIterator.h"
 
 #include UE_INLINE_GENERATED_CPP_BY_NAME(NiagaraWorldManager)
 
diff --git a/Engine/Plugins/FX/Niagara/Source/Niagara/Public/NiagaraGpuReadbackManager.h b/Engine/Plugins/FX/Niagara/Source/Niagara/Public/NiagaraGpuReadbackManager.h
--- a/Engine/Plugins/FX/Niagara/Source/Niagara/Public/NiagaraGpuReadbackManager.h
+++ b/Engine/Plugins/FX/Niagara/Source/Niagara/Public/NiagaraGpuReadbackManager.h
@@ -2,12 +2,8 @@
 
 #pragma once
 
-#include "CoreMinimal.h"
 #include "Containers/Queue.h"
-#include "NiagaraCommon.h"
 #include "RenderGraphResources.h"
-#include "RHI.h"
-#include "RHICommandList.h"
 
 class FRDGBuilder;
 
