Adding missing include.

diff --git a/Engine/Plugins/FX/Niagara/Source/Niagara/Private/NiagaraDataInterfaceRenderTargetVolume.cpp b/Engine/Plugins/FX/Niagara/Source/Niagara/Private/NiagaraDataInterfaceRenderTargetVolume.cpp
--- a/Engine/Plugins/FX/Niagara/Source/Niagara/Private/NiagaraDataInterfaceRenderTargetVolume.cpp
+++ b/Engine/Plugins/FX/Niagara/Source/Niagara/Private/NiagaraDataInterfaceRenderTargetVolume.cpp
@@ -10,6 +10,7 @@
 #include "NDIRenderTargetVolumeSimCacheData.h"
 #include "NiagaraDataInterfaceRenderTargetCommon.h"
 #include "NiagaraGpuComputeDispatchInterface.h"
+#include "NiagaraSettings.h"
 #include "NiagaraSystemInstance.h"
 #include "NiagaraStats.h"
 #include "NiagaraShaderParametersBuilder.h"
