Fix resource leak around NiagaraScriptShader, moving naked pointers to SharedPtr/UniquePtr

diff --git a/Engine/Plugins/FX/Niagara/Source/Niagara/Classes/NiagaraScript.h b/Engine/Plugins/FX/Niagara/Source/Niagara/Classes/NiagaraScript.h
--- a/Engine/Plugins/FX/Niagara/Source/Niagara/Classes/NiagaraScript.h
+++ b/Engine/Plugins/FX/Niagara/Source/Niagara/Classes/NiagaraScript.h
@@ -1032,7 +1032,7 @@ public:
 
 	// Infrastructure for GPU compute Shaders
 #if WITH_EDITOR
-	NIAGARA_API void CacheResourceShadersForCooking(EShaderPlatform ShaderPlatform, TArray<FNiagaraShaderScript*>& InOutCachedResources, const ITargetPlatform* TargetPlatform = nullptr);
+	NIAGARA_API void CacheResourceShadersForCooking(EShaderPlatform ShaderPlatform, TArray<TUniquePtr<FNiagaraShaderScript>>& InOutCachedResources, const ITargetPlatform* TargetPlatform = nullptr);
 
 	NIAGARA_API void CacheResourceShadersForRendering(bool bRegenerateId, bool bForceRecompile=false);
 	void BeginCacheForCookedPlatformData(const ITargetPlatform *TargetPlatform);
@@ -1043,7 +1043,6 @@ public:
 	TArray<FName> FindShaderFormatsForCooking(const ITargetPlatform* TargetPlatform) const;
 
 #endif // WITH_EDITOR
-	FNiagaraShaderScript* AllocateResource();
 	FNiagaraShaderScript* GetRenderThreadScript()
 	{
 		return ScriptResource.Get();
@@ -1250,11 +1249,11 @@ private:
 	FNiagaraVMExecutableDataId& GetLastGeneratedVMId(const FGuid& VersionGuid = FGuid()) const;
 #endif
 
-	TUniquePtr<FNiagaraShaderScript> ScriptResource;
+	TSharedPtr<FNiagaraShaderScript> ScriptResource;
 
 #if WITH_EDITORONLY_DATA
 	TArray<FNiagaraShaderScript> LoadedScriptResources;
-	FNiagaraShaderScript* ScriptResourcesByFeatureLevel[ERHIFeatureLevel::Num];
+	TSharedPtr<FNiagaraShaderScript> ScriptResourcesByFeatureLevel[ERHIFeatureLevel::Num];
 #endif
 
 	/** Compute shader compiled for this script */
@@ -1269,7 +1268,7 @@ private:
 
 #if WITH_EDITORONLY_DATA
 	/* script resources being cached for cooking. */
-	TMap<const class ITargetPlatform*, TArray<FNiagaraShaderScript*>> CachedScriptResourcesForCooking;
+	TMap<const class ITargetPlatform*, TArray<TUniquePtr<FNiagaraShaderScript>>> CachedScriptResourcesForCooking;
 
 	UPROPERTY(Transient)
 	TArray<TObjectPtr<UObject>> ActiveCompileRoots;
diff --git a/Engine/Plugins/FX/Niagara/Source/Niagara/Private/NiagaraScript.cpp b/Engine/Plugins/FX/Niagara/Source/Niagara/Private/NiagaraScript.cpp
--- a/Engine/Plugins/FX/Niagara/Source/Niagara/Private/NiagaraScript.cpp
+++ b/Engine/Plugins/FX/Niagara/Source/Niagara/Private/NiagaraScript.cpp
@@ -940,7 +940,7 @@ UNiagaraScript::UNiagaraScript(const FObjectInitializer& ObjectInitializer)
 #endif
 {
 #if WITH_EDITORONLY_DATA
-	ScriptResource = MakeUnique<FNiagaraShaderScript>();
+	ScriptResource = MakeShared<FNiagaraShaderScript>();
 	ScriptResource->OnCompilationComplete().AddUniqueDynamic(this, &UNiagaraScript::RaiseOnGPUCompilationComplete);
 
 	RapidIterationParameters.DebugName = *GetFullName();
@@ -1940,10 +1940,10 @@ void UNiagaraScript::SaveShaderStableKeys(const class ITargetPlatform* TP)
 	SaveKeyVal.ClassNameAndObjectPath.SetCompactFullNameFromObject(this);
 	static FName FName_Niagara(TEXT("Niagara"));
 	SaveKeyVal.MaterialDomain = FName_Niagara;
-	const TArray<FNiagaraShaderScript*>* ScriptResourcesToSavePtr = CachedScriptResourcesForCooking.Find(TP);
+	const TArray<TUniquePtr<FNiagaraShaderScript>>* ScriptResourcesToSavePtr = CachedScriptResourcesForCooking.Find(TP);
 	if (ScriptResourcesToSavePtr != nullptr)
 	{
-		for (FNiagaraShaderScript* Resource : *ScriptResourcesToSavePtr)
+		for (const TUniquePtr<FNiagaraShaderScript>& Resource : *ScriptResourcesToSavePtr)
 		{
 			if (Resource)
 			{
@@ -2313,9 +2313,9 @@ void UNiagaraScript::SetPreviewFeatureLevel(ERHIFeatureLevel::Type InPreviewFeat
 			continue;
 		}
 
-		for ( FNiagaraShaderScript* ExistingScript : Script->ScriptResourcesByFeatureLevel )
+		for (TSharedPtr<FNiagaraShaderScript>& ExistingScript : Script->ScriptResourcesByFeatureLevel )
 		{
-			if (ExistingScript == Script->ScriptResource.Get())
+			if (ExistingScript == Script->ScriptResource)
 			{
 				Script->CacheResourceShadersForRendering(true, false);
 				break;
@@ -3283,7 +3283,7 @@ void UNiagaraScript::BeginCacheForCookedPlatformData(const ITargetPlatform *Targ
 			return;
 		}
 
-		TArray<FNiagaraShaderScript*>& CachedScriptResourcesForPlatform = CachedScriptResourcesForCooking.FindOrAdd(TargetPlatform);
+		TArray<TUniquePtr<FNiagaraShaderScript>>& CachedScriptResourcesForPlatform = CachedScriptResourcesForCooking.FindOrAdd(TargetPlatform);
 
 		// Cache for all the shader formats that the cooking target requires
 		for (const FName& ShaderFormat : ShaderFormatsForCooking)
@@ -3313,14 +3313,14 @@ bool UNiagaraScript::IsCachedCookedPlatformDataLoaded(const ITargetPlatform* Tar
 
 		if (!bHasOutstandingCompilationRequests)
 		{
-			const TArray<FNiagaraShaderScript*>* CachedScriptResourcesForPlatform = CachedScriptResourcesForCooking.Find(TargetPlatform);
+			const TArray<TUniquePtr<FNiagaraShaderScript>>* CachedScriptResourcesForPlatform = CachedScriptResourcesForCooking.Find(TargetPlatform);
 			if (CachedScriptResourcesForPlatform)
 			{
-				for (const auto& MaterialResource : *CachedScriptResourcesForPlatform)
+				for (const TUniquePtr<FNiagaraShaderScript>& MaterialResource : *CachedScriptResourcesForPlatform)
 				{
 					if (MaterialResource->IsCompilationFinished() == false)
 					{
-						// For now, finish compilation here until we can make sure compilation is finished in the cook commandlet asyncronously before serialize
+						// For now, finish compilation here until we can make sure compilation is finished in the cook commandlet asynchronously before serialize
 						MaterialResource->FinishCompilation();
 
 						if (MaterialResource->IsCompilationFinished() == false)
@@ -3340,7 +3340,7 @@ bool UNiagaraScript::IsCachedCookedPlatformDataLoaded(const ITargetPlatform* Tar
 	return true;
 }
 
-void UNiagaraScript::CacheResourceShadersForCooking(EShaderPlatform ShaderPlatform, TArray<FNiagaraShaderScript*>& InOutCachedResources, const ITargetPlatform* TargetPlatform)
+void UNiagaraScript::CacheResourceShadersForCooking(EShaderPlatform ShaderPlatform, TArray<TUniquePtr<FNiagaraShaderScript>>& InOutCachedResources, const ITargetPlatform* TargetPlatform)
 {
 	if (CanBeRunOnGpu())
 	{
@@ -3348,7 +3348,7 @@ void UNiagaraScript::CacheResourceShadersForCooking(EShaderPlatform ShaderPlatfo
 		if (Usage == ENiagaraScriptUsage::ParticleGPUComputeScript)
 		{
 			ERHIFeatureLevel::Type TargetFeatureLevel = GetMaxSupportedFeatureLevel(ShaderPlatform);
-			const auto FindExistingScriptPredicate = [&](const FNiagaraShaderScript* ExistingScript)
+			const auto FindExistingScriptPredicate = [&](const TUniquePtr<FNiagaraShaderScript>& ExistingScript)
 			{
 				return ExistingScript->MatchesScript(TargetFeatureLevel, ShaderPlatform, CachedScriptVMId);
 			};
@@ -3359,8 +3359,7 @@ void UNiagaraScript::CacheResourceShadersForCooking(EShaderPlatform ShaderPlatfo
 				return;
 			}
 
-			FNiagaraShaderScript *ResourceToCache = nullptr;
-			FNiagaraShaderScript* NewResource = AllocateResource();
+			TUniquePtr<FNiagaraShaderScript> NewResource = MakeUnique<FNiagaraShaderScript>();
 			check(CachedScriptVMId.CompilerVersionID.IsValid());
 			check(CachedScriptVMId.BaseScriptCompileHash.IsValid());
 
@@ -3368,16 +3367,13 @@ void UNiagaraScript::CacheResourceShadersForCooking(EShaderPlatform ShaderPlatfo
 				CachedScriptVMId.GetAdditionalVariableStrings(),
 				CachedScriptVMId.BaseScriptCompileHash,	CachedScriptVMId.ReferencedCompileHashes,
 				CachedScriptVMId.bUsesRapidIterationParams, GetFriendlyName());
-			ResourceToCache = NewResource;
 
-			check(ResourceToCache);
-
-			CacheShadersForResources(ResourceToCache, false, false, true, TargetPlatform);
+			CacheShadersForResources(NewResource.Get(), false, false, true, TargetPlatform);
 
 			INiagaraModule& NiagaraModule = FModuleManager::GetModuleChecked<INiagaraModule>(TEXT("Niagara"));
 			NiagaraModule.ProcessShaderCompilationQueue();
 
-			InOutCachedResources.Add(ResourceToCache);
+			InOutCachedResources.Emplace(MoveTemp(NewResource));
 		}
 	}
 }
@@ -3424,11 +3420,7 @@ void UNiagaraScript::CacheResourceShadersForRendering(bool bRegenerateId, bool b
 		// Regenerate this script's Id if requested
 		for (int32 Idx = 0; Idx < ERHIFeatureLevel::Num; Idx++)
 		{
-			if (ScriptResourcesByFeatureLevel[Idx])
-			{
-				ScriptResourcesByFeatureLevel[Idx]->ReleaseShaderMap();
-				ScriptResourcesByFeatureLevel[Idx] = nullptr;
-			}
+			ScriptResourcesByFeatureLevel[Idx] = nullptr;
 		}
 	}
 
@@ -3455,7 +3447,7 @@ void UNiagaraScript::CacheResourceShadersForRendering(bool bRegenerateId, bool b
 			ScriptResource->BuildScriptParametersMetadata(CachedScriptVM.ShaderScriptParametersMetadata);
 
 			CacheShadersForResources(ScriptResource.Get(), true);
-			ScriptResourcesByFeatureLevel[CacheFeatureLevel] = ScriptResource.Get();
+			ScriptResourcesByFeatureLevel[CacheFeatureLevel] = ScriptResource;
 		}
 		else
 		{
@@ -3825,7 +3817,7 @@ void UNiagaraScript::SerializeNiagaraShaderMaps(FArchive& Ar, int32 NiagaraVer,
 		Ar.UsingCustomVersion(FReleaseObjectVersion::GUID);
 
 		int32 NumResourcesToSave = 0;
-		const TArray<FNiagaraShaderScript*>* ScriptResourcesToSavePtr = nullptr;
+		TArray<TUniquePtr<FNiagaraShaderScript>>* ScriptResourcesToSavePtr = nullptr;
 
 		if (Ar.IsCooking())
 		{
@@ -3840,7 +3832,7 @@ void UNiagaraScript::SerializeNiagaraShaderMaps(FArchive& Ar, int32 NiagaraVer,
 
 		if (ScriptResourcesToSavePtr != nullptr)
 		{
-			for (FNiagaraShaderScript* ScriptResourceToSave : (*ScriptResourcesToSavePtr))
+			for (TUniquePtr<FNiagaraShaderScript>& ScriptResourceToSave : (*ScriptResourcesToSavePtr))
 			{
 				checkf(ScriptResourceToSave != nullptr, TEXT("Invalid script resource was cached"));
 				ScriptResourceToSave->SerializeShaderMap(Ar);
@@ -3889,7 +3881,7 @@ void UNiagaraScript::SerializeNiagaraShaderMaps(FArchive& Ar, int32 NiagaraVer,
 					{
 						if (GMaxRHIShaderPlatform == ShaderMap->GetShaderPlatform())
 						{
-							ScriptResource = MakeUnique<FNiagaraShaderScript>(Resource);
+							ScriptResource = MakeShared<FNiagaraShaderScript>(Resource);
 						}
 					}
 				}
@@ -3911,13 +3903,13 @@ void UNiagaraScript::ProcessSerializedShaderMaps()
 		if (LoadedShaderMap && LoadedShaderMap->GetShaderPlatform() == GMaxRHIShaderPlatform)
 		{
 			HasScriptResource = true;
-			ScriptResource = MakeUnique<FNiagaraShaderScript>(LoadedResource);
+			ScriptResource = MakeShared<FNiagaraShaderScript>(LoadedResource);
 			ScriptResource->OnCompilationComplete().AddUniqueDynamic(this, &UNiagaraScript::RaiseOnGPUCompilationComplete);
 
 			ERHIFeatureLevel::Type LoadedFeatureLevel = LoadedShaderMap->GetShaderMapId().FeatureLevel;
 			if (!ScriptResourcesByFeatureLevel[LoadedFeatureLevel])
 			{
-				ScriptResourcesByFeatureLevel[LoadedFeatureLevel] = AllocateResource();
+				ScriptResourcesByFeatureLevel[LoadedFeatureLevel] = MakeShared<FNiagaraShaderScript>();
 			}
 
 			ScriptResourcesByFeatureLevel[LoadedFeatureLevel]->SetShaderMap(LoadedShaderMap);
@@ -3938,11 +3930,6 @@ void UNiagaraScript::ProcessSerializedShaderMaps()
 	}
 }
 
-FNiagaraShaderScript* UNiagaraScript::AllocateResource()
-{
-	return new FNiagaraShaderScript();
-}
-
 #if WITH_EDITORONLY_DATA
 TArray<ENiagaraScriptUsage> FVersionedNiagaraScriptData::GetSupportedUsageContexts() const
 {
diff --git a/Engine/Plugins/FX/Niagara/Source/NiagaraEditor/Private/ViewModels/NiagaraScriptStatsViewModel.cpp b/Engine/Plugins/FX/Niagara/Source/NiagaraEditor/Private/ViewModels/NiagaraScriptStatsViewModel.cpp
--- a/Engine/Plugins/FX/Niagara/Source/NiagaraEditor/Private/ViewModels/NiagaraScriptStatsViewModel.cpp
+++ b/Engine/Plugins/FX/Niagara/Source/NiagaraEditor/Private/ViewModels/NiagaraScriptStatsViewModel.cpp
@@ -147,12 +147,13 @@ namespace NiagaraScriptStatsLocal
 						NiagaraEmitter.GetEmitterData()->ForEachScript(
 							[&](UNiagaraScript* NiagaraScript)
 							{
-								TArray<FNiagaraShaderScript*> NewResources;
+								TArray<TUniquePtr<FNiagaraShaderScript>> NewResources;
 								NiagaraScript->CacheResourceShadersForCooking(ShaderPlatform, NewResources);
 								if (NewResources.Num() > 0)
 								{
 									ensure(NewResources.Num() == 1);
-									ShaderScripts.Emplace(NewResources[0]);
+									FNiagaraShaderScript* ShaderScript = NewResources[0].Release();
+									ShaderScripts.Emplace(ShaderScript);
 								}
 							}
 						);
@@ -240,7 +241,7 @@ namespace NiagaraScriptStatsLocal
 		bool bIsGPU = false;
 		bool bHasError = false;
 		FString ResultsString;
-		TArray<TUniquePtr<class FNiagaraShaderScript, FNiagaraScriptSafeDelete>> ShaderScripts;
+		TArray<TUniquePtr<FNiagaraShaderScript, FNiagaraScriptSafeDelete>> ShaderScripts;
 	};
 }
 
