Niagara - Asset permission improvements.
+ Allow external code to specify whether or not an asset should be filtered (unsupported) based on class usage.  This is used to prevent class based usage from filtering out engine plugin supplied content.
+ Add a stack error for modules and dynamic inputs using unsupported assets.
+ Update the graph node error handling to handle unsupported classes and unsupported assets.

diff --git a/Engine/Plugins/FX/Niagara/Source/NiagaraEditor/Private/NiagaraEditorSettings.cpp b/Engine/Plugins/FX/Niagara/Source/NiagaraEditor/Private/NiagaraEditorSettings.cpp
--- a/Engine/Plugins/FX/Niagara/Source/NiagaraEditor/Private/NiagaraEditorSettings.cpp
+++ b/Engine/Plugins/FX/Niagara/Source/NiagaraEditor/Private/NiagaraEditorSettings.cpp
@@ -665,6 +665,11 @@ void UNiagaraEditorSettings::SetOnIsClassPathAllowed(const FOnIsClassPathAllowed
 	OnIsClassPathAllowedDelegate = InOnIsClassPathAllowed;
 }
 
+void UNiagaraEditorSettings::SetOnShouldFilterAssetByClassUsage(const FOnShouldFilterAssetByClassUsage& InOnShouldFilterAssetByClassUsage)
+{
+	OnShouldFilterAssetByClassUsage = InOnShouldFilterAssetByClassUsage;
+}
+
 bool UNiagaraEditorSettings::IsAllowedClass(const UClass* InClass) const
 {
 	return OnIsClassAllowedDelegate.IsBound() == false || OnIsClassAllowedDelegate.Execute(InClass);
@@ -719,41 +724,57 @@ UObject::FAssetRegistryTag UNiagaraEditorSettings::CreateClassUsageAssetRegistry
 	return UObject::FAssetRegistryTag(ClassUsageListTagName, ClassUsageList, FAssetRegistryTag::TT_Hidden);
 }
 
+bool UNiagaraEditorSettings::IsAllowedAssetObjectByClassUsageInternal(const UObject& AssetObject, TSet<const UObject*>& CheckedAssetObjects) const
+{
+	CheckedAssetObjects.Add(&AssetObject);
+
+	TArray<UObject*> ObjectsInPackage;
+	GetObjectsWithPackage(AssetObject.GetOutermost(), ObjectsInPackage);
+	for (UObject* ObjectInPackage : ObjectsInPackage)
+	{
+		if (ShouldTrackClassUsage(ObjectInPackage->GetClass()) && IsAllowedClass(ObjectInPackage->GetClass()) == false)
+		{
+			if (GbLogFoundButNotAllowedAssets)
+			{
+				UE_LOG(LogNiagaraEditor, Log, TEXT("Asset %s is not allowed due to object %s with class %s which is not allowed in this editor context."),
+					*AssetObject.GetPathName(), *ObjectInPackage->GetPathName(), *ObjectInPackage->GetClass()->GetClassPathName().ToString());
+			}
+			return false;
+		}
+
+		const UNiagaraNode* NiagaraNode = Cast<UNiagaraNode>(ObjectInPackage);
+		if (NiagaraNode != nullptr &&
+			NiagaraNode->GetReferencedAsset() != nullptr &&
+			OnShouldFilterAssetByClassUsage.Execute(FTopLevelAssetPath(NiagaraNode->GetReferencedAsset()->GetPathName())) &&
+			CheckedAssetObjects.Contains(NiagaraNode->GetReferencedAsset()) == false)
+		{
+			if (IsAllowedAssetObjectByClassUsageInternal(*NiagaraNode->GetReferencedAsset(), CheckedAssetObjects) == false)
+			{
+				UE_LOG(LogNiagaraEditor, Log, TEXT("Asset %s is not allowed due to referenced asset %s which is not allowed in this editor context."),
+					*AssetObject.GetPathName(), *NiagaraNode->GetReferencedAsset()->GetPathName());
+				return false;
+			}
+		}
+	}
+	return true;
+}
+
 bool UNiagaraEditorSettings::IsAllowedAssetByClassUsage(const FAssetData& InAssetData) const
 {
-	if (OnIsClassAllowedDelegate.IsBound() == false)
+	if (OnShouldFilterAssetByClassUsage.IsBound() == false ||
+		OnIsClassAllowedDelegate.IsBound() == false ||
+		OnShouldFilterAssetByClassUsage.Execute(FTopLevelAssetPath(InAssetData.GetObjectPathString())) == false)
 	{
 		return true;
 	}
 
-	const UObject* AssetObject = nullptr;
 	bool bClassDataFound = false;
 	bool bInvalidClassFound = false;
-	if (InAssetData.IsAssetLoaded())
+	if (InAssetData.IsAssetLoaded() && InAssetData.GetAsset() != nullptr)
 	{
-		AssetObject = InAssetData.GetAsset();
-		if (AssetObject != nullptr)
-		{
-			bClassDataFound = true;
-			TArray<UObject*> ObjectsInPackage;
-			GetObjectsWithPackage(AssetObject->GetOutermost(), ObjectsInPackage);
-			for (UObject* ObjectInPackage : ObjectsInPackage)
-			{
-				if (ShouldTrackClassUsage(ObjectInPackage->GetClass()) && IsAllowedClass(ObjectInPackage->GetClass()) == false)
-				{
-					bInvalidClassFound = true;
-					if (GbLogFoundButNotAllowedAssets)
-					{
-						UE_LOG(LogNiagaraEditor, Log, TEXT("Asset %s is not allowed due to object %s with class %s which is not allowed in this editor context."),
-							*InAssetData.GetFullName(), *ObjectInPackage->GetPathName(), *ObjectInPackage->GetClass()->GetClassPathName().ToString());
-					}
-					else
-					{
-						break;
-					}
-				}
-			}
-		}
+		TSet<const UObject*> CheckedAssetObjects;
+		bClassDataFound = true;
+		bInvalidClassFound = IsAllowedAssetObjectByClassUsageInternal(*InAssetData.GetAsset(), CheckedAssetObjects) == false;
 	}
 	else
 	{
@@ -791,9 +812,21 @@ bool UNiagaraEditorSettings::IsAllowedAssetByClassUsage(const FAssetData& InAsse
 			}
 		}
 	}
-
 	return bClassDataFound && bInvalidClassFound == false;
 }
 
+bool UNiagaraEditorSettings::IsAllowedAssetObjectByClassUsage(const UObject& InAssetObject) const
+{
+	if (OnShouldFilterAssetByClassUsage.IsBound() == false ||
+		OnIsClassAllowedDelegate.IsBound() == false ||
+		OnShouldFilterAssetByClassUsage.Execute(FTopLevelAssetPath(InAssetObject.GetPathName())) == false)
+	{
+		return true;
+	}
+
+	TSet<const UObject*> CheckedAssetObjects;
+	return IsAllowedAssetObjectByClassUsageInternal(InAssetObject, CheckedAssetObjects);
+}
+
 #undef LOCTEXT_NAMESPACE
 
diff --git a/Engine/Plugins/FX/Niagara/Source/NiagaraEditor/Private/ViewModels/Stack/NiagaraStackFunctionInput.cpp b/Engine/Plugins/FX/Niagara/Source/NiagaraEditor/Private/ViewModels/Stack/NiagaraStackFunctionInput.cpp
--- a/Engine/Plugins/FX/Niagara/Source/NiagaraEditor/Private/ViewModels/Stack/NiagaraStackFunctionInput.cpp
+++ b/Engine/Plugins/FX/Niagara/Source/NiagaraEditor/Private/ViewModels/Stack/NiagaraStackFunctionInput.cpp
@@ -760,6 +760,18 @@ void UNiagaraStackFunctionInput::RefreshChildrenInternal(const TArray<UNiagaraSt
 						true));
 				}
 
+				const UNiagaraEditorSettings* NiagaraEditorSettings = GetDefault<UNiagaraEditorSettings>();
+				if (InputValues.DynamicNode->FunctionScript != nullptr &&
+					NiagaraEditorSettings->IsAllowedAssetObjectByClassUsage(*InputValues.DynamicNode->FunctionScript) == false)
+				{
+					NewIssues.Add(FStackIssue(
+						EStackIssueSeverity::Error,
+						LOCTEXT("UnsupportedDynamicInputScriptShort", "Unsupported Dynamic Input Script"),
+						LOCTEXT("UnsupportedDynamicInputScriptLong", "This dynamic input uses a script which uses types which are unsupported in this editor context.  This dynamic input must either be deleted, or the referenced script must be fixed."),
+						GetStackEditorDataKey(),
+						false));
+				}
+
 				if (!ScriptData->NoteMessage.IsEmptyOrWhitespace())
 				{
 					FStackIssue NoteIssue = FStackIssue(
diff --git a/Engine/Plugins/FX/Niagara/Source/NiagaraEditor/Private/ViewModels/Stack/NiagaraStackModuleItem.cpp b/Engine/Plugins/FX/Niagara/Source/NiagaraEditor/Private/ViewModels/Stack/NiagaraStackModuleItem.cpp
--- a/Engine/Plugins/FX/Niagara/Source/NiagaraEditor/Private/ViewModels/Stack/NiagaraStackModuleItem.cpp
+++ b/Engine/Plugins/FX/Niagara/Source/NiagaraEditor/Private/ViewModels/Stack/NiagaraStackModuleItem.cpp
@@ -889,7 +889,7 @@ UNiagaraStackEntry::FStackIssueFixDelegate UNiagaraStackModuleItem::GetUpgradeVe
 
 void UNiagaraStackModuleItem::RefreshIssues(TArray<FStackIssue>& NewIssues)
 {
-	if (!GetIsEnabled() || GetSystemViewModel()->GetIsForDataProcessingOnly())
+	if (GetSystemViewModel()->GetIsForDataProcessingOnly())
 	{
 		NewIssues.Empty();
 		return;
@@ -900,6 +900,25 @@ void UNiagaraStackModuleItem::RefreshIssues(TArray<FStackIssue>& NewIssues)
 		return;
 	}
 
+	if (FunctionCallNode->FunctionScript != nullptr)
+	{
+		const UNiagaraEditorSettings* EditorSettings = GetDefault<UNiagaraEditorSettings>();
+		if (EditorSettings->IsAllowedAssetObjectByClassUsage(*FunctionCallNode->FunctionScript) == false)
+		{
+			NewIssues.Add(FStackIssue(
+				EStackIssueSeverity::Error,
+				LOCTEXT("UnsupportedModuleScriptShort", "Unsupported Module Script"),
+				LOCTEXT("UnsupportedModuleScriptLong", "This module uses a script which uses types which are unsupported in this editor context.  This module must either be deleted, or the referenced script must be fixed."),
+				GetStackEditorDataKey(),
+				false));
+		}
+	}
+
+	if (!GetIsEnabled())
+	{
+		return;
+	}
+
 	FNiagaraStackGraphUtilities::CheckForDeprecatedScriptVersion(FunctionCallNode, GetStackEditorDataKey(), GetUpgradeVersionFix(), NewIssues);
 
 	FVersionedNiagaraScriptData* ScriptData = FunctionCallNode->GetScriptData();
diff --git a/Engine/Plugins/FX/Niagara/Source/NiagaraEditor/Private/Widgets/SNiagaraGraphNode.cpp b/Engine/Plugins/FX/Niagara/Source/NiagaraEditor/Private/Widgets/SNiagaraGraphNode.cpp
--- a/Engine/Plugins/FX/Niagara/Source/NiagaraEditor/Private/Widgets/SNiagaraGraphNode.cpp
+++ b/Engine/Plugins/FX/Niagara/Source/NiagaraEditor/Private/Widgets/SNiagaraGraphNode.cpp
@@ -62,10 +62,14 @@ void SNiagaraGraphNode::CreateOutputSideAddButton(TSharedPtr<SVerticalBox> Outpu
 void SNiagaraGraphNode::UpdateErrorInfo()
 {
 	const UNiagaraEditorSettings* NiagaraEditorSettings = GetDefault<UNiagaraEditorSettings>();
-	if (GraphNode != nullptr && GraphNode->IsA<UNiagaraNode>() && NiagaraEditorSettings->IsAllowedClass(GraphNode->GetClass()) == false)
+	if (NiagaraNode.IsValid())
 	{
-		ErrorMsg = FString(TEXT("UNSUPPORTED!"));
-		ErrorColor = FAppStyle::GetColor("ErrorReporting.BackgroundColor");
+		if (NiagaraEditorSettings->IsAllowedClass(GraphNode->GetClass()) == false || 
+			(NiagaraNode->GetReferencedAsset() != nullptr && NiagaraEditorSettings->IsAllowedAssetObjectByClassUsage(*NiagaraNode->GetReferencedAsset()) == false))
+		{
+			ErrorMsg = FString(TEXT("UNSUPPORTED!"));
+			ErrorColor = FAppStyle::GetColor("ErrorReporting.BackgroundColor");
+		}
 	}
 	else
 	{
diff --git a/Engine/Plugins/FX/Niagara/Source/NiagaraEditor/Public/NiagaraEditorSettings.h b/Engine/Plugins/FX/Niagara/Source/NiagaraEditor/Public/NiagaraEditorSettings.h
--- a/Engine/Plugins/FX/Niagara/Source/NiagaraEditor/Public/NiagaraEditorSettings.h
+++ b/Engine/Plugins/FX/Niagara/Source/NiagaraEditor/Public/NiagaraEditorSettings.h
@@ -368,6 +368,7 @@ public:
 
 	DECLARE_DELEGATE_RetVal_OneParam(bool, FOnIsClassAllowed, const UClass* /*InClass*/);
 	DECLARE_DELEGATE_RetVal_OneParam(bool, FOnIsClassPathAllowed, const FTopLevelAssetPath& /*InClassPath*/);
+	DECLARE_DELEGATE_RetVal_OneParam(bool, FOnShouldFilterAssetByClassUsage, const FTopLevelAssetPath& /*InAssetPath*/)
 	
 	/** Sets a delegate that allows external code to restrict which features can be used within the niagara editor by filtering which classes are allowed. */
 	void SetOnIsClassAllowed(const FOnIsClassAllowed& InOnIsClassAllowed);
@@ -375,6 +376,9 @@ public:
 	/** Sets a delegate that allows external code to restrict which features can be used within the niagara editor by filtering which classes are allowed by class path. */
 	void SetOnIsClassPathAllowed(const FOnIsClassPathAllowed& InOnIsClassPathAllowed);
 
+	/** Sets a delegate that allows external code to restrict which features can be used within the niagara editor by filtering which assets are allowed. */
+	void SetOnShouldFilterAssetByClassUsage(const FOnShouldFilterAssetByClassUsage& InOnShouldFilterAssetByClassUsage);
+
 	/** Returns whether or not the supplied class can be used in the current editor context. */
 	bool IsAllowedClass(const UClass* InClass) const;
 
@@ -387,8 +391,11 @@ public:
 	FAssetRegistryTag CreateClassUsageAssetRegistryTag(const UObject* SourceObject) const;
 
 	bool IsAllowedAssetByClassUsage(const FAssetData& InAssetData) const;
+	bool IsAllowedAssetObjectByClassUsage(const UObject& InAssetObject) const;
 
 private:
+	bool IsAllowedAssetObjectByClassUsageInternal(const UObject& InAssetObject, TSet<const UObject*>& CheckedAssetObjects) const;
+
 	void SetupNamespaceMetadata();
 	void BuildCachedPlaybackSpeeds() const;
 	bool ShouldTrackClassUsage(const UClass* InClass) const;
@@ -479,6 +486,7 @@ private:
 
 	FOnIsClassAllowed OnIsClassAllowedDelegate;
 	FOnIsClassPathAllowed OnIsClassPathAllowedDelegate;
+	FOnShouldFilterAssetByClassUsage OnShouldFilterAssetByClassUsage;
 
 	TArray<UClass*> TrackedUsageBaseClasses;
 
