Fixing non-editor non-unity compiles.

diff --git a/Engine/Plugins/FX/Niagara/Source/Niagara/Private/NiagaraPerfBaseline.cpp b/Engine/Plugins/FX/Niagara/Source/Niagara/Private/NiagaraPerfBaseline.cpp
--- a/Engine/Plugins/FX/Niagara/Source/Niagara/Private/NiagaraPerfBaseline.cpp
+++ b/Engine/Plugins/FX/Niagara/Source/Niagara/Private/NiagaraPerfBaseline.cpp
@@ -15,6 +15,7 @@
 #include "CanvasTypes.h"
 #include "HAL/FileManager.h"
 #include "UnrealClient.h"
+#include "UObject/UObjectIterator.h" // IWYU pragma: keep
 
 #include UE_INLINE_GENERATED_CPP_BY_NAME(NiagaraPerfBaseline)
 
