Fixed ARM64EC compile issue.

diff --git a/Engine/Plugins/FX/Niagara/Source/Niagara/Niagara.Build.cs b/Engine/Plugins/FX/Niagara/Source/Niagara/Niagara.Build.cs
--- a/Engine/Plugins/FX/Niagara/Source/Niagara/Niagara.Build.cs
+++ b/Engine/Plugins/FX/Niagara/Source/Niagara/Niagara.Build.cs
@@ -81,7 +81,7 @@ public class Niagara : ModuleRules
             });
         }
 
-		if (Target.IsInPlatformGroup(UnrealPlatformGroup.Windows) && Target.Architecture.bIsX64)
+		if (Target.IsInPlatformGroup(UnrealPlatformGroup.Windows) && Target.WindowsPlatform.Architecture != UnrealArch.Arm64)
 		{ 
 			AddEngineThirdPartyPrivateStaticDependencies(Target,
 				"IntelTBB",
