Fixing Android nonunity error.

diff --git a/Engine/Plugins/FX/Niagara/Source/NiagaraShader/Private/NiagaraGenerateMips.cpp b/Engine/Plugins/FX/Niagara/Source/NiagaraShader/Private/NiagaraGenerateMips.cpp
--- a/Engine/Plugins/FX/Niagara/Source/NiagaraShader/Private/NiagaraGenerateMips.cpp
+++ b/Engine/Plugins/FX/Niagara/Source/NiagaraShader/Private/NiagaraGenerateMips.cpp
@@ -4,6 +4,9 @@
 #include "GlobalShader.h"
 #include "RenderGraphUtils.h"
 #include "RHIStaticStates.h"
+#if PLATFORM_ANDROID
+#include "DataDrivenShaderPlatformInfo.h"
+#endif
 
 #include UE_INLINE_GENERATED_CPP_BY_NAME(NiagaraGenerateMips)
 
