- Fix crash if version is missing, instead return an invalid Guid

diff --git a/Engine/Plugins/FX/Niagara/Source/Niagara/Private/NiagaraScript.cpp b/Engine/Plugins/FX/Niagara/Source/Niagara/Private/NiagaraScript.cpp
--- a/Engine/Plugins/FX/Niagara/Source/Niagara/Private/NiagaraScript.cpp
+++ b/Engine/Plugins/FX/Niagara/Source/Niagara/Private/NiagaraScript.cpp
@@ -4035,8 +4035,8 @@ bool UNiagaraScript::LegacyCanBeRunOnGpu() const
 #if WITH_EDITORONLY_DATA
 FGuid UNiagaraScript::GetBaseChangeID(const FGuid& VersionGuid) const
 {
-	UNiagaraScriptSourceBase* Source = GetScriptData(VersionGuid)->Source;
-	return Source->GetChangeID(); 
+	const FVersionedNiagaraScriptData* ScriptData = GetScriptData(VersionGuid);
+	return ScriptData && ScriptData->Source ? ScriptData->Source->GetChangeID() : FGuid();
 }
 
 ENiagaraScriptCompileStatus UNiagaraScript::GetLastCompileStatus() const
