Fixed system parameter panel selecting parameters in scratch pad via name instead of via guid.
This caused wrong variables to be displayed

diff --git a/Engine/Plugins/FX/Niagara/Source/NiagaraEditor/Private/ViewModels/NiagaraParameterPanelViewModel.cpp b/Engine/Plugins/FX/Niagara/Source/NiagaraEditor/Private/ViewModels/NiagaraParameterPanelViewModel.cpp
--- a/Engine/Plugins/FX/Niagara/Source/NiagaraEditor/Private/ViewModels/NiagaraParameterPanelViewModel.cpp
+++ b/Engine/Plugins/FX/Niagara/Source/NiagaraEditor/Private/ViewModels/NiagaraParameterPanelViewModel.cpp
@@ -1033,14 +1033,26 @@ void FNiagaraSystemToolkitParameterPanelViewModel::OnINiagaraParameterPanelViewM
 void FNiagaraSystemToolkitParameterPanelViewModel::OnParameterItemSelected(const FNiagaraParameterPanelItem& SelectedItem, ESelectInfo::Type SelectInfo) const
 {
 	SelectedVariable = SelectedItem.GetVariable();
+	const FGuid SelectedVariableGuid = SelectedItem.GetVariableMetaData().GetVariableGuid();
 	bool bFound = false;
+	
 	for (const UNiagaraGraph* Graph : GetEditableGraphsConst())
 	{
-		UNiagaraScriptVariable* ParameterScriptVariable = Graph->GetScriptVariable(SelectedItem.GetVariable().GetName());
-		if (ParameterScriptVariable != nullptr)
+		TArray<TObjectPtr<UNiagaraScriptVariable>> ScriptVariables;
+		Graph->GetAllMetaData().GenerateValueArray(ScriptVariables);
+
+		for(TObjectPtr<UNiagaraScriptVariable> ScriptVariable : ScriptVariables)
+		{
+			if(ScriptVariable->Metadata.GetVariableGuid() == SelectedVariableGuid)
+			{
+				VariableObjectSelection->SetSelectedObject(ScriptVariable);
+				bFound = true;
+				break;
+			}
+		}
+
+		if(bFound)
 		{
-			VariableObjectSelection->SetSelectedObject(ParameterScriptVariable);
-			bFound = true;
 			break;
 		}
 	}
@@ -1051,7 +1063,6 @@ void FNiagaraSystemToolkitParameterPanelViewModel::OnParameterItemSelected(const
 	}
 
 	InvalidateCachedDependencies();
-
 }
 
 bool FNiagaraSystemToolkitParameterPanelViewModel::IsVariableSelected(FNiagaraVariableBase& InVar) const
