Make visuals for convert nodes more compact

diff --git a/Engine/Plugins/FX/Niagara/Source/NiagaraEditor/Private/NiagaraNodeConvert.cpp b/Engine/Plugins/FX/Niagara/Source/NiagaraEditor/Private/NiagaraNodeConvert.cpp
--- a/Engine/Plugins/FX/Niagara/Source/NiagaraEditor/Private/NiagaraNodeConvert.cpp
+++ b/Engine/Plugins/FX/Niagara/Source/NiagaraEditor/Private/NiagaraNodeConvert.cpp
@@ -127,6 +127,17 @@ UNiagaraNodeConvert::UNiagaraNodeConvert() : UNiagaraNodeWithDynamicPins(), bIsW
 
 }
 
+void UNiagaraNodeConvert::PostLoad()
+{
+	Super::PostLoad();
+
+	if (IsLocalConstantValue())
+	{
+		// collapse convert nodes with simple constant values on load
+		SetWiringShown(false);
+	}
+}
+
 void UNiagaraNodeConvert::AllocateDefaultPins()
 {
 	CreateAddPin(EGPD_Input);
@@ -628,6 +639,7 @@ void UNiagaraNodeConvert::AutowireNewNode(UEdGraphPin* FromPin)
 				}
 			}
 		}
+		SetWiringShown(false);
 	}
 	else
 	{
@@ -842,6 +854,27 @@ bool UNiagaraNodeConvert::IsWiringShown() const
 void UNiagaraNodeConvert::SetWiringShown(bool bInShown)
 {
 	bIsWiringShown = bInShown;
+	VisualsChangedDelegate.Broadcast(this);
+}
+
+bool UNiagaraNodeConvert::IsLocalConstantValue() const
+{
+	FPinCollectorArray Outputs;
+	FPinCollectorArray Inputs;
+	GetOutputPins(Outputs);
+	GetInputPins(Inputs);
+	if (Outputs.Num() == 2 && IsAddPin(Outputs.Last()) && Inputs.Num() >= 2 && IsAddPin(Inputs.Last()))
+	{
+		for (UEdGraphPin* InPin : Inputs)
+		{
+			if (InPin->LinkedTo.Num() != 0)
+			{
+				return false;
+			}
+		}
+		return true;
+	}
+	return false;
 }
 
 void UNiagaraNodeConvert::RemoveExpandedRecord(const FNiagaraConvertPinRecord& InRecord)
diff --git a/Engine/Plugins/FX/Niagara/Source/NiagaraEditor/Private/NiagaraNodeConvert.h b/Engine/Plugins/FX/Niagara/Source/NiagaraEditor/Private/NiagaraNodeConvert.h
--- a/Engine/Plugins/FX/Niagara/Source/NiagaraEditor/Private/NiagaraNodeConvert.h
+++ b/Engine/Plugins/FX/Niagara/Source/NiagaraEditor/Private/NiagaraNodeConvert.h
@@ -97,6 +97,10 @@ public:
 
 	UNiagaraNodeConvert();
 
+	//~ Begin UObject interface
+	virtual void PostLoad() override;
+	//~ End UObject interface
+
 	//~ UEdGraphNode interface
 	virtual void AllocateDefaultPins() override;
 	virtual TSharedPtr<SGraphNode> CreateVisualWidget() override;
@@ -136,12 +140,14 @@ public:
 	void AddExpandedRecord(const FNiagaraConvertPinRecord& InRecord);
 	/** Is this socket expanded?*/
 	bool HasExpandedRecord(const FNiagaraConvertPinRecord& InRecord);
+
+	/** Returns true if the node has a constant value input and a single output */
+	bool IsLocalConstantValue() const;
+	
 private:
 	//~ EdGraphNode interface
 	virtual void OnPinRemoved(UEdGraphPin* Pin) override;
 
-private:
-
 	//A swizzle string set externally to instruct the autowiring code.
 	UPROPERTY()
 	FString AutowireSwizzle;
diff --git a/Engine/Plugins/FX/Niagara/Source/NiagaraEditor/Private/NiagaraNodeParameterMapFor.cpp b/Engine/Plugins/FX/Niagara/Source/NiagaraEditor/Private/NiagaraNodeParameterMapFor.cpp
--- a/Engine/Plugins/FX/Niagara/Source/NiagaraEditor/Private/NiagaraNodeParameterMapFor.cpp
+++ b/Engine/Plugins/FX/Niagara/Source/NiagaraEditor/Private/NiagaraNodeParameterMapFor.cpp
@@ -3,7 +3,6 @@
 #include "NiagaraNodeParameterMapFor.h"
 #include "EdGraphSchema_Niagara.h"
 #include "NiagaraEditorUtilities.h"
-#include "Widgets/SNiagaraGraphNodeConvert.h"
 #include "NiagaraHlslTranslator.h"
 #include "EdGraph/EdGraphNode.h"
 
diff --git a/Engine/Plugins/FX/Niagara/Source/NiagaraEditor/Private/Widgets/SNiagaraGraphNodeConvert.cpp b/Engine/Plugins/FX/Niagara/Source/NiagaraEditor/Private/Widgets/SNiagaraGraphNodeConvert.cpp
--- a/Engine/Plugins/FX/Niagara/Source/NiagaraEditor/Private/Widgets/SNiagaraGraphNodeConvert.cpp
+++ b/Engine/Plugins/FX/Niagara/Source/NiagaraEditor/Private/Widgets/SNiagaraGraphNodeConvert.cpp
@@ -33,7 +33,7 @@ void SNiagaraGraphNodeConvert::SetDefaultTitleAreaWidget(TSharedRef<SOverlay> De
 		.OnCheckStateChanged(this, &SNiagaraGraphNodeConvert::ToggleShowWiring)
 		.IsChecked(this, &SNiagaraGraphNodeConvert::GetToggleButtonChecked)
 		.Cursor(EMouseCursor::Default)
-		.ToolTipText(LOCTEXT("ToggleShaderCode_Tooltip", "Toggle visibility of shader code."))
+		.ToolTipText(LOCTEXT("ToggleConvertNode_Tooltip", "Toggle visibility of convert node wiring."))
 		.Style(FAppStyle::Get(), "Graph.Node.AdvancedView")
 		[
 			SNew(SHorizontalBox)
@@ -130,6 +130,22 @@ void SNiagaraGraphNodeConvert::UpdateGraphNode()
 		ConvertNodeViewModel = MakeShareable(new FNiagaraConvertNodeViewModel(*ConvertNode));
 	}
 	SNiagaraGraphNode::UpdateGraphNode();
+	
+	if (ConvertNode && ConvertNode->IsLocalConstantValue())
+	{
+		// hide input pins when the node is just displaying a constant value
+		for (TSharedRef<SGraphPin> InPin : InputPins)
+		{
+			InPin->SetPinColorModifier(ConvertNode->IsWiringShown() ? FLinearColor::White : FLinearColor::Transparent);
+		}
+	}
+
+	// set visibility of add pins
+	if (InputPins.Num() > 0 && OutputPins.Num() > 0)
+	{
+		InputPins.Last()->SetVisibility(ConvertNode->IsWiringShown() ? EVisibility::Visible : EVisibility::Collapsed);
+		OutputPins.Last()->SetVisibility(ConvertNode->IsWiringShown() ? EVisibility::Visible : EVisibility::Collapsed);
+	}
 }
 
 void SNiagaraGraphNodeConvert::AddPin(const TSharedRef<SGraphPin>& PinToAdd)
