- Fix for Niagara not having the right feature level when loading the editor with a feature level active

diff --git a/Engine/Plugins/FX/Niagara/Source/NiagaraEditor/Private/NiagaraEditorModule.cpp b/Engine/Plugins/FX/Niagara/Source/NiagaraEditor/Private/NiagaraEditorModule.cpp
--- a/Engine/Plugins/FX/Niagara/Source/NiagaraEditor/Private/NiagaraEditorModule.cpp
+++ b/Engine/Plugins/FX/Niagara/Source/NiagaraEditor/Private/NiagaraEditorModule.cpp
@@ -1405,6 +1405,9 @@ void FNiagaraEditorModule::OnPostEngineInit()
 
 		PreviewFeatureLevelChangedHandle = EditorEngine->OnPreviewFeatureLevelChanged().AddStatic(UNiagaraScript::SetPreviewFeatureLevel);
 
+		// Ensure we have the right feature level set as the editor may already be in one before we get here
+		UNiagaraScript::SetPreviewFeatureLevel(GEditor->DefaultWorldFeatureLevel);
+
 		// Handle a re-import for mesh renderers
 		if (UImportSubsystem* ImportSubsystem = GEditor->GetEditorSubsystem<UImportSubsystem>())
 		{
