Fixed static analysis warning

diff --git a/Engine/Plugins/FX/Niagara/Source/NiagaraEditor/Private/Widgets/SNiagaraGraphNodeConvert.cpp b/Engine/Plugins/FX/Niagara/Source/NiagaraEditor/Private/Widgets/SNiagaraGraphNodeConvert.cpp
--- a/Engine/Plugins/FX/Niagara/Source/NiagaraEditor/Private/Widgets/SNiagaraGraphNodeConvert.cpp
+++ b/Engine/Plugins/FX/Niagara/Source/NiagaraEditor/Private/Widgets/SNiagaraGraphNodeConvert.cpp
@@ -141,7 +141,7 @@ void SNiagaraGraphNodeConvert::UpdateGraphNode()
 	}
 
 	// set visibility of add pins
-	if (InputPins.Num() > 0 && OutputPins.Num() > 0)
+	if (InputPins.Num() > 0 && OutputPins.Num() > 0 && ConvertNode)
 	{
 		InputPins.Last()->SetVisibility(ConvertNode->IsWiringShown() ? EVisibility::Visible : EVisibility::Collapsed);
 		OutputPins.Last()->SetVisibility(ConvertNode->IsWiringShown() ? EVisibility::Visible : EVisibility::Collapsed);
