Fixed LWC bug in camera DI

diff --git a/Engine/Plugins/FX/Niagara/Source/Niagara/Classes/NiagaraScript.h b/Engine/Plugins/FX/Niagara/Source/Niagara/Classes/NiagaraScript.h
--- a/Engine/Plugins/FX/Niagara/Source/Niagara/Classes/NiagaraScript.h
+++ b/Engine/Plugins/FX/Niagara/Source/Niagara/Classes/NiagaraScript.h
@@ -1086,10 +1086,10 @@ public:
 
 	TOptional<FVector3d> GetCompilerTag(const FNiagaraVariableBase& InVar, const FNiagaraParameterStore* FallbackParameterStore) const
 	{
-		TOptional<FVector3f> Optional(*GetCompilerTag<FVector3f>(InVar, FallbackParameterStore));
+		TOptional<FVector3f> Optional = GetCompilerTag<FVector3f>(InVar, FallbackParameterStore);
 		if (Optional.IsSet())
 		{
-			return TOptional<FVector3d>((FVector3d)Optional.GetValue());
+			return TOptional(FVector3d(Optional.GetValue()));
 		}
 		return TOptional<FVector3d>();
 	}
diff --git a/Engine/Plugins/FX/Niagara/Source/Niagara/Private/NiagaraDataInterfaceCamera.cpp b/Engine/Plugins/FX/Niagara/Source/Niagara/Private/NiagaraDataInterfaceCamera.cpp
--- a/Engine/Plugins/FX/Niagara/Source/Niagara/Private/NiagaraDataInterfaceCamera.cpp
+++ b/Engine/Plugins/FX/Niagara/Source/Niagara/Private/NiagaraDataInterfaceCamera.cpp
@@ -536,11 +536,11 @@ void UNiagaraDataInterfaceCamera::CalculateParticleDistances(FVectorVMExternalFu
 	FNDIInputParam<FNiagaraID> IDParam(Context);
 	FNDIInputParam<FNiagaraPosition> ParticlePosParam(Context);
 
-	FVector CameraPos = (FVector)InstData->CameraLocation;
+	FNiagaraPosition CameraPos = InstData->CameraLocation;
 	for (int32 i = 0; i < Context.GetNumInstances(); ++i)
 	{
 		FDistanceData DistanceData;
-		FVector ParticlePos = InstData->LWCConverter.ConvertSimulationPositionToWorld(ParticlePosParam.GetAndAdvance());
+		FNiagaraPosition ParticlePos = ParticlePosParam.GetAndAdvance();
 		DistanceData.ParticleID = IDParam.GetAndAdvance();
 		DistanceData.DistanceSquared = (ParticlePos - CameraPos).SizeSquared();
 		InstData->DistanceSortQueue.Enqueue(DistanceData);
