Fixes for attribute trimming with Niagara particle scripts
-fixes Particles.Module.AttributeName attributes from getting incorrectly culled
-fixes dependency search to start with the most recent write rather than the first read

diff --git a/Engine/Plugins/FX/Niagara/Source/NiagaraEditor/Private/NiagaraAttributeTrimmer.cpp b/Engine/Plugins/FX/Niagara/Source/NiagaraEditor/Private/NiagaraAttributeTrimmer.cpp
--- a/Engine/Plugins/FX/Niagara/Source/NiagaraEditor/Private/NiagaraAttributeTrimmer.cpp
+++ b/Engine/Plugins/FX/Niagara/Source/NiagaraEditor/Private/NiagaraAttributeTrimmer.cpp
@@ -157,6 +157,18 @@ namespace NiagaraAttributeTrimming
 							{
 								if (Pin->Direction == EEdGraphPinDirection::EGPD_Input)
 								{
+									// setting the following to false will allow the trimmer to remove additional attributes.  A value of false
+									// will limit our dependency search to the function inputs rather than the execution path.
+									constexpr bool bBackCompat_IncludeExecPinsForImpureFunctions = true;
+
+									if (!bBackCompat_IncludeExecPinsForImpureFunctions)
+									{
+										if (UEdGraphSchema_Niagara::PinToTypeDefinition(Pin) == FNiagaraTypeDefinition::GetParameterMapDef())
+										{
+											continue;
+										}
+									}
+
 									ImpureFunctionInputs.Emplace(Pin, *Namespace);
 								}
 							}
@@ -183,7 +195,7 @@ namespace NiagaraAttributeTrimming
 		//	-generates the set of input pins to nodes that contribute to the 'expression' of the supplied output pin (EndPin)
 		//		An expression is the set of nodes that define what is being set in a MapSet.  It does not traverse through
 		//		Exec pins as those nodes would be part of another expression
-		//	-generates the list of custom HLSL nodes that are encounterd in the traversal
+		//	-generates the list of custom HLSL nodes that are encountered in the traversal
 		void FindDependencies(const FModuleScopedPin& EndPin, FDependencyChain& Dependencies)
 		{
 			FindDependenciesInternal(EndPin, Dependencies, nullptr);
@@ -194,7 +206,7 @@ namespace NiagaraAttributeTrimming
 		//	-generates the set of input pins to nodes that contribute to the 'expression' of the supplied output pin (EndPin)
 		//		An expression is the set of nodes that define what is being set in a MapSet.  It does not traverse through
 		//		Exec pins as those nodes would be part of another expression
-		//	-generates the list of custom HLSL nodes that are encounterd in the traversal
+		//	-generates the list of custom HLSL nodes that are encountered in the traversal
 		void FindDependenciesInternal(const FModuleScopedPin& EndPin, FDependencyChain& Dependencies, TArray<UNiagaraNodeInput*>* InputNodes)
 		{
 			TStringBuilder<1024> NamespaceBuilder;
@@ -234,7 +246,7 @@ namespace NiagaraAttributeTrimming
 								NamespaceBuilder.Reset();
 								if (CurrentPin.ModuleName != NAME_None)
 								{
-									NamespaceBuilder << CurrentPin.ModuleName.ToString();
+									CurrentPin.ModuleName.AppendString(NamespaceBuilder);
 									NamespaceBuilder << TEXT(".");
 								}
 								NamespaceBuilder << FunctionCallNode->GetFunctionName();
@@ -445,7 +457,7 @@ namespace NiagaraAttributeTrimming
 		return FModuleScopedPin();
 	}
 
-	static void AddDefaultBinding(const FNiagaraParameterMapHistory& ParamMap, int32 VariableIndex, const FModuleScopedPin& ReadPin, FDependencyChain& ResolvedDependencies)
+	static int32 FindDefaultBinding(const FNiagaraParameterMapHistory& ParamMap, int32 VariableIndex, const FModuleScopedPin& ReadPin)
 	{
 		if (const UNiagaraGraph* OwnerGraph = Cast<const UNiagaraGraph>(ReadPin.Pin->GetOwningNode()->GetGraph()))
 		{
@@ -461,19 +473,11 @@ namespace NiagaraAttributeTrimming
 			TOptional<ENiagaraDefaultMode> DefaultMode = OwnerGraph->GetDefaultMode(ParamVariable, &DefaultBinding);
 			if (DefaultMode.IsSet() && *DefaultMode == ENiagaraDefaultMode::Binding && DefaultBinding.IsValid())
 			{
-				int32 BoundVariableIndex = ParamMap.FindVariableByName(DefaultBinding.GetName());
-
-				if (BoundVariableIndex != INDEX_NONE)
-				{
-					for (const FNiagaraParameterMapHistory::FReadHistory& ReadHistoryEntry : ParamMap.PerVariableReadHistory[BoundVariableIndex])
-					{
-						// for now adding the first pin to the dependencies works, but this needs to be revisited
-						ResolvedDependencies.Pins.Add(ReadHistoryEntry.ReadPin);
-						break;
-					}
-				}
+				return ParamMap.FindVariableByName(DefaultBinding.GetName());
 			}
 		}
+
+		return INDEX_NONE;
 	}
 
 	// for a specific read of a variable finds the actual name of the attribute being read
@@ -533,7 +537,51 @@ namespace NiagaraAttributeTrimming
 								}
 								else if (SourceVariableIndex != INDEX_NONE)
 								{
-									AddDefaultBinding(ParamMap, SourceVariableIndex, ReadPin, ResolvedDependencies);
+									const int32 DefaultBoundVariable = FindDefaultBinding(ParamMap, SourceVariableIndex, ReadPin);
+
+									// if we are bound to a variable we need to grab the last relevant write to that variable and continue our resolution
+									if (DefaultBoundVariable != INDEX_NONE)
+									{
+										bool bValidDefaultWriteFound = false;
+
+										const UNiagaraNode* ReadPinOwningNode = Cast<const UNiagaraNode>(ReadPin.Pin->GetOwningNode());
+										const int32 ReadNodeVisitationIndex = ParamMap.MapNodeVisitations.IndexOfByKey(ReadPinOwningNode);
+										if (ReadNodeVisitationIndex != INDEX_NONE)
+										{
+											const int32 VariableWriteCount = ParamMap.PerVariableWriteHistory[DefaultBoundVariable].Num();
+											for (int32 VariableWriteIt = VariableWriteCount - 1; VariableWriteIt >= 0; --VariableWriteIt)
+											{
+												const FModuleScopedPin& DefaultWritePin = ParamMap.PerVariableWriteHistory[DefaultBoundVariable][VariableWriteIt];
+												const UNiagaraNode* DefaultWritePinOwningNode = Cast<const UNiagaraNode>(DefaultWritePin.Pin->GetOwningNode());
+												const int32 DefaultWriteNodeVisitationIndex = ParamMap.MapNodeVisitations.IndexOfByKey(DefaultWritePinOwningNode);
+												if (DefaultWriteNodeVisitationIndex != INDEX_NONE && DefaultWriteNodeVisitationIndex < ReadNodeVisitationIndex)
+												{
+													PinsToResolve.Add(DefaultWritePin);
+													bValidDefaultWriteFound = true;
+													break;
+												}
+											}
+										}
+
+										// if no valid pin was found that writes into the default bound variable within this parameter map, then we assume
+										// that it will be present in a previous parameter map, and so we just record the earliest read from this parameter
+										// map as a dependency
+
+										if (!ParamMap.PerVariableReadHistory[DefaultBoundVariable].IsEmpty())
+										{
+											// setting the following to false will allow the trimmer to remove additional attributes.  In the cases
+											// where an attribute is written to and used only in a single stage the value could be considered local and not
+											// need to be an attribute.  If bBackCompat_AlwaysIncludeFirstReadPin is true then any time we read from that
+											// variable we'll treat it as a dependent.  If it's false then the writes that contribute to that variable will 
+											// be resolved as part of the dependency resolution.
+											constexpr bool bBackCompat_AlwaysIncludeFirstReadPin = true;
+
+											if (!bValidDefaultWriteFound || bBackCompat_AlwaysIncludeFirstReadPin)
+											{
+												ResolvedDependencies.Pins.Add(ParamMap.PerVariableReadHistory[DefaultBoundVariable][0].ReadPin);
+											}
+										}
+									}
 								}
 							}
 						}
@@ -661,8 +709,14 @@ static void TrimAttributes_Aggressive(const FNiagaraCompileRequestDuplicateData*
 
 				const FName VariableName = *VariableNameString;
 
-				if (VariableNameString.StartsWith(FNiagaraConstants::StackContextNamespaceString))
+				if (ParamMap->FindVariableByName(VariableName) != INDEX_NONE)
 				{
+					AttributesToPreserve.Add(VariableName);
+				}
+				else
+				{
+					// we also need to check to see if the attribute has resolved namespaces (i.e. Particles.Module.MyAttribute or StackContext.MyAttribute)
+					// for which we'll just search through the parameter maps variables directly
 					const int32 AliasedIndex = ParamMap->VariablesWithOriginalAliasesIntact.IndexOfByPredicate([&](const FNiagaraVariable& Variable)
 					{
 						return Variable.GetName() == VariableName;
@@ -673,11 +727,6 @@ static void TrimAttributes_Aggressive(const FNiagaraCompileRequestDuplicateData*
 						AttributesToPreserve.Add(ParamMap->Variables[AliasedIndex].GetName());
 					}
 				}
-				else
-				if (ParamMap->FindVariableByName(VariableName) != INDEX_NONE)
-				{
-					AttributesToPreserve.Add(VariableName);
-				}
 			}
 		}
 	}
diff --git a/Engine/Source/Runtime/Core/Private/UObject/UE5ReleaseStreamObjectVersion.cpp b/Engine/Source/Runtime/Core/Private/UObject/UE5ReleaseStreamObjectVersion.cpp
--- a/Engine/Source/Runtime/Core/Private/UObject/UE5ReleaseStreamObjectVersion.cpp
+++ b/Engine/Source/Runtime/Core/Private/UObject/UE5ReleaseStreamObjectVersion.cpp
@@ -12,7 +12,7 @@ TMap<FGuid, FGuid> FUE5ReleaseStreamObjectVersion::GetSystemGuids()
 	SystemGuids.Add(DevGuids.MATERIALSHADERMAP_DERIVEDDATA_VER, FGuid("79090A66F9D94B5BB285D49E5D39468E"));
 	SystemGuids.Add(DevGuids.NANITE_DERIVEDDATA_VER, FGuid("AC88CFBDDA614A9C8488F64D2DAF699D"));
 	SystemGuids.Add(DevGuids.NIAGARASHADERMAP_DERIVEDDATA_VER, FGuid("6BE35B6FACB34568970120F9BC8DAB80"));
-	SystemGuids.Add(DevGuids.Niagara_LatestScriptCompileVersion, FGuid("4882CE3BA4DF4D6F8D8B7828CED1C1A4"));
+	SystemGuids.Add(DevGuids.Niagara_LatestScriptCompileVersion, FGuid("0CD19494DF6F4B549FC12A0CCC01A02E"));
 	SystemGuids.Add(DevGuids.SkeletalMeshDerivedDataVersion, FGuid("25C49E579B3142DDA2A8C14037267679"));
 
 	return SystemGuids;
