Niagara - Refactor the stack function input cache so that it stores the inputs as variables rather than storing pointers to the pins themselves since their lifetime can be unpredictable due to how we abuse the graphs.

diff --git a/Engine/Plugins/FX/Niagara/Source/NiagaraEditor/Private/NiagaraGraphDataCache.cpp b/Engine/Plugins/FX/Niagara/Source/NiagaraEditor/Private/NiagaraGraphDataCache.cpp
--- a/Engine/Plugins/FX/Niagara/Source/NiagaraEditor/Private/NiagaraGraphDataCache.cpp
+++ b/Engine/Plugins/FX/Niagara/Source/NiagaraEditor/Private/NiagaraGraphDataCache.cpp
@@ -2,6 +2,7 @@
 
 #include "NiagaraGraphDataCache.h"
 
+#include "EdGraphSchema_Niagara.h"
 #include "NiagaraNodeFunctionCall.h"
 #include "NiagaraNodeAssignment.h"
 
@@ -26,11 +27,11 @@ FNiagaraGraphDataCache::FNiagaraGraphDataCache()
 {
 }
 
-void FNiagaraGraphDataCache::GetStackFunctionInputPinsInternal(
+void FNiagaraGraphDataCache::GetStackFunctionInputsInternal(
 	const UNiagaraNodeFunctionCall& FunctionCallNode,
 	const UNiagaraGraph* CalledGraph,
 	TConstArrayView<FNiagaraVariable> StaticVars,
-	TArray<const UEdGraphPin*>& OutInputPins,
+	TArray<FNiagaraVariable>& OutInputVariables,
 	FCompileConstantResolver ConstantResolver,
 	FNiagaraStackGraphUtilities::ENiagaraGetStackFunctionInputPinsOptions Options,
 	bool bIgnoreDisabled,
@@ -38,17 +39,23 @@ void FNiagaraGraphDataCache::GetStackFunctionInputPinsInternal(
 {
 	// we don't currently support hashing the internal details of the translator within FStackFunctionInputPinKey so in the case
 	// we've been supplied a translator we'll just fall back to not using the cache
+	TArray<const UEdGraphPin*> InputPins;
 	if (const FHlslNiagaraTranslator* Translator = ConstantResolver.GetTranslator())
 	{
 		FNiagaraStackGraphUtilities::GetStackFunctionInputPinsWithoutCache(
 			FunctionCallNode,
 			StaticVars,
-			OutInputPins,
+			InputPins,
 			ConstantResolver,
 			Options,
 			bIgnoreDisabled,
 			bFilterForCompilation);
 
+		for (const UEdGraphPin* InputPin : InputPins)
+		{
+			OutInputVariables.Add(UEdGraphSchema_Niagara::PinToNiagaraVariable(InputPin));
+		}
+
 		return;
 	}
 
@@ -59,36 +66,40 @@ void FNiagaraGraphDataCache::GetStackFunctionInputPinsInternal(
 	}
 
 	FStackFunctionInputPinKey CacheKey(FunctionCallNode, CalledGraph, StaticVars, ConstantResolver, Options, bIgnoreDisabled, bFilterForCompilation);
-	
+
 	bool bFoundInCache = false;
 	if(bUseCache)
 	{
 		if (const FStackFunctionInputPinValue* CachedValue = StackFunctionInputPinCache.FindAndTouch(CacheKey))
 		{
 			bFoundInCache = true;
-			OutInputPins = CachedValue->InputPins;
+			OutInputVariables = CachedValue->InputVariables;
 		}
 	}
 
 	if(bUseCache == false || bFoundInCache == false)
 	{
-		FStackFunctionInputPinValue NewValue;
 		FNiagaraStackGraphUtilities::GetStackFunctionInputPinsWithoutCache(
 			FunctionCallNode,
 			StaticVars,
-			NewValue.InputPins,
+			InputPins,
 			ConstantResolver,
 			Options,
 			bIgnoreDisabled,
 			bFilterForCompilation);
 
-		OutInputPins = NewValue.InputPins;
+		for(const UEdGraphPin* InputPin : InputPins)
+		{
+			OutInputVariables.Add(UEdGraphSchema_Niagara::PinToNiagaraVariable(InputPin));
+		}
 
 		if(bUseCache)
 		{
+			FStackFunctionInputPinValue NewValue;
+			NewValue.InputVariables = OutInputVariables;
 			StackFunctionInputPinCache.Add(CacheKey, NewValue);
 		}
-	}	
+	}
 
 	if (GNiagaraGraphDataCacheValidation)
 	{
@@ -102,43 +113,41 @@ void FNiagaraGraphDataCache::GetStackFunctionInputPinsInternal(
 			bIgnoreDisabled,
 			bFilterForCompilation);
 
-		if (!ensureMsgf(ValidationPins == OutInputPins, TEXT("FNiagaraGraphDataCache failed to accurately collect StackFunctionInputPins for Function {0}"), *FunctionCallNode.GetPathName()))
+		TArray<FNiagaraVariable> ValidationVariables;
+		for (const UEdGraphPin* ValidationPin : ValidationPins)
+		{
+			ValidationVariables.Add(UEdGraphSchema_Niagara::PinToNiagaraVariable(ValidationPin));
+		}
+
+		if (!ensureMsgf(ValidationVariables == OutInputVariables, TEXT("FNiagaraGraphDataCache failed to accurately collect StackFunctionInputPins for Function {0}"), *FunctionCallNode.GetPathName()))
 		{
-			OutInputPins = ValidationPins;
+			OutInputVariables = ValidationVariables;
 		}
 	}
 }
 
-void FNiagaraGraphDataCache::GetStackFunctionInputPins(const UNiagaraNodeFunctionCall& FunctionCallNode, TConstArrayView<FNiagaraVariable> StaticVars, TArray<const UEdGraphPin*>& OutInputPins, FCompileConstantResolver ConstantResolver, FNiagaraStackGraphUtilities::ENiagaraGetStackFunctionInputPinsOptions Options, bool bIgnoreDisabled)
+void FNiagaraGraphDataCache::GetStackFunctionInputs(const UNiagaraNodeFunctionCall& FunctionCallNode, TConstArrayView<FNiagaraVariable> StaticVars, TArray<FNiagaraVariable>& OutInputVariables, FCompileConstantResolver ConstantResolver, FNiagaraStackGraphUtilities::ENiagaraGetStackFunctionInputPinsOptions Options, bool bIgnoreDisabled)
 {
 	if (const UNiagaraGraph* CalledGraph = FunctionCallNode.GetCalledGraph())
 	{
-		GetStackFunctionInputPinsInternal(FunctionCallNode, CalledGraph, StaticVars, OutInputPins, ConstantResolver, Options, bIgnoreDisabled, false /*bFilterForCompilation*/);
+		GetStackFunctionInputsInternal(FunctionCallNode, CalledGraph, StaticVars, OutInputVariables, ConstantResolver, Options, bIgnoreDisabled, false /*bFilterForCompilation*/);
 	}
 }
 
-void FNiagaraGraphDataCache::GetStackFunctionInputPins(const UNiagaraNodeFunctionCall& FunctionCallNode, TConstArrayView<FNiagaraVariable> StaticVars, TArray<const UEdGraphPin*>& OutInputPins, TSet<const UEdGraphPin*>& OutHiddenPins, FCompileConstantResolver ConstantResolver, FNiagaraStackGraphUtilities::ENiagaraGetStackFunctionInputPinsOptions Options, bool bIgnoreDisabled)
+void FNiagaraGraphDataCache::GetStackFunctionInputs(const UNiagaraNodeFunctionCall& FunctionCallNode, TConstArrayView<FNiagaraVariable> StaticVars, TArray<FNiagaraVariable>& OutInputVariables, TSet<FNiagaraVariable>& OutHiddenVariables, FCompileConstantResolver ConstantResolver, FNiagaraStackGraphUtilities::ENiagaraGetStackFunctionInputPinsOptions Options, bool bIgnoreDisabled)
 {
 	if (const UNiagaraGraph* CalledGraph = FunctionCallNode.GetCalledGraph())
 	{
-		TArray<const UEdGraphPin*> FilteredPins;
-
-		GetStackFunctionInputPinsInternal(FunctionCallNode, CalledGraph, StaticVars, OutInputPins, ConstantResolver, Options, bIgnoreDisabled, false /*bFilterForCompilation*/);
-		GetStackFunctionInputPinsInternal(FunctionCallNode, CalledGraph, StaticVars, FilteredPins, ConstantResolver, Options, bIgnoreDisabled, true /*bFilterForCompilation*/);
+		TArray<FNiagaraVariable> FilteredVariables;
 
-		// generate hidden pins
-		auto PinsMatch = [](const UEdGraphPin* Lhs, const UEdGraphPin* Rhs) -> bool
-		{
-			return Lhs->GetFName() == Rhs->GetFName()
-				&& Lhs->PinType.PinCategory == Rhs->PinType.PinCategory
-				&& Lhs->PinType.PinSubCategoryObject == Rhs->PinType.PinSubCategoryObject;
-		};
+		GetStackFunctionInputsInternal(FunctionCallNode, CalledGraph, StaticVars, OutInputVariables, ConstantResolver, Options, bIgnoreDisabled, false /*bFilterForCompilation*/);
+		GetStackFunctionInputsInternal(FunctionCallNode, CalledGraph, StaticVars, FilteredVariables, ConstantResolver, Options, bIgnoreDisabled, true /*bFilterForCompilation*/);
 
-		for (const UEdGraphPin* InputPin : OutInputPins)
+		for (const FNiagaraVariable& InputVariable : OutInputVariables)
 		{
-			if (!FilteredPins.ContainsByPredicate([&](const UEdGraphPin* CompilationPin) { return PinsMatch(InputPin, CompilationPin); }))
+			if (FilteredVariables.Contains(InputVariable) == false)
 			{
-				OutHiddenPins.Add(InputPin);
+				OutHiddenVariables.Add(InputVariable);
 			}
 		}
 	}
diff --git a/Engine/Plugins/FX/Niagara/Source/NiagaraEditor/Private/NiagaraNodeFunctionCall.cpp b/Engine/Plugins/FX/Niagara/Source/NiagaraEditor/Private/NiagaraNodeFunctionCall.cpp
--- a/Engine/Plugins/FX/Niagara/Source/NiagaraEditor/Private/NiagaraNodeFunctionCall.cpp
+++ b/Engine/Plugins/FX/Niagara/Source/NiagaraEditor/Private/NiagaraNodeFunctionCall.cpp
@@ -1386,7 +1386,7 @@ void UNiagaraNodeFunctionCall::ChangeScriptVersion(FGuid NewScriptVersion, const
 	if (FunctionScript)
 	{
 		// Automatically remove old inputs so it does not show a bunch of warnings to the user
-		TMap<FName, FEdGraphPinType> FunctionInputNames;
+		TMap<FName, FNiagaraTypeDefinition> FunctionInputNames;
 		FPinCollectorArray OverridePins;
 		UNiagaraNodeParameterMapSet* OverrideNode = FNiagaraStackGraphUtilities::GetStackFunctionOverrideNode(*this);
 		if (OverrideNode != nullptr)
@@ -1394,11 +1394,11 @@ void UNiagaraNodeFunctionCall::ChangeScriptVersion(FGuid NewScriptVersion, const
 			OverrideNode->Modify();
 			OverrideNode->GetInputPins(OverridePins);
 			
-			TArray<const UEdGraphPin*> ModuleInputPins;
-			FNiagaraStackGraphUtilities::GetStackFunctionInputPins(*this, ModuleInputPins, UpgradeContext.ConstantResolver, FNiagaraStackGraphUtilities::ENiagaraGetStackFunctionInputPinsOptions::ModuleInputsOnly);
-			for (const UEdGraphPin* InputPin : ModuleInputPins)
+			TArray<FNiagaraVariable> ModuleInputVariables;
+			FNiagaraStackGraphUtilities::GetStackFunctionInputs(*this, ModuleInputVariables, UpgradeContext.ConstantResolver, FNiagaraStackGraphUtilities::ENiagaraGetStackFunctionInputPinsOptions::ModuleInputsOnly);
+			for (const FNiagaraVariable& InputVariable : ModuleInputVariables)
 			{
-				FunctionInputNames.Add(FNiagaraParameterHandle(InputPin->PinName).GetName(), InputPin->PinType);
+				FunctionInputNames.Add(FNiagaraParameterHandle(InputVariable.GetName()).GetName(), InputVariable.GetType());
 			}
 		}
 		for (UEdGraphPin* OverridePin : OverridePins)
@@ -1408,14 +1408,15 @@ void UNiagaraNodeFunctionCall::ChangeScriptVersion(FGuid NewScriptVersion, const
 				continue;
 			}
 			FName InputName = FNiagaraParameterHandle(OverridePin->PinName).GetName();
-			FEdGraphPinType* PinType = FunctionInputNames.Find(InputName);
-			if (PinType && *PinType != OverridePin->PinType)
+			FNiagaraTypeDefinition* InputType = FunctionInputNames.Find(InputName);
+			FNiagaraTypeDefinition OverridePinType = UEdGraphSchema_Niagara::PinTypeToTypeDefinition(OverridePin->PinType);
+			if (InputType != nullptr && *InputType != OverridePinType)
 			{
 				// looks like the type of the module input changed - we'll change the override pin type as well
 				OverridePin->Modify();
-				OverridePin->PinType = *PinType;
+				OverridePin->PinType = UEdGraphSchema_Niagara::TypeDefinitionToPinType(*InputType);
 			}
-			else if (PinType == nullptr)
+			else if (InputType == nullptr)
 			{
 				// the input doesn't exist any more, delete the override pin
 				TArray<TWeakObjectPtr<UNiagaraDataInterface>> RemovedDataObjects;
diff --git a/Engine/Plugins/FX/Niagara/Source/NiagaraEditor/Private/NiagaraScriptMergeManager.cpp b/Engine/Plugins/FX/Niagara/Source/NiagaraEditor/Private/NiagaraScriptMergeManager.cpp
--- a/Engine/Plugins/FX/Niagara/Source/NiagaraEditor/Private/NiagaraScriptMergeManager.cpp
+++ b/Engine/Plugins/FX/Niagara/Source/NiagaraEditor/Private/NiagaraScriptMergeManager.cpp
@@ -252,21 +252,16 @@ FNiagaraStackFunctionMergeAdapter::FNiagaraStackFunctionMergeAdapter(const FVers
 	// can be ignored since they're not actually overrides.  This is usually not an issue due to the PreparateRapidIterationParameters call in the emitter
 	// editor, but modifications to modules can cause inconsistency here in the emitter.
 	TArray<FNiagaraVariable> RapidIterationInputDefaultValues;
-	const UEdGraphSchema_Niagara* NiagaraSchema = GetDefault<UEdGraphSchema_Niagara>();
 	if (InFunctionCallNode.FunctionScript != nullptr)
 	{
 		FCompileConstantResolver Resolver(InOwningEmitter, ENiagaraScriptUsage::Function);
-		TArray<const UEdGraphPin*> FunctionInputPins;
-		GetStackFunctionInputPins(*FunctionCallNode, FunctionInputPins, Resolver, FNiagaraStackGraphUtilities::ENiagaraGetStackFunctionInputPinsOptions::ModuleInputsOnly, false);
-		
-		TArray<FNiagaraVariable> FunctionInputVariables;
 		TArray<FName> FunctionInputVariableNames;
-		for (const UEdGraphPin* FunctionInputPin : FunctionInputPins)
+		TArray<FNiagaraVariable> FunctionInputVariables;
+		GetStackFunctionInputs(*FunctionCallNode, FunctionInputVariables, Resolver, FNiagaraStackGraphUtilities::ENiagaraGetStackFunctionInputPinsOptions::ModuleInputsOnly, false);
+		for (const FNiagaraVariable& FunctionInputVariable : FunctionInputVariables)
 		{
-			FNiagaraVariable FunctionInputVariable = NiagaraSchema->PinToNiagaraVariable(FunctionInputPin);
 			if (FunctionInputVariable.IsValid() && FNiagaraStackGraphUtilities::IsRapidIterationType(FunctionInputVariable.GetType()))
 			{
-				FunctionInputVariables.Add(FunctionInputVariable);
 				FunctionInputVariableNames.Add(FunctionInputVariable.GetName());
 			}
 		}
@@ -282,7 +277,7 @@ FNiagaraStackFunctionMergeAdapter::FNiagaraStackFunctionMergeAdapter(const FVers
 			if (FunctionInputDefaultPin != nullptr)
 			{
 				// Try to get the default value from the default pin.
-				FNiagaraVariable FunctionInputDefaultVariable = NiagaraSchema->PinToNiagaraVariable(FunctionInputDefaultPin);
+				FNiagaraVariable FunctionInputDefaultVariable = UEdGraphSchema_Niagara::PinToNiagaraVariable(FunctionInputDefaultPin);
 				if (FunctionInputDefaultVariable.GetData() != nullptr)
 				{
 					FunctionInputVariable.SetData(FunctionInputDefaultVariable.GetData());
diff --git a/Engine/Plugins/FX/Niagara/Source/NiagaraEditor/Private/NiagaraScriptSource.cpp b/Engine/Plugins/FX/Niagara/Source/NiagaraEditor/Private/NiagaraScriptSource.cpp
--- a/Engine/Plugins/FX/Niagara/Source/NiagaraEditor/Private/NiagaraScriptSource.cpp
+++ b/Engine/Plugins/FX/Niagara/Source/NiagaraEditor/Private/NiagaraScriptSource.cpp
@@ -340,7 +340,7 @@ bool UNiagaraScriptSource::AddModuleIfMissing(FString ModulePath, ENiagaraScript
 	return false;
 }
 
-void UNiagaraScriptSource::FixupRenamedParameters(UNiagaraNodeFunctionCall* FunctionCallNode, TConstArrayView<const UEdGraphPin*> ModuleInputPins, TConstArrayView<FNiagaraVariable> ModulePinVariables, FNiagaraParameterStore& RapidIterationParameters, const TArray<FNiagaraVariable>& OldRapidIterationVariables, const FVersionedNiagaraEmitter& VersionedEmitter, ENiagaraScriptUsage ScriptUsage) const
+void UNiagaraScriptSource::FixupRenamedParameters(UNiagaraNodeFunctionCall* FunctionCallNode, TConstArrayView<FNiagaraVariable> ModuleInputVariables, FNiagaraParameterStore& RapidIterationParameters, const TArray<FNiagaraVariable>& OldRapidIterationVariables, const FVersionedNiagaraEmitter& VersionedEmitter, ENiagaraScriptUsage ScriptUsage) const
 {
 	// the rapid iteration parameters and the function input pins use different variable naming schemes, so most of this is just used to convert one name to the other 
 	UNiagaraGraph* Graph = FunctionCallNode->GetCalledGraph();
@@ -351,7 +351,7 @@ void UNiagaraScriptSource::FixupRenamedParameters(UNiagaraNodeFunctionCall* Func
 	TArray<FString, TInlineAllocator<16>> ModulePinConstantName;
 	TArray<int32, TInlineAllocator<16>> ModulePinNameSpaceEnd;
 
-	const int32 ModuleInputPinCount = ModuleInputPins.Num();
+	const int32 ModuleInputPinCount = ModuleInputVariables.Num();
 	ModulePinMetaData.Reserve(ModuleInputPinCount);
 	ModulePinAliasedName.Reserve(ModuleInputPinCount);
 	ModulePinConstantName.Reserve(ModuleInputPinCount);
@@ -361,11 +361,10 @@ void UNiagaraScriptSource::FixupRenamedParameters(UNiagaraNodeFunctionCall* Func
 
 	for (int32 ModuleInputPinIt = 0; ModuleInputPinIt < ModuleInputPinCount; ++ModuleInputPinIt)
 	{
-		const UEdGraphPin* ModuleInputPin = ModuleInputPins[ModuleInputPinIt];
-		const FNiagaraVariable& InputVar = ModulePinVariables[ModuleInputPinIt];
+		const FNiagaraVariable& ModuleInputVariable = ModuleInputVariables[ModuleInputPinIt];
 
-		ModulePinMetaData.Add(Graph->GetMetaData(InputVar));
-		FNiagaraParameterHandle AliasedFunctionInputHandle = FNiagaraParameterHandle::CreateAliasedModuleParameterHandle(ModuleInputPin->PinName, FunctionName);
+		ModulePinMetaData.Add(Graph->GetMetaData(ModuleInputVariable));
+		FNiagaraParameterHandle AliasedFunctionInputHandle = FNiagaraParameterHandle::CreateAliasedModuleParameterHandle(ModuleInputVariable.GetName(), FunctionName);
 		const FName& AliasedName = ModulePinAliasedName.Add_GetRef(AliasedFunctionInputHandle.GetParameterHandleString());
 		const FString& ConstantName = ModulePinConstantName.Add_GetRef(FNiagaraUtilities::CreateRapidIterationConstantName(AliasedName, *UniqueEmitterName, ScriptUsage));
 
@@ -379,7 +378,7 @@ void UNiagaraScriptSource::FixupRenamedParameters(UNiagaraNodeFunctionCall* Func
 		{
 			for (int32 ModuleInputPinIt = 0; ModuleInputPinIt < ModuleInputPinCount; ++ModuleInputPinIt)
 			{
-				const FNiagaraVariable& InputVar = ModulePinVariables[ModuleInputPinIt];
+				const FNiagaraVariable& InputVar = ModuleInputVariables[ModuleInputPinIt];
 				const FStringView ConstantName(ModulePinConstantName[ModuleInputPinIt]);
 				const int32 NameSpaceEnd = ModulePinNameSpaceEnd[ModuleInputPinIt];
 
@@ -434,7 +433,7 @@ void UNiagaraScriptSource::FixupRenamedParameters(UNiagaraNodeFunctionCall* Func
 	}
 }
 
-void UNiagaraScriptSource::InitializeNewParameters(UNiagaraNodeFunctionCall* FunctionCallNode, TConstArrayView<const UEdGraphPin*> FunctionInputPins, TConstArrayView<FNiagaraVariable> PinVariables, FNiagaraParameterStore& RapidIterationParameters, const FVersionedNiagaraEmitter& VersionedEmitter, ENiagaraScriptUsage ScriptUsage, TSet<FName>& ValidRapidIterationParameterNames) const
+void UNiagaraScriptSource::InitializeNewParameters(UNiagaraNodeFunctionCall* FunctionCallNode, TConstArrayView<FNiagaraVariable> ModuleInputVariables, FNiagaraParameterStore& RapidIterationParameters, const FVersionedNiagaraEmitter& VersionedEmitter, ENiagaraScriptUsage ScriptUsage, TSet<FName>& ValidRapidIterationParameterNames) const
 {
 	const FString UniqueEmitterName = VersionedEmitter.Emitter ? VersionedEmitter.Emitter->GetUniqueEmitterName() : FString();
 	UNiagaraGraph* Graph = FunctionCallNode->GetCalledGraph();
@@ -445,7 +444,7 @@ void UNiagaraScriptSource::InitializeNewParameters(UNiagaraNodeFunctionCall* Fun
 	TArray<FNiagaraParameterHandle, TInlineAllocator<16>> AliasHandles;
 	TArray<FNiagaraVariable, TInlineAllocator<16>> Parameters;
 
-	const int32 FunctionInputPinCount = FunctionInputPins.Num();
+	const int32 FunctionInputPinCount = ModuleInputVariables.Num();
 	PinNames.Reserve(FunctionInputPinCount);
 	DefaultPins.Reserve(FunctionInputPinCount);
 	AliasHandles.Reserve(FunctionInputPinCount);
@@ -455,26 +454,25 @@ void UNiagaraScriptSource::InitializeNewParameters(UNiagaraNodeFunctionCall* Fun
 
 	for (int32 FunctionInputPinIt = 0; FunctionInputPinIt < FunctionInputPinCount; ++FunctionInputPinIt)
 	{
-		const UEdGraphPin* FunctionInputPin = FunctionInputPins[FunctionInputPinIt];
-		const FNiagaraTypeDefinition& InputType = PinVariables[FunctionInputPinIt].GetType();
-		if (InputType.IsValid() == false)
+		const FNiagaraVariable& ModuleInputVariable = ModuleInputVariables[FunctionInputPinIt];
+		if (ModuleInputVariable.GetType().IsValid() == false)
 		{
 			UE_LOG(LogNiagaraEditor, Error, TEXT("Invalid input type found while attempting initialize new rapid iteration parameters. Function Node: %s %s Input Name: %s"),
-				*FunctionCallNode->GetPathName(), *FunctionCallNode->GetFunctionName(), *FunctionInputPin->GetName());
+				*FunctionCallNode->GetPathName(), *FunctionCallNode->GetFunctionName(), *ModuleInputVariable.GetName().ToString());
 			continue;
 		}
 
-		if (FNiagaraStackGraphUtilities::IsRapidIterationType(InputType))
+		if (FNiagaraStackGraphUtilities::IsRapidIterationType(ModuleInputVariable.GetType()))
 		{
-			FNiagaraParameterHandle AliasedFunctionInputHandle = FNiagaraParameterHandle::CreateAliasedModuleParameterHandle(FunctionInputPin->PinName, FunctionNodeName);
-			FNiagaraVariable RapidIterationParameter = FNiagaraStackGraphUtilities::CreateRapidIterationParameter(UniqueEmitterName, ScriptUsage, AliasedFunctionInputHandle.GetParameterHandleString(), InputType);
+			FNiagaraParameterHandle AliasedFunctionInputHandle = FNiagaraParameterHandle::CreateAliasedModuleParameterHandle(ModuleInputVariable.GetName(), FunctionNodeName);
+			FNiagaraVariable RapidIterationParameter = FNiagaraStackGraphUtilities::CreateRapidIterationParameter(UniqueEmitterName, ScriptUsage, AliasedFunctionInputHandle.GetParameterHandleString(), ModuleInputVariable.GetType());
 			ValidRapidIterationParameterNames.Add(RapidIterationParameter.GetName());
 			int32 ParameterIndex = RapidIterationParameters.IndexOf(RapidIterationParameter);
 
 			// Only set a value for the parameter if it's not already set.
 			if (ParameterIndex == INDEX_NONE)
 			{
-				PinNames.Add(FunctionInputPin->PinName);
+				PinNames.Add(ModuleInputVariable.GetName());
 				AliasHandles.Add(AliasedFunctionInputHandle);
 				Parameters.Add(RapidIterationParameter);
 			}
@@ -540,7 +538,6 @@ void UNiagaraScriptSource::CleanUpOldAndInitializeNewRapidIterationParameters(co
 		{
 			TArray<UNiagaraNode*> Nodes;
 			NodeGraph->BuildTraversal(Nodes, OutputNode);
-			const UEdGraphSchema_Niagara* Schema = GetDefault<UEdGraphSchema_Niagara>();
 
 			for(UNiagaraNode* Node : Nodes)
 			{
@@ -554,25 +551,16 @@ void UNiagaraScriptSource::CleanUpOldAndInitializeNewRapidIterationParameters(co
 				UNiagaraGraph* Graph = FunctionCallNode->GetCalledGraph();
 				const UEdGraphSchema_Niagara* NiagaraSchema = Graph->GetNiagaraSchema();
 
-				TArray<const UEdGraphPin*> ModuleInputPins;
+				TArray<FNiagaraVariable> ModuleInputVariables;
 				FCompileConstantResolver ConstantResolver(Emitter, ScriptUsage);
 
-				GetStackFunctionInputPins(*FunctionCallNode, ModuleInputPins, ConstantResolver, FNiagaraStackGraphUtilities::ENiagaraGetStackFunctionInputPinsOptions::ModuleInputsOnly);
+				GetStackFunctionInputs(*FunctionCallNode, ModuleInputVariables, ConstantResolver, FNiagaraStackGraphUtilities::ENiagaraGetStackFunctionInputPinsOptions::ModuleInputsOnly);
 
 				// if we found any module input pins, then we can continue with fixing up and initializing them
-				if (!ModuleInputPins.IsEmpty())
+				if (!ModuleInputVariables.IsEmpty())
 				{
-					TArray<FNiagaraVariable> ModuleInputPinVariables;
-					ModuleInputPinVariables.Reserve(ModuleInputPins.Num());
-					for (const UEdGraphPin* InputPin : ModuleInputPins)
-					{
-						// PinToNiagaraVariable can be costly (when in sufficient numbers) so do this work once for
-						// both FixupRenamedParameters & InitializeNewParameters
-						ModuleInputPinVariables.Add(NiagaraSchema->PinToNiagaraVariable(InputPin));
-					}
-
-					FixupRenamedParameters(FunctionCallNode, ModuleInputPins, ModuleInputPinVariables, RapidIterationParameters, OldRapidIterationVariables, Emitter, ScriptUsage);
-					InitializeNewParameters(FunctionCallNode, ModuleInputPins, ModuleInputPinVariables, RapidIterationParameters, Emitter, OutputNode->GetUsage(), ValidRapidIterationParameterNames);
+					FixupRenamedParameters(FunctionCallNode, ModuleInputVariables, RapidIterationParameters, OldRapidIterationVariables, Emitter, ScriptUsage);
+					InitializeNewParameters(FunctionCallNode, ModuleInputVariables, RapidIterationParameters, Emitter, OutputNode->GetUsage(), ValidRapidIterationParameterNames);
 				}
 			}
 		}
diff --git a/Engine/Plugins/FX/Niagara/Source/NiagaraEditor/Private/NiagaraStackFunctionInputBinder.cpp b/Engine/Plugins/FX/Niagara/Source/NiagaraEditor/Private/NiagaraStackFunctionInputBinder.cpp
--- a/Engine/Plugins/FX/Niagara/Source/NiagaraEditor/Private/NiagaraStackFunctionInputBinder.cpp
+++ b/Engine/Plugins/FX/Niagara/Source/NiagaraEditor/Private/NiagaraStackFunctionInputBinder.cpp
@@ -94,19 +94,17 @@ bool FNiagaraStackFunctionInputBinder::TryBindInternal(
 	}
 	FunctionCallNode = InFunctionCallNode;
 
-	TArray<const UEdGraphPin*> InputPins;
-	TSet<const UEdGraphPin*> HiddenInputPins;
-	FNiagaraStackGraphUtilities::GetStackFunctionInputPins(*FunctionCallNode, InputPins, HiddenInputPins, InConstantResolver, FNiagaraStackGraphUtilities::ENiagaraGetStackFunctionInputPinsOptions::ModuleInputsOnly);
+	TArray<FNiagaraVariable> InputVariables;
+	TSet<FNiagaraVariable> HiddenInputVariables;
+	FNiagaraStackGraphUtilities::GetStackFunctionInputs(*FunctionCallNode, InputVariables, HiddenInputVariables, InConstantResolver, FNiagaraStackGraphUtilities::ENiagaraGetStackFunctionInputPinsOptions::ModuleInputsOnly);
 	bool bInputFound = false;
-	const UEdGraphSchema_Niagara* Schema = GetDefault<UEdGraphSchema_Niagara>();
-	for (const UEdGraphPin* InputPin : InputPins)
+	for (const FNiagaraVariable& InputVariable : InputVariables)
 	{
-		if (HiddenInputPins.Contains(InputPin))
+		if (HiddenInputVariables.Contains(InputVariable))
 		{
 			continue;
 		}
 
-		FNiagaraVariable InputVariable = Schema->PinToNiagaraVariable(InputPin);
 		if(InputMatchesCallback.Execute(InputVariable))
 		{
 			InputName = InputVariable.GetName();
diff --git a/Engine/Plugins/FX/Niagara/Source/NiagaraEditor/Private/ViewModels/NiagaraCurveSelectionViewModel.cpp b/Engine/Plugins/FX/Niagara/Source/NiagaraEditor/Private/ViewModels/NiagaraCurveSelectionViewModel.cpp
--- a/Engine/Plugins/FX/Niagara/Source/NiagaraEditor/Private/ViewModels/NiagaraCurveSelectionViewModel.cpp
+++ b/Engine/Plugins/FX/Niagara/Source/NiagaraEditor/Private/ViewModels/NiagaraCurveSelectionViewModel.cpp
@@ -509,18 +509,16 @@ TSharedPtr<FNiagaraCurveSelectionTreeNode> UNiagaraCurveSelectionViewModel::Crea
 		}
 	}
 
-	const UEdGraphSchema_Niagara* NiagaraSchema = GetDefault<UEdGraphSchema_Niagara>();
-	TArray<const UEdGraphPin*> InputPins;
-	TSet<const UEdGraphPin*> HiddenPins;
-	FNiagaraStackGraphUtilities::GetStackFunctionInputPins(FunctionCallNode, InputPins, HiddenPins, ConstantResolver, FNiagaraStackGraphUtilities::ENiagaraGetStackFunctionInputPinsOptions::ModuleInputsOnly);
-	for (const UEdGraphPin* InputPin : InputPins)
+	TArray<FNiagaraVariable> InputVariables;
+	TSet<FNiagaraVariable> HiddenVariables;
+	FNiagaraStackGraphUtilities::GetStackFunctionInputs(FunctionCallNode, InputVariables, HiddenVariables, ConstantResolver, FNiagaraStackGraphUtilities::ENiagaraGetStackFunctionInputPinsOptions::ModuleInputsOnly);
+	for (FNiagaraVariable InputVariable : InputVariables)
 	{
-		if (HiddenPins.Contains(InputPin))
+		if (HiddenVariables.Contains(InputVariable))
 		{
 			continue;
 		}
 
-		FNiagaraVariable InputVariable = NiagaraSchema->PinToNiagaraVariable(InputPin);
 		if (InputVariable.IsValid() && InputVariable.GetType().IsDataInterface() && InputVariable.GetType().GetClass()->IsChildOf(UNiagaraDataInterfaceCurveBase::StaticClass()))
 		{
 			UNiagaraDataInterfaceCurveBase* DataInterfaceToDisplay = nullptr;
diff --git a/Engine/Plugins/FX/Niagara/Source/NiagaraEditor/Private/ViewModels/Stack/NiagaraStackFunctionInputCollection.cpp b/Engine/Plugins/FX/Niagara/Source/NiagaraEditor/Private/ViewModels/Stack/NiagaraStackFunctionInputCollection.cpp
--- a/Engine/Plugins/FX/Niagara/Source/NiagaraEditor/Private/ViewModels/Stack/NiagaraStackFunctionInputCollection.cpp
+++ b/Engine/Plugins/FX/Niagara/Source/NiagaraEditor/Private/ViewModels/Stack/NiagaraStackFunctionInputCollection.cpp
@@ -290,9 +290,8 @@ void UNiagaraStackFunctionInputCollectionBase::AppendInputsForFunctionCall(FFunc
 	FVersionedNiagaraEmitter Emitter = GetEmitterViewModel().IsValid() ? GetEmitterViewModel()->GetEmitter() : FVersionedNiagaraEmitter();
 	FVersionedNiagaraEmitterData* EmitterData = Emitter.GetEmitterData();
 
-	TSet<const UEdGraphPin*> HiddenPins;
-	TArray<const UEdGraphPin*> InputPins;
-	TSet<const UEdGraphPin*> SummaryViewPins;
+	TSet<FNiagaraVariable> HiddenVariables;
+	TArray<FNiagaraVariable> InputVariables;
 	FCompileConstantResolver ConstantResolver;
 	if (GetEmitterViewModel().IsValid())
 	{
@@ -303,7 +302,7 @@ void UNiagaraStackFunctionInputCollectionBase::AppendInputsForFunctionCall(FFunc
 		// if we don't have an emitter model, we must be in a system context
 		ConstantResolver = FCompileConstantResolver(&GetSystemViewModel()->GetSystem(), FNiagaraStackGraphUtilities::GetOutputNodeUsage(*InputFunctionCallNode));
 	}
-	GetStackFunctionInputPins(*InputFunctionCallNode, InputPins, HiddenPins, ConstantResolver, FNiagaraStackGraphUtilities::ENiagaraGetStackFunctionInputPinsOptions::ModuleInputsOnly);
+	GetStackFunctionInputs(*InputFunctionCallNode, InputVariables, HiddenVariables, ConstantResolver, FNiagaraStackGraphUtilities::ENiagaraGetStackFunctionInputPinsOptions::ModuleInputsOnly);
 
 	const UEdGraphSchema_Niagara* NiagaraSchema = GetDefault<UEdGraphSchema_Niagara>();
 
@@ -311,7 +310,7 @@ void UNiagaraStackFunctionInputCollectionBase::AppendInputsForFunctionCall(FFunc
 	TArray<FName> DuplicateInputNames;
 	TArray<FName> ValidAliasedInputNames;
 	TMap<FName, UEdGraphPin*> StaticSwitchInputs;
-	TArray<const UEdGraphPin*> PinsWithInvalidTypes;
+	TArray<FNiagaraVariable> InputsWithInvalidTypes;
 
 	UNiagaraGraph* InputFunctionGraph = InputFunctionCallNode->GetCalledGraph();
 	
@@ -354,23 +353,22 @@ void UNiagaraStackFunctionInputCollectionBase::AppendInputsForFunctionCall(FFunc
 
 	
 	// Gather input data
-	for (const UEdGraphPin* InputPin : InputPins)
+	for (const FNiagaraVariable& InputVariable : InputVariables)
 	{
-		if (ProcessedInputNames.Contains(InputPin->PinName))
+		if (ProcessedInputNames.Contains(InputVariable.GetName()))
 		{
-			DuplicateInputNames.AddUnique(InputPin->PinName);
+			DuplicateInputNames.AddUnique(InputVariable.GetName());
 			continue;
 		}
-		ProcessedInputNames.Add(InputPin->PinName);
+		ProcessedInputNames.Add(InputVariable.GetName());
 
-		FNiagaraVariable InputVariable = NiagaraSchema->PinToNiagaraVariable(InputPin);
 		if (InputVariable.GetType().IsValid() == false)
 		{
-			PinsWithInvalidTypes.Add(InputPin);
+			InputsWithInvalidTypes.Add(InputVariable);
 			continue;
 		}
 		ValidAliasedInputNames.Add(
-			FNiagaraParameterHandle::CreateAliasedModuleParameterHandle(FNiagaraParameterHandle(InputPin->PinName), InputFunctionCallNode).GetParameterHandleString());
+			FNiagaraParameterHandle::CreateAliasedModuleParameterHandle(FNiagaraParameterHandle(InputVariable.GetName()), InputFunctionCallNode).GetParameterHandleString());
 
 		TOptional<FNiagaraVariableMetaData> InputMetaData;
 		if (InputFunctionGraph != nullptr)
@@ -385,19 +383,15 @@ void UNiagaraStackFunctionInputCollectionBase::AppendInputsForFunctionCall(FFunc
 		int32 EditorSortPriority = InputMetaData.IsSet() ? InputMetaData->EditorSortPriority : 0;
 		TOptional<FText> DisplayName;
 
-		bool bShouldShowInSummary = GetSummaryViewState(InputFunctionCallNode, InputVariable, InputMetaData, InputCategory, DisplayName, EditorSortPriority);		
-		if (bShouldShowInSummary)
-		{
-			SummaryViewPins.Add(InputPin);
-		}		
-		bool bIsInputHidden = HiddenPins.Contains(InputPin);
-		FInputData InputData = { InputPin, InputVariable.GetType(), EditorSortPriority, DisplayName, InputCategory, false, bIsInputHidden, bShouldShowInSummary, ModuleNode, InputFunctionCallNode };
+		bool bShouldShowInSummary = GetSummaryViewState(InputFunctionCallNode, InputVariable, InputMetaData, InputCategory, DisplayName, EditorSortPriority);			
+		bool bIsInputHidden = HiddenVariables.Contains(InputVariable);
+		FInputData InputData = { InputVariable, EditorSortPriority, DisplayName, InputCategory, false, bIsInputHidden, bShouldShowInSummary, ModuleNode, InputFunctionCallNode };
 		int32 Index = State.InputDataCollection.Add(InputData);
 
 		if (InputVariable.GetType().IsStatic())
 		{
 			FNiagaraParentData& ParentData = State.ParentMapping.FindOrAdd(InputVariable.GetName());
-			ParentData.ParentPin = InputPin;
+			ParentData.ParentVariable = InputVariable;
 		}
 
 		// set up the data for the parent-child mapping
@@ -437,7 +431,7 @@ void UNiagaraStackFunctionInputCollectionBase::AppendInputsForFunctionCall(FFunc
 		FNiagaraVariable InputVariable = NiagaraSchema->PinToNiagaraVariable(InputPin);
 		if (InputVariable.GetType().IsValid() == false)
 		{
-			PinsWithInvalidTypes.Add(InputPin);
+			InputsWithInvalidTypes.Add(InputVariable);
 			continue;
 		}
 
@@ -458,19 +452,15 @@ void UNiagaraStackFunctionInputCollectionBase::AppendInputsForFunctionCall(FFunc
 		TOptional<FText> DisplayName;
 						
 		bool bShouldShowInSummary = GetSummaryViewState(InputFunctionCallNode, InputVariable, InputMetaData, InputCategory, DisplayName, EditorSortPriority);		
-		if (bShouldShowInSummary)
-		{
-			SummaryViewPins.Add(InputPin);
-		}		
 		bool bIsInputHidden = HiddenSwitchPins.Contains(InputPin);
-		FInputData InputData = { InputPin, InputVariable.GetType(), EditorSortPriority, DisplayName, InputCategory, true, bIsInputHidden, bShouldShowInSummary, ModuleNode, InputFunctionCallNode };
+		FInputData InputData = { InputVariable, EditorSortPriority, DisplayName, InputCategory, true, bIsInputHidden, bShouldShowInSummary, ModuleNode, InputFunctionCallNode };
 		int32 Index = State.InputDataCollection.Add(InputData);
 
 		// set up the data for the parent-child mapping
 		if (InputMetaData)
 		{
 			FNiagaraParentData& ParentData = State.ParentMapping.FindOrAdd(SwitchPinName);
-			ParentData.ParentPin = InputPin;
+			ParentData.ParentVariable = InputVariable;
 			if (!InputMetaData->ParentAttribute.IsNone())
 			{
 				if (InputMetaData->ParentAttribute.ToString().StartsWith(PARAM_MAP_MODULE_STR))
@@ -486,7 +476,7 @@ void UNiagaraStackFunctionInputCollectionBase::AppendInputsForFunctionCall(FFunc
 		}
 	}
 
-	RefreshIssues(InputFunctionCallNode, DuplicateInputNames, ValidAliasedInputNames, PinsWithInvalidTypes, StaticSwitchInputs, NewIssues);
+	RefreshIssues(InputFunctionCallNode, DuplicateInputNames, ValidAliasedInputNames, InputsWithInvalidTypes, StaticSwitchInputs, NewIssues);
 }
 
 void UNiagaraStackFunctionInputCollectionBase::ApplyAllFunctionInputsToChildren(FFunctionCallNodesState& State, const TArray<UNiagaraStackEntry*>& CurrentChildren, TArray<UNiagaraStackEntry*>& NewChildren, TArray<FStackIssue>& NewIssues, bool bShouldApplySummaryFilter)
@@ -498,17 +488,17 @@ void UNiagaraStackFunctionInputCollectionBase::ApplyAllFunctionInputsToChildren(
 		if (Data.ChildIndices.Num() == 0) { continue; }
 		for (FInputData& InputData : State.InputDataCollection)
 		{
-			if (InputData.Pin != Data.ParentPin) { continue; }
+			if (InputData.InputVariable != Data.ParentVariable) { continue; }
 			if (InputData.bIsChild)
 			{
-				AddInvalidChildStackIssue(InputData.Pin->PinName, NewIssues);
+				AddInvalidChildStackIssue(InputData.InputVariable.GetName(), NewIssues);
 				continue;
 			}
 			for (int32 ChildIndex : Data.ChildIndices)
 			{
 				if (State.InputDataCollection[ChildIndex].Children.Num() > 0)
 				{
-					AddInvalidChildStackIssue(State.InputDataCollection[ChildIndex].Pin->PinName, NewIssues);
+					AddInvalidChildStackIssue(State.InputDataCollection[ChildIndex].InputVariable.GetName(), NewIssues);
 					continue;
 				}
 				State.InputDataCollection[ChildIndex].bIsChild = true;
@@ -558,7 +548,7 @@ void UNiagaraStackFunctionInputCollectionBase::ApplyAllFunctionInputsToChildren(
 		{
 			return A.SortKey < B.SortKey;
 		}
-		return A.Pin->PinName.LexicalLess(B.Pin->PinName);
+		return A.InputVariable.GetName().LexicalLess(B.InputVariable.GetName());
 	};
 
 	// Sort child and parent data separately
@@ -593,7 +583,7 @@ void UNiagaraStackFunctionInputCollectionBase::ApplyAllFunctionInputsToChildren(
 
 
 void UNiagaraStackFunctionInputCollectionBase::RefreshIssues(UNiagaraNodeFunctionCall* InputFunctionCallNode, const TArray<FName>& DuplicateInputNames, 
-	const TArray<FName>& ValidAliasedInputNames, const TArray<const UEdGraphPin*>& PinsWithInvalidTypes, const TMap<FName, UEdGraphPin*>& StaticSwitchInputs, TArray<FStackIssue>& NewIssues)
+	const TArray<FName>& ValidAliasedInputNames, const TArray<FNiagaraVariable>& InputsWithInvalidTypes, const TMap<FName, UEdGraphPin*>& StaticSwitchInputs, TArray<FStackIssue>& NewIssues)
 {
 	if (!GetIsEnabled())
 	{
@@ -698,13 +688,13 @@ void UNiagaraStackFunctionInputCollectionBase::RefreshIssues(UNiagaraNodeFunctio
 	}
 
 	// Generate issues for invalid types.
-	for (const UEdGraphPin* PinWithInvalidType : PinsWithInvalidTypes)
+	for (const FNiagaraVariable& InputWithInvalidType : InputsWithInvalidTypes)
 	{
 		FStackIssue InputWithInvalidTypeError(
 			EStackIssueSeverity::Error,
-			FText::Format(LOCTEXT("InputWithInvalidTypeSummaryFormat", "Input has an invalid type: {0}"), FText::FromName(PinWithInvalidType->PinName)),
+			FText::Format(LOCTEXT("InputWithInvalidTypeSummaryFormat", "Input has an invalid type: {0}"), FText::FromName(InputWithInvalidType.GetName())),
 			FText::Format(LOCTEXT("InputWithInvalidTypeFormat", "The input {0} on function {1} has a type which is invalid.\nThe type of this input doesn't exist anymore.\nThe type must be brought back into the project or this input must be removed from the script."),
-				FText::FromName(PinWithInvalidType->PinName), GetUserFriendlyFunctionName(InputFunctionCallNode)),
+				FText::FromName(InputWithInvalidType.GetName()), GetUserFriendlyFunctionName(InputFunctionCallNode)),
 			GetStackEditorDataKey(),
 			false);
 		NewIssues.Add(InputWithInvalidTypeError);
@@ -790,7 +780,7 @@ void UNiagaraStackFunctionInputCollectionBase::AddInputToCategory(const FInputDa
 		}
 		NewChildren.Add(InputCategory);
 	}
-	InputCategory->AddInput(InputData.ModuleNode, InputData.InputFunctionCallNode, InputData.Pin->PinName, InputData.Type, InputData.bIsStatic ? EStackParameterBehavior::Static : EStackParameterBehavior::Dynamic, InputData.DisplayName, InputData.bIsHidden, InputData.bIsChild);
+	InputCategory->AddInput(InputData.ModuleNode, InputData.InputFunctionCallNode, InputData.InputVariable.GetName(), InputData.InputVariable.GetType(), InputData.bIsStatic ? EStackParameterBehavior::Static : EStackParameterBehavior::Dynamic, InputData.DisplayName, InputData.bIsHidden, InputData.bIsChild);
 }
 
 bool UNiagaraStackFunctionInputCollectionBase::FilterByActiveSection(const UNiagaraStackEntry& Child) const
diff --git a/Engine/Plugins/FX/Niagara/Source/NiagaraEditor/Private/ViewModels/Stack/NiagaraStackGraphUtilities.cpp b/Engine/Plugins/FX/Niagara/Source/NiagaraEditor/Private/ViewModels/Stack/NiagaraStackGraphUtilities.cpp
--- a/Engine/Plugins/FX/Niagara/Source/NiagaraEditor/Private/ViewModels/Stack/NiagaraStackGraphUtilities.cpp
+++ b/Engine/Plugins/FX/Niagara/Source/NiagaraEditor/Private/ViewModels/Stack/NiagaraStackGraphUtilities.cpp
@@ -947,26 +947,26 @@ void ExtractInputPinsFromHistory(FNiagaraParameterMapHistory& History, UEdGraph*
 	}
 }
 
-void FNiagaraStackGraphUtilities::GetStackFunctionInputPins(const UNiagaraNodeFunctionCall& FunctionCallNode, TArray<const UEdGraphPin*>& OutInputPins, TSet<const UEdGraphPin*>& OutHiddenPins, FCompileConstantResolver ConstantResolver, ENiagaraGetStackFunctionInputPinsOptions Options, bool bIgnoreDisabled)
+void FNiagaraStackGraphUtilities::GetStackFunctionInputs(const UNiagaraNodeFunctionCall& FunctionCallNode, TArray<FNiagaraVariable>& OutInputVariables, TSet<FNiagaraVariable>& OutHiddenVariables, FCompileConstantResolver ConstantResolver, ENiagaraGetStackFunctionInputPinsOptions Options, bool bIgnoreDisabled)
 {
 	TArray<FNiagaraVariable> StaticVars;
 	FNiagaraStackGraphUtilitiesImpl::GetFunctionStaticVariables(FunctionCallNode, StaticVars);
 
-	FNiagaraEditorModule::Get().GetGraphDataCache().GetStackFunctionInputPins(FunctionCallNode, StaticVars, OutInputPins, OutHiddenPins, ConstantResolver, Options, bIgnoreDisabled);
+	FNiagaraEditorModule::Get().GetGraphDataCache().GetStackFunctionInputs(FunctionCallNode, StaticVars, OutInputVariables, OutHiddenVariables, ConstantResolver, Options, bIgnoreDisabled);
 }
 
-void FNiagaraStackGraphUtilities::GetStackFunctionInputPins(const UNiagaraNodeFunctionCall& FunctionCallNode, TArray<const UEdGraphPin*>& OutInputPins, FCompileConstantResolver ConstantResolver, ENiagaraGetStackFunctionInputPinsOptions Options, bool bIgnoreDisabled)
+void FNiagaraStackGraphUtilities::GetStackFunctionInputs(const UNiagaraNodeFunctionCall& FunctionCallNode, TArray<FNiagaraVariable>& OutInputVariables, FCompileConstantResolver ConstantResolver, ENiagaraGetStackFunctionInputPinsOptions Options, bool bIgnoreDisabled)
 {
 	TArray<FNiagaraVariable> StaticVars;
 	FNiagaraStackGraphUtilitiesImpl::GetFunctionStaticVariables(FunctionCallNode, StaticVars);
 
-	FNiagaraEditorModule::Get().GetGraphDataCache().GetStackFunctionInputPins(FunctionCallNode, StaticVars, OutInputPins, ConstantResolver, Options, bIgnoreDisabled);
+	FNiagaraEditorModule::Get().GetGraphDataCache().GetStackFunctionInputs(FunctionCallNode, StaticVars, OutInputVariables, ConstantResolver, Options, bIgnoreDisabled);
 }
 
-void FNiagaraStackGraphUtilities::GetStackFunctionInputPins(const UNiagaraNodeFunctionCall& FunctionCallNode, TArray<const UEdGraphPin*>& OutInputPins, ENiagaraGetStackFunctionInputPinsOptions Options /*= ENiagaraGetStackFunctionInputPinsOptions::AllInputs*/, bool bIgnoreDisabled /*= false*/)
+void FNiagaraStackGraphUtilities::GetStackFunctionInputs(const UNiagaraNodeFunctionCall& FunctionCallNode, TArray<FNiagaraVariable>& OutInputVariables, ENiagaraGetStackFunctionInputPinsOptions Options /*= ENiagaraGetStackFunctionInputPinsOptions::AllInputs*/, bool bIgnoreDisabled /*= false*/)
 {
 	FCompileConstantResolver EmptyResolver;
-	GetStackFunctionInputPins(FunctionCallNode, OutInputPins, EmptyResolver, Options, bIgnoreDisabled);
+	GetStackFunctionInputs(FunctionCallNode, OutInputVariables, EmptyResolver, Options, bIgnoreDisabled);
 }
 
 void FNiagaraStackGraphUtilities::GetStackFunctionInputPinsWithoutCache(
@@ -1094,13 +1094,12 @@ TOptional<FNiagaraStackGraphUtilities::FMatchingFunctionInputData> FNiagaraStack
 	TOptional<FNiagaraVariableMetaData> MetaData;
 				
 	// first, we check the normal module inputs
-	TArray<const UEdGraphPin*> InputPins;
+	TArray<FNiagaraVariable> InputVariables;
 	FNiagaraStackGraphUtilities::ENiagaraGetStackFunctionInputPinsOptions Options = FNiagaraStackGraphUtilities::ENiagaraGetStackFunctionInputPinsOptions::AllInputs;
-	FNiagaraStackGraphUtilities::GetStackFunctionInputPins(FunctionCallNode, InputPins, ConstantResolver, Options, true);
+	FNiagaraStackGraphUtilities::GetStackFunctionInputs(FunctionCallNode, InputVariables, ConstantResolver, Options, true);
 				
-	for(const UEdGraphPin* InputPin : InputPins)
+	for(const FNiagaraVariable& InputVariable : InputVariables)
 	{
-		FNiagaraVariable InputVariable = UEdGraphSchema_Niagara::PinToNiagaraVariable(InputPin);
 		if (InputVariable.GetType().IsValid() == false)
 		{
 			continue;
@@ -1109,7 +1108,7 @@ TOptional<FNiagaraStackGraphUtilities::FMatchingFunctionInputData> FNiagaraStack
 		MetaData = FunctionCallNode.GetCalledGraph()->GetMetaData(InputVariable);
 		if(MetaData.IsSet() && MetaData.GetValue().GetVariableGuid() == InputGuid)
 		{
-			Result.InputName = InputPin->PinName;
+			Result.InputName = InputVariable.GetName();
 			Result.Type = InputVariable.GetType();
 			Result.MetaData = MetaData.GetValue();
 
diff --git a/Engine/Plugins/FX/Niagara/Source/NiagaraEditor/Private/ViewModels/Stack/NiagaraStackScriptItemGroup.cpp b/Engine/Plugins/FX/Niagara/Source/NiagaraEditor/Private/ViewModels/Stack/NiagaraStackScriptItemGroup.cpp
--- a/Engine/Plugins/FX/Niagara/Source/NiagaraEditor/Private/ViewModels/Stack/NiagaraStackScriptItemGroup.cpp
+++ b/Engine/Plugins/FX/Niagara/Source/NiagaraEditor/Private/ViewModels/Stack/NiagaraStackScriptItemGroup.cpp
@@ -359,11 +359,11 @@ private:
 
 		UNiagaraNodeAssignment* NewAssignmentModule = FNiagaraStackGraphUtilities::AddParameterModuleToStack(Vars, *OutputNode, TargetIndex,DefaultVals );
 		
-		TArray<const UEdGraphPin*> InputPins;
-		FNiagaraStackGraphUtilities::GetStackFunctionInputPins(*NewAssignmentModule, InputPins, FNiagaraStackGraphUtilities::ENiagaraGetStackFunctionInputPinsOptions::AllInputs);
-		if (InputPins.Num() == 1)
+		TArray<FNiagaraVariable> InputVariables;
+		FNiagaraStackGraphUtilities::GetStackFunctionInputs(*NewAssignmentModule, InputVariables, FNiagaraStackGraphUtilities::ENiagaraGetStackFunctionInputPinsOptions::AllInputs);
+		if (InputVariables.Num() == 1)
 		{
-			FString FunctionInputEditorDataKey = FNiagaraStackGraphUtilities::GenerateStackFunctionInputEditorDataKey(*NewAssignmentModule, InputPins[0]->PinName);
+			FString FunctionInputEditorDataKey = FNiagaraStackGraphUtilities::GenerateStackFunctionInputEditorDataKey(*NewAssignmentModule, InputVariables[0].GetName());
 			if (bRenameParameterOnAdd)
 			{
 				StackEditorData.SetStackEntryIsRenamePending(FunctionInputEditorDataKey, true);
diff --git a/Engine/Plugins/FX/Niagara/Source/NiagaraEditor/Public/NiagaraGraphDataCache.h b/Engine/Plugins/FX/Niagara/Source/NiagaraEditor/Public/NiagaraGraphDataCache.h
--- a/Engine/Plugins/FX/Niagara/Source/NiagaraEditor/Public/NiagaraGraphDataCache.h
+++ b/Engine/Plugins/FX/Niagara/Source/NiagaraEditor/Public/NiagaraGraphDataCache.h
@@ -19,15 +19,15 @@ class FNiagaraGraphDataCache
 public:
 	FNiagaraGraphDataCache();
 
-	void GetStackFunctionInputPins(const UNiagaraNodeFunctionCall& FunctionCallNode, TConstArrayView<FNiagaraVariable> StaticVars, TArray<const UEdGraphPin*>& OutInputPins, FCompileConstantResolver ConstantResolver, FNiagaraStackGraphUtilities::ENiagaraGetStackFunctionInputPinsOptions Options, bool bIgnoreDisabled);
-	void GetStackFunctionInputPins(const UNiagaraNodeFunctionCall& FunctionCallNode, TConstArrayView<FNiagaraVariable> StaticVars, TArray<const UEdGraphPin*>& OutInputPins, TSet<const UEdGraphPin*>& OutHiddenPins, FCompileConstantResolver ConstantResolver, FNiagaraStackGraphUtilities::ENiagaraGetStackFunctionInputPinsOptions Options, bool bIgnoreDisabled);
+	void GetStackFunctionInputs(const UNiagaraNodeFunctionCall& FunctionCallNode, TConstArrayView<FNiagaraVariable> StaticVars, TArray<FNiagaraVariable>& OutInputVariables, FCompileConstantResolver ConstantResolver, FNiagaraStackGraphUtilities::ENiagaraGetStackFunctionInputPinsOptions Options, bool bIgnoreDisabled);
+	void GetStackFunctionInputs(const UNiagaraNodeFunctionCall& FunctionCallNode, TConstArrayView<FNiagaraVariable> StaticVars, TArray<FNiagaraVariable>& OutInputVariables, TSet<FNiagaraVariable>& OutHiddenVariables, FCompileConstantResolver ConstantResolver, FNiagaraStackGraphUtilities::ENiagaraGetStackFunctionInputPinsOptions Options, bool bIgnoreDisabled);
 
 private:
-	void GetStackFunctionInputPinsInternal(
+	void GetStackFunctionInputsInternal(
 		const UNiagaraNodeFunctionCall& FunctionCallNode,
 		const UNiagaraGraph* CalledGraph,
 		TConstArrayView<FNiagaraVariable> StaticVars,
-		TArray<const UEdGraphPin*>& OutInputPins,
+		TArray<FNiagaraVariable>& OutInputVariables,
 		FCompileConstantResolver ConstantResolver,
 		FNiagaraStackGraphUtilities::ENiagaraGetStackFunctionInputPinsOptions Options,
 		bool bIgnoreDisabled,
@@ -62,7 +62,7 @@ private:
 
 	struct FStackFunctionInputPinValue
 	{
-		TArray<const UEdGraphPin*> InputPins;
+		TArray<FNiagaraVariable> InputVariables;
 	};
 
 	// cache holding a finite number (specific value pulled from a cvar) of results to previous calls
diff --git a/Engine/Plugins/FX/Niagara/Source/NiagaraEditor/Public/NiagaraScriptSource.h b/Engine/Plugins/FX/Niagara/Source/NiagaraEditor/Public/NiagaraScriptSource.h
--- a/Engine/Plugins/FX/Niagara/Source/NiagaraEditor/Public/NiagaraScriptSource.h
+++ b/Engine/Plugins/FX/Niagara/Source/NiagaraEditor/Public/NiagaraScriptSource.h
@@ -54,8 +54,8 @@ class UNiagaraScriptSource : public UNiagaraScriptSourceBase
 
 	NIAGARAEDITOR_API virtual bool AddModuleIfMissing(FString ModulePath, ENiagaraScriptUsage Usage, bool& bOutFoundModule)override;
 
-	virtual void FixupRenamedParameters(UNiagaraNodeFunctionCall* FunctionCallNode, TConstArrayView<const UEdGraphPin*> FunctionPins, TConstArrayView<FNiagaraVariable> PinVariables, FNiagaraParameterStore& RapidIterationParameters, const TArray<FNiagaraVariable>& OldRapidIterationVariables, const FVersionedNiagaraEmitter& Emitter, ENiagaraScriptUsage ScriptUsage) const;
-	virtual void InitializeNewParameters(UNiagaraNodeFunctionCall* FunctionCallNode, TConstArrayView<const UEdGraphPin*> FunctionPins, TConstArrayView<FNiagaraVariable> PinVariables, FNiagaraParameterStore& RapidIterationParameters, const FVersionedNiagaraEmitter& VersionedEmitter, ENiagaraScriptUsage ScriptUsage, TSet<FName>& ValidRapidIterationParameterNames) const;
+	virtual void FixupRenamedParameters(UNiagaraNodeFunctionCall* FunctionCallNode, TConstArrayView<FNiagaraVariable> ModuleInputVariables, FNiagaraParameterStore& RapidIterationParameters, const TArray<FNiagaraVariable>& OldRapidIterationVariables, const FVersionedNiagaraEmitter& Emitter, ENiagaraScriptUsage ScriptUsage) const;
+	virtual void InitializeNewParameters(UNiagaraNodeFunctionCall* FunctionCallNode, TConstArrayView<FNiagaraVariable> ModuleInputVariables, FNiagaraParameterStore& RapidIterationParameters, const FVersionedNiagaraEmitter& VersionedEmitter, ENiagaraScriptUsage ScriptUsage, TSet<FName>& ValidRapidIterationParameterNames) const;
 
 	virtual void CleanUpOldAndInitializeNewRapidIterationParameters(const FVersionedNiagaraEmitter& Emitter, ENiagaraScriptUsage ScriptUsage, FGuid ScriptUsageId, FNiagaraParameterStore& RapidIterationParameters) const override;
 	virtual void ForceGraphToRecompileOnNextCheck() override;
diff --git a/Engine/Plugins/FX/Niagara/Source/NiagaraEditor/Public/ViewModels/Stack/NiagaraStackFunctionInputCollection.h b/Engine/Plugins/FX/Niagara/Source/NiagaraEditor/Public/ViewModels/Stack/NiagaraStackFunctionInputCollection.h
--- a/Engine/Plugins/FX/Niagara/Source/NiagaraEditor/Public/ViewModels/Stack/NiagaraStackFunctionInputCollection.h
+++ b/Engine/Plugins/FX/Niagara/Source/NiagaraEditor/Public/ViewModels/Stack/NiagaraStackFunctionInputCollection.h
@@ -43,8 +43,7 @@ protected:
 
 	struct FInputData
 	{
-		const UEdGraphPin* Pin;
-		FNiagaraTypeDefinition Type;
+		FNiagaraVariable InputVariable;
 		int32 SortKey;
 		TOptional<FText> DisplayName;
 		FText Category;
@@ -57,12 +56,11 @@ protected:
 		
 		TArray<FInputData*> Children;
 		bool bIsChild = false;
-
 	};
 
 	struct FNiagaraParentData
 	{
-		const UEdGraphPin* ParentPin;
+		FNiagaraVariable ParentVariable;
 		TArray<int32> ChildIndices;
 	};
 
@@ -83,7 +81,7 @@ protected:
 	
 	void ApplyAllFunctionInputsToChildren(FFunctionCallNodesState& State, const TArray<UNiagaraStackEntry*>& CurrentChildren, TArray<UNiagaraStackEntry*>& NewChildren, TArray<FStackIssue>& NewIssues, bool bShouldApplySummaryFilter);
 
-	void RefreshIssues(UNiagaraNodeFunctionCall* InputFunctionCallNode, const TArray<FName>& DuplicateInputNames, const TArray<FName>& ValidAliasedInputNames, const TArray<const UEdGraphPin*>& PinsWithInvalidTypes, const TMap<FName, UEdGraphPin*>& StaticSwitchInputs, TArray<FStackIssue>& NewIssues);
+	void RefreshIssues(UNiagaraNodeFunctionCall* InputFunctionCallNode, const TArray<FName>& DuplicateInputNames, const TArray<FName>& ValidAliasedInputNames, const TArray<FNiagaraVariable>& InputsWithInvalidTypes, const TMap<FName, UEdGraphPin*>& StaticSwitchInputs, TArray<FStackIssue>& NewIssues);
 
 	void OnFunctionInputsChanged();
 
diff --git a/Engine/Plugins/FX/Niagara/Source/NiagaraEditor/Public/ViewModels/Stack/NiagaraStackGraphUtilities.h b/Engine/Plugins/FX/Niagara/Source/NiagaraEditor/Public/ViewModels/Stack/NiagaraStackGraphUtilities.h
--- a/Engine/Plugins/FX/Niagara/Source/NiagaraEditor/Public/ViewModels/Stack/NiagaraStackGraphUtilities.h
+++ b/Engine/Plugins/FX/Niagara/Source/NiagaraEditor/Public/ViewModels/Stack/NiagaraStackGraphUtilities.h
@@ -96,11 +96,11 @@ namespace FNiagaraStackGraphUtilities
 	};
 
 	/* Returns the input pins for the given function call node. Try not to use this method if possible and use the version that accepts a FCompileConstantResolver parameter. This method should only be used if the current context is completely outside of any system or emitter (e.g. inside the graph editor itself) and constants cannot be resolved at all. */
-	void GetStackFunctionInputPins(const UNiagaraNodeFunctionCall& FunctionCallNode, TArray<const UEdGraphPin*>& OutInputPins, ENiagaraGetStackFunctionInputPinsOptions Options, bool bIgnoreDisabled = false);
+	void GetStackFunctionInputs(const UNiagaraNodeFunctionCall& FunctionCallNode, TArray<FNiagaraVariable>& OutInputVariables, ENiagaraGetStackFunctionInputPinsOptions Options, bool bIgnoreDisabled = false);
 
 	/* Returns the input pins for the given function call node. */
-	void GetStackFunctionInputPins(const UNiagaraNodeFunctionCall& FunctionCallNode, TArray<const UEdGraphPin*>& OutInputPins, TSet<const UEdGraphPin*>& OutHiddenPins, FCompileConstantResolver ConstantResolver, ENiagaraGetStackFunctionInputPinsOptions Options, bool bIgnoreDisabled = false);
-	void GetStackFunctionInputPins(const UNiagaraNodeFunctionCall& FunctionCallNode, TArray<const UEdGraphPin*>& OutInputPins, FCompileConstantResolver ConstantResolver, ENiagaraGetStackFunctionInputPinsOptions Options, bool bIgnoreDisabled = false);
+	void GetStackFunctionInputs(const UNiagaraNodeFunctionCall& FunctionCallNode, TArray<FNiagaraVariable>& OutInputVariables, TSet<FNiagaraVariable>& OutHiddenVariables, FCompileConstantResolver ConstantResolver, ENiagaraGetStackFunctionInputPinsOptions Options, bool bIgnoreDisabled = false);
+	void GetStackFunctionInputs(const UNiagaraNodeFunctionCall& FunctionCallNode, TArray<FNiagaraVariable>& OutInputVariables, FCompileConstantResolver ConstantResolver, ENiagaraGetStackFunctionInputPinsOptions Options, bool bIgnoreDisabled = false);
 
 	/* Returns the input pins for the given function call node.  Bypasses the module level caching of this data and generates it each call. */
 	void GetStackFunctionInputPinsWithoutCache(
