- Deprecate Niagara variable string set methods in favor of using FNames to avoid unnecessary conversions

diff --git a/Engine/Plugins/FX/Niagara/Source/Niagara/Private/NiagaraComponent.cpp b/Engine/Plugins/FX/Niagara/Source/Niagara/Private/NiagaraComponent.cpp
--- a/Engine/Plugins/FX/Niagara/Source/Niagara/Private/NiagaraComponent.cpp
+++ b/Engine/Plugins/FX/Niagara/Source/Niagara/Private/NiagaraComponent.cpp
@@ -2506,7 +2506,7 @@ void UNiagaraComponent::SetVariableActor(FName InVariableName, AActor* InValue)
 
 void UNiagaraComponent::SetNiagaraVariableActor(const FString& InVariableName, AActor* InValue)
 {
-	SetNiagaraVariableObject(InVariableName, InValue);
+	SetVariableObject(FName(*InVariableName), InValue);
 }
 
 void UNiagaraComponent::SetVariableObject(FName InVariableName, UObject* InValue)
diff --git a/Engine/Plugins/FX/Niagara/Source/Niagara/Public/NiagaraComponent.h b/Engine/Plugins/FX/Niagara/Source/Niagara/Public/NiagaraComponent.h
--- a/Engine/Plugins/FX/Niagara/Source/Niagara/Public/NiagaraComponent.h
+++ b/Engine/Plugins/FX/Niagara/Source/Niagara/Public/NiagaraComponent.h
@@ -421,6 +421,7 @@ public:
 	int32 GetRandomSeedOffset() const { return RandomSeedOffset; }
 
 	/** Sets a Niagara FLinearColor parameter by name, overriding locally if necessary.*/
+	UE_DEPRECATED(5.3, "This method will be removed in a future release.  Please update to use the FName variant")
 	UFUNCTION(BlueprintCallable, Category = Niagara, meta = (DisplayName = "Set Niagara Variable By String (LinearColor)", Keywords="user parameter variable color"))
 	void SetNiagaraVariableLinearColor(const FString& InVariableName, const FLinearColor& InValue);
 
@@ -429,6 +430,7 @@ public:
 	void SetVariableLinearColor(FName InVariableName, const FLinearColor& InValue);
 
 	/** Sets a Niagara Vector4 parameter by name, overriding locally if necessary.*/
+	UE_DEPRECATED(5.3, "This method will be removed in a future release.  Please update to use the FName variant")
 	UFUNCTION(BlueprintCallable, Category = Niagara, meta = (DisplayName = "Set Niagara Variable By String (Vector4)", Keywords="user parameter variable vector"))
 	void SetNiagaraVariableVec4(const FString& InVariableName, const FVector4& InValue);
 
@@ -437,6 +439,7 @@ public:
 	void SetVariableVec4(FName InVariableName, const FVector4& InValue);
 
 	/** Sets a Niagara quaternion parameter by name, overriding locally if necessary.*/
+	UE_DEPRECATED(5.3, "This method will be removed in a future release.  Please update to use the FName variant")
 	UFUNCTION(BlueprintCallable, Category = Niagara, meta = (DisplayName = "Set Niagara Variable By String (Quaternion)", Keywords="user parameter variable quaternion rotation"))
 	void SetNiagaraVariableQuat(const FString& InVariableName, const FQuat& InValue);
 
@@ -445,6 +448,7 @@ public:
 	void SetVariableQuat(FName InVariableName, const FQuat& InValue);
 
 	/** Sets a Niagara matrix parameter by name, overriding locally if necessary.*/
+	UE_DEPRECATED(5.3, "This method will be removed in a future release.  Please update to use the FName variant")
 	UFUNCTION(BlueprintCallable, Category = Niagara, meta = (DisplayName = "Set Niagara Variable By String (Matrix)", Keywords="user parameter variable matrix"))
 	void SetNiagaraVariableMatrix(const FString& InVariableName, const FMatrix& InValue);
 
@@ -453,6 +457,7 @@ public:
 	void SetVariableMatrix(FName InVariableName, const FMatrix& InValue);
 
 	/** Sets a Niagara Vector3 parameter by name, overriding locally if necessary.*/
+	UE_DEPRECATED(5.3, "This method will be removed in a future release.  Please update to use the FName variant")
 	UFUNCTION(BlueprintCallable, Category = Niagara, meta = (DisplayName = "Set Niagara Variable By String (Vector3)", Keywords="user parameter variable vector"))
 	void SetNiagaraVariableVec3(const FString& InVariableName, FVector InValue);
 
@@ -461,6 +466,7 @@ public:
 	void SetVariableVec3(FName InVariableName, FVector InValue);
 
 	/** Sets a Niagara Position parameter by name, overriding locally if necessary.*/
+	UE_DEPRECATED(5.3, "This method will be removed in a future release.  Please update to use the FName variant")
 	UFUNCTION(BlueprintCallable, Category = Niagara, meta = (DisplayName = "Set Niagara Variable By String (Position)", Keywords="user parameter variable vector position lwc"))
 	void SetNiagaraVariablePosition(const FString& InVariableName, FVector InValue);
 
@@ -469,6 +475,7 @@ public:
 	void SetVariablePosition(FName InVariableName, FVector InValue);
 
 	/** Sets a Niagara Vector2 parameter by name, overriding locally if necessary.*/
+	UE_DEPRECATED(5.3, "This method will be removed in a future release.  Please update to use the FName variant")
 	UFUNCTION(BlueprintCallable, Category = Niagara, meta = (DisplayName = "Set Niagara Variable By String (Vector2)", Keywords="user parameter variable vector"))
 	void SetNiagaraVariableVec2(const FString& InVariableName, FVector2D InValue);
 
@@ -477,6 +484,7 @@ public:
 	void SetVariableVec2(FName InVariableName, FVector2D InValue);
 
 	/** Sets a Niagara float parameter by name, overriding locally if necessary.*/
+	UE_DEPRECATED(5.3, "This method will be removed in a future release.  Please update to use the FName variant")
 	UFUNCTION(BlueprintCallable, Category = Niagara, meta = (DisplayName = "Set Niagara Variable By String (Float)", Keywords="user parameter variable float"))
 	void SetNiagaraVariableFloat(const FString& InVariableName, float InValue);
 
@@ -485,6 +493,7 @@ public:
 	void SetVariableFloat(FName InVariableName, float InValue);
 
 	/** Sets a Niagara int parameter by name, overriding locally if necessary.*/
+	UE_DEPRECATED(5.3, "This method will be removed in a future release.  Please update to use the FName variant")
 	UFUNCTION(BlueprintCallable, Category = Niagara, meta = (DisplayName = "Set Niagara Variable By String (Int32)", Keywords="user parameter variable int"))
 	void SetNiagaraVariableInt(const FString& InVariableName, int32 InValue);
 
@@ -493,6 +502,7 @@ public:
 	void SetVariableInt(FName InVariableName, int32 InValue);
 
 	/** Sets a Niagara bool parameter by name, overriding locally if necessary.*/
+	UE_DEPRECATED(5.3, "This method will be removed in a future release.  Please update to use the FName variant")
 	UFUNCTION(BlueprintCallable, Category = Niagara, meta = (DisplayName = "Set Niagara Variable By String (Bool)", Keywords="user parameter variable bool"))
 	void SetNiagaraVariableBool(const FString& InVariableName, bool InValue);
 
@@ -500,12 +510,14 @@ public:
 	UFUNCTION(BlueprintCallable, Category = Niagara, meta = (DisplayName = "Set Niagara Variable (Bool)", Keywords="user parameter variable bool"))
 	void SetVariableBool(FName InVariableName, bool InValue);
 
+	UE_DEPRECATED(5.3, "This method will be removed in a future release.  Please update to use the FName variant")
 	UFUNCTION(BlueprintCallable, Category = Niagara, meta = (DisplayName = "Set Niagara Variable By String (Actor)", Keywords="user parameter variable actor"))
 	void SetNiagaraVariableActor(const FString& InVariableName, AActor* Actor);
 
 	UFUNCTION(BlueprintCallable, Category = Niagara, meta = (DisplayName = "Set Niagara Variable (Actor)", Keywords="user parameter variable actor"))
 	void SetVariableActor(FName InVariableName, AActor* Actor);
 
+	UE_DEPRECATED(5.3, "This method will be removed in a future release.  Please update to use the FName variant")
 	UFUNCTION(BlueprintCallable, Category = Niagara, meta = (DisplayName = "Set Niagara Variable By String (Object)", Keywords="user parameter variable object"))
 	void SetNiagaraVariableObject(const FString& InVariableName, UObject* Object);
 
