- Add specific op to select a static int from a static bool
- This allows for more complex graph operations like creating a bit mask out of static bool conditions

diff --git a/Engine/Plugins/FX/Niagara/Source/NiagaraEditor/Private/NiagaraEditorCommon.cpp b/Engine/Plugins/FX/Niagara/Source/NiagaraEditor/Private/NiagaraEditorCommon.cpp
--- a/Engine/Plugins/FX/Niagara/Source/NiagaraEditor/Private/NiagaraEditorCommon.cpp
+++ b/Engine/Plugins/FX/Niagara/Source/NiagaraEditor/Private/NiagaraEditorCommon.cpp
@@ -1470,9 +1470,34 @@ void FNiagaraOpInfo::Init()
 	Op->AddedInputTypeRestrictions.Add(IntType);
 	Op->AddedInputFormatting = TEXT("{A} >> {B}");
 	OpInfoMap.Add(Op->Name) = Idx;
+
+
+	Idx = OpInfos.AddDefaulted();
+	Op = &OpInfos[Idx];
+	Op->Category = IntCategory;
+	Op->FriendlyName = NSLOCTEXT("NiagaraOpInfo", "SelectStaticIntName", "Select Int (Static)");
+	Op->Description = NSLOCTEXT("NiagaraOpInfo", "SelectStaticIntDesc", "Result = A ? B : C");
+	Op->Inputs.Add(FNiagaraOpInOutInfo(A, FNiagaraTypeDefinition::GetBoolDef().ToStaticDef(), AText, AText, TEXT("false")));
+	Op->Inputs.Add(FNiagaraOpInOutInfo(B, IntType.ToStaticDef(), BText, BText, Default_IntZero));
+	Op->Inputs.Add(FNiagaraOpInOutInfo(C, IntType.ToStaticDef(), CText, CText, Default_IntZero));
+	Op->Outputs.Add(FNiagaraOpInOutInfo(Result, IntType.ToStaticDef(), ResultText, ResultText, Default_IntZero, TEXT("A ? B : C")));
+	Op->BuildName(TEXT("SelectStaticInt"), IntCategoryName);
+	Op->bSupportsStaticResolution = true;
+	Op->StaticVariableResolveFunction.BindLambda(
+		[=](const TArray<int32>& InPinValues)
+	{
+		if (InPinValues.Num() != 3)
+			return 0;
+
+		const int32 Return = InPinValues[0] != 0 ? InPinValues[1] : InPinValues[2];
+		return Return;
+	}
+	);
+	OpInfoMap.Add(Op->Name) = Idx;
+
+
 	//////////////////////////////////////////////////////////////////////////
 	// Boolean Only Ops
-
 	FString Default_BoolZero(TEXT("false"));
 	FString Default_BoolOne(TEXT("true"));
 	FText BoolCategory = NSLOCTEXT("NiagaraOpInfo", "BoolOpCategory", "Boolean");
