- Allow static bools to be used as edit conditions

diff --git a/Engine/Plugins/FX/Niagara/Source/NiagaraEditor/Private/Customizations/NiagaraScriptVariableCustomization.cpp b/Engine/Plugins/FX/Niagara/Source/NiagaraEditor/Private/Customizations/NiagaraScriptVariableCustomization.cpp
--- a/Engine/Plugins/FX/Niagara/Source/NiagaraEditor/Private/Customizations/NiagaraScriptVariableCustomization.cpp
+++ b/Engine/Plugins/FX/Niagara/Source/NiagaraEditor/Private/Customizations/NiagaraScriptVariableCustomization.cpp
@@ -130,7 +130,7 @@ void FNiagaraScriptVariableDetails::CustomizeDetails(IDetailLayoutBuilder& Detai
 	// Always hide the default value variant. We generate a node for this property to enable PostEditChangeProperty(...) events but do not modify it via the default generated customization.
 	DetailBuilder.HideProperty("DefaultValueVariant", UNiagaraScriptVariable::StaticClass());
 
-	if (Variable->Variable.GetType() != FNiagaraTypeDefinition::GetBoolDef())
+	if (!Variable->Variable.GetType().IsSameBaseDefinition(FNiagaraTypeDefinition::GetBoolDef()))
 	{
 		DetailBuilder.HideProperty(GET_MEMBER_NAME_CHECKED(UNiagaraScriptVariable, Metadata.bInlineEditConditionToggle));
 		DetailBuilder.HideProperty(GET_MEMBER_NAME_CHECKED(UNiagaraScriptVariable, Metadata.InlineParameterBoolOverride));
diff --git a/Engine/Plugins/FX/Niagara/Source/NiagaraEditor/Private/NiagaraStackFunctionInputBinder.cpp b/Engine/Plugins/FX/Niagara/Source/NiagaraEditor/Private/NiagaraStackFunctionInputBinder.cpp
--- a/Engine/Plugins/FX/Niagara/Source/NiagaraEditor/Private/NiagaraStackFunctionInputBinder.cpp
+++ b/Engine/Plugins/FX/Niagara/Source/NiagaraEditor/Private/NiagaraStackFunctionInputBinder.cpp
@@ -260,7 +260,7 @@ void FNiagaraStackFunctionInputBinder::SetData(const uint8* InValue, int32 InSiz
 	checkf(InSize == InputType.GetSize(), TEXT("Set value size doesn't match bound value size"));
 	if (FMemory::Memcmp(GetData().GetData(), InValue, InSize) != 0)
 	{
-		if (RapidIterationParameter.IsValid())
+		if ( RapidIterationParameter.IsValid() )
 		{
 			Script->Modify();
 			Script->RapidIterationParameters.SetParameterData(InValue, RapidIterationParameter, true);
@@ -271,7 +271,10 @@ void FNiagaraStackFunctionInputBinder::SetData(const uint8* InValue, int32 InSiz
 				DependentScript->RapidIterationParameters.SetParameterData(InValue, RapidIterationParameter, true);
 			}
 		}
-		else
+
+		// Static variables are RI parameters and will therefore not force a recompile unless we take this path
+		// Not taking this path will result in the Script StaticVariablesWritten being out of date
+		if ( !RapidIterationParameter.IsValid() || InputType.IsStatic() )
 		{
 			checkf(FunctionCallNode.IsValid(), TEXT("Bound function call is no longer valid"));
 
