- Fix for GPU ribbons not showing in instanced stereo rendering

diff --git a/Engine/Plugins/FX/Niagara/Shaders/Private/NiagaraRibbonVertexFactory.ush b/Engine/Plugins/FX/Niagara/Shaders/Private/NiagaraRibbonVertexFactory.ush
--- a/Engine/Plugins/FX/Niagara/Shaders/Private/NiagaraRibbonVertexFactory.ush
+++ b/Engine/Plugins/FX/Niagara/Shaders/Private/NiagaraRibbonVertexFactory.ush
@@ -24,8 +24,8 @@
 #define SLICEVERTEXDATA_TEXTUREV_OFFSET 4
 #define SLICEVERTEXDATA_STRIDE 5
 
-#define INDIRECTDRAWARGS_INTERPCOUNT_OFFSET 10
-#define INDIRECTDRAWARGS_ONEOVERINTERPCOUNT_OFFSET 11
+#define INDEX_GEN_NUM_SUB_SEGMENTS_OFFSET			15		// Duplicated in NiagaraRibbonCommon.ush
+#define INDEX_GEN_ONE_OVER_NUM_SUB_SEGMENTS_OFFSET	16		// Duplicated in NiagaraRibbonCommon.ush
 
 uint UnpackParticleId(uint InterpIdx)
 {
@@ -291,7 +291,7 @@ uint GetInterpCount()
 	BRANCH
 	if (NiagaraRibbonVFLooseParameters.IndirectDrawOutputOffset != -1)
 	{
-		return NiagaraRibbonVFLooseParameters.IndirectDrawOutput[NiagaraRibbonVFLooseParameters.IndirectDrawOutputOffset + INDIRECTDRAWARGS_INTERPCOUNT_OFFSET];
+		return NiagaraRibbonVFLooseParameters.IndirectDrawOutput[NiagaraRibbonVFLooseParameters.IndirectDrawOutputOffset + INDEX_GEN_NUM_SUB_SEGMENTS_OFFSET];
 	}
 	else
 	{
@@ -304,7 +304,7 @@ float GetOneOverInterpCount()
 	BRANCH
 	if (NiagaraRibbonVFLooseParameters.IndirectDrawOutputOffset != -1)
 	{
-		return asfloat(NiagaraRibbonVFLooseParameters.IndirectDrawOutput[NiagaraRibbonVFLooseParameters.IndirectDrawOutputOffset + INDIRECTDRAWARGS_ONEOVERINTERPCOUNT_OFFSET]);
+		return asfloat(NiagaraRibbonVFLooseParameters.IndirectDrawOutput[NiagaraRibbonVFLooseParameters.IndirectDrawOutputOffset + INDEX_GEN_ONE_OVER_NUM_SUB_SEGMENTS_OFFSET]);
 	}
 	else
 	{
diff --git a/Engine/Plugins/FX/Niagara/Shaders/Private/Ribbons/NiagaraRibbonCommon.ush b/Engine/Plugins/FX/Niagara/Shaders/Private/Ribbons/NiagaraRibbonCommon.ush
--- a/Engine/Plugins/FX/Niagara/Shaders/Private/Ribbons/NiagaraRibbonCommon.ush
+++ b/Engine/Plugins/FX/Niagara/Shaders/Private/Ribbons/NiagaraRibbonCommon.ush
@@ -37,14 +37,17 @@
 #define VERTEX_GEN_OUTPUT_DATA_TESSELLATION_CURRENT_FRAME_AVERAGE_WIDTH 14
 #endif
 
+// This layout must match FNiagaraRibbonIndirectDrawBufferLayout
+#define INDEX_GEN_EXECUTE_INDIRECT_OFFSET			0
+#define INDEX_GEN_DRAW_INDIRECT_OFFSET				3
+#define INDEX_GEN_STEREO_DRAW_INDIRECT_OFFSET		8
 
-#define INDEX_GEN_INDIRECT_ARGS_STRIDE 12
-#define INDEX_GEN_EXECUTE_INDIRECT_OFFSET 0
-#define INDEX_GEN_TESSELLATION_FACTOR_OFFSET 3
-#define INDEX_GEN_INDIRECT_DRAW_OFFSET 4
-#define INDEX_GEN_NUM_SEGMENTS_OFFSET 9
-#define INDEX_GEN_NUM_SUB_SEGMENTS_OFFSET 10
-#define INDEX_GEN_ONE_OVER_NUM_SUB_SEGMENTS_OFFSET 11
+#define INDEX_GEN_TESSELLATION_FACTOR_OFFSET		13
+#define INDEX_GEN_NUM_SEGMENTS_OFFSET				14
+#define INDEX_GEN_NUM_SUB_SEGMENTS_OFFSET			15		// Duplicated in NiagaraRibbonVertexFactory.ush
+#define INDEX_GEN_ONE_OVER_NUM_SUB_SEGMENTS_OFFSET	16		// Duplicated in NiagaraRibbonVertexFactory.ush
+
+#define INDEX_GEN_INDIRECT_ARGS_STRIDE				17
 
 
 // Total particle Count
diff --git a/Engine/Plugins/FX/Niagara/Shaders/Private/Ribbons/NiagaraRibbonInitializeIndices.usf b/Engine/Plugins/FX/Niagara/Shaders/Private/Ribbons/NiagaraRibbonInitializeIndices.usf
--- a/Engine/Plugins/FX/Niagara/Shaders/Private/Ribbons/NiagaraRibbonInitializeIndices.usf
+++ b/Engine/Plugins/FX/Niagara/Shaders/Private/Ribbons/NiagaraRibbonInitializeIndices.usf
@@ -43,14 +43,20 @@ void InitializeIndices(uint3 DispatchThreadId : SV_DispatchThreadID)
 		IndirectDrawOutput[IndirectDrawOffset + INDEX_GEN_EXECUTE_INDIRECT_OFFSET + 0] = 0;
 		IndirectDrawOutput[IndirectDrawOffset + INDEX_GEN_EXECUTE_INDIRECT_OFFSET + 1] = 0;
 		IndirectDrawOutput[IndirectDrawOffset + INDEX_GEN_EXECUTE_INDIRECT_OFFSET + 2] = 0;
-		IndirectDrawOutput[IndirectDrawOffset + INDEX_GEN_TESSELLATION_FACTOR_OFFSET] = 0;
 			
-		IndirectDrawOutput[IndirectDrawOffset + INDEX_GEN_INDIRECT_DRAW_OFFSET + 0] = 0;
-		IndirectDrawOutput[IndirectDrawOffset + INDEX_GEN_INDIRECT_DRAW_OFFSET + 1] = 0;
-		IndirectDrawOutput[IndirectDrawOffset + INDEX_GEN_INDIRECT_DRAW_OFFSET + 2] = 0;
-		IndirectDrawOutput[IndirectDrawOffset + INDEX_GEN_INDIRECT_DRAW_OFFSET + 3] = 0;
-		IndirectDrawOutput[IndirectDrawOffset + INDEX_GEN_INDIRECT_DRAW_OFFSET + 4] = 0;
+		IndirectDrawOutput[IndirectDrawOffset + INDEX_GEN_DRAW_INDIRECT_OFFSET + 0] = 0;
+		IndirectDrawOutput[IndirectDrawOffset + INDEX_GEN_DRAW_INDIRECT_OFFSET + 1] = 0;
+		IndirectDrawOutput[IndirectDrawOffset + INDEX_GEN_DRAW_INDIRECT_OFFSET + 2] = 0;
+		IndirectDrawOutput[IndirectDrawOffset + INDEX_GEN_DRAW_INDIRECT_OFFSET + 3] = 0;
+		IndirectDrawOutput[IndirectDrawOffset + INDEX_GEN_DRAW_INDIRECT_OFFSET + 4] = 0;
+	
+		IndirectDrawOutput[IndirectDrawOffset + INDEX_GEN_STEREO_DRAW_INDIRECT_OFFSET + 0] = 0;
+		IndirectDrawOutput[IndirectDrawOffset + INDEX_GEN_STEREO_DRAW_INDIRECT_OFFSET + 1] = 0;
+		IndirectDrawOutput[IndirectDrawOffset + INDEX_GEN_STEREO_DRAW_INDIRECT_OFFSET + 2] = 0;
+		IndirectDrawOutput[IndirectDrawOffset + INDEX_GEN_STEREO_DRAW_INDIRECT_OFFSET + 3] = 0;
+		IndirectDrawOutput[IndirectDrawOffset + INDEX_GEN_STEREO_DRAW_INDIRECT_OFFSET + 4] = 0;
 	
+		IndirectDrawOutput[IndirectDrawOffset + INDEX_GEN_TESSELLATION_FACTOR_OFFSET] = 0;
 		IndirectDrawOutput[IndirectDrawOffset + INDEX_GEN_NUM_SEGMENTS_OFFSET] = 0;
 		IndirectDrawOutput[IndirectDrawOffset + INDEX_GEN_NUM_SUB_SEGMENTS_OFFSET] = 0;
 		IndirectDrawOutput[IndirectDrawOffset + INDEX_GEN_ONE_OVER_NUM_SUB_SEGMENTS_OFFSET] = asuint(0.0f);
@@ -135,24 +141,23 @@ void InitializeIndices(uint3 DispatchThreadId : SV_DispatchThreadID)
 	IndirectDrawOutput[IndirectDrawOffset + INDEX_GEN_EXECUTE_INDIRECT_OFFSET + 0] = NumThreadGroups;
 	IndirectDrawOutput[IndirectDrawOffset + INDEX_GEN_EXECUTE_INDIRECT_OFFSET + 1] = 1;
 	IndirectDrawOutput[IndirectDrawOffset + INDEX_GEN_EXECUTE_INDIRECT_OFFSET + 2] = 1;
-	IndirectDrawOutput[IndirectDrawOffset + INDEX_GEN_TESSELLATION_FACTOR_OFFSET] = SegmentTessellation;
 
-			
-	// Index Count
-	IndirectDrawOutput[IndirectDrawOffset + INDEX_GEN_INDIRECT_DRAW_OFFSET + 0] = SegmentCount * SegmentTessellation * TrianglesPerSegment * 3;
-	// Num Instances
-	IndirectDrawOutput[IndirectDrawOffset + INDEX_GEN_INDIRECT_DRAW_OFFSET + 1] = 1;
-	// First Index Offset
-	IndirectDrawOutput[IndirectDrawOffset + INDEX_GEN_INDIRECT_DRAW_OFFSET + 2] = 0;
-	// First Vertex Offset
-	IndirectDrawOutput[IndirectDrawOffset + INDEX_GEN_INDIRECT_DRAW_OFFSET + 3] = 0;
-	// First Instance Offset
-	IndirectDrawOutput[IndirectDrawOffset + INDEX_GEN_INDIRECT_DRAW_OFFSET + 4] = 0;
-	
-	// SegmentCount
+	// DrawIndirect Params: IndexCount, NumInstances, FirstIndexOffset, FirstVertexOffset, FirstInstanceOffset
+	IndirectDrawOutput[IndirectDrawOffset + INDEX_GEN_DRAW_INDIRECT_OFFSET + 0] = SegmentCount * SegmentTessellation * TrianglesPerSegment * 3;
+	IndirectDrawOutput[IndirectDrawOffset + INDEX_GEN_DRAW_INDIRECT_OFFSET + 1] = 1;
+	IndirectDrawOutput[IndirectDrawOffset + INDEX_GEN_DRAW_INDIRECT_OFFSET + 2] = 0;
+	IndirectDrawOutput[IndirectDrawOffset + INDEX_GEN_DRAW_INDIRECT_OFFSET + 3] = 0;
+	IndirectDrawOutput[IndirectDrawOffset + INDEX_GEN_DRAW_INDIRECT_OFFSET + 4] = 0;
+
+	IndirectDrawOutput[IndirectDrawOffset + INDEX_GEN_STEREO_DRAW_INDIRECT_OFFSET + 0] = SegmentCount * SegmentTessellation * TrianglesPerSegment * 3;
+	IndirectDrawOutput[IndirectDrawOffset + INDEX_GEN_STEREO_DRAW_INDIRECT_OFFSET + 1] = 2;
+	IndirectDrawOutput[IndirectDrawOffset + INDEX_GEN_STEREO_DRAW_INDIRECT_OFFSET + 2] = 0;
+	IndirectDrawOutput[IndirectDrawOffset + INDEX_GEN_STEREO_DRAW_INDIRECT_OFFSET + 3] = 0;
+	IndirectDrawOutput[IndirectDrawOffset + INDEX_GEN_STEREO_DRAW_INDIRECT_OFFSET + 4] = 0;
+
+	// Misc data for generation
+	IndirectDrawOutput[IndirectDrawOffset + INDEX_GEN_TESSELLATION_FACTOR_OFFSET] = SegmentTessellation;
 	IndirectDrawOutput[IndirectDrawOffset + INDEX_GEN_NUM_SEGMENTS_OFFSET] = SegmentCount;
-	// SubSegmentCount
 	IndirectDrawOutput[IndirectDrawOffset + INDEX_GEN_NUM_SUB_SEGMENTS_OFFSET] = SegmentTessellation;
-	// 1/SubSegmentCount
 	IndirectDrawOutput[IndirectDrawOffset + INDEX_GEN_ONE_OVER_NUM_SUB_SEGMENTS_OFFSET] = asuint(OneOverSubSegmentTessellation);
 }
diff --git a/Engine/Plugins/FX/Niagara/Source/Niagara/Private/NiagaraRendererRibbons.cpp b/Engine/Plugins/FX/Niagara/Source/Niagara/Private/NiagaraRendererRibbons.cpp
--- a/Engine/Plugins/FX/Niagara/Source/Niagara/Private/NiagaraRendererRibbons.cpp
+++ b/Engine/Plugins/FX/Niagara/Source/Niagara/Private/NiagaraRendererRibbons.cpp
@@ -228,30 +228,22 @@ struct FNiagaraRibbonCommandBufferLayout
 	float TessCurrentFrame_AverageWidth;
 };
 
-
+// This data must match INDEX_GEN_INDIRECT_ARGS_STRIDE in NiagaraRibbonCommon.ush
+// Be careful if we ever allocate more than 1 of these as ExecuteIndirect arguments have boundary restrictions on some platforms
 struct FNiagaraRibbonIndirectDrawBufferLayout
 {
-	static constexpr int32 NumElements = 12;
-	static constexpr int32 GenerateIndicesCommandOffset = 0;
-	static constexpr int32 IndirectDrawCommandIndex = 4;
-	static constexpr int32 IndirectDrawCommandByteOffset = IndirectDrawCommandIndex * sizeof(uint32);
-
-	// This is passed from InitializeIndices to GenerateIndices
-	uint32 IndexGenIndirectArgsXDim;
-	uint32 IndexGenIndirectArgsYDim;
-	uint32 IndexGenIndirectArgsZDim;
-	uint32 TessellationFactor;
-
-	// This is the indirect draw args and then resulting information for the vertex shader	
-	uint32 IndexCount;
-	uint32 NumInstances;
-	uint32 FirstIndexOffset;
-	uint32 FirstVertexOffset;
-	uint32 FirstInstanceOffset;
-	
-	uint32 NumSegments;
-	uint32 NumSubSegments;
-	float OneOverSubSegmentCount;	
+	FRHIDispatchIndirectParametersNoPadding	IndexGenIndirectArgs;			//  0 - 3 uints
+	FRHIDrawIndexedIndirectParameters		DrawIndirectParameters;			//  3 - 5 uints
+	FRHIDrawIndexedIndirectParameters		StereoDrawIndirectParameters;	//  8 - 5 uints
+	
+	uint32	TessellationFactor;												// 13 - 1 uint
+	uint32	NumSegments;													// 14 - 1 uint
+	uint32	NumSubSegments;													// 15 - 1 uint
+	float	OneOverSubSegmentCount;											// 16 - 1 uint
+
+	static constexpr int32 DrawIndirectParametersByteOffset			= 3 * sizeof(uint32);
+	static constexpr int32 StereoDrawIndirectParametersByteOffset	= 8 * sizeof(uint32);
+	static constexpr int32 NumElements								= 17;
 };
 
 struct FNiagaraRibbonIndexBuffer final : FIndexBuffer
@@ -337,9 +329,6 @@ struct FNiagaraRibbonRenderingFrameViewResources
 	FRWBuffer						IndirectDrawBuffer;
 	FNiagaraIndexGenerationInput	IndexGenerationSettings;
 
-	int32 IndirectDrawBufferStartOffset = FNiagaraRibbonIndirectDrawBufferLayout::IndirectDrawCommandIndex;
-	int32 IndirectDrawBufferStartByteOffset = FNiagaraRibbonIndirectDrawBufferLayout::IndirectDrawCommandByteOffset;
-
 	~FNiagaraRibbonRenderingFrameViewResources()
 	{
 		UniformBuffer.SafeRelease();
@@ -2072,7 +2061,7 @@ inline void FNiagaraRendererRibbons::SetupMeshBatchAndCollectorResourceForView(c
 	{
 		MeshElement.NumPrimitives = 0;
 		MeshElement.IndirectArgsBuffer = RenderingViewResources->IndirectDrawBuffer.Buffer;
-		MeshElement.IndirectArgsOffset = RenderingViewResources->IndirectDrawBufferStartByteOffset;
+		MeshElement.IndirectArgsOffset = View->IsInstancedStereoPass() ? FNiagaraRibbonIndirectDrawBufferLayout::StereoDrawIndirectParametersByteOffset : FNiagaraRibbonIndirectDrawBufferLayout::DrawIndirectParametersByteOffset;
 	}
 	else
 	{
