Niagara - When loading curve DIs from cooked packages don't update the lut since it's already up to date and the curve data isn't restored from the editor cache until the derived class deserializes.

diff --git a/Engine/Plugins/FX/Niagara/Source/Niagara/Private/NiagaraDataInterfaceCurveBase.cpp b/Engine/Plugins/FX/Niagara/Source/Niagara/Private/NiagaraDataInterfaceCurveBase.cpp
--- a/Engine/Plugins/FX/Niagara/Source/Niagara/Private/NiagaraDataInterfaceCurveBase.cpp
+++ b/Engine/Plugins/FX/Niagara/Source/Niagara/Private/NiagaraDataInterfaceCurveBase.cpp
@@ -207,8 +207,10 @@ void UNiagaraDataInterfaceCurveBase::Serialize(FArchive& Ar)
 #if WITH_EDITORONLY_DATA
 		HasEditorData = !Ar.IsFilterEditorOnly();
 		// Sometimes curves are out of date which needs to be tracked down
-		// Temporarily we will make sure they are up to date in editor builds
-		if (HasEditorData && GetClass() != UNiagaraDataInterfaceCurveBase::StaticClass())
+		// Temporarily we will make sure they are up to date in editor builds,
+		// but not for cooked editor packages since they're already up to date,
+		// and the curve data is not valid at this point.
+		if (HasEditorData && Ar.IsLoadingFromCookedPackage() == false && GetClass() != UNiagaraDataInterfaceCurveBase::StaticClass())
 		{
 			UpdateLUT(true);
 		}
