Rollback changes to BeginCacheForCookedPlatformData/IsCachedCookedPlatformDataLoaded pattern for Niagara systems as it was causing problems iwth jupiter ddc fill commandlet
#rnx

diff --git a/Engine/Plugins/FX/Niagara/Source/Niagara/Private/NiagaraSystem.cpp b/Engine/Plugins/FX/Niagara/Source/Niagara/Private/NiagaraSystem.cpp
--- a/Engine/Plugins/FX/Niagara/Source/Niagara/Private/NiagaraSystem.cpp
+++ b/Engine/Plugins/FX/Niagara/Source/Niagara/Private/NiagaraSystem.cpp
@@ -264,20 +264,18 @@ void UNiagaraSystem::BeginCacheForCookedPlatformData(const ITargetPlatform *Targ
 	
 	EnsureFullyLoaded();
 #if WITH_EDITORONLY_DATA
-	if (bNeedsRequestCompile)
-	{
-		RequestCompile(false);
-	}
+	// todo - when we can get compilation work to be properly asynchronous we need to move this over to the intended behavior where
+	// this function kicks off the work and we poll for results in IsCachedCookedPlatformDataLoaded().  With the current setup we
+	// will timeout DDC fill jobs while polling for the work to be complete.  Not entirely sure why, but my guess is that we get too
+	// many systems in the process of being compiled and we iteratively block the gamethread advancing each of them starving progress
+	WaitForCompilationComplete();
 #endif
 }
 
 bool UNiagaraSystem::IsCachedCookedPlatformDataLoaded(const ITargetPlatform* TargetPlatform)
 {
-	if (Super::IsCachedCookedPlatformDataLoaded(TargetPlatform))
-	{
-		return PollForCompilationComplete(false);
-	}
-	return false;
+	// todo - see note above in BeginCacheForCookedPlatformData()
+	return Super::IsCachedCookedPlatformDataLoaded(TargetPlatform);
 }
 
 void UNiagaraSystem::HandleVariableRenamed(const FNiagaraVariable& InOldVariable, const FNiagaraVariable& InNewVariable, bool bUpdateContexts)
