- Optional per channel masking for dynamic parameters
- Looks for a static int with the binding name + ChannelMask which will be used to determine what channels should be read from the material vs the dynamic parameters

diff --git a/Engine/Plugins/FX/Niagara/Source/Niagara/Classes/NiagaraScript.h b/Engine/Plugins/FX/Niagara/Source/Niagara/Classes/NiagaraScript.h
--- a/Engine/Plugins/FX/Niagara/Source/Niagara/Classes/NiagaraScript.h
+++ b/Engine/Plugins/FX/Niagara/Source/Niagara/Classes/NiagaraScript.h
@@ -1058,6 +1058,25 @@ public:
 	NIAGARA_API bool IsScriptCompilationPending(bool bGPUScript) const;
 	NIAGARA_API bool DidScriptCompilationSucceed(bool bGPUScript) const;
 
+#if WITH_EDITORONLY_DATA
+	template<typename T>
+	TOptional<T> GetStaticVariableValue(const FNiagaraVariableBase& InVar) const
+	{
+		for (const FNiagaraVariable& StaticVariable : CachedScriptVM.StaticVariablesWritten)
+		{
+			if (static_cast<const FNiagaraVariableBase&>(StaticVariable) == InVar)
+			{
+				if (StaticVariable.IsDataAllocated())
+				{
+					return TOptional<T>(StaticVariable.GetValue<T>());
+				}
+			}
+		}
+
+		return TOptional<T>();
+	}
+#endif
+
 	template<typename T>
 	TOptional<T> GetCompilerTag(const FNiagaraVariableBase& InVar, const FNiagaraParameterStore* FallbackParameterStore = nullptr) const
 	{
diff --git a/Engine/Plugins/FX/Niagara/Source/Niagara/Private/NiagaraMeshRendererProperties.cpp b/Engine/Plugins/FX/Niagara/Source/Niagara/Private/NiagaraMeshRendererProperties.cpp
--- a/Engine/Plugins/FX/Niagara/Source/Niagara/Private/NiagaraMeshRendererProperties.cpp
+++ b/Engine/Plugins/FX/Niagara/Source/Niagara/Private/NiagaraMeshRendererProperties.cpp
@@ -486,10 +486,10 @@ void UNiagaraMeshRendererProperties::CacheFromCompiledData(const FNiagaraDataSet
 	RendererLayoutWithCustomSorting.SetVariableFromBinding(CompiledData, CustomSortingBinding, ENiagaraMeshVFLayout::CustomSorting);
 	RendererLayoutWithCustomSorting.SetVariableFromBinding(CompiledData, SubImageIndexBinding, ENiagaraMeshVFLayout::SubImage);
 	RendererLayoutWithCustomSorting.SetVariableFromBinding(CompiledData, CameraOffsetBinding, ENiagaraMeshVFLayout::CameraOffset);
-	MaterialParamValidMask  = RendererLayoutWithCustomSorting.SetVariableFromBinding(CompiledData, DynamicMaterialBinding, ENiagaraMeshVFLayout::DynamicParam0) ? 0x1 : 0;
-	MaterialParamValidMask |= RendererLayoutWithCustomSorting.SetVariableFromBinding(CompiledData, DynamicMaterial1Binding, ENiagaraMeshVFLayout::DynamicParam1) ? 0x2 : 0;
-	MaterialParamValidMask |= RendererLayoutWithCustomSorting.SetVariableFromBinding(CompiledData, DynamicMaterial2Binding, ENiagaraMeshVFLayout::DynamicParam2) ? 0x4 : 0;
-	MaterialParamValidMask |= RendererLayoutWithCustomSorting.SetVariableFromBinding(CompiledData, DynamicMaterial3Binding, ENiagaraMeshVFLayout::DynamicParam3) ? 0x8 : 0;
+	RendererLayoutWithCustomSorting.SetVariableFromBinding(CompiledData, DynamicMaterialBinding,  ENiagaraMeshVFLayout::DynamicParam0);
+	RendererLayoutWithCustomSorting.SetVariableFromBinding(CompiledData, DynamicMaterial1Binding, ENiagaraMeshVFLayout::DynamicParam1);
+	RendererLayoutWithCustomSorting.SetVariableFromBinding(CompiledData, DynamicMaterial2Binding, ENiagaraMeshVFLayout::DynamicParam2);
+	RendererLayoutWithCustomSorting.SetVariableFromBinding(CompiledData, DynamicMaterial3Binding, ENiagaraMeshVFLayout::DynamicParam3);
 	if (NeedsPreciseMotionVectors())
 	{
 		RendererLayoutWithCustomSorting.SetVariableFromBinding(CompiledData, PrevPositionBinding, ENiagaraMeshVFLayout::PrevPosition);
@@ -510,10 +510,10 @@ void UNiagaraMeshRendererProperties::CacheFromCompiledData(const FNiagaraDataSet
 	RendererLayoutWithoutCustomSorting.SetVariableFromBinding(CompiledData, NormalizedAgeBinding, ENiagaraMeshVFLayout::NormalizedAge);
 	RendererLayoutWithoutCustomSorting.SetVariableFromBinding(CompiledData, SubImageIndexBinding, ENiagaraMeshVFLayout::SubImage);
 	RendererLayoutWithoutCustomSorting.SetVariableFromBinding(CompiledData, CameraOffsetBinding, ENiagaraMeshVFLayout::CameraOffset);
-	MaterialParamValidMask =  RendererLayoutWithoutCustomSorting.SetVariableFromBinding(CompiledData, DynamicMaterialBinding, ENiagaraMeshVFLayout::DynamicParam0) ? 0x1 : 0;
-	MaterialParamValidMask |= RendererLayoutWithoutCustomSorting.SetVariableFromBinding(CompiledData, DynamicMaterial1Binding, ENiagaraMeshVFLayout::DynamicParam1) ? 0x2 : 0;
-	MaterialParamValidMask |= RendererLayoutWithoutCustomSorting.SetVariableFromBinding(CompiledData, DynamicMaterial2Binding, ENiagaraMeshVFLayout::DynamicParam2) ? 0x4 : 0;
-	MaterialParamValidMask |= RendererLayoutWithoutCustomSorting.SetVariableFromBinding(CompiledData, DynamicMaterial3Binding, ENiagaraMeshVFLayout::DynamicParam3) ? 0x8 : 0;
+	const bool bDynamicParam0Valid = RendererLayoutWithoutCustomSorting.SetVariableFromBinding(CompiledData, DynamicMaterialBinding,  ENiagaraMeshVFLayout::DynamicParam0);
+	const bool bDynamicParam1Valid = RendererLayoutWithoutCustomSorting.SetVariableFromBinding(CompiledData, DynamicMaterial1Binding, ENiagaraMeshVFLayout::DynamicParam1);
+	const bool bDynamicParam2Valid = RendererLayoutWithoutCustomSorting.SetVariableFromBinding(CompiledData, DynamicMaterial2Binding, ENiagaraMeshVFLayout::DynamicParam2);
+	const bool bDynamicParam3Valid = RendererLayoutWithoutCustomSorting.SetVariableFromBinding(CompiledData, DynamicMaterial3Binding, ENiagaraMeshVFLayout::DynamicParam3);
 	if (NeedsPreciseMotionVectors())
 	{
 		RendererLayoutWithoutCustomSorting.SetVariableFromBinding(CompiledData, PrevPositionBinding, ENiagaraMeshVFLayout::PrevPosition);
@@ -523,6 +523,16 @@ void UNiagaraMeshRendererProperties::CacheFromCompiledData(const FNiagaraDataSet
 		RendererLayoutWithoutCustomSorting.SetVariableFromBinding(CompiledData, PrevVelocityBinding, ENiagaraMeshVFLayout::PrevVelocity);
 	}
 	RendererLayoutWithoutCustomSorting.Finalize();
+
+#if WITH_EDITORONLY_DATA
+	// Build dynamic parameter mask
+	// Serialize in cooked builds
+	const FVersionedNiagaraEmitterData* EmitterData = GetEmitterData();
+	MaterialParamValidMask  = bDynamicParam0Valid ? GetDynamicParameterChannelMask(EmitterData, DynamicMaterialBinding.GetName(), 0xf) << 0 : 0;
+	MaterialParamValidMask |= bDynamicParam1Valid ? GetDynamicParameterChannelMask(EmitterData, DynamicMaterial1Binding.GetName(), 0xf) << 4 : 0;
+	MaterialParamValidMask |= bDynamicParam2Valid ? GetDynamicParameterChannelMask(EmitterData, DynamicMaterial2Binding.GetName(), 0xf) << 8 : 0;
+	MaterialParamValidMask |= bDynamicParam3Valid ? GetDynamicParameterChannelMask(EmitterData, DynamicMaterial3Binding.GetName(), 0xf) << 12 : 0;
+#endif
 }
 
 void UNiagaraMeshRendererProperties::UpdateMICs()
diff --git a/Engine/Plugins/FX/Niagara/Source/Niagara/Private/NiagaraRendererMeshes.cpp b/Engine/Plugins/FX/Niagara/Source/Niagara/Private/NiagaraRendererMeshes.cpp
--- a/Engine/Plugins/FX/Niagara/Source/Niagara/Private/NiagaraRendererMeshes.cpp
+++ b/Engine/Plugins/FX/Niagara/Source/Niagara/Private/NiagaraRendererMeshes.cpp
@@ -903,19 +903,19 @@ FNiagaraMeshUniformBufferRef FNiagaraRendererMeshes::CreateVFUniformBuffer(const
 					break;
 				case ENiagaraMeshVFLayout::Type::DynamicParam0:
 					FMemory::Memcpy(&Params.DefaultDynamicMaterialParameter0, ParameterBoundData + VFBoundOffsetsInParamStore[i], sizeof(FVector4f));
-					Params.MaterialParamValidMask |= 0x1;
+					Params.MaterialParamValidMask |= 0x000f;
 					break;
 				case ENiagaraMeshVFLayout::Type::DynamicParam1:
 					FMemory::Memcpy(&Params.DefaultDynamicMaterialParameter1, ParameterBoundData + VFBoundOffsetsInParamStore[i], sizeof(FVector4f));
-					Params.MaterialParamValidMask |= 0x2;
+					Params.MaterialParamValidMask |= 0x00f0;
 					break;
 				case ENiagaraMeshVFLayout::Type::DynamicParam2:
 					FMemory::Memcpy(&Params.DefaultDynamicMaterialParameter2, ParameterBoundData + VFBoundOffsetsInParamStore[i], sizeof(FVector4f));
-					Params.MaterialParamValidMask |= 0x4;
+					Params.MaterialParamValidMask |= 0x0f00;
 					break;
 				case ENiagaraMeshVFLayout::Type::DynamicParam3:
 					FMemory::Memcpy(&Params.DefaultDynamicMaterialParameter3, ParameterBoundData + VFBoundOffsetsInParamStore[i], sizeof(FVector4f));
-					Params.MaterialParamValidMask |= 0x8;
+					Params.MaterialParamValidMask |= 0xf000;
 					break;
 				case ENiagaraMeshVFLayout::Type::CustomSorting:
 					// unsupported for now...
diff --git a/Engine/Plugins/FX/Niagara/Source/Niagara/Private/NiagaraRendererProperties.cpp b/Engine/Plugins/FX/Niagara/Source/Niagara/Private/NiagaraRendererProperties.cpp
--- a/Engine/Plugins/FX/Niagara/Source/Niagara/Private/NiagaraRendererProperties.cpp
+++ b/Engine/Plugins/FX/Niagara/Source/Niagara/Private/NiagaraRendererProperties.cpp
@@ -547,6 +547,34 @@ void UNiagaraRendererProperties::UpdateMaterialParametersMIC(const FNiagaraRende
 	}
 }
 
+int32 UNiagaraRendererProperties::GetDynamicParameterChannelMask(const FVersionedNiagaraEmitterData* EmitterData, FName BindingName, int32 DefaultChannelMask) const
+{
+	if (EmitterData == nullptr)
+	{
+		return DefaultChannelMask;
+	}
+
+	TOptional<int32> ChannelMask;
+	EmitterData->ForEachScript(
+		[&ChannelMask, &BindingName](const UNiagaraScript* NiagaraScript)
+		{
+			const FNiagaraVMExecutableData& VMExecData = NiagaraScript->GetVMExecutableData();
+
+			FNameBuilder NameBuilder;
+			BindingName.ToString(NameBuilder);
+			NameBuilder.Append(TEXT("ChannelMask"));
+
+			const FNiagaraVariableBase ChannelMaskVariable(FNiagaraTypeDefinition::GetIntDef().ToStaticDef(), FName(NameBuilder));
+			TOptional<int32> ChannelMaskValue = NiagaraScript->GetStaticVariableValue<int32>(ChannelMaskVariable);
+			if (ChannelMaskValue.IsSet())
+			{
+				ChannelMask = ChannelMask.Get(0) | ChannelMaskValue.GetValue();
+			}
+		}
+	);
+	return ChannelMask.Get(DefaultChannelMask);
+}
+
 FNiagaraVariable UNiagaraRendererProperties::GetBoundAttribute(const FNiagaraVariableAttributeBinding* Binding) const
 {
 	if (Binding->GetParamMapBindableVariable().IsValid())
diff --git a/Engine/Plugins/FX/Niagara/Source/Niagara/Private/NiagaraRendererSprites.cpp b/Engine/Plugins/FX/Niagara/Source/Niagara/Private/NiagaraRendererSprites.cpp
--- a/Engine/Plugins/FX/Niagara/Source/Niagara/Private/NiagaraRendererSprites.cpp
+++ b/Engine/Plugins/FX/Niagara/Source/Niagara/Private/NiagaraRendererSprites.cpp
@@ -707,19 +707,19 @@ FNiagaraSpriteUniformBufferRef FNiagaraRendererSprites::CreateViewUniformBuffer(
 					break;
 				case ENiagaraSpriteVFLayout::Type::MaterialParam0:
 					FMemory::Memcpy(&PerViewUniformParameters.DefaultDynamicMaterialParameter0, DynamicDataSprites->ParameterDataBound.GetData() + VFBoundOffsetsInParamStore[i], sizeof(FVector4f));
-					PerViewUniformParameters.MaterialParamValidMask |= 0x1;
+					PerViewUniformParameters.MaterialParamValidMask |= 0x000f;
 					break;
 				case ENiagaraSpriteVFLayout::Type::MaterialParam1:
 					FMemory::Memcpy(&PerViewUniformParameters.DefaultDynamicMaterialParameter1, DynamicDataSprites->ParameterDataBound.GetData() + VFBoundOffsetsInParamStore[i], sizeof(FVector4f));
-					PerViewUniformParameters.MaterialParamValidMask |= 0x2;
+					PerViewUniformParameters.MaterialParamValidMask |= 0x00f0;
 					break;
 				case ENiagaraSpriteVFLayout::Type::MaterialParam2:
 					FMemory::Memcpy(&PerViewUniformParameters.DefaultDynamicMaterialParameter2, DynamicDataSprites->ParameterDataBound.GetData() + VFBoundOffsetsInParamStore[i], sizeof(FVector4f));
-					PerViewUniformParameters.MaterialParamValidMask |= 0x4;
+					PerViewUniformParameters.MaterialParamValidMask |= 0x0f00;
 					break;
 				case ENiagaraSpriteVFLayout::Type::MaterialParam3:
 					FMemory::Memcpy(&PerViewUniformParameters.DefaultDynamicMaterialParameter3, DynamicDataSprites->ParameterDataBound.GetData() + VFBoundOffsetsInParamStore[i], sizeof(FVector4f));
-					PerViewUniformParameters.MaterialParamValidMask |= 0x8;
+					PerViewUniformParameters.MaterialParamValidMask |= 0xf000;
 					break;
 				case ENiagaraSpriteVFLayout::Type::CameraOffset:
 					FMemory::Memcpy(&PerViewUniformParameters.DefaultCamOffset, DynamicDataSprites->ParameterDataBound.GetData() + VFBoundOffsetsInParamStore[i], sizeof(float));
diff --git a/Engine/Plugins/FX/Niagara/Source/Niagara/Private/NiagaraRibbonRendererProperties.cpp b/Engine/Plugins/FX/Niagara/Source/Niagara/Private/NiagaraRibbonRendererProperties.cpp
--- a/Engine/Plugins/FX/Niagara/Source/Niagara/Private/NiagaraRibbonRendererProperties.cpp
+++ b/Engine/Plugins/FX/Niagara/Source/Niagara/Private/NiagaraRibbonRendererProperties.cpp
@@ -404,10 +404,10 @@ void UNiagaraRibbonRendererProperties::CacheFromCompiledData(const FNiagaraDataS
 	RendererLayout.SetVariableFromBinding(CompiledData, V0RangeOverrideBinding, ENiagaraRibbonVFLayout::V0RangeOverride);
 	RendererLayout.SetVariableFromBinding(CompiledData, U1OverrideBinding, ENiagaraRibbonVFLayout::U1Override);
 	RendererLayout.SetVariableFromBinding(CompiledData, V1RangeOverrideBinding, ENiagaraRibbonVFLayout::V1RangeOverride);
-	MaterialParamValidMask  = RendererLayout.SetVariableFromBinding(CompiledData, DynamicMaterialBinding, ENiagaraRibbonVFLayout::MaterialParam0) ? 1 : 0;
-	MaterialParamValidMask |= RendererLayout.SetVariableFromBinding(CompiledData, DynamicMaterial1Binding, ENiagaraRibbonVFLayout::MaterialParam1) ? 2 : 0;
-	MaterialParamValidMask |= RendererLayout.SetVariableFromBinding(CompiledData, DynamicMaterial2Binding, ENiagaraRibbonVFLayout::MaterialParam2) ? 4 : 0;
-	MaterialParamValidMask |= RendererLayout.SetVariableFromBinding(CompiledData, DynamicMaterial3Binding, ENiagaraRibbonVFLayout::MaterialParam3) ? 8 : 0;
+	const bool bDynamicParam0Valid = RendererLayout.SetVariableFromBinding(CompiledData, DynamicMaterialBinding,  ENiagaraRibbonVFLayout::MaterialParam0);
+	const bool bDynamicParam1Valid = RendererLayout.SetVariableFromBinding(CompiledData, DynamicMaterial1Binding, ENiagaraRibbonVFLayout::MaterialParam1);
+	const bool bDynamicParam2Valid = RendererLayout.SetVariableFromBinding(CompiledData, DynamicMaterial2Binding, ENiagaraRibbonVFLayout::MaterialParam2);
+	const bool bDynamicParam3Valid = RendererLayout.SetVariableFromBinding(CompiledData, DynamicMaterial3Binding, ENiagaraRibbonVFLayout::MaterialParam3);
 
 	if (NeedsPreciseMotionVectors())
 	{
@@ -420,6 +420,16 @@ void UNiagaraRibbonRendererProperties::CacheFromCompiledData(const FNiagaraDataS
 	RendererLayout.SetVariableFromBinding(CompiledData, RibbonLinkOrderBinding, ENiagaraRibbonVFLayout::LinkOrder);
 	
 	RendererLayout.Finalize();
+
+#if WITH_EDITORONLY_DATA
+	// Build dynamic parameter mask
+	// Serialized in cooked builds
+	const FVersionedNiagaraEmitterData* EmitterData = GetEmitterData();
+	MaterialParamValidMask  = bDynamicParam0Valid ? GetDynamicParameterChannelMask(EmitterData, DynamicMaterialBinding.GetName(),  0xf) <<  0 : 0;
+	MaterialParamValidMask |= bDynamicParam1Valid ? GetDynamicParameterChannelMask(EmitterData, DynamicMaterial1Binding.GetName(), 0xf) <<  4 : 0;
+	MaterialParamValidMask |= bDynamicParam2Valid ? GetDynamicParameterChannelMask(EmitterData, DynamicMaterial2Binding.GetName(), 0xf) <<  8 : 0;
+	MaterialParamValidMask |= bDynamicParam3Valid ? GetDynamicParameterChannelMask(EmitterData, DynamicMaterial3Binding.GetName(), 0xf) << 12 : 0;
+#endif
 }
 
 void UNiagaraRibbonRendererProperties::UpdateMICs()
diff --git a/Engine/Plugins/FX/Niagara/Source/Niagara/Private/NiagaraSpriteRendererProperties.cpp b/Engine/Plugins/FX/Niagara/Source/Niagara/Private/NiagaraSpriteRendererProperties.cpp
--- a/Engine/Plugins/FX/Niagara/Source/Niagara/Private/NiagaraSpriteRendererProperties.cpp
+++ b/Engine/Plugins/FX/Niagara/Source/Niagara/Private/NiagaraSpriteRendererProperties.cpp
@@ -295,6 +295,7 @@ void UNiagaraSpriteRendererProperties::CacheFromCompiledData(const FNiagaraDataS
 	UpdateSourceModeDerivates(SourceMode);
 	UpdateMICs();
 
+	// Initialize layout
 	const int32 NumLayoutVars = NeedsPreciseMotionVectors() ? ENiagaraSpriteVFLayout::Num_Max : ENiagaraSpriteVFLayout::Num_Default;
 	RendererLayoutWithCustomSort.Initialize(NumLayoutVars);
 	RendererLayoutWithCustomSort.SetVariableFromBinding(CompiledData, PositionBinding, ENiagaraSpriteVFLayout::Position);
@@ -311,10 +312,10 @@ void UNiagaraSpriteRendererProperties::CacheFromCompiledData(const FNiagaraDataS
 	RendererLayoutWithCustomSort.SetVariableFromBinding(CompiledData, NormalizedAgeBinding, ENiagaraSpriteVFLayout::NormalizedAge);
 	RendererLayoutWithCustomSort.SetVariableFromBinding(CompiledData, MaterialRandomBinding, ENiagaraSpriteVFLayout::MaterialRandom);
 	RendererLayoutWithCustomSort.SetVariableFromBinding(CompiledData, CustomSortingBinding, ENiagaraSpriteVFLayout::CustomSorting);
-	MaterialParamValidMask  = RendererLayoutWithCustomSort.SetVariableFromBinding(CompiledData, DynamicMaterialBinding, ENiagaraSpriteVFLayout::MaterialParam0) ? 0x1 : 0;
-	MaterialParamValidMask |= RendererLayoutWithCustomSort.SetVariableFromBinding(CompiledData, DynamicMaterial1Binding, ENiagaraSpriteVFLayout::MaterialParam1) ? 0x2 : 0;
-	MaterialParamValidMask |= RendererLayoutWithCustomSort.SetVariableFromBinding(CompiledData, DynamicMaterial2Binding, ENiagaraSpriteVFLayout::MaterialParam2) ? 0x4 : 0;
-	MaterialParamValidMask |= RendererLayoutWithCustomSort.SetVariableFromBinding(CompiledData, DynamicMaterial3Binding, ENiagaraSpriteVFLayout::MaterialParam3) ? 0x8 : 0;
+	RendererLayoutWithCustomSort.SetVariableFromBinding(CompiledData, DynamicMaterialBinding, ENiagaraSpriteVFLayout::MaterialParam0);
+	RendererLayoutWithCustomSort.SetVariableFromBinding(CompiledData, DynamicMaterial1Binding, ENiagaraSpriteVFLayout::MaterialParam1);
+	RendererLayoutWithCustomSort.SetVariableFromBinding(CompiledData, DynamicMaterial2Binding, ENiagaraSpriteVFLayout::MaterialParam2);
+	RendererLayoutWithCustomSort.SetVariableFromBinding(CompiledData, DynamicMaterial3Binding, ENiagaraSpriteVFLayout::MaterialParam3);
 	if (NeedsPreciseMotionVectors())
 	{
 		RendererLayoutWithCustomSort.SetVariableFromBinding(CompiledData, PrevPositionBinding, ENiagaraSpriteVFLayout::PrevPosition);
@@ -342,10 +343,10 @@ void UNiagaraSpriteRendererProperties::CacheFromCompiledData(const FNiagaraDataS
 	RendererLayoutWithoutCustomSort.SetVariableFromBinding(CompiledData, PivotOffsetBinding, ENiagaraSpriteVFLayout::PivotOffset);
 	RendererLayoutWithoutCustomSort.SetVariableFromBinding(CompiledData, NormalizedAgeBinding, ENiagaraSpriteVFLayout::NormalizedAge);
 	RendererLayoutWithoutCustomSort.SetVariableFromBinding(CompiledData, MaterialRandomBinding, ENiagaraSpriteVFLayout::MaterialRandom);
-	MaterialParamValidMask = RendererLayoutWithoutCustomSort.SetVariableFromBinding(CompiledData, DynamicMaterialBinding, ENiagaraSpriteVFLayout::MaterialParam0) ? 0x1 : 0;
-	MaterialParamValidMask |= RendererLayoutWithoutCustomSort.SetVariableFromBinding(CompiledData, DynamicMaterial1Binding, ENiagaraSpriteVFLayout::MaterialParam1) ? 0x2 : 0;
-	MaterialParamValidMask |= RendererLayoutWithoutCustomSort.SetVariableFromBinding(CompiledData, DynamicMaterial2Binding, ENiagaraSpriteVFLayout::MaterialParam2) ? 0x4 : 0;
-	MaterialParamValidMask |= RendererLayoutWithoutCustomSort.SetVariableFromBinding(CompiledData, DynamicMaterial3Binding, ENiagaraSpriteVFLayout::MaterialParam3) ? 0x8 : 0;
+	const bool bDynamicParam0Valid = RendererLayoutWithoutCustomSort.SetVariableFromBinding(CompiledData, DynamicMaterialBinding,  ENiagaraSpriteVFLayout::MaterialParam0);
+	const bool bDynamicParam1Valid = RendererLayoutWithoutCustomSort.SetVariableFromBinding(CompiledData, DynamicMaterial1Binding, ENiagaraSpriteVFLayout::MaterialParam1);
+	const bool bDynamicParam2Valid = RendererLayoutWithoutCustomSort.SetVariableFromBinding(CompiledData, DynamicMaterial2Binding, ENiagaraSpriteVFLayout::MaterialParam2);
+	const bool bDynamicParam3Valid = RendererLayoutWithoutCustomSort.SetVariableFromBinding(CompiledData, DynamicMaterial3Binding, ENiagaraSpriteVFLayout::MaterialParam3);
 	if (NeedsPreciseMotionVectors())
 	{
 		RendererLayoutWithoutCustomSort.SetVariableFromBinding(CompiledData, PrevPositionBinding, ENiagaraSpriteVFLayout::PrevPosition);
@@ -358,6 +359,16 @@ void UNiagaraSpriteRendererProperties::CacheFromCompiledData(const FNiagaraDataS
 		RendererLayoutWithoutCustomSort.SetVariableFromBinding(CompiledData, PrevPivotOffsetBinding, ENiagaraSpriteVFLayout::PrevPivotOffset);
 	}
 	RendererLayoutWithoutCustomSort.Finalize();
+
+#if WITH_EDITORONLY_DATA
+	// Build dynamic parameter mask
+	// Serialize in cooked builds
+	const FVersionedNiagaraEmitterData* EmitterData = GetEmitterData();
+	MaterialParamValidMask  = bDynamicParam0Valid ? GetDynamicParameterChannelMask(EmitterData, DynamicMaterialBinding.GetName(),  0xf) <<  0 : 0;
+	MaterialParamValidMask |= bDynamicParam1Valid ? GetDynamicParameterChannelMask(EmitterData, DynamicMaterial1Binding.GetName(), 0xf) <<  4 : 0;
+	MaterialParamValidMask |= bDynamicParam2Valid ? GetDynamicParameterChannelMask(EmitterData, DynamicMaterial2Binding.GetName(), 0xf) <<  8 : 0;
+	MaterialParamValidMask |= bDynamicParam3Valid ? GetDynamicParameterChannelMask(EmitterData, DynamicMaterial3Binding.GetName(), 0xf) << 12 : 0;
+#endif
 }
 
 #if WITH_EDITORONLY_DATA
diff --git a/Engine/Plugins/FX/Niagara/Source/Niagara/Public/NiagaraMeshRendererProperties.h b/Engine/Plugins/FX/Niagara/Source/Niagara/Public/NiagaraMeshRendererProperties.h
--- a/Engine/Plugins/FX/Niagara/Source/Niagara/Public/NiagaraMeshRendererProperties.h
+++ b/Engine/Plugins/FX/Niagara/Source/Niagara/Public/NiagaraMeshRendererProperties.h
@@ -449,7 +449,9 @@ public:
 	uint32 NumFlipbookFrames;
 #endif
 
+	UPROPERTY()
 	uint32 MaterialParamValidMask = 0;
+
 	FNiagaraRendererLayout RendererLayoutWithCustomSorting;
 	FNiagaraRendererLayout RendererLayoutWithoutCustomSorting;
 
diff --git a/Engine/Plugins/FX/Niagara/Source/Niagara/Public/NiagaraRendererProperties.h b/Engine/Plugins/FX/Niagara/Source/Niagara/Public/NiagaraRendererProperties.h
--- a/Engine/Plugins/FX/Niagara/Source/Niagara/Public/NiagaraRendererProperties.h
+++ b/Engine/Plugins/FX/Niagara/Source/Niagara/Public/NiagaraRendererProperties.h
@@ -436,6 +436,8 @@ protected:
 	/** Utility function to updates MICs. */
 	void UpdateMaterialParametersMIC(const FNiagaraRendererMaterialParameters& MaterialParameters, TObjectPtr<UMaterialInterface>& InOutMaterial, TObjectPtr<UMaterialInstanceConstant>& InOutMIC);
 	void UpdateMaterialParametersMIC(const FNiagaraRendererMaterialParameters& MaterialParameters, TArrayView<UMaterialInterface*> Materials, TArray<TObjectPtr<UMaterialInstanceConstant>>& InOutMICs);
+
+	int32 GetDynamicParameterChannelMask(const FVersionedNiagaraEmitterData* EmitterData, FName BindingName, int32 DefaultChannelMask) const;
 #endif
 
 #if WITH_EDITOR
diff --git a/Engine/Plugins/FX/Niagara/Source/Niagara/Public/NiagaraRibbonRendererProperties.h b/Engine/Plugins/FX/Niagara/Source/Niagara/Public/NiagaraRibbonRendererProperties.h
--- a/Engine/Plugins/FX/Niagara/Source/Niagara/Public/NiagaraRibbonRendererProperties.h
+++ b/Engine/Plugins/FX/Niagara/Source/Niagara/Public/NiagaraRibbonRendererProperties.h
@@ -462,6 +462,8 @@ public:
 	UPROPERTY(Transient)
 	FNiagaraVariableAttributeBinding PrevRibbonTwistBinding;
 
+	UPROPERTY()
+	uint32 MaterialParamValidMask = 0;
 
 	bool								bSortKeyDataSetAccessorIsAge = false;
 	FNiagaraDataSetAccessor<float>		SortKeyDataSetAccessor;
@@ -483,7 +485,6 @@ public:
 	
 	FNiagaraDataSetAccessor<float>		RibbonLinkOrderDataSetAccessor;
 
-	uint32 MaterialParamValidMask = 0;
 	FNiagaraRendererLayout RendererLayout;
 
 protected:
diff --git a/Engine/Plugins/FX/Niagara/Source/Niagara/Public/NiagaraSpriteRendererProperties.h b/Engine/Plugins/FX/Niagara/Source/Niagara/Public/NiagaraSpriteRendererProperties.h
--- a/Engine/Plugins/FX/Niagara/Source/Niagara/Public/NiagaraSpriteRendererProperties.h
+++ b/Engine/Plugins/FX/Niagara/Source/Niagara/Public/NiagaraSpriteRendererProperties.h
@@ -392,11 +392,13 @@ public:
 	void CacheDerivedData();
 #endif
 
+	UPROPERTY()
+	uint32 MaterialParamValidMask = 0;
+
 	const TArray<FVector2f>& GetCutoutData() const { return DerivedData.BoundingGeometry; }
 
 	FNiagaraRendererLayout RendererLayoutWithCustomSort;
 	FNiagaraRendererLayout RendererLayoutWithoutCustomSort;
-	uint32 MaterialParamValidMask = 0;
 
 protected:
 	void InitBindings();
diff --git a/Engine/Shaders/Private/MaterialTemplate.ush b/Engine/Shaders/Private/MaterialTemplate.ush
--- a/Engine/Shaders/Private/MaterialTemplate.ush
+++ b/Engine/Shaders/Private/MaterialTemplate.ush
@@ -235,21 +235,46 @@ struct FMaterialParticleParameters
 float4 GetDynamicParameter(FMaterialParticleParameters Parameters, float4 Default, int ParameterIndex=0)
 {
 #if (NIAGARA_PARTICLE_FACTORY)
-	switch ( ParameterIndex)
+	switch ( ParameterIndex )
 	{
 	#if (DYNAMIC_PARAMETERS_MASK & 1)
-		case 0:	return (Parameters.DynamicParameterValidMask & 1) != 0 ? Parameters.DynamicParameter : Default;
+		case 0:
+			return float4(
+				(Parameters.DynamicParameterValidMask & 0x0001) == 0 ? Default.x : Parameters.DynamicParameter.x,
+				(Parameters.DynamicParameterValidMask & 0x0002) == 0 ? Default.y : Parameters.DynamicParameter.y,
+				(Parameters.DynamicParameterValidMask & 0x0004) == 0 ? Default.z : Parameters.DynamicParameter.z,
+				(Parameters.DynamicParameterValidMask & 0x0008) == 0 ? Default.w : Parameters.DynamicParameter.w
+			);
 	#endif
 	#if (DYNAMIC_PARAMETERS_MASK & 2)
-		case 1:	return (Parameters.DynamicParameterValidMask & 2) != 0 ? Parameters.DynamicParameter1 : Default;
+		case 1:
+			return float4(
+				(Parameters.DynamicParameterValidMask & 0x0010) == 0 ? Default.x : Parameters.DynamicParameter1.x,
+				(Parameters.DynamicParameterValidMask & 0x0020) == 0 ? Default.y : Parameters.DynamicParameter1.y,
+				(Parameters.DynamicParameterValidMask & 0x0040) == 0 ? Default.z : Parameters.DynamicParameter1.z,
+				(Parameters.DynamicParameterValidMask & 0x0080) == 0 ? Default.w : Parameters.DynamicParameter1.w
+			);
 	#endif
 	#if (DYNAMIC_PARAMETERS_MASK & 4)
-		case 2:	return (Parameters.DynamicParameterValidMask & 4) != 0 ? Parameters.DynamicParameter2 : Default;
+		case 2:
+			return float4(
+				(Parameters.DynamicParameterValidMask & 0x0100) == 0 ? Default.x : Parameters.DynamicParameter2.x,
+				(Parameters.DynamicParameterValidMask & 0x0200) == 0 ? Default.y : Parameters.DynamicParameter2.y,
+				(Parameters.DynamicParameterValidMask & 0x0400) == 0 ? Default.z : Parameters.DynamicParameter2.z,
+				(Parameters.DynamicParameterValidMask & 0x0800) == 0 ? Default.w : Parameters.DynamicParameter2.w
+			);
 	#endif	
 	#if (DYNAMIC_PARAMETERS_MASK & 8)
-		case 3:	return (Parameters.DynamicParameterValidMask & 8) != 0 ? Parameters.DynamicParameter3 : Default;
+		case 3:
+			return float4(
+				(Parameters.DynamicParameterValidMask & 0x1000) == 0 ? Default.x : Parameters.DynamicParameter3.x,
+				(Parameters.DynamicParameterValidMask & 0x2000) == 0 ? Default.y : Parameters.DynamicParameter3.y,
+				(Parameters.DynamicParameterValidMask & 0x4000) == 0 ? Default.z : Parameters.DynamicParameter3.z,
+				(Parameters.DynamicParameterValidMask & 0x8000) == 0 ? Default.w : Parameters.DynamicParameter3.w
+			);
 	#endif
-		default: return Default;
+		default:
+			return Default;
 	}
 #elif (PARTICLE_FACTORY)
 	if ( ParameterIndex == 0 )
