- Handle non normalized quats inside Niagara Decal Renderer to avoid triggering IsRotationNormalized

#rnx

diff --git a/Engine/Plugins/FX/Niagara/Source/Niagara/Private/NiagaraRendererDecals.cpp b/Engine/Plugins/FX/Niagara/Source/Niagara/Private/NiagaraRendererDecals.cpp
--- a/Engine/Plugins/FX/Niagara/Source/Niagara/Private/NiagaraRendererDecals.cpp
+++ b/Engine/Plugins/FX/Niagara/Source/Niagara/Private/NiagaraRendererDecals.cpp
@@ -200,7 +200,7 @@ FNiagaraDynamicDataBase* FNiagaraRendererDecals::GenerateDynamicData(const FNiag
 
 			// Grab Decal Attributes
 			const FVector3f SimPos = PositionReader.GetSafe(ParticleIndex, DefaultPos);
-			const FQuat4f SimRot = RotationReader.GetSafe(ParticleIndex, DefaultRot);
+			const FQuat4f SimRot = RotationReader.GetSafe(ParticleIndex, DefaultRot).GetNormalized();
 			const FVector Position = bUseLocalSpace ? LocalToWorld.TransformPosition(FVector(SimPos)) : LwcConverter.ConvertSimulationPositionToWorld(SimPos);
 			const FQuat Rotation = bUseLocalSpace ? LocalToWorld.TransformRotation(FQuat(SimRot)) : FQuat(SimRot);
 			const FVector Size = FVector(SizeReader.GetSafe(ParticleIndex, DefaultSize) * 0.5f);
@@ -241,7 +241,7 @@ FNiagaraDynamicDataBase* FNiagaraRendererDecals::GenerateDynamicData(const FNiag
 		if (bVisible && VisTag == RendererProperties->RendererVisibility)
 		{
 			const FVector3f SimPos = ParameterStore.GetParameterValueOrDefault(RendererProperties->PositionBinding.GetParamMapBindableVariable(), DefaultPos);
-			const FQuat4f SimRot = ParameterStore.GetParameterValueOrDefault(RendererProperties->DecalOrientationBinding.GetParamMapBindableVariable(), DefaultRot);
+			const FQuat4f SimRot = ParameterStore.GetParameterValueOrDefault(RendererProperties->DecalOrientationBinding.GetParamMapBindableVariable(), DefaultRot).GetNormalized();
 			const FVector Position = bUseLocalSpace ? LocalToWorld.TransformPosition(FVector(SimPos)) : LwcConverter.ConvertSimulationPositionToWorld(SimPos);
 			const FQuat Rotation = bUseLocalSpace ? LocalToWorld.TransformRotation(FQuat(SimRot)) : FQuat(SimRot);
 			const FVector Size = FVector(ParameterStore.GetParameterValueOrDefault(RendererProperties->DecalSizeBinding.GetParamMapBindableVariable(), DefaultSize) * 0.5f);
