Stack type ui revamp

diff --git a/Engine/Plugins/FX/Niagara/Source/Niagara/Public/NiagaraTypes.h b/Engine/Plugins/FX/Niagara/Source/Niagara/Public/NiagaraTypes.h
--- a/Engine/Plugins/FX/Niagara/Source/Niagara/Public/NiagaraTypes.h
+++ b/Engine/Plugins/FX/Niagara/Source/Niagara/Public/NiagaraTypes.h
@@ -11,6 +11,8 @@
 #include "Engine/Texture2D.h"
 #endif
 #include "NiagaraCore.h"
+#include "Math/UnitConversion.h"
+#include "Math/UnitConversion.inl"
 
 #include "NiagaraTypes.generated.h"
 
@@ -686,6 +688,10 @@ struct NIAGARA_API FNiagaraVariableMetaData
 
 	UPROPERTY(EditAnywhere, Category = "Variable", meta = (SkipForCompileHash = "true"))
 	FText CategoryName;
+	
+	/** The unit to display next to input fields for this parameter - note that this is only a visual indicator and does not change any of the calculations. */
+	UPROPERTY(EditAnywhere, Category = "Variable", meta = (SkipForCompileHash = "true"))
+	EUnit DisplayUnit = EUnit::Unspecified;
 
 	/** Declares that this input is advanced and should only be visible if expanded inputs have been expanded. */
 	UPROPERTY(EditAnywhere, Category = "Variable", meta = (SkipForCompileHash = "true"))
diff --git a/Engine/Plugins/FX/Niagara/Source/NiagaraEditor/Private/Customizations/NiagaraScriptVariableCustomization.cpp b/Engine/Plugins/FX/Niagara/Source/NiagaraEditor/Private/Customizations/NiagaraScriptVariableCustomization.cpp
--- a/Engine/Plugins/FX/Niagara/Source/NiagaraEditor/Private/Customizations/NiagaraScriptVariableCustomization.cpp
+++ b/Engine/Plugins/FX/Niagara/Source/NiagaraEditor/Private/Customizations/NiagaraScriptVariableCustomization.cpp
@@ -729,6 +729,7 @@ void FNiagaraScriptVariableHierarchyDetails::CustomizeDetails(IDetailLayoutBuild
 	DetailBuilder.HideProperty("DefaultValueVariant", UNiagaraScriptVariable::StaticClass());
 
 	DetailBuilder.HideProperty(GET_MEMBER_NAME_CHECKED(UNiagaraScriptVariable, Metadata.CategoryName));
+	DetailBuilder.HideProperty(GET_MEMBER_NAME_CHECKED(UNiagaraScriptVariable, Metadata.DisplayUnit));
 	DetailBuilder.HideProperty(GET_MEMBER_NAME_CHECKED(UNiagaraScriptVariable, Metadata.bAdvancedDisplay));
 	DetailBuilder.HideProperty(GET_MEMBER_NAME_CHECKED(UNiagaraScriptVariable, Metadata.bOverrideColor));
 	DetailBuilder.HideProperty(GET_MEMBER_NAME_CHECKED(UNiagaraScriptVariable, Metadata.bEnableBoolOverride));
diff --git a/Engine/Plugins/FX/Niagara/Source/NiagaraEditor/Private/NiagaraEditorStyle.cpp b/Engine/Plugins/FX/Niagara/Source/NiagaraEditor/Private/NiagaraEditorStyle.cpp
--- a/Engine/Plugins/FX/Niagara/Source/NiagaraEditor/Private/NiagaraEditorStyle.cpp
+++ b/Engine/Plugins/FX/Niagara/Source/NiagaraEditor/Private/NiagaraEditorStyle.cpp
@@ -14,7 +14,6 @@
 #include "Settings/EditorStyleSettings.h"
 #include "Styling/StarshipCoreStyle.h"
 #include "Styling/StyleColors.h"
-#include "Styling/ToolBarStyle.h"
 
 TSharedPtr<FNiagaraEditorStyle> FNiagaraEditorStyle::NiagaraEditorStyle = nullptr;
 
@@ -375,7 +374,7 @@ void FNiagaraEditorStyle::InitSelectedEmitter()
 void FNiagaraEditorStyle::InitToolbarIcons()
 {
 	Set("NiagaraEditor.Refresh", new CORE_IMAGE_BRUSH_SVG("Starship/Common/Update", Icon20x20));
-	Set("NiagaraEditor.ApplyScratchPadChanges", new IMAGE_BRUSH_SVG("Icons/Commands/ApplyScratch", Icon40x40));		
+	Set("NiagaraEditor.ApplyScratchPadChanges", new IMAGE_BRUSH_SVG("Icons/Commands/ApplyScratch", Icon40x40));
 	Set("NiagaraEditor.OverviewNode.IsolatedColor", FLinearColor::Yellow);
 	Set("NiagaraEditor.OverviewNode.NotIsolatedColor", FLinearColor::Transparent);
 }
@@ -388,6 +387,13 @@ void FNiagaraEditorStyle::InitIcons()
 	Set("NiagaraEditor.Module.AddPin", new IMAGE_BRUSH("Icons/PlusSymbol_12x", Icon12x12, FLinearColor::Gray));
 	Set("NiagaraEditor.Module.RemovePin", new IMAGE_BRUSH("Icons/MinusSymbol_12x", Icon12x12, FLinearColor::Gray));
 	Set("NiagaraEditor.Message.CustomNote", new IMAGE_BRUSH("Icons/icon_custom_note_16x", Icon16x16));
+	Set("NiagaraEditor.Module.DynamicInput", new IMAGE_BRUSH_SVG("Icons/DynamicInput", Icon16x16));
+	Set("NiagaraEditor.Module.TypeIconPill", new IMAGE_BRUSH_SVG("Icons/TypeIconPill", FVector2D(5.0f, 16.0f)));
+	
+	FSlateBrush* SlateBrush = new FSlateBrush();
+	*SlateBrush = FAppStyle::Get().GetWidgetStyle<FEditableTextBoxStyle>("NormalEditableTextBox").BackgroundImageNormal;
+	SlateBrush->TintColor = FLinearColor(0.01f, 0.01f, 0.01f);
+	Set("NiagaraEditor.Module.InputTypeBorder", SlateBrush);
 }
 
 void FNiagaraEditorStyle::InitOverview()
diff --git a/Engine/Plugins/FX/Niagara/Source/NiagaraEditor/Private/TypeEditorUtilities/NiagaraBoolTypeEditorUtilities.cpp b/Engine/Plugins/FX/Niagara/Source/NiagaraEditor/Private/TypeEditorUtilities/NiagaraBoolTypeEditorUtilities.cpp
--- a/Engine/Plugins/FX/Niagara/Source/NiagaraEditor/Private/TypeEditorUtilities/NiagaraBoolTypeEditorUtilities.cpp
+++ b/Engine/Plugins/FX/Niagara/Source/NiagaraEditor/Private/TypeEditorUtilities/NiagaraBoolTypeEditorUtilities.cpp
@@ -59,7 +59,7 @@ private:
 	bool bBoolValue;
 };
 
-TSharedPtr<SNiagaraParameterEditor> FNiagaraEditorBoolTypeUtilities::CreateParameterEditor(const FNiagaraTypeDefinition& ParameterType) const
+TSharedPtr<SNiagaraParameterEditor> FNiagaraEditorBoolTypeUtilities::CreateParameterEditor(const FNiagaraTypeDefinition& ParameterType, EUnit DisplayUnit) const
 {
 	return SNew(SNiagaraBoolParameterEditor);
 }
diff --git a/Engine/Plugins/FX/Niagara/Source/NiagaraEditor/Private/TypeEditorUtilities/NiagaraBoolTypeEditorUtilities.h b/Engine/Plugins/FX/Niagara/Source/NiagaraEditor/Private/TypeEditorUtilities/NiagaraBoolTypeEditorUtilities.h
--- a/Engine/Plugins/FX/Niagara/Source/NiagaraEditor/Private/TypeEditorUtilities/NiagaraBoolTypeEditorUtilities.h
+++ b/Engine/Plugins/FX/Niagara/Source/NiagaraEditor/Private/TypeEditorUtilities/NiagaraBoolTypeEditorUtilities.h
@@ -12,7 +12,7 @@ class FNiagaraEditorBoolTypeUtilities : public FNiagaraEditorTypeUtilities
 public:
 	//~ INiagaraEditorTypeUtilities interface.
 	virtual bool CanCreateParameterEditor() const override { return true; }
-	virtual TSharedPtr<SNiagaraParameterEditor> CreateParameterEditor(const FNiagaraTypeDefinition& ParameterType) const override;
+	virtual TSharedPtr<SNiagaraParameterEditor> CreateParameterEditor(const FNiagaraTypeDefinition& ParameterType, EUnit DisplayUnit) const override;
 	virtual bool CanHandlePinDefaults() const override;
 	virtual FString GetPinDefaultStringFromValue(const FNiagaraVariable& AllocatedVariable) const override;
 	virtual bool SetValueFromPinDefaultString(const FString& StringValue, FNiagaraVariable& Variable) const override;
diff --git a/Engine/Plugins/FX/Niagara/Source/NiagaraEditor/Private/TypeEditorUtilities/NiagaraColorTypeEditorUtilities.cpp b/Engine/Plugins/FX/Niagara/Source/NiagaraEditor/Private/TypeEditorUtilities/NiagaraColorTypeEditorUtilities.cpp
--- a/Engine/Plugins/FX/Niagara/Source/NiagaraEditor/Private/TypeEditorUtilities/NiagaraColorTypeEditorUtilities.cpp
+++ b/Engine/Plugins/FX/Niagara/Source/NiagaraEditor/Private/TypeEditorUtilities/NiagaraColorTypeEditorUtilities.cpp
@@ -202,7 +202,7 @@ void FNiagaraEditorColorTypeUtilities::UpdateVariableWithDefaultValue(FNiagaraVa
 	Variable.SetValue<FLinearColor>(FLinearColor(1, 1, 1, 1));
 }
 
-TSharedPtr<SNiagaraParameterEditor> FNiagaraEditorColorTypeUtilities::CreateParameterEditor(const FNiagaraTypeDefinition& ParameterType) const
+TSharedPtr<SNiagaraParameterEditor> FNiagaraEditorColorTypeUtilities::CreateParameterEditor(const FNiagaraTypeDefinition& ParameterType, EUnit DisplayUnit) const
 {
 	return SNew(SNiagaraColorParameterEditor);
 }
diff --git a/Engine/Plugins/FX/Niagara/Source/NiagaraEditor/Private/TypeEditorUtilities/NiagaraColorTypeEditorUtilities.h b/Engine/Plugins/FX/Niagara/Source/NiagaraEditor/Private/TypeEditorUtilities/NiagaraColorTypeEditorUtilities.h
--- a/Engine/Plugins/FX/Niagara/Source/NiagaraEditor/Private/TypeEditorUtilities/NiagaraColorTypeEditorUtilities.h
+++ b/Engine/Plugins/FX/Niagara/Source/NiagaraEditor/Private/TypeEditorUtilities/NiagaraColorTypeEditorUtilities.h
@@ -14,7 +14,7 @@ public:
 	virtual bool CanProvideDefaultValue() const override { return true; }
 	virtual void UpdateVariableWithDefaultValue(FNiagaraVariable& Variable) const override;
 	virtual bool CanCreateParameterEditor() const override { return true; }
-	virtual TSharedPtr<SNiagaraParameterEditor> CreateParameterEditor(const FNiagaraTypeDefinition& ParameterType) const override;
+	virtual TSharedPtr<SNiagaraParameterEditor> CreateParameterEditor(const FNiagaraTypeDefinition& ParameterType, EUnit DisplayUnit) const override;
 	virtual bool CanHandlePinDefaults() const override;
 	virtual FString GetPinDefaultStringFromValue(const FNiagaraVariable& AllocatedVariable) const override;
 	virtual bool SetValueFromPinDefaultString(const FString& StringValue, FNiagaraVariable& Variable) const override;
diff --git a/Engine/Plugins/FX/Niagara/Source/NiagaraEditor/Private/TypeEditorUtilities/NiagaraEnumTypeEditorUtilities.cpp b/Engine/Plugins/FX/Niagara/Source/NiagaraEditor/Private/TypeEditorUtilities/NiagaraEnumTypeEditorUtilities.cpp
--- a/Engine/Plugins/FX/Niagara/Source/NiagaraEditor/Private/TypeEditorUtilities/NiagaraEnumTypeEditorUtilities.cpp
+++ b/Engine/Plugins/FX/Niagara/Source/NiagaraEditor/Private/TypeEditorUtilities/NiagaraEnumTypeEditorUtilities.cpp
@@ -76,7 +76,7 @@ void FNiagaraEditorEnumTypeUtilities::UpdateVariableWithDefaultValue(FNiagaraVar
 	Variable.SetValue<FNiagaraInt32>(EnumIntValue);
 }
 
-TSharedPtr<SNiagaraParameterEditor> FNiagaraEditorEnumTypeUtilities::CreateParameterEditor(const FNiagaraTypeDefinition& ParameterType) const
+TSharedPtr<SNiagaraParameterEditor> FNiagaraEditorEnumTypeUtilities::CreateParameterEditor(const FNiagaraTypeDefinition& ParameterType, EUnit DisplayUnit) const
 {
 	return SNew(SNiagaraEnumParameterEditor, ParameterType.GetEnum());
 }
diff --git a/Engine/Plugins/FX/Niagara/Source/NiagaraEditor/Private/TypeEditorUtilities/NiagaraEnumTypeEditorUtilities.h b/Engine/Plugins/FX/Niagara/Source/NiagaraEditor/Private/TypeEditorUtilities/NiagaraEnumTypeEditorUtilities.h
--- a/Engine/Plugins/FX/Niagara/Source/NiagaraEditor/Private/TypeEditorUtilities/NiagaraEnumTypeEditorUtilities.h
+++ b/Engine/Plugins/FX/Niagara/Source/NiagaraEditor/Private/TypeEditorUtilities/NiagaraEnumTypeEditorUtilities.h
@@ -13,7 +13,7 @@ public:
 	virtual bool CanProvideDefaultValue() const override;
 	virtual void UpdateVariableWithDefaultValue(FNiagaraVariable& Variable) const override;
 	virtual bool CanCreateParameterEditor() const override { return true; }
-	virtual TSharedPtr<SNiagaraParameterEditor> CreateParameterEditor(const FNiagaraTypeDefinition& ParameterType) const override;
+	virtual TSharedPtr<SNiagaraParameterEditor> CreateParameterEditor(const FNiagaraTypeDefinition& ParameterType, EUnit DisplayUnit) const override;
 	virtual bool CanHandlePinDefaults() const override;
 	virtual FString GetPinDefaultStringFromValue(const FNiagaraVariable& AllocatedVariable) const override;
 	virtual bool SetValueFromPinDefaultString(const FString& StringValue, FNiagaraVariable& Variable) const override;
diff --git a/Engine/Plugins/FX/Niagara/Source/NiagaraEditor/Private/TypeEditorUtilities/NiagaraFloatTypeEditorUtilities.cpp b/Engine/Plugins/FX/Niagara/Source/NiagaraEditor/Private/TypeEditorUtilities/NiagaraFloatTypeEditorUtilities.cpp
--- a/Engine/Plugins/FX/Niagara/Source/NiagaraEditor/Private/TypeEditorUtilities/NiagaraFloatTypeEditorUtilities.cpp
+++ b/Engine/Plugins/FX/Niagara/Source/NiagaraEditor/Private/TypeEditorUtilities/NiagaraFloatTypeEditorUtilities.cpp
@@ -1,11 +1,13 @@
 // Copyright Epic Games, Inc. All Rights Reserved.
 
 #include "NiagaraFloatTypeEditorUtilities.h"
+
+#include "GraphEditorSettings.h"
 #include "SNiagaraParameterEditor.h"
 #include "NiagaraTypes.h"
 #include "NiagaraEditorStyle.h"
-
-#include "Widgets/Input/SSpinBox.h"
+#include "Math/UnitConversion.h"
+#include "Widgets/Input/SNumericEntryBox.h"
 
 class SNiagaraFloatParameterEditor : public SNiagaraParameterEditor
 {
@@ -13,26 +15,34 @@ public:
 	SLATE_BEGIN_ARGS(SNiagaraFloatParameterEditor) { }
 	SLATE_END_ARGS();
 
-	void Construct(const FArguments& InArgs)
+	void Construct(const FArguments& InArgs, EUnit DisplayUnit)
 	{
 		SNiagaraParameterEditor::Construct(SNiagaraParameterEditor::FArguments()
 			.MinimumDesiredWidth(DefaultInputSize)
 			.MaximumDesiredWidth(DefaultInputSize));
 
+		const UGraphEditorSettings* Settings = GetDefault<UGraphEditorSettings>();
 		ChildSlot
 		[
-			SNew(SSpinBox<float>)
+			SNew(SNumericEntryBox<float>)
 			.Font(FAppStyle::Get().GetFontStyle("PropertyWindow.NormalFont"))
 			.MinValue(TOptional<float>())
 			.MaxValue(TOptional<float>())
 			.MaxSliderValue(TOptional<float>())
 			.MinSliderValue(TOptional<float>())
-			.Delta(0.0f)
 			.Value(this, &SNiagaraFloatParameterEditor::GetValue)
 			.OnValueChanged(this, &SNiagaraFloatParameterEditor::ValueChanged)
 			.OnValueCommitted(this, &SNiagaraFloatParameterEditor::ValueCommitted)
 			.OnBeginSliderMovement(this, &SNiagaraFloatParameterEditor::BeginSliderMovement)
 			.OnEndSliderMovement(this, &SNiagaraFloatParameterEditor::EndSliderMovement)
+			.TypeInterface(GetTypeInterface<float>(DisplayUnit))
+			.AllowSpin(true)
+			.LabelPadding(FMargin(3) )
+			.LabelLocation(SNumericEntryBox<float>::ELabelLocation::Inside)
+			.Label()
+			[
+				SNumericEntryBox<int32>::BuildNarrowColorLabel(Settings->FloatPinTypeColor)
+			]
 		];
 	}
 
@@ -61,7 +71,7 @@ private:
 		ExecuteOnEndValueChange();
 	}
 
-	float GetValue() const
+	TOptional<float> GetValue() const
 	{
 		return FloatValue;
 	}
@@ -81,12 +91,12 @@ private:
 	}
 
 private:
-	float FloatValue;
+	float FloatValue = 0;
 };
 
-TSharedPtr<SNiagaraParameterEditor> FNiagaraEditorFloatTypeUtilities::CreateParameterEditor(const FNiagaraTypeDefinition& ParameterType) const
+TSharedPtr<SNiagaraParameterEditor> FNiagaraEditorFloatTypeUtilities::CreateParameterEditor(const FNiagaraTypeDefinition& ParameterType, EUnit DisplayUnit) const
 {
-	return SNew(SNiagaraFloatParameterEditor);
+	return SNew(SNiagaraFloatParameterEditor, DisplayUnit);
 }
 
 bool FNiagaraEditorFloatTypeUtilities::CanHandlePinDefaults() const
diff --git a/Engine/Plugins/FX/Niagara/Source/NiagaraEditor/Private/TypeEditorUtilities/NiagaraFloatTypeEditorUtilities.h b/Engine/Plugins/FX/Niagara/Source/NiagaraEditor/Private/TypeEditorUtilities/NiagaraFloatTypeEditorUtilities.h
--- a/Engine/Plugins/FX/Niagara/Source/NiagaraEditor/Private/TypeEditorUtilities/NiagaraFloatTypeEditorUtilities.h
+++ b/Engine/Plugins/FX/Niagara/Source/NiagaraEditor/Private/TypeEditorUtilities/NiagaraFloatTypeEditorUtilities.h
@@ -13,7 +13,7 @@ class FNiagaraEditorFloatTypeUtilities : public FNiagaraEditorTypeUtilities
 public:
 	//~ INiagaraEditorTypeUtilities interface.
 	virtual bool CanCreateParameterEditor() const override { return true; }
-	virtual TSharedPtr<SNiagaraParameterEditor> CreateParameterEditor(const FNiagaraTypeDefinition& ParameterType) const override;
+	virtual TSharedPtr<SNiagaraParameterEditor> CreateParameterEditor(const FNiagaraTypeDefinition& ParameterType, EUnit DisplayUnit) const override;
 	virtual bool CanHandlePinDefaults() const override;
 	virtual FString GetPinDefaultStringFromValue(const FNiagaraVariable& AllocatedVariable) const override;
 	virtual bool SetValueFromPinDefaultString(const FString& StringValue, FNiagaraVariable& Variable) const override;
diff --git a/Engine/Plugins/FX/Niagara/Source/NiagaraEditor/Private/TypeEditorUtilities/NiagaraIntegerTypeEditorUtilities.cpp b/Engine/Plugins/FX/Niagara/Source/NiagaraEditor/Private/TypeEditorUtilities/NiagaraIntegerTypeEditorUtilities.cpp
--- a/Engine/Plugins/FX/Niagara/Source/NiagaraEditor/Private/TypeEditorUtilities/NiagaraIntegerTypeEditorUtilities.cpp
+++ b/Engine/Plugins/FX/Niagara/Source/NiagaraEditor/Private/TypeEditorUtilities/NiagaraIntegerTypeEditorUtilities.cpp
@@ -1,12 +1,13 @@
 // Copyright Epic Games, Inc. All Rights Reserved.
 
 #include "NiagaraIntegerTypeEditorUtilities.h"
+
+#include "GraphEditorSettings.h"
 #include "SNiagaraParameterEditor.h"
 #include "NiagaraTypes.h"
 #include "NiagaraEditorStyle.h"
 
-#include "Widgets/Input/SSpinBox.h"
-#include "Widgets/Input/SComboBox.h"
+#include "Widgets/Input/SNumericEntryBox.h"
 
 class SNiagaraIntegerParameterEditor : public SNiagaraParameterEditor
 {
@@ -14,26 +15,34 @@ public:
 	SLATE_BEGIN_ARGS(SNiagaraIntegerParameterEditor) { }
 	SLATE_END_ARGS();
 
-	void Construct(const FArguments& InArgs)
+	void Construct(const FArguments& InArgs, EUnit DisplayUnit)
 	{
 		SNiagaraParameterEditor::Construct(SNiagaraParameterEditor::FArguments()
 			.MinimumDesiredWidth(DefaultInputSize)
 			.MaximumDesiredWidth(DefaultInputSize));
 
+		const UGraphEditorSettings* Settings = GetDefault<UGraphEditorSettings>();
 		ChildSlot
 		[
-			SNew(SSpinBox<int32>)
+			SNew(SNumericEntryBox<int32>)
 			.Font(FAppStyle::Get().GetFontStyle("PropertyWindow.NormalFont"))
 			.MinValue(TOptional<int32>())
 			.MaxValue(TOptional<int32>())
 			.MaxSliderValue(TOptional<int32>())
 			.MinSliderValue(TOptional<int32>())
-			.Delta(0.001f)
 			.Value(this, &SNiagaraIntegerParameterEditor::GetValue)
 			.OnValueChanged(this, &SNiagaraIntegerParameterEditor::ValueChanged)
 			.OnValueCommitted(this, &SNiagaraIntegerParameterEditor::ValueCommitted)
 			.OnBeginSliderMovement(this, &SNiagaraIntegerParameterEditor::BeginSliderMovement)
 			.OnEndSliderMovement(this, &SNiagaraIntegerParameterEditor::EndSliderMovement)
+			.TypeInterface(GetTypeInterface<int32>(DisplayUnit))
+			.AllowSpin(true)
+			.LabelPadding(FMargin(3))
+			.LabelLocation(SNumericEntryBox<int32>::ELabelLocation::Inside)
+			.Label()
+			[
+				SNumericEntryBox<int32>::BuildNarrowColorLabel(Settings->IntPinTypeColor)
+			]
 		];
 	}
 
@@ -62,7 +71,7 @@ private:
 		ExecuteOnEndValueChange();
 	}
 
-	int32 GetValue() const
+	TOptional<int32> GetValue() const
 	{
 		return IntValue;
 	}
@@ -82,12 +91,12 @@ private:
 	}
 
 private:
-	int32 IntValue;
+	int32 IntValue = 0;
 };
 
-TSharedPtr<SNiagaraParameterEditor> FNiagaraEditorIntegerTypeUtilities::CreateParameterEditor(const FNiagaraTypeDefinition& ParameterType) const
+TSharedPtr<SNiagaraParameterEditor> FNiagaraEditorIntegerTypeUtilities::CreateParameterEditor(const FNiagaraTypeDefinition& ParameterType, EUnit DisplayUnit) const
 {
-	return SNew(SNiagaraIntegerParameterEditor);
+	return SNew(SNiagaraIntegerParameterEditor, DisplayUnit);
 }
 
 bool FNiagaraEditorIntegerTypeUtilities::CanHandlePinDefaults() const
diff --git a/Engine/Plugins/FX/Niagara/Source/NiagaraEditor/Private/TypeEditorUtilities/NiagaraIntegerTypeEditorUtilities.h b/Engine/Plugins/FX/Niagara/Source/NiagaraEditor/Private/TypeEditorUtilities/NiagaraIntegerTypeEditorUtilities.h
--- a/Engine/Plugins/FX/Niagara/Source/NiagaraEditor/Private/TypeEditorUtilities/NiagaraIntegerTypeEditorUtilities.h
+++ b/Engine/Plugins/FX/Niagara/Source/NiagaraEditor/Private/TypeEditorUtilities/NiagaraIntegerTypeEditorUtilities.h
@@ -11,7 +11,7 @@ class FNiagaraEditorIntegerTypeUtilities : public FNiagaraEditorTypeUtilities
 public:
 	//~ INiagaraEditorTypeUtilities interface.
 	virtual bool CanCreateParameterEditor() const override { return true; }
-	virtual TSharedPtr<SNiagaraParameterEditor> CreateParameterEditor(const FNiagaraTypeDefinition& ParameterType) const override;
+	virtual TSharedPtr<SNiagaraParameterEditor> CreateParameterEditor(const FNiagaraTypeDefinition& ParameterType, EUnit DisplayUnit) const override;
 	virtual bool CanHandlePinDefaults() const override;
 	virtual FString GetPinDefaultStringFromValue(const FNiagaraVariable& AllocatedVariable) const override;
 	virtual bool SetValueFromPinDefaultString(const FString& StringValue, FNiagaraVariable& Variable) const override;
diff --git a/Engine/Plugins/FX/Niagara/Source/NiagaraEditor/Private/TypeEditorUtilities/NiagaraVectorTypeEditorUtilities.cpp b/Engine/Plugins/FX/Niagara/Source/NiagaraEditor/Private/TypeEditorUtilities/NiagaraVectorTypeEditorUtilities.cpp
--- a/Engine/Plugins/FX/Niagara/Source/NiagaraEditor/Private/TypeEditorUtilities/NiagaraVectorTypeEditorUtilities.cpp
+++ b/Engine/Plugins/FX/Niagara/Source/NiagaraEditor/Private/TypeEditorUtilities/NiagaraVectorTypeEditorUtilities.cpp
@@ -15,24 +15,19 @@ public:
 		SLATE_ARGUMENT(int32, ComponentCount)
 	SLATE_END_ARGS();
 
-	void Construct(const FArguments& InArgs)
+	void Construct(const FArguments& InArgs, EUnit DisplayUnit)
 	{
 		SNiagaraParameterEditor::Construct(SNiagaraParameterEditor::FArguments()
 			.MinimumDesiredWidth(DefaultInputSize * InArgs._ComponentCount)
 			.MaximumDesiredWidth(DefaultInputSize * InArgs._ComponentCount));
 
-		ComponentLabels.Add(NSLOCTEXT("VectorParameterEditor", "XLabel", "X"));
-		ComponentLabels.Add(NSLOCTEXT("VectorParameterEditor", "YLabel", "Y"));
-		ComponentLabels.Add(NSLOCTEXT("VectorParameterEditor", "ZLabel", "Z"));
-		ComponentLabels.Add(NSLOCTEXT("VectorParameterEditor", "WLabel", "W"));
-
 		TSharedRef<SHorizontalBox> ComponentBox = SNew(SHorizontalBox);
 		for (int32 ComponentIndex = 0; ComponentIndex < InArgs._ComponentCount; ++ComponentIndex)
 		{
 			ComponentBox->AddSlot()
 			.Padding(ComponentIndex == 0 ? 0 : 3, 0, 0, 0)
 			[
-				ConstructComponentWidget(ComponentIndex)
+				ConstructComponentWidget(ComponentIndex, DisplayUnit)
 			];
 		}
 
@@ -50,8 +45,17 @@ protected:
 	virtual void SetValue(int32 Index, float Value) = 0;
 
 private:
-	TSharedRef<SWidget> ConstructComponentWidget(int32 Index)
-	{
+	TSharedRef<SWidget> ConstructComponentWidget(int32 Index, EUnit DisplayUnit)
+	{
+		const static TArray Colors = {
+			&SNumericEntryBox<int32>::RedLabelBackgroundColor,
+			&SNumericEntryBox<int32>::GreenLabelBackgroundColor,
+			&SNumericEntryBox<int32>::BlueLabelBackgroundColor,
+			&FLinearColor::White
+		};
+		
+		TSharedRef<SWidget>  LabelWidget = SNumericEntryBox<float>::BuildNarrowColorLabel(*Colors[Index]);
+		
 		return SNew(SNumericEntryBox<float>)
 		.Font(FAppStyle::Get().GetFontStyle("PropertyWindow.NormalFont"))
 		.MinValue(TOptional<float>())
@@ -64,13 +68,13 @@ private:
 		.OnValueCommitted(this, &SNiagaraVectorParameterEditorBase::ValueCommitted, Index)
 		.OnBeginSliderMovement(this, &SNiagaraVectorParameterEditorBase::BeginSliderMovement)
 		.OnEndSliderMovement(this, &SNiagaraVectorParameterEditorBase::EndSliderMovement)
+		.TypeInterface(GetTypeInterface<float>(DisplayUnit))
 		.AllowSpin(true)
-		.LabelVAlign(EVerticalAlignment::VAlign_Center)
+		.LabelPadding( FMargin(3) )
+		.LabelLocation( SNumericEntryBox<float>::ELabelLocation::Inside )
 		.Label()
 		[
-			SNew(STextBlock)
-			.TextStyle(FNiagaraEditorStyle::Get(), "NiagaraEditor.ParameterText")
-			.Text(ComponentLabels[Index])
+			LabelWidget
 		];
 	}
 
@@ -102,9 +106,6 @@ private:
 			ValueChanged(Value, Index);
 		}
 	}
-
-private:
-	TArray<FText> ComponentLabels;
 };
 
 
@@ -114,11 +115,11 @@ public:
 	SLATE_BEGIN_ARGS(SNiagaraVector2ParameterEditor) { }
 	SLATE_END_ARGS();
 
-	void Construct(const FArguments& InArgs)
+	void Construct(const FArguments& InArgs, EUnit DisplayUnit)
 	{
 		SNiagaraVectorParameterEditorBase::Construct(
 			SNiagaraVectorParameterEditorBase::FArguments()
-			.ComponentCount(2));
+			.ComponentCount(2), DisplayUnit);
 	}
 
 	virtual void UpdateInternalValueFromStruct(TSharedRef<FStructOnScope> Struct) override
@@ -148,9 +149,9 @@ private:
 	FVector2f VectorValue;
 };
 
-TSharedPtr<SNiagaraParameterEditor> FNiagaraEditorVector2TypeUtilities::CreateParameterEditor(const FNiagaraTypeDefinition& ParameterType) const
+TSharedPtr<SNiagaraParameterEditor> FNiagaraEditorVector2TypeUtilities::CreateParameterEditor(const FNiagaraTypeDefinition& ParameterType, EUnit DisplayUnit) const
 {
-	return SNew(SNiagaraVector2ParameterEditor);
+	return SNew(SNiagaraVector2ParameterEditor, DisplayUnit);
 }
 
 bool FNiagaraEditorVector2TypeUtilities::CanHandlePinDefaults() const
@@ -192,11 +193,11 @@ public:
 	SLATE_BEGIN_ARGS(SNiagaraVector3ParameterEditor) { }
 	SLATE_END_ARGS();
 
-	void Construct(const FArguments& InArgs)
+	void Construct(const FArguments& InArgs, EUnit DisplayUnit)
 	{
 		SNiagaraVectorParameterEditorBase::Construct(
 			SNiagaraVectorParameterEditorBase::FArguments()
-			.ComponentCount(3));
+			.ComponentCount(3), DisplayUnit);
 	}
 
 	virtual void UpdateInternalValueFromStruct(TSharedRef<FStructOnScope> Struct) override
@@ -226,9 +227,9 @@ private:
 	FVector3f VectorValue;
 };
 
-TSharedPtr<SNiagaraParameterEditor> FNiagaraEditorVector3TypeUtilities::CreateParameterEditor(const FNiagaraTypeDefinition& ParameterType) const
+TSharedPtr<SNiagaraParameterEditor> FNiagaraEditorVector3TypeUtilities::CreateParameterEditor(const FNiagaraTypeDefinition& ParameterType, EUnit DisplayUnit) const
 {
-	return SNew(SNiagaraVector3ParameterEditor);
+	return SNew(SNiagaraVector3ParameterEditor, DisplayUnit);
 }
 
 bool FNiagaraEditorVector3TypeUtilities::CanHandlePinDefaults() const
@@ -274,11 +275,11 @@ public:
 	SLATE_BEGIN_ARGS(SNiagaraVector4ParameterEditor) { }
 	SLATE_END_ARGS();
 
-	void Construct(const FArguments& InArgs)
+	void Construct(const FArguments& InArgs, EUnit DisplayUnit)
 	{
 		SNiagaraVectorParameterEditorBase::Construct(
 			SNiagaraVectorParameterEditorBase::FArguments()
-			.ComponentCount(4));
+			.ComponentCount(4), DisplayUnit);
 	}
 
 	virtual void UpdateInternalValueFromStruct(TSharedRef<FStructOnScope> Struct) override
@@ -308,9 +309,9 @@ private:
 	FVector4f VectorValue;
 };
 
-TSharedPtr<SNiagaraParameterEditor> FNiagaraEditorVector4TypeUtilities::CreateParameterEditor(const FNiagaraTypeDefinition& ParameterType) const
+TSharedPtr<SNiagaraParameterEditor> FNiagaraEditorVector4TypeUtilities::CreateParameterEditor(const FNiagaraTypeDefinition& ParameterType, EUnit DisplayUnit) const
 {
-	return SNew(SNiagaraVector4ParameterEditor);
+	return SNew(SNiagaraVector4ParameterEditor, DisplayUnit);
 }
 
 bool FNiagaraEditorVector4TypeUtilities::CanHandlePinDefaults() const
@@ -356,11 +357,11 @@ public:
 	SLATE_BEGIN_ARGS(SNiagaraQuatParameterEditor) { }
 	SLATE_END_ARGS();
 
-	void Construct(const FArguments& InArgs)
+	void Construct(const FArguments& InArgs, EUnit DisplayUnit)
 	{
 		SNiagaraVectorParameterEditorBase::Construct(
 			SNiagaraVectorParameterEditorBase::FArguments()
-			.ComponentCount(4));
+			.ComponentCount(4), DisplayUnit);
 	}
 
 	virtual void UpdateInternalValueFromStruct(TSharedRef<FStructOnScope> Struct) override
@@ -419,9 +420,9 @@ private:
 	FQuat4f VectorValue;
 };
 
-TSharedPtr<SNiagaraParameterEditor> FNiagaraEditorQuatTypeUtilities::CreateParameterEditor(const FNiagaraTypeDefinition& ParameterType) const
+TSharedPtr<SNiagaraParameterEditor> FNiagaraEditorQuatTypeUtilities::CreateParameterEditor(const FNiagaraTypeDefinition& ParameterType, EUnit DisplayUnit) const
 {
-	return SNew(SNiagaraQuatParameterEditor);
+	return SNew(SNiagaraQuatParameterEditor, DisplayUnit);
 }
 
 bool FNiagaraEditorQuatTypeUtilities::CanHandlePinDefaults() const
diff --git a/Engine/Plugins/FX/Niagara/Source/NiagaraEditor/Private/TypeEditorUtilities/NiagaraVectorTypeEditorUtilities.h b/Engine/Plugins/FX/Niagara/Source/NiagaraEditor/Private/TypeEditorUtilities/NiagaraVectorTypeEditorUtilities.h
--- a/Engine/Plugins/FX/Niagara/Source/NiagaraEditor/Private/TypeEditorUtilities/NiagaraVectorTypeEditorUtilities.h
+++ b/Engine/Plugins/FX/Niagara/Source/NiagaraEditor/Private/TypeEditorUtilities/NiagaraVectorTypeEditorUtilities.h
@@ -12,7 +12,7 @@ class FNiagaraEditorVector2TypeUtilities : public FNiagaraEditorTypeUtilities
 public:
 	//~ INiagaraEditorTypeUtilities interface.
 	virtual bool CanCreateParameterEditor() const override { return true; }
-	virtual TSharedPtr<SNiagaraParameterEditor> CreateParameterEditor(const FNiagaraTypeDefinition& ParameterType) const override;
+	virtual TSharedPtr<SNiagaraParameterEditor> CreateParameterEditor(const FNiagaraTypeDefinition& ParameterType, EUnit DisplayUnit) const override;
 	virtual bool CanHandlePinDefaults() const override;
 	virtual FString GetPinDefaultStringFromValue(const FNiagaraVariable& AllocatedVariable) const override;
 	virtual bool SetValueFromPinDefaultString(const FString& StringValue, FNiagaraVariable& Variable) const override;
@@ -26,7 +26,7 @@ class FNiagaraEditorVector3TypeUtilities : public FNiagaraEditorTypeUtilities
 public:
 	//~ INiagaraEditorTypeUtilities interface.
 	virtual bool CanCreateParameterEditor() const override { return true; }
-	virtual TSharedPtr<SNiagaraParameterEditor> CreateParameterEditor(const FNiagaraTypeDefinition& ParameterType) const override;
+	virtual TSharedPtr<SNiagaraParameterEditor> CreateParameterEditor(const FNiagaraTypeDefinition& ParameterType, EUnit DisplayUnit) const override;
 	virtual bool CanHandlePinDefaults() const override;
 	virtual FString GetPinDefaultStringFromValue(const FNiagaraVariable& AllocatedVariable) const override;
 	virtual bool SetValueFromPinDefaultString(const FString& StringValue, FNiagaraVariable& Variable) const override;
@@ -40,7 +40,7 @@ class FNiagaraEditorVector4TypeUtilities : public FNiagaraEditorTypeUtilities
 public:
 	//~ INiagaraEditorTypeUtilities interface.
 	virtual bool CanCreateParameterEditor() const override { return true; }
-	virtual TSharedPtr<SNiagaraParameterEditor> CreateParameterEditor(const FNiagaraTypeDefinition& ParameterType) const override;
+	virtual TSharedPtr<SNiagaraParameterEditor> CreateParameterEditor(const FNiagaraTypeDefinition& ParameterType, EUnit DisplayUnit) const override;
 	virtual bool CanHandlePinDefaults() const override;
 	virtual FString GetPinDefaultStringFromValue(const FNiagaraVariable& AllocatedVariable) const override;
 	virtual bool SetValueFromPinDefaultString(const FString& StringValue, FNiagaraVariable& Variable) const override;
@@ -56,7 +56,7 @@ public:
 	virtual bool CanProvideDefaultValue() const override { return true; }
 	virtual void UpdateVariableWithDefaultValue(FNiagaraVariable& Variable) const override;
 	virtual bool CanCreateParameterEditor() const override { return true; }
-	virtual TSharedPtr<SNiagaraParameterEditor> CreateParameterEditor(const FNiagaraTypeDefinition& ParameterType) const override;
+	virtual TSharedPtr<SNiagaraParameterEditor> CreateParameterEditor(const FNiagaraTypeDefinition& ParameterType, EUnit DisplayUnit) const override;
 	virtual bool CanHandlePinDefaults() const override;
 	virtual FString GetPinDefaultStringFromValue(const FNiagaraVariable& AllocatedVariable) const override;
 	virtual bool SetValueFromPinDefaultString(const FString& StringValue, FNiagaraVariable& Variable) const override;
diff --git a/Engine/Plugins/FX/Niagara/Source/NiagaraEditor/Private/ViewModels/Stack/NiagaraStackFunctionInput.cpp b/Engine/Plugins/FX/Niagara/Source/NiagaraEditor/Private/ViewModels/Stack/NiagaraStackFunctionInput.cpp
--- a/Engine/Plugins/FX/Niagara/Source/NiagaraEditor/Private/ViewModels/Stack/NiagaraStackFunctionInput.cpp
+++ b/Engine/Plugins/FX/Niagara/Source/NiagaraEditor/Private/ViewModels/Stack/NiagaraStackFunctionInput.cpp
@@ -232,6 +232,11 @@ const FNiagaraTypeDefinition& UNiagaraStackFunctionInput::GetInputType() const
 	return InputType;
 }
 
+EUnit UNiagaraStackFunctionInput::GetInputDisplayUnit() const
+{
+	return InputMetaData.IsSet() ? InputMetaData->DisplayUnit : EUnit::Unspecified;
+}
+
 FText UNiagaraStackFunctionInput::GetTooltipText() const
 {
 	FText Description = InputMetaData.IsSet() ? InputMetaData->Description : FText::GetEmpty();
diff --git a/Engine/Plugins/FX/Niagara/Source/NiagaraEditor/Public/INiagaraEditorTypeUtilities.h b/Engine/Plugins/FX/Niagara/Source/NiagaraEditor/Public/INiagaraEditorTypeUtilities.h
--- a/Engine/Plugins/FX/Niagara/Source/NiagaraEditor/Public/INiagaraEditorTypeUtilities.h
+++ b/Engine/Plugins/FX/Niagara/Source/NiagaraEditor/Public/INiagaraEditorTypeUtilities.h
@@ -4,6 +4,7 @@
 #include "Templates/SharedPointer.h"
 #include "Delegates/Delegate.h"
 #include "Internationalization/Text.h"
+#include "Math/UnitConversion.h"
 
 class FStructOnScope;
 class SNiagaraParameterEditor;
@@ -24,7 +25,7 @@ public:
 
 	virtual bool CanCreateParameterEditor() const = 0;
 
-	virtual TSharedPtr<SNiagaraParameterEditor> CreateParameterEditor(const FNiagaraTypeDefinition& ParameterType) const = 0;
+	virtual TSharedPtr<SNiagaraParameterEditor> CreateParameterEditor(const FNiagaraTypeDefinition& ParameterType, EUnit DisplayUnit = EUnit::Unspecified) const = 0;
 
 	virtual bool CanCreateDataInterfaceEditor() const = 0;
 
@@ -54,7 +55,7 @@ public:
 	virtual bool CanProvideDefaultValue() const override { return false; }
 	virtual void UpdateVariableWithDefaultValue(FNiagaraVariable& Variable) const override { }
 	virtual bool CanCreateParameterEditor() const override { return false; }
-	virtual TSharedPtr<SNiagaraParameterEditor> CreateParameterEditor(const FNiagaraTypeDefinition& ParameterType) const override { return TSharedPtr<SNiagaraParameterEditor>(); }
+	virtual TSharedPtr<SNiagaraParameterEditor> CreateParameterEditor(const FNiagaraTypeDefinition& ParameterType, EUnit DisplayUnit) const override { return TSharedPtr<SNiagaraParameterEditor>(); }
 	virtual bool CanCreateDataInterfaceEditor() const override { return false; };
 	virtual TSharedPtr<SWidget> CreateDataInterfaceEditor(UObject* DataInterface, FNotifyValueChanged DataInterfaceChangedHandler) const override { return TSharedPtr<SWidget>(); }
 	virtual bool CanHandlePinDefaults() const override { return false; }
diff --git a/Engine/Plugins/FX/Niagara/Source/NiagaraEditor/Public/SNiagaraParameterEditor.h b/Engine/Plugins/FX/Niagara/Source/NiagaraEditor/Public/SNiagaraParameterEditor.h
--- a/Engine/Plugins/FX/Niagara/Source/NiagaraEditor/Public/SNiagaraParameterEditor.h
+++ b/Engine/Plugins/FX/Niagara/Source/NiagaraEditor/Public/SNiagaraParameterEditor.h
@@ -4,6 +4,7 @@
 #include "Widgets/SCompoundWidget.h"
 #include "Widgets/DeclarativeSyntaxSupport.h"
 #include "UObject/StructOnScope.h"
+#include "Widgets/Input/NumericUnitTypeInterface.inl"
 
 /** Base class for inline parameter editors. These editors are expected to maintain an internal value which
 	is populated from a parameter struct. */
@@ -77,6 +78,16 @@ protected:
 	/** Executes the OnValueChanged delegate. */
 	void ExecuteOnValueChanged();
 
+	template<typename NumericType>
+	static TSharedPtr<TNumericUnitTypeInterface<NumericType>> GetTypeInterface(EUnit TypeUnit)
+	{
+		if (FUnitConversion::Settings().ShouldDisplayUnits())
+		{
+			return MakeShareable(new TNumericUnitTypeInterface<NumericType>(TypeUnit));
+		}
+		return MakeShareable(new TNumericUnitTypeInterface<NumericType>(EUnit::Unspecified));
+	}
+
 protected:
 	static const float DefaultInputSize;
 
@@ -109,4 +120,5 @@ private:
 
 	/** Sets the desired horizontal alignment of this parameter editor for it's parent container. */
 	EVerticalAlignment VerticalAlignment;
-};
\ No newline at end of file
+};
+
diff --git a/Engine/Plugins/FX/Niagara/Source/NiagaraEditor/Public/ViewModels/Stack/NiagaraStackFunctionInput.h b/Engine/Plugins/FX/Niagara/Source/NiagaraEditor/Public/ViewModels/Stack/NiagaraStackFunctionInput.h
--- a/Engine/Plugins/FX/Niagara/Source/NiagaraEditor/Public/ViewModels/Stack/NiagaraStackFunctionInput.h
+++ b/Engine/Plugins/FX/Niagara/Source/NiagaraEditor/Public/ViewModels/Stack/NiagaraStackFunctionInput.h
@@ -91,6 +91,9 @@ public:
 	/** Gets the type of this input. */
 	const FNiagaraTypeDefinition& GetInputType() const;
 
+	/** Gets the unit of this input. */
+	EUnit GetInputDisplayUnit() const;
+
 	//~ UNiagaraStackEntry interface
 	virtual FText GetDisplayName() const override;
 	virtual FText GetTooltipText() const override;
diff --git a/Engine/Plugins/FX/Niagara/Source/NiagaraEditorWidgets/Private/Stack/SNiagaraStackFunctionInputValue.cpp b/Engine/Plugins/FX/Niagara/Source/NiagaraEditorWidgets/Private/Stack/SNiagaraStackFunctionInputValue.cpp
--- a/Engine/Plugins/FX/Niagara/Source/NiagaraEditorWidgets/Private/Stack/SNiagaraStackFunctionInputValue.cpp
+++ b/Engine/Plugins/FX/Niagara/Source/NiagaraEditorWidgets/Private/Stack/SNiagaraStackFunctionInputValue.cpp
@@ -18,11 +18,9 @@
 #include "NiagaraEditorWidgetsStyle.h"
 #include "NiagaraEditorWidgetsUtilities.h"
 #include "NiagaraNodeAssignment.h"
-#include "NiagaraNodeCustomHlsl.h"
 #include "NiagaraNodeFunctionCall.h"
 #include "NiagaraParameterCollection.h"
 #include "NiagaraScriptVariable.h"
-#include "NiagaraScriptVariable.h"
 #include "NiagaraSettings.h"
 #include "NiagaraSystem.h"
 #include "PropertyEditorModule.h"
@@ -70,7 +68,6 @@ void SNiagaraStackFunctionInputValue::Construct(const FArguments& InArgs, UNiaga
 	SyntaxHighlighter = FNiagaraHLSLSyntaxHighlighter::Create();
 
 	TSharedPtr<SHorizontalBox> ChildrenBox;
-	
 	ChildSlot
 	[
 		SAssignNew(ChildrenBox, SHorizontalBox)
@@ -104,16 +101,64 @@ void SNiagaraStackFunctionInputValue::Construct(const FArguments& InArgs, UNiaga
 				.Padding(0, 0, 3, 0)
 				[
 					// Value Icon
-					SNew(SBox)
-					.WidthOverride(TextIconSize)
-					.VAlign(VAlign_Center)
+					SNew(SHorizontalBox)
 					.Visibility(this, &SNiagaraStackFunctionInputValue::GetInputIconVisibility)
+					+ SHorizontalBox::Slot()
+					.AutoWidth()
+					[
+						// icons without type info (e.g. data interfaces)
+						SNew(SBox)
+						.WidthOverride(TextIconSize)
+						.VAlign(VAlign_Center)
+						.Visibility_Lambda([this]()
+						{
+							return (FunctionInput->GetValueMode() == UNiagaraStackFunctionInput::EValueMode::Dynamic || FunctionInput->GetValueMode() == UNiagaraStackFunctionInput::EValueMode::Linked) ? EVisibility::Collapsed : EVisibility::Visible;
+						})
+						[
+							SNew(STextBlock)
+							.Font(FAppStyle::Get().GetFontStyle("FontAwesome.10"))
+							.Text(this, &SNiagaraStackFunctionInputValue::GetInputIconText)
+							.ToolTipText(this, &SNiagaraStackFunctionInputValue::GetInputIconToolTip)
+							.ColorAndOpacity(this, &SNiagaraStackFunctionInputValue::GetInputIconColor)
+						]
+					]
+					+ SHorizontalBox::Slot()
+					.AutoWidth()
 					[
-						SNew(STextBlock)
-						.Font(FAppStyle::Get().GetFontStyle("FontAwesome.10"))
-						.Text(this, &SNiagaraStackFunctionInputValue::GetInputIconText)
-						.ToolTipText(this, &SNiagaraStackFunctionInputValue::GetInputIconToolTip)
-						.ColorAndOpacity(this, &SNiagaraStackFunctionInputValue::GetInputIconColor)
+						// icons type info pill
+						SNew(SBorder)
+						.BorderImage(FNiagaraEditorStyle::Get().GetBrush("NiagaraEditor.Module.InputTypeBorder"))
+						.BorderBackgroundColor(FLinearColor::White)
+						.Padding(FMargin(2, 2, 2, 2))
+						[
+							SNew(SHorizontalBox)
+							.Visibility_Lambda([this]()
+							{
+								return (FunctionInput->GetValueMode() == UNiagaraStackFunctionInput::EValueMode::Dynamic || FunctionInput->GetValueMode() == UNiagaraStackFunctionInput::EValueMode::Linked) ? EVisibility::Visible : EVisibility::Collapsed;
+							})
+							+ SHorizontalBox::Slot()
+							.VAlign(VAlign_Center)
+							.HAlign(HAlign_Center)
+							.AutoWidth()
+							[
+								SNew(SImage)
+								.Visibility_Lambda([this](){ return (FunctionInput->GetValueMode() == UNiagaraStackFunctionInput::EValueMode::Dynamic || FunctionInput->GetValueMode() == UNiagaraStackFunctionInput::EValueMode::Linked) ? EVisibility::Visible : EVisibility::Collapsed; })
+								.ToolTipText(this, &SNiagaraStackFunctionInputValue::GetInputIconToolTip)
+								.ColorAndOpacity_Lambda([this]() { return UEdGraphSchema_Niagara::GetTypeColor(FunctionInput->GetInputType()); })
+								.Image(FNiagaraEditorStyle::Get().GetBrush("NiagaraEditor.Module.TypeIconPill"))
+							]
+							+ SHorizontalBox::Slot()
+							.VAlign(VAlign_Center)
+							.HAlign(HAlign_Center)
+							.AutoWidth()
+							[
+								SNew(SImage)
+								.Visibility_Lambda([this](){ return FunctionInput->GetValueMode() == UNiagaraStackFunctionInput::EValueMode::Dynamic ? EVisibility::Visible : EVisibility::Collapsed; })
+								.ToolTipText(this, &SNiagaraStackFunctionInputValue::GetInputIconToolTip)
+								.ColorAndOpacity(FLinearColor::White)
+								.Image(FNiagaraEditorStyle::Get().GetBrush("NiagaraEditor.Module.DynamicInput"))
+							]
+						]
 					]
 				]
 				+ SHorizontalBox::Slot()
@@ -477,7 +522,7 @@ TSharedRef<SWidget> SNiagaraStackFunctionInputValue::ConstructLocalValueStructWi
 		TSharedPtr<INiagaraEditorTypeUtilities, ESPMode::ThreadSafe> TypeEditorUtilities = NiagaraEditorModule.GetTypeUtilities(FunctionInput->GetInputType());
 		if (TypeEditorUtilities.IsValid() && TypeEditorUtilities->CanCreateParameterEditor())
 		{
-			TSharedPtr<SNiagaraParameterEditor> ParameterEditor = TypeEditorUtilities->CreateParameterEditor(FunctionInput->GetInputType());
+			TSharedPtr<SNiagaraParameterEditor> ParameterEditor = TypeEditorUtilities->CreateParameterEditor(FunctionInput->GetInputType(), FunctionInput->GetInputDisplayUnit());
 			ParameterEditor->UpdateInternalValueFromStruct(DisplayedLocalValueStruct.ToSharedRef());
 			ParameterEditor->SetOnBeginValueChange(SNiagaraParameterEditor::FOnValueChange::CreateSP(
 				this, &SNiagaraStackFunctionInputValue::ParameterBeginValueChange));
@@ -599,6 +644,12 @@ FText SNiagaraStackFunctionInputValue::GetDataValueText() const
 	}
 }
 
+UEnum* GetDisplayUnitEnum()
+{
+	static UEnum* UnitEnum = FindObjectChecked<UEnum>(nullptr, TEXT("/Script/CoreUObject.EUnit"));
+	return UnitEnum;
+}
+
 FText SNiagaraStackFunctionInputValue::GetDynamicValueText() const
 {
 	if (UNiagaraNodeFunctionCall* NodeFunctionCall = FunctionInput->GetDynamicInputNode())
@@ -612,12 +663,16 @@ FText SNiagaraStackFunctionInputValue::GetDynamicValueText() const
 			}
 		}
 		FString FunctionName = NodeFunctionCall->FunctionScript ? NodeFunctionCall->FunctionScript->GetName() : NodeFunctionCall->Signature.Name.ToString();
-		return FText::FromString(FName::NameToDisplayString(FunctionName, false));
-	}
-	else
-	{
-		return LOCTEXT("InvalidDynamicDisplayName", "(Invalid)");
+		FText DisplayString = FText::FromString(FName::NameToDisplayString(FunctionName, false));
+		EUnit DisplayUnit = FunctionInput->GetInputDisplayUnit();
+		if (DisplayUnit != EUnit::Unspecified)
+		{
+			UEnum* Enum = GetDisplayUnitEnum();
+			return FText::Format(FText::FromString("{0} ({1})"), DisplayString, FText::FromString(Enum->GetNameStringByValue(static_cast<int64>(DisplayUnit))));
+		}
+		return DisplayString;
 	}
+	return LOCTEXT("InvalidDynamicDisplayName", "(Invalid)");
 }
 
 FText SNiagaraStackFunctionInputValue::GetDefaultFunctionText() const
