Fixed a compilation bug with disabled nodes. The execution pin is not the first pin on a function call node if it silently inits some data interfaces with invisible pins.

diff --git a/Engine/Plugins/FX/Niagara/Source/NiagaraEditor/Private/NiagaraHlslTranslator.cpp b/Engine/Plugins/FX/Niagara/Source/NiagaraEditor/Private/NiagaraHlslTranslator.cpp
--- a/Engine/Plugins/FX/Niagara/Source/NiagaraEditor/Private/NiagaraHlslTranslator.cpp
+++ b/Engine/Plugins/FX/Niagara/Source/NiagaraEditor/Private/NiagaraHlslTranslator.cpp
@@ -9419,10 +9419,16 @@ int32 FHlslNiagaraTranslator::CompileOutputPin(const UEdGraphPin* InPin)
 	if (!Node->IsNodeEnabled())
 	{
 		// if the node is disabled (which commonly happens when a module is disabled in the stack), we skip it and follow the execution pin to the next node
-		UEdGraphPin* InputPin = Node->GetInputPin(0);
-		if (Node->GetOutputPin(0) == Pin && InputPin && InputPin->LinkedTo.Num() == 1 && Node->IsParameterMapPin(Pin))
+		FPinCollectorArray InputPins;
+		Node->GetInputPins(InputPins);
+		UEdGraphPin* OutputPin = Node->GetOutputPin(0);
+		for (int32 PinIndex = 0; PinIndex < InputPins.Num(); PinIndex++)
 		{
-			return CompileOutputPin(InputPin->LinkedTo[0]);
+			UEdGraphPin* ExecPin = InputPins[PinIndex];
+			if (OutputPin == Pin && ExecPin && ExecPin->LinkedTo.Num() == 1 && Node->IsParameterMapPin(Pin))
+			{
+				return CompileOutputPin(ExecPin->LinkedTo[0]);
+			}
 		}
 		Error(LOCTEXT("TraceDisabledPinFailed", "Failed to trace output pin of disabled node to a valid input!"), Node, Pin);
 		return INDEX_NONE;
