Fixed parameter map get & set nodes not allowing to delete previously added pins.
The new 'dynamic pins' array should handle this but the code has not accounted for previously existing pins in content.

diff --git a/Engine/Plugins/FX/Niagara/Source/NiagaraEditor/Private/NiagaraNodeParameterMapBase.cpp b/Engine/Plugins/FX/Niagara/Source/NiagaraEditor/Private/NiagaraNodeParameterMapBase.cpp
--- a/Engine/Plugins/FX/Niagara/Source/NiagaraEditor/Private/NiagaraNodeParameterMapBase.cpp
+++ b/Engine/Plugins/FX/Niagara/Source/NiagaraEditor/Private/NiagaraNodeParameterMapBase.cpp
@@ -167,6 +167,11 @@ bool UNiagaraNodeParameterMapBase::CanRenamePin(const UEdGraphPin* Pin) const
 	return FNiagaraParameterUtilities::TestCanRenameWithMessage(Pin->PinName, Unused);
 }
 
+bool UNiagaraNodeParameterMapBase::CanRemovePin(const UEdGraphPin* Pin) const
+{
+	return IsAddPin(Pin) == false && IsParameterMapPin(Pin) == false;
+}
+
 bool UNiagaraNodeParameterMapBase::GetIsPinEditNamespaceModifierPending(const UEdGraphPin* Pin)
 {
 	return PinsGuidsWithEditNamespaceModifierPending.Contains(Pin->PersistentGuid);
diff --git a/Engine/Plugins/FX/Niagara/Source/NiagaraEditor/Private/NiagaraNodeParameterMapBase.h b/Engine/Plugins/FX/Niagara/Source/NiagaraEditor/Private/NiagaraNodeParameterMapBase.h
--- a/Engine/Plugins/FX/Niagara/Source/NiagaraEditor/Private/NiagaraNodeParameterMapBase.h
+++ b/Engine/Plugins/FX/Niagara/Source/NiagaraEditor/Private/NiagaraNodeParameterMapBase.h
@@ -41,7 +41,8 @@ public:
 	virtual bool CanRenamePinFromContextMenu(const UEdGraphPin* Pin) const override { return false; }
 
 	virtual bool CanRenamePin(const UEdGraphPin* Pin) const override;
-
+	virtual bool CanRemovePin(const UEdGraphPin* Pin) const override;
+	
 	virtual FName GetNewPinDefaultNamespace() const { return PARAM_MAP_LOCAL_MODULE_STR; }
 
 	bool GetIsPinEditNamespaceModifierPending(const UEdGraphPin* Pin);
