Speculative fix for nonUnity compile error

diff --git a/Engine/Plugins/FX/Niagara/Source/NiagaraEditor/Private/Customizations/NiagaraScriptVariableCustomization.cpp b/Engine/Plugins/FX/Niagara/Source/NiagaraEditor/Private/Customizations/NiagaraScriptVariableCustomization.cpp
--- a/Engine/Plugins/FX/Niagara/Source/NiagaraEditor/Private/Customizations/NiagaraScriptVariableCustomization.cpp
+++ b/Engine/Plugins/FX/Niagara/Source/NiagaraEditor/Private/Customizations/NiagaraScriptVariableCustomization.cpp
@@ -12,7 +12,6 @@
 #include "EdGraphSchema_Niagara.h"
 #include "INiagaraEditorTypeUtilities.h"
 #include "NiagaraConstants.h"
-#include "NiagaraEditorCommon.h"
 #include "NiagaraEditorModule.h"
 #include "NiagaraEditorStyle.h"
 #include "NiagaraGraph.h"
diff --git a/Engine/Plugins/FX/Niagara/Source/NiagaraEditor/Public/SNiagaraParameterEditor.h b/Engine/Plugins/FX/Niagara/Source/NiagaraEditor/Public/SNiagaraParameterEditor.h
--- a/Engine/Plugins/FX/Niagara/Source/NiagaraEditor/Public/SNiagaraParameterEditor.h
+++ b/Engine/Plugins/FX/Niagara/Source/NiagaraEditor/Public/SNiagaraParameterEditor.h
@@ -4,6 +4,7 @@
 #include "Widgets/SCompoundWidget.h"
 #include "Widgets/DeclarativeSyntaxSupport.h"
 #include "UObject/StructOnScope.h"
+#include "Widgets/Input/NumericTypeInterface.h"
 #include "Widgets/Input/NumericUnitTypeInterface.inl"
 
 /** Base class for inline parameter editors. These editors are expected to maintain an internal value which
