Fixed ARM64 compile issue.

diff --git a/Engine/Plugins/FX/Niagara/Source/Niagara/Niagara.Build.cs b/Engine/Plugins/FX/Niagara/Source/Niagara/Niagara.Build.cs
--- a/Engine/Plugins/FX/Niagara/Source/Niagara/Niagara.Build.cs
+++ b/Engine/Plugins/FX/Niagara/Source/Niagara/Niagara.Build.cs
@@ -13,7 +13,7 @@ public class Niagara : ModuleRules
 
 		// Enable OpenVDB under Windows
 		// @todo: we need builds for OpenVDB for platforms other than windows
-		if (Target.IsInPlatformGroup(UnrealPlatformGroup.Windows))
+		if (Target.IsInPlatformGroup(UnrealPlatformGroup.Windows) && Target.WindowsPlatform.Architecture != UnrealArch.Arm64)
 		{
 			// Specific to OpenVDB support
 			bUseRTTI = true;
