- Fix incorrect component pool not InUsed list warning
- If the NC pool was cleaned up any in flight components would still exist and could be returned to the pool
- For example an UpdateContext could request the pool is cleared, which we do, but any existing components taken from the pool will still exist and therefore on return will produce an incorrect error

#rnx

diff --git a/Engine/Plugins/FX/Niagara/Source/Niagara/Private/NiagaraComponentPool.cpp b/Engine/Plugins/FX/Niagara/Source/Niagara/Private/NiagaraComponentPool.cpp
--- a/Engine/Plugins/FX/Niagara/Source/Niagara/Private/NiagaraComponentPool.cpp
+++ b/Engine/Plugins/FX/Niagara/Source/Niagara/Private/NiagaraComponentPool.cpp
@@ -79,11 +79,6 @@ void FNCPool::Cleanup()
 		}
 	}
 	FreeElements.Empty();
-
-#if ENABLE_NC_POOL_DEBUGGING
-	InUseComponents_Auto.Empty();
-	InUseComponents_Manual.Empty();
-#endif
 }
 
 UNiagaraComponent* FNCPool::Acquire(UWorld* World, UNiagaraSystem* Template, ENCPoolMethod PoolingMethod, bool bForceNew)
@@ -124,18 +119,6 @@ UNiagaraComponent* FNCPool::Acquire(UWorld* World, UNiagaraSystem* Template, ENC
 
 	RetElem.Component->PoolingMethod = PoolingMethod;
 
-#if ENABLE_NC_POOL_DEBUGGING
-	if (PoolingMethod == ENCPoolMethod::AutoRelease)
-	{
-		InUseComponents_Auto.Emplace(RetElem.Component);
-	}
-	else if (PoolingMethod == ENCPoolMethod::ManualRelease)
-	{
-		InUseComponents_Manual.Emplace(RetElem.Component);
-	}
-
-	MaxUsed = FMath::Max(MaxUsed, InUseComponents_Manual.Num() + InUseComponents_Auto.Num());
-#endif 
 	return RetElem.Component;
 }
 
@@ -144,23 +127,6 @@ void FNCPool::Reclaim(UNiagaraComponent* Component, const double CurrentTimeSeco
 	check(Component);
 	check(Component->GetAsset());
 
-#if ENABLE_NC_POOL_DEBUGGING
-	bool bWasInList = false;
-	if (Component->PoolingMethod == ENCPoolMethod::AutoRelease)
-	{
-		bWasInList = InUseComponents_Auto.RemoveSingleSwap(Component) > 0;
-	}
-	else if (Component->PoolingMethod == ENCPoolMethod::ManualRelease)
-	{
-		bWasInList = InUseComponents_Manual.RemoveSingleSwap(Component) > 0;
-	}
-	
-	if(!bWasInList)
-	{
-		UE_LOG(LogNiagara, Error, TEXT("World Niagara System Pool is reclaiming a component that is not in it's InUse list!"));
-	}
-#endif
-
 	//Don't add back to the pool if we're no longer pooling or we've hit our max resident pool size. Also if the component's world is in the process of tearing down.
 	if (GbEnableNiagaraSystemPooling != 0 && FreeElements.Num() < (int32)Component->GetAsset()->MaxPoolSize && Component->GetWorld()->bIsTearingDown == false)
 	{
@@ -244,19 +210,6 @@ void FNCPool::KillUnusedComponents(double KillTime, UNiagaraSystem* Template)
 		}
 	}
 	FreeElements.Shrink();
-
-#if ENABLE_NC_POOL_DEBUGGING
-	// Clean up any in use components that have been cleared out from under the pool. This could happen in someone manually destroys a component for example.
-	if ( InUseComponents_Manual.RemoveAllSwap([](const TWeakObjectPtr<UNiagaraComponent>& WeakPtr) { return !WeakPtr.IsValid(); }) > 0 )
-	{
-		UE_LOG(LogNiagara, Log, TEXT("Manual Pooled NC has been destroyed! Possibly via a DestroyComponent() call. You should not destroy these but rather call ReleaseToPool on the component so it can be re-used. |\t System: %s"), *Template->GetFullName());
-	}
-
-	if (InUseComponents_Auto.RemoveAllSwap([](const TWeakObjectPtr<UNiagaraComponent>& WeakPtr) { return !WeakPtr.IsValid(); }) > 0)
-	{
-		UE_LOG(LogNiagara, Log, TEXT("Auto Pooled NC has been destroyed! Possibly via a DestroyComponent() call. You should not destroy these manually. Just deactivate them and allow then to be reclaimed by the pool automatically. |\t System: %s"), *Template->GetFullName());
-	}
-#endif
 }
 
 //////////////////////////////////////////////////////////////////////////
@@ -433,6 +386,19 @@ UNiagaraComponent* UNiagaraComponentPool::CreateWorldParticleSystem(UNiagaraSyst
 		{
 			FNCPool& Pool = WorldParticleSystemPools.FindOrAdd(Template);
 			Component = Pool.Acquire(World, Template, PoolingMethod);
+		#if ENABLE_NC_POOL_DEBUGGING
+			if (Component != nullptr)
+			{
+				if (PoolingMethod == ENCPoolMethod::AutoRelease)
+				{
+					InUseComponents_Auto.Add(Component);
+				}
+				else if (ensure(PoolingMethod == ENCPoolMethod::ManualRelease))
+				{
+					InUseComponents_Manual.Add(Component);
+				}
+			}
+		#endif
 		}
 	}
 	else
@@ -496,6 +462,22 @@ void UNiagaraComponentPool::ReclaimWorldParticleSystem(UNiagaraComponent* Compon
 			NCPool = &WorldParticleSystemPools.Add(Asset);
 		}
 
+	#if ENABLE_NC_POOL_DEBUGGING
+		bool bWasInList = false;
+		if (Component->PoolingMethod == ENCPoolMethod::AutoRelease)
+		{
+			bWasInList = InUseComponents_Auto.RemoveSingleSwap(Component) > 0;
+		}
+		else if (ensure(Component->PoolingMethod == ENCPoolMethod::ManualRelease))
+		{
+			bWasInList = InUseComponents_Manual.RemoveSingleSwap(Component) > 0;
+		}
+
+		if (!bWasInList)
+		{
+			UE_LOG(LogNiagara, Error, TEXT("World Niagara System Pool is reclaiming a component that is not in it's InUse list!"));
+		}
+	#endif
 		NCPool->Reclaim(Component, CurrentTime);
 	}
 	else
@@ -590,27 +572,31 @@ void UNiagaraComponentPool::Dump()
 			}
 		}
 		uint32 InUseMemUsage = 0;
-		for (auto WeakComponent : Pool.InUseComponents_Auto)
+		int32 NumInUseAutoComponents = 0;
+		int32 NumInUseManualComponents = 0;
+		for (auto WeakComponent : InUseComponents_Auto)
 		{
 			UNiagaraComponent* Component = WeakComponent.Get();
-			if (ensureAlways(Component))
+			if (ensureAlways(Component) && (Component->GetAsset() == System))
 			{
 				InUseMemUsage += Component->GetApproxMemoryUsage();				
+				++NumInUseAutoComponents;
 			}
 		}
-		for (auto WeakComponent : Pool.InUseComponents_Auto)
+		for (auto WeakComponent : InUseComponents_Manual)
 		{
 			UNiagaraComponent* Component = WeakComponent.Get();
-			if (ensureAlways(Component))
+			if (ensureAlways(Component) && (Component->GetAsset() == System))
 			{
 				InUseMemUsage += Component->GetApproxMemoryUsage();
+				++NumInUseManualComponents;
 			}
 		}
 
 		TotalMemUsage += FreeMemUsage;
 		TotalMemUsage += InUseMemUsage;
 
-		DumpStr += FString::Printf(TEXT("Free: %d (%uB) \t|\t Used(Auto - Manual): %d - %d (%uB) \t|\t MaxUsed: %d \t|\t System: %s\n"), Pool.FreeElements.Num(), FreeMemUsage, Pool.InUseComponents_Auto.Num(), Pool.InUseComponents_Manual.Num(), InUseMemUsage, Pool.MaxUsed, *System->GetFullName());
+		DumpStr += FString::Printf(TEXT("Free: %d (%uB) \t|\t Used(Auto - Manual): %d - %d (%uB) \t|\t System: %s\n"), Pool.FreeElements.Num(), FreeMemUsage, NumInUseAutoComponents, NumInUseManualComponents, InUseMemUsage, *System->GetFullName());
 	}
 
 	UE_LOG(LogNiagara, Log, TEXT("***************************************"));
