Clamped MoveTo index when drag & dropping an array element in the stack. Fixes a crash with an index out of bounds access.

diff --git a/Engine/Plugins/FX/Niagara/Source/NiagaraEditor/Private/ViewModels/Stack/NiagaraStackPropertyRow.cpp b/Engine/Plugins/FX/Niagara/Source/NiagaraEditor/Private/ViewModels/Stack/NiagaraStackPropertyRow.cpp
--- a/Engine/Plugins/FX/Niagara/Source/NiagaraEditor/Private/ViewModels/Stack/NiagaraStackPropertyRow.cpp
+++ b/Engine/Plugins/FX/Niagara/Source/NiagaraEditor/Private/ViewModels/Stack/NiagaraStackPropertyRow.cpp
@@ -209,6 +209,15 @@ TOptional<UNiagaraStackEntry::FDropRequestResponse> UNiagaraStackPropertyRow::Dr
 		TSharedPtr<IPropertyHandle> DraggedPropertyHandle = DraggedPropertyRow->GetDetailTreeNode()->CreatePropertyHandle();
 		TSharedPtr<IPropertyHandle> TargetPropertyHandle = DetailTreeNode->CreatePropertyHandle();
 		int32 IndexOffset = DropRequest.DropZone == EItemDropZone::AboveItem ? 0 : 1;
+		uint32 NumElements = 0;
+		TargetPropertyHandle->GetParentHandle()->AsArray()->GetNumElements(NumElements);
+
+		// we clamp the offset
+		if((uint32) (TargetPropertyHandle->GetIndexInArray() + IndexOffset) >= (NumElements - 1))
+		{
+			IndexOffset = 0;
+		}
+		
 		TargetPropertyHandle->GetParentHandle()->AsArray()->MoveElementTo(DraggedPropertyHandle->GetIndexInArray(), TargetPropertyHandle->GetIndexInArray() + IndexOffset);
 	}
 	return CanDropResponse;
