Fix with merge manager getting mismatched variables with the default pins
#rnx

diff --git a/Engine/Plugins/FX/Niagara/Source/NiagaraEditor/Private/NiagaraScriptMergeManager.cpp b/Engine/Plugins/FX/Niagara/Source/NiagaraEditor/Private/NiagaraScriptMergeManager.cpp
--- a/Engine/Plugins/FX/Niagara/Source/NiagaraEditor/Private/NiagaraScriptMergeManager.cpp
+++ b/Engine/Plugins/FX/Niagara/Source/NiagaraEditor/Private/NiagaraScriptMergeManager.cpp
@@ -254,15 +254,25 @@ FNiagaraStackFunctionMergeAdapter::FNiagaraStackFunctionMergeAdapter(const FVers
 	TArray<FNiagaraVariable> RapidIterationInputDefaultValues;
 	if (InFunctionCallNode.FunctionScript != nullptr)
 	{
-		FCompileConstantResolver Resolver(InOwningEmitter, ENiagaraScriptUsage::Function);
-		TArray<FName> FunctionInputVariableNames;
 		TArray<FNiagaraVariable> FunctionInputVariables;
-		GetStackFunctionInputs(*FunctionCallNode, FunctionInputVariables, Resolver, FNiagaraStackGraphUtilities::ENiagaraGetStackFunctionInputPinsOptions::ModuleInputsOnly, false);
-		for (const FNiagaraVariable& FunctionInputVariable : FunctionInputVariables)
+		TArray<FName> FunctionInputVariableNames;
+
 		{
-			if (FunctionInputVariable.IsValid() && FNiagaraStackGraphUtilities::IsRapidIterationType(FunctionInputVariable.GetType()))
+			FCompileConstantResolver Resolver(InOwningEmitter, ENiagaraScriptUsage::Function);
+			TArray<FNiagaraVariable> AllFunctionInputVariables;
+			GetStackFunctionInputs(*FunctionCallNode, AllFunctionInputVariables, Resolver, FNiagaraStackGraphUtilities::ENiagaraGetStackFunctionInputPinsOptions::ModuleInputsOnly, false);
+
+			const int32 AllFunctionInputVariableCount = AllFunctionInputVariables.Num();
+			FunctionInputVariables.Reserve(AllFunctionInputVariableCount);
+			FunctionInputVariableNames.Reserve(AllFunctionInputVariableCount);
+
+			for (const FNiagaraVariable& FunctionInputVariable : AllFunctionInputVariables)
 			{
-				FunctionInputVariableNames.Add(FunctionInputVariable.GetName());
+				if (FunctionInputVariable.IsValid() && FNiagaraStackGraphUtilities::IsRapidIterationType(FunctionInputVariable.GetType()))
+				{
+					FunctionInputVariables.Add(FunctionInputVariable);
+					FunctionInputVariableNames.Add(FunctionInputVariable.GetName());
+				}
 			}
 		}
 
