- Deprecate GetCameraOcclusionCircle and replace with a new function to fix steppy behavior with off screen samples

diff --git a/Engine/Plugins/FX/Niagara/Shaders/Private/NiagaraDataInterfaceOcclusion.ush b/Engine/Plugins/FX/Niagara/Shaders/Private/NiagaraDataInterfaceOcclusion.ush
--- a/Engine/Plugins/FX/Niagara/Shaders/Private/NiagaraDataInterfaceOcclusion.ush
+++ b/Engine/Plugins/FX/Niagara/Shaders/Private/NiagaraDataInterfaceOcclusion.ush
@@ -127,6 +127,65 @@ void QueryOcclusionFactorWithCircleGPU_{ParameterName}(in float3 In_SampleCenter
 	Out_SampleFraction = Steps == 0 ? 0 : (TotalSamples / (Rays * Steps + 1));
 }
 
+void QueryOcclusionFactorWithCircle_{ParameterName}(float3 WorldPosition, float WorldDiameter, bool bIncludeCenterSample, int NumberOfRings, int SamplesPerRing, out float VisibilityFraction, out float OnScreenFraction, out float DepthPassFraction)
+{
+	const float PI = 3.14159265;
+	const float SPIRAL_TURN = 2.0f * PI * 0.61803399; // use golden ratio to rotate sample pattern each ring so we get a spiral
+
+	const FLWCVector3 LWCWorldPosition = MakeLWCVector3({ParameterName}_SystemLWCTile, WorldPosition);
+	const float SampleDepth = LWCToFloat(LWCAbs(LWCDot(LWCSubtract(LWCWorldPosition, PrimaryView.WorldViewOrigin), View.ViewForward.xyz)));
+
+	const float3 PreViewPosition = LWCToFloat(LWCAdd(LWCWorldPosition, PrimaryView.PreViewTranslation));
+	const float4 ClipPosition = mul(float4(PreViewPosition, 1), PrimaryView.TranslatedWorldToClip);
+	const float2 ScreenPosition = ClipPosition.xy / ClipPosition.w;
+	const float2 ScreenUV = (ScreenPosition * float2(0.5f, -0.5f)) + 0.5f;
+
+	int TotalSamples = 0;
+	int UnoccludedSamples = 0;
+	int OnScreenSamples = 0;
+
+	// Take center tap
+	if ( bIncludeCenterSample )
+	{
+		++TotalSamples;
+		if ( all(ScreenUV >= 0.0f) && all(ScreenUV < 1.0f) )
+		{
+			UnoccludedSamples += CalcSceneDepth(ScreenUV) > SampleDepth;
+			++OnScreenSamples;
+		}
+	}
+
+	// Take ring taps
+	if ( NumberOfRings > 0 && SamplesPerRing > 0 )
+	{
+		const float WorldDepth = dot(LWCToFloat(LWCSubtract(LWCWorldPosition, PrimaryView.WorldCameraOrigin)), PrimaryView.ViewForward);
+		const float2 ScreenRadiusUV = (WorldDiameter * 0.5f) / (GetTanHalfFieldOfView() * 2.0f * WorldDepth);
+
+		const float RadStep = 2.0f * PI / float(SamplesPerRing + 1);
+		float Rad = 0.0f;
+		for ( int iRing=0; iRing < NumberOfRings; ++iRing )
+		{
+			const float2 RingRadiusUV = ScreenRadiusUV * (float(iRing + 1) / float(NumberOfRings));
+			Rad += SPIRAL_TURN;
+			for ( int iSample=0; iSample < SamplesPerRing; ++iSample )
+			{
+				const float2 SampleUV = ScreenUV + (RingRadiusUV * float2(sin(Rad), cos(Rad)));
+				if ( all(SampleUV >= 0.0f) && all(SampleUV < 1.0f) )
+				{
+					UnoccludedSamples += CalcSceneDepth(SampleUV) > SampleDepth;
+					++OnScreenSamples;
+				}
+				Rad += RadStep;
+			}
+		}
+		TotalSamples += NumberOfRings * SamplesPerRing;
+	}
+
+	VisibilityFraction = TotalSamples > 0 ? float(UnoccludedSamples) / float(TotalSamples) : 0.0f;
+	OnScreenFraction = TotalSamples > 0 ? float(OnScreenSamples) / float(TotalSamples) : 0.0f;
+	DepthPassFraction = OnScreenSamples > 0 ? float(UnoccludedSamples) / float(OnScreenSamples) : 0.0f;
+}
+
 void QueryCloudOcclusionWithCircle_{ParameterName}(float3 WorldPosition, float WorldDiameter, bool bIncludeCenterSample, int NumberOfRings, int SamplesPerRing, out float VisibilityFraction, out float SampleFraction, out float3 AtmosphereTransmittance)
 {
 	const float PI = 3.14159265;
diff --git a/Engine/Plugins/FX/Niagara/Source/Niagara/Classes/NiagaraDataInterfaceOcclusion.h b/Engine/Plugins/FX/Niagara/Source/Niagara/Classes/NiagaraDataInterfaceOcclusion.h
--- a/Engine/Plugins/FX/Niagara/Source/Niagara/Classes/NiagaraDataInterfaceOcclusion.h
+++ b/Engine/Plugins/FX/Niagara/Source/Niagara/Classes/NiagaraDataInterfaceOcclusion.h
@@ -37,12 +37,6 @@ public:
 	virtual bool CanExecuteOnTarget(ENiagaraSimTarget Target) const override { return Target == ENiagaraSimTarget::GPUComputeSim; }
 	virtual bool RequiresDepthBuffer() const override { return true; }
 	//UNiagaraDataInterface Interface
-	
-private:
-	static const TCHAR* TemplateShaderFilePath;
-	static const FName GetCameraOcclusionRectangleName;
-	static const FName GetCameraOcclusionCircleName;
-	static const FName QueryCloudOcclusionWithCircleName;
 };
 
 struct FNiagaraDataIntefaceProxyOcclusionQuery : public FNiagaraDataInterfaceProxy
diff --git a/Engine/Plugins/FX/Niagara/Source/Niagara/Private/NiagaraDataInterfaceOcclusion.cpp b/Engine/Plugins/FX/Niagara/Source/Niagara/Private/NiagaraDataInterfaceOcclusion.cpp
--- a/Engine/Plugins/FX/Niagara/Source/Niagara/Private/NiagaraDataInterfaceOcclusion.cpp
+++ b/Engine/Plugins/FX/Niagara/Source/Niagara/Private/NiagaraDataInterfaceOcclusion.cpp
@@ -16,10 +16,14 @@
 
 #define LOCTEXT_NAMESPACE "NiagaraDataInterfaceOcclusion"
 
-const TCHAR* UNiagaraDataInterfaceOcclusion::TemplateShaderFilePath = TEXT("/Plugin/FX/Niagara/Private/NiagaraDataInterfaceOcclusion.ush");
-const FName UNiagaraDataInterfaceOcclusion::GetCameraOcclusionRectangleName(TEXT("QueryOcclusionFactorWithRectangleGPU"));
-const FName UNiagaraDataInterfaceOcclusion::GetCameraOcclusionCircleName(TEXT("QueryOcclusionFactorWithCircleGPU"));
-const FName UNiagaraDataInterfaceOcclusion::QueryCloudOcclusionWithCircleName(TEXT("QueryCloudOcclusionWithCircle"));
+namespace NDIOcclusionImpl
+{
+	static const TCHAR*	TemplateShaderFilePath = TEXT("/Plugin/FX/Niagara/Private/NiagaraDataInterfaceOcclusion.ush");
+	static const FName	GetCameraOcclusionRectangleName(TEXT("QueryOcclusionFactorWithRectangleGPU"));
+	static const FName	Deprecated_GetCameraOcclusionCircleName(TEXT("QueryOcclusionFactorWithCircleGPU"));
+	static const FName	QueryOcclusionFactorWithCircleName(TEXT("QueryOcclusionFactorWithCircle"));
+	static const FName	QueryCloudOcclusionWithCircleName(TEXT("QueryCloudOcclusionWithCircle"));
+}
 
 struct FNiagaraOcclusionDIFunctionVersion
 {
@@ -53,20 +57,25 @@ void UNiagaraDataInterfaceOcclusion::PostInitProperties()
 
 void UNiagaraDataInterfaceOcclusion::GetFunctions(TArray<FNiagaraFunctionSignature>& OutFunctions)
 {
-	FNiagaraFunctionSignature DefaultSig;
-	DefaultSig.bMemberFunction = true;
-	DefaultSig.bRequiresContext = false;
-	DefaultSig.bSupportsCPU = false;
-	DefaultSig.AddInput(FNiagaraVariable(FNiagaraTypeDefinition(GetClass()), TEXT("Occlusion interface")));
-	DefaultSig.SetFunctionVersion(FNiagaraOcclusionDIFunctionVersion::LatestVersion);
+	using namespace NDIOcclusionImpl;
 
 	static const FText VisibilityFractionDescription(LOCTEXT("VisibilityFractionDescription", "Returns a value 0..1 depending on how many of the samples on the screen were occluded.\nFor example, a value of 0.3 means that 70% of visible samples were occluded.\nIf the sample fraction is 0 then this also returns 0."));
+	static const FText OnScreenFractionDescription(LOCTEXT("OnScreenFractionDescription", "Returns a valid 0..1 depending on how may of the samples are on screen.\nFor example, a valid of 0.3 means 70% of the samples are off screen and 30% are on screen."));
+	static const FText DepthPassFractionDescription(LOCTEXT("DepthPassFractionDescription", "Returns a value 0..1 depending on how many of the samples passed the depth test.\nFor example, a value of 0.3 means 70% of the depth test samples are closer to the screen.\nNote this fraction only includes on screen samples."));
 	static const FText SampleFractionDescription(LOCTEXT("SampleFractionDescription", "Returns a value 0..1 depending on how many samples were inside the viewport or outside of it.\nFor example, a value of 0.3 means that 70% of samples were outside the current viewport and therefore not visible."));
 	static const FText CircleCenterPosDescription(LOCTEXT("CircleCenterPosDescription", "This world space position where the center of the sample circle should be."));
 	static const FText SampleWindowDiameterDescription(LOCTEXT("SampleWindowDiameterDescription", "The world space diameter of the circle to sample.\nIf the particle is a spherical sprite then this is the sprite size."));
 	static const FText SamplesPerRingDescription(LOCTEXT("SamplesPerRingDescription", "The number of samples for each ring inside the circle.\nThe total number of samples is NumRings * SamplesPerRing."));
 	static const FText NumberOfSampleRingsDescription(LOCTEXT("NumberOfSampleRingsDescription", "This number of concentric rings to sample inside the circle.\nThe total number of samples is NumRings * SamplesPerRing."));
 	static const FText IncludeCenterSampleDescription(LOCTEXT("IncludeCenterSampleDescription", "When enabled we sample the center of the circle in addition to the rings."));
+
+	FNiagaraFunctionSignature DefaultSig;
+	DefaultSig.bMemberFunction = true;
+	DefaultSig.bRequiresContext = false;
+	DefaultSig.bSupportsCPU = false;
+	DefaultSig.AddInput(FNiagaraVariable(FNiagaraTypeDefinition(GetClass()), TEXT("Occlusion interface")));
+	DefaultSig.SetFunctionVersion(FNiagaraOcclusionDIFunctionVersion::LatestVersion);
+
 	{
 		FNiagaraFunctionSignature& Sig = OutFunctions.Add_GetRef(DefaultSig);
 		Sig.Name = GetCameraOcclusionRectangleName;
@@ -81,8 +90,8 @@ void UNiagaraDataInterfaceOcclusion::GetFunctions(TArray<FNiagaraFunctionSignatu
 
 	{
 		FNiagaraFunctionSignature& Sig = OutFunctions.Add_GetRef(DefaultSig);
-		Sig.Name = GetCameraOcclusionCircleName;
-		Sig.SetDescription(LOCTEXT("GetCameraOcclusionCircleFunctionDescription", "This function returns the occlusion factor of a sprite. It samples the depth buffer in concentric rings around the given world position and compares each sample with the camera distance."));
+		Sig.Name = Deprecated_GetCameraOcclusionCircleName;
+		Sig.bSoftDeprecatedFunction = true;
 		Sig.AddInput(FNiagaraVariable(FNiagaraTypeDefinition::GetPositionDef(), TEXT("Sample Center World Position")), CircleCenterPosDescription);
 		Sig.AddInput(FNiagaraVariable(FNiagaraTypeDefinition::GetFloatDef(), TEXT("Sample Window Diameter World")), SampleWindowDiameterDescription);
 		Sig.AddInput(FNiagaraVariable(FNiagaraTypeDefinition::GetFloatDef(), TEXT("Samples per ring")), SamplesPerRingDescription);
@@ -91,6 +100,29 @@ void UNiagaraDataInterfaceOcclusion::GetFunctions(TArray<FNiagaraFunctionSignatu
 		Sig.AddOutput(FNiagaraVariable(FNiagaraTypeDefinition::GetFloatDef(), TEXT("Sample Fraction")), SampleFractionDescription);
 	}
 
+	{
+		FNiagaraFunctionSignature& Sig = OutFunctions.Add_GetRef(DefaultSig);
+		Sig.Name = QueryOcclusionFactorWithCircleName;
+		Sig.SetDescription(LOCTEXT("QueryCloudOcclusionWithCircleDescription", "Returns the cloud occlusion factor for the world position. "));
+		Sig.Inputs.Emplace(FNiagaraTypeDefinition::GetPositionDef(), TEXT("WorldPosition"));
+		Sig.Inputs.Emplace(FNiagaraTypeDefinition::GetFloatDef(), TEXT("WorldDiameter"));
+		Sig.Inputs.Emplace_GetRef(FNiagaraTypeDefinition::GetBoolDef(), TEXT("IncludeCenterSample")).SetValue(true);
+		Sig.Inputs.Emplace_GetRef(FNiagaraTypeDefinition::GetIntDef(), TEXT("NumberOfRings")).SetValue(1);
+		Sig.Inputs.Emplace_GetRef(FNiagaraTypeDefinition::GetIntDef(), TEXT("SamplesPerRing")).SetValue(1);
+		Sig.Outputs.Emplace(FNiagaraTypeDefinition::GetFloatDef(), TEXT("VisibilityFraction"));
+		Sig.Outputs.Emplace(FNiagaraTypeDefinition::GetFloatDef(), TEXT("OnScreenFraction"));
+		Sig.Outputs.Emplace(FNiagaraTypeDefinition::GetFloatDef(), TEXT("DepthPassFraction"));
+
+		Sig.SetInputDescription(Sig.Inputs[1], CircleCenterPosDescription);
+		Sig.SetInputDescription(Sig.Inputs[2], SampleWindowDiameterDescription);
+		Sig.SetInputDescription(Sig.Inputs[3], IncludeCenterSampleDescription);
+		Sig.SetInputDescription(Sig.Inputs[4], NumberOfSampleRingsDescription);
+		Sig.SetInputDescription(Sig.Inputs[5], SamplesPerRingDescription);
+		Sig.SetOutputDescription(Sig.Outputs[0], VisibilityFractionDescription);
+		Sig.SetOutputDescription(Sig.Outputs[1], OnScreenFractionDescription);
+		Sig.SetOutputDescription(Sig.Outputs[2], DepthPassFractionDescription);
+	}
+
 	{
 		FNiagaraFunctionSignature& Sig = OutFunctions.Add_GetRef(DefaultSig);
 		Sig.Name = QueryCloudOcclusionWithCircleName;
@@ -117,6 +149,8 @@ void UNiagaraDataInterfaceOcclusion::GetFunctions(TArray<FNiagaraFunctionSignatu
 #if WITH_EDITORONLY_DATA
 bool UNiagaraDataInterfaceOcclusion::AppendCompileHash(FNiagaraCompileHashVisitor* InVisitor) const
 {
+	using namespace NDIOcclusionImpl;
+
 	if (!Super::AppendCompileHash(InVisitor))
 		return false;
 
@@ -132,11 +166,14 @@ void UNiagaraDataInterfaceOcclusion::GetCommonHLSL(FString& OutHLSL)
 
 bool UNiagaraDataInterfaceOcclusion::GetFunctionHLSL(const FNiagaraDataInterfaceGPUParamInfo& ParamInfo, const FNiagaraDataInterfaceGeneratedFunction& FunctionInfo, int FunctionInstanceIndex, FString& OutHLSL)
 {
+	using namespace NDIOcclusionImpl;
+
 	static const TSet<FName> ValidGpuFunctions =
 	{
 		GetCameraOcclusionRectangleName,
-		GetCameraOcclusionCircleName,
+		QueryOcclusionFactorWithCircleName,
 		QueryCloudOcclusionWithCircleName,
+		Deprecated_GetCameraOcclusionCircleName,
 	};
 	return ValidGpuFunctions.Contains(FunctionInfo.DefinitionName);
 }
@@ -162,6 +199,8 @@ bool UNiagaraDataInterfaceOcclusion::UpgradeFunctionCall(FNiagaraFunctionSignatu
 
 void UNiagaraDataInterfaceOcclusion::GetParameterDefinitionHLSL(const FNiagaraDataInterfaceGPUParamInfo& ParamInfo,	FString& OutHLSL)
 {
+	using namespace NDIOcclusionImpl;
+
 	Super::GetParameterDefinitionHLSL(ParamInfo, OutHLSL);
 
 	const TMap<FString, FStringFormatArg> TemplateArgs =
