- Fix for media textures not capturing the size
- They are neither Texture2D or RT2D so we capture the surface size and assume a mip count of 1

diff --git a/Engine/Plugins/FX/Niagara/Source/Niagara/Private/NiagaraDataInterfaceTexture.cpp b/Engine/Plugins/FX/Niagara/Source/Niagara/Private/NiagaraDataInterfaceTexture.cpp
--- a/Engine/Plugins/FX/Niagara/Source/Niagara/Private/NiagaraDataInterfaceTexture.cpp
+++ b/Engine/Plugins/FX/Niagara/Source/Niagara/Private/NiagaraDataInterfaceTexture.cpp
@@ -275,17 +275,23 @@ bool UNiagaraDataInterfaceTexture::PerInstanceTick(void* PerInstanceData, FNiaga
 	UTexture* CurrentTexture = InstanceData->UserParamBinding.GetValueOrDefault<UTexture>(Texture);
 	FIntPoint CurrentTextureSize = FIntPoint::ZeroValue;
 	int32 CurrentTextureMipLevels = 0;
-	if (UTexture2D* CurrentTexture2D = Cast<UTexture2D>(CurrentTexture) )
-	{
-		CurrentTextureSize.X = int32(CurrentTexture->GetSurfaceWidth());
-		CurrentTextureSize.Y = int32(CurrentTexture->GetSurfaceHeight());
-		CurrentTextureMipLevels = CurrentTexture2D->GetNumMips();
-	}
-	else if (UTextureRenderTarget2D* CurrentTexture2DRT = Cast<UTextureRenderTarget2D>(CurrentTexture))
+
+	if (CurrentTexture != nullptr)
 	{
 		CurrentTextureSize.X = int32(CurrentTexture->GetSurfaceWidth());
 		CurrentTextureSize.Y = int32(CurrentTexture->GetSurfaceHeight());
-		CurrentTextureMipLevels = CurrentTexture2DRT->GetNumMips();
+		if (UTexture2D* CurrentTexture2D = Cast<UTexture2D>(CurrentTexture))
+		{
+			CurrentTextureMipLevels = CurrentTexture2D->GetNumMips();
+		}
+		else if (UTextureRenderTarget2D* CurrentTexture2DRT = Cast<UTextureRenderTarget2D>(CurrentTexture))
+		{
+			CurrentTextureMipLevels = CurrentTexture2DRT->GetNumMips();
+		}
+		else
+		{
+			CurrentTextureMipLevels = 1;
+		}
 	}
 
 	if ( (InstanceData->CurrentTexture != CurrentTexture) || (InstanceData->CurrentTextureSize != CurrentTextureSize) || (InstanceData->CurrentTextureMipLevels != CurrentTextureMipLevels) )
