Couple of small fixes to avoid crashing when dealing with bad data when cooking niagara systems (in this case variables with invalid types because dependent enums are missing)

diff --git a/Engine/Plugins/FX/Niagara/Source/Niagara/Private/NiagaraParameterStore.cpp b/Engine/Plugins/FX/Niagara/Source/Niagara/Private/NiagaraParameterStore.cpp
--- a/Engine/Plugins/FX/Niagara/Source/Niagara/Private/NiagaraParameterStore.cpp
+++ b/Engine/Plugins/FX/Niagara/Source/Niagara/Private/NiagaraParameterStore.cpp
@@ -788,7 +788,7 @@ bool FNiagaraParameterStore::RemoveParameter(const FNiagaraVariableBase& ToRemov
 					NewOffsets.Add(FNiagaraVariableWithOffset(ExistingVar, Offset, FNiagaraLwcStructConverter()));
 					NewUObjects[Offset] = UObjects[ExistingOffset];
 				}
-				else
+				else if (ExistingVar.IsValid())
 				{
 					int32 Offset = NewData.Num();
 					int32 ParamSize = ExistingVar.GetSizeInBytes();
@@ -796,6 +796,10 @@ bool FNiagaraParameterStore::RemoveParameter(const FNiagaraVariableBase& ToRemov
 					NewData.AddUninitialized(ParamSize);
 					FMemory::Memcpy(NewData.GetData() + Offset, ParameterData.GetData() + ExistingOffset, ParamSize);
 				}
+				else
+				{
+					UE_LOG(LogNiagara, Warning, TEXT("RemoveParameter: Encountered invalid variable in parameter store (%s) -> variable: %s"), *GetPathNameSafe(Owner.Get()), *ExistingVar.GetName().ToString());
+				}
 			}
 		}
 
diff --git a/Engine/Plugins/FX/Niagara/Source/Niagara/Public/NiagaraTypes.h b/Engine/Plugins/FX/Niagara/Source/Niagara/Public/NiagaraTypes.h
--- a/Engine/Plugins/FX/Niagara/Source/Niagara/Public/NiagaraTypes.h
+++ b/Engine/Plugins/FX/Niagara/Source/Niagara/Public/NiagaraTypes.h
@@ -1002,7 +1002,8 @@ public:
 
 	FORCEINLINE bool IsUObject() const
 	{
-		return GetStruct()->IsChildOf<UObject>();
+		const UStruct* ClassStruct = GetStruct();
+		return ClassStruct ? ClassStruct->IsChildOf<UObject>() : false;
 	}
 
 	bool IsEnum() const { return UnderlyingType == UT_Enum; }
