- Stack search now attempts to restore the search result after changing something in the stack while the search is ongoing
- A stack entry can now only appear once in  in search results
- Refactored the search functionality to use an TOptional<FSearchResult> for the focused search result instead of maintaining an index into an array of search results.

Minor changes:
Renamed "Refresh" function to "RequestRefreshDeferred" as that's what it was doing
Removed another function that did the same thing

diff --git a/Engine/Plugins/FX/Niagara/Source/NiagaraEditor/Private/ViewModels/NiagaraSystemViewModel.cpp b/Engine/Plugins/FX/Niagara/Source/NiagaraEditor/Private/ViewModels/NiagaraSystemViewModel.cpp
--- a/Engine/Plugins/FX/Niagara/Source/NiagaraEditor/Private/ViewModels/NiagaraSystemViewModel.cpp
+++ b/Engine/Plugins/FX/Niagara/Source/NiagaraEditor/Private/ViewModels/NiagaraSystemViewModel.cpp
@@ -1809,7 +1809,7 @@ void FNiagaraSystemViewModel::NotifyDataObjectChanged(TArray<UObject*> ChangedOb
 						UNiagaraStackViewModel* StackViewModel = EmitterHandleViewModel->GetEmitterStackViewModel();
 						if(ensure(StackViewModel))
 						{
-							StackViewModel->Refresh();
+							StackViewModel->RequestRefreshDeferred();
 						}
 					}
 				}
@@ -3015,11 +3015,11 @@ void FNiagaraSystemViewModel::RefreshStackViewModels()
 {
 	if (SystemStackViewModel)
 	{
-		SystemStackViewModel->Refresh();
+		SystemStackViewModel->RequestRefreshDeferred();
 	}
 	for (TSharedRef<FNiagaraEmitterHandleViewModel> EmitterHandleViewModel : EmitterHandleViewModels)
 	{
-		EmitterHandleViewModel->GetEmitterStackViewModel()->Refresh();
+		EmitterHandleViewModel->GetEmitterStackViewModel()->RequestRefreshDeferred();
 	}
 }
 
diff --git a/Engine/Plugins/FX/Niagara/Source/NiagaraEditor/Private/ViewModels/Stack/NiagaraStackViewModel.cpp b/Engine/Plugins/FX/Niagara/Source/NiagaraEditor/Private/ViewModels/Stack/NiagaraStackViewModel.cpp
--- a/Engine/Plugins/FX/Niagara/Source/NiagaraEditor/Private/ViewModels/Stack/NiagaraStackViewModel.cpp
+++ b/Engine/Plugins/FX/Niagara/Source/NiagaraEditor/Private/ViewModels/Stack/NiagaraStackViewModel.cpp
@@ -121,7 +121,7 @@ void UNiagaraStackViewModel::InitializeWithViewModels(TSharedPtr<FNiagaraSystemV
 
 			if (UNiagaraEmitter* Emitter = EmitterViewModel->GetEmitter().Emitter)
 			{
-				EmitterViewModel->GetOrCreateEditorData().OnSummaryViewStateChanged().AddUObject(this, &UNiagaraStackViewModel::EntryRequestFullRefreshDeferred);
+				EmitterViewModel->GetOrCreateEditorData().OnSummaryViewStateChanged().AddUObject(this, &UNiagaraStackViewModel::RequestRefreshDeferred);
 			}
 
 		}
@@ -137,7 +137,7 @@ void UNiagaraStackViewModel::InitializeWithViewModels(TSharedPtr<FNiagaraSystemV
 		StackRoot->OnExpansionInOverviewChanged().AddUObject(this, &UNiagaraStackViewModel::EntryExpansionInOverviewChanged);
 		StackRoot->OnDataObjectModified().AddUObject(this, &UNiagaraStackViewModel::EntryDataObjectModified);
 		StackRoot->OnRequestFullRefresh().AddUObject(this, &UNiagaraStackViewModel::EntryRequestFullRefresh);
-		StackRoot->OnRequestFullRefreshDeferred().AddUObject(this, &UNiagaraStackViewModel::EntryRequestFullRefreshDeferred);
+		StackRoot->OnRequestFullRefreshDeferred().AddUObject(this, &UNiagaraStackViewModel::RequestRefreshDeferred);
 		RootEntry = StackRoot;
 		RootEntries.Add(RootEntry);
 
@@ -159,7 +159,7 @@ void UNiagaraStackViewModel::InitializeWithRootEntry(UNiagaraStackEntry* InRootE
 	RootEntry->OnExpansionChanged().AddUObject(this, &UNiagaraStackViewModel::EntryExpansionChanged);
 	RootEntry->OnExpansionInOverviewChanged().AddUObject(this, &UNiagaraStackViewModel::EntryExpansionInOverviewChanged);
 	RootEntry->OnRequestFullRefresh().AddUObject(this, &UNiagaraStackViewModel::EntryRequestFullRefresh);
-	RootEntry->OnRequestFullRefreshDeferred().AddUObject(this, &UNiagaraStackViewModel::EntryRequestFullRefreshDeferred);
+	RootEntry->OnRequestFullRefreshDeferred().AddUObject(this, &UNiagaraStackViewModel::RequestRefreshDeferred);
 	RootEntries.Add(RootEntry);
 
 	bExternalRootEntry = true;
@@ -208,7 +208,7 @@ void UNiagaraStackViewModel::Reset()
 	TopLevelViewModels.Empty();
 
 	CurrentIssueCycleIndex = -1;
-	CurrentFocusedSearchMatchIndex = -1;
+	FocusedSearchResultCache.Reset();
 	bRestartSearch = false;
 	bRefreshPending = false;
 	bValidatorUpdatePending = false;
@@ -246,7 +246,7 @@ void UNiagaraStackViewModel::Tick()
 		{
 			RootEntry->RefreshChildren();
 			bRefreshPending = false;
-			InvalidateSearchResults();
+			RestartSearch();
 			bValidatorUpdatePending = true;
 		}
 
@@ -261,6 +261,7 @@ void UNiagaraStackViewModel::Tick()
 
 void UNiagaraStackViewModel::ResetSearchText()
 {
+	FocusedSearchResultCache.Reset();
 	CurrentSearchText = FText::GetEmpty();
 	// restarting the search with empty text will reset search results
 	bRestartSearch = true;	
@@ -284,6 +285,11 @@ void UNiagaraStackViewModel::OnSearchTextChanged(const FText& SearchText)
 		// also this can be triggered by multiple events, so better wait
 		bRestartSearch = true;
 	}
+
+	if(SearchText.IsEmpty())
+	{
+		FocusedSearchResultCache.Reset();
+	}
 }
 
 bool UNiagaraStackViewModel::IsSearching()
@@ -296,26 +302,54 @@ const TArray<UNiagaraStackViewModel::FSearchResult>& UNiagaraStackViewModel::Get
 	return CurrentSearchResults;
 }
 
-UNiagaraStackEntry* UNiagaraStackViewModel::GetCurrentFocusedEntry()
+UNiagaraStackEntry* UNiagaraStackViewModel::GetCurrentFocusedEntry() const
 {
-	if (CurrentFocusedSearchMatchIndex >= 0)
+	if (FocusedSearchResultCache.IsSet() && FocusedSearchResultCache->GetEntry() != nullptr && !FocusedSearchResultCache->GetEntry()->IsFinalized())
 	{
-		FSearchResult FocusedMatch = CurrentSearchResults[CurrentFocusedSearchMatchIndex];
-		return FocusedMatch.GetEntry();
+		return FocusedSearchResultCache->GetEntry();
 	}
+	
 	return nullptr;
 }
 
+int32 UNiagaraStackViewModel::GetCurrentFocusedEntryIndex() const
+{
+	if (FocusedSearchResultCache.IsSet() && FocusedSearchResultCache->GetEntry() != nullptr && !FocusedSearchResultCache->GetEntry()->IsFinalized())
+	{
+		for(int32 SearchResultIndex = 0; SearchResultIndex < CurrentSearchResults.Num(); SearchResultIndex++)
+		{
+			const FSearchResult& SearchResult = CurrentSearchResults[SearchResultIndex];
+			if(SearchResult.GetEntry() == FocusedSearchResultCache->GetEntry())
+			{
+				return SearchResultIndex;
+			}
+		}
+	}
+
+	return INDEX_NONE;
+}
+
 void UNiagaraStackViewModel::AddSearchScrollOffset(int NumberOfSteps)
 {
-	CurrentFocusedSearchMatchIndex += NumberOfSteps;
-	if (CurrentFocusedSearchMatchIndex >= CurrentSearchResults.Num())
+	int32 CurrentFocusedEntryIndex = GetCurrentFocusedEntryIndex();
+	int32 NewFocusedEntryIndex = CurrentFocusedEntryIndex + NumberOfSteps;
+	
+	if (NewFocusedEntryIndex >= CurrentSearchResults.Num())
 	{
-		CurrentFocusedSearchMatchIndex = 0;
+		NewFocusedEntryIndex = 0;
 	}
-	if (CurrentFocusedSearchMatchIndex < 0)
+	if (NewFocusedEntryIndex < 0)
 	{
-		CurrentFocusedSearchMatchIndex = CurrentSearchResults.Num() - 1;
+		NewFocusedEntryIndex = CurrentSearchResults.Num() - 1;
+	}
+
+	if(CurrentSearchResults.IsValidIndex(NewFocusedEntryIndex))
+	{
+		FocusedSearchResultCache = CurrentSearchResults[NewFocusedEntryIndex];
+	}
+	else
+	{
+		FocusedSearchResultCache.Reset();
 	}
 }
 
@@ -445,7 +479,6 @@ void UNiagaraStackViewModel::SearchTick()
 			SearchResult.GetEntry()->SetIsSearchResult(false);
 		}
 		CurrentSearchResults.Empty();
-		CurrentFocusedSearchMatchIndex = -1;
 		// generates ItemsToSearch, these will be processed on tick, in batches
 		if (CurrentSearchText.IsEmpty() == false)
 		{
@@ -472,7 +505,7 @@ void UNiagaraStackViewModel::SearchTick()
 			{
 				TArray<UNiagaraStackEntry::FStackSearchItem> SearchItems;
 				EntryToProcess->GetSearchItems(SearchItems);
-				TSet<FName> MatchedKeys;
+				
 				for (UNiagaraStackEntry::FStackSearchItem SearchItem : SearchItems)
 				{
 					if (!EntryToProcess->GetStackEditorDataKey().IsEmpty())
@@ -482,12 +515,9 @@ void UNiagaraStackViewModel::SearchTick()
 
 					if (ItemMatchesSearchCriteria(SearchItem)) 
 					{
-						if (MatchedKeys.Contains(SearchItem.Key) == false)
-						{
-							EntryToProcess->SetIsSearchResult(true);
-							CurrentSearchResults.Add({ ItemsToSearch[0].EntryPath, SearchItem });
-							MatchedKeys.Add(SearchItem.Key);
-						}
+						EntryToProcess->SetIsSearchResult(true);
+						CurrentSearchResults.Add({ ItemsToSearch[0].EntryPath, SearchItem });
+						break;
 					}
 				}
 			}
@@ -498,6 +528,26 @@ void UNiagaraStackViewModel::SearchTick()
 		{
 			// The search can change the child filtering so refresh it when the search finishes.
 			RootEntry->RefreshFilteredChildren();
+			
+			// we attempt to restore the focused search result we had from before
+			int32 CurrentSearchIndex = GetCurrentFocusedEntryIndex();
+			if(CurrentSearchIndex != INDEX_NONE)
+			{
+				FocusedSearchResultCache = CurrentSearchResults[CurrentSearchIndex];
+			}
+			else
+			{
+				// if we don't have a valid search result from before, we select the first result by default, if available
+				if(CurrentSearchResults.Num() > 0)
+				{
+					FocusedSearchResultCache = CurrentSearchResults[0];
+				}
+				else
+				{
+					FocusedSearchResultCache.Reset();
+				}
+			}
+			
 			SearchCompletedDelegate.Broadcast();
 		}
 	}
@@ -542,9 +592,12 @@ void UNiagaraStackViewModel::GeneratePathForEntry(UNiagaraStackEntry* Root, UNia
 	}
 }
 
-void UNiagaraStackViewModel::InvalidateSearchResults()
+void UNiagaraStackViewModel::RestartSearch()
 {
-	bRestartSearch = true;
+	if(bRestartSearch == false)
+	{
+		bRestartSearch = true;
+	}
 }
 
 UNiagaraStackEntry* UNiagaraStackViewModel::GetRootEntry()
@@ -609,7 +662,7 @@ void UNiagaraStackViewModel::SetShowAllAdvanced(bool bInShowAllAdvanced)
 		}
 	}
 
-	InvalidateSearchResults();
+	RestartSearch();
 	RootEntry->RefreshFilteredChildren();
 }
 
@@ -636,7 +689,7 @@ void UNiagaraStackViewModel::SetShowOutputs(bool bInShowOutputs)
 	}
 	
 	// Showing outputs changes indenting so a full refresh is needed.
-	InvalidateSearchResults();
+	RestartSearch();
 	RootEntry->RefreshChildren();
 }
 
@@ -663,7 +716,7 @@ void UNiagaraStackViewModel::SetShowLinkedInputs(bool bInShowLinkedInputs)
 	}
 
 	// Showing linked inputs changes indenting so a full refresh is needed.
-	InvalidateSearchResults();
+	RestartSearch();
 	RootEntry->RefreshChildren();
 }
 
@@ -689,7 +742,7 @@ void UNiagaraStackViewModel::SetShowOnlyIssues(bool bInShowOnlyIssues)
 		}
 	}
 
-	InvalidateSearchResults();
+	RestartSearch();
 	RootEntry->RefreshChildren();
 }
 
@@ -734,7 +787,7 @@ void UNiagaraStackViewModel::EntryStructureChanged(ENiagaraStructureChangedFlags
 	// if the structure changed, additionally invalidate search results
 	if((Flags & StructureChanged) != 0)
 	{
-		InvalidateSearchResults();
+		RestartSearch();
 	}
 }
 
@@ -744,7 +797,7 @@ void UNiagaraStackViewModel::EntryDataObjectModified(TArray<UObject*> ChangedObj
 	{
 		SystemViewModel.Pin()->NotifyDataObjectChanged(ChangedObjects, ChangeType);
 	}
-	InvalidateSearchResults();
+	RestartSearch();
 	DataObjectChangedDelegate.Broadcast(ChangedObjects, ChangeType);
 }
 
@@ -754,7 +807,7 @@ void UNiagaraStackViewModel::EntryRequestFullRefresh()
 	RootEntry->RefreshChildren();
 }
 
-void UNiagaraStackViewModel::EntryRequestFullRefreshDeferred()
+void UNiagaraStackViewModel::RequestRefreshDeferred()
 {
 	bRefreshPending = true;
 }
diff --git a/Engine/Plugins/FX/Niagara/Source/NiagaraEditor/Public/ViewModels/Stack/NiagaraStackViewModel.h b/Engine/Plugins/FX/Niagara/Source/NiagaraEditor/Public/ViewModels/Stack/NiagaraStackViewModel.h
--- a/Engine/Plugins/FX/Niagara/Source/NiagaraEditor/Public/ViewModels/Stack/NiagaraStackViewModel.h
+++ b/Engine/Plugins/FX/Niagara/Source/NiagaraEditor/Public/ViewModels/Stack/NiagaraStackViewModel.h
@@ -50,7 +50,7 @@ public:
 	{
 		TArray<UNiagaraStackEntry*> EntryPath;
 		UNiagaraStackEntry::FStackSearchItem MatchingItem;
-		UNiagaraStackEntry* GetEntry()
+		UNiagaraStackEntry* GetEntry() const
 		{
 			return EntryPath.Num() > 0 ? 
 				EntryPath[EntryPath.Num() - 1] : 
@@ -120,8 +120,8 @@ public:
 	void SetSearchTextExternal(const FText& NewSearchText);
 	bool IsSearching();
 	const TArray<FSearchResult>& GetCurrentSearchResults();
-	int GetCurrentFocusedMatchIndex() const { return CurrentFocusedSearchMatchIndex; }
-	UNiagaraStackEntry* GetCurrentFocusedEntry();
+	UNiagaraStackEntry* GetCurrentFocusedEntry() const;
+	int32 GetCurrentFocusedEntryIndex() const;
 	void AddSearchScrollOffset(int NumberOfSteps);
 
 	void OnCycleThroughIssues(TSharedPtr<FTopLevelViewModel> TopLevelToCycle);
@@ -141,7 +141,7 @@ public:
 
 	void Reset();
 	bool HasIssues() const;
-	void Refresh() { bRefreshPending = true; }
+	void RequestRefreshDeferred();
 	void RequestValidationUpdate() { bValidatorUpdatePending = true; }
 
 	bool ShouldHideDisabledModules() const;
@@ -168,7 +168,6 @@ private:
 	void EntryStructureChanged(ENiagaraStructureChangedFlags Flags);
 	void EntryDataObjectModified(TArray<UObject*> ChangedObjects, ENiagaraDataObjectChange ChangeType);
 	void EntryRequestFullRefresh();
-	void EntryRequestFullRefreshDeferred();
 	void RefreshTopLevelViewModels();
 	void RefreshHasIssues();
 	void EmitterParentRemoved();
@@ -179,7 +178,7 @@ private:
 	bool ItemMatchesSearchCriteria(UNiagaraStackEntry::FStackSearchItem SearchItem);
 	void GeneratePathForEntry(UNiagaraStackEntry* Root, UNiagaraStackEntry* Entry, TArray<UNiagaraStackEntry*> CurrentPath, TArray<UNiagaraStackEntry*>& EntryPath) const;
 
-	void InvalidateSearchResults();
+	void RestartSearch();
 	void UpdateStackWithValidationResults();
 
 private:
@@ -203,11 +202,11 @@ private:
 
 	// ~Search stuff
 	FText CurrentSearchText;
-	int CurrentFocusedSearchMatchIndex;
 	FOnSearchCompleted SearchCompletedDelegate;
 	TArray<FSearchWorkItem> ItemsToSearch;
 	TArray<FSearchResult> CurrentSearchResults;
 	static const double MaxSearchTime;
+	TOptional<FSearchResult> FocusedSearchResultCache;
 	bool bRestartSearch;
 	bool bRefreshPending;
 	bool bValidatorUpdatePending;
diff --git a/Engine/Plugins/FX/Niagara/Source/NiagaraEditorWidgets/Private/SNiagaraStack.cpp b/Engine/Plugins/FX/Niagara/Source/NiagaraEditorWidgets/Private/SNiagaraStack.cpp
--- a/Engine/Plugins/FX/Niagara/Source/NiagaraEditorWidgets/Private/SNiagaraStack.cpp
+++ b/Engine/Plugins/FX/Niagara/Source/NiagaraEditorWidgets/Private/SNiagaraStack.cpp
@@ -408,7 +408,7 @@ static void ExpandAllInPath(const TArray<UNiagaraStackEntry*>& EntryPath)
 
 FReply SNiagaraStack::ScrollToNextMatch()
 {
-	int NextMatchIndex = StackViewModel->GetCurrentFocusedMatchIndex() + 1;
+	int NextMatchIndex = StackViewModel->GetCurrentFocusedEntryIndex() + 1;
 	TArray<UNiagaraStackViewModel::FSearchResult> CurrentSearchResults = StackViewModel->GetCurrentSearchResults();
 	if (CurrentSearchResults.Num() != 0)
 	{
@@ -426,7 +426,7 @@ FReply SNiagaraStack::ScrollToNextMatch()
 
 FReply SNiagaraStack::ScrollToPreviousMatch()
 {
-	const int PreviousMatchIndex = StackViewModel->GetCurrentFocusedMatchIndex() - 1;
+	const int PreviousMatchIndex = StackViewModel->GetCurrentFocusedEntryIndex() - 1;
 	TArray<UNiagaraStackViewModel::FSearchResult> CurrentSearchResults = StackViewModel->GetCurrentSearchResults();
 	if (CurrentSearchResults.Num() != 0)
 	{
@@ -453,8 +453,7 @@ void SNiagaraStack::AddSearchScrollOffset(int NumberOfSteps)
 	}
 
 	StackViewModel->AddSearchScrollOffset(NumberOfSteps);
-
-	StackTree->RequestScrollIntoView(StackViewModel->GetCurrentFocusedEntry());
+	ScrollFocusedEntryIntoView();
 }
 
 TOptional<SSearchBox::FSearchResultData> SNiagaraStack::GetSearchResultData() const
@@ -463,7 +462,7 @@ TOptional<SSearchBox::FSearchResultData> SNiagaraStack::GetSearchResultData() co
 	{
 		return TOptional<SSearchBox::FSearchResultData>();
 	}
-	return TOptional<SSearchBox::FSearchResultData>({ StackViewModel->GetCurrentSearchResults().Num(), StackViewModel->GetCurrentFocusedMatchIndex() + 1 });
+	return TOptional<SSearchBox::FSearchResultData>({ StackViewModel->GetCurrentSearchResults().Num(), StackViewModel->GetCurrentFocusedEntryIndex() + 1 });
 }
 
 bool SNiagaraStack::GetIsSearching() const
@@ -592,7 +591,7 @@ FReply SNiagaraStack::OnRowAcceptDrop(const FDragDropEvent& InDragDropEvent, EIt
 void SNiagaraStack::OnStackSearchComplete()
 {
 	ExpandSearchResults();
-	ScrollToNextMatch();
+	ScrollFocusedEntryIntoView();
 }
 
 void SNiagaraStack::ExpandSearchResults()
@@ -634,6 +633,11 @@ FSlateColor SNiagaraStack::GetTextColorForItem(UNiagaraStackEntry* Item) const
 	return FSlateColor::UseForeground();
 }
 
+void SNiagaraStack::ScrollFocusedEntryIntoView()
+{
+	StackTree->RequestScrollIntoView(StackViewModel->GetCurrentFocusedEntry());
+}
+
 TSharedRef<ITableRow> SNiagaraStack::OnGenerateRowForStackItem(UNiagaraStackEntry* Item, const TSharedRef<STableViewBase>& OwnerTable)
 {
 	TSharedRef<SNiagaraStackTableRow> Container = ConstructContainerForItem(Item);
diff --git a/Engine/Plugins/FX/Niagara/Source/NiagaraEditorWidgets/Private/SNiagaraStack.h b/Engine/Plugins/FX/Niagara/Source/NiagaraEditorWidgets/Private/SNiagaraStack.h
--- a/Engine/Plugins/FX/Niagara/Source/NiagaraEditorWidgets/Private/SNiagaraStack.h
+++ b/Engine/Plugins/FX/Niagara/Source/NiagaraEditorWidgets/Private/SNiagaraStack.h
@@ -93,6 +93,7 @@ private:
 	void OnSearchBoxTextCommitted(const FText& NewText, ETextCommit::Type CommitInfo);	
 	void OnSearchBoxSearch(SSearchBox::SearchDirection Direction);
 	FSlateColor GetTextColorForItem(UNiagaraStackEntry* Item) const;
+	void ScrollFocusedEntryIntoView();
 	void AddSearchScrollOffset(int NumberOfSteps);
 	void OnStackSearchComplete();
 	void ExpandSearchResults();
