- Remove deprecated & hidden functions from the HLSL display / auto complete list
- Deprecate the old InterlockedMaxFloatGridValueFunction function

diff --git a/Engine/Plugins/FX/Niagara/Source/Niagara/Private/NiagaraDataInterfaceRasterizationGrid3D.cpp b/Engine/Plugins/FX/Niagara/Source/Niagara/Private/NiagaraDataInterfaceRasterizationGrid3D.cpp
--- a/Engine/Plugins/FX/Niagara/Source/Niagara/Private/NiagaraDataInterfaceRasterizationGrid3D.cpp
+++ b/Engine/Plugins/FX/Niagara/Source/Niagara/Private/NiagaraDataInterfaceRasterizationGrid3D.cpp
@@ -354,9 +354,8 @@ void UNiagaraDataInterfaceRasterizationGrid3D::GetFunctions(TArray<FNiagaraFunct
 
 		Sig.bSupportsCPU = false;
 		Sig.bSupportsGPU = true;
-
-		Sig.bSupportsCPU = false;
-		Sig.bSupportsGPU = true;
+		Sig.bSoftDeprecatedFunction = true;
+		Sig.bHidden = true;
 #if WITH_EDITORONLY_DATA
 		Sig.Description = NSLOCTEXT("Niagara", "NiagaraDataInterfaceGridColl2D_SetValueFunction", "Set the value at a specific index. Note that this is an older way of working with Grids. Consider using the SetFloat or other typed, named functions or parameter map variables with StackContext namespace instead.");
 #endif
diff --git a/Engine/Plugins/FX/Niagara/Source/NiagaraEditor/Private/EdGraphSchema_Niagara.cpp b/Engine/Plugins/FX/Niagara/Source/NiagaraEditor/Private/EdGraphSchema_Niagara.cpp
--- a/Engine/Plugins/FX/Niagara/Source/NiagaraEditor/Private/EdGraphSchema_Niagara.cpp
+++ b/Engine/Plugins/FX/Niagara/Source/NiagaraEditor/Private/EdGraphSchema_Niagara.cpp
@@ -2193,6 +2193,11 @@ TArray<TPair<FString, FString>> UEdGraphSchema_Niagara::GetDataInterfaceFunction
 			FunctionPrototypes.Reserve(FunctionSignatures.Num());
 			for (const FNiagaraFunctionSignature& FunctionSignature : FunctionSignatures)
 			{
+				if (FunctionSignature.bSoftDeprecatedFunction || FunctionSignature.bHidden )
+				{
+					continue;
+				}
+
 				FString Prototype = FHlslNiagaraTranslator::GenerateFunctionHlslPrototype(DataInterfaceName, FunctionSignature);
 				if ( Prototype.Len() > 0 )
 				{
