Define IHeterogeneousVolumeInterface on FPrimitiveSceneProxy for communicating per-volume parameters to the heterogeneous volumes system.

Implement IHeterogeneousVolumeInterface for FNiagaraSceneProxy when the element SceneProxy is declared to be a heterogeneous volume.

diff --git a/Engine/Plugins/FX/Niagara/Source/Niagara/Private/NiagaraComponent.cpp b/Engine/Plugins/FX/Niagara/Source/Niagara/Private/NiagaraComponent.cpp
--- a/Engine/Plugins/FX/Niagara/Source/Niagara/Private/NiagaraComponent.cpp
+++ b/Engine/Plugins/FX/Niagara/Source/Niagara/Private/NiagaraComponent.cpp
@@ -195,6 +195,31 @@ FNiagaraSceneProxy::FNiagaraSceneProxy(UNiagaraComponent* InComponent)
 		bAlwaysHasVelocity = RenderData->HasAnyMotionBlurEnabled();
 		bIsHeterogeneousVolume = RenderData->HasAnyHeterogeneousVolumesEnabled();
 
+		// Initialize Heterogeneous Volume properties
+		if (bIsHeterogeneousVolume)
+		{
+			// Temporarily dependening on user variables until a NiagaraVolumeRenderer is created
+			FNiagaraVariable VoxelResolutionVariable(FNiagaraTypeDefinition::GetIntDef(), TEXT("ResolutionMaxAxis"));
+			float ResolutionMaxAxis = InComponent->GetOverrideParameters().GetParameterValueOrDefault(VoxelResolutionVariable, 0);
+			ResolutionMaxAxis = FMath::Max(ResolutionMaxAxis, 0);
+
+			FNiagaraVariable WorldSpaceSizeVariable(FNiagaraTypeDefinition::GetVec3Def(), TEXT("WorldSpaceSize"));
+			FVector3f WorldSpaceSize = InComponent->GetOverrideParameters().GetParameterValueOrDefault(WorldSpaceSizeVariable, FVector3f(0));
+
+			float WorldSpaceSizeMaxInv = WorldSpaceSize.GetMax() > 0.0 ? 1.0 / WorldSpaceSize.GetMax() : 0.0;
+			FVector3f ResolutionFactor = WorldSpaceSize * WorldSpaceSizeMaxInv;
+
+			FVector3f VolumeResolutionV3f = ResolutionFactor * ResolutionMaxAxis;
+			FIntVector VolumeResolutionInt3 = FIntVector(
+				FMath::CeilToInt(VolumeResolutionV3f.X),
+				FMath::CeilToInt(VolumeResolutionV3f.Y),
+				FMath::CeilToInt(VolumeResolutionV3f.Z));
+
+			HeterogeneousVolumeData.VoxelResolution = VolumeResolutionInt3;
+			HeterogeneousVolumeData.LightingDownsampleFactor = 1;
+			HeterogeneousVolumeData.MinimumVoxelSize = 0.1;
+		}
+
 		SystemStatID = InComponent->GetAsset()->GetStatID(false, false);
 	}
 }
diff --git a/Engine/Plugins/FX/Niagara/Source/Niagara/Public/NiagaraSceneProxy.h b/Engine/Plugins/FX/Niagara/Source/Niagara/Public/NiagaraSceneProxy.h
--- a/Engine/Plugins/FX/Niagara/Source/Niagara/Public/NiagaraSceneProxy.h
+++ b/Engine/Plugins/FX/Niagara/Source/Niagara/Public/NiagaraSceneProxy.h
@@ -4,6 +4,7 @@
 
 #include "Particles/ParticlePerfStats.h"
 #include "PrimitiveSceneProxy.h"
+#include "HeterogeneousVolumeInterface.h"
 
 class FNiagaraGpuComputeDispatchInterface;
 class FNiagaraSystemRenderData;
@@ -89,6 +90,8 @@ private:
 
 	uint32 GetAllocatedSize() const;
 
+	virtual const IHeterogeneousVolumeInterface* GetHeterogeneousVolumeInterface() const override { return &HeterogeneousVolumeData; }
+
 private:
 	/** Custom Uniform Buffers, allows us to have renderer specific data packed inside such as pre-skinned bounds. */
 	mutable TMap<uint32, TUniformBuffer<FPrimitiveUniformShaderParameters>*> CustomUniformBuffers;
@@ -103,5 +106,8 @@ private:
 	TStatId SystemStatID;
 
 	FDynamicData DynamicData;
+
+	// Consider wrapping in a TUniquePtr
+	FHeterogeneousVolumeData HeterogeneousVolumeData;
 };
 
diff --git a/Engine/Source/Runtime/Engine/Private/HeterogeneousVolumeInterface.cpp b/Engine/Source/Runtime/Engine/Private/HeterogeneousVolumeInterface.cpp
new file mode 100644
--- /dev/null
+++ b/Engine/Source/Runtime/Engine/Private/HeterogeneousVolumeInterface.cpp
@@ -0,0 +1,22 @@
+// Copyright Epic Games, Inc. All Rights Reserved.
+
+/*=============================================================================
+	HeterogeneousVolumeInterface.cpp: Primitive scene proxy implementation.
+=============================================================================*/
+
+#include "HeterogeneousVolumeInterface.h"
+
+IHeterogeneousVolumeInterface::~IHeterogeneousVolumeInterface()
+{
+}
+
+FHeterogeneousVolumeData::FHeterogeneousVolumeData()
+	: VoxelResolution(FIntVector::ZeroValue)
+	, MinimumVoxelSize(0.1)
+	, LightingDownsampleFactor(1.0)
+{
+}
+
+FHeterogeneousVolumeData::~FHeterogeneousVolumeData()
+{
+}
diff --git a/Engine/Source/Runtime/Engine/Public/HeterogeneousVolumeInterface.h b/Engine/Source/Runtime/Engine/Public/HeterogeneousVolumeInterface.h
new file mode 100644
--- /dev/null
+++ b/Engine/Source/Runtime/Engine/Public/HeterogeneousVolumeInterface.h
@@ -0,0 +1,36 @@
+// Copyright Epic Games, Inc. All Rights Reserved.
+
+#pragma once
+
+#include "CoreMinimal.h"
+
+class IHeterogeneousVolumeInterface
+{
+public:
+	ENGINE_API virtual ~IHeterogeneousVolumeInterface();
+
+	// Volume
+	ENGINE_API virtual FIntVector GetVoxelResolution() const = 0;
+	ENGINE_API virtual float GetMinimumVoxelSize() const = 0;
+
+	// Lighting
+	ENGINE_API virtual float GetLightingDownsampleFactor() const = 0;
+};
+
+class FHeterogeneousVolumeData : public IHeterogeneousVolumeInterface
+{
+public:
+	ENGINE_API FHeterogeneousVolumeData();
+	ENGINE_API virtual ~FHeterogeneousVolumeData();
+
+	// Volume
+	ENGINE_API virtual FIntVector GetVoxelResolution() const { return VoxelResolution; }
+	ENGINE_API virtual float GetMinimumVoxelSize() const { return MinimumVoxelSize; }
+
+	// Lighting
+	ENGINE_API virtual float GetLightingDownsampleFactor() const { return LightingDownsampleFactor; }
+
+	FIntVector VoxelResolution;
+	float MinimumVoxelSize;
+	float LightingDownsampleFactor;
+};
diff --git a/Engine/Source/Runtime/Engine/Public/PrimitiveSceneProxy.h b/Engine/Source/Runtime/Engine/Public/PrimitiveSceneProxy.h
--- a/Engine/Source/Runtime/Engine/Public/PrimitiveSceneProxy.h
+++ b/Engine/Source/Runtime/Engine/Public/PrimitiveSceneProxy.h
@@ -33,6 +33,7 @@ class FColorVertexBuffer;
 struct FInstanceUpdateCmdBuffer;
 class FRayTracingGeometry;
 class FVertexFactory;
+class IHeterogeneousVolumeInterface;
 
 namespace Nanite
 {
@@ -795,6 +796,11 @@ public:
 		return bIsHeterogeneousVolume;
 	}
 
+	virtual const IHeterogeneousVolumeInterface* GetHeterogeneousVolumeInterface() const
+	{
+		return nullptr;
+	}
+
 	/**
 	 * Returns true if all meshes drawn by this proxy support GPU scene. 
 	 */
