Removed DynamicPins array and refactored use cases that required it. This fixes various UI issues in regards to pin context actions

diff --git a/Engine/Plugins/FX/Niagara/Source/NiagaraEditor/Private/NiagaraNodeConvert.h b/Engine/Plugins/FX/Niagara/Source/NiagaraEditor/Private/NiagaraNodeConvert.h
--- a/Engine/Plugins/FX/Niagara/Source/NiagaraEditor/Private/NiagaraNodeConvert.h
+++ b/Engine/Plugins/FX/Niagara/Source/NiagaraEditor/Private/NiagaraNodeConvert.h
@@ -106,7 +106,7 @@ public:
 	virtual TSharedPtr<SGraphNode> CreateVisualWidget() override;
 	virtual void  AutowireNewNode(UEdGraphPin* FromPin)override;
 	virtual FText GetNodeTitle(ENodeTitleType::Type TitleType) const override;
-	virtual bool CanMovePin(const UEdGraphPin* Pin, int32 DirectionToMove) const override { return false; }
+	//virtual bool CanMovePin(const UEdGraphPin* Pin, int32 DirectionToMove) const override { return false; }
 
 	//~ UNiagaraNode interface
 	virtual void Compile(class FHlslNiagaraTranslator* Translator, TArray<int32>& Outputs) override;
diff --git a/Engine/Plugins/FX/Niagara/Source/NiagaraEditor/Private/NiagaraNodeFunctionCall.cpp b/Engine/Plugins/FX/Niagara/Source/NiagaraEditor/Private/NiagaraNodeFunctionCall.cpp
--- a/Engine/Plugins/FX/Niagara/Source/NiagaraEditor/Private/NiagaraNodeFunctionCall.cpp
+++ b/Engine/Plugins/FX/Niagara/Source/NiagaraEditor/Private/NiagaraNodeFunctionCall.cpp
@@ -108,7 +108,6 @@ void UNiagaraNodeFunctionCall::PostLoad()
 		const int32 NiagaraVer = GetLinkerCustomVersion(FNiagaraCustomVersion::GUID);
 		if (NiagaraVer < FNiagaraCustomVersion::DynamicPinNodeFixup && AllowDynamicPins())
 		{
-			DynamicPins.Empty();
 			bool bFoundInputAdd = false;
 			bool bFoundOutputAdd = false;
 			for (const UEdGraphPin* Pin : Pins)
@@ -124,10 +123,6 @@ void UNiagaraNodeFunctionCall::PostLoad()
 						bFoundOutputAdd = true;
 					}
 				}
-				else
-				{
-					DynamicPins.Add(Pin->PinId);
-				}
 			}
 
 			if (!bFoundInputAdd)
@@ -262,6 +257,34 @@ void UNiagaraNodeFunctionCall::AddOrphanedStaticSwitchPinForDataRetention(FNiaga
 	NewPin->bOrphanedPin = true;
 }
 
+void UNiagaraNodeFunctionCall::RemoveAllDynamicPins()
+{
+	FScopedTransaction RemovePinTransaction(LOCTEXT("RemoveAllDynamicPinsTransaction", "Remove all dynamic pins"));
+	if(AllowDynamicPins())
+	{
+		FPinCollectorArray PinArray;
+		GetInputPins(PinArray);
+
+		for(UEdGraphPin* Pin : PinArray)
+		{
+			if(CanModifyPin(Pin) && IsBaseSignatureOfDataInterfaceFunction(Pin) == false)
+			{
+				RemoveDynamicPin(Pin);
+			}
+		}
+
+		GetOutputPins(PinArray);
+
+		for(UEdGraphPin* Pin : PinArray)
+		{
+			if(CanModifyPin(Pin) && IsBaseSignatureOfDataInterfaceFunction(Pin) == false)
+			{
+				RemoveDynamicPin(Pin);
+			}
+		}
+	}
+}
+
 UEdGraphPin* UNiagaraNodeFunctionCall::AddStaticSwitchInputPin(FNiagaraVariable Input)
 {
 	UNiagaraGraph* Graph = GetCalledGraph();
@@ -298,6 +321,24 @@ UEdGraphPin* UNiagaraNodeFunctionCall::AddStaticSwitchInputPin(FNiagaraVariable
 	return NewPin;
 }
 
+bool UNiagaraNodeFunctionCall::CanModifyPin(const UEdGraphPin* Pin) const
+{
+	if(IsAddPin(Pin))
+	{
+		return false;
+	}
+	
+	if(Signature.IsValid() && (Signature.VariadicInput() || Signature.VariadicOutput()))
+	{
+		if(IsBaseSignatureOfDataInterfaceFunction(Pin) == false)
+		{
+			return true;
+		}
+	}
+
+	return false;
+}
+
 void UNiagaraNodeFunctionCall::PostEditChangeProperty(struct FPropertyChangedEvent& PropertyChangedEvent)
 {
 	if (PropertyChangedEvent.Property != nullptr)
@@ -1587,11 +1628,12 @@ void UNiagaraNodeFunctionCall::RefreshSignature()
 
 				for (int32 i = 0; i < FoundPins.Num(); ++i)
 				{
-					UEdGraphPin* Pin = FoundPins[i];
+					UEdGraphPin* InputPin = FoundPins[i];
+					FNiagaraVariable InputVariable = UEdGraphSchema_Niagara::PinToNiagaraVariable(InputPin);
 
-					if (DynamicPins.Contains(Pin->PinId))
+					if(!BaseSig->Inputs.Contains(InputVariable))
 					{
-						Signature.AddInput(UEdGraphSchema_Niagara::PinToNiagaraVariable(Pin), FText::FromString(Pin->PinToolTip));
+						Signature.AddInput(InputVariable, FText::FromString(InputPin->PinToolTip));
 					}
 				}
 
@@ -1599,11 +1641,12 @@ void UNiagaraNodeFunctionCall::RefreshSignature()
 
 				for (int32 i = 0; i < FoundPins.Num(); ++i)
 				{
-					UEdGraphPin* Pin = FoundPins[i];
+					UEdGraphPin* OutputPin = FoundPins[i];
+					FNiagaraVariable InputVariable = UEdGraphSchema_Niagara::PinToNiagaraVariable(OutputPin);
 
-					if (DynamicPins.Contains(Pin->PinId))
+					if(!BaseSig->Outputs.Contains(InputVariable))
 					{
-						Signature.AddOutput(UEdGraphSchema_Niagara::PinToNiagaraVariable(Pin), FText::FromString(Pin->PinToolTip));
+						Signature.AddOutput(UEdGraphSchema_Niagara::PinToNiagaraVariable(OutputPin), FText::FromString(OutputPin->PinToolTip));
 					}
 				}
 
@@ -1614,6 +1657,41 @@ void UNiagaraNodeFunctionCall::RefreshSignature()
 	Signature = FNiagaraFunctionSignature();
 }
 
+bool UNiagaraNodeFunctionCall::IsBaseSignatureOfDataInterfaceFunction(const UEdGraphPin* Pin) const
+{
+	if ((Signature.Inputs.Num() > 0) && Signature.Inputs[0].GetType().IsDataInterface())
+	{
+		UNiagaraDataInterface* CDO = CastChecked<UNiagaraDataInterface>(Signature.Inputs[0].GetType().GetClass()->GetDefaultObject());
+		TArray<FNiagaraFunctionSignature> BaseDIFuncs;
+		CDO->GetFunctions(BaseDIFuncs);
+		if (FNiagaraFunctionSignature* BaseSig = BaseDIFuncs.FindByPredicate([&](const FNiagaraFunctionSignature& CheckSig) { return Signature.Name == CheckSig.Name; }))
+		{
+			FPinCollectorArray FoundPins;
+			TArray<FNiagaraVariable> InputOrOutputVariables;
+
+			if(Pin->Direction == EGPD_Input)
+			{
+				GetInputPins(FoundPins);
+				InputOrOutputVariables = BaseSig->Inputs;
+			}
+			else
+			{
+				GetOutputPins(FoundPins);
+				InputOrOutputVariables = BaseSig->Outputs;
+			}
+			
+			FNiagaraVariable Variable = UEdGraphSchema_Niagara::PinToNiagaraVariable(Pin);
+
+			if(InputOrOutputVariables.Contains(Variable))
+			{
+				return true;
+			}
+		}
+	}
+
+	return false;
+}
+
 void UNiagaraNodeFunctionCall::ComputeNodeName(FString SuggestedName, bool bForceSuggestion)
 {
 	FString FunctionName = FunctionScript ? FunctionScript->GetName() : Signature.GetNameString();
diff --git a/Engine/Plugins/FX/Niagara/Source/NiagaraEditor/Private/NiagaraNodeParameterMapBase.cpp b/Engine/Plugins/FX/Niagara/Source/NiagaraEditor/Private/NiagaraNodeParameterMapBase.cpp
--- a/Engine/Plugins/FX/Niagara/Source/NiagaraEditor/Private/NiagaraNodeParameterMapBase.cpp
+++ b/Engine/Plugins/FX/Niagara/Source/NiagaraEditor/Private/NiagaraNodeParameterMapBase.cpp
@@ -158,7 +158,7 @@ void UNiagaraNodeParameterMapBase::SetPinName(UEdGraphPin* InPin, const FName& I
 
 bool UNiagaraNodeParameterMapBase::CanRenamePin(const UEdGraphPin* Pin) const
 {
-	if (IsAddPin(Pin))
+	if(Super::CanRenamePin(Pin) == false)
 	{
 		return false;
 	}
@@ -167,11 +167,6 @@ bool UNiagaraNodeParameterMapBase::CanRenamePin(const UEdGraphPin* Pin) const
 	return FNiagaraParameterUtilities::TestCanRenameWithMessage(Pin->PinName, Unused);
 }
 
-bool UNiagaraNodeParameterMapBase::CanRemovePin(const UEdGraphPin* Pin) const
-{
-	return IsAddPin(Pin) == false && IsParameterMapPin(Pin) == false;
-}
-
 bool UNiagaraNodeParameterMapBase::GetIsPinEditNamespaceModifierPending(const UEdGraphPin* Pin)
 {
 	return PinsGuidsWithEditNamespaceModifierPending.Contains(Pin->PersistentGuid);
diff --git a/Engine/Plugins/FX/Niagara/Source/NiagaraEditor/Private/NiagaraNodeParameterMapBase.h b/Engine/Plugins/FX/Niagara/Source/NiagaraEditor/Private/NiagaraNodeParameterMapBase.h
--- a/Engine/Plugins/FX/Niagara/Source/NiagaraEditor/Private/NiagaraNodeParameterMapBase.h
+++ b/Engine/Plugins/FX/Niagara/Source/NiagaraEditor/Private/NiagaraNodeParameterMapBase.h
@@ -33,15 +33,13 @@ public:
 
 	virtual void GetNodeContextMenuActions(class UToolMenu* Menu, class UGraphNodeContextMenuContext* Context) const override;
 	virtual void GetPinHoverText(const UEdGraphPin& Pin, FString& HoverTextOut) const override;
-	virtual bool CanModifyExistingPins() const override { return false; }
 	virtual bool IncludeParentNodeContextMenu() const override { return true; }
 
 	void SetPinName(UEdGraphPin* InPin, const FName& InName);
 
 	virtual bool CanRenamePinFromContextMenu(const UEdGraphPin* Pin) const override { return false; }
-
+	
 	virtual bool CanRenamePin(const UEdGraphPin* Pin) const override;
-	virtual bool CanRemovePin(const UEdGraphPin* Pin) const override;
 	
 	virtual FName GetNewPinDefaultNamespace() const { return PARAM_MAP_LOCAL_MODULE_STR; }
 
diff --git a/Engine/Plugins/FX/Niagara/Source/NiagaraEditor/Private/NiagaraNodeParameterMapFor.cpp b/Engine/Plugins/FX/Niagara/Source/NiagaraEditor/Private/NiagaraNodeParameterMapFor.cpp
--- a/Engine/Plugins/FX/Niagara/Source/NiagaraEditor/Private/NiagaraNodeParameterMapFor.cpp
+++ b/Engine/Plugins/FX/Niagara/Source/NiagaraEditor/Private/NiagaraNodeParameterMapFor.cpp
@@ -38,6 +38,25 @@ void UNiagaraNodeParameterMapFor::GetPinHoverText(const UEdGraphPin& Pin, FStrin
 	Super::GetPinHoverText(Pin, HoverTextOut);
 }
 
+bool UNiagaraNodeParameterMapFor::CanModifyPin(const UEdGraphPin* Pin) const
+{
+	if(Super::CanModifyPin(Pin) == false)
+	{
+		return false;
+	}
+
+	FPinCollectorArray InputPins;
+	GetInputPins(InputPins);
+
+	// we don't allow the modification of the iteration count pin
+	if(Pin == InputPins[1])
+	{
+		return false;
+	}
+
+	return true;
+}
+
 void UNiagaraNodeParameterMapFor::Compile(FHlslNiagaraTranslator* Translator, TArray<int32>& Outputs)
 {
 	FPinCollectorArray InputPins;
@@ -133,6 +152,25 @@ void UNiagaraNodeParameterMapForWithContinue::GetPinHoverText(const UEdGraphPin&
 	Super::GetPinHoverText(Pin, HoverTextOut);
 }
 
+bool UNiagaraNodeParameterMapForWithContinue::CanModifyPin(const UEdGraphPin* Pin) const
+{
+	if(Super::CanModifyPin(Pin) == false)
+	{
+		return false;
+	}
+
+	FPinCollectorArray InputPins;
+	GetInputPins(InputPins);
+
+	// we don't allow the modification of the iteration count pin or the iteration enabled pin
+	if(Pin == InputPins[1] || Pin == InputPins[2])
+	{
+		return false;
+	}
+
+	return true;
+}
+
 bool UNiagaraNodeParameterMapForWithContinue::SkipPinCompilation(UEdGraphPin* Pin) const
 {
 	FPinCollectorArray InputPins;
diff --git a/Engine/Plugins/FX/Niagara/Source/NiagaraEditor/Private/NiagaraNodeParameterMapFor.h b/Engine/Plugins/FX/Niagara/Source/NiagaraEditor/Private/NiagaraNodeParameterMapFor.h
--- a/Engine/Plugins/FX/Niagara/Source/NiagaraEditor/Private/NiagaraNodeParameterMapFor.h
+++ b/Engine/Plugins/FX/Niagara/Source/NiagaraEditor/Private/NiagaraNodeParameterMapFor.h
@@ -24,6 +24,8 @@ public:
 	virtual void GetPinHoverText(const UEdGraphPin& Pin, FString& HoverTextOut) const override;
 	
 protected:
+	virtual bool CanModifyPin(const UEdGraphPin* Pin) const override;
+	
 	virtual bool SkipPinCompilation(UEdGraphPin* Pin) const override;
 };
 
@@ -42,6 +44,7 @@ public:
 	virtual void GetPinHoverText(const UEdGraphPin& Pin, FString& HoverTextOut) const override;
 	
 protected:
+	virtual bool CanModifyPin(const UEdGraphPin* Pin) const override;
 	virtual bool SkipPinCompilation(UEdGraphPin* Pin) const override;
 };
 
diff --git a/Engine/Plugins/FX/Niagara/Source/NiagaraEditor/Private/NiagaraNodeParameterMapGet.cpp b/Engine/Plugins/FX/Niagara/Source/NiagaraEditor/Private/NiagaraNodeParameterMapGet.cpp
--- a/Engine/Plugins/FX/Niagara/Source/NiagaraEditor/Private/NiagaraNodeParameterMapGet.cpp
+++ b/Engine/Plugins/FX/Niagara/Source/NiagaraEditor/Private/NiagaraNodeParameterMapGet.cpp
@@ -198,6 +198,16 @@ void UNiagaraNodeParameterMapGet::OnPinRemoved(UEdGraphPin* InRemovedPin)
 	Super::OnPinRemoved(InRemovedPin);
 }
 
+bool UNiagaraNodeParameterMapGet::CanModifyPin(const UEdGraphPin* Pin) const
+{
+	if(Super::CanModifyPin(Pin) == false || Pin->Direction == EGPD_Input)
+	{
+		return false;
+	}
+
+	return true;
+}
+
 void UNiagaraNodeParameterMapGet::OnNewTypedPinAdded(UEdGraphPin*& NewPin)
 {
 	if (NewPin->Direction == EEdGraphPinDirection::EGPD_Output)
diff --git a/Engine/Plugins/FX/Niagara/Source/NiagaraEditor/Private/NiagaraNodeParameterMapGet.h b/Engine/Plugins/FX/Niagara/Source/NiagaraEditor/Private/NiagaraNodeParameterMapGet.h
--- a/Engine/Plugins/FX/Niagara/Source/NiagaraEditor/Private/NiagaraNodeParameterMapGet.h
+++ b/Engine/Plugins/FX/Niagara/Source/NiagaraEditor/Private/NiagaraNodeParameterMapGet.h
@@ -52,6 +52,8 @@ public:
 	virtual EEdGraphPinDirection GetPinDirectionForNewParameters() override { return EEdGraphPinDirection::EGPD_Output; };
 
 protected:
+	virtual bool CanModifyPin(const UEdGraphPin* Pin) const override;
+	
 	virtual void OnNewTypedPinAdded(UEdGraphPin*& NewPin) override;
 	virtual void OnPinRenamed(UEdGraphPin* RenamedPin, const FString& OldName) override;
 	virtual void OnPinRemoved(UEdGraphPin* InRemovedPin) override;
diff --git a/Engine/Plugins/FX/Niagara/Source/NiagaraEditor/Private/NiagaraNodeStaticSwitch.cpp b/Engine/Plugins/FX/Niagara/Source/NiagaraEditor/Private/NiagaraNodeStaticSwitch.cpp
--- a/Engine/Plugins/FX/Niagara/Source/NiagaraEditor/Private/NiagaraNodeStaticSwitch.cpp
+++ b/Engine/Plugins/FX/Niagara/Source/NiagaraEditor/Private/NiagaraNodeStaticSwitch.cpp
@@ -124,6 +124,16 @@ TArray<int32> UNiagaraNodeStaticSwitch::GetOptionValues() const
 	return OptionValues;
 }
 
+bool UNiagaraNodeStaticSwitch::CanModifyPin(const UEdGraphPin* Pin) const
+{
+	if(IsAddPin(Pin))
+	{
+		return false;
+	}
+
+	return true;
+}
+
 void UNiagaraNodeStaticSwitch::PreChange(const UUserDefinedEnum* Changed, FEnumEditorUtils::EEnumEditorChangeInfo ChangedType)
 {
 	// do nothing
diff --git a/Engine/Plugins/FX/Niagara/Source/NiagaraEditor/Private/NiagaraNodeStaticSwitch.h b/Engine/Plugins/FX/Niagara/Source/NiagaraEditor/Private/NiagaraNodeStaticSwitch.h
--- a/Engine/Plugins/FX/Niagara/Source/NiagaraEditor/Private/NiagaraNodeStaticSwitch.h
+++ b/Engine/Plugins/FX/Niagara/Source/NiagaraEditor/Private/NiagaraNodeStaticSwitch.h
@@ -112,6 +112,7 @@ protected:
 	virtual TArray<int32> GetOptionValues() const override;
 	//~ End UNiagaraNodeUsageSelector Interface
 
+	virtual bool CanModifyPin(const UEdGraphPin* Pin) const override;
 private:
 	/** INotifyOnEnumChanged interface */
 	virtual void PreChange(const UUserDefinedEnum* Changed, FEnumEditorUtils::EEnumEditorChangeInfo ChangedType) override;
diff --git a/Engine/Plugins/FX/Niagara/Source/NiagaraEditor/Private/NiagaraNodeWithDynamicPins.cpp b/Engine/Plugins/FX/Niagara/Source/NiagaraEditor/Private/NiagaraNodeWithDynamicPins.cpp
--- a/Engine/Plugins/FX/Niagara/Source/NiagaraEditor/Private/NiagaraNodeWithDynamicPins.cpp
+++ b/Engine/Plugins/FX/Niagara/Source/NiagaraEditor/Private/NiagaraNodeWithDynamicPins.cpp
@@ -58,8 +58,6 @@ void UNiagaraNodeWithDynamicPins::PinConnectionListChanged(UEdGraphPin* Pin)
 		Pin->PinName = NewPinName;
 		Pin->PinFriendlyName = NewPinFriendlyName;
 
-		DynamicPins.Add(Pin->PinId);
-
 		CreateAddPin(Pin->Direction);
 		OnNewTypedPinAdded(Pin);
 		MarkNodeRequiresSynchronization(__FUNCTION__, true);
@@ -117,8 +115,6 @@ UEdGraphPin* UNiagaraNodeWithDynamicPins::RequestNewTypedPin(EEdGraphPinDirectio
 	AddPin->Modify();
 	AddPin->PinType = Schema->TypeDefinitionToPinType(Type);
 	AddPin->PinName = InName;
-	
-	DynamicPins.Add(AddPin->PinId);
 
 	CreateAddPin(Direction);
 	// we pass the pointer in as reference in case we want to reallocate so the overriding node has a chance to restore the pointer
@@ -153,27 +149,29 @@ bool UNiagaraNodeWithDynamicPins::IsExecPin(const UEdGraphPin* Pin) const
 	return Pin->PinType == Schema->TypeDefinitionToPinType(FNiagaraTypeDefinition::GetParameterMapDef()) && Pin->PinName == TEXT("");
 }
 
-bool UNiagaraNodeWithDynamicPins::CanRenamePin(const UEdGraphPin* Pin) const
+bool UNiagaraNodeWithDynamicPins::CanModifyPin(const UEdGraphPin* Pin) const
 {
-	if (CanModifyExistingPins() == false )
+	if(IsAddPin(Pin) || IsParameterMapPin(Pin))
 	{
-		return DynamicPins.Contains(Pin->PinId);
+		return false;
 	}
-	return IsAddPin(Pin) == false || IsExecPin(Pin);
+
+	return true;
+}
+
+bool UNiagaraNodeWithDynamicPins::CanRenamePin(const UEdGraphPin* Pin) const
+{	
+	return CanModifyPin(Pin) && !IsExecPin(Pin) && !Pin->bOrphanedPin;
 }
 
 bool UNiagaraNodeWithDynamicPins::CanRemovePin(const UEdGraphPin* Pin) const
 {
-	if (CanModifyExistingPins() == false)
-	{
-		return DynamicPins.Contains(Pin->PinId);
-	}
-	return IsAddPin(Pin) == false || IsExecPin(Pin);
+	return CanModifyPin(Pin) && !IsExecPin(Pin);
 }
 
 bool UNiagaraNodeWithDynamicPins::CanMovePin(const UEdGraphPin* Pin, int32 DirectionToMove) const
 {
-	if(IsAddPin(Pin) || IsExecPin(Pin) || IsParameterMapPin(Pin) || Pin->bOrphanedPin || (CanModifyExistingPins() == false && DynamicPins.Contains(Pin->PinId) == false))
+	if(!CanModifyPin(Pin) || IsExecPin(Pin) || Pin->bOrphanedPin)
 	{
 		return false;
 	}
@@ -192,7 +190,7 @@ bool UNiagaraNodeWithDynamicPins::CanMovePin(const UEdGraphPin* Pin, int32 Direc
 	if(PinArray.IsValidIndex(Index + DirectionToMove))
 	{
 		const UEdGraphPin* PinToMoveTo = PinArray[Index + DirectionToMove];
-		if(IsAddPin(PinToMoveTo) || IsParameterMapPin(PinToMoveTo) || PinToMoveTo->bOrphanedPin || (CanModifyExistingPins() == false && DynamicPins.Contains(PinToMoveTo->PinId) == false))
+		if(!CanModifyPin(PinToMoveTo) || PinToMoveTo->bOrphanedPin)
 		{
 			return false;
 		}
@@ -470,7 +468,6 @@ void UNiagaraNodeWithDynamicPins::AddExistingParameter(FNiagaraVariable Paramete
 
 void UNiagaraNodeWithDynamicPins::RemoveDynamicPin(UEdGraphPin* Pin)
 {
-	DynamicPins.Remove(Pin->PinId);
 	RemovePin(Pin);
 	MarkNodeRequiresSynchronization(__FUNCTION__, true);
 }
@@ -480,7 +477,6 @@ FText UNiagaraNodeWithDynamicPins::GetPinNameText(UEdGraphPin* Pin) const
 	return FText::FromName(Pin->PinName);
 }
 
-
 void UNiagaraNodeWithDynamicPins::PinNameTextCommitted(const FText& Text, ETextCommit::Type CommitType, UEdGraphPin* Pin)
 {
 	if (CommitType == ETextCommit::OnEnter)
@@ -511,23 +507,5 @@ void UNiagaraNodeWithDynamicPins::MoveDynamicPinFromMenu(UEdGraphPin* Pin, int32
 	MoveDynamicPin(Pin, DirectionToMove);
 }
 
-void UNiagaraNodeWithDynamicPins::RemoveAllDynamicPins()
-{
-	FScopedTransaction RemovePinTransaction(LOCTEXT("RemoveAllDynamicPinsTransaction", "Remove all dynamic pins"));
-	while(DynamicPins.Num() > 0)
-	{
-		FGuid& PinToRemoveGuid = DynamicPins.Last();
-		if(UEdGraphPin** Pin = Pins.FindByPredicate([&](UEdGraphPin* Pin){ return Pin->PinId == PinToRemoveGuid;}))
-		{
-			RemoveDynamicPin(*Pin);
-		}
-		else
-		{
-			checkf(0, TEXT("Dynamic pin could not be found."));
-			break;
-		}
-	}
-}
-
 #undef LOCTEXT_NAMESPACE
 
diff --git a/Engine/Plugins/FX/Niagara/Source/NiagaraEditor/Public/NiagaraNodeCustomHlsl.h b/Engine/Plugins/FX/Niagara/Source/NiagaraEditor/Public/NiagaraNodeCustomHlsl.h
--- a/Engine/Plugins/FX/Niagara/Source/NiagaraEditor/Public/NiagaraNodeCustomHlsl.h
+++ b/Engine/Plugins/FX/Niagara/Source/NiagaraEditor/Public/NiagaraNodeCustomHlsl.h
@@ -56,7 +56,6 @@ public:
 
 protected:
 	virtual bool AllowDynamicPins() const override { return true; }
-	virtual bool CanModifyExistingPins() const { return true; }
 	virtual bool GetValidateDataInterfaces() const override { return false; }
 	virtual bool VerifyEditablePinName(const FText& InName, FText& OutErrorMessage, const UEdGraphPin* InGraphPinObj) const override;
 	virtual bool IsPinNameEditableUponCreation(const UEdGraphPin* Pin) const override;
diff --git a/Engine/Plugins/FX/Niagara/Source/NiagaraEditor/Public/NiagaraNodeFunctionCall.h b/Engine/Plugins/FX/Niagara/Source/NiagaraEditor/Public/NiagaraNodeFunctionCall.h
--- a/Engine/Plugins/FX/Niagara/Source/NiagaraEditor/Public/NiagaraNodeFunctionCall.h
+++ b/Engine/Plugins/FX/Niagara/Source/NiagaraEditor/Public/NiagaraNodeFunctionCall.h
@@ -187,16 +187,17 @@ public:
 		orphaned. This allows a previously available static switch value to be retained on the node even if the the switch is no longer exposed. */
 	void AddOrphanedStaticSwitchPinForDataRetention(FNiagaraVariableBase StaticSwitchVariable, const FString& StaticSwitchPinDefault);
 
+	void RemoveAllDynamicPins();
 protected:
 	UEdGraphPin* AddStaticSwitchInputPin(FNiagaraVariable Input);
 
 	virtual bool GetValidateDataInterfaces() const { return true; };
 
 	virtual bool AllowDynamicPins() const override { return Signature.VariadicInput() || Signature.VariadicOutput(); }
+	virtual bool CanModifyPin(const UEdGraphPin* Pin) const override;
 	virtual bool CanRenamePin(const UEdGraphPin* Pin) const override { return AllowDynamicPins() && Super::CanRenamePin(Pin); }
 	virtual bool CanRemovePin(const UEdGraphPin* Pin) const override { return AllowDynamicPins() && Super::CanRemovePin(Pin); }
 	virtual bool CanMovePin(const UEdGraphPin* Pin, int32 DirectionToMove) const override { return AllowDynamicPins() && Super::CanMovePin(Pin, DirectionToMove); }
-	virtual bool CanModifyExistingPins() const override { return false; }
 
 	virtual void OnNewTypedPinAdded(UEdGraphPin*& NewPin) override { Super::OnNewTypedPinAdded(NewPin); RefreshSignature(); }
 	virtual void OnPinRenamed(UEdGraphPin* RenamedPin, const FString& OldPinName) override { Super::OnPinRenamed(RenamedPin, OldPinName); RefreshSignature(); }
@@ -205,6 +206,8 @@ protected:
 
 	void RefreshSignature();
 
+	bool IsBaseSignatureOfDataInterfaceFunction(const UEdGraphPin* Pin) const;
+	
 	/** Resets the node name based on the referenced script or signature. Guaranteed unique within a given graph instance.*/
 	void ComputeNodeName(FString SuggestedName = FString(), bool bForceSuggestion = false);
 
diff --git a/Engine/Plugins/FX/Niagara/Source/NiagaraEditor/Public/NiagaraNodeWithDynamicPins.h b/Engine/Plugins/FX/Niagara/Source/NiagaraEditor/Public/NiagaraNodeWithDynamicPins.h
--- a/Engine/Plugins/FX/Niagara/Source/NiagaraEditor/Public/NiagaraNodeWithDynamicPins.h
+++ b/Engine/Plugins/FX/Niagara/Source/NiagaraEditor/Public/NiagaraNodeWithDynamicPins.h
@@ -49,9 +49,6 @@ public:
 	/** Convenience method to determine whether this Node is a Map Get or Map Set when adding a parameter through the parameter panel. */
 	virtual EEdGraphPinDirection GetPinDirectionForNewParameters() { return EEdGraphPinDirection::EGPD_MAX; };
 
-	/** Will remove any dynamic pins. */
-	void RemoveAllDynamicPins();
-
 protected:
 	virtual bool AllowDynamicPins() const { return true; }
 
@@ -70,15 +67,15 @@ protected:
 
 	virtual bool CanRenamePinFromContextMenu(const UEdGraphPin* Pin) const { return CanRenamePin(Pin); }
 
+	/** Called to determine if this node can modify a given pin at all. */
+	virtual bool CanModifyPin(const UEdGraphPin* Pin) const;
+	
 	/** Called to determine if a pin can be renamed by the user. */
 	virtual bool CanRenamePin(const UEdGraphPin* Pin) const;
 
 	/** Called to determine if a pin can be removed by the user. */
 	virtual bool CanRemovePin(const UEdGraphPin* Pin) const;
-
-	/** Called to determine if this node can modify pins that existed on node creation and were not added dynamically. If true, the only pin that cannot be modified is the Add pin. */
-	virtual bool CanModifyExistingPins() const { return true; }
-
+	
 	/** Called to determine if a pin can be moved by the user. Negative values for up, positive for down. */
 	virtual bool CanMovePin(const UEdGraphPin* Pin, int32 DirectionToMove) const;
 
@@ -109,8 +106,4 @@ private:
 public:
 	/** The sub category for add pins. */
 	static const FName AddPinSubCategory;
-
-	/** Track dynamic pins so that we only allow removing of dynamic pins. */
-	UPROPERTY()
-	TArray<FGuid> DynamicPins;
 };
