- Add missing scoped transaction for array moves

diff --git a/Engine/Plugins/FX/Niagara/Source/NiagaraEditor/Private/ViewModels/Stack/NiagaraStackPropertyRow.cpp b/Engine/Plugins/FX/Niagara/Source/NiagaraEditor/Private/ViewModels/Stack/NiagaraStackPropertyRow.cpp
--- a/Engine/Plugins/FX/Niagara/Source/NiagaraEditor/Private/ViewModels/Stack/NiagaraStackPropertyRow.cpp
+++ b/Engine/Plugins/FX/Niagara/Source/NiagaraEditor/Private/ViewModels/Stack/NiagaraStackPropertyRow.cpp
@@ -6,6 +6,7 @@
 #include "IDetailTreeNode.h"
 #include "PropertyHandle.h"
 #include "IDetailPropertyRow.h"
+#include "ScopedTransaction.h"
 
 #include UE_INLINE_GENERATED_CPP_BY_NAME(NiagaraStackPropertyRow)
 
@@ -208,17 +209,21 @@ TOptional<UNiagaraStackEntry::FDropRequestResponse> UNiagaraStackPropertyRow::Dr
 		UNiagaraStackPropertyRow* DraggedPropertyRow = CastChecked<UNiagaraStackPropertyRow>(StackEntryDragDropOp->GetDraggedEntries()[0]);
 		TSharedPtr<IPropertyHandle> DraggedPropertyHandle = DraggedPropertyRow->GetDetailTreeNode()->CreatePropertyHandle();
 		TSharedPtr<IPropertyHandle> TargetPropertyHandle = DetailTreeNode->CreatePropertyHandle();
+		TSharedPtr<IPropertyHandle> ParentHandle = TargetPropertyHandle->GetParentHandle();
 		int32 IndexOffset = DropRequest.DropZone == EItemDropZone::AboveItem ? 0 : 1;
 		uint32 NumElements = 0;
-		TargetPropertyHandle->GetParentHandle()->AsArray()->GetNumElements(NumElements);
+		ParentHandle->AsArray()->GetNumElements(NumElements);
 
 		// we clamp the offset
 		if((uint32) (TargetPropertyHandle->GetIndexInArray() + IndexOffset) >= (NumElements - 1))
 		{
 			IndexOffset = 0;
 		}
-		
-		TargetPropertyHandle->GetParentHandle()->AsArray()->MoveElementTo(DraggedPropertyHandle->GetIndexInArray(), TargetPropertyHandle->GetIndexInArray() + IndexOffset);
+
+		FScopedTransaction Transaction(NSLOCTEXT("NiagaraStackPropertyRow", "DropArrayItem", "Move Array Item"));
+		ParentHandle->NotifyPreChange();
+		ParentHandle->AsArray()->MoveElementTo(DraggedPropertyHandle->GetIndexInArray(), TargetPropertyHandle->GetIndexInArray() + IndexOffset);
+		ParentHandle->NotifyPostChange(EPropertyChangeType::ArrayMove);
 	}
 	return CanDropResponse;
 }
