Removed DynamicPins array and refactored use cases that required it. This fixes various UI issues in regards to pin context actions

diff --git a/Engine/Plugins/FX/Niagara/Source/NiagaraEditor/Private/NiagaraNodeConvert.h b/Engine/Plugins/FX/Niagara/Source/NiagaraEditor/Private/NiagaraNodeConvert.h
--- a/Engine/Plugins/FX/Niagara/Source/NiagaraEditor/Private/NiagaraNodeConvert.h
+++ b/Engine/Plugins/FX/Niagara/Source/NiagaraEditor/Private/NiagaraNodeConvert.h
@@ -106,7 +106,6 @@ public:
 	virtual TSharedPtr<SGraphNode> CreateVisualWidget() override;
 	virtual void  AutowireNewNode(UEdGraphPin* FromPin)override;
 	virtual FText GetNodeTitle(ENodeTitleType::Type TitleType) const override;
-	//virtual bool CanMovePin(const UEdGraphPin* Pin, int32 DirectionToMove) const override { return false; }
 
 	//~ UNiagaraNode interface
 	virtual void Compile(class FHlslNiagaraTranslator* Translator, TArray<int32>& Outputs) override;
