- Add enable condition for emitter bounds so that if system fixed bounds are enabled the emitter ones are greyed out.

diff --git a/Engine/Plugins/FX/Niagara/Source/NiagaraEditor/Private/NiagaraEmitterDetailsCustomization.cpp b/Engine/Plugins/FX/Niagara/Source/NiagaraEditor/Private/NiagaraEmitterDetailsCustomization.cpp
--- a/Engine/Plugins/FX/Niagara/Source/NiagaraEditor/Private/NiagaraEmitterDetailsCustomization.cpp
+++ b/Engine/Plugins/FX/Niagara/Source/NiagaraEditor/Private/NiagaraEmitterDetailsCustomization.cpp
@@ -6,6 +6,7 @@
 #include "DetailCategoryBuilder.h"
 #include "NiagaraEmitter.h"
 #include "NiagaraEditorModule.h"
+#include "NiagaraSystem.h"
 #include "Toolkits/NiagaraSystemToolkit.h"
 #include "ViewModels/NiagaraEmitterViewModel.h"
 #include "ViewModels/TNiagaraViewModelManager.h"
@@ -21,6 +22,7 @@ void CustomizeEmitterData(IDetailLayoutBuilder& InDetailLayout, UNiagaraEmitter*
 	TNiagaraViewModelManager<UNiagaraEmitter, FNiagaraEmitterViewModel>::GetAllViewModelsForObject(EmitterBeingCustomized, ExistingViewModels);
 	FVersionedNiagaraEmitter VersionedNiagaraEmitter = ExistingViewModels[0]->GetEmitter();
 	FVersionedNiagaraEmitterData* EmitterData = VersionedNiagaraEmitter.GetEmitterData();
+	UNiagaraSystem* NiagaraSystem = EmitterBeingCustomized->GetTypedOuter<UNiagaraSystem>();
 	TSharedPtr<FStructOnScope> StructData = MakeShareable(new FStructOnScope(FVersionedNiagaraEmitterData::StaticStruct(), reinterpret_cast<uint8*>(EmitterData)));
 	for (const auto& Entry : CategoryPropertyMap)
 	{
@@ -33,23 +35,27 @@ void CustomizeEmitterData(IDetailLayoutBuilder& InDetailLayout, UNiagaraEmitter*
 		for (IDetailPropertyRow* PropertyRow : OutRows)
 		{
 			TSharedPtr<IPropertyHandle> PropertyHandle = PropertyRow->GetPropertyHandle();
-			if (!Entry.Value.Contains(PropertyHandle->GetProperty()))
+			FProperty* Property = PropertyHandle->GetProperty();
+			if (!Entry.Value.Contains(Property))
 			{
 				PropertyHandle->MarkHiddenByCustomization();
 				PropertyRow->Visibility(EVisibility::Collapsed);
+				continue;
 			}
-			else
-			{
-				FProperty* ChildProperty = PropertyHandle->GetProperty();
-				const auto& PostEditChangeLambda = [VersionedNiagaraEmitter, ChildProperty]
-				{
-					FPropertyChangedEvent ChangeEvent(ChildProperty);
-					VersionedNiagaraEmitter.Emitter->PostEditChangeVersionedProperty(ChangeEvent, VersionedNiagaraEmitter.Version);
-				};
 
-				PropertyRow->GetPropertyHandle()->SetOnChildPropertyValueChanged(FSimpleDelegate::CreateLambda(PostEditChangeLambda));
-				PropertyRow->GetPropertyHandle()->SetOnPropertyValueChanged(FSimpleDelegate::CreateLambda(PostEditChangeLambda));
+			if (NiagaraSystem && (Property->GetFName() == GET_MEMBER_NAME_CHECKED(FVersionedNiagaraEmitterData, CalculateBoundsMode) || Property->GetFName() == GET_MEMBER_NAME_CHECKED(FVersionedNiagaraEmitterData, FixedBounds)) )
+			{
+				PropertyRow->IsEnabled(TAttribute<bool>::Create([VersionedNiagaraEmitter, WeakNiagaraSystem=TWeakObjectPtr<UNiagaraSystem>(NiagaraSystem)]() { return !WeakNiagaraSystem.IsValid() || WeakNiagaraSystem.Get()->bFixedBounds == false; }));
 			}
+
+			const auto& PostEditChangeLambda = [VersionedNiagaraEmitter, Property]
+			{
+				FPropertyChangedEvent ChangeEvent(Property);
+				VersionedNiagaraEmitter.Emitter->PostEditChangeVersionedProperty(ChangeEvent, VersionedNiagaraEmitter.Version);
+			};
+
+			PropertyRow->GetPropertyHandle()->SetOnChildPropertyValueChanged(FSimpleDelegate::CreateLambda(PostEditChangeLambda));
+			PropertyRow->GetPropertyHandle()->SetOnPropertyValueChanged(FSimpleDelegate::CreateLambda(PostEditChangeLambda));
 		}
 	}
 }
