- Don't include invalid components in Outliner list
- Fix Outliner crash when capturing SimCache from invalid components

#rnx

diff --git a/Engine/Plugins/FX/Niagara/Source/Niagara/Private/NiagaraDebuggerClient.cpp b/Engine/Plugins/FX/Niagara/Source/Niagara/Private/NiagaraDebuggerClient.cpp
--- a/Engine/Plugins/FX/Niagara/Source/Niagara/Private/NiagaraDebuggerClient.cpp
+++ b/Engine/Plugins/FX/Niagara/Source/Niagara/Private/NiagaraDebuggerClient.cpp
@@ -73,7 +73,7 @@ bool FNiagaraDebuggerClient::Tick(float DeltaSeconds)
 		FNiagaraSimCacheCaptureInfo& SimCacheCapture = *It;
 		if (SimCacheCapture.Process())
 		{
-			if (MessageEndpoint.IsValid() && Connection.IsValid())
+			if (MessageEndpoint.IsValid() && Connection.IsValid() && SimCacheCapture.SimCache.IsValid())
 			{
 				//Sim cache is complete. Send it back to the connected debugger.
 				FNiagaraSystemSimCacheCaptureReply* SimCacheCaptureReply = FMessageEndpoint::MakeMessage<FNiagaraSystemSimCacheCaptureReply>();
@@ -353,129 +353,131 @@ bool FNiagaraDebuggerClient::UpdateOutliner(float DeltaSeconds)
 		for (TObjectIterator<UNiagaraComponent> CompIt; CompIt; ++CompIt)
 		{
 			UNiagaraComponent* Comp = *CompIt;
-			if (Comp)
+			if (!::IsValid(Comp))
 			{
-				UWorld* World = Comp->GetWorld();
-				FNiagaraOutlinerWorldData& WorldData = Message->OutlinerData.WorldData.FindOrAdd( World ? World->GetPathName() : TEXT("Null World") );
-				if (World)
-				{
-					WorldData.bHasBegunPlay = World->HasBegunPlay();
-					WorldData.WorldType = uint8(World->WorldType);
-					WorldData.NetMode = uint8(World->GetNetMode());
+				continue;
+			}
 
-					#if WITH_PARTICLE_PERF_STATS
-					if(StatsListener)
+			UWorld* World = Comp->GetWorld();
+			FNiagaraOutlinerWorldData& WorldData = Message->OutlinerData.WorldData.FindOrAdd( World ? World->GetPathName() : TEXT("Null World") );
+			if (World)
+			{
+				WorldData.bHasBegunPlay = World->HasBegunPlay();
+				WorldData.WorldType = uint8(World->WorldType);
+				WorldData.NetMode = uint8(World->GetNetMode());
+
+				#if WITH_PARTICLE_PERF_STATS
+				if(StatsListener)
+				{
+					if (FAccumulatedParticlePerfStats* WorldStats = StatsListener->GetStats(World))
 					{
-						if (FAccumulatedParticlePerfStats* WorldStats = StatsListener->GetStats(World))
-						{
-							WorldData.AveragePerFrameTime.GameThread = WorldStats->GetGameThreadStats().GetPerFrameAvg();
-							WorldData.AveragePerFrameTime.RenderThread = WorldStats->GetRenderThreadStats().GetPerFrameAvg();
+						WorldData.AveragePerFrameTime.GameThread = WorldStats->GetGameThreadStats().GetPerFrameAvg();
+						WorldData.AveragePerFrameTime.RenderThread = WorldStats->GetRenderThreadStats().GetPerFrameAvg();
 
-							WorldData.MaxPerFrameTime.GameThread = WorldStats->GetGameThreadStats().GetPerFrameMax();
-							WorldData.MaxPerFrameTime.RenderThread = WorldStats->GetRenderThreadStats().GetPerFrameMax();
-						}
+						WorldData.MaxPerFrameTime.GameThread = WorldStats->GetGameThreadStats().GetPerFrameMax();
+						WorldData.MaxPerFrameTime.RenderThread = WorldStats->GetRenderThreadStats().GetPerFrameMax();
 					}
-					#endif
 				}
+				#endif
+			}
 
-				UNiagaraSystem* System = Comp->GetAsset();
-				FNiagaraOutlinerSystemData& Instances = WorldData.Systems.FindOrAdd(System ? System->GetPathName() : TEXT("Null System"));
-				if (System)
+			UNiagaraSystem* System = Comp->GetAsset();
+			FNiagaraOutlinerSystemData& Instances = WorldData.Systems.FindOrAdd(System ? System->GetPathName() : TEXT("Null System"));
+			if (System)
+			{
+				//Add System specific data.
+				#if WITH_PARTICLE_PERF_STATS
+				if(StatsListener)
 				{
-					//Add System specific data.
-					#if WITH_PARTICLE_PERF_STATS
-					if(StatsListener)
+					if (FAccumulatedParticlePerfStats* SystemStats = StatsListener->GetStats(System))
 					{
-						if (FAccumulatedParticlePerfStats* SystemStats = StatsListener->GetStats(System))
-						{
-							Instances.AveragePerFrameTime.GameThread = SystemStats->GetGameThreadStats().GetPerFrameAvg();
-							Instances.AveragePerFrameTime.RenderThread = SystemStats->GetRenderThreadStats().GetPerFrameAvg();
+						Instances.AveragePerFrameTime.GameThread = SystemStats->GetGameThreadStats().GetPerFrameAvg();
+						Instances.AveragePerFrameTime.RenderThread = SystemStats->GetRenderThreadStats().GetPerFrameAvg();
 
-							Instances.MaxPerFrameTime.GameThread = SystemStats->GetGameThreadStats().GetPerFrameMax();
-							Instances.MaxPerFrameTime.RenderThread = SystemStats->GetRenderThreadStats().GetPerFrameMax();
+						Instances.MaxPerFrameTime.GameThread = SystemStats->GetGameThreadStats().GetPerFrameMax();
+						Instances.MaxPerFrameTime.RenderThread = SystemStats->GetRenderThreadStats().GetPerFrameMax();
 
 
-							Instances.AveragePerInstanceTime.GameThread = SystemStats->GetGameThreadStats().GetPerInstanceAvg();
-							Instances.AveragePerInstanceTime.RenderThread = SystemStats->GetRenderThreadStats().GetPerInstanceAvg();
+						Instances.AveragePerInstanceTime.GameThread = SystemStats->GetGameThreadStats().GetPerInstanceAvg();
+						Instances.AveragePerInstanceTime.RenderThread = SystemStats->GetRenderThreadStats().GetPerInstanceAvg();
 
-							Instances.MaxPerInstanceTime.GameThread = SystemStats->GetGameThreadStats().GetPerInstanceMax();
-							Instances.MaxPerInstanceTime.RenderThread = SystemStats->GetRenderThreadStats().GetPerInstanceMax();
-						}
+						Instances.MaxPerInstanceTime.GameThread = SystemStats->GetGameThreadStats().GetPerInstanceMax();
+						Instances.MaxPerInstanceTime.RenderThread = SystemStats->GetRenderThreadStats().GetPerInstanceMax();
 					}
-					#endif
 				}
+				#endif
+			}
 
-				FNiagaraOutlinerSystemInstanceData& InstData = Instances.SystemInstances.AddDefaulted_GetRef();
-				InstData.ComponentName = Comp->GetPathName();
+			FNiagaraOutlinerSystemInstanceData& InstData = Instances.SystemInstances.AddDefaulted_GetRef();
+			InstData.ComponentName = Comp->GetPathName();
 
-				if (FNiagaraSystemInstanceControllerPtr InstController = Comp->GetSystemInstanceController())
-				{
-					InstData.ActualExecutionState = InstController->GetActualExecutionState();
-					InstData.RequestedExecutionState = InstController->GetRequestedExecutionState();
+			if (FNiagaraSystemInstanceControllerPtr InstController = Comp->GetSystemInstanceController())
+			{
+				InstData.ActualExecutionState = InstController->GetActualExecutionState();
+				InstData.RequestedExecutionState = InstController->GetRequestedExecutionState();
 
-					InstData.ScalabilityState = Comp->DebugCachedScalabilityState;
-					InstData.bPendingKill = !IsValidChecked(Comp) || Comp->IsUnreachable();
-					InstData.bUsingCullProxy = Comp->IsUsingCullProxy();
+				InstData.ScalabilityState = Comp->DebugCachedScalabilityState;
+				InstData.bPendingKill = !IsValidChecked(Comp) || Comp->IsUnreachable();
+				InstData.bUsingCullProxy = Comp->IsUsingCullProxy();
 
-					InstData.PoolMethod = Comp->PoolingMethod;
+				InstData.PoolMethod = Comp->PoolingMethod;
 
-					// TODO: need to be able to access at least a shadow copy of per-emitter execution state and num particles. For now, we'll just allow unsafe access
-					if (FNiagaraSystemInstance* Inst = InstController->GetSystemInstance_Unsafe())
+				// TODO: need to be able to access at least a shadow copy of per-emitter execution state and num particles. For now, we'll just allow unsafe access
+				if (FNiagaraSystemInstance* Inst = InstController->GetSystemInstance_Unsafe())
+				{
+					InstData.TickGroup = Inst->CalculateTickGroup();
+					if (const FNiagaraSystemGpuComputeProxy* SystemInstanceComputeProxy = Inst->GetSystemGpuComputeProxy())
 					{
-						InstData.TickGroup = Inst->CalculateTickGroup();
-						if (const FNiagaraSystemGpuComputeProxy* SystemInstanceComputeProxy = Inst->GetSystemGpuComputeProxy())
-						{
-							InstData.GpuTickStage = SystemInstanceComputeProxy->GetComputeTickStage();
-						}
-						InstData.LWCTile = Inst->GetLWCTile();
-						InstData.bIsSolo = Inst->IsSolo();
-						InstData.bRequiresDistanceFieldData = Inst->RequiresDistanceFieldData();
-						InstData.bRequiresDepthBuffer = Inst->RequiresDepthBuffer();
-						InstData.bRequiresEarlyViewData = Inst->RequiresEarlyViewData();
-						InstData.bRequiresViewUniformBuffer = Inst->RequiresViewUniformBuffer();
-						InstData.bRequiresRayTracingScene = Inst->RequiresRayTracingScene();
-
-						InstData.Emitters.Reserve(Inst->GetEmitters().Num());
-						for (TSharedRef<FNiagaraEmitterInstance, ESPMode::ThreadSafe>& EmitterInst : Inst->GetEmitters())
+						InstData.GpuTickStage = SystemInstanceComputeProxy->GetComputeTickStage();
+					}
+					InstData.LWCTile = Inst->GetLWCTile();
+					InstData.bIsSolo = Inst->IsSolo();
+					InstData.bRequiresDistanceFieldData = Inst->RequiresDistanceFieldData();
+					InstData.bRequiresDepthBuffer = Inst->RequiresDepthBuffer();
+					InstData.bRequiresEarlyViewData = Inst->RequiresEarlyViewData();
+					InstData.bRequiresViewUniformBuffer = Inst->RequiresViewUniformBuffer();
+					InstData.bRequiresRayTracingScene = Inst->RequiresRayTracingScene();
+
+					InstData.Emitters.Reserve(Inst->GetEmitters().Num());
+					for (TSharedRef<FNiagaraEmitterInstance, ESPMode::ThreadSafe>& EmitterInst : Inst->GetEmitters())
+					{
+						FNiagaraOutlinerEmitterInstanceData& EmitterData = InstData.Emitters.AddDefaulted_GetRef();
+						FVersionedNiagaraEmitter VersionedEmitter = EmitterInst->GetCachedEmitter();
+						if (VersionedEmitter.Emitter)
 						{
-							FNiagaraOutlinerEmitterInstanceData& EmitterData = InstData.Emitters.AddDefaulted_GetRef();
-							FVersionedNiagaraEmitter VersionedEmitter = EmitterInst->GetCachedEmitter();
-							if (VersionedEmitter.Emitter)
-							{
-								//TODO: This is a bit wasteful to copy the name into each instance data. Though we can't rely on the debugger side data matchin the actul running data on the device.
-								//We need to build a shared representation of the asset data from the client that we then reference from this per instance data.
-								EmitterData.EmitterName = VersionedEmitter.Emitter->GetUniqueEmitterName();
-								EmitterData.SimTarget = VersionedEmitter.GetEmitterData()->SimTarget;
-								//Move all above to a shared asset representation.
-
-								EmitterData.ExecState = EmitterInst->GetExecutionState();
-								EmitterData.NumParticles = EmitterInst->GetNumParticles();
-
-								EmitterData.bRequiresPersistentIDs = VersionedEmitter.GetEmitterData()->RequiresPersistentIDs();
-							}
+							//TODO: This is a bit wasteful to copy the name into each instance data. Though we can't rely on the debugger side data matchin the actul running data on the device.
+							//We need to build a shared representation of the asset data from the client that we then reference from this per instance data.
+							EmitterData.EmitterName = VersionedEmitter.Emitter->GetUniqueEmitterName();
+							EmitterData.SimTarget = VersionedEmitter.GetEmitterData()->SimTarget;
+							//Move all above to a shared asset representation.
+
+							EmitterData.ExecState = EmitterInst->GetExecutionState();
+							EmitterData.NumParticles = EmitterInst->GetNumParticles();
+
+							EmitterData.bRequiresPersistentIDs = VersionedEmitter.GetEmitterData()->RequiresPersistentIDs();
 						}
 					}
 				}
-				else
-				{
-					InstData.ActualExecutionState = ENiagaraExecutionState::Num;
-					InstData.RequestedExecutionState = ENiagaraExecutionState::Num;
-				}
+			}
+			else
+			{
+				InstData.ActualExecutionState = ENiagaraExecutionState::Num;
+				InstData.RequestedExecutionState = ENiagaraExecutionState::Num;
+			}
 
 #if WITH_PARTICLE_PERF_STATS
-				if(StatsListener)
+			if(StatsListener)
+			{
+				if (FAccumulatedParticlePerfStats* ComponentStats = StatsListener->GetStats(Comp))
 				{
-					if (FAccumulatedParticlePerfStats* ComponentStats = StatsListener->GetStats(Comp))
-					{
-						InstData.AverageTime.GameThread = ComponentStats->GetGameThreadStats().GetPerFrameAvg();
-						InstData.AverageTime.RenderThread = ComponentStats->GetRenderThreadStats().GetPerFrameAvg();
+					InstData.AverageTime.GameThread = ComponentStats->GetGameThreadStats().GetPerFrameAvg();
+					InstData.AverageTime.RenderThread = ComponentStats->GetRenderThreadStats().GetPerFrameAvg();
 
-						InstData.MaxTime.GameThread = ComponentStats->GetGameThreadStats().GetPerFrameMax();
-						InstData.MaxTime.RenderThread = ComponentStats->GetRenderThreadStats().GetPerFrameMax();
-					}
+					InstData.MaxTime.GameThread = ComponentStats->GetGameThreadStats().GetPerFrameMax();
+					InstData.MaxTime.RenderThread = ComponentStats->GetRenderThreadStats().GetPerFrameMax();
 				}
-#endif
 			}
+#endif
 		}
 
 		//TODO: Add any component less systems too if and when they are a thing.
@@ -507,7 +509,10 @@ bool FNiagaraSimCacheCaptureInfo::Process()
 	if (Comp == nullptr || ProcessedFrames >= Request.CaptureDelayFrames + Request.CaptureFrames)
 	{
 		//Capture is complete		
-		SimCache->EndWrite();
+		if (SimCache.IsValid())
+		{
+			SimCache->EndWrite();
+		}
 		return true;
 	}
 
