Fixed Custom HLSL pins not being modifiable.
Also added visibility delegates for action "Reset to Default" and "Promote to Parameter".
Also changed the Pin Section name used in Niagara to use the parent section name so we no longer have 2 sections called "Pin Actions" at times

diff --git a/Engine/Plugins/FX/Niagara/Source/NiagaraEditor/Private/EdGraphSchema_Niagara.cpp b/Engine/Plugins/FX/Niagara/Source/NiagaraEditor/Private/EdGraphSchema_Niagara.cpp
--- a/Engine/Plugins/FX/Niagara/Source/NiagaraEditor/Private/EdGraphSchema_Niagara.cpp
+++ b/Engine/Plugins/FX/Niagara/Source/NiagaraEditor/Private/EdGraphSchema_Niagara.cpp
@@ -2336,39 +2336,49 @@ void UEdGraphSchema_Niagara::GetContextMenuActions(UToolMenu* Menu, UGraphNodeCo
 	const UEdGraphPin* InGraphPin = Context->Pin;
 	if (InGraphPin)
 	{
+		const FName SectionName = "EdGraphSchemaPinActions";
+		FToolMenuSection& Section = Menu->FindOrAddSection(SectionName);
+		FNiagaraTypeDefinition InGraphPinType = PinToTypeDefinition(InGraphPin);
+		if (InGraphPinType == FNiagaraTypeDefinition::GetGenericNumericDef() && InGraphPin->LinkedTo.Num() == 0)
 		{
-			const FName SectionName = "EdGraphSchema_NiagaraPinActions";
-			FToolMenuSection& Section = Menu->AddSection(SectionName, LOCTEXT("PinActionsMenuHeader", "Pin Actions"));
-			FNiagaraTypeDefinition InGraphPinType = PinToTypeDefinition(InGraphPin);
-			if (InGraphPinType == FNiagaraTypeDefinition::GetGenericNumericDef() && InGraphPin->LinkedTo.Num() == 0)
-			{
-				Section.AddSubMenu(
-					"ConvertNumericSpecific",
-					LOCTEXT("ConvertNumericSpecific", "Convert Numeric To..."),
-					LOCTEXT("ConvertNumericSpecificToolTip", "Convert Numeric pin to the specific typed pin."),
-				FNewToolMenuDelegate::CreateUObject((UEdGraphSchema_Niagara*const)this, &UEdGraphSchema_Niagara::GetNumericConversionToSubMenuActions, SectionName, const_cast<UEdGraphPin*>(InGraphPin)));
-			}
-			else if ( InGraphPinType.IsDataInterface() )
-			{
-				GenerateDataInterfacePinMenu(Menu, SectionName, InGraphPin, InGraphPinType);
-			}
+			Section.AddSubMenu(
+				"ConvertNumericSpecific",
+				LOCTEXT("ConvertNumericSpecific", "Convert Numeric To..."),
+				LOCTEXT("ConvertNumericSpecificToolTip", "Convert Numeric pin to the specific typed pin."),
+			FNewToolMenuDelegate::CreateUObject((UEdGraphSchema_Niagara*const)this, &UEdGraphSchema_Niagara::GetNumericConversionToSubMenuActions, SectionName, const_cast<UEdGraphPin*>(InGraphPin)));
+		}
+		else if ( InGraphPinType.IsDataInterface() )
+		{
+			GenerateDataInterfacePinMenu(Menu, SectionName, InGraphPin, InGraphPinType);
+		}
 
-			if (InGraphPin->Direction == EEdGraphPinDirection::EGPD_Input)
+		if (InGraphPin->Direction == EEdGraphPinDirection::EGPD_Input)
+		{
+			FUIAction PromoteToParameterAction(FExecuteAction::CreateUObject((UEdGraphSchema_Niagara*const)this, &UEdGraphSchema_Niagara::PromoteSinglePinToParameter, const_cast<UEdGraphPin*>(InGraphPin)),					
+				FCanExecuteAction::CreateStatic(&UEdGraphSchema_Niagara::CanPromoteSinglePinToParameter, InGraphPin));
+			PromoteToParameterAction.IsActionVisibleDelegate = FIsActionButtonVisible::CreateStatic(&UEdGraphSchema_Niagara::CanPromoteSinglePinToParameter, InGraphPin);
+			
+			Section.AddMenuEntry(
+			"PromoteToParameter",
+			LOCTEXT("PromoteToParameter", "Promote to Parameter"),
+			LOCTEXT("PromoteToParameterTooltip", "Create a parameter argument and connect this pin to that parameter."),
+				FSlateIcon(),
+				PromoteToParameterAction
+			);
+			
+			if (InGraphPin->LinkedTo.Num() == 0 && InGraphPin->bDefaultValueIsIgnored == false)
 			{
-				Section.AddMenuEntry("PromoteToParameter", LOCTEXT("PromoteToParameter", "Promote to Parameter"), LOCTEXT("PromoteToParameterTooltip", "Create a parameter argument and connect this pin to that parameter."), FSlateIcon(),
-					FUIAction(FExecuteAction::CreateUObject((UEdGraphSchema_Niagara*const)this, &UEdGraphSchema_Niagara::PromoteSinglePinToParameter, const_cast<UEdGraphPin*>(InGraphPin)),
-						FCanExecuteAction::CreateStatic(&UEdGraphSchema_Niagara::CanPromoteSinglePinToParameter, InGraphPin)));
-				if (InGraphPin->LinkedTo.Num() == 0 && InGraphPin->bDefaultValueIsIgnored == false)
-				{
-					Section.AddMenuEntry(
-						"ResetInputToDefault",
-						LOCTEXT("ResetInputToDefault", "Reset to Default"),  // TODO(mv): This is currently broken
-						LOCTEXT("ResetInputToDefaultToolTip", "Reset this input to its default value."),
-						FSlateIcon(),
-						FUIAction(
-							FExecuteAction::CreateUObject((UEdGraphSchema_Niagara*const)this, &UEdGraphSchema_Niagara::ResetPinToAutogeneratedDefaultValue, const_cast<UEdGraphPin*>(InGraphPin), true),
-							FCanExecuteAction::CreateStatic(&CanResetPinToDefault, this, InGraphPin)));
-				}
+				FUIAction ResetToDefaultAction(FExecuteAction::CreateUObject((UEdGraphSchema_Niagara*const)this, &UEdGraphSchema_Niagara::ResetPinToAutogeneratedDefaultValue, const_cast<UEdGraphPin*>(InGraphPin), true),
+					FCanExecuteAction::CreateStatic(&CanResetPinToDefault, this, InGraphPin));
+				ResetToDefaultAction.IsActionVisibleDelegate = FIsActionButtonVisible::CreateStatic(&CanResetPinToDefault, this, InGraphPin);
+				
+				Section.AddMenuEntry(
+					"ResetInputToDefault",
+					LOCTEXT("ResetInputToDefault", "Reset to Default"),
+					LOCTEXT("ResetInputToDefaultToolTip", "Reset this input to its default value."),
+					FSlateIcon(),
+					ResetToDefaultAction
+				);
 			}
 		}
 	}
@@ -2407,7 +2417,7 @@ void UEdGraphSchema_Niagara::GetContextMenuActions(UToolMenu* Menu, UGraphNodeCo
 				FNewToolMenuDelegate::CreateUObject((UEdGraphSchema_Niagara*const)this, &UEdGraphSchema_Niagara::GetNumericConversionToSubMenuActionsAll, SectionName, const_cast<UNiagaraNode*>(Node)));
 		}
 		
-		FToolMenuSection& Section = Menu->AddSection("EdGraphSchema_NiagaraNodeActions", LOCTEXT("NodeActionsMenuHeader", "Node Actions"));
+		FToolMenuSection& Section = Menu->AddSection("EdGraphSchema_NiagaraNodeActions", LOCTEXT("NodeActionsMenuHeader", "Niagara Node Actions"));
 		Section.AddMenuEntry("ToggleEnabledState", LOCTEXT("ToggleEnabledState", "Toggle Enabled State"), LOCTEXT("ToggleEnabledStateTooltip", "Toggle this node between Enbled (default) and Disabled (skipped from compilation)."), FSlateIcon(),
 			FUIAction(FExecuteAction::CreateUObject((UEdGraphSchema_Niagara*const)this, &UEdGraphSchema_Niagara::ToggleNodeEnabledState, const_cast<UNiagaraNode*>(Node))));
 		Section.AddMenuEntry("RefreshNode", LOCTEXT("RefreshNode", "Refresh Node"), LOCTEXT("RefreshNodeTooltip", "Refresh this node."), FSlateIcon(),
diff --git a/Engine/Plugins/FX/Niagara/Source/NiagaraEditor/Public/NiagaraNodeCustomHlsl.h b/Engine/Plugins/FX/Niagara/Source/NiagaraEditor/Public/NiagaraNodeCustomHlsl.h
--- a/Engine/Plugins/FX/Niagara/Source/NiagaraEditor/Public/NiagaraNodeCustomHlsl.h
+++ b/Engine/Plugins/FX/Niagara/Source/NiagaraEditor/Public/NiagaraNodeCustomHlsl.h
@@ -62,14 +62,11 @@ protected:
 	virtual bool IsPinNameEditable(const UEdGraphPin* Pin) const override;
 	virtual bool CommitEditablePinName(const FText& InName, UEdGraphPin* InGraphPinObj, bool bSuppressEvents = false) override;
 	virtual bool CancelEditablePinName(const FText& InName, UEdGraphPin* InGraphPinObj) override;
-	
+
+	virtual bool CanModifyPin(const UEdGraphPin* Pin) const override { return UNiagaraNodeWithDynamicPins::CanModifyPin(Pin); }
 	virtual bool CanRenamePin(const UEdGraphPin* Pin) const override { return UNiagaraNodeWithDynamicPins::CanRenamePin(Pin); }
-	virtual bool CanRemovePin(const UEdGraphPin* Pin) const override {
-		return UNiagaraNodeWithDynamicPins::CanRemovePin(Pin);
-	}
-	virtual bool CanMovePin(const UEdGraphPin* Pin, int32 DirectionToMove) const override {
-		return UNiagaraNodeWithDynamicPins::CanMovePin(Pin, DirectionToMove);
-	}
+	virtual bool CanRemovePin(const UEdGraphPin* Pin) const override { return UNiagaraNodeWithDynamicPins::CanRemovePin(Pin); }
+	virtual bool CanMovePin(const UEdGraphPin* Pin, int32 DirectionToMove) const override {	return UNiagaraNodeWithDynamicPins::CanMovePin(Pin, DirectionToMove); }
 
 	/** Called when a new typed pin is added by the user. */
 	virtual void OnNewTypedPinAdded(UEdGraphPin*& NewPin) override;
