- Improve tooltip to notify why bounds may be greyed out

#rnx

diff --git a/Engine/Plugins/FX/Niagara/Source/Niagara/Classes/NiagaraEmitter.h b/Engine/Plugins/FX/Niagara/Source/Niagara/Classes/NiagaraEmitter.h
--- a/Engine/Plugins/FX/Niagara/Source/Niagara/Classes/NiagaraEmitter.h
+++ b/Engine/Plugins/FX/Niagara/Source/Niagara/Classes/NiagaraEmitter.h
@@ -297,11 +297,17 @@ struct NIAGARA_API FVersionedNiagaraEmitterData
 	UPROPERTY(EditAnywhere, BlueprintReadWrite, Category = "Emitter")
 	ENiagaraSimTarget SimTarget = ENiagaraSimTarget::CPUSim;
 
-	/** How should we calculate bounds for the emitter. */
+	/**
+	How should we calculate bounds for the emitter.
+	Note: If this is greyed out it means fixed bounds are enabled in the System Properties and these bounds are therefore ignored.
+	*/
 	UPROPERTY(EditAnywhere, Category = "Emitter")
 	ENiagaraEmitterCalculateBoundMode CalculateBoundsMode = ENiagaraEmitterCalculateBoundMode::Dynamic;
 	
-	/** The fixed bounding box value. CalculateBoundsMode is the condition whether the fixed bounds can be edited. */
+	/**
+	The fixed bounding box value. CalculateBoundsMode is the condition whether the fixed bounds can be edited.
+	Note: If this is greyed out it means fixed bounds are enabled in the System Properties and these bounds are therefore ignored.
+	*/
 	UPROPERTY(EditAnywhere, BlueprintReadWrite, Category = "Emitter", meta = (EditConditionHides, EditCondition = "CalculateBoundsMode == ENiagaraEmitterCalculateBoundMode::Fixed"))
 	FBox FixedBounds;
 
