Change a checkSlow to a check to see if there's actual cases where we are getting DDC conflicts with NiagaraShaders
#rnx

diff --git a/Engine/Plugins/FX/Niagara/Source/NiagaraShader/Private/NiagaraShader.cpp b/Engine/Plugins/FX/Niagara/Source/NiagaraShader/Private/NiagaraShader.cpp
--- a/Engine/Plugins/FX/Niagara/Source/NiagaraShader/Private/NiagaraShader.cpp
+++ b/Engine/Plugins/FX/Niagara/Source/NiagaraShader/Private/NiagaraShader.cpp
@@ -487,7 +487,7 @@ void FNiagaraShaderMap::LoadFromDerivedDataCache(const FNiagaraShaderScript* Scr
 				// Deserialize from the cached data
 				if (InOutShaderMap->Serialize(Ar))
 				{
-					checkSlow(InOutShaderMap->GetShaderMapId() == ShaderMapId);
+					check(InOutShaderMap->GetShaderMapId() == ShaderMapId);
 
 					// Register in the global map
 					InOutShaderMap->Register(Platform);
