Implemented a Niagara data interface for SparseVolumeTexture. Supports sampling, loading and querying texture size and number of mip levels.

#rnx

diff --git a/Engine/Plugins/FX/Niagara/Shaders/Private/NiagaraDataInterfaceSparseVolumeTextureTemplate.ush b/Engine/Plugins/FX/Niagara/Shaders/Private/NiagaraDataInterfaceSparseVolumeTextureTemplate.ush
new file mode 100644
--- /dev/null
+++ b/Engine/Plugins/FX/Niagara/Shaders/Private/NiagaraDataInterfaceSparseVolumeTextureTemplate.ush
@@ -0,0 +1,40 @@
+// Copyright Epic Games, Inc. All Rights Reserved.
+
+#include "/Engine/Private/SparseVolumeTexture/SparseVolumeTextureCommon.ush"
+
+Texture3D<uint>	{ParameterName}_PageTableTexture;
+Texture3D		{ParameterName}_PhysicalTileDataATexture;
+Texture3D		{ParameterName}_PhysicalTileDataBTexture;
+uint4			{ParameterName}_PackedUniforms0;
+uint4			{ParameterName}_PackedUniforms1;
+int3			{ParameterName}_TextureSize;
+int				{ParameterName}_MipLevels;
+
+
+void LoadSparseVolumeTexture_{ParameterName}(in int TexelX, in int TexelY, in int TexelZ, in int MipLevel, out float4 OutAttributesA, out float4 OutAttributesB)
+{
+	const FSparseVolumeTextureUniforms Uniforms = SparseVolumeTextureUnpackUniforms({ParameterName}_PackedUniforms0, {ParameterName}_PackedUniforms1);
+	const int3 VoxelCoord = SparseVolumeTextureLoadPageTable({ParameterName}_PageTableTexture, Uniforms, int3(TexelX, TexelY, TexelZ), MipLevel);
+	OutAttributesA = {ParameterName}_PhysicalTileDataATexture.Load(int4(VoxelCoord, 0));
+	OutAttributesB = {ParameterName}_PhysicalTileDataBTexture.Load(int4(VoxelCoord, 0));
+}
+
+void SampleSparseVolumeTexture_{ParameterName}(in float3 UVW, in int MipLevel, out float4 OutAttributesA, out float4 OutAttributesB)
+{
+	const FSparseVolumeTextureUniforms Uniforms = SparseVolumeTextureUnpackUniforms({ParameterName}_PackedUniforms0, {ParameterName}_PackedUniforms1);
+	const float3 VoxelCoord = SparseVolumeTextureSamplePageTable({ParameterName}_PageTableTexture, Uniforms, UVW, MipLevel);
+	OutAttributesA = {ParameterName}_PhysicalTileDataATexture.Load(int4(int3(VoxelCoord), 0));
+	OutAttributesB = {ParameterName}_PhysicalTileDataBTexture.Load(int4(int3(VoxelCoord), 0));
+}
+
+void GetTextureDimensions_{ParameterName}(int MipLevel, out float3 OutValue)
+{
+	OutValue.x = float(max({ParameterName}_TextureSize.x >> MipLevel, 1));
+	OutValue.y = float(max({ParameterName}_TextureSize.y >> MipLevel, 1));
+	OutValue.z = float(max({ParameterName}_TextureSize.z >> MipLevel, 1));
+}
+
+void GetNumMipLevels_{ParameterName}(out int OutMipLevels)
+{
+	OutMipLevels = {ParameterName}_MipLevels;
+}
diff --git a/Engine/Plugins/FX/Niagara/Source/Niagara/Classes/NiagaraDataInterfaceSparseVolumeTexture.h b/Engine/Plugins/FX/Niagara/Source/Niagara/Classes/NiagaraDataInterfaceSparseVolumeTexture.h
new file mode 100644
--- /dev/null
+++ b/Engine/Plugins/FX/Niagara/Source/Niagara/Classes/NiagaraDataInterfaceSparseVolumeTexture.h
@@ -0,0 +1,80 @@
+// Copyright Epic Games, Inc. All Rights Reserved.
+
+#pragma once
+
+#include "NiagaraDataInterface.h"
+#include "NiagaraCommon.h"
+#include "VectorVM.h"
+#include "NiagaraDataInterfaceSparseVolumeTexture.generated.h"
+
+class USparseVolumeTexture;
+
+/** Data Interface allowing sampling of a sparse volume texture */
+UCLASS(EditInlineNew, Category = "Texture", meta = (DisplayName = "Sparse Volume Texture Sample"))
+class NIAGARA_API UNiagaraDataInterfaceSparseVolumeTexture : public UNiagaraDataInterface
+{
+	GENERATED_UCLASS_BODY()
+
+	BEGIN_SHADER_PARAMETER_STRUCT(FShaderParameters, )
+		SHADER_PARAMETER_TEXTURE(Texture3D<uint>, PageTableTexture)
+		SHADER_PARAMETER_TEXTURE(Texture3D, PhysicalTileDataATexture)
+		SHADER_PARAMETER_TEXTURE(Texture3D, PhysicalTileDataBTexture)
+		SHADER_PARAMETER(FUintVector4, PackedUniforms0)
+		SHADER_PARAMETER(FUintVector4, PackedUniforms1)
+		SHADER_PARAMETER(FIntVector3, TextureSize)
+		SHADER_PARAMETER(int32, MipLevels)
+	END_SHADER_PARAMETER_STRUCT()
+
+public:
+	UPROPERTY(EditAnywhere, Category = "SparseVolumeTexture")
+	TObjectPtr<USparseVolumeTexture> SparseVolumeTexture;
+
+	UPROPERTY(EditAnywhere, Category = "SparseVolumeTexture", meta = (ToolTip = "When valid the user parameter is used as the texture rather than the one on the data interface"))
+	FNiagaraUserParameterBinding SparseVolumeTextureUserParameter;
+
+	//UObject Interface
+	virtual void PostInitProperties() override;
+	virtual void PostLoad() override;
+	virtual void Serialize(FArchive& Ar) override;
+	//UObject Interface End
+
+	//UNiagaraDataInterface Interface
+	virtual void GetFunctions(TArray<FNiagaraFunctionSignature>& OutFunctions) override;
+	virtual void GetVMExternalFunction(const FVMExternalFunctionBindingInfo& BindingInfo, void* InstanceData, FVMExternalFunction& OutFunc) override;
+	virtual bool CanExecuteOnTarget(ENiagaraSimTarget Target) const override { return Target == ENiagaraSimTarget::GPUComputeSim; }
+
+	virtual int32 PerInstanceDataSize() const override;
+	virtual bool InitPerInstanceData(void* PerInstanceData, FNiagaraSystemInstance* SystemInstance) override;
+	virtual void DestroyPerInstanceData(void* PerInstanceData, FNiagaraSystemInstance* SystemInstance) override;
+
+	virtual bool HasPreSimulateTick() const override { return true; }
+	virtual bool PerInstanceTick(void* PerInstanceData, FNiagaraSystemInstance* SystemInstance, float DeltaSeconds) override;
+	//UNiagaraDataInterface Interface End
+
+	virtual bool Equals(const UNiagaraDataInterface* Other) const override;
+
+	// GPU sim functionality
+#if WITH_EDITORONLY_DATA
+	virtual bool AppendCompileHash(FNiagaraCompileHashVisitor* InVisitor) const override;
+	virtual void GetParameterDefinitionHLSL(const FNiagaraDataInterfaceGPUParamInfo& ParamInfo, FString& OutHLSL) override;
+	virtual bool GetFunctionHLSL(const FNiagaraDataInterfaceGPUParamInfo& ParamInfo, const FNiagaraDataInterfaceGeneratedFunction& FunctionInfo, int FunctionInstanceIndex, FString& OutHLSL) override;
+	virtual bool UpgradeFunctionCall(FNiagaraFunctionSignature& FunctionSignature) override;
+#endif
+	virtual void BuildShaderParameters(FNiagaraShaderParametersBuilder& ShaderParametersBuilder) const override;
+	virtual void SetShaderParameters(const FNiagaraDataInterfaceSetShaderParametersContext& Context) const override;
+
+	void VMGetTextureDimensions(FVectorVMExternalFunctionContext& Context);
+	void VMGetNumMipLevels(FVectorVMExternalFunctionContext& Context);
+
+	void SetTexture(USparseVolumeTexture* InSparseVolumeTexture);
+
+protected:
+	virtual bool CopyToInternal(UNiagaraDataInterface* Destination) const override;
+
+protected:
+	static const TCHAR* TemplateShaderFilePath;
+	static const FName LoadSparseVolumeTextureName;
+	static const FName SampleSparseVolumeTextureName;
+	static const FName GetTextureDimensionsName;
+	static const FName GetNumMipLevelsName;
+};
diff --git a/Engine/Plugins/FX/Niagara/Source/Niagara/Private/NiagaraDataInterfaceSparseVolumeTexture.cpp b/Engine/Plugins/FX/Niagara/Source/Niagara/Private/NiagaraDataInterfaceSparseVolumeTexture.cpp
new file mode 100644
--- /dev/null
+++ b/Engine/Plugins/FX/Niagara/Source/Niagara/Private/NiagaraDataInterfaceSparseVolumeTexture.cpp
@@ -0,0 +1,379 @@
+// Copyright Epic Games, Inc. All Rights Reserved.
+
+#include "NiagaraDataInterfaceSparseVolumeTexture.h"
+#include "SparseVolumeTexture/SparseVolumeTexture.h"
+#include "SparseVolumeTexture/SparseVolumeTextureSceneProxy.h"
+#include "NiagaraComputeExecutionContext.h"
+#include "NiagaraCustomVersion.h"
+#include "NiagaraGpuComputeDispatchInterface.h"
+#include "NiagaraShader.h"
+#include "NiagaraShaderParametersBuilder.h"
+#include "NiagaraSystemInstance.h"
+#include "RHIStaticStates.h"
+#include "ShaderCompilerCore.h"
+
+#include UE_INLINE_GENERATED_CPP_BY_NAME(NiagaraDataInterfaceSparseVolumeTexture)
+
+#define LOCTEXT_NAMESPACE "UNiagaraDataInterfaceSparseVolumeTexture"
+
+const TCHAR* UNiagaraDataInterfaceSparseVolumeTexture::TemplateShaderFilePath = TEXT("/Plugin/FX/Niagara/Private/NiagaraDataInterfaceSparseVolumeTextureTemplate.ush");
+const FName UNiagaraDataInterfaceSparseVolumeTexture::LoadSparseVolumeTextureName(TEXT("LoadSparseVolumeTexture"));
+const FName UNiagaraDataInterfaceSparseVolumeTexture::SampleSparseVolumeTextureName(TEXT("SampleSparseVolumeTexture"));
+const FName UNiagaraDataInterfaceSparseVolumeTexture::GetTextureDimensionsName(TEXT("GetSparseVolumeTextureDimensions"));
+const FName UNiagaraDataInterfaceSparseVolumeTexture::GetNumMipLevelsName(TEXT("GetSparseVolumeTextureNumMipLevels"));
+
+struct FNDISparseVolumeTextureFunctionVersion
+{
+	enum Type
+	{
+		InitialVersion = 0,
+
+		VersionPlusOne,
+		LatestVersion = VersionPlusOne - 1
+	};
+};
+
+struct FNDISparseVolumeTextureInstanceData_GameThread
+{
+	TWeakObjectPtr<USparseVolumeTexture> CurrentTexture = nullptr;
+	FUintVector4 CurrentPackedUniforms0 = FUintVector4();
+	FUintVector4 CurrentPackedUniforms1 = FUintVector4();
+	FIntVector3 CurrentTextureSize = FIntVector3::ZeroValue;
+	int32 CurrentTextureMipLevels = 0;
+	FNiagaraParameterDirectBinding<UObject*> UserParamBinding;
+};
+
+struct FNDISparseVolumeTextureInstanceData_RenderThread
+{
+	FTextureRHIRef PageTableTextureRHI = nullptr;
+	FTextureRHIRef PhysicalTileDataATextureRHI = nullptr;
+	FTextureRHIRef PhysicalTileDataBTextureRHI = nullptr;
+	FUintVector4 PackedUniforms0 = FUintVector4();
+	FUintVector4 PackedUniforms1 = FUintVector4();
+	FIntVector3 TextureSize = FIntVector3::ZeroValue;
+	int32 MipLevels = 0;
+};
+
+struct FNiagaraDataInterfaceProxySparseVolumeTexture : public FNiagaraDataInterfaceProxy
+{
+	TMap<FNiagaraSystemInstanceID, FNDISparseVolumeTextureInstanceData_RenderThread> InstanceData_RT;
+
+	virtual void ConsumePerInstanceDataFromGameThread(void* PerInstanceData, const FNiagaraSystemInstanceID& Instance) override { checkNoEntry(); }
+	virtual int32 PerInstanceDataPassedToRenderThreadSize() const override { return 0; }
+};
+
+UNiagaraDataInterfaceSparseVolumeTexture::UNiagaraDataInterfaceSparseVolumeTexture(FObjectInitializer const& ObjectInitializer)
+	: Super(ObjectInitializer)
+	, SparseVolumeTexture(nullptr)
+{
+	Proxy.Reset(new FNiagaraDataInterfaceProxySparseVolumeTexture());
+
+	FNiagaraTypeDefinition Def(UObject::StaticClass());
+	SparseVolumeTextureUserParameter.Parameter.SetType(Def);
+}
+
+void UNiagaraDataInterfaceSparseVolumeTexture::PostInitProperties()
+{
+	Super::PostInitProperties();
+
+	if (HasAnyFlags(RF_ClassDefaultObject))
+	{
+		ENiagaraTypeRegistryFlags Flags = ENiagaraTypeRegistryFlags::AllowAnyVariable | ENiagaraTypeRegistryFlags::AllowParameter;
+		FNiagaraTypeRegistry::Register(FNiagaraTypeDefinition(GetClass()), Flags);
+	}
+}
+
+void UNiagaraDataInterfaceSparseVolumeTexture::PostLoad()
+{
+	Super::PostLoad();
+}
+
+void UNiagaraDataInterfaceSparseVolumeTexture::Serialize(FArchive& Ar)
+{
+	Super::Serialize(Ar);
+}
+
+void UNiagaraDataInterfaceSparseVolumeTexture::GetFunctions(TArray<FNiagaraFunctionSignature>& OutFunctions)
+{
+	FNiagaraFunctionSignature DefaultGpuSig;
+	DefaultGpuSig.bMemberFunction = true;
+	DefaultGpuSig.bRequiresContext = false;
+	DefaultGpuSig.bSupportsCPU = false;
+	DefaultGpuSig.bSupportsGPU = true;
+	DefaultGpuSig.Inputs.Emplace(FNiagaraTypeDefinition(GetClass()), TEXT("SparseVolumeTexture"));
+	DefaultGpuSig.SetFunctionVersion(FNDISparseVolumeTextureFunctionVersion::LatestVersion);
+
+	{
+		FNiagaraFunctionSignature& Sig = OutFunctions.Add_GetRef(DefaultGpuSig);
+		Sig.Name = LoadSparseVolumeTextureName;
+		Sig.Inputs.Emplace(FNiagaraTypeDefinition::GetIntDef(), TEXT("TexelX"));
+		Sig.Inputs.Emplace(FNiagaraTypeDefinition::GetIntDef(), TEXT("TexelY"));
+		Sig.Inputs.Emplace(FNiagaraTypeDefinition::GetIntDef(), TEXT("TexelZ"));
+		Sig.Inputs.Emplace(FNiagaraTypeDefinition::GetIntDef(), TEXT("MipLevel"));
+		Sig.Outputs.Emplace(FNiagaraTypeDefinition::GetVec4Def(), TEXT("AttributesA"));
+		Sig.Outputs.Emplace(FNiagaraTypeDefinition::GetVec4Def(), TEXT("AttributesB"));
+		Sig.SetDescription(LOCTEXT("SparseVolumeTextureLoadDesc", "Read a voxel from the provided location & mip without any filtering or sampling."));
+	}
+
+	{
+		FNiagaraFunctionSignature& Sig = OutFunctions.Add_GetRef(DefaultGpuSig);
+		Sig.Name = SampleSparseVolumeTextureName;
+		Sig.Inputs.Emplace(FNiagaraTypeDefinition::GetVec3Def(), TEXT("UVW"));
+		Sig.Inputs.Emplace(FNiagaraTypeDefinition::GetIntDef(), TEXT("MipLevel"));
+		Sig.Outputs.Emplace(FNiagaraTypeDefinition::GetVec4Def(), TEXT("AttributesA"));
+		Sig.Outputs.Emplace(FNiagaraTypeDefinition::GetVec4Def(), TEXT("AttributesB"));
+		Sig.SetDescription(LOCTEXT("SparseVolumeTextureLoadDesc", "Sample supplied mip level from input sparse volume texture at the specified UVW coordinates."));
+	}
+
+	{
+		FNiagaraFunctionSignature& Sig = OutFunctions.AddDefaulted_GetRef();
+		Sig.Name = GetTextureDimensionsName;
+		Sig.bSupportsCPU = true;
+		Sig.Inputs.Emplace(FNiagaraTypeDefinition::GetIntDef(), TEXT("MipLevel"));
+		Sig.Outputs.Emplace(FNiagaraTypeDefinition::GetVec3Def(), TEXT("Dimensions3D"));
+		Sig.SetDescription(LOCTEXT("SparseVolumeTextureDimsDesc", "Get the dimensions of the provided mip level."));
+	}
+
+	{
+		FNiagaraFunctionSignature& Sig = OutFunctions.AddDefaulted_GetRef();
+		Sig.Name = GetNumMipLevelsName;
+		Sig.bSupportsCPU = true;
+		Sig.Outputs.Emplace(FNiagaraTypeDefinition::GetIntDef(), TEXT("NumMipLevels"));
+		Sig.SetDescription(LOCTEXT("SparseVolumeGetNumMipLevelsDesc", "Get the number of mip levels."));
+	}
+}
+
+void UNiagaraDataInterfaceSparseVolumeTexture::GetVMExternalFunction(const FVMExternalFunctionBindingInfo& BindingInfo, void* InstanceData, FVMExternalFunction& OutFunc)
+{
+	if (BindingInfo.Name == GetTextureDimensionsName)
+	{
+		check(BindingInfo.GetNumInputs() == 2 && BindingInfo.GetNumOutputs() == 3);
+		OutFunc = FVMExternalFunction::CreateUObject(this, &UNiagaraDataInterfaceSparseVolumeTexture::VMGetTextureDimensions);
+	}
+	else if (BindingInfo.Name == GetNumMipLevelsName)
+	{
+		check(BindingInfo.GetNumInputs() == 1 && BindingInfo.GetNumOutputs() == 1);
+		OutFunc = FVMExternalFunction::CreateUObject(this, &UNiagaraDataInterfaceSparseVolumeTexture::VMGetNumMipLevels);
+	}
+}
+
+int32 UNiagaraDataInterfaceSparseVolumeTexture::PerInstanceDataSize() const
+{
+	return sizeof(FNDISparseVolumeTextureInstanceData_GameThread);
+}
+
+bool UNiagaraDataInterfaceSparseVolumeTexture::InitPerInstanceData(void* PerInstanceData, FNiagaraSystemInstance* SystemInstance)
+{
+	FNDISparseVolumeTextureInstanceData_GameThread* InstanceData = new (PerInstanceData) FNDISparseVolumeTextureInstanceData_GameThread();
+	InstanceData->UserParamBinding.Init(SystemInstance->GetInstanceParameters(), SparseVolumeTextureUserParameter.Parameter);
+	return true;
+}
+
+void UNiagaraDataInterfaceSparseVolumeTexture::DestroyPerInstanceData(void* PerInstanceData, FNiagaraSystemInstance* SystemInstance)
+{
+	FNDISparseVolumeTextureInstanceData_GameThread* InstanceData = static_cast<FNDISparseVolumeTextureInstanceData_GameThread*>(PerInstanceData);
+	InstanceData->~FNDISparseVolumeTextureInstanceData_GameThread();
+
+	ENQUEUE_RENDER_COMMAND(NDISparseVolumeTextureTexture_RemoveInstance)
+	(
+		[RT_Proxy = GetProxyAs<FNiagaraDataInterfaceProxySparseVolumeTexture>(), RT_InstanceID = SystemInstance->GetId()](FRHICommandListImmediate&)
+		{
+			RT_Proxy->InstanceData_RT.Remove(RT_InstanceID);
+		}
+	);
+}
+
+bool UNiagaraDataInterfaceSparseVolumeTexture::PerInstanceTick(void* PerInstanceData, FNiagaraSystemInstance* SystemInstance, float DeltaSeconds)
+{
+	FNDISparseVolumeTextureInstanceData_GameThread* InstanceData = static_cast<FNDISparseVolumeTextureInstanceData_GameThread*>(PerInstanceData);
+
+	USparseVolumeTexture* CurrentTexture = InstanceData->UserParamBinding.GetValueOrDefault<USparseVolumeTexture>(SparseVolumeTexture);
+	FUintVector4 CurrentPackedUniforms0 = FUintVector4();
+	FUintVector4 CurrentPackedUniforms1 = FUintVector4();
+	FIntVector3 CurrentTextureSize = FIntVector3::ZeroValue;
+	int32 CurrentTextureMipLevels = 0;
+	if (CurrentTexture)
+	{
+		CurrentTexture->GetPackedUniforms(CurrentPackedUniforms0, CurrentPackedUniforms1);
+		CurrentTextureSize = FIntVector3(CurrentTexture->GetVolumeResolution());
+		CurrentTextureMipLevels = CurrentTexture->GetNumMipLevels();
+	}
+
+	if ((InstanceData->CurrentTexture != CurrentTexture) 
+		|| (InstanceData->CurrentTextureSize != CurrentTextureSize) 
+		|| (InstanceData->CurrentTextureMipLevels != CurrentTextureMipLevels)
+		|| (InstanceData->CurrentPackedUniforms0 != CurrentPackedUniforms0)
+		|| (InstanceData->CurrentPackedUniforms1 != CurrentPackedUniforms1))
+	{
+		InstanceData->CurrentTexture = CurrentTexture;
+		InstanceData->CurrentTextureSize = CurrentTextureSize;
+		InstanceData->CurrentTextureMipLevels = CurrentTextureMipLevels;
+		InstanceData->CurrentPackedUniforms0 = CurrentPackedUniforms0;
+		InstanceData->CurrentPackedUniforms1 = CurrentPackedUniforms1;
+
+		ENQUEUE_RENDER_COMMAND(NDISparseVolumeTexture_UpdateInstance)
+		(
+			[RT_Proxy = GetProxyAs<FNiagaraDataInterfaceProxySparseVolumeTexture>(), RT_InstanceID = SystemInstance->GetId(), RT_Texture = CurrentTexture, RT_Packed0 = CurrentPackedUniforms0, 
+			RT_Packed1 = CurrentPackedUniforms1, RT_TextureSize = CurrentTextureSize, RT_MipLevels = CurrentTextureMipLevels] (FRHICommandListImmediate&)
+			{
+				FNDISparseVolumeTextureInstanceData_RenderThread& RTInstanceData = RT_Proxy->InstanceData_RT.FindOrAdd(RT_InstanceID);
+				const FSparseVolumeTextureSceneProxy* SceneProxy = RT_Texture ? RT_Texture->GetSparseVolumeTextureSceneProxy() : nullptr;
+				if (RT_Texture && SceneProxy)
+				{
+					FRHITexture* PageTableTexture = SceneProxy->GetPageTableTextureRHI();
+					FRHITexture* TextureA = SceneProxy->GetPhysicalTileDataATextureRHI();
+					FRHITexture* TextureB = SceneProxy->GetPhysicalTileDataBTextureRHI();
+
+					RTInstanceData.PageTableTextureRHI = PageTableTexture;
+					RTInstanceData.PhysicalTileDataATextureRHI = TextureA;
+					RTInstanceData.PhysicalTileDataBTextureRHI = TextureB;
+					RTInstanceData.PackedUniforms0 = RT_Packed0;
+					RTInstanceData.PackedUniforms1 = RT_Packed1;
+					RTInstanceData.TextureSize = RT_TextureSize;
+					RTInstanceData.MipLevels = RT_MipLevels;
+				}
+				else
+				{
+					RTInstanceData.PageTableTextureRHI = nullptr;
+					RTInstanceData.PhysicalTileDataATextureRHI = nullptr;
+					RTInstanceData.PhysicalTileDataBTextureRHI = nullptr;
+					RTInstanceData.PackedUniforms0 = FUintVector4();
+					RTInstanceData.PackedUniforms1 = FUintVector4();
+					RTInstanceData.TextureSize = FIntVector3::ZeroValue;
+					RTInstanceData.MipLevels = 0;
+				}
+			}
+		);
+	}
+
+	return false;
+}
+
+bool UNiagaraDataInterfaceSparseVolumeTexture::Equals(const UNiagaraDataInterface* Other) const
+{
+	if (!Super::Equals(Other))
+	{
+		return false;
+	}
+	const UNiagaraDataInterfaceSparseVolumeTexture* OtherSparseVolumeTexture = CastChecked<const UNiagaraDataInterfaceSparseVolumeTexture>(Other);
+	return OtherSparseVolumeTexture->SparseVolumeTexture == SparseVolumeTexture && OtherSparseVolumeTexture->SparseVolumeTextureUserParameter == SparseVolumeTextureUserParameter;
+}
+
+#if WITH_EDITORONLY_DATA
+bool UNiagaraDataInterfaceSparseVolumeTexture::AppendCompileHash(FNiagaraCompileHashVisitor* InVisitor) const
+{
+	bool bSuccess = Super::AppendCompileHash(InVisitor);
+	InVisitor->UpdateShaderFile(TemplateShaderFilePath);
+	bSuccess &= InVisitor->UpdateShaderParameters<FShaderParameters>();
+	return bSuccess;
+}
+
+void UNiagaraDataInterfaceSparseVolumeTexture::GetParameterDefinitionHLSL(const FNiagaraDataInterfaceGPUParamInfo& ParamInfo, FString& OutHLSL)
+{
+	const TMap<FString, FStringFormatArg> TemplateArgs =
+	{
+		{TEXT("ParameterName"),	ParamInfo.DataInterfaceHLSLSymbol},
+	};
+	AppendTemplateHLSL(OutHLSL, TemplateShaderFilePath, TemplateArgs);
+}
+
+bool UNiagaraDataInterfaceSparseVolumeTexture::GetFunctionHLSL(const FNiagaraDataInterfaceGPUParamInfo& ParamInfo, const FNiagaraDataInterfaceGeneratedFunction& FunctionInfo, int FunctionInstanceIndex, FString& OutHLSL)
+{
+	static const TSet<FName> ValidGpuFunctions =
+	{
+		LoadSparseVolumeTextureName,
+		SampleSparseVolumeTextureName,
+		GetTextureDimensionsName,
+		GetNumMipLevelsName,
+	};
+
+	return ValidGpuFunctions.Contains(FunctionInfo.DefinitionName);
+}
+
+bool UNiagaraDataInterfaceSparseVolumeTexture::UpgradeFunctionCall(FNiagaraFunctionSignature& FunctionSignature)
+{
+	return false;
+}
+#endif
+
+void UNiagaraDataInterfaceSparseVolumeTexture::BuildShaderParameters(FNiagaraShaderParametersBuilder& ShaderParametersBuilder) const
+{
+	ShaderParametersBuilder.AddNestedStruct<FShaderParameters>();
+}
+
+void UNiagaraDataInterfaceSparseVolumeTexture::SetShaderParameters(const FNiagaraDataInterfaceSetShaderParametersContext& Context) const
+{
+	FNiagaraDataInterfaceProxySparseVolumeTexture& TextureProxy = Context.GetProxy<FNiagaraDataInterfaceProxySparseVolumeTexture>();
+	FNDISparseVolumeTextureInstanceData_RenderThread* RTInstanceData = TextureProxy.InstanceData_RT.Find(Context.GetSystemInstanceID());
+
+	FShaderParameters* Parameters = Context.GetParameterNestedStruct<FShaderParameters>();
+	if (RTInstanceData && RTInstanceData->PageTableTextureRHI.IsValid() && (RTInstanceData->PhysicalTileDataATextureRHI.IsValid() || RTInstanceData->PhysicalTileDataBTextureRHI.IsValid()))
+	{
+		Parameters->PageTableTexture = RTInstanceData->PageTableTextureRHI;
+		Parameters->PhysicalTileDataATexture = RTInstanceData->PhysicalTileDataATextureRHI.IsValid() ? RTInstanceData->PhysicalTileDataATextureRHI : GBlackVolumeTexture->TextureRHI;
+		Parameters->PhysicalTileDataBTexture = RTInstanceData->PhysicalTileDataBTextureRHI.IsValid() ? RTInstanceData->PhysicalTileDataBTextureRHI : GBlackVolumeTexture->TextureRHI;
+		Parameters->PackedUniforms0 = RTInstanceData->PackedUniforms0;
+		Parameters->PackedUniforms1 = RTInstanceData->PackedUniforms1;
+		Parameters->TextureSize = RTInstanceData->TextureSize;
+		Parameters->MipLevels = RTInstanceData->MipLevels;
+	}
+	else
+	{
+		Parameters->PageTableTexture = GBlackUintVolumeTexture->TextureRHI;
+		Parameters->PhysicalTileDataATexture = GBlackVolumeTexture->TextureRHI;
+		Parameters->PhysicalTileDataBTexture = GBlackVolumeTexture->TextureRHI;
+		Parameters->PackedUniforms0 = FUintVector4();
+		Parameters->PackedUniforms1 = FUintVector4();
+		Parameters->TextureSize = FIntVector3::ZeroValue;
+		Parameters->MipLevels = 0;
+	}
+}
+
+void UNiagaraDataInterfaceSparseVolumeTexture::VMGetTextureDimensions(FVectorVMExternalFunctionContext& Context)
+{
+	VectorVM::FUserPtrHandler<FNDISparseVolumeTextureInstanceData_GameThread> InstData(Context);
+	FNDIInputParam<int32>		InMipLevel(Context);
+	FNDIOutputParam<FVector3f>	OutSize(Context);
+
+	for (int32 i = 0; i < Context.GetNumInstances(); ++i)
+	{
+		const int32 MipLevel = InMipLevel.GetAndAdvance();
+		const FVector3f TextureSize(
+			float(FMath::Max(InstData->CurrentTextureSize.X >> MipLevel, 1)),
+			float(FMath::Max(InstData->CurrentTextureSize.Y >> MipLevel, 1)),
+			float(FMath::Max(InstData->CurrentTextureSize.Z >> MipLevel, 1))
+		);
+		OutSize.SetAndAdvance(TextureSize);
+	}
+}
+
+void UNiagaraDataInterfaceSparseVolumeTexture::VMGetNumMipLevels(FVectorVMExternalFunctionContext& Context)
+{
+	VectorVM::FUserPtrHandler<FNDISparseVolumeTextureInstanceData_GameThread> InstData(Context);
+	FNDIOutputParam<int32> OutNumMipLevels(Context);
+
+	for (int32 i = 0; i < Context.GetNumInstances(); ++i)
+	{
+		OutNumMipLevels.SetAndAdvance(InstData->CurrentTextureMipLevels);
+	}
+}
+
+void UNiagaraDataInterfaceSparseVolumeTexture::SetTexture(USparseVolumeTexture* InSparseVolumeTexture)
+{
+	SparseVolumeTexture = InSparseVolumeTexture;
+}
+
+bool UNiagaraDataInterfaceSparseVolumeTexture::CopyToInternal(UNiagaraDataInterface* Destination) const
+{
+	if (!Super::CopyToInternal(Destination))
+	{
+		return false;
+	}
+	UNiagaraDataInterfaceSparseVolumeTexture* DestinationTexture = CastChecked<UNiagaraDataInterfaceSparseVolumeTexture>(Destination);
+	DestinationTexture->SparseVolumeTexture = SparseVolumeTexture;
+	DestinationTexture->SparseVolumeTextureUserParameter = SparseVolumeTextureUserParameter;
+
+	return true;
+}
+
+#undef LOCTEXT_NAMESPACE
diff --git a/Engine/Shaders/Private/SparseVolumeTexture/SparseVolumeTextureCommon.ush b/Engine/Shaders/Private/SparseVolumeTexture/SparseVolumeTextureCommon.ush
--- a/Engine/Shaders/Private/SparseVolumeTexture/SparseVolumeTextureCommon.ush
+++ b/Engine/Shaders/Private/SparseVolumeTexture/SparseVolumeTextureCommon.ush
@@ -48,6 +48,23 @@ float3 SparseVolumeTextureSamplePageTable(Texture3D<uint> PageTable, const FSpar
 	return VoxelCoord;
 }
 
+int3 SparseVolumeTextureLoadPageTable(Texture3D<uint> PageTable, const FSparseVolumeTextureUniforms Uniforms, const int3 TexelCoord, int MipLevel = 0)
+{
+	if (MipLevel < Uniforms.MinMipLevel || MipLevel > Uniforms.MaxMipLevel)
+	{
+		return 0; // Point to null tile
+	}
+	const float RcpMipLevelFactor = rcp(float(1u << (uint)MipLevel));
+	const float3 VolumeMipPageCoord = (TexelCoord + 0.5f) / Uniforms.TileSize;
+	const float3 MipPageTableOffset = floor(Uniforms.PageTableOffset * RcpMipLevelFactor);
+	const float3 PageTableCoord = VolumeMipPageCoord - MipPageTableOffset;
+
+	const uint PackedPhysicalTileCoord = PageTable.Load(int4(int3(PageTableCoord), (MipLevel - Uniforms.MinMipLevel))).x;
+
+	const float3 VoxelCoord = SparseVolumeTextureGetVoxelCoord(PackedPhysicalTileCoord, Uniforms.TileSize, PageTableCoord);
+	return int3(VoxelCoord);
+}
+
 float4 SparseVolumeTextureSamplePhysicalTileData(Texture3D PhysicalTileDataA, Texture3D PhysicalTileDataB, const float3 VoxelCoord, const int PhysicalTileDataIndex)
 {
 	switch (PhysicalTileDataIndex)
