Fix for LOCTEXT warning in NiagaraDataInterfaceSparseVolumeTexture.cpp

#rnx

diff --git a/Engine/Plugins/FX/Niagara/Source/Niagara/Private/NiagaraDataInterfaceSparseVolumeTexture.cpp b/Engine/Plugins/FX/Niagara/Source/Niagara/Private/NiagaraDataInterfaceSparseVolumeTexture.cpp
--- a/Engine/Plugins/FX/Niagara/Source/Niagara/Private/NiagaraDataInterfaceSparseVolumeTexture.cpp
+++ b/Engine/Plugins/FX/Niagara/Source/Niagara/Private/NiagaraDataInterfaceSparseVolumeTexture.cpp
@@ -122,7 +122,7 @@ void UNiagaraDataInterfaceSparseVolumeTexture::GetFunctions(TArray<FNiagaraFunct
 		Sig.Inputs.Emplace(FNiagaraTypeDefinition::GetIntDef(), TEXT("MipLevel"));
 		Sig.Outputs.Emplace(FNiagaraTypeDefinition::GetVec4Def(), TEXT("AttributesA"));
 		Sig.Outputs.Emplace(FNiagaraTypeDefinition::GetVec4Def(), TEXT("AttributesB"));
-		Sig.SetDescription(LOCTEXT("SparseVolumeTextureLoadDesc", "Sample supplied mip level from input sparse volume texture at the specified UVW coordinates."));
+		Sig.SetDescription(LOCTEXT("SparseVolumeTextureSampleDesc", "Sample supplied mip level from input sparse volume texture at the specified UVW coordinates."));
 	}
 
 	{
