- Only set the parent if it changed
- Fixes a -game issue in some cases where the parent doesn't change but the internal call on the MIC will hit a check

#rnx

diff --git a/Engine/Plugins/FX/Niagara/Source/Niagara/Private/NiagaraRendererProperties.cpp b/Engine/Plugins/FX/Niagara/Source/Niagara/Private/NiagaraRendererProperties.cpp
--- a/Engine/Plugins/FX/Niagara/Source/Niagara/Private/NiagaraRendererProperties.cpp
+++ b/Engine/Plugins/FX/Niagara/Source/Niagara/Private/NiagaraRendererProperties.cpp
@@ -493,8 +493,12 @@ void UNiagaraRendererProperties::UpdateMaterialParametersMIC(const FNiagaraRende
 	if (InOutMIC == nullptr)
 	{
 		InOutMIC = NewObject<UMaterialInstanceConstant>(this);
+		InOutMIC->SetParentEditorOnly(InOutMaterial);
+	}
+	else if (InOutMIC->Parent != InOutMaterial)
+	{
+		InOutMIC->SetParentEditorOnly(InOutMaterial);
 	}
-	InOutMIC->SetParentEditorOnly(InOutMaterial);
 
 	if (UpdateMaterialStaticParameters(MaterialParameters, InOutMIC) == false)
 	{
@@ -532,8 +536,20 @@ void UNiagaraRendererProperties::UpdateMaterialParametersMIC(const FNiagaraRende
 		}
 
 		//-OPT: We should be able to reuse rather than create
-		UMaterialInstanceConstant* MIC = MICPool.Num() > 0 ? MICPool.Pop() : NewObject<UMaterialInstanceConstant>(this);
-		MIC->SetParentEditorOnly(Material);
+		UMaterialInstanceConstant* MIC = nullptr;
+		if (MICPool.Num() > 0)
+		{
+			MIC = MICPool.Pop();
+			if (MIC->Parent != Material)
+			{
+				MIC->SetParentEditorOnly(Material);
+			}
+		}
+		else
+		{
+			MIC = NewObject<UMaterialInstanceConstant>(this);
+			MIC->SetParentEditorOnly(Material);
+		}
 
 		if (UpdateMaterialStaticParameters(MaterialParameters, MIC) == false)
 		{
