- Add SystemInstanceState information to Debug HUD
- This allows you to see the internal state, i.e. are we really paused vs the VM state

#rnx

diff --git a/Engine/Plugins/FX/Niagara/Source/Niagara/Private/NiagaraDebugHud.cpp b/Engine/Plugins/FX/Niagara/Source/Niagara/Private/NiagaraDebugHud.cpp
--- a/Engine/Plugins/FX/Niagara/Source/Niagara/Private/NiagaraDebugHud.cpp
+++ b/Engine/Plugins/FX/Niagara/Source/Niagara/Private/NiagaraDebugHud.cpp
@@ -2584,6 +2584,7 @@ void FNiagaraDebugHud::DrawComponents(FNiagaraWorldManager* WorldManager, UCanva
 	// Draw in world components
 	UEnum* ExecutionStateEnum = StaticEnum<ENiagaraExecutionState>();
 	UEnum* PoolingMethodEnum = StaticEnum<ENCPoolMethod>();
+	UEnum* SystemInstanceState = StaticEnum<ENiagaraSystemInstanceState>();
 	for (TWeakObjectPtr<UNiagaraComponent> WeakComponent : InWorldComponents)
 	{
 		UNiagaraComponent* NiagaraComponent = WeakComponent.Get();
@@ -2740,7 +2741,12 @@ void FNiagaraDebugHud::DrawComponents(FNiagaraWorldManager* WorldManager, UCanva
 
 				if (Settings.SystemDebugVerbosity == ENiagaraDebugHudVerbosity::Verbose)
 				{
-					StringBuilder.Appendf(TEXT("System ActualState %s - RequestedState %s\n"), *ExecutionStateEnum->GetNameStringByIndex((int32)SystemInstance->GetActualExecutionState()), *ExecutionStateEnum->GetNameStringByIndex((int32)SystemInstance->GetRequestedExecutionState()));
+					StringBuilder.Appendf(
+						TEXT("System InstanceState(%s) VMActualState(%s) VMRequestedState(%s)\n"),
+						*SystemInstanceState->GetNameStringByIndex((int32)SystemInstance->SystemInstanceState),
+						*ExecutionStateEnum->GetNameStringByIndex((int32)SystemInstance->GetActualExecutionState()),
+						*ExecutionStateEnum->GetNameStringByIndex((int32)SystemInstance->GetRequestedExecutionState())
+					);
 					if (NiagaraComponent->PoolingMethod != ENCPoolMethod::None)
 					{
 						StringBuilder.Appendf(TEXT("Pooled - %s\n"), *PoolingMethodEnum->GetNameStringByIndex((int32)NiagaraComponent->PoolingMethod));
