- Modify DI HLSL menu to contain the function description

diff --git a/Engine/Plugins/FX/Niagara/Source/NiagaraEditor/Private/EdGraphSchema_Niagara.cpp b/Engine/Plugins/FX/Niagara/Source/NiagaraEditor/Private/EdGraphSchema_Niagara.cpp
--- a/Engine/Plugins/FX/Niagara/Source/NiagaraEditor/Private/EdGraphSchema_Niagara.cpp
+++ b/Engine/Plugins/FX/Niagara/Source/NiagaraEditor/Private/EdGraphSchema_Niagara.cpp
@@ -2160,9 +2160,9 @@ void UEdGraphSchema_Niagara::GetNumericConversionToSubMenuActionsAll(UToolMenu*
 	}
 }
 
-TArray<TPair<FString, FString>> UEdGraphSchema_Niagara::GetDataInterfaceFunctionPrototypes(const UEdGraphPin* InGraphPin)
+TArray<TTuple<FString, FString, FText>> UEdGraphSchema_Niagara::GetDataInterfaceFunctionPrototypes(const UEdGraphPin* InGraphPin)
 {
-	TArray<TPair<FString, FString>> FunctionPrototypes;
+	TArray<TTuple<FString, FString, FText>> FunctionPrototypes;
 	if (InGraphPin)
 	{
 		FNiagaraTypeDefinition GraphPinType = UEdGraphSchema_Niagara::PinToTypeDefinition(InGraphPin);
@@ -2202,7 +2202,8 @@ TArray<TPair<FString, FString>> UEdGraphSchema_Niagara::GetDataInterfaceFunction
 				if ( Prototype.Len() > 0 )
 				{
 					Prototype.Append(TEXT("\r\n"));
-					FunctionPrototypes.Emplace(FunctionSignature.GetNameString(), Prototype);
+					FText Description = FText::Format(LOCTEXT("DataInterfacePrototypeDescriptionFormat", "{0}\n{1}"), FText::FromString(Prototype), FunctionSignature.GetDescription());
+					FunctionPrototypes.Emplace(FunctionSignature.GetNameString(), Prototype, Description);
 				}
 			}
 		}
@@ -2212,13 +2213,13 @@ TArray<TPair<FString, FString>> UEdGraphSchema_Niagara::GetDataInterfaceFunction
 
 void UEdGraphSchema_Niagara::GenerateDataInterfacePinMenu(UToolMenu* ToolMenu, const FName SectionName, const UEdGraphPin* GraphPin, FNiagaraTypeDefinition TypeDef) const
 {
-	TArray<TPair<FString, FString>> FunctionPrototypes = GetDataInterfaceFunctionPrototypes(GraphPin);
+	TArray<TTuple<FString, FString, FText>> FunctionPrototypes = GetDataInterfaceFunctionPrototypes(GraphPin);
 	if (FunctionPrototypes.Num () == 0)
 	{
 		return;
 	}
 
-	Algo::Sort(FunctionPrototypes, [](const TPair<FString, FString>& Lhs, const TPair<FString, FString>& Rhs) { return Lhs.Key < Rhs.Key; });
+	Algo::Sort(FunctionPrototypes, [](const TTuple<FString, FString, FText>& Lhs, const TTuple<FString, FString, FText>& Rhs) { return Lhs.Get<0>() < Rhs.Get<0>(); });
 
 	// Make the menu
 	FToolMenuSection& Section = ToolMenu->FindOrAddSection(SectionName);
@@ -2232,7 +2233,7 @@ void UEdGraphSchema_Niagara::GenerateDataInterfacePinMenu(UToolMenu* ToolMenu, c
 				FString AllPrototypes;
 				for ( const auto& Prototype : FunctionPrototypes )
 				{
-					AllPrototypes += Prototype.Value;
+					AllPrototypes += Prototype.Get<1>();
 				}
 
 				FToolMenuSection& Section = ToolMenu->FindOrAddSection(SectionName);
@@ -2249,10 +2250,10 @@ void UEdGraphSchema_Niagara::GenerateDataInterfacePinMenu(UToolMenu* ToolMenu, c
 				{
 					Section.AddMenuEntry(
 						NAME_None,
-						FText::FromString(Prototype.Key),
-						FText::FromString(Prototype.Value),
+						FText::FromString(Prototype.Get<0>()),
+						Prototype.Get<2>(),
 						FSlateIcon(),
-						FUIAction(FExecuteAction::CreateLambda([ClipboardText=Prototype.Value]() { FPlatformApplicationMisc::ClipboardCopy(*ClipboardText); }))
+						FUIAction(FExecuteAction::CreateLambda([ClipboardText=Prototype.Get<1>()]() { FPlatformApplicationMisc::ClipboardCopy(*ClipboardText); }))
 					);
 				}
 			}
diff --git a/Engine/Plugins/FX/Niagara/Source/NiagaraEditor/Private/Widgets/SNiagaraGraphNodeCustomHlsl.cpp b/Engine/Plugins/FX/Niagara/Source/NiagaraEditor/Private/Widgets/SNiagaraGraphNodeCustomHlsl.cpp
--- a/Engine/Plugins/FX/Niagara/Source/NiagaraEditor/Private/Widgets/SNiagaraGraphNodeCustomHlsl.cpp
+++ b/Engine/Plugins/FX/Niagara/Source/NiagaraEditor/Private/Widgets/SNiagaraGraphNodeCustomHlsl.cpp
@@ -130,8 +130,8 @@ void SNiagaraGraphNodeCustomHlsl::CreateBelowPinControls(TSharedPtr<SVerticalBox
 				TStringBuilder<256> Options;
 				for (int i = 0; (i + AutoCompletePageIndex * 9) < FunctionPrototypes.Num() && i < 9; i++)
 				{
-					TPair<FString, FString> Prototype = FunctionPrototypes[i + AutoCompletePageIndex * 9];
-					Options.Appendf(TEXT("\n(%i) %s"), i + 1, *Prototype.Key);
+					const FString& Prototype = FunctionPrototypes[i + AutoCompletePageIndex * 9].Get<0>();
+					Options.Appendf(TEXT("\n(%i) %s"), i + 1, *Prototype);
 				}
 				FText HiddenText = FText::GetEmpty();
 				if (FunctionPrototypes.Num() > (AutoCompletePageIndex + 1) * 9)
@@ -251,14 +251,14 @@ FReply SNiagaraGraphNodeCustomHlsl::OnShaderTextKeyChar(const FGeometry&, const
 
 						for (int i = 0; i < FunctionPrototypes.Num(); i++)
 						{
-							TPair<FString, FString> Prototype = FunctionPrototypes[i];
+							FString Prototype = FunctionPrototypes[i].Get<1>();
 							int32 NameDot = 0;
-							if (Prototype.Value.FindChar('.', NameDot))
+							if (Prototype.FindChar('.', NameDot))
 							{
-								Prototype.Value.RightChopInline(NameDot + 1);
+								Prototype.RightChopInline(NameDot + 1);
 							}
-							Prototype.Value.TrimStartAndEndInline();
-							AutoCompleteOptions.Add(Prototype.Value);
+							Prototype.TrimStartAndEndInline();
+							AutoCompleteOptions.Add(Prototype);
 						}						
 						MenuAnchor->SetIsOpen(true, false);
 						break;
diff --git a/Engine/Plugins/FX/Niagara/Source/NiagaraEditor/Private/Widgets/SNiagaraGraphNodeCustomHlsl.h b/Engine/Plugins/FX/Niagara/Source/NiagaraEditor/Private/Widgets/SNiagaraGraphNodeCustomHlsl.h
--- a/Engine/Plugins/FX/Niagara/Source/NiagaraEditor/Private/Widgets/SNiagaraGraphNodeCustomHlsl.h
+++ b/Engine/Plugins/FX/Niagara/Source/NiagaraEditor/Private/Widgets/SNiagaraGraphNodeCustomHlsl.h
@@ -43,6 +43,6 @@ private:
 
 	bool bIsAutoCompleteActive = false;
 	TArray<FString> AutoCompleteOptions;
-	TArray<TPair<FString, FString>> FunctionPrototypes;
+	TArray<TTuple<FString, FString, FText>> FunctionPrototypes;
 	int32 AutoCompletePageIndex = 0;
 };
diff --git a/Engine/Plugins/FX/Niagara/Source/NiagaraEditor/Public/EdGraphSchema_Niagara.h b/Engine/Plugins/FX/Niagara/Source/NiagaraEditor/Public/EdGraphSchema_Niagara.h
--- a/Engine/Plugins/FX/Niagara/Source/NiagaraEditor/Public/EdGraphSchema_Niagara.h
+++ b/Engine/Plugins/FX/Niagara/Source/NiagaraEditor/Public/EdGraphSchema_Niagara.h
@@ -186,7 +186,9 @@ class NIAGARAEDITOR_API UEdGraphSchema_Niagara : public UEdGraphSchema
 	
 	bool PinTypesValidForNumericConversion(FEdGraphPinType AType, FEdGraphPinType BType) const;	
 
-	static TArray<TPair<FString, FString>> GetDataInterfaceFunctionPrototypes(const UEdGraphPin* GraphPin);
+	// Return type is Name, Prototype, Description
+	static TArray<TTuple<FString, FString, FText>> GetDataInterfaceFunctionPrototypes(const UEdGraphPin* GraphPin);
+
 private:
 	void GetNumericConversionToSubMenuActions(class UToolMenu* Menu, const FName SectionName, UEdGraphPin* InGraphPin);
 	void GetNumericConversionToSubMenuActionsAll(class UToolMenu* Menu, const FName SectionName, UNiagaraNode* InGraphPin);
