- Remove ToString from name builders

#rnx

diff --git a/Engine/Plugins/FX/Niagara/Source/Niagara/Private/NiagaraScriptExecutionParameterStore.cpp b/Engine/Plugins/FX/Niagara/Source/Niagara/Private/NiagaraScriptExecutionParameterStore.cpp
--- a/Engine/Plugins/FX/Niagara/Source/Niagara/Private/NiagaraScriptExecutionParameterStore.cpp
+++ b/Engine/Plugins/FX/Niagara/Source/Niagara/Private/NiagaraScriptExecutionParameterStore.cpp
@@ -326,7 +326,7 @@ void FNiagaraScriptExecutionParameterStore::AddScriptParams(UNiagaraScript* Scri
 			NameBuilder.Append(FNiagaraConstants::InternalNamespaceString);
 			NameBuilder.AppendChar(TEXT('.'));
 			Info.Name.AppendString(NameBuilder);
-			ParameterName = FName(NameBuilder.ToString()); ;
+			ParameterName = FName(NameBuilder);
 		}
 
 		FNiagaraVariable Var(Info.Type, ParameterName);
diff --git a/Engine/Plugins/FX/Niagara/Source/Niagara/Private/NiagaraTypes.cpp b/Engine/Plugins/FX/Niagara/Source/Niagara/Private/NiagaraTypes.cpp
--- a/Engine/Plugins/FX/Niagara/Source/Niagara/Private/NiagaraTypes.cpp
+++ b/Engine/Plugins/FX/Niagara/Source/Niagara/Private/NiagaraTypes.cpp
@@ -20,7 +20,7 @@ void FNiagaraVariableBase::SetNamespacedName(const FString& InNamespace, FName I
 	NameBuilder.Append(InNamespace);
 	NameBuilder.AppendChar(TEXT('.'));
 	InVariableName.AppendString(NameBuilder);
-	Name = FName(NameBuilder.ToString());
+	Name = FName(NameBuilder);
 }
 
 bool FNiagaraVariableBase::RemoveRootNamespace(const FStringView& ExpectedNamespace)
