- Fix a potential crash with Niagara Arrays that are modified during construction scripts

#rnx

diff --git a/Engine/Plugins/FX/Niagara/Source/Niagara/Private/NiagaraDataInterfaceArrayFunctionLibrary.cpp b/Engine/Plugins/FX/Niagara/Source/Niagara/Private/NiagaraDataInterfaceArrayFunctionLibrary.cpp
--- a/Engine/Plugins/FX/Niagara/Source/Niagara/Private/NiagaraDataInterfaceArrayFunctionLibrary.cpp
+++ b/Engine/Plugins/FX/Niagara/Source/Niagara/Private/NiagaraDataInterfaceArrayFunctionLibrary.cpp
@@ -46,6 +46,13 @@ void SetNiagaraArray(UNiagaraComponent* NiagaraComponent, FName OverrideName, TC
 			ArrayDI->CopyTo(VariantDI);
 			VariantDI->SetVariantArrayData(InArray);
 			NiagaraComponent->SetParameterOverride(FNiagaraVariableBase(FNiagaraTypeDefinition(TDataInterace::StaticClass()), OverrideName), FNiagaraVariant(VariantDI));
+
+			// We reinitialize the system as we need to flush out any queue'ed up GPU ticks as the DI could be GCed.
+			// This would normally occur naturally on the next system tick, but we may not tick the system before rendering
+			if (NiagaraComponent->IsActive())
+			{
+				NiagaraComponent->ReinitializeSystem();
+			}
 		}
 #endif
 	}
@@ -79,10 +86,15 @@ void SetNiagaraArrayValue(UNiagaraComponent* NiagaraComponent, FName OverrideNam
 			ArrayDI->CopyTo(VariantDI);
 			VariantDI->SetVariantArrayValue(Index, Value, bSizeToFit);
 			NiagaraComponent->SetParameterOverride(FNiagaraVariableBase(FNiagaraTypeDefinition(TDataInterace::StaticClass()), OverrideName), FNiagaraVariant(VariantDI));
+
+			// We reinitialize the system as we need to flush out any queue'ed up GPU ticks as the DI could be GCed.
+			// This would normally occur naturally on the next system tick, but we may not tick the system before rendering
+			if (NiagaraComponent->IsActive())
+			{
+				NiagaraComponent->ReinitializeSystem();
+			}
 		}
 #endif
-
-		//-TODO:FIXME:SetNiagaraVariantArray<TArrayType, TDataInterace>(NiagaraComponent, OverrideName, ArrayDI);
 	}
 }
 
