- Fix for decal renderer debug drawing not working

diff --git a/Engine/Plugins/FX/Niagara/Source/Niagara/Private/NiagaraRendererDecals.cpp b/Engine/Plugins/FX/Niagara/Source/Niagara/Private/NiagaraRendererDecals.cpp
--- a/Engine/Plugins/FX/Niagara/Source/Niagara/Private/NiagaraRendererDecals.cpp
+++ b/Engine/Plugins/FX/Niagara/Source/Niagara/Private/NiagaraRendererDecals.cpp
@@ -23,7 +23,9 @@ namespace NiagaraRendererDecalsLocal
 		ECVF_Default
 	);
 
-#if UE_BUILD_SHIPPING
+#define NIAGARARENDERERDECALS_DEBUG_ENABLED !UE_BUILD_SHIPPING
+
+#if NIAGARARENDERERDECALS_DEBUG_ENABLED
 	static bool GDrawDebug = false;
 	static FAutoConsoleVariableRef CVarDrawDebug(
 		TEXT("fx.Niagara.DecalRenderer.DrawDebug"),
@@ -60,8 +62,9 @@ namespace NiagaraRendererDecalsLocal
 
 					const float Fade = (WorldTime - UpdateParams.FadeStartDelay) / UpdateParams.FadeDuration;
 					const FBox BoundsBox(-FVector(1), FVector(1));
-					const uint8 IntFade = FMath::Clamp(uint8(255.0f * Fade), 0, 255);
-					LineBatcher->DrawSolidBox(BoundsBox, UpdateParams.Transform, FColor(IntFade, IntFade, IntFade, 128), 0, 0.0f);
+					//const uint8 IntFade = FMath::Clamp(uint8(255.0f * Fade), 0, 255);
+					const FColor DecalColor = UpdateParams.DecalColor.ToFColor(false);
+					LineBatcher->DrawSolidBox(BoundsBox, UpdateParams.Transform, DecalColor, 0, 0.0f);
 
 					LineBatcher->DrawCircle(UpdateParams.Bounds.Origin, FVector::XAxisVector, FVector::YAxisVector, FColor::Red, UpdateParams.Bounds.SphereRadius, 16, 0);
 					LineBatcher->DrawCircle(UpdateParams.Bounds.Origin, FVector::XAxisVector, FVector::ZAxisVector, FColor::Red, UpdateParams.Bounds.SphereRadius, 16, 0);
@@ -284,7 +287,7 @@ FNiagaraDynamicDataBase* FNiagaraRendererDecals::GenerateDynamicData(const FNiag
 	// Send updates to RT
 	if (DecalUpdates.Num() > 0)
 	{
-	#if UE_BUILD_SHIPPING
+	#if NIAGARARENDERERDECALS_DEBUG_ENABLED
 		DrawDebugDecal(World, DecalUpdates, World->TimeSeconds);
 	#endif
 		World->Scene->BatchUpdateDecals(MoveTemp(DecalUpdates));
