Fixed a bug with async raytrace data interface not issuing traces

diff --git a/Engine/Plugins/FX/Niagara/Shaders/Private/NiagaraDataInterfaceAsyncGpuTrace.ush b/Engine/Plugins/FX/Niagara/Shaders/Private/NiagaraDataInterfaceAsyncGpuTrace.ush
--- a/Engine/Plugins/FX/Niagara/Shaders/Private/NiagaraDataInterfaceAsyncGpuTrace.ush
+++ b/Engine/Plugins/FX/Niagara/Shaders/Private/NiagaraDataInterfaceAsyncGpuTrace.ush
@@ -47,7 +47,7 @@ void NDIAsyncGpuTrace_ReserveRayTraceIndex(uint MaxRayTraceCount, RWBuffer<uint>
 	Out_StartQueryID = -1;
 	Out_IndicesValid = false;
 
-	if ((In_TraceCount + RWRayTraceCounts[RayTraceCountsOffset + 0]) < MaxRayTraceCount)
+	if ((In_TraceCount + RWRayTraceCounts[RayTraceCountsOffset + 0]) <= MaxRayTraceCount)
 	{
 		uint RequestIndex = 0;
 		InterlockedAdd(RWRayTraceCounts[RayTraceCountsOffset + 0], In_TraceCount, RequestIndex);
diff --git a/Engine/Plugins/FX/Niagara/Shaders/Private/NiagaraDataInterfaceAsyncGpuTraceTemplate.ush b/Engine/Plugins/FX/Niagara/Shaders/Private/NiagaraDataInterfaceAsyncGpuTraceTemplate.ush
--- a/Engine/Plugins/FX/Niagara/Shaders/Private/NiagaraDataInterfaceAsyncGpuTraceTemplate.ush
+++ b/Engine/Plugins/FX/Niagara/Shaders/Private/NiagaraDataInterfaceAsyncGpuTraceTemplate.ush
@@ -15,7 +15,7 @@ void IssueAsyncRayTraceGpu_{ParameterName}_UEImpureCall(int In_QueryID, float3 I
 	NDIAsyncGpuTrace_IssueAsyncRayTrace({ParameterName}_MaxRayTraceCount, {ParameterName}_RWRayRequests, {ParameterName}_RayRequestsOffset, {ParameterName}_RWRayTraceCounts, {ParameterName}_RayTraceCountsOffset, {ParameterName}_SystemLWCTile, In_QueryID, In_TraceStart, In_TraceEnd, In_CollisionGroup, Out_IsQueryValid);
 }
 
-void ReserveRayTraceIndex_{ParameterName}_UEImpureCall(int In_TraceCount, out int Out_StartQueryID, out bool Out_IndicesValid)
+void ReserveAsyncRayTraceGpu_{ParameterName}_UEImpureCall(int In_TraceCount, out int Out_StartQueryID, out bool Out_IndicesValid)
 {
 	NDIAsyncGpuTrace_ReserveRayTraceIndex({ParameterName}_MaxRayTraceCount, {ParameterName}_RWRayTraceCounts, {ParameterName}_RayTraceCountsOffset, In_TraceCount, Out_StartQueryID, Out_IndicesValid);
 }
