Added warning to ray trace data interface when the selected trace mode doesn't work with the project settings

diff --git a/Engine/Plugins/FX/Niagara/Source/Niagara/Classes/NiagaraDataInterfaceAsyncGpuTrace.h b/Engine/Plugins/FX/Niagara/Source/Niagara/Classes/NiagaraDataInterfaceAsyncGpuTrace.h
--- a/Engine/Plugins/FX/Niagara/Source/Niagara/Classes/NiagaraDataInterfaceAsyncGpuTrace.h
+++ b/Engine/Plugins/FX/Niagara/Source/Niagara/Classes/NiagaraDataInterfaceAsyncGpuTrace.h
@@ -36,6 +36,7 @@ public:
 
 #if WITH_EDITOR
 	virtual void PostEditChangeProperty(struct FPropertyChangedEvent& PropertyChangedEvent) override;
+	virtual void GetFeedback(UNiagaraSystem* InAsset, UNiagaraComponent* InComponent, TArray<FNiagaraDataInterfaceError>& OutErrors, TArray<FNiagaraDataInterfaceFeedback>& OutWarnings, TArray<FNiagaraDataInterfaceFeedback>& OutInfo) override;
 #endif
 	//UObject Interface End
 
diff --git a/Engine/Plugins/FX/Niagara/Source/Niagara/Private/DataInterface/NiagaraDataInterfaceAsyncGpuTrace.cpp b/Engine/Plugins/FX/Niagara/Source/Niagara/Private/DataInterface/NiagaraDataInterfaceAsyncGpuTrace.cpp
--- a/Engine/Plugins/FX/Niagara/Source/Niagara/Private/DataInterface/NiagaraDataInterfaceAsyncGpuTrace.cpp
+++ b/Engine/Plugins/FX/Niagara/Source/Niagara/Private/DataInterface/NiagaraDataInterfaceAsyncGpuTrace.cpp
@@ -11,11 +11,8 @@
 #include "NiagaraSystem.h"
 #include "NiagaraTypes.h"
 #include "NiagaraSystemInstance.h"
-#include "RenderResource.h"
+#include "SceneManagement.h"
 #include "RenderGraphBuilder.h"
-#include "Shader.h"
-#include "ShaderCore.h"
-#include "ShaderCompilerCore.h"
 
 #include UE_INLINE_GENERATED_CPP_BY_NAME(NiagaraDataInterfaceAsyncGpuTrace)
 
@@ -542,6 +539,25 @@ void UNiagaraDataInterfaceAsyncGpuTrace::PostEditChangeProperty(struct FProperty
 		}
 	}
 }
+
+void UNiagaraDataInterfaceAsyncGpuTrace::GetFeedback(UNiagaraSystem* InAsset, UNiagaraComponent* InComponent, TArray<FNiagaraDataInterfaceError>& OutErrors, TArray<FNiagaraDataInterfaceFeedback>& OutWarnings, TArray<FNiagaraDataInterfaceFeedback>& OutInfo)
+{
+	Super::GetFeedback(InAsset, InComponent, OutErrors, OutWarnings, OutInfo);
+	IConsoleVariable* EnableRayTrace = IConsoleManager::Get().FindConsoleVariable(TEXT("fx.Niagara.AsyncGpuTrace.HWRayTraceEnabled"));
+	IConsoleVariable* EnableDistanceField = IConsoleManager::Get().FindConsoleVariable(TEXT("fx.Niagara.AsyncGpuTrace.GlobalSdfEnabled"));
+	if (TraceProvider == ENDICollisionQuery_AsyncGpuTraceProvider::HWRT && (EnableRayTrace == nullptr || !EnableRayTrace->GetBool() || !IsRayTracingEnabled()))
+	{
+		FText ErrorDescription = LOCTEXT("NiagaraAsyncGpuTraceNoHWRayTraceWarning", "Hardware raytracing is currently disabled");
+		FText ErrorSummary = LOCTEXT("NiagaraAsyncGpuTraceNoHWRayTraceWarningSummary", "The data interface is configured to use hardware raytracing, which is currently disabled.\nThe data interface will NOT work with the current project settings.\n\nTo enable ray tracing support, all of the following must be true:\n* Default RHI is DirectX 12 in the project settings\n* Ray Tracing is enabled in the project settings\n* CVar fx.Niagara.AsyncGpuTrace.HWRayTraceEnabled is set to 1\n* CVar r.Raytracing.Enable is set to 1");
+		OutWarnings.Emplace(ErrorDescription, ErrorSummary, FNiagaraDataInterfaceFix());
+	}
+	if (TraceProvider == ENDICollisionQuery_AsyncGpuTraceProvider::GSDF && (EnableDistanceField == nullptr || !EnableDistanceField->GetBool() || !DoesProjectSupportDistanceFields()))
+	{
+		FText ErrorDescription = LOCTEXT("NiagaraAsyncGpuTraceNoDFTraceWarning", "Distance field tracing is currently disabled");
+		FText ErrorSummary = LOCTEXT("NiagaraAsyncGpuTraceNoDFWarningSummary", "The data interface is configured to use distance fielde tracing, which is currently disabled.\nThe data interface will NOT work with the current project settings.\n\nAside from project settings, check the following cvars:\n* fx.Niagara.AsyncGpuTrace.GlobalSdfEnabled\n* r.GenerateMeshDistanceFields\n* r.DistanceFields.SupportEvenIfHardwareRayTracingSupported (only if ray tracing is disabled)");
+		OutWarnings.Emplace(ErrorDescription, ErrorSummary, FNiagaraDataInterfaceFix());
+	}
+}
 #endif
 
 #undef LOCTEXT_NAMESPACE
