Work around an ICE in Niagara (UE-179196).

diff --git a/Engine/Plugins/FX/Niagara/Source/NiagaraEditor/Private/EdGraphSchema_Niagara.cpp b/Engine/Plugins/FX/Niagara/Source/NiagaraEditor/Private/EdGraphSchema_Niagara.cpp
--- a/Engine/Plugins/FX/Niagara/Source/NiagaraEditor/Private/EdGraphSchema_Niagara.cpp
+++ b/Engine/Plugins/FX/Niagara/Source/NiagaraEditor/Private/EdGraphSchema_Niagara.cpp
@@ -2248,12 +2248,13 @@ void UEdGraphSchema_Niagara::GenerateDataInterfacePinMenu(UToolMenu* ToolMenu, c
 
 				for (const auto& Prototype : FunctionPrototypes)
 				{
+					const FString SaveClipboardText = Prototype.Get<1>();
 					Section.AddMenuEntry(
 						NAME_None,
 						FText::FromString(Prototype.Get<0>()),
 						Prototype.Get<2>(),
 						FSlateIcon(),
-						FUIAction(FExecuteAction::CreateLambda([ClipboardText=Prototype.Get<1>()]() { FPlatformApplicationMisc::ClipboardCopy(*ClipboardText); }))
+						FUIAction(FExecuteAction::CreateLambda([ClipboardText=SaveClipboardText]() { FPlatformApplicationMisc::ClipboardCopy(*ClipboardText); }))
 					);
 				}
 			}
