Added drag/drop support for emitter assets in the system editor

diff --git a/Engine/Plugins/FX/Niagara/Source/NiagaraEditor/Private/EdGraphSchema_NiagaraSystemOverview.cpp b/Engine/Plugins/FX/Niagara/Source/NiagaraEditor/Private/EdGraphSchema_NiagaraSystemOverview.cpp
--- a/Engine/Plugins/FX/Niagara/Source/NiagaraEditor/Private/EdGraphSchema_NiagaraSystemOverview.cpp
+++ b/Engine/Plugins/FX/Niagara/Source/NiagaraEditor/Private/EdGraphSchema_NiagaraSystemOverview.cpp
@@ -2,6 +2,11 @@
 
 #include "EdGraphSchema_NiagaraSystemOverview.h"
 #include "EdGraphNode_Comment.h"
+#include "NiagaraEditorUtilities.h"
+#include "NiagaraEmitter.h"
+#include "NiagaraSystem.h"
+#include "NiagaraSystemEditorData.h"
+#include "ScopedTransaction.h"
 
 #include UE_INLINE_GENERATED_CPP_BY_NAME(EdGraphSchema_NiagaraSystemOverview)
 
@@ -34,5 +39,84 @@ void UEdGraphSchema_NiagaraSystemOverview::ForceVisualizationCacheClear() const
 	CurrentCacheRefreshID++;
 }
 
+void UEdGraphSchema_NiagaraSystemOverview::DroppedAssetsOnGraph(const TArray<FAssetData>& Assets, const FVector2D& GraphPosition, UEdGraph* Graph) const
+{
+	TArray<UNiagaraEmitter*> Emitters;
+
+	for (const FAssetData& Data : Assets)
+	{
+		UObject* Asset = Data.GetAsset();
+		if (UNiagaraEmitter* Emitter = Cast<UNiagaraEmitter>(Asset))
+		{
+			Emitters.Add(Emitter);
+		}
+	}
+
+	UNiagaraSystem* System = Graph->GetTypedOuter<UNiagaraSystem>();
+	if (Emitters.Num() > 0 && System)
+	{
+		UNiagaraSystemEditorData* SystemEditorData = CastChecked<UNiagaraSystemEditorData>(System->GetEditorData(), ECastCheckedType::NullChecked);
+		// only insert emitters if we're in a system graph and not another emitter
+		if (SystemEditorData->GetOwningSystemIsPlaceholder() == false)
+		{
+			FScopedTransaction AddEmitterTransaction(LOCTEXT("NiagaraEditorDropEmitter", "Niagara: Drag and Drop Emitter"));
+
+			FVector2D NewNodePosition = GraphPosition;
+			for (UNiagaraEmitter* Emitter : Emitters)
+			{
+				FGuid Handle = FNiagaraEditorUtilities::AddEmitterToSystem(*System, *Emitter, Emitter->GetExposedVersion().VersionGuid);
+
+				// place the node at the mouse position from the drag operation
+				TArray<UNiagaraOverviewNode*> OverviewNodes;
+				SystemEditorData->GetSystemOverviewGraph()->GetNodesOfClass<UNiagaraOverviewNode>(OverviewNodes);
+				UNiagaraOverviewNode* OverviewNode = OverviewNodes.Last();
+				if (OverviewNode->GetEmitterHandleGuid() == Handle)
+				{
+					OverviewNode->NodePosX = NewNodePosition.X;
+					OverviewNode->NodePosY = NewNodePosition.Y;
+					NewNodePosition.X += 250;
+				}
+			}
+			System->OnSystemPostEditChange().Broadcast(System);
+		}
+	}
+}
+
+void UEdGraphSchema_NiagaraSystemOverview::GetAssetsGraphHoverMessage(const TArray<FAssetData>& Assets, const UEdGraph* HoverGraph, FString& OutTooltipText, bool& OutOkIcon) const
+{
+	OutOkIcon = false;
+
+	UNiagaraSystem* System = HoverGraph->GetTypedOuter<UNiagaraSystem>();
+	if (!System)
+	{
+		return;
+	}
+	UNiagaraSystemEditorData* SystemEditorData = CastChecked<UNiagaraSystemEditorData>(System->GetEditorData(), ECastCheckedType::NullChecked);
+	if (SystemEditorData->GetOwningSystemIsPlaceholder())
+	{
+		for (const FAssetData& AssetData : Assets)
+		{
+			if (UNiagaraEmitter* Emitter = Cast<UNiagaraEmitter>(AssetData.GetAsset()))
+			{
+				FText Tooltip = LOCTEXT("NiagaraEditorDropEmitterInEmitterTooltip", "Emitters can only be added to systems");
+				OutTooltipText = Tooltip.ToString();
+				break;
+			}
+		}
+		return;
+	}
+	
+	for (const FAssetData& AssetData : Assets)
+	{
+		if (UNiagaraEmitter* Emitter = Cast<UNiagaraEmitter>(AssetData.GetAsset()))
+		{
+			FText Tooltip = LOCTEXT("NiagaraEditorDropEmitterTooltip", "Add emitter to system");
+			OutTooltipText = Tooltip.ToString();
+			OutOkIcon = true;
+			break;
+		}
+	}
+}
+
 #undef LOCTEXT_NAMESPACE
 
diff --git a/Engine/Plugins/FX/Niagara/Source/NiagaraEditor/Public/EdGraphSchema_NiagaraSystemOverview.h b/Engine/Plugins/FX/Niagara/Source/NiagaraEditor/Public/EdGraphSchema_NiagaraSystemOverview.h
--- a/Engine/Plugins/FX/Niagara/Source/NiagaraEditor/Public/EdGraphSchema_NiagaraSystemOverview.h
+++ b/Engine/Plugins/FX/Niagara/Source/NiagaraEditor/Public/EdGraphSchema_NiagaraSystemOverview.h
@@ -22,6 +22,8 @@ class NIAGARAEDITOR_API UEdGraphSchema_NiagaraSystemOverview : public UEdGraphSc
 	virtual bool IsCacheVisualizationOutOfDate(int32 InVisualizationCacheID) const override;
 	virtual int32 GetCurrentVisualizationCacheID() const override;
 	virtual void ForceVisualizationCacheClear() const override;
+	virtual void DroppedAssetsOnGraph(const TArray<FAssetData>& Assets, const FVector2D& GraphPosition, UEdGraph* Graph) const override;
+	virtual void GetAssetsGraphHoverMessage(const TArray<FAssetData>& Assets, const UEdGraph* HoverGraph, FString& OutTooltipText, bool& OutOkIcon) const override;
 
 private:
 	static int32 CurrentCacheRefreshID;
