Added decal icon for decal renderer

diff --git a/Engine/Plugins/FX/Niagara/Source/Niagara/Private/NiagaraDecalRendererProperties.cpp b/Engine/Plugins/FX/Niagara/Source/Niagara/Private/NiagaraDecalRendererProperties.cpp
--- a/Engine/Plugins/FX/Niagara/Source/Niagara/Private/NiagaraDecalRendererProperties.cpp
+++ b/Engine/Plugins/FX/Niagara/Source/Niagara/Private/NiagaraDecalRendererProperties.cpp
@@ -5,6 +5,7 @@
 #include "NiagaraEmitterInstance.h"
 #include "NiagaraRenderer.h"
 #include "NiagaraRendererDecals.h"
+#include "Components/DecalComponent.h"
 #include "Materials/MaterialInstanceConstant.h"
 #include "Modules/ModuleManager.h"
 
@@ -14,7 +15,6 @@
 #include "Widgets/Images/SImage.h"
 #include "Widgets/Text/STextBlock.h"
 #include "Widgets/SWidget.h"
-#include "Styling/SlateBrush.h"
 #include "Styling/SlateIconFinder.h"
 #include "AssetThumbnail.h"
 #endif
@@ -304,7 +304,13 @@ const TArray<FNiagaraVariable>& UNiagaraDecalRendererProperties::GetOptionalAttr
 
 void UNiagaraDecalRendererProperties::GetRendererWidgets(const FNiagaraEmitterInstance* InEmitter, TArray<TSharedPtr<SWidget>>& OutWidgets, TSharedPtr<FAssetThumbnailPool> InThumbnailPool) const
 {
-	OutWidgets.Add(SNew(SImage).Image(FSlateIconFinder::FindIconBrushForClass(GetClass())));
+	TSharedRef<SWidget> Widget = SNew(SImage).Image(GetStackIcon());
+	OutWidgets.Add(Widget);
+}
+
+const FSlateBrush* UNiagaraDecalRendererProperties::GetStackIcon() const
+{
+	return FSlateIconFinder::FindIconBrushForClass(UDecalComponent::StaticClass());
 }
 
 void UNiagaraDecalRendererProperties::GetRendererTooltipWidgets(const FNiagaraEmitterInstance* InEmitter, TArray<TSharedPtr<SWidget>>& OutWidgets, TSharedPtr<FAssetThumbnailPool> InThumbnailPool) const
diff --git a/Engine/Plugins/FX/Niagara/Source/Niagara/Public/NiagaraDecalRendererProperties.h b/Engine/Plugins/FX/Niagara/Source/Niagara/Public/NiagaraDecalRendererProperties.h
--- a/Engine/Plugins/FX/Niagara/Source/Niagara/Public/NiagaraDecalRendererProperties.h
+++ b/Engine/Plugins/FX/Niagara/Source/Niagara/Public/NiagaraDecalRendererProperties.h
@@ -42,9 +42,10 @@ public:
 #if WITH_EDITORONLY_DATA
 	virtual const TArray<FNiagaraVariable>& GetOptionalAttributes() override;
 	virtual void GetRendererWidgets(const FNiagaraEmitterInstance* InEmitter, TArray<TSharedPtr<SWidget>>& OutWidgets, TSharedPtr<FAssetThumbnailPool> InThumbnailPool) const override;
+	virtual const FSlateBrush* GetStackIcon() const override;
 	virtual void GetRendererTooltipWidgets(const FNiagaraEmitterInstance* InEmitter, TArray<TSharedPtr<SWidget>>& OutWidgets, TSharedPtr<FAssetThumbnailPool> InThumbnailPool) const override;
 	virtual void GetRendererFeedback(const FVersionedNiagaraEmitter& InEmitter, TArray<FNiagaraRendererFeedback>& OutErrors, TArray<FNiagaraRendererFeedback>& OutWarnings, TArray<FNiagaraRendererFeedback>& OutInfo) const override;
-	TArray<FNiagaraVariable> GetBoundAttributes() const override;
+	virtual TArray<FNiagaraVariable> GetBoundAttributes() const override;
 	virtual void RenameVariable(const FNiagaraVariableBase& OldVariable, const FNiagaraVariableBase& NewVariable, const FVersionedNiagaraEmitter& InEmitter) override;
 	virtual void RemoveVariable(const FNiagaraVariableBase& OldVariable, const FVersionedNiagaraEmitter& InEmitter) override;
 #endif // WITH_EDITORONLY_DATA
@@ -52,7 +53,7 @@ public:
 	virtual void UpdateSourceModeDerivates(ENiagaraRendererSourceDataMode InSourceMode, bool bFromPropertyEdit = false) override;
 	virtual ENiagaraRendererSourceDataMode GetCurrentSourceMode() const override { return SourceMode; }
 	virtual bool PopulateRequiredBindings(FNiagaraParameterStore& InParameterStore) override;
-	virtual bool NeedsMIDsForMaterials() const { return MaterialParameters.HasAnyBindings(); }
+	virtual bool NeedsMIDsForMaterials() const override { return MaterialParameters.HasAnyBindings(); }
 	//UNiagaraRendererProperties Interface END
 
 	void UpdateMICs();
