Fixed editor only parameters not being marked as RF_Transactional.
This fixes issues with transactions not properly being recorded in Niagara

diff --git a/Engine/Plugins/FX/Niagara/Source/NiagaraEditor/Private/NiagaraEditorData.cpp b/Engine/Plugins/FX/Niagara/Source/NiagaraEditor/Private/NiagaraEditorData.cpp
--- a/Engine/Plugins/FX/Niagara/Source/NiagaraEditor/Private/NiagaraEditorData.cpp
+++ b/Engine/Plugins/FX/Niagara/Source/NiagaraEditor/Private/NiagaraEditorData.cpp
@@ -15,6 +15,8 @@ void UNiagaraEditorParametersAdapter::PostLoad()
 {
 	Super::PostLoad();
 
+	SetFlags(RF_Transactional);
+	
 	// Remove null entries from EditorOnlyScriptVars predating null checks on add methods.
 	for (int32 Idx = EditorOnlyScriptVars.Num() - 1; Idx > -1; --Idx)
 	{
diff --git a/Engine/Plugins/FX/Niagara/Source/NiagaraEditor/Private/NiagaraEditorModule.cpp b/Engine/Plugins/FX/Niagara/Source/NiagaraEditor/Private/NiagaraEditorModule.cpp
--- a/Engine/Plugins/FX/Niagara/Source/NiagaraEditor/Private/NiagaraEditorModule.cpp
+++ b/Engine/Plugins/FX/Niagara/Source/NiagaraEditor/Private/NiagaraEditorModule.cpp
@@ -211,6 +211,7 @@ class FNiagaraEditorOnlyDataUtilities : public INiagaraEditorOnlyDataUtilities
 		if (System != nullptr)
 		{
 			UNiagaraSystemEditorData* SystemEditorData = NewObject<UNiagaraSystemEditorData>(InOuter);
+			SystemEditorData->SetFlags(RF_Transactional);
 			SystemEditorData->SynchronizeOverviewGraphWithSystem(*System);
 			SystemEditorData->InitOnSyncScriptVariables(*System);
 			return SystemEditorData;
@@ -220,7 +221,10 @@ class FNiagaraEditorOnlyDataUtilities : public INiagaraEditorOnlyDataUtilities
 
 	virtual UNiagaraEditorParametersAdapterBase* CreateDefaultEditorParameters(UObject* InOuter) const override
 	{
-		return NewObject<UNiagaraEditorParametersAdapter>(InOuter);
+		UNiagaraEditorParametersAdapter* Adapter = NewObject<UNiagaraEditorParametersAdapter>(InOuter);
+		Adapter->SetFlags(RF_Transactional);
+		
+		return Adapter;
 	}
 
 	virtual UObject::FAssetRegistryTag CreateClassUsageAssetRegistryTag(const UObject* SourceObject) const override
