Changing a few last bits of code to use batched parameters.

diff --git a/Engine/Plugins/FX/Niagara/Source/Niagara/Private/NiagaraGpuComputeDispatch.cpp b/Engine/Plugins/FX/Niagara/Source/Niagara/Private/NiagaraGpuComputeDispatch.cpp
--- a/Engine/Plugins/FX/Niagara/Source/Niagara/Private/NiagaraGpuComputeDispatch.cpp
+++ b/Engine/Plugins/FX/Niagara/Source/Niagara/Private/NiagaraGpuComputeDispatch.cpp
@@ -1618,14 +1618,16 @@ void FNiagaraGpuComputeDispatch::DispatchStage(FRDGBuilder& GraphBuilder, const
 				FRHIComputeShader* RHIComputeShader = ComputeShader.GetComputeShader();
 				SetComputePipelineState(RHICmdList, RHIComputeShader);
 
+				FRHIBatchedShaderParameters& BatchedParameters = RHICmdList.GetScratchShaderParameters();
+
 				if (ComputeShader->ExternalConstantBufferParam[0].IsBound())
 				{
-					RHICmdList.SetShaderUniformBuffer(RHIComputeShader, ComputeShader->ExternalConstantBufferParam[0].GetBaseIndex(), TickPtr->GetExternalUniformBuffer(*InstanceDataPtr, false));
+					BatchedParameters.SetShaderUniformBuffer(ComputeShader->ExternalConstantBufferParam[0].GetBaseIndex(), TickPtr->GetExternalUniformBuffer(*InstanceDataPtr, false));
 				}
 				if (ComputeShader->ExternalConstantBufferParam[1].IsBound())
 				{
 					check(InstanceDataPtr->Context->HasInterpolationParameters);
-					RHICmdList.SetShaderUniformBuffer(RHIComputeShader, ComputeShader->ExternalConstantBufferParam[1].GetBaseIndex(), TickPtr->GetExternalUniformBuffer(*InstanceDataPtr, true));
+					BatchedParameters.SetShaderUniformBuffer(ComputeShader->ExternalConstantBufferParam[1].GetBaseIndex(), TickPtr->GetExternalUniformBuffer(*InstanceDataPtr, true));
 				}
 
 				FNiagaraEmptyUAVPoolScopedAccess UAVPoolAccessScope(GetEmptyUAVPool());
@@ -1633,7 +1635,10 @@ void FNiagaraGpuComputeDispatch::DispatchStage(FRDGBuilder& GraphBuilder, const
 				FNiagaraGpuProfileScope GpuProfileDispatchScope(RHICmdList, this, FNiagaraGpuProfileEvent(*InstanceDataPtr, *SimStageDataPtr, InstanceDataPtr == &TickPtr->GetInstances()[0]));
 
 				FComputeShaderUtils::ValidateIndirectArgsBuffer(SimStageDataPtr->DispatchArgs.IndirectBuffer->GetSize(), SimStageDataPtr->DispatchArgs.IndirectOffset);
-				SetShaderParameters<FRHICommandList, FNiagaraShader>(RHICmdList, ComputeShader, RHIComputeShader, ShaderParametersMetadata, *DispatchParameters);
+				SetShaderParameters(BatchedParameters, ComputeShader, ShaderParametersMetadata, *DispatchParameters);
+
+				RHICmdList.SetBatchedShaderParameters(RHIComputeShader, BatchedParameters);
+
 				RHICmdList.DispatchIndirectComputeShader(SimStageDataPtr->DispatchArgs.IndirectBuffer->GetIndirectRHICallBuffer(), SimStageDataPtr->DispatchArgs.IndirectOffset);
 				UnsetShaderUAVs<FRHICommandList, FNiagaraShader>(RHICmdList, ComputeShader, RHIComputeShader);
 			}
@@ -1706,21 +1711,26 @@ void FNiagaraGpuComputeDispatch::DispatchStage(FRDGBuilder& GraphBuilder, const
 				FRHIComputeShader* RHIComputeShader = ComputeShader.GetComputeShader();
 				SetComputePipelineState(RHICmdList, RHIComputeShader);
 
+				FRHIBatchedShaderParameters& BatchedParameters = RHICmdList.GetScratchShaderParameters();
+
 				if (ComputeShader->ExternalConstantBufferParam[0].IsBound())
 				{
-					RHICmdList.SetShaderUniformBuffer(RHIComputeShader, ComputeShader->ExternalConstantBufferParam[0].GetBaseIndex(), TickPtr->GetExternalUniformBuffer(*InstanceDataPtr, false));
+					BatchedParameters.SetShaderUniformBuffer(ComputeShader->ExternalConstantBufferParam[0].GetBaseIndex(), TickPtr->GetExternalUniformBuffer(*InstanceDataPtr, false));
 				}
 				if (ComputeShader->ExternalConstantBufferParam[1].IsBound())
 				{
 					check(InstanceDataPtr->Context->HasInterpolationParameters);
-					RHICmdList.SetShaderUniformBuffer(RHIComputeShader, ComputeShader->ExternalConstantBufferParam[1].GetBaseIndex(), TickPtr->GetExternalUniformBuffer(*InstanceDataPtr, true));
+					BatchedParameters.SetShaderUniformBuffer(ComputeShader->ExternalConstantBufferParam[1].GetBaseIndex(), TickPtr->GetExternalUniformBuffer(*InstanceDataPtr, true));
 				}
 
 				FNiagaraEmptyUAVPoolScopedAccess UAVPoolAccessScope(GetEmptyUAVPool());
 
 				FNiagaraGpuProfileScope GpuProfileDispatchScope(RHICmdList, this, FNiagaraGpuProfileEvent(*InstanceDataPtr, *SimStageDataPtr, InstanceDataPtr == &TickPtr->GetInstances()[0]));
 
-				SetShaderParameters<FRHICommandList, FNiagaraShader>(RHICmdList, ComputeShader, RHIComputeShader, ShaderParametersMetadata, *DispatchParameters);
+				SetShaderParameters(BatchedParameters, ComputeShader, ShaderParametersMetadata, *DispatchParameters);
+
+				RHICmdList.SetBatchedShaderParameters(RHIComputeShader, BatchedParameters);
+
 				DispatchComputeShader(RHICmdList, ComputeShader, ThreadGroupCount.X, ThreadGroupCount.Y, ThreadGroupCount.Z);
 				UnsetShaderUAVs<FRHICommandList, FNiagaraShader>(RHICmdList, ComputeShader, RHIComputeShader);
 			}
diff --git a/Engine/Source/Runtime/Renderer/Private/RendererUtils.cpp b/Engine/Source/Runtime/Renderer/Private/RendererUtils.cpp
--- a/Engine/Source/Runtime/Renderer/Private/RendererUtils.cpp
+++ b/Engine/Source/Runtime/Renderer/Private/RendererUtils.cpp
@@ -60,9 +60,9 @@ public:
 	}
 
 	// Shader parameter structs don't have a way to push variable sized data yet. So the we use the old shader parameter API.
-	void SetPlatformData(FRHIComputeCommandList& RHICmdList, const void* PlatformDataPtr, uint32 PlatformDataSize)
+	void SetParameters(FRHIBatchedShaderParameters& BatchedParameters, const void* PlatformDataPtr, uint32 PlatformDataSize)
 	{
-		RHICmdList.SetShaderParameter(RHICmdList.GetBoundComputeShader(), PlatformDataParam.GetBufferIndex(), PlatformDataParam.GetBaseIndex(), PlatformDataSize, PlatformDataPtr);
+		BatchedParameters.SetShaderParameter(PlatformDataParam.GetBufferIndex(), PlatformDataParam.GetBaseIndex(), PlatformDataSize, PlatformDataPtr);
 	}
 
 private:
@@ -158,8 +158,8 @@ void FRenderTargetWriteMask::Decode(
 		}
 
 		SetComputePipelineState(RHICmdList, DecodeCS.GetComputeShader());
-		SetShaderParameters(RHICmdList, DecodeCS, DecodeCS.GetComputeShader(), *PassParameters);
-		DecodeCS->SetPlatformData(RHICmdList, PlatformDataPtr, PlatformDataSize);
+
+		SetShaderParametersMixedCS(RHICmdList, DecodeCS, *PassParameters, PlatformDataPtr, PlatformDataSize);
 
 		RHICmdList.DispatchComputeShader(
 			FMath::DivideAndRoundUp((uint32)RTWriteMaskDims.X, FRTWriteMaskDecodeCS::ThreadGroupSizeX),
diff --git a/Engine/Source/Runtime/Renderer/Private/SceneRendering.cpp b/Engine/Source/Runtime/Renderer/Private/SceneRendering.cpp
--- a/Engine/Source/Runtime/Renderer/Private/SceneRendering.cpp
+++ b/Engine/Source/Runtime/Renderer/Private/SceneRendering.cpp
@@ -5544,7 +5544,10 @@ void AddResolveSceneDepthPass(FRDGBuilder& GraphBuilder, const FViewInfo& View,
 
 		SetGraphicsPipelineState(RHICmdList, GraphicsPSOInit, 0);
 		RHICmdList.SetBlendFactor(FLinearColor::White);
-		RHICmdList.SetShaderTexture(ResolvePixelShader, TextureIndex, SourceTextureRHI);
+
+		FRHIBatchedShaderParameters& BatchedParameters = RHICmdList.GetScratchShaderParameters();
+		BatchedParameters.SetShaderTexture(TextureIndex, SourceTextureRHI);
+		RHICmdList.SetBatchedShaderParameters(ResolvePixelShader, BatchedParameters);
 
 		SetShaderParametersLegacyVS(RHICmdList, ResolveVertexShader, ResolveRect, ResolveRect, DepthExtent.X, DepthExtent.Y);
 
