- Give generated MICs a name based off the parent material
- This makes them easier to track in profiling tools

#rnx

diff --git a/Engine/Plugins/FX/Niagara/Source/Niagara/Private/NiagaraRendererProperties.cpp b/Engine/Plugins/FX/Niagara/Source/Niagara/Private/NiagaraRendererProperties.cpp
--- a/Engine/Plugins/FX/Niagara/Source/Niagara/Private/NiagaraRendererProperties.cpp
+++ b/Engine/Plugins/FX/Niagara/Source/Niagara/Private/NiagaraRendererProperties.cpp
@@ -490,13 +490,17 @@ void UNiagaraRendererProperties::UpdateMaterialParametersMIC(const FNiagaraRende
 		return;
 	}
 
+	FNameBuilder NameBuilder;
+	InOutMaterial->GetFName().ToString(NameBuilder);
+	NameBuilder.Append(TEXT("_MIC"));
 	if (InOutMIC == nullptr)
 	{
-		InOutMIC = NewObject<UMaterialInstanceConstant>(this);
+		InOutMIC = NewObject<UMaterialInstanceConstant>(this, FName(NameBuilder));
 		InOutMIC->SetParentEditorOnly(InOutMaterial);
 	}
 	else if (InOutMIC->Parent != InOutMaterial)
 	{
+		InOutMIC->Rename(NameBuilder.ToString());
 		InOutMIC->SetParentEditorOnly(InOutMaterial);
 	}
 
@@ -536,18 +540,23 @@ void UNiagaraRendererProperties::UpdateMaterialParametersMIC(const FNiagaraRende
 		}
 
 		//-OPT: We should be able to reuse rather than create
+		FNameBuilder NameBuilder;
+		Material->GetFName().ToString(NameBuilder);
+		NameBuilder.Append(TEXT("_MIC"));
+
 		UMaterialInstanceConstant* MIC = nullptr;
 		if (MICPool.Num() > 0)
 		{
 			MIC = MICPool.Pop();
 			if (MIC->Parent != Material)
 			{
+				MIC->Rename(NameBuilder.ToString());
 				MIC->SetParentEditorOnly(Material);
 			}
 		}
 		else
 		{
-			MIC = NewObject<UMaterialInstanceConstant>(this);
+			MIC = NewObject<UMaterialInstanceConstant>(this, FName(NameBuilder));
 			MIC->SetParentEditorOnly(Material);
 		}
 
