Removal of incorrect default pins that would show up as 'None'

diff --git a/Engine/Plugins/FX/Niagara/Source/NiagaraEditor/Private/NiagaraNodeParameterMapGet.cpp b/Engine/Plugins/FX/Niagara/Source/NiagaraEditor/Private/NiagaraNodeParameterMapGet.cpp
--- a/Engine/Plugins/FX/Niagara/Source/NiagaraEditor/Private/NiagaraNodeParameterMapGet.cpp
+++ b/Engine/Plugins/FX/Niagara/Source/NiagaraEditor/Private/NiagaraNodeParameterMapGet.cpp
@@ -356,6 +356,11 @@ void UNiagaraNodeParameterMapGet::PostLoad()
 	Super::PostLoad();
 	FPinCollectorArray OutputPins;
 	GetOutputPins(OutputPins);
+
+	FPinCollectorArray InputPins;
+	GetInputPins(InputPins);
+
+	TSet<UEdGraphPin*> CorrectParameterDefaultPins;
 	for (int32 i = 0; i < OutputPins.Num(); i++)
 	{
 		UEdGraphPin* OutputPin = OutputPins[i];
@@ -367,7 +372,7 @@ void UNiagaraNodeParameterMapGet::PostLoad()
 		UEdGraphPin* DefaultPin = GetDefaultPin(OutputPin);
 		if (DefaultPin == nullptr)
 		{
-			CreateDefaultPin(OutputPin);
+			DefaultPin = CreateDefaultPin(OutputPin);
 		}
 		else
 		{
@@ -375,6 +380,22 @@ void UNiagaraNodeParameterMapGet::PostLoad()
 		}
 
 		OutputPin->PinType.PinSubCategory = UNiagaraNodeParameterMapBase::ParameterPinSubCategory;
+		
+		CorrectParameterDefaultPins.Add(DefaultPin);
+	}
+
+	// we remove default input pins that aren't supposed to be there.
+	for(UEdGraphPin* InputPin : InputPins)
+	{
+		if(IsParameterMapPin(InputPin) || InputPin->bOrphanedPin)
+		{
+			continue;
+		}
+		
+		if(!CorrectParameterDefaultPins.Contains(InputPin))
+		{
+			RemovePin(InputPin);
+		}
 	}
 }
 
