Fixed a bug where sim caching can lead to minute-long hitches with large number of particles

diff --git a/Engine/Plugins/FX/Niagara/Source/Niagara/Private/NiagaraSimCache.cpp b/Engine/Plugins/FX/Niagara/Source/Niagara/Private/NiagaraSimCache.cpp
--- a/Engine/Plugins/FX/Niagara/Source/Niagara/Private/NiagaraSimCache.cpp
+++ b/Engine/Plugins/FX/Niagara/Source/Niagara/Private/NiagaraSimCache.cpp
@@ -430,9 +430,17 @@ bool UNiagaraSimCache::EndWrite()
 					return;
 				}
 
+				TMap<int32, int32> BufferIndexMap;
+				BufferIndexMap.Reserve(BufferB.InterpMapping.Num());
+				for (int32 i = 0; i < BufferB.InterpMapping.Num(); ++i)
+				{
+					BufferIndexMap.Add(BufferB.InterpMapping[i], i);
+				}
+				
 				for (int32 i = 0; i < BufferA.InterpMapping.Num(); ++i)
 				{
-					BufferA.InterpMapping[i] = BufferB.InterpMapping.IndexOfByKey(BufferA.InterpMapping[i]);
+					int32* Mapping = BufferIndexMap.Find(BufferA.InterpMapping[i]);
+					BufferA.InterpMapping[i] = Mapping ? *Mapping : INDEX_NONE;
 				}
 			};
 
