- Modify Particle GPU Read fail message as it can mean either not found or not GPU sim

#rnx

diff --git a/Engine/Plugins/FX/Niagara/Source/Niagara/Private/NiagaraDataInterfaceParticleRead.cpp b/Engine/Plugins/FX/Niagara/Source/Niagara/Private/NiagaraDataInterfaceParticleRead.cpp
--- a/Engine/Plugins/FX/Niagara/Source/Niagara/Private/NiagaraDataInterfaceParticleRead.cpp
+++ b/Engine/Plugins/FX/Niagara/Source/Niagara/Private/NiagaraDataInterfaceParticleRead.cpp
@@ -2150,10 +2150,9 @@ void UNiagaraDataInterfaceParticleRead::SetShaderParameters(const FNiagaraDataIn
 
 	if (InstanceData->SourceEmitterGPUContext == nullptr)
 	{
-		// This means the source emitter isn't running on GPU.
 		if (!InstanceData->bSourceEmitterNotGPUErrorShown)
 		{
-			UE_LOG(LogNiagara, Error, TEXT("GPU particle read DI is set to access CPU emitter '%s'."), *InstanceData->DebugSourceName);
+			UE_LOG(LogNiagara, Error, TEXT("GPU particle read could not find GPU context for '%s', emitter either doesn't exist or is not a GPU simulation."), *InstanceData->DebugSourceName);
 			InstanceData->bSourceEmitterNotGPUErrorShown = true;
 		}
 		SetErrorParams(true);
