- WIP: Niagara Volume Renderer

#rnx

diff --git a/Engine/Plugins/FX/Niagara/Source/Niagara/Private/NiagaraModule.cpp b/Engine/Plugins/FX/Niagara/Source/Niagara/Private/NiagaraModule.cpp
--- a/Engine/Plugins/FX/Niagara/Source/Niagara/Private/NiagaraModule.cpp
+++ b/Engine/Plugins/FX/Niagara/Source/Niagara/Private/NiagaraModule.cpp
@@ -29,6 +29,7 @@
 #include "NiagaraMeshRendererProperties.h"
 #include "NiagaraRibbonRendererProperties.h"
 #include "NiagaraComponentRendererProperties.h"
+#include "NiagaraVolumeRendererProperties.h"
 #include "NiagaraCustomVersion.h"
 #include "NiagaraShaderModule.h"
 #include "UObject/CoreRedirects.h"
@@ -424,6 +425,7 @@ void INiagaraModule::StartupModule()
 	UNiagaraRibbonRendererProperties::InitCDOPropertiesAfterModuleStartup();
 	UNiagaraMeshRendererProperties::InitCDOPropertiesAfterModuleStartup();
 	UNiagaraComponentRendererProperties::InitCDOPropertiesAfterModuleStartup();
+	UNiagaraVolumeRendererProperties::InitCDOPropertiesAfterModuleStartup();
 
 	// Register the data interface CDO finder with the shader module..
 	INiagaraShaderModule& NiagaraShaderModule = FModuleManager::LoadModuleChecked<INiagaraShaderModule>("NiagaraShader");
diff --git a/Engine/Plugins/FX/Niagara/Source/Niagara/Private/NiagaraRendererVolumes.cpp b/Engine/Plugins/FX/Niagara/Source/Niagara/Private/NiagaraRendererVolumes.cpp
new file mode 100644
--- /dev/null
+++ b/Engine/Plugins/FX/Niagara/Source/Niagara/Private/NiagaraRendererVolumes.cpp
@@ -0,0 +1,337 @@
+// Copyright Epic Games, Inc. All Rights Reserved.
+
+#include "NiagaraRendererVolumes.h"
+
+#include "NiagaraCullProxyComponent.h"
+#include "NiagaraVolumeRendererProperties.h"
+#include "NiagaraEmitter.h"
+#include "NiagaraSceneProxy.h"
+#include "NiagaraSystemInstance.h"
+
+#include "Async/Async.h"
+#include "Components/LineBatchComponent.h"
+#include "Engine/World.h"
+#include "SceneInterface.h"
+
+namespace NiagaraRendererVolumesLocal
+{
+	static bool GRendererEnabled = true;
+	static FAutoConsoleVariableRef CVarRendererEnabled(
+		TEXT("fx.Niagara.VolumeRenderer.Enabled"),
+		GRendererEnabled,
+		TEXT("If == 0, Niagara Volume Renderers are disabled."),
+		ECVF_Default
+	);
+
+	struct FVolumeDynamicData : public FNiagaraDynamicDataBase
+	{
+		FVolumeDynamicData(const FNiagaraEmitterInstance* InEmitter)
+			: FNiagaraDynamicDataBase(InEmitter)
+		{
+		}
+
+		virtual void ApplyMaterialOverride(int32 MaterialIndex, UMaterialInterface* MaterialOverride) override
+		{
+			if (MaterialIndex == 0 && MaterialOverride)
+			{
+				Material = MaterialOverride->GetRenderProxy();
+			}
+		}
+
+		FMaterialRenderProxy*	Material = nullptr;
+		FNiagaraPosition		DefaultPosition = FVector3f::ZeroVector;
+		FQuat4f					DefaultRotation = UNiagaraVolumeRendererProperties::GetDefaultVolumeRotation();
+		FVector3f				DefaultScale = UNiagaraVolumeRendererProperties::GetDefaultVolumeScale();
+		int32					DefaultRendererVisibilityTag = 0;
+		int32					DefaultVolumeResolutionMaxAxis = UNiagaraVolumeRendererProperties::GetDefaultVolumeResolutionMaxAxis();
+		FVector3f				DefaultVolumeWorldSpaceSize = UNiagaraVolumeRendererProperties::GetDefaultVolumeWorldSpaceSize();
+	};
+
+	struct FHeterogeneousVolumeInstances
+	{
+		struct FInstance
+		{
+			FBoxSphereBounds	LocalBounds;
+			FMatrix				LocalToWorld;
+		};
+
+		TArray<FInstance>	Instances;
+	};
+}
+
+FNiagaraRendererVolumes::FNiagaraRendererVolumes(ERHIFeatureLevel::Type FeatureLevel, const UNiagaraRendererProperties* InProperties, const FNiagaraEmitterInstance* Emitter)
+	: FNiagaraRenderer(FeatureLevel, InProperties, Emitter)
+	, VertexFactory(FeatureLevel, "FNiagaraRendererVolumes")
+{
+	const UNiagaraVolumeRendererProperties* Properties = CastChecked<const UNiagaraVolumeRendererProperties>(InProperties);
+	check(Properties);
+
+	SourceMode						= Properties->SourceMode;
+	RendererVisibilityTag			= Properties->RendererVisibility;
+	PositionDataSetAccessor			= Properties->PositionDataSetAccessor;
+	RotationDataSetAccessor			= Properties->RotationDataSetAccessor;
+	ScaleDataSetAccessor			= Properties->ScaleDataSetAccessor;
+	RendererVisibilityTagAccessor	= Properties->RendererVisibilityTagAccessor;
+	VolumeResolutionMaxAxisAccessor = Properties->VolumeResolutionMaxAxisAccessor;
+	VolumeWorldSpaceSizeAccessor	= Properties->VolumeWorldSpaceSizeAccessor;
+
+	// Generate Attribute -> Renderer Parameter Store Bindings
+	const TArray<const FNiagaraVariableAttributeBinding*>& VFAttributeBindings = Properties->GetAttributeBindings();
+
+	check(UE_ARRAY_COUNT(VFBoundOffsetsInParamStore) == VFAttributeBindings.Num());
+	for ( int32 i=0; i < UE_ARRAY_COUNT(VFBoundOffsetsInParamStore); ++i )
+	{
+		VFBoundOffsetsInParamStore[i] = INDEX_NONE;
+	
+		if (VFAttributeBindings[i] && VFAttributeBindings[i]->CanBindToHostParameterMap())
+		{
+			VFBoundOffsetsInParamStore[i] = Emitter->GetRendererBoundVariables().IndexOf(VFAttributeBindings[i]->GetParamMapBindableVariable());
+			bAnyVFBoundOffsets |= VFBoundOffsetsInParamStore[i] != INDEX_NONE;
+		}
+	}
+}
+
+FNiagaraRendererVolumes::~FNiagaraRendererVolumes()
+{
+}
+
+void FNiagaraRendererVolumes::CreateRenderThreadResources()
+{
+	FNiagaraRenderer::CreateRenderThreadResources();
+
+	FLocalVertexFactory::FDataType VFData;
+	VFData.PositionComponent = FVertexStreamComponent(&GNullVertexBuffer, 0, 0, VET_Float3);
+	VFData.PositionComponentSRV = GNullVertexBuffer.VertexBufferSRV;
+
+	VFData.TangentBasisComponents[0] = FVertexStreamComponent(&GNullVertexBuffer, 0, 0, VET_PackedNormal, EVertexStreamUsage::ManualFetch);
+	VFData.TangentBasisComponents[1] = FVertexStreamComponent(&GNullVertexBuffer, 0, 0, VET_PackedNormal, EVertexStreamUsage::ManualFetch);
+	VFData.TangentsSRV = GNullVertexBuffer.VertexBufferSRV;
+
+	VFData.NumTexCoords = 0;
+	VFData.TextureCoordinatesSRV = GNullVertexBuffer.VertexBufferSRV;
+
+	VFData.ColorIndexMask = 0;
+	VFData.ColorComponent = FVertexStreamComponent(&GNullColorVertexBuffer, 0, 0, VET_Color, EVertexStreamUsage::ManualFetch);
+	VFData.ColorComponentsSRV = GNullColorVertexBuffer.VertexBufferSRV;
+
+	VertexFactory.SetData(VFData);
+	VertexFactory.InitResource();
+}
+
+void FNiagaraRendererVolumes::ReleaseRenderThreadResources()
+{
+	FNiagaraRenderer::ReleaseRenderThreadResources();
+
+	VertexFactory.ReleaseResource();
+}
+
+bool FNiagaraRendererVolumes::IsMaterialValid(const UMaterialInterface* Material) const
+{
+	const UMaterial* BaseMaterial = Material ? Material->GetMaterial() : nullptr;
+	return BaseMaterial ? BaseMaterial->CheckMaterialUsage_Concurrent(MATUSAGE_HeterogeneousVolumes) : false;
+}
+
+FNiagaraDynamicDataBase* FNiagaraRendererVolumes::GenerateDynamicData(const FNiagaraSceneProxy* Proxy, const UNiagaraRendererProperties* InProperties, const FNiagaraEmitterInstance* Emitter) const
+{
+	using namespace NiagaraRendererVolumesLocal;
+
+	const UNiagaraVolumeRendererProperties* Properties = CastChecked<const UNiagaraVolumeRendererProperties>(InProperties);
+	if (!Properties || !IsRendererEnabled(Properties, Emitter))
+	{
+		return nullptr;
+	}
+
+	if (SourceMode != ENiagaraRendererSourceDataMode::Emitter)
+	{
+		if (SimTarget == ENiagaraSimTarget::GPUComputeSim)
+		{
+			return nullptr;
+		}
+
+		FNiagaraDataBuffer* DataToRender = Emitter->GetData().GetCurrentData();
+		if (!DataToRender || DataToRender->GetNumInstances() == 0)
+		{
+			return nullptr;
+		}
+	}
+
+	check(BaseMaterials_GT.Num() == 1);
+	check(BaseMaterials_GT[0]->CheckMaterialUsage_Concurrent(MATUSAGE_HeterogeneousVolumes));
+
+	// Setup Dynamic Data
+	FVolumeDynamicData* VolumeDynamicData = new FVolumeDynamicData(Emitter);
+	VolumeDynamicData->Material = BaseMaterials_GT[0]->GetRenderProxy();
+	VolumeDynamicData->SetMaterialRelevance(BaseMaterialRelevance_GT);
+	VolumeDynamicData->DefaultRendererVisibilityTag = Properties->RendererVisibility;
+
+	// Pull Default Values from Renderer Bounds Parameters
+	if (bAnyVFBoundOffsets)
+	{
+		const FNiagaraParameterStore& ParameterData = Emitter->GetRendererBoundVariables();
+
+		if (VFBoundOffsetsInParamStore[ENiagaraVolumeVFLayout::Position] != INDEX_NONE)
+		{
+			FMemory::Memcpy(&VolumeDynamicData->DefaultPosition, ParameterData.GetParameterData(VFBoundOffsetsInParamStore[ENiagaraVolumeVFLayout::Position]), sizeof(FVector3f));
+		}
+		if (VFBoundOffsetsInParamStore[ENiagaraVolumeVFLayout::Rotation] != INDEX_NONE)
+		{
+			FMemory::Memcpy(&VolumeDynamicData->DefaultRotation, ParameterData.GetParameterData(VFBoundOffsetsInParamStore[ENiagaraVolumeVFLayout::Rotation]), sizeof(FQuat4f));
+		}
+		if (VFBoundOffsetsInParamStore[ENiagaraVolumeVFLayout::Scale] != INDEX_NONE)
+		{
+			FMemory::Memcpy(&VolumeDynamicData->DefaultScale, ParameterData.GetParameterData(VFBoundOffsetsInParamStore[ENiagaraVolumeVFLayout::Scale]), sizeof(FVector3f));
+		}
+		if (VFBoundOffsetsInParamStore[ENiagaraVolumeVFLayout::RendererVisibilityTag] != INDEX_NONE)
+		{
+			FMemory::Memcpy(&VolumeDynamicData->DefaultRendererVisibilityTag, ParameterData.GetParameterData(VFBoundOffsetsInParamStore[ENiagaraVolumeVFLayout::RendererVisibilityTag]), sizeof(int32));
+		}
+		if (VFBoundOffsetsInParamStore[ENiagaraVolumeVFLayout::VolumeResolutionMaxAxis] != INDEX_NONE)
+		{
+			FMemory::Memcpy(&VolumeDynamicData->DefaultVolumeResolutionMaxAxis, ParameterData.GetParameterData(VFBoundOffsetsInParamStore[ENiagaraVolumeVFLayout::VolumeResolutionMaxAxis]), sizeof(int32));
+		}
+		if (VFBoundOffsetsInParamStore[ENiagaraVolumeVFLayout::VolumeWorldSpaceSize] != INDEX_NONE)
+		{
+			FMemory::Memcpy(&VolumeDynamicData->DefaultVolumeWorldSpaceSize, ParameterData.GetParameterData(VFBoundOffsetsInParamStore[ENiagaraVolumeVFLayout::VolumeWorldSpaceSize]), sizeof(FVector3f));
+		}
+	}
+
+	// If we are in emitter mode and the vis tag doesn't match we won't send any dynamic data
+	//-OPT: Perhaps remove the new / delete?  Path seems unlikely vs cluttering the code
+	if (SourceMode == ENiagaraRendererSourceDataMode::Emitter && VolumeDynamicData->DefaultRendererVisibilityTag != RendererVisibilityTag)
+	{
+		delete VolumeDynamicData;
+		return nullptr;
+	}
+
+	// Process Material Bindings
+	if (Properties->MaterialParameters.HasAnyBindings())
+	{
+		ProcessMaterialParameterBindings(Properties->MaterialParameters, Emitter, MakeArrayView(BaseMaterials_GT));
+	}
+
+	return VolumeDynamicData;
+}
+
+int FNiagaraRendererVolumes::GetDynamicDataSize() const
+{
+	using namespace NiagaraRendererVolumesLocal;
+	return sizeof(FVolumeDynamicData);
+}
+
+void FNiagaraRendererVolumes::GetDynamicMeshElements(const TArray<const FSceneView*>& Views, const FSceneViewFamily& ViewFamily, uint32 VisibilityMap, FMeshElementCollector& Collector, const FNiagaraSceneProxy* SceneProxy) const
+{
+	using namespace NiagaraRendererVolumesLocal;
+	
+	const FVolumeDynamicData* VolumeDynamicData = static_cast<const FVolumeDynamicData*>(DynamicDataRender);
+	if (VolumeDynamicData == nullptr)
+	{
+		return;
+	}
+
+	FNiagaraDataBuffer* ParticleDataBuffer = VolumeDynamicData->GetParticleDataToRender();
+	if (!ParticleDataBuffer || (SourceMode == ENiagaraRendererSourceDataMode::Particles && ParticleDataBuffer->GetNumInstances() == 0))
+	{
+		return;
+	}
+
+	for (int32 ViewIndex = 0; ViewIndex < Views.Num(); ViewIndex++)
+	{
+		if (VisibilityMap & (1 << ViewIndex))
+		{
+			const FSceneView* View = Views[ViewIndex];
+			if (View->bIsInstancedStereoEnabled && IStereoRendering::IsStereoEyeView(*View) && !IStereoRendering::IsAPrimaryView(*View))
+			{
+				// We don't have to generate batches for non-primary views in stereo instance rendering
+				continue;
+			}
+
+			FHeterogeneousVolumeInstances* VolumeInstances = &Collector.AllocateOneFrameResource<FHeterogeneousVolumeInstances>();
+			if (SourceMode == ENiagaraRendererSourceDataMode::Particles)
+			{
+				const FNiagaraDataSet* DataSet = ParticleDataBuffer->GetOwner();
+
+				FNiagaraDataSetReaderFloat<FNiagaraPosition>	PositionReader = PositionDataSetAccessor.GetReader(*DataSet);
+				FNiagaraDataSetReaderFloat<FQuat4f>				RotationReader = RotationDataSetAccessor.GetReader(*DataSet);
+				FNiagaraDataSetReaderFloat<FVector3f>			ScaleReader = ScaleDataSetAccessor.GetReader(*DataSet);
+				FNiagaraDataSetReaderInt32<int32>				RendererVisibilityTagReader = RendererVisibilityTagAccessor.GetReader(*DataSet);
+				FNiagaraDataSetReaderInt32<int32>				VolumeResolutionMaxAxisReader = VolumeResolutionMaxAxisAccessor.GetReader(*DataSet);
+				FNiagaraDataSetReaderFloat<FVector3f>			VolumeWorldSpaceSizeReader = VolumeWorldSpaceSizeAccessor.GetReader(*DataSet);
+
+				const bool bUseLocalSpace = UseLocalSpace(SceneProxy);
+				const FMatrix LocalToWorld = SceneProxy->GetLocalToWorld();
+				const FVector LWCTileOffset = FVector(SceneProxy->GetLWCRenderTile()) * FLargeWorldRenderScalar::GetTileSize();
+
+				VolumeInstances->Instances.Reserve(ParticleDataBuffer->GetNumInstances());
+				for (uint32 i=0; i < ParticleDataBuffer->GetNumInstances(); ++i)
+				{
+					const FVector	Position = FVector(PositionReader.GetSafe(i, VolumeDynamicData->DefaultPosition));
+					const FQuat		Rotation = FQuat(RotationReader.GetSafe(i, VolumeDynamicData->DefaultRotation).GetNormalized());
+					const FVector	Scale = FVector(ScaleReader.GetSafe(i, VolumeDynamicData->DefaultScale));
+					const int32		VisibilityTag = RendererVisibilityTagReader.GetSafe(i, VolumeDynamicData->DefaultRendererVisibilityTag);
+					const int32		ResolutionMaxAxis = VolumeResolutionMaxAxisReader.GetSafe(i, VolumeDynamicData->DefaultVolumeResolutionMaxAxis);
+					const FVector3f	WorldSpaceSize = VolumeWorldSpaceSizeReader.GetSafe(i, VolumeDynamicData->DefaultVolumeWorldSpaceSize);
+
+					if (VisibilityTag != RendererVisibilityTag || FMath::IsNearlyZero(Scale.SquaredLength()))
+					{
+						continue;
+					}
+					
+					FMatrix InstanceToWorld = FScaleMatrix::Make(Scale) * FQuatRotationTranslationMatrix::Make(Rotation, Position);
+					if ( bUseLocalSpace )
+					{
+						InstanceToWorld = LocalToWorld * InstanceToWorld;
+					}
+					else
+					{
+						InstanceToWorld.SetOrigin(InstanceToWorld.GetOrigin() + LWCTileOffset);
+					}
+
+					//-TODO: Push per instance data in here
+					FHeterogeneousVolumeInstances::FInstance& VolumeInstance = VolumeInstances->Instances.AddDefaulted_GetRef();
+					VolumeInstance.LocalBounds	= SceneProxy->GetLocalBounds();
+					VolumeInstance.LocalToWorld	= InstanceToWorld;
+				}
+			}
+			else
+			{
+				//-TODO: Matches current behavior but should pull emitter bounds / position binding location
+				FHeterogeneousVolumeInstances::FInstance& VolumeInstance = VolumeInstances->Instances.AddDefaulted_GetRef();
+				VolumeInstance.LocalBounds	= SceneProxy->GetLocalBounds();
+				VolumeInstance.LocalToWorld	= SceneProxy->GetLocalToWorld();
+			}
+
+			if (VolumeInstances->Instances.Num() == 0)
+			{
+				continue;
+			}
+
+			FMeshBatch& Mesh				= Collector.AllocateMesh();
+			Mesh.VertexFactory				= &VertexFactory;
+			Mesh.MaterialRenderProxy		= VolumeDynamicData->Material;
+			Mesh.LCI						= nullptr;
+			Mesh.ReverseCulling				= SceneProxy->IsLocalToWorldDeterminantNegative() ? true : false;
+			Mesh.CastShadow					= false;
+			Mesh.Type						= PT_TriangleStrip;
+			Mesh.bDisableBackfaceCulling	= true;
+			Mesh.bCanApplyViewModeOverrides = true;
+			Mesh.bUseWireframeSelectionColoring = SceneProxy->IsSelected();
+
+			FMeshBatchElement& BatchElement	= Mesh.Elements[0];
+			BatchElement.IndexBuffer		= nullptr;
+			BatchElement.FirstIndex			= 0;
+			BatchElement.MinVertexIndex		= 0;
+			BatchElement.MaxVertexIndex		= 3;
+			BatchElement.NumPrimitives		= 2;
+			BatchElement.BaseVertexIndex	= 0;
+			//-TODO:BatchElement.VertexFactoryUserData = VolumeInstances;
+			Collector.AddMesh(0, Mesh);
+		}
+	}
+}
+
+#if RHI_RAYTRACING
+void FNiagaraRendererVolumes::GetDynamicRayTracingInstances(FRayTracingMaterialGatheringContext& Context, TArray<FRayTracingInstance>& OutRayTracingInstances, const FNiagaraSceneProxy* SceneProxy)
+{
+	//-TODO: Add support for raytracing
+}
+#endif //RHI_RAYTRACING
diff --git a/Engine/Plugins/FX/Niagara/Source/Niagara/Private/NiagaraVolumeRendererProperties.cpp b/Engine/Plugins/FX/Niagara/Source/Niagara/Private/NiagaraVolumeRendererProperties.cpp
new file mode 100644
--- /dev/null
+++ b/Engine/Plugins/FX/Niagara/Source/Niagara/Private/NiagaraVolumeRendererProperties.cpp
@@ -0,0 +1,384 @@
+// Copyright Epic Games, Inc. All Rights Reserved.
+
+#include "NiagaraVolumeRendererProperties.h"
+#include "NiagaraConstants.h"
+#include "NiagaraEmitterInstance.h"
+#include "NiagaraRenderer.h"
+#include "NiagaraRendererVolumes.h"
+#include "Materials/MaterialInstanceConstant.h"
+#include "Modules/ModuleManager.h"
+
+#include UE_INLINE_GENERATED_CPP_BY_NAME(NiagaraVolumeRendererProperties)
+
+#if WITH_EDITOR
+#include "Widgets/Images/SImage.h"
+#include "Widgets/Text/STextBlock.h"
+#include "Widgets/SWidget.h"
+#include "Styling/SlateBrush.h"
+#include "Styling/SlateIconFinder.h"
+#include "AssetThumbnail.h"
+#endif
+
+#define LOCTEXT_NAMESPACE "UNiagaraVolumeRendererProperties"
+
+namespace NiagaraVolumeRendererPropertiesLocal
+{
+	TArray<TWeakObjectPtr<UNiagaraVolumeRendererProperties>> RendererPropertiesToDeferredInit;
+
+	template<typename TValueType>
+	static FNiagaraVariable MakeNiagaraVariableWithValue(const FNiagaraTypeDefinition& TypeDef, const FName& Name, const TValueType& Value)
+	{
+		FNiagaraVariable Variable(TypeDef, Name);
+		Variable.SetValue(Value);
+		return Variable;
+	}
+
+	static FNiagaraVariable GetVolumeRotationVariable()
+	{
+		static FNiagaraVariable Variable = MakeNiagaraVariableWithValue(FNiagaraTypeDefinition::GetQuatDef(), TEXT("Particles.VolumeRotation"), UNiagaraVolumeRendererProperties::GetDefaultVolumeRotation());
+		return Variable;
+	}
+
+	static FNiagaraVariable GetVolumeScaleVariable()
+	{
+		static FNiagaraVariable Variable = MakeNiagaraVariableWithValue(FNiagaraTypeDefinition::GetVec3Def(), TEXT("Particles.VolumeScale"), UNiagaraVolumeRendererProperties::GetDefaultVolumeScale());
+		return Variable;
+	}
+
+	static FNiagaraVariable GetVolumeResolutionMaxAxisVariable()
+	{
+		static FNiagaraVariable Variable = MakeNiagaraVariableWithValue(FNiagaraTypeDefinition::GetIntDef(), TEXT("Particles.VolumeResolutionMaxAxis"), UNiagaraVolumeRendererProperties::GetDefaultVolumeResolutionMaxAxis());
+		return Variable;
+	}
+
+	static FNiagaraVariable GetVolumeWorldSpaceSizeVariable()
+	{
+		static FNiagaraVariable Variable = MakeNiagaraVariableWithValue(FNiagaraTypeDefinition::GetVec3Def(), TEXT("Particles.VolumeWorldSpaceSize"), UNiagaraVolumeRendererProperties::GetDefaultVolumeWorldSpaceSize());
+		return Variable;
+	}
+
+	static void SetupBindings(UNiagaraVolumeRendererProperties* Props)
+	{
+		if (Props->PositionBinding.IsValid())
+		{
+			return;
+		}
+		Props->PositionBinding					= FNiagaraConstants::GetAttributeDefaultBinding(SYS_PARAM_PARTICLES_POSITION);
+		Props->RotationBinding.Setup(GetVolumeRotationVariable(), GetVolumeRotationVariable());
+		Props->ScaleBinding.Setup(GetVolumeScaleVariable(), GetVolumeScaleVariable());
+		Props->RendererVisibilityTagBinding		= FNiagaraConstants::GetAttributeDefaultBinding(SYS_PARAM_PARTICLES_VISIBILITY_TAG);
+		Props->VolumeResolutionMaxAxisBinding.Setup(GetVolumeResolutionMaxAxisVariable(), GetVolumeResolutionMaxAxisVariable());
+		Props->VolumeWorldSpaceSizeBinding.Setup(GetVolumeWorldSpaceSizeVariable(), GetVolumeWorldSpaceSizeVariable());
+
+	#if WITH_EDITORONLY_DATA
+		Props->MaterialParameterBinding.SetUsage(ENiagaraParameterBindingUsage::NotParticle);
+		Props->MaterialParameterBinding.SetAllowedObjects({ UMaterialInterface::StaticClass() });
+	#endif
+	}
+}
+
+UNiagaraVolumeRendererProperties::UNiagaraVolumeRendererProperties()
+{
+	AttributeBindings.Reserve(ENiagaraVolumeVFLayout::Num);
+	AttributeBindings.Add(&PositionBinding);
+	AttributeBindings.Add(&RotationBinding);
+	AttributeBindings.Add(&ScaleBinding);
+	AttributeBindings.Add(&RendererVisibilityTagBinding);
+	AttributeBindings.Add(&VolumeResolutionMaxAxisBinding);
+	AttributeBindings.Add(&VolumeWorldSpaceSizeBinding);
+}
+
+void UNiagaraVolumeRendererProperties::PostLoad()
+{
+	Super::PostLoad();
+	
+#if WITH_EDITORONLY_DATA
+	ChangeToPositionBinding(PositionBinding);
+#endif
+	PostLoadBindings(ENiagaraRendererSourceDataMode::Particles);
+}
+
+void UNiagaraVolumeRendererProperties::PostInitProperties()
+{
+	using namespace NiagaraVolumeRendererPropertiesLocal;
+
+	Super::PostInitProperties();
+
+	if (HasAnyFlags(RF_ClassDefaultObject) == false)
+	{
+		// We can end up hitting PostInitProperties before the Niagara Module has initialized bindings this needs, mark this object for deferred init and early out.
+		if (FModuleManager::Get().IsModuleLoaded("Niagara") == false)
+		{
+			NiagaraVolumeRendererPropertiesLocal::RendererPropertiesToDeferredInit.Add(this);
+			return;
+		}
+		SetupBindings(this);
+	}
+}
+
+void UNiagaraVolumeRendererProperties::Serialize(FArchive& Ar)
+{
+	// MIC will replace the main material during serialize
+	// Be careful if adding code that looks at the material to make sure you get the correct one
+	{
+//#if WITH_EDITORONLY_DATA
+//		TOptional<TGuardValue<TObjectPtr<UMaterialInterface>>> MICGuard;
+//		if (Ar.IsSaving() && Ar.IsCooking() && MICMaterial)
+//		{
+//			MICGuard.Emplace(Material, MICMaterial);
+//		}
+//#endif
+
+		Super::Serialize(Ar);
+	}
+}
+
+#if WITH_EDITORONLY_DATA
+void UNiagaraVolumeRendererProperties::PostEditChangeProperty(struct FPropertyChangedEvent& PropertyChangedEvent)
+{
+	Super::PostEditChangeProperty(PropertyChangedEvent);
+
+	const FName PropertyName = PropertyChangedEvent.GetPropertyName();
+	const FName MemberPropertyName = PropertyChangedEvent.GetMemberPropertyName();
+
+	if (PropertyName == GET_MEMBER_NAME_CHECKED(UNiagaraVolumeRendererProperties, SourceMode))
+	{
+		UpdateSourceModeDerivates(SourceMode, true);
+	}
+
+	// Update our MICs if we change material / material bindings
+	//-OPT: Could narrow down further to only static materials
+	if ((PropertyName == GET_MEMBER_NAME_CHECKED(UNiagaraVolumeRendererProperties, Material)) ||
+		(MemberPropertyName == GET_MEMBER_NAME_CHECKED(UNiagaraVolumeRendererProperties, MaterialParameters)))
+	{
+		UpdateMICs();
+	}
+}
+#endif// WITH_EDITORONLY_DATA
+
+void UNiagaraVolumeRendererProperties::InitCDOPropertiesAfterModuleStartup()
+{
+	using namespace NiagaraVolumeRendererPropertiesLocal;
+
+	UNiagaraVolumeRendererProperties* CDO = CastChecked<UNiagaraVolumeRendererProperties>(UNiagaraVolumeRendererProperties::StaticClass()->GetDefaultObject());
+	SetupBindings(CDO);
+
+	for (TWeakObjectPtr<UNiagaraVolumeRendererProperties>& WeakProps : NiagaraVolumeRendererPropertiesLocal::RendererPropertiesToDeferredInit)
+	{
+		if (UNiagaraVolumeRendererProperties* Props = WeakProps.Get())
+		{
+			SetupBindings(Props);
+		}
+	}
+}
+
+FNiagaraRenderer* UNiagaraVolumeRendererProperties::CreateEmitterRenderer(ERHIFeatureLevel::Type FeatureLevel, const FNiagaraEmitterInstance* Emitter, const FNiagaraSystemInstanceController& InController)
+{
+	FNiagaraRenderer* NewRenderer = new FNiagaraRendererVolumes(FeatureLevel, this, Emitter);
+	NewRenderer->Initialize(this, Emitter, InController);
+	return NewRenderer;
+}
+
+//class FNiagaraBoundsCalculatorVolumes : public FNiagaraBoundsCalculator
+//{
+//public:
+//	FNiagaraBoundsCalculatorVolumes(const FNiagaraDataSetAccessor<FNiagaraPosition>& InPositionAccessor, const FNiagaraDataSetAccessor<FVector3f>& InVolumeWorldSpaceSizeAccessor)
+//		: PositionAccessor(InPositionAccessor)
+//		, VolumeWorldSpaceSizeAccessor(InVolumeWorldSpaceSizeAccessor)
+//	{}
+//
+//	FNiagaraBoundsCalculatorVolumes() = delete;
+//	virtual ~FNiagaraBoundsCalculatorVolumes() = default;
+//
+//	virtual void InitAccessors(const FNiagaraDataSetCompiledData* CompiledData) override {}
+//	virtual FBox CalculateBounds(const FTransform& SystemTransform, const FNiagaraDataSet& DataSet, const int32 NumInstances) const override
+//	{
+//		if (!NumInstances || !PositionAccessor.IsValid())
+//		{
+//			return FBox(ForceInit);
+//		}
+//
+//		FNiagaraPosition BoundsMin(ForceInitToZero);
+//		FNiagaraPosition BoundsMax(ForceInitToZero);
+//		PositionAccessor.GetReader(DataSet).GetMinMax(BoundsMin, BoundsMax);
+//		FBox Bounds(BoundsMin, BoundsMax);
+//
+//		if (VolumeWorldSpaceSizeAccessor.IsValid())
+//		{
+//			Bounds = Bounds.ExpandBy(VolumeSizeAccessor.GetReader(DataSet).GetMax().GetAbsMax());
+//		}
+//		return Bounds;
+//	}
+//
+//protected:
+//	const FNiagaraDataSetAccessor<FNiagaraPosition>&	PositionAccessor;
+//	const FNiagaraDataSetAccessor<FVector3f>&			VolumeWorldSpaceSizeAccessor;
+//};
+
+FNiagaraBoundsCalculator* UNiagaraVolumeRendererProperties::CreateBoundsCalculator()
+{
+	//-TODO: Dynamic Bounds Calculation
+	//if (GetCurrentSourceMode() == ENiagaraRendererSourceDataMode::Emitter)
+	//{
+	//	return nullptr;
+	//}
+
+	//return new FNiagaraBoundsCalculatorVolumes(PositionDataSetAccessor, DecalSizeDataSetAccessor);
+	return nullptr;
+}
+
+void UNiagaraVolumeRendererProperties::GetUsedMaterials(const FNiagaraEmitterInstance* InEmitter, TArray<UMaterialInterface*>& OutMaterials) const
+{
+	UMaterialInterface* MaterialInterface = nullptr;
+	if (InEmitter != nullptr)
+	{
+		MaterialInterface = Cast<UMaterialInterface>(InEmitter->FindBinding(MaterialParameterBinding.Parameter));
+	}
+
+#if WITH_EDITORONLY_DATA
+	MaterialInterface = MaterialInterface ? MaterialInterface : ToRawPtr(MICMaterial);
+#endif
+
+	OutMaterials.Add(MaterialInterface ? MaterialInterface : ToRawPtr(Material));
+}
+
+void UNiagaraVolumeRendererProperties::CacheFromCompiledData(const FNiagaraDataSetCompiledData* CompiledData)
+{
+	UpdateSourceModeDerivates(SourceMode);
+	UpdateMICs();
+
+	PositionDataSetAccessor.Init(CompiledData, PositionBinding.GetDataSetBindableVariable().GetName());
+	RotationDataSetAccessor.Init(CompiledData, RotationBinding.GetDataSetBindableVariable().GetName());
+	ScaleDataSetAccessor.Init(CompiledData, ScaleBinding.GetDataSetBindableVariable().GetName());
+	RendererVisibilityTagAccessor.Init(CompiledData, RendererVisibilityTagBinding.GetDataSetBindableVariable().GetName());
+	VolumeResolutionMaxAxisAccessor.Init(CompiledData, VolumeResolutionMaxAxisBinding.GetDataSetBindableVariable().GetName());
+	VolumeWorldSpaceSizeAccessor.Init(CompiledData, VolumeWorldSpaceSizeBinding.GetDataSetBindableVariable().GetName());
+	RendererVisibilityTagAccessor.Init(CompiledData, RendererVisibilityTagBinding.GetDataSetBindableVariable().GetName());
+}
+
+void UNiagaraVolumeRendererProperties::UpdateSourceModeDerivates(ENiagaraRendererSourceDataMode InSourceMode, bool bFromPropertyEdit)
+{
+	Super::UpdateSourceModeDerivates(InSourceMode, bFromPropertyEdit);
+
+	if ( UNiagaraEmitter* SrcEmitter = GetTypedOuter<UNiagaraEmitter>() )
+	{
+		for (FNiagaraMaterialAttributeBinding& MaterialParamBinding : MaterialParameters.AttributeBindings)
+		{
+			MaterialParamBinding.CacheValues(SrcEmitter);
+		}
+	}
+}
+
+bool UNiagaraVolumeRendererProperties::PopulateRequiredBindings(FNiagaraParameterStore& InParameterStore)
+{
+	bool bAnyAdded = Super::PopulateRequiredBindings(InParameterStore);
+
+	for (const FNiagaraVariableAttributeBinding* Binding : AttributeBindings)
+	{
+		if (Binding && Binding->CanBindToHostParameterMap())
+		{
+			InParameterStore.AddParameter(Binding->GetParamMapBindableVariable(), false);
+			bAnyAdded = true;
+		}
+	}
+
+	for (const FNiagaraMaterialAttributeBinding& MaterialParamBinding : MaterialParameters.AttributeBindings)
+	{
+		InParameterStore.AddParameter(MaterialParamBinding.GetParamMapBindableVariable(), false);
+		bAnyAdded = true;
+	}
+
+	return bAnyAdded;
+}
+
+#if WITH_EDITORONLY_DATA
+const TArray<FNiagaraVariable>& UNiagaraVolumeRendererProperties::GetOptionalAttributes()
+{
+	using namespace NiagaraVolumeRendererPropertiesLocal;
+
+	static TArray<FNiagaraVariable> Attrs;
+	if (Attrs.Num() == 0)
+	{
+		Attrs.Add(SYS_PARAM_PARTICLES_POSITION);
+		//Attrs.Add(SYS_PARAM_PARTICLES_VISIBILITY_TAG);
+	}
+	return Attrs;
+}
+
+void UNiagaraVolumeRendererProperties::GetRendererWidgets(const FNiagaraEmitterInstance* InEmitter, TArray<TSharedPtr<SWidget>>& OutWidgets, TSharedPtr<FAssetThumbnailPool> InThumbnailPool) const
+{
+	OutWidgets.Add(SNew(SImage).Image(FSlateIconFinder::FindIconBrushForClass(GetClass())));
+}
+
+void UNiagaraVolumeRendererProperties::GetRendererTooltipWidgets(const FNiagaraEmitterInstance* InEmitter, TArray<TSharedPtr<SWidget>>& OutWidgets, TSharedPtr<FAssetThumbnailPool> InThumbnailPool) const
+{
+	OutWidgets.Add(SNew(STextBlock).Text(LOCTEXT("VolumeRenderer", "Volume Renderer")));
+}
+
+void UNiagaraVolumeRendererProperties::GetRendererFeedback(const FVersionedNiagaraEmitter& InEmitter, TArray<FNiagaraRendererFeedback>& OutErrors, TArray<FNiagaraRendererFeedback>& OutWarnings, TArray<FNiagaraRendererFeedback>& OutInfo) const
+{
+	Super::GetRendererFeedback(InEmitter, OutErrors, OutWarnings, OutInfo);
+
+	if ( const FVersionedNiagaraEmitterData* EmitterData = InEmitter.GetEmitterData() )
+	{
+		if (SourceMode == ENiagaraRendererSourceDataMode::Particles && EmitterData->SimTarget == ENiagaraSimTarget::GPUComputeSim)
+		{
+			OutErrors.Emplace(
+				LOCTEXT("GPUNotSupportedDesc", "Volume renderer does not support reading from GPU particles."),
+				LOCTEXT("GPUNotSupportedSummary", "GPU particles not support.")
+				//LOCTEXT("GPUNotSupportedFix", "Switch emitter to CPU particles."),
+				//FNiagaraRendererFeedbackFix::CreateLambda(
+				//	[this]()
+				//	{
+				//		const_cast<UNiagaraRibbonRendererProperties*>(this)->MultiPlaneCount = FMath::Clamp(this->MultiPlaneCount - 1, 1, 16);
+				//	}
+				//),
+			);
+		}
+	}
+
+	if (MaterialParameters.HasAnyBindings())
+	{
+		TArray<UMaterialInterface*> Materials;
+		GetUsedMaterials(nullptr, Materials);
+		MaterialParameters.GetFeedback(Materials, OutWarnings);
+	}
+}
+
+TArray<FNiagaraVariable> UNiagaraVolumeRendererProperties::GetBoundAttributes() const
+{
+	TArray<FNiagaraVariable> BoundAttributes = Super::GetBoundAttributes();
+	BoundAttributes.Reserve(BoundAttributes.Num() + MaterialParameters.AttributeBindings.Num());
+
+	for (const FNiagaraMaterialAttributeBinding& MaterialParamBinding : MaterialParameters.AttributeBindings)
+	{
+		BoundAttributes.AddUnique(MaterialParamBinding.GetParamMapBindableVariable());
+	}
+	return BoundAttributes;
+}
+
+void UNiagaraVolumeRendererProperties::RenameVariable(const FNiagaraVariableBase& OldVariable, const FNiagaraVariableBase& NewVariable, const FVersionedNiagaraEmitter& InEmitter)
+{
+	Super::RenameVariable(OldVariable, NewVariable, InEmitter);
+#if WITH_EDITORONLY_DATA
+	MaterialParameters.RenameVariable(OldVariable, NewVariable, InEmitter, GetCurrentSourceMode());
+#endif
+}
+
+void UNiagaraVolumeRendererProperties::RemoveVariable(const FNiagaraVariableBase& OldVariable, const FVersionedNiagaraEmitter& InEmitter)
+{
+	Super::RemoveVariable(OldVariable, InEmitter);
+#if WITH_EDITORONLY_DATA
+	MaterialParameters.RemoveVariable(OldVariable, InEmitter, GetCurrentSourceMode());
+#endif
+}
+#endif // WITH_EDITORONLY_DATA
+
+void UNiagaraVolumeRendererProperties::UpdateMICs()
+{
+#if WITH_EDITORONLY_DATA
+	UpdateMaterialParametersMIC(MaterialParameters, Material, MICMaterial);
+#endif
+}
+
+#undef LOCTEXT_NAMESPACE
diff --git a/Engine/Plugins/FX/Niagara/Source/Niagara/Public/NiagaraRendererVolumes.h b/Engine/Plugins/FX/Niagara/Source/Niagara/Public/NiagaraRendererVolumes.h
new file mode 100644
--- /dev/null
+++ b/Engine/Plugins/FX/Niagara/Source/Niagara/Public/NiagaraRendererVolumes.h
@@ -0,0 +1,49 @@
+// Copyright Epic Games, Inc. All Rights Reserved.
+
+#pragma once
+
+#include "NiagaraRenderer.h"
+#include "NiagaraVolumeRendererProperties.h"
+
+#include "LocalVertexFactory.h"
+
+class USceneComponent;
+
+class NIAGARA_API FNiagaraRendererVolumes : public FNiagaraRenderer
+{
+public:
+	explicit FNiagaraRendererVolumes(ERHIFeatureLevel::Type FeatureLevel, const UNiagaraRendererProperties *InProps, const FNiagaraEmitterInstance* Emitter);
+	~FNiagaraRendererVolumes();
+
+	//FNiagaraRenderer interface
+	virtual void CreateRenderThreadResources() override;
+	virtual void ReleaseRenderThreadResources() override;
+
+	virtual bool IsMaterialValid(const UMaterialInterface* Material) const override;
+
+	virtual FNiagaraDynamicDataBase *GenerateDynamicData(const FNiagaraSceneProxy* Proxy, const UNiagaraRendererProperties* InProperties, const FNiagaraEmitterInstance* Emitter) const override;
+	virtual int GetDynamicDataSize() const override;
+
+	void GetDynamicMeshElements(const TArray<const FSceneView*>& Views, const FSceneViewFamily& ViewFamily, uint32 VisibilityMap, FMeshElementCollector& Collector, const FNiagaraSceneProxy* SceneProxy) const;
+#if RHI_RAYTRACING
+	void GetDynamicRayTracingInstances(FRayTracingMaterialGatheringContext& Context, TArray<FRayTracingInstance>& OutRayTracingInstances, const FNiagaraSceneProxy* SceneProxy);
+#endif //RHI_RAYTRACING
+	//FNiagaraRenderer interface END
+
+
+protected:
+	ENiagaraRendererSourceDataMode	SourceMode = ENiagaraRendererSourceDataMode::Emitter;
+	int32							RendererVisibilityTag = 0;
+
+	FNiagaraDataSetAccessor<FNiagaraPosition>	PositionDataSetAccessor;
+	FNiagaraDataSetAccessor<FQuat4f>			RotationDataSetAccessor;
+	FNiagaraDataSetAccessor<FVector3f>			ScaleDataSetAccessor;
+	FNiagaraDataSetAccessor<int32>				RendererVisibilityTagAccessor;
+	FNiagaraDataSetAccessor<int32>				VolumeResolutionMaxAxisAccessor;
+	FNiagaraDataSetAccessor<FVector3f>			VolumeWorldSpaceSizeAccessor;
+
+	bool							bAnyVFBoundOffsets = false;
+	int32							VFBoundOffsetsInParamStore[int32(ENiagaraVolumeVFLayout::Num)];
+
+	FLocalVertexFactory				VertexFactory;
+};
diff --git a/Engine/Plugins/FX/Niagara/Source/Niagara/Public/NiagaraVolumeRendererProperties.h b/Engine/Plugins/FX/Niagara/Source/Niagara/Public/NiagaraVolumeRendererProperties.h
new file mode 100644
--- /dev/null
+++ b/Engine/Plugins/FX/Niagara/Source/Niagara/Public/NiagaraVolumeRendererProperties.h
@@ -0,0 +1,142 @@
+// Copyright Epic Games, Inc. All Rights Reserved.
+
+#pragma once
+
+#include "CoreMinimal.h"
+#include "UObject/ObjectMacros.h"
+#include "NiagaraRendererProperties.h"
+#include "NiagaraCommon.h"
+#include "NiagaraDataSetAccessor.h"
+#include "NiagaraParameterBinding.h"
+#include "NiagaraVolumeRendererProperties.generated.h"
+
+class UMaterialInstanceConstant;
+class UMaterialInterface;
+class FNiagaraEmitterInstance;
+class SWidget;
+
+namespace ENiagaraVolumeVFLayout
+{
+	enum Type
+	{
+		Position,
+		Rotation,
+		Scale,
+		RendererVisibilityTag,
+		VolumeResolutionMaxAxis,
+		VolumeWorldSpaceSize,
+
+		Num
+	};
+};
+
+UCLASS(editinlinenew, MinimalAPI, meta = (DisplayName = "Volume Renderer"))
+class UNiagaraVolumeRendererProperties : public UNiagaraRendererProperties
+{
+public:
+	GENERATED_BODY()
+
+	UNiagaraVolumeRendererProperties();
+
+	//UObject Interface
+	virtual void PostLoad() override;
+	virtual void PostInitProperties() override;
+	virtual void Serialize(FArchive& Ar) override;
+#if WITH_EDITORONLY_DATA
+	virtual void PostEditChangeProperty(struct FPropertyChangedEvent& PropertyChangedEvent) override;
+#endif// WITH_EDITORONLY_DATA
+	//UObject Interface END
+
+	static void InitCDOPropertiesAfterModuleStartup();
+
+	//~ UNiagaraRendererProperties interface
+	virtual FNiagaraRenderer* CreateEmitterRenderer(ERHIFeatureLevel::Type FeatureLevel, const FNiagaraEmitterInstance* Emitter, const FNiagaraSystemInstanceController& InController) override;
+	virtual class FNiagaraBoundsCalculator* CreateBoundsCalculator() override;
+	virtual void GetUsedMaterials(const FNiagaraEmitterInstance* InEmitter, TArray<UMaterialInterface*>& OutMaterials) const override;
+	virtual bool IsSimTargetSupported(ENiagaraSimTarget InSimTarget) const override { return true; };
+#if WITH_EDITORONLY_DATA
+	virtual const TArray<FNiagaraVariable>& GetOptionalAttributes() override;
+	virtual void GetRendererWidgets(const FNiagaraEmitterInstance* InEmitter, TArray<TSharedPtr<SWidget>>& OutWidgets, TSharedPtr<FAssetThumbnailPool> InThumbnailPool) const override;
+	virtual void GetRendererTooltipWidgets(const FNiagaraEmitterInstance* InEmitter, TArray<TSharedPtr<SWidget>>& OutWidgets, TSharedPtr<FAssetThumbnailPool> InThumbnailPool) const override;
+	virtual void GetRendererFeedback(const FVersionedNiagaraEmitter& InEmitter, TArray<FNiagaraRendererFeedback>& OutErrors, TArray<FNiagaraRendererFeedback>& OutWarnings, TArray<FNiagaraRendererFeedback>& OutInfo) const override;
+	TArray<FNiagaraVariable> GetBoundAttributes() const override;
+	virtual void RenameVariable(const FNiagaraVariableBase& OldVariable, const FNiagaraVariableBase& NewVariable, const FVersionedNiagaraEmitter& InEmitter) override;
+	virtual void RemoveVariable(const FNiagaraVariableBase& OldVariable, const FVersionedNiagaraEmitter& InEmitter) override;
+#endif // WITH_EDITORONLY_DATA
+	virtual void CacheFromCompiledData(const FNiagaraDataSetCompiledData* CompiledData) override;
+	virtual void UpdateSourceModeDerivates(ENiagaraRendererSourceDataMode InSourceMode, bool bFromPropertyEdit = false) override;
+	virtual ENiagaraRendererSourceDataMode GetCurrentSourceMode() const override { return SourceMode; }
+	virtual bool PopulateRequiredBindings(FNiagaraParameterStore& InParameterStore) override;
+	virtual bool NeedsMIDsForMaterials() const override { return MaterialParameters.HasAnyBindings(); }
+	virtual bool UseHeterogeneousVolumes() const override { return true; }
+	//UNiagaraRendererProperties Interface END
+
+	void UpdateMICs();
+
+	static FQuat4f GetDefaultVolumeRotation() { return FQuat4f::Identity; }
+	static FVector3f GetDefaultVolumeScale() { return FVector3f(1.0f, 1.0f, 1.0f); }
+	static int32 GetDefaultVolumeResolutionMaxAxis() { return 0; }
+	static FVector3f GetDefaultVolumeWorldSpaceSize() { return FVector3f::ZeroVector; }
+
+	/** What material to use for the volume. */
+	UPROPERTY(EditAnywhere, Category = "Volume Rendering")
+	TObjectPtr<UMaterialInterface> Material;
+
+#if WITH_EDITORONLY_DATA
+	UPROPERTY(transient)
+	TObjectPtr<UMaterialInstanceConstant> MICMaterial;
+#endif
+
+	/** Binding to material. */
+	UPROPERTY(EditAnywhere, Category = "Volume Rendering")
+	FNiagaraParameterBinding MaterialParameterBinding;
+
+	/** Whether or not to draw a single element for the Emitter or to draw the particles.*/
+	//-TODO: Unhind when particle support is added
+	//UPROPERTY(EditAnywhere, Category = "Volume Rendering")
+	ENiagaraRendererSourceDataMode SourceMode = ENiagaraRendererSourceDataMode::Emitter;
+
+	/** If a render visibility tag is present, particles whose tag matches this value will be visible in this renderer. */
+	UPROPERTY(EditAnywhere, Category = "Volume Rendering")
+	int32 RendererVisibility = 0;
+
+	/** Position binding for the center of the volume. */
+	//-TODO: Unhind when binding is supported
+	//UPROPERTY(EditAnywhere, Category = "Bindings")
+	FNiagaraVariableAttributeBinding PositionBinding;
+
+	/** Rotation binding for the volume. */
+	//-TODO: Unhind when binding is supported
+	//UPROPERTY(EditAnywhere, Category = "Bindings")
+	FNiagaraVariableAttributeBinding RotationBinding;
+
+	/** Scale binding for the volume. */
+	//-TODO: Unhind when binding is supported
+	//UPROPERTY(EditAnywhere, Category = "Bindings")
+	FNiagaraVariableAttributeBinding ScaleBinding;
+
+	/** Visibility tag binding, when valid the returned values is compared with RendererVisibility. */
+	UPROPERTY(EditAnywhere, Category = "Bindings")
+	FNiagaraVariableAttributeBinding RendererVisibilityTagBinding;
+
+	/** TODO */
+	//-TODO: Unhind when binding is supported
+	//UPROPERTY(EditAnywhere, Category = "Bindings")
+	FNiagaraVariableAttributeBinding VolumeResolutionMaxAxisBinding;
+
+	/** TODO */
+	//-TODO: Unhind when binding is supported
+	//UPROPERTY(EditAnywhere, Category = "Bindings")
+	FNiagaraVariableAttributeBinding VolumeWorldSpaceSizeBinding;
+
+	/** If this array has entries, we will create a MaterialInstanceDynamic per Emitter instance from Material and set the Material parameters using the Niagara simulation variables listed.*/
+	UPROPERTY(EditAnywhere, Category = "Bindings")
+	FNiagaraRendererMaterialParameters MaterialParameters;
+
+	FNiagaraDataSetAccessor<FNiagaraPosition>	PositionDataSetAccessor;
+	FNiagaraDataSetAccessor<FQuat4f>			RotationDataSetAccessor;
+	FNiagaraDataSetAccessor<FVector3f>			ScaleDataSetAccessor;
+	FNiagaraDataSetAccessor<int32>				RendererVisibilityTagAccessor;
+	FNiagaraDataSetAccessor<int32>				VolumeResolutionMaxAxisAccessor;
+	FNiagaraDataSetAccessor<FVector3f>			VolumeWorldSpaceSizeAccessor;
+};
diff --git a/Engine/Plugins/FX/Niagara/Source/NiagaraEditor/Private/NiagaraEditorModule.cpp b/Engine/Plugins/FX/Niagara/Source/NiagaraEditor/Private/NiagaraEditorModule.cpp
--- a/Engine/Plugins/FX/Niagara/Source/NiagaraEditor/Private/NiagaraEditorModule.cpp
+++ b/Engine/Plugins/FX/Niagara/Source/NiagaraEditor/Private/NiagaraEditorModule.cpp
@@ -155,6 +155,7 @@
 #include "NiagaraLightRendererProperties.h"
 #include "NiagaraRibbonRendererProperties.h"
 #include "NiagaraSpriteRendererProperties.h"
+#include "NiagaraVolumeRendererProperties.h"
 #include "Misc/ScopedSlowTask.h"
 
 #include "NiagaraActions.h"
@@ -1845,6 +1846,18 @@ void FNiagaraEditorModule::RegisterDefaultRendererFactories()
 			NewRenderer->Material = Cast<UMaterialInterface>(DefaultMaterial.TryLoad());
 			return NewRenderer;
 		})));
+
+	RegisterRendererCreationInfo(FNiagaraRendererCreationInfo(
+		UNiagaraVolumeRendererProperties::StaticClass()->GetDisplayNameText(),
+		FText::FromString(UNiagaraVolumeRendererProperties::StaticClass()->GetDescription()),
+		UNiagaraVolumeRendererProperties::StaticClass()->GetClassPathName(),
+		FNiagaraRendererCreationInfo::FRendererFactory::CreateLambda([](UObject* OuterEmitter)
+		{
+			UNiagaraVolumeRendererProperties* NewRenderer = NewObject<UNiagaraVolumeRendererProperties>(OuterEmitter, NAME_None, RF_Transactional);
+			//FSoftObjectPath DefaultMaterial(TEXT("Material'/Niagara/DefaultAssets/DefaultDecalMaterial.DefaultDecalMaterial'"));
+			//NewRenderer->Material = Cast<UMaterialInterface>(DefaultMaterial.TryLoad());
+			return NewRenderer;
+		})));
 }
 
 void FNiagaraEditorModule::RegisterSettings()
