Fixed the script variable customization to also account for parameter pins that are disconnected from an output node.
This fixes a bug in which a disconnected parameter pin wouldn't show its custom pin if the value mode was set to 'custom'.
Also, you no longer need to reselect a parameter after adding it to a parameter map get node to make it show the default value

diff --git a/Engine/Plugins/FX/Niagara/Source/NiagaraEditor/Private/Customizations/NiagaraScriptVariableCustomization.cpp b/Engine/Plugins/FX/Niagara/Source/NiagaraEditor/Private/Customizations/NiagaraScriptVariableCustomization.cpp
--- a/Engine/Plugins/FX/Niagara/Source/NiagaraEditor/Private/Customizations/NiagaraScriptVariableCustomization.cpp
+++ b/Engine/Plugins/FX/Niagara/Source/NiagaraEditor/Private/Customizations/NiagaraScriptVariableCustomization.cpp
@@ -43,6 +43,13 @@ FNiagaraScriptVariableDetails::FNiagaraScriptVariableDetails()
 
 FNiagaraScriptVariableDetails::~FNiagaraScriptVariableDetails()
 {
+	if(ParameterMapWatcherDelegate.IsValid() && Variable.IsValid())
+	{
+		if (UNiagaraGraph* Graph = Cast<UNiagaraGraph>(Variable->GetOuter()))
+		{
+			Graph->RemoveOnGraphChangedHandler(ParameterMapWatcherDelegate);	
+		}
+	}
 }
 
 UEdGraphPin* FNiagaraScriptVariableDetails::GetAnyDefaultPin()
@@ -279,6 +286,23 @@ void FNiagaraScriptVariableDetails::CustomizeDetailsParameterDefinitionsSynchron
 	}
 }
 
+void FNiagaraScriptVariableDetails::OnGraphChanged(const FEdGraphEditAction& EdGraphEditAction,	TWeakObjectPtr<UNiagaraScriptVariable> ScriptVariable)
+{
+	if(!ScriptVariable.IsValid())
+	{
+		return;
+	}
+	
+	if(GetAnyDefaultPin() != nullptr)
+	{
+		Refresh();
+		if (UNiagaraGraph* Graph = Cast<UNiagaraGraph>(Variable->GetOuter()))
+		{
+			Graph->RemoveOnGraphChangedHandler(ParameterMapWatcherDelegate);	
+		}
+	}
+}
+
 void FNiagaraScriptVariableDetails::AddGraphDefaultValueCustomRow(IDetailCategoryBuilder& CategoryBuilder)
 {
 	if (UEdGraphPin* Pin = GetAnyDefaultPin())
@@ -317,6 +341,11 @@ void FNiagaraScriptVariableDetails::AddGraphDefaultValueCustomRow(IDetailCategor
 	}
 	else
 	{
+		if (UNiagaraGraph* Graph = Cast<UNiagaraGraph>(Variable->GetOuter()))
+		{
+			ParameterMapWatcherDelegate = Graph->AddOnGraphChangedHandler(FOnGraphChanged::FDelegate::CreateSP(this, &FNiagaraScriptVariableDetails::OnGraphChanged, Variable));
+		}
+		
 		if (Variable->DefaultMode == ENiagaraDefaultMode::Value)
 		{
 			FDetailWidgetRow& DefaultValueWidget = CategoryBuilder.AddCustomRow(LOCTEXT("DefaultValueFilterText", "Default Value"));
@@ -325,7 +354,7 @@ void FNiagaraScriptVariableDetails::AddGraphDefaultValueCustomRow(IDetailCategor
 			[
 				SNew(STextBlock)
 				.Font(FNiagaraEditorStyle::Get().GetFontStyle("NiagaraEditor.ParameterFont"))
-				.Text(NSLOCTEXT("ScriptVariableCustomization", "MissingDefaults", "To set default, add to a Map Get node that is wired to the graph."))
+				.Text(NSLOCTEXT("ScriptVariableCustomization", "MissingDefaults", "To set the default value, add this parameter to a Parameter Map Get node."))
 			];
 		}
 	}
diff --git a/Engine/Plugins/FX/Niagara/Source/NiagaraEditor/Private/Customizations/NiagaraScriptVariableCustomization.h b/Engine/Plugins/FX/Niagara/Source/NiagaraEditor/Private/Customizations/NiagaraScriptVariableCustomization.h
--- a/Engine/Plugins/FX/Niagara/Source/NiagaraEditor/Private/Customizations/NiagaraScriptVariableCustomization.h
+++ b/Engine/Plugins/FX/Niagara/Source/NiagaraEditor/Private/Customizations/NiagaraScriptVariableCustomization.h
@@ -64,6 +64,8 @@ private:
 	void CustomizeDetailsStaticSwitchScriptVariable(IDetailLayoutBuilder& DetailBuilder);
 	void CustomizeDetailsParameterDefinitionsSynchronizedScriptVariable(IDetailLayoutBuilder& DetailBuilder);
 
+	void OnGraphChanged(const FEdGraphEditAction& EdGraphEditAction, TWeakObjectPtr<UNiagaraScriptVariable> ScriptVariable);
+	
 	void AddGraphDefaultValueCustomRow(IDetailCategoryBuilder& CategoryBuilder);
 	void AddLibraryDefaultValueCustomRow(IDetailCategoryBuilder& CategoryBuilder, bool bInLibraryAsset);
 
@@ -94,6 +96,8 @@ private:
 	// Bool to track whether a parameters value has actually changed at any point during OnBeginValueChanged(), OnValueChanged() and OnEndValueChanged() to prevent excessive refreshing.
 	bool bParameterValueChangedDuringOnValueChanged;
 
+	FDelegateHandle ParameterMapWatcherDelegate;
+	
 	// Cached value of enum selector for choosing the library default mode.
 	int32 LibraryDefaultModeValue;
 
diff --git a/Engine/Plugins/FX/Niagara/Source/NiagaraEditor/Private/NiagaraGraph.cpp b/Engine/Plugins/FX/Niagara/Source/NiagaraEditor/Private/NiagaraGraph.cpp
--- a/Engine/Plugins/FX/Niagara/Source/NiagaraEditor/Private/NiagaraGraph.cpp
+++ b/Engine/Plugins/FX/Niagara/Source/NiagaraEditor/Private/NiagaraGraph.cpp
@@ -1026,44 +1026,31 @@ void UNiagaraGraph::ReplaceScriptReferences(UNiagaraScript* OldScript, UNiagaraS
 
 TArray<UEdGraphPin*> UNiagaraGraph::FindParameterMapDefaultValuePins(const FName VariableName) const
 {
-	TArray<UEdGraphPin*> DefaultPins;
+	TArray<UEdGraphPin*> OutDefaultPins;
+	
+	TArray<UNiagaraNodeParameterMapGet*> MapGetNodes;
+	GetNodesOfClass<UNiagaraNodeParameterMapGet>(MapGetNodes);
 
-	TArray<UNiagaraNode*> NodesTraversed;
-	FPinCollectorArray OutputPins;
-	for (UEdGraphNode* Node : Nodes)
+	for(UNiagaraNodeParameterMapGet* MapGet : MapGetNodes)
 	{
-		UNiagaraNodeOutput* OutNode = Cast<UNiagaraNodeOutput>(Node);
-		if (!OutNode)
-		{
-			continue;
-		}
-		NodesTraversed.Reset();
-		BuildTraversal(NodesTraversed, OutNode);
-
-		for (UNiagaraNode* NiagaraNode : NodesTraversed)
+		FPinCollectorArray OutputPins;
+		MapGet->GetOutputPins(OutputPins);
+		for (UEdGraphPin* OutputPin : OutputPins)
 		{
-			UNiagaraNodeParameterMapGet* GetNode = Cast<UNiagaraNodeParameterMapGet>(NiagaraNode);
-			if (!GetNode)
+			if (VariableName != OutputPin->PinName || OutputPin->PinType.PinSubCategory != UNiagaraNodeParameterMapBase::ParameterPinSubCategory)
 			{
 				continue;
 			}
-			OutputPins.Reset();
-			GetNode->GetOutputPins(OutputPins);
-			for (UEdGraphPin* OutputPin : OutputPins)
+			
+			if (UEdGraphPin* Pin = MapGet->GetDefaultPin(OutputPin))
 			{
-				if (VariableName != OutputPin->PinName || OutputPin->PinType.PinSubCategory != UNiagaraNodeParameterMapBase::ParameterPinSubCategory)
-				{
-					continue;
-				}
-				if (UEdGraphPin* Pin = GetNode->GetDefaultPin(OutputPin))
-				{
-					check(Pin->Direction == EEdGraphPinDirection::EGPD_Input);
-					DefaultPins.AddUnique(Pin);
-				}
+				check(Pin->Direction == EEdGraphPinDirection::EGPD_Input);
+				OutDefaultPins.AddUnique(Pin);
 			}
 		}
 	}
-	return DefaultPins;
+	
+	return OutDefaultPins;
 }
 
 // Impacts of this function include changing the default value of UNiagaraScriptVariable::Variable & DefaultMode (note that this won't impact the
diff --git a/Engine/Plugins/FX/Niagara/Source/NiagaraEditor/Public/NiagaraGraph.h b/Engine/Plugins/FX/Niagara/Source/NiagaraEditor/Public/NiagaraGraph.h
--- a/Engine/Plugins/FX/Niagara/Source/NiagaraEditor/Public/NiagaraGraph.h
+++ b/Engine/Plugins/FX/Niagara/Source/NiagaraEditor/Public/NiagaraGraph.h
@@ -277,7 +277,7 @@ class UNiagaraGraph : public UEdGraph
 	UEdGraphPin* FindParameterMapDefaultValuePin(const FName VariableName, ENiagaraScriptUsage InUsage, ENiagaraScriptUsage InParentUsage) const;
 	void MultiFindParameterMapDefaultValuePins(TConstArrayView<FName> VariableNames, ENiagaraScriptUsage InUsage, ENiagaraScriptUsage InParentUsage, TArrayView<UEdGraphPin*> DefaultPins) const;
 
-	/** Walk through the graph for an ParameterMapGet nodes and find all matching default pins for VariableName, irrespective of usage. */
+	/** Find all matching default pins for VariableName, irrespective of usage. */
 	TArray<UEdGraphPin*> FindParameterMapDefaultValuePins(const FName VariableName) const;
 
 	/** Rebuilds the internally stored ParameterRefrenceMap if it has been marked dirty.  Used to explicitly trigger the rebuild rather than
