Fix for non unity compile issue

diff --git a/Engine/Plugins/FX/Niagara/Source/NiagaraEditor/Private/Customizations/NiagaraScriptVariableCustomization.cpp b/Engine/Plugins/FX/Niagara/Source/NiagaraEditor/Private/Customizations/NiagaraScriptVariableCustomization.cpp
--- a/Engine/Plugins/FX/Niagara/Source/NiagaraEditor/Private/Customizations/NiagaraScriptVariableCustomization.cpp
+++ b/Engine/Plugins/FX/Niagara/Source/NiagaraEditor/Private/Customizations/NiagaraScriptVariableCustomization.cpp
@@ -19,6 +19,7 @@
 #include "NiagaraParameterDefinitions.h"
 #include "NiagaraScriptVariable.h"
 #include "SNiagaraParameterEditor.h"
+#include "GraphEditAction.h"
 
 #include UE_INLINE_GENERATED_CPP_BY_NAME(NiagaraScriptVariableCustomization)
 
diff --git a/Engine/Plugins/FX/Niagara/Source/NiagaraEditor/Private/Customizations/NiagaraScriptVariableCustomization.h b/Engine/Plugins/FX/Niagara/Source/NiagaraEditor/Private/Customizations/NiagaraScriptVariableCustomization.h
--- a/Engine/Plugins/FX/Niagara/Source/NiagaraEditor/Private/Customizations/NiagaraScriptVariableCustomization.h
+++ b/Engine/Plugins/FX/Niagara/Source/NiagaraEditor/Private/Customizations/NiagaraScriptVariableCustomization.h
@@ -12,7 +12,7 @@
  
 class IDetailCategoryBuilder;
 class UEdGraphPin;
-
+struct FEdGraphEditAction;
 
 /** Intermediate representations for default mode set on parameter definition script variables. Maps to ENiagaraDefaultMode and bOverrideParameterDefinitionsDefaultValue of UNiagaraScriptVariable. */
 UENUM()
@@ -64,7 +64,7 @@ private:
 	void CustomizeDetailsStaticSwitchScriptVariable(IDetailLayoutBuilder& DetailBuilder);
 	void CustomizeDetailsParameterDefinitionsSynchronizedScriptVariable(IDetailLayoutBuilder& DetailBuilder);
 
-	void OnGraphChanged(const FEdGraphEditAction& EdGraphEditAction, TWeakObjectPtr<UNiagaraScriptVariable> ScriptVariable);
+	void OnGraphChanged(const FEdGraphEditAction& EdGraphEditAction, TWeakObjectPtr<class UNiagaraScriptVariable> ScriptVariable);
 	
 	void AddGraphDefaultValueCustomRow(IDetailCategoryBuilder& CategoryBuilder);
 	void AddLibraryDefaultValueCustomRow(IDetailCategoryBuilder& CategoryBuilder, bool bInLibraryAsset);
