Collection of Niagara load time optimizations
-Updates NiagaraGraph to remove many unecessary references to UNiagaraScriptVaraible (most significantly those that are declared by child graphs).  Makes a dramatic improvement to the number of objects that get serialized with a system and improves load times.
-Improvements to reduce the cost of UNiagaraSystem::RequestCompile() in the case of a warm DDC
-when using baked out RI parameters we now store prepared RI parameters in the DDC so that with a warm DDC we can apply them without having to run the relatively expesnive PrepareRapidIterationParameters step
-some small optimizations around generating FNiagaraVMExecutableDataId

diff --git a/Engine/Plugins/FX/Niagara/Source/Niagara/Private/NiagaraAsyncCompile.h b/Engine/Plugins/FX/Niagara/Source/Niagara/Private/NiagaraAsyncCompile.h
--- a/Engine/Plugins/FX/Niagara/Source/Niagara/Private/NiagaraAsyncCompile.h
+++ b/Engine/Plugins/FX/Niagara/Source/Niagara/Private/NiagaraAsyncCompile.h
@@ -20,7 +20,7 @@ struct FEmitterCompiledScriptPair
 	uint32 PendingJobID = INDEX_NONE; // this is the ID for any active shader compiler worker job
 	FNiagaraVMExecutableDataId CompileId;
 	TSharedPtr<FNiagaraVMExecutableData> CompileResults;
-	double CompileTime = 0;
+	float CompileTime = 0.0f;
 };
 
 UENUM()
