- Move various Emitter deprecated properties into editor only data

#rnx

diff --git a/Engine/Plugins/FX/Niagara/Source/Niagara/Classes/NiagaraEmitter.h b/Engine/Plugins/FX/Niagara/Source/Niagara/Classes/NiagaraEmitter.h
--- a/Engine/Plugins/FX/Niagara/Source/Niagara/Classes/NiagaraEmitter.h
+++ b/Engine/Plugins/FX/Niagara/Source/Niagara/Classes/NiagaraEmitter.h
@@ -940,14 +940,8 @@ private:
 	FOnEmitterCompiled OnGPUScriptCompiledDelegate;
 
 	void RaiseOnEmitterGPUCompiled(UNiagaraScript* InScript, const FGuid& ScriptVersion);
-#endif
-
-	bool bFullyLoaded = false;
 
-	UPROPERTY()
-	FString UniqueEmitterName;
-
-	/** Use property in struct returned from GetEmitterData() instead */ 
+	/** Use property in struct returned from GetEmitterData() instead */
 	UPROPERTY(meta = (DeprecatedProperty))
 	TArray<TObjectPtr<UNiagaraRendererProperties>> RendererProperties_DEPRECATED;
 	
@@ -967,7 +961,6 @@ private:
 	UPROPERTY(meta = (DeprecatedProperty))
 	TArray<FName> SharedEventGeneratorIds_DEPRECATED;
 
-#if WITH_EDITORONLY_DATA
 	UPROPERTY(meta = (DeprecatedProperty))
 	TObjectPtr<UNiagaraEmitter> Parent_DEPRECATED;
 
@@ -977,9 +970,13 @@ private:
 	/** Subscriptions to definitions of parameters. */
 	UPROPERTY()
 	TArray<FParameterDefinitionsSubscription> ParameterDefinitionsSubscriptions;
-
 #endif
 
+	bool bFullyLoaded = false;
+
+	UPROPERTY()
+	FString UniqueEmitterName;
+
 #if WITH_EDITOR
 	FOnPropertiesChanged OnPropertiesChangedDelegate;
 	FOnRenderersChanged OnRenderersChangedDelegate;
