- Change from Load to [] operator and use int3 vs int4

#rnx

diff --git a/Engine/Plugins/FX/Niagara/Source/Niagara/Private/NiagaraDataInterfaceGrid3DCollection.cpp b/Engine/Plugins/FX/Niagara/Source/Niagara/Private/NiagaraDataInterfaceGrid3DCollection.cpp
--- a/Engine/Plugins/FX/Niagara/Source/Niagara/Private/NiagaraDataInterfaceGrid3DCollection.cpp
+++ b/Engine/Plugins/FX/Niagara/Source/Niagara/Private/NiagaraDataInterfaceGrid3DCollection.cpp
@@ -2147,7 +2147,7 @@ bool UNiagaraDataInterfaceGrid3DCollection::GetFunctionHLSL(const FNiagaraDataIn
 			RGBAHlsl.Appendf(TEXT("    val = 0;\n"));
 			RGBAHlsl.Appendf(TEXT("    [branch] if (In_AttributeIndex == 0)\n"));
 			RGBAHlsl.Appendf(TEXT("    {\n"));
-			RGBAHlsl.Appendf(TEXT("		   float{GridChannels} TmpVal = {OutputGrid}.Load(int4(In_IndexX, In_IndexY, In_IndexZ, 0));\n"));
+			RGBAHlsl.Appendf(TEXT("		   float{GridChannels} TmpVal = {OutputGrid}[int3(In_IndexX, In_IndexY, In_IndexZ)];\n"));
 			RGBAHlsl.Appendf(TEXT("        TmpVal.r = In_Value;\n"));
 			RGBAHlsl.Appendf(TEXT("        {OutputGrid}[int3(In_IndexX, In_IndexY, In_IndexZ)] = TmpVal;\n"));
 			RGBAHlsl.Appendf(TEXT("    }\n"));
@@ -2156,7 +2156,7 @@ bool UNiagaraDataInterfaceGrid3DCollection::GetFunctionHLSL(const FNiagaraDataIn
 			{
 				RGBAHlsl.Appendf(TEXT("    else if (In_AttributeIndex == %d)\n"), i + 1);
 				RGBAHlsl.Appendf(TEXT("    {\n"));
-				RGBAHlsl.Appendf(TEXT("		   float{GridChannels} TmpVal = {OutputGrid}.Load(int4(In_IndexX, In_IndexY, In_IndexZ, 0));\n"));
+				RGBAHlsl.Appendf(TEXT("		   float{GridChannels} TmpVal = {OutputGrid}[int3(In_IndexX, In_IndexY, In_IndexZ)];\n"));
 				RGBAHlsl.Appendf(TEXT("        TmpVal.%s = In_Value;\n"), *FGrid3DCollectionAttributeHlslWriter::Channels[i + 1]);
 				RGBAHlsl.Appendf(TEXT("        {OutputGrid}[int3(In_IndexX, In_IndexY, In_IndexZ)] = TmpVal;\n"));
 				RGBAHlsl.Appendf(TEXT("    }\n"));
diff --git a/Engine/Source/Runtime/Core/Private/UObject/UE5ReleaseStreamObjectVersion.cpp b/Engine/Source/Runtime/Core/Private/UObject/UE5ReleaseStreamObjectVersion.cpp
--- a/Engine/Source/Runtime/Core/Private/UObject/UE5ReleaseStreamObjectVersion.cpp
+++ b/Engine/Source/Runtime/Core/Private/UObject/UE5ReleaseStreamObjectVersion.cpp
@@ -12,7 +12,7 @@ TMap<FGuid, FGuid> FUE5ReleaseStreamObjectVersion::GetSystemGuids()
 	SystemGuids.Add(DevGuids.MATERIALSHADERMAP_DERIVEDDATA_VER, FGuid("79090A66F9D94B5BB285D49E5D39468E"));
 	SystemGuids.Add(DevGuids.NANITE_DERIVEDDATA_VER, FGuid("AC88CFBDDA614A9C8488F64D2DAF699D"));
 	SystemGuids.Add(DevGuids.NIAGARASHADERMAP_DERIVEDDATA_VER, FGuid("6BE35B6FACB34568970120F9BC8DAB80"));
-	SystemGuids.Add(DevGuids.Niagara_LatestScriptCompileVersion, FGuid("0CD19494DF6F4B549FC12A0CCC01A02E"));
+	SystemGuids.Add(DevGuids.Niagara_LatestScriptCompileVersion, FGuid("CCA1C9E89704F54E9DF3639F7E75E78E"));
 	SystemGuids.Add(DevGuids.SkeletalMeshDerivedDataVersion, FGuid("25C49E579B3142DDA2A8C14037267679"));
 
 	return SystemGuids;
