Fix up an issue which seems like it could lead to effects being culled if their DI are being dirtied over multiple frames.  Resolve log spammed warning as well.

diff --git a/Engine/Plugins/FX/Niagara/Source/Niagara/Private/NiagaraScriptExecutionContext.cpp b/Engine/Plugins/FX/Niagara/Source/Niagara/Private/NiagaraScriptExecutionContext.cpp
--- a/Engine/Plugins/FX/Niagara/Source/Niagara/Private/NiagaraScriptExecutionContext.cpp
+++ b/Engine/Plugins/FX/Niagara/Source/Niagara/Private/NiagaraScriptExecutionContext.cpp
@@ -663,9 +663,9 @@ bool FNiagaraSystemScriptExecutionContext::Tick(class FNiagaraSystemInstance* In
 			const TArray<UNiagaraDataInterface*>& DataInterfaces = GetDataInterfaces();
 
 			const int32 FunctionCount = ScriptExecutableData.CalledVMExternalFunctions.Num();
-			FunctionTable.Reset();
-			FunctionTable.SetNum(FunctionCount);
-			ExtFunctionInfo.AddDefaulted(FunctionCount);
+
+			FunctionTable.Init(nullptr, FunctionCount);
+			ExtFunctionInfo.Init(FExternalFuncInfo(), FunctionCount);
 
 			const FNiagaraScriptExecutionParameterStore* ScriptParameterStore = Script->GetExecutionReadyParameterStore(ENiagaraSimTarget::CPUSim);
 			check(ScriptParameterStore != nullptr);
