- Force ExecutionState to be captured in renderer only mode otherwise emitter may not show up
- Fix copying volume data if the row / slice strides do not match source data

diff --git a/Engine/Plugins/FX/Niagara/Source/Niagara/Private/NiagaraDataInterfaceRenderTargetVolume.cpp b/Engine/Plugins/FX/Niagara/Source/Niagara/Private/NiagaraDataInterfaceRenderTargetVolume.cpp
--- a/Engine/Plugins/FX/Niagara/Source/Niagara/Private/NiagaraDataInterfaceRenderTargetVolume.cpp
+++ b/Engine/Plugins/FX/Niagara/Source/Niagara/Private/NiagaraDataInterfaceRenderTargetVolume.cpp
@@ -752,7 +752,7 @@ bool UNiagaraDataInterfaceRenderTargetVolume::SimCacheReadFrame(UObject* Storage
 					{
 						if (FRenderTargetVolumeRWInstanceData_RenderThread* InstanceData_RT = RT_Proxy->SystemInstancesToProxyData_RT.Find(RT_InstanceID))
 						{
-							RT_VolumeCacheData->Fill3DTexture(RT_FrameIndex, InstanceData_RT->RenderTarget->GetRHI());
+							RT_VolumeCacheData->Fill3DTexture_RenderThread(RT_FrameIndex, InstanceData_RT->RenderTarget->GetRHI(), RHICmdList);
 						}
 					});
 		}
diff --git a/Engine/Plugins/FX/Niagara/Source/Niagara/Private/NiagaraSimCache.cpp b/Engine/Plugins/FX/Niagara/Source/Niagara/Private/NiagaraSimCache.cpp
--- a/Engine/Plugins/FX/Niagara/Source/Niagara/Private/NiagaraSimCache.cpp
+++ b/Engine/Plugins/FX/Niagara/Source/Niagara/Private/NiagaraSimCache.cpp
@@ -125,6 +125,8 @@ bool UNiagaraSimCache::BeginWrite(FNiagaraSimCacheCreateParameters InCreateParam
 	// When in rendering only mode ask all renderers which attributes we need to capture
 	if ( CreateParameters.AttributeCaptureMode == ENiagaraSimCacheAttributeCaptureMode::RenderingOnly )
 	{
+		FNameBuilder NameBuilder;
+
 	#if WITH_EDITORONLY_DATA
 		for ( const FNiagaraEmitterHandle& EmitterHandle : Helper.NiagaraSystem->GetEmitterHandles() )
 		{
@@ -133,7 +135,7 @@ bool UNiagaraSimCache::BeginWrite(FNiagaraSimCacheCreateParameters InCreateParam
 				{
 					for (FNiagaraVariableBase BoundAttribute : RenderProperties->GetBoundAttributes())
 					{
-						FNameBuilder NameBuilder;
+						NameBuilder.Reset();
 						NameBuilder.Append(EmitterHandle.GetUniqueInstanceName());
 						NameBuilder.AppendChar(TEXT('.'));
 						BoundAttribute.GetName().AppendString(NameBuilder);
@@ -144,6 +146,10 @@ bool UNiagaraSimCache::BeginWrite(FNiagaraSimCacheCreateParameters InCreateParam
 			);
 		}
 	#endif
+		static const FName NAME_ExecutionState("ExecutionState");
+		static const FName NAME_SystemExecutionState("System.ExecutionState");
+		CreateParameters.ExplicitCaptureAttributes.AddUnique(NAME_SystemExecutionState);
+
 		for (const TSharedRef<FNiagaraEmitterInstance, ESPMode::ThreadSafe>& EmitterInstance : Helper.SystemInstance->GetEmitters())
 		{
 			const FNiagaraParameterStore& RendererBindings = EmitterInstance->GetRendererBoundVariables();
@@ -151,6 +157,12 @@ bool UNiagaraSimCache::BeginWrite(FNiagaraSimCacheCreateParameters InCreateParam
 			{
 				CreateParameters.ExplicitCaptureAttributes.AddUnique(Variable.GetName());
 			}
+
+			NameBuilder.Reset();
+			NameBuilder.Append(EmitterInstance->GetEmitterHandle().GetUniqueInstanceName());
+			NameBuilder.AppendChar(TEXT('.'));
+			NAME_ExecutionState.AppendString(NameBuilder);
+			CreateParameters.ExplicitCaptureAttributes.AddUnique(NameBuilder.ToString());
 		}
 
 		FNiagaraDataInterfaceUtilities::ForEachDataInterface(
diff --git a/Engine/Plugins/FX/Niagara/Source/Niagara/Private/VolumeCache.cpp b/Engine/Plugins/FX/Niagara/Source/Niagara/Private/VolumeCache.cpp
--- a/Engine/Plugins/FX/Niagara/Source/Niagara/Private/VolumeCache.cpp
+++ b/Engine/Plugins/FX/Niagara/Source/Niagara/Private/VolumeCache.cpp
@@ -30,14 +30,14 @@ public:
 		DenseGridPtr = nullptr;
 	}
 
-	virtual void Init(FIntVector Resolution);
-	virtual bool LoadFile(FString Path, int frame);
-	virtual bool UnloadFile(int frame);
-	virtual bool LoadRange(FString Path, int Start, int End);
-	virtual void UnloadAll();
-	virtual bool Fill3DTexture_RenderThread(int frame, FTextureRHIRef TextureToFill, FRHICommandListImmediate& RHICmdList);
-	virtual bool Fill3DTexture(int frame, FTextureRHIRef TextureToFill);
-	
+	virtual void Init(FIntVector Resolution) override;
+	virtual bool LoadFile(FString Path, int frame) override;
+	virtual bool UnloadFile(int frame) override;
+	virtual bool LoadRange(FString Path, int Start, int End) override;
+	virtual void UnloadAll() override;
+	virtual bool Fill3DTexture_RenderThread(int frame, FTextureRHIRef TextureToFill, FRHICommandListImmediate& RHICmdList) override;
+	virtual bool Fill3DTexture(int frame, FTextureRHIRef TextureToFill) override;
+
 private:
 	TMap<int32, Vec4Grid::Ptr> OpenVDBGrids;
 	TQueue<int32> CacheQueue;
@@ -306,12 +306,28 @@ bool FOpenVDBCacheData::Fill3DTexture_RenderThread(int frame, FTextureRHIRef Tex
 		else
 		{
 			const FUpdateTextureRegion3D UpdateRegion(0, 0, 0, 0, 0, 0, DenseResolution.X, DenseResolution.Y, DenseResolution.Z);			
-			const SIZE_T MemorySize = static_cast<SIZE_T>(UpdateRegion.Width) * static_cast<SIZE_T>(UpdateRegion.Height) * static_cast<SIZE_T>(UpdateRegion.Depth) * static_cast<SIZE_T>(FormatSize);
-
 			FUpdateTexture3DData TheData = RHICmdList.BeginUpdateTexture3D(TextureToFill, 0, UpdateRegion);
-			
-			FMemory::Memcpy(TheData.Data, DataPtr, MemorySize);
 
+			const SIZE_T RowSize = static_cast<SIZE_T>(UpdateRegion.Width) * static_cast<SIZE_T>(FormatSize);
+			const SIZE_T SliceSize = static_cast<SIZE_T>(UpdateRegion.Height) * RowSize;
+			if (TheData.RowPitch != RowSize || TheData.DepthPitch != SliceSize)
+			{
+				for (int32 z=0; z < DenseResolution.Z; ++z)
+				{
+					uint8* DestRow = TheData.Data + (SIZE_T(TheData.DepthPitch) * SIZE_T(z));
+					for (int32 y=0; y < DenseResolution.Y; ++y)
+					{
+						FMemory::Memcpy(DestRow, DataPtr, RowSize);
+						DestRow += TheData.RowPitch;
+						DataPtr += RowSize;
+					}
+				}
+			}
+			else
+			{
+				const SIZE_T MemorySize = static_cast<SIZE_T>(UpdateRegion.Depth) * SliceSize;
+				FMemory::Memcpy(TheData.Data, DataPtr, MemorySize);
+			}
 			RHICmdList.EndUpdateTexture3D(TheData);						
 		}
 
@@ -325,29 +341,14 @@ bool FOpenVDBCacheData::Fill3DTexture(int frame, FTextureRHIRef TextureToFill)
 {
 	if (OpenVDBGrids.Contains(frame) && OpenVDBGrids[frame] != nullptr)
 	{
-		FScopeLock ScopeLock(&DenseGridGuard);
-		uint8* DataPtr = (uint8*)DenseGridPtr->data();
-
-		const int32 FormatSize = GPixelFormats[TextureToFill->GetFormat()].BlockBytes;
-
-		if (TextureToFill->GetSizeXYZ() != DenseResolution)
-		{
-			UE_LOG(LogVolumeCache, Warning, TEXT("Target texture resolution %s doesn't match OpenVDB grid resolution %s."), *TextureToFill->GetSizeXYZ().ToString(), *DenseResolution.ToString());
-			return false;
-		}
-
-		const FUpdateTextureRegion3D UpdateRegion(0, 0, 0, 0, 0, 0, DenseResolution.X, DenseResolution.Y, DenseResolution.Z);
-		const SIZE_T MemorySize = static_cast<SIZE_T>(UpdateRegion.Width) * static_cast<SIZE_T>(UpdateRegion.Height) * static_cast<SIZE_T>(UpdateRegion.Depth) * static_cast<SIZE_T>(FormatSize);
-
-		FUpdateTexture3DData TheData = RHIBeginUpdateTexture3D(TextureToFill, 0, UpdateRegion);
-
-		FMemory::Memcpy(TheData.Data, DataPtr, MemorySize);
-
-		RHIEndUpdateTexture3D(TheData);		
-
+		ENQUEUE_RENDER_COMMAND(FUpdateData)(
+			[this, frame, TextureToFill](FRHICommandListImmediate& RHICmdList)
+			{
+				Fill3DTexture_RenderThread(frame, TextureToFill, RHICmdList);
+			}
+		);
 		return true;
 	}
-
 	return false;
 }
 
