Add Sim Cache Capture to Niagara Editor

diff --git a/Engine/Plugins/FX/Niagara/Source/Niagara/Classes/NiagaraSimCacheCapture.cpp b/Engine/Plugins/FX/Niagara/Source/Niagara/Classes/NiagaraSimCacheCapture.cpp
--- a/Engine/Plugins/FX/Niagara/Source/Niagara/Classes/NiagaraSimCacheCapture.cpp
+++ b/Engine/Plugins/FX/Niagara/Source/Niagara/Classes/NiagaraSimCacheCapture.cpp
@@ -31,6 +31,7 @@ void FNiagaraSimCacheCapture::CaptureNiagaraSimCache(UNiagaraSimCache* SimCache,
 	WeakCaptureSimCache = SimCache;
 	WeakCaptureComponent = NiagaraComponent;
 	CaptureParameters = InCaptureParameters;
+	CaptureFrameCounter = 0;
 
 	if(SimCache == nullptr || NiagaraComponent == nullptr)
 	{
diff --git a/Engine/Plugins/FX/Niagara/Source/NiagaraEditor/Private/Toolkits/NiagaraSystemToolkit.cpp b/Engine/Plugins/FX/Niagara/Source/NiagaraEditor/Private/Toolkits/NiagaraSystemToolkit.cpp
--- a/Engine/Plugins/FX/Niagara/Source/NiagaraEditor/Private/Toolkits/NiagaraSystemToolkit.cpp
+++ b/Engine/Plugins/FX/Niagara/Source/NiagaraEditor/Private/Toolkits/NiagaraSystemToolkit.cpp
@@ -281,6 +281,9 @@ void FNiagaraSystemToolkit::InitializeInternal(const EToolkitMode::Type Mode, co
 	BakerViewModel = MakeShared<FNiagaraBakerViewModel>();
 	BakerViewModel->Initialize(SystemViewModel);
 
+	SimCacheViewModel = MakeShared<FNiagaraSimCacheViewModel>();
+	SimCacheViewModel->Initialize(nullptr);
+
 	SystemViewModel->OnEmitterHandleViewModelsChanged().AddSP(this, &FNiagaraSystemToolkit::RefreshParameters);
 	SystemViewModel->GetSelectionViewModel()->OnSystemIsSelectedChanged().AddSP(this, &FNiagaraSystemToolkit::OnSystemSelectionChanged);
 	SystemViewModel->GetSelectionViewModel()->OnEmitterHandleIdSelectionChanged().AddSP(this, &FNiagaraSystemToolkit::OnSystemSelectionChanged);
@@ -826,6 +829,11 @@ TSharedPtr<FNiagaraSystemViewModel> FNiagaraSystemToolkit::GetSystemViewModel()
 	return SystemViewModel;
 }
 
+TSharedPtr<FNiagaraSimCacheViewModel> FNiagaraSystemToolkit::GetSimCacheViewModel()
+{
+	return SimCacheViewModel;
+}
+
 TSharedPtr<FNiagaraSystemGraphSelectionViewModel> FNiagaraSystemToolkit::GetSystemGraphSelectionViewModel()
 {
 	return SystemViewModel->GetSystemGraphSelectionViewModel();
diff --git a/Engine/Plugins/FX/Niagara/Source/NiagaraEditor/Private/Toolkits/NiagaraSystemToolkit.h b/Engine/Plugins/FX/Niagara/Source/NiagaraEditor/Private/Toolkits/NiagaraSystemToolkit.h
--- a/Engine/Plugins/FX/Niagara/Source/NiagaraEditor/Private/Toolkits/NiagaraSystemToolkit.h
+++ b/Engine/Plugins/FX/Niagara/Source/NiagaraEditor/Private/Toolkits/NiagaraSystemToolkit.h
@@ -10,6 +10,7 @@
 
 #include "ISequencer.h"
 #include "ISequencerTrackEditor.h"
+#include "ViewModels/NiagaraSimCacheViewModel.h"
 
 #include "ViewModels/NiagaraEmitterViewModel.h"
 #include "Widgets/SItemSelector.h"
@@ -96,6 +97,8 @@ public:
 
 	TSharedPtr<FNiagaraSystemViewModel> GetSystemViewModel();
 	TSharedPtr<FNiagaraSystemGraphSelectionViewModel> GetSystemGraphSelectionViewModel();
+
+	TSharedPtr<FNiagaraSimCacheViewModel> GetSimCacheViewModel();
 	
 	// @todo This is a hack for now until we reconcile the default toolbar with application modes [duplicated from counterpart in Blueprint Editor]
 	void RegisterToolbarTab(const TSharedRef<class FTabManager>& TabManager);
@@ -249,6 +252,9 @@ private:
 	TSharedPtr<FNiagaraSystemToolkitParameterDefinitionsPanelViewModel> ParameterDefinitionsPanelViewModel;
 	TSharedPtr<class SNiagaraParameterPanel> ParameterPanel;
 
+	/** Sim Cache editor view model**/
+	TSharedPtr<FNiagaraSimCacheViewModel> SimCacheViewModel;
+
 	TSharedPtr<FNiagaraObjectSelection> ObjectSelectionForParameterMapView;
 
 	bool bChangesDiscarded;
diff --git a/Engine/Plugins/FX/Niagara/Source/NiagaraEditor/Private/Toolkits/SystemToolkitModes/NiagaraSystemToolkitModeBase.cpp b/Engine/Plugins/FX/Niagara/Source/NiagaraEditor/Private/Toolkits/SystemToolkitModes/NiagaraSystemToolkitModeBase.cpp
--- a/Engine/Plugins/FX/Niagara/Source/NiagaraEditor/Private/Toolkits/SystemToolkitModes/NiagaraSystemToolkitModeBase.cpp
+++ b/Engine/Plugins/FX/Niagara/Source/NiagaraEditor/Private/Toolkits/SystemToolkitModes/NiagaraSystemToolkitModeBase.cpp
@@ -32,6 +32,10 @@
 #include "ViewModels/NiagaraScratchPadScriptViewModel.h"
 #include "Widgets/SNiagaraSelectedObjectsDetails.h"
 #include "NiagaraObjectSelection.h"
+#include "Widgets/SNiagaraDebugCaptureView.h"
+#include "Widgets/SNiagaraSimCacheView.h"
+#include "Widgets/SNiagaraSimCacheViewTimeline.h"
+#include "Widgets/SNiagaraSimCacheViewTransportControls.h"
 #include "Widgets/SNiagaraSystemUserParameters.h"
 #include "Customizations/NiagaraComponentDetails.h"
 #include "Styling/StyleColors.h"
@@ -50,6 +54,8 @@ const FName FNiagaraSystemToolkitModeBase::SystemParameterDefinitionsTabID(TEXT(
 const FName FNiagaraSystemToolkitModeBase::SelectedEmitterStackTabID(TEXT("NiagaraSystemEditor_SelectedEmitterStack"));
 const FName FNiagaraSystemToolkitModeBase::SelectedEmitterGraphTabID(TEXT("NiagaraSystemEditor_SelectedEmitterGraph"));
 const FName FNiagaraSystemToolkitModeBase::DebugSpreadsheetTabID(TEXT("NiagaraSystemEditor_DebugAttributeSpreadsheet"));
+const FName FNiagaraSystemToolkitModeBase::DebugCaptureTabID(TEXT("NiagaraSystemEditor_DebugCacheCapture"));
+const FName FNiagaraSystemToolkitModeBase::DebugCacheSpreadsheetTabID(TEXT("NiagaraSystemEditor_DebugCacheSpreadsheet"));
 const FName FNiagaraSystemToolkitModeBase::PreviewSettingsTabId(TEXT("NiagaraSystemEditor_PreviewSettings"));
 const FName FNiagaraSystemToolkitModeBase::GeneratedCodeTabID(TEXT("NiagaraSystemEditor_GeneratedCode"));
 const FName FNiagaraSystemToolkitModeBase::MessageLogTabID(TEXT("NiagaraSystemEditor_MessageLog"));
@@ -266,6 +272,16 @@ void FNiagaraSystemToolkitModeBase::RegisterTabFactories(TSharedPtr<FTabManager>
 		.SetGroup(WorkspaceMenuCategory.ToSharedRef())
 		.SetIcon(FSlateIcon(FNiagaraEditorStyle::Get().GetStyleSetName(), "Tab.Spreadsheet"));
 
+	InTabManager->RegisterTabSpawner(DebugCaptureTabID, FOnSpawnTab::CreateSP(this, &FNiagaraSystemToolkitModeBase::SpawnTab_DebugCapture))
+		.SetDisplayName(LOCTEXT("DebugCapture", "Attribute Capture"))
+		.SetGroup(WorkspaceMenuCategory.ToSharedRef())
+		.SetIcon(FSlateIcon(FNiagaraEditorStyle::Get().GetStyleSetName(), "Tab.Spreadsheet"));
+
+	InTabManager->RegisterTabSpawner(DebugCacheSpreadsheetTabID, FOnSpawnTab::CreateSP(this, &FNiagaraSystemToolkitModeBase::SpawnTab_DebugCacheSpreadsheet))
+		.SetDisplayName(LOCTEXT("DebugSpreadshseet", "Attribute Spreadsheet"))
+		.SetGroup(WorkspaceMenuCategory.ToSharedRef())
+		.SetIcon(FSlateIcon(FNiagaraEditorStyle::Get().GetStyleSetName(), "Tab.Spreadsheet"));
+
 	InTabManager->RegisterTabSpawner(PreviewSettingsTabId, FOnSpawnTab::CreateSP(this, &FNiagaraSystemToolkitModeBase::SpawnTab_PreviewSettings))
 		.SetDisplayName(LOCTEXT("PreviewSceneSettingsTab", "Preview Scene Settings"))
 		.SetGroup(WorkspaceMenuCategory.ToSharedRef())
@@ -659,6 +675,74 @@ TSharedRef<SDockTab> FNiagaraSystemToolkitModeBase::SpawnTab_DebugSpreadsheet(co
 	return SpawnedTab;
 }
 
+TSharedRef<SDockTab> FNiagaraSystemToolkitModeBase::SpawnTab_DebugCapture(const FSpawnTabArgs& Args)
+{
+	check(Args.GetTabId().TabType == DebugCaptureTabID);
+	TSharedPtr<SNiagaraDebugCaptureView> DebugCaptureTab;
+
+	TSharedRef<SDockTab> SpawnedTab =
+		SNew(SDockTab)
+		[
+			SAssignNew(DebugCaptureTab, SNiagaraDebugCaptureView, SystemToolkit.Pin()->GetSystemViewModel().ToSharedRef(), SystemToolkit.Pin()->GetSimCacheViewModel().ToSharedRef())
+		];
+
+	DebugCaptureTab->OnRequestSpreadsheetTab().BindLambda([this]()
+	{
+		SystemToolkit.Pin()->GetTabManager()->TryInvokeTab(DebugCacheSpreadsheetTabID);
+	});
+	
+	return SpawnedTab;
+}
+
+TSharedRef<SDockTab> FNiagaraSystemToolkitModeBase::SpawnTab_DebugCacheSpreadsheet(const FSpawnTabArgs& Args)
+{
+	check(Args.GetTabId().TabType == DebugCacheSpreadsheetTabID);
+
+	TSharedRef<SNiagaraSimCacheView> SimCacheSpreadsheetView =
+		SNew(SNiagaraSimCacheView)
+		.SimCacheViewModel(SystemToolkit.Pin()->GetSimCacheViewModel());
+
+	const TSharedRef<SVerticalBox> Contents = SNew(SVerticalBox);
+
+	Contents->AddSlot()
+		.FillHeight(1.0f)
+		[
+			SNew(SVerticalBox)
+			+SVerticalBox::Slot()
+			[
+				SimCacheSpreadsheetView
+			]
+			+SVerticalBox::Slot()
+			.AutoHeight()
+			[
+				SNew(SHorizontalBox)
+				+SHorizontalBox::Slot()
+				.AutoWidth()
+				.Padding(2.0f)
+				[
+					SNew(SNiagaraSimCacheViewTransportControls)
+					.WeakViewModel(SystemToolkit.Pin()->GetSimCacheViewModel())
+				]
+				+SHorizontalBox::Slot()
+				.Padding(4.0f)
+				[
+					SNew(SNiagaraSimCacheViewTimeline)
+					.WeakViewModel(SystemToolkit.Pin()->GetSimCacheViewModel())
+				]
+				
+			]
+			
+		];
+
+	TSharedRef<SDockTab> SpawnedTab =
+		SNew(SDockTab)
+		[
+			Contents
+		];
+
+	return SpawnedTab;
+}
+
 TSharedRef<SDockTab> FNiagaraSystemToolkitModeBase::SpawnTab_GeneratedCode(const FSpawnTabArgs& Args)
 {
 	check(Args.GetTabId().TabType == GeneratedCodeTabID);
diff --git a/Engine/Plugins/FX/Niagara/Source/NiagaraEditor/Private/Toolkits/SystemToolkitModes/NiagaraSystemToolkitModeBase.h b/Engine/Plugins/FX/Niagara/Source/NiagaraEditor/Private/Toolkits/SystemToolkitModes/NiagaraSystemToolkitModeBase.h
--- a/Engine/Plugins/FX/Niagara/Source/NiagaraEditor/Private/Toolkits/SystemToolkitModes/NiagaraSystemToolkitModeBase.h
+++ b/Engine/Plugins/FX/Niagara/Source/NiagaraEditor/Private/Toolkits/SystemToolkitModes/NiagaraSystemToolkitModeBase.h
@@ -35,6 +35,8 @@ private:
 	TSharedRef<SDockTab> SpawnTab_SelectedEmitterStack(const FSpawnTabArgs& Args);
 	TSharedRef<SDockTab> SpawnTab_SelectedEmitterGraph(const FSpawnTabArgs& Args);
 	TSharedRef<SDockTab> SpawnTab_DebugSpreadsheet(const FSpawnTabArgs& Args);
+	TSharedRef<SDockTab> SpawnTab_DebugCapture(const FSpawnTabArgs& Args);
+	TSharedRef<SDockTab> SpawnTab_DebugCacheSpreadsheet(const FSpawnTabArgs& Args);
 	TSharedRef<SDockTab> SpawnTab_PreviewSettings(const FSpawnTabArgs& Args);
 	TSharedRef<SDockTab> SpawnTab_GeneratedCode(const FSpawnTabArgs& Args);
 	TSharedRef<SDockTab> SpawnTab_MessageLog(const FSpawnTabArgs& Args);
@@ -77,6 +79,8 @@ public:
 	static const FName SelectedEmitterStackTabID;
 	static const FName SelectedEmitterGraphTabID;
 	static const FName DebugSpreadsheetTabID;
+	static const FName DebugCaptureTabID;
+	static const FName DebugCacheSpreadsheetTabID;
 	static const FName PreviewSettingsTabId;
 	static const FName GeneratedCodeTabID;
 	static const FName MessageLogTabID;
diff --git a/Engine/Plugins/FX/Niagara/Source/NiagaraEditor/Private/ViewModels/NiagaraSimCacheViewModel.cpp b/Engine/Plugins/FX/Niagara/Source/NiagaraEditor/Private/ViewModels/NiagaraSimCacheViewModel.cpp
--- a/Engine/Plugins/FX/Niagara/Source/NiagaraEditor/Private/ViewModels/NiagaraSimCacheViewModel.cpp
+++ b/Engine/Plugins/FX/Niagara/Source/NiagaraEditor/Private/ViewModels/NiagaraSimCacheViewModel.cpp
@@ -85,9 +85,25 @@ void FNiagaraSimCacheViewModel::SetupPreviewComponentAndInstance()
 	
 }
 
+TConstArrayView<FNiagaraSimCacheViewModel::FComponentInfo> FNiagaraSimCacheViewModel::GetComponentInfos(
+	int32 InEmitterIndex) const
+{
+	if( InEmitterIndex == INDEX_NONE)
+	{
+		return SystemComponentInfos;
+	}
+
+	if(EmitterComponentInfos.IsValidIndex(InEmitterIndex))
+	{
+		return EmitterComponentInfos[InEmitterIndex];
+	}
+	
+	return TConstArrayView<FComponentInfo>();
+}
+
 FText FNiagaraSimCacheViewModel::GetComponentText(const FName ComponentName, const int32 InstanceIndex) const
 {
-	const FComponentInfo* ComponentInfo = ComponentInfos->FindByPredicate([ComponentName](const FComponentInfo& FoundInfo) { return FoundInfo.Name == ComponentName; });
+	const FComponentInfo* ComponentInfo = GetCurrentComponentInfos().FindByPredicate([ComponentName](const FComponentInfo& FoundInfo) { return FoundInfo.Name == ComponentName; });
 
 	if (ComponentInfo)
 	{
@@ -207,6 +223,7 @@ void FNiagaraSimCacheViewModel::OnCacheModified(UNiagaraSimCache* SimCache)
 	{
 		if ( ThisSimCache == SimCache )
 		{
+			SetFrameIndex(0);
 			UpdateComponentInfos();
 			UpdateCachedFrame();
 			OnSimCacheChangedDelegate.Broadcast();
@@ -241,8 +258,6 @@ void FNiagaraSimCacheViewModel::UpdateCachedFrame()
 
 	NumInstances = EmitterIndex == INDEX_NONE ? 1 : SimCache->GetEmitterNumInstances(EmitterIndex, FrameIndex);
 	const FName EmitterName = EmitterIndex == INDEX_NONE ? NAME_None : SimCache->GetEmitterName(EmitterIndex);
-	// Cached frame data needs to match the component info being viewed. Update here so they match.
-	ComponentInfos = EmitterIndex == INDEX_NONE ? &SystemComponentInfos : &EmitterComponentInfos[EmitterIndex];
 
 	SimCache->ForEachEmitterAttribute(EmitterIndex,
 		[&](const FNiagaraSimCacheVariable& Variable)
@@ -413,11 +428,25 @@ void FNiagaraSimCacheViewModel::BuildEntries(TWeakPtr<SNiagaraSimCacheTreeView>
 	const TSharedRef<FNiagaraSimCacheTreeItem> SharedSystemTreeItem = MakeShared<FNiagaraSimCacheTreeItem>(OwningTreeView);
 	const TSharedRef<FNiagaraSimCacheOverviewSystemItem> SharedSystemBufferItem = MakeShared<FNiagaraSimCacheOverviewSystemItem>();
 
-	SharedSystemTreeItem->SetDisplayName(LOCTEXT("SystemInstance", "System Instance"));
-	SharedSystemBufferItem->SetDisplayName(LOCTEXT("SystemInstance", "System Instance"));
+	if(IsCacheValid())
+	{
+		SharedSystemTreeItem->SetDisplayName(LOCTEXT("SystemInstance", "System Instance"));
+		SharedSystemBufferItem->SetDisplayName(LOCTEXT("SystemInstance", "System Instance"));
+	}
+	else
+	{
+		SharedSystemTreeItem->SetDisplayName(LOCTEXT("InvalidCache", "Invalid Cache"));
+		SharedSystemBufferItem->SetDisplayName(LOCTEXT("InvalidCache", "Invalid Cache"));
+	}
 
 	RootEntries.Add(SharedSystemTreeItem);
 	BufferEntries.Add(SharedSystemBufferItem);
+
+	if(!IsCacheValid())
+	{
+		UpdateCurrentEntries();
+		return;
+	}
 		
 	BuildTreeItemChildren(SharedSystemTreeItem, OwningTreeView);
 		
diff --git a/Engine/Plugins/FX/Niagara/Source/NiagaraEditor/Private/ViewModels/NiagaraSimCacheViewModel.h b/Engine/Plugins/FX/Niagara/Source/NiagaraEditor/Private/ViewModels/NiagaraSimCacheViewModel.h
--- a/Engine/Plugins/FX/Niagara/Source/NiagaraEditor/Private/ViewModels/NiagaraSimCacheViewModel.h
+++ b/Engine/Plugins/FX/Niagara/Source/NiagaraEditor/Private/ViewModels/NiagaraSimCacheViewModel.h
@@ -41,9 +41,9 @@ public:
 
 	void SetupPreviewComponentAndInstance();
 
-	TConstArrayView<FComponentInfo> GetCurrentComponentInfos() const { return *ComponentInfos; }
+	TConstArrayView<FComponentInfo> GetCurrentComponentInfos() const { return GetComponentInfos(EmitterIndex); }
 
-	TConstArrayView<FComponentInfo> GetComponentInfos(int32 InEmitterIndex) const { return  InEmitterIndex == INDEX_NONE ? SystemComponentInfos : EmitterComponentInfos[InEmitterIndex]; }
+	TConstArrayView<FComponentInfo> GetComponentInfos(int32 InEmitterIndex) const;
 
 	TConstArrayView<FString> GetComponentFilters() const { return ComponentFilterArray; }
 
@@ -126,9 +126,6 @@ private:
 	// Number of particles in the given frame
 	int32 NumInstances = 0;
 
-	// Info about the variables in the cache that are currently being viewed in the spreadsheet
-	TArray<FComponentInfo>* ComponentInfos = &SystemComponentInfos;
-
 	// Cached Component infos for this system
 	TArray<FComponentInfo>				SystemComponentInfos;
 	TArray<TArray<FComponentInfo>>		EmitterComponentInfos;
diff --git a/Engine/Plugins/FX/Niagara/Source/NiagaraEditor/Private/Widgets/SNiagaraDebugCaptureView.cpp b/Engine/Plugins/FX/Niagara/Source/NiagaraEditor/Private/Widgets/SNiagaraDebugCaptureView.cpp
new file mode 100644
--- /dev/null
+++ b/Engine/Plugins/FX/Niagara/Source/NiagaraEditor/Private/Widgets/SNiagaraDebugCaptureView.cpp
@@ -0,0 +1,196 @@
+﻿// Copyright Epic Games, Inc. All Rights Reserved.
+
+#include "SNiagaraDebugCaptureView.h"
+
+#include "PropertyEditorModule.h"
+#include "ViewModels/NiagaraEmitterHandleViewModel.h"
+#include "ViewModels/NiagaraSystemSelectionViewModel.h"
+#include "Widgets/Input/SSpinBox.h"
+#include "Widgets/Input/SButton.h"
+#include "Widgets/Layout/SUniformGridPanel.h"
+
+#define LOCTEXT_NAMESPACE "NiagaraDebugCaptureView"
+
+void SNiagaraDebugCaptureView::OnNumFramesChanged(int32 InNumFrames)
+{
+	NumFrames = FMath::Max(1, InNumFrames);
+}
+
+void SNiagaraDebugCaptureView::Construct(const FArguments& InArgs, const TSharedRef<FNiagaraSystemViewModel> InSystemViewModel, const TSharedRef<FNiagaraSimCacheViewModel> InSimCacheViewModel)
+{
+	TargetComponent = InSystemViewModel->GetPreviewComponent();
+	SimCacheViewModel = InSimCacheViewModel;
+	SystemViewModel = InSystemViewModel;
+
+	CapturedCache = NewObject<UNiagaraSimCache>(GetTransientPackage());
+	SimCacheViewModel.Get()->Initialize(CapturedCache);
+	CapturedCache->SetFlags(RF_Standalone);
+
+	FDetailsViewArgs DetailsViewArgs;
+	DetailsViewArgs.bShowOptions = false;
+	
+	ChildSlot
+	[
+		SNew(SSplitter)
+		.Orientation(Orient_Vertical)
+		+SSplitter::Slot()
+		.SizeRule(SSplitter::SizeToContent)
+		[
+			// Capture controls
+			SNew(SSplitter)
+			.Orientation(Orient_Vertical)
+			+SSplitter::Slot()
+			.Resizable(false)
+			.SizeRule(SSplitter::SizeToContent)
+			[
+				// Header
+				SNew(SBorder)
+				.BorderImage(FAppStyle::GetNoBrush())
+				.Padding(5.0f)
+				[
+					SNew(STextBlock)
+					.Justification(ETextJustify::Center)
+					.Text(LOCTEXT("CaptureControlsLabel", "Capture"))
+				]
+			]
+			+SSplitter::Slot()
+			.SizeRule(SSplitter::SizeToContent)
+			[
+				SNew(SUniformGridPanel)
+				.SlotPadding(2.0f)
+				// Capture Single Frame
+				+ SUniformGridPanel::Slot(0, 0)
+				[
+					SNew(SButton)
+					.OnClicked(this, &SNiagaraDebugCaptureView::OnSingleFrameSelected)
+					.ToolTipText(LOCTEXT("SingleFrameButtonTooltip", "Capture a single frame."))
+					.HAlign(HAlign_Center)
+					.VAlign(VAlign_Center)
+					[
+						SNew(STextBlock)
+						.Text(LOCTEXT("SingleFrameButton", "Single"))
+					]
+				]
+				// Capture Multi Frame
+				+SUniformGridPanel::Slot(1, 0)
+				[
+					SNew(SButton)
+					.OnClicked(this, &SNiagaraDebugCaptureView::OnMultiFrameSelected)
+					.ToolTipText(LOCTEXT("MultiFrameButtonTooltip", "Capture multiple frames."))
+					.HAlign(HAlign_Center)
+					.VAlign(VAlign_Center)
+					[
+						SNew(STextBlock)
+						.Text(LOCTEXT("MultiFrameButton", "Multi"))
+					]
+				]
+				// Capture Until Complete
+				+SUniformGridPanel::Slot(2, 0)
+				[
+					SNew(SHorizontalBox)
+					+SHorizontalBox::Slot()
+					.HAlign(HAlign_Center)
+					.VAlign(VAlign_Center)
+					[
+						SNew(STextBlock)
+						.Text(LOCTEXT("NumberOfFrames", "Frames"))
+					]
+					+SHorizontalBox::Slot()
+					[
+						SNew(SSpinBox<int32>)
+						.Value(1)
+						.MinValue(1)
+						.OnValueChanged(this, &SNiagaraDebugCaptureView::OnNumFramesChanged)
+					]
+				]
+			]
+		]
+		+SSplitter::Slot()
+		.Value(0.99f)
+		[
+			SNew(SNiagaraSimCacheOverview)
+			.SimCacheViewModel(SimCacheViewModel)
+		]
+	];
+}
+
+SNiagaraDebugCaptureView::~SNiagaraDebugCaptureView()
+{
+	if(CapturedCache)
+	{
+		CapturedCache->ClearFlags(RF_Standalone);
+		CapturedCache->MarkAsGarbage();
+	}
+}
+
+FReply SNiagaraDebugCaptureView::OnSingleFrameSelected()
+{
+	if(!bIsCaptureActive && TargetComponent && CapturedCache)
+	{
+		SystemViewModel.Get()->GetSequencer().Get()->OnPlay(false);
+
+		const FNiagaraSimCacheCreateParameters CreateParameters;
+
+		SimCacheCapture.CaptureCurrentFrameImmediate(CapturedCache, CreateParameters, TargetComponent, CapturedCache, true);
+
+		if(CapturedCache)
+		{
+			OnCaptureComplete(CapturedCache);
+		}
+		
+	}
+	return FReply::Handled();
+}
+
+FReply SNiagaraDebugCaptureView::OnMultiFrameSelected()
+{
+	if(!bIsCaptureActive && TargetComponent && CapturedCache)
+	{
+		const FNiagaraSimCacheCreateParameters CreateParameters;
+		
+		FNiagaraSimCacheCaptureParameters CaptureParameters;
+		CaptureParameters.NumFrames = NumFrames;
+
+		SystemViewModel.Get()->GetSequencer().Get()->OnPlay(false);
+		
+		bIsCaptureActive = true;
+
+		SimCacheCapture.CaptureNiagaraSimCache(CapturedCache, CreateParameters, TargetComponent, CaptureParameters);
+		
+		SimCacheCapture.OnCaptureComplete().AddSP(this, &SNiagaraDebugCaptureView::OnCaptureComplete);
+	}
+
+	return FReply::Handled();
+}
+
+void SNiagaraDebugCaptureView::OnCaptureComplete(UNiagaraSimCache* CapturedSimCache)
+{
+	bIsCaptureActive = false;
+	SystemViewModel.Get()->GetSequencer().Get()->Pause();
+				
+	TArray<FGuid> SelectedEmitterHandleIds = SystemViewModel->GetSelectionViewModel()->GetSelectedEmitterHandleIds();
+
+	if(SelectedEmitterHandleIds.Num())
+	{
+		TSharedPtr<FNiagaraEmitterHandleViewModel> SelectedEmitterHandleViewModel = SystemViewModel->GetEmitterHandleViewModelById(SelectedEmitterHandleIds[0]);
+
+		if(SelectedEmitterHandleViewModel.IsValid())
+		{
+			for(int32 i = 0; i < SimCacheViewModel->GetNumEmitterLayouts(); ++i)
+			{
+				if(SimCacheViewModel->GetEmitterLayoutName(i) == SelectedEmitterHandleViewModel->GetName())
+				{
+					if(SimCacheViewModel->GetEmitterIndex() != i)
+					{
+						SimCacheViewModel->SetEmitterIndex(i);
+					}
+					break;
+				}
+			}
+		}
+	}
+
+	RequestSpreadsheetTab.ExecuteIfBound();
+}
+
+#undef LOCTEXT_NAMESPACE
diff --git a/Engine/Plugins/FX/Niagara/Source/NiagaraEditor/Private/Widgets/SNiagaraDebugCaptureView.h b/Engine/Plugins/FX/Niagara/Source/NiagaraEditor/Private/Widgets/SNiagaraDebugCaptureView.h
new file mode 100644
--- /dev/null
+++ b/Engine/Plugins/FX/Niagara/Source/NiagaraEditor/Private/Widgets/SNiagaraDebugCaptureView.h
@@ -0,0 +1,41 @@
+﻿// Copyright Epic Games, Inc. All Rights Reserved.
+
+#pragma once
+#include "NiagaraSimCacheCapture.h"
+#include "ViewModels/NiagaraSimCacheViewModel.h"
+#include "ViewModels/NiagaraSystemViewModel.h"
+
+class FNiagaraSimCacheCapture;
+class UNiagaraSimCache;
+
+class SNiagaraDebugCaptureView : public SCompoundWidget
+{
+public:
+	SLATE_BEGIN_ARGS(SNiagaraDebugCaptureView)
+	{}
+
+	SLATE_END_ARGS();
+
+	DECLARE_DELEGATE(FOnRequestSpreadsheetTab)
+
+	void OnNumFramesChanged(int32 InNumFrames);
+	void Construct(const FArguments& InArgs, const TSharedRef<FNiagaraSystemViewModel> SystemViewModel, const TSharedRef<FNiagaraSimCacheViewModel> SimCacheViewModel);
+	virtual ~SNiagaraDebugCaptureView() override;
+	FOnRequestSpreadsheetTab& OnRequestSpreadsheetTab() {return RequestSpreadsheetTab; }
+
+protected:
+	FReply OnSingleFrameSelected();
+	FReply OnMultiFrameSelected();
+	void OnCaptureComplete(UNiagaraSimCache* CapturedSimCache);
+	
+	FNiagaraSimCacheCapture SimCacheCapture;
+	UNiagaraSimCache* CapturedCache = nullptr;
+
+	UNiagaraComponent* TargetComponent = nullptr;
+	int32 NumFrames = 1;
+	bool bIsCaptureActive = false;
+	TSharedPtr<FNiagaraSimCacheViewModel> SimCacheViewModel;
+	TSharedPtr<FNiagaraSystemViewModel> SystemViewModel;
+	FOnRequestSpreadsheetTab RequestSpreadsheetTab;
+	
+};
diff --git a/Engine/Plugins/FX/Niagara/Source/NiagaraEditor/Private/Widgets/SNiagaraSimCacheOverview.cpp b/Engine/Plugins/FX/Niagara/Source/NiagaraEditor/Private/Widgets/SNiagaraSimCacheOverview.cpp
--- a/Engine/Plugins/FX/Niagara/Source/NiagaraEditor/Private/Widgets/SNiagaraSimCacheOverview.cpp
+++ b/Engine/Plugins/FX/Niagara/Source/NiagaraEditor/Private/Widgets/SNiagaraSimCacheOverview.cpp
@@ -33,6 +33,14 @@ public:
 };
 
 
+void SNiagaraSimCacheOverview::OnSimCacheChanged()
+{
+	if(BufferListView.IsValid())
+	{
+		BufferListView->RebuildList();
+	}
+}
+
 void SNiagaraSimCacheOverview::Construct(const FArguments& InArgs)
 {
 	ViewModel = InArgs._SimCacheViewModel;
@@ -45,6 +53,8 @@ void SNiagaraSimCacheOverview::Construct(const FArguments& InArgs)
 	.OnGenerateRow(this, &SNiagaraSimCacheOverview::OnGenerateRowForItem)
 	.OnSelectionChanged(this, &SNiagaraSimCacheOverview::OnListSelectionChanged);
 
+	ViewModel.Get()->OnSimCacheChanged().AddSP(this, &SNiagaraSimCacheOverview::OnSimCacheChanged);
+
 	const float MinSplitterSlotSize = 30.0f; 
 
 	ChildSlot
diff --git a/Engine/Plugins/FX/Niagara/Source/NiagaraEditor/Private/Widgets/SNiagaraSimCacheOverview.h b/Engine/Plugins/FX/Niagara/Source/NiagaraEditor/Private/Widgets/SNiagaraSimCacheOverview.h
--- a/Engine/Plugins/FX/Niagara/Source/NiagaraEditor/Private/Widgets/SNiagaraSimCacheOverview.h
+++ b/Engine/Plugins/FX/Niagara/Source/NiagaraEditor/Private/Widgets/SNiagaraSimCacheOverview.h
@@ -77,6 +77,7 @@ public:
 		SLATE_ARGUMENT(TSharedPtr<FNiagaraSimCacheViewModel>, SimCacheViewModel)
 	SLATE_END_ARGS()
 
+	void OnSimCacheChanged();
 	void Construct(const FArguments& InArgs);
 
 	TSharedRef<ITableRow> OnGenerateRowForItem(TSharedRef<FNiagaraSimCacheOverviewItem> Item, const TSharedRef<STableViewBase>& Owner);
diff --git a/Engine/Plugins/FX/Niagara/Source/NiagaraEditor/Private/Widgets/SNiagaraSimCacheTreeView.cpp b/Engine/Plugins/FX/Niagara/Source/NiagaraEditor/Private/Widgets/SNiagaraSimCacheTreeView.cpp
--- a/Engine/Plugins/FX/Niagara/Source/NiagaraEditor/Private/Widgets/SNiagaraSimCacheTreeView.cpp
+++ b/Engine/Plugins/FX/Niagara/Source/NiagaraEditor/Private/Widgets/SNiagaraSimCacheTreeView.cpp
@@ -227,6 +227,17 @@ bool SSimCacheTreeViewVisibilityWidget::IsItemSelected() const
 
 ///// Tree View Widget
 
+
+void SNiagaraSimCacheTreeView::SetupRootEntries()
+{
+	TArray<TSharedRef<FNiagaraSimCacheTreeItem>>* RootEntries = ViewModel->GetCurrentRootEntries();
+
+	if(RootEntries && !RootEntries->IsEmpty())
+	{
+		TreeView->SetItemExpansion((*RootEntries)[0], true);
+	}
+}
+
 void SNiagaraSimCacheTreeView::Construct(const FArguments& InArgs)
 {
 	constexpr float ItemHeight = 50.0f;
@@ -234,11 +245,9 @@ void SNiagaraSimCacheTreeView::Construct(const FArguments& InArgs)
 	ViewModel = InArgs._SimCacheViewModel;
 
 	ViewModel->OnBufferChanged().AddSP(this, &SNiagaraSimCacheTreeView::OnBufferChanged);
+	ViewModel->OnSimCacheChanged().AddSP(this, &SNiagaraSimCacheTreeView::OnSimCacheChanged);
 
 	ViewModel->BuildEntries(SharedThis(this));
-
-	TArray<TSharedRef<FNiagaraSimCacheTreeItem>>* RootEntries = ViewModel->GetCurrentRootEntries();
-
 	
 	TreeView = SNew(STreeView<TSharedRef<FNiagaraSimCacheTreeItem>>)
 	.ItemHeight(ItemHeight)
@@ -247,11 +256,7 @@ void SNiagaraSimCacheTreeView::Construct(const FArguments& InArgs)
 	.OnGenerateRow(this, &SNiagaraSimCacheTreeView::OnGenerateRow)
 	.OnGetChildren(this, &SNiagaraSimCacheTreeView::OnGetChildren);
 
-	if(RootEntries && !RootEntries->IsEmpty())
-	{
-		TreeView->SetItemExpansion((*RootEntries)[0], true);
-	}
-
+	SetupRootEntries();
 
 	ChildSlot
 	[
@@ -295,6 +300,12 @@ void SNiagaraSimCacheTreeView::OnBufferChanged()
 	TreeView->SetItemExpansion((*ViewModel->GetCurrentRootEntries())[0], true);
 }
 
+void SNiagaraSimCacheTreeView::OnSimCacheChanged()
+{
+	ViewModel->BuildEntries(SharedThis(this));
+	SetupRootEntries();
+}
+
 void SNiagaraSimCacheTreeView::RecursiveAddToSelectionFilter(TArray<TSharedRef<FNiagaraSimCacheTreeItem>>& ArrayToAdd)
 {
 	for(const TSharedRef<FNiagaraSimCacheTreeItem>& Item : ArrayToAdd)
diff --git a/Engine/Plugins/FX/Niagara/Source/NiagaraEditor/Private/Widgets/SNiagaraSimCacheTreeView.h b/Engine/Plugins/FX/Niagara/Source/NiagaraEditor/Private/Widgets/SNiagaraSimCacheTreeView.h
--- a/Engine/Plugins/FX/Niagara/Source/NiagaraEditor/Private/Widgets/SNiagaraSimCacheTreeView.h
+++ b/Engine/Plugins/FX/Niagara/Source/NiagaraEditor/Private/Widgets/SNiagaraSimCacheTreeView.h
@@ -68,6 +68,8 @@ public:
 	SLATE_ARGUMENT(TSharedPtr<FNiagaraSimCacheViewModel>, SimCacheViewModel)
 SLATE_END_ARGS()
 
+	void OnSimCacheChanged();
+	void SetupRootEntries();
 	void Construct(const FArguments& InArgs);
 	bool HasSelectionFilter() const {return !SelectionForFilter.IsEmpty();}
 	bool IsFilterActive() const;
diff --git a/Engine/Plugins/FX/Niagara/Source/NiagaraEditor/Private/Widgets/SNiagaraSimCacheView.cpp b/Engine/Plugins/FX/Niagara/Source/NiagaraEditor/Private/Widgets/SNiagaraSimCacheView.cpp
--- a/Engine/Plugins/FX/Niagara/Source/NiagaraEditor/Private/Widgets/SNiagaraSimCacheView.cpp
+++ b/Engine/Plugins/FX/Niagara/Source/NiagaraEditor/Private/Widgets/SNiagaraSimCacheView.cpp
@@ -47,17 +47,23 @@ public:
 
 	virtual TSharedRef<SWidget> GenerateWidgetForColumn(const FName& InColumnName) override
 	{
+		
+		if(!SimCacheViewModel.Get()->IsCacheValid())
+		{
+			return SNullWidget::NullWidget;
+		}
+
 		const int32 InstanceIndex = *RowIndexPtr;
+		
 		if (InColumnName == NAME_Instance)
 		{
 			return SNew(STextBlock)
 				.Text(FText::AsNumber(InstanceIndex));
 		}
-		else
-		{
-			return SNew(STextBlock)
-				.Text(SimCacheViewModel->GetComponentText(InColumnName, InstanceIndex));
-		}
+		
+		return SNew(STextBlock)
+			.Text(SimCacheViewModel->GetComponentText(InColumnName, InstanceIndex));
+		
 	}
 
 	TSharedPtr<int32>						RowIndexPtr;
@@ -199,7 +205,7 @@ void SNiagaraSimCacheView::GenerateColumns()
 	//  Give columns a width to prevent them from being shrunk when filtering. 
 	constexpr float ManualWidth = 125.0f;
 	HeaderRowWidget->ClearColumns();
-
+	
 	// Generate instance count column
 	HeaderRowWidget->AddColumn(
 		SHeaderRow::Column(NAME_Instance)
@@ -211,7 +217,7 @@ void SNiagaraSimCacheView::GenerateColumns()
 		.ManualWidth(ManualWidth)
 		.SortMode(EColumnSortMode::None)
 	);
-		
+	
 	// Generate a column for each component
 	for (const FNiagaraSimCacheViewModel::FComponentInfo& ComponentInfo : SimCacheViewModel->GetCurrentComponentInfos())
 	{
@@ -224,9 +230,11 @@ void SNiagaraSimCacheView::GenerateColumns()
 			.VAlignCell(EVerticalAlignment::VAlign_Fill)
 			.FillWidth(1.0f)
 			.ManualWidth(ManualWidth)
+			.ShouldGenerateWidget(TAttribute<bool>::Create(TAttribute<bool>::FGetter::CreateSP(this, &SNiagaraSimCacheView::GetShouldGenerateWidget, ComponentInfo.Name)))
 			.SortMode(EColumnSortMode::None)
 		);
 	}
+
 }
 
 void SNiagaraSimCacheView::UpdateColumns(const bool bReset)
@@ -251,39 +259,9 @@ void SNiagaraSimCacheView::UpdateColumns(const bool bReset)
 	const bool bColumnsUpToDate = bReset && !bApplyFilter;
 
 	FString ColumnName;
-
-	auto ComponentFilterPred = [&ColumnName](const FString& ComponentFilter)
-	{
-		return ColumnName.Equals(ComponentFilter);
-	};
-
-	auto StringFilterPred = [&ColumnName](const FString& StringFilter)
-	{
-		return ColumnName.Contains(StringFilter);
-	};
 	
 	if(!bColumnsUpToDate)
 	{
-		for (SHeaderRow::FColumn Column : HeaderRowWidget->GetColumns())
-		{
-			ColumnName = Column.DefaultText.Get().ToString();
-
-			// Never filter instance count column.
-			if(ColumnName.Equals(NAME_Instance.ToString()))
-			{
-				continue;
-			}
-
-			bool bPassedFilter = true;
-
-			if(bApplyFilter)
-			{
-				bPassedFilter = bUseComponentFilter ? FilterArray.ContainsByPredicate(ComponentFilterPred) : FilterArray.ContainsByPredicate(StringFilterPred);
-			}
-			
-			HeaderRowWidget->SetShowGeneratedColumn(Column.ColumnId, bPassedFilter);
-		}
-
 		HeaderRowWidget->RefreshColumns();
 	}
 }
@@ -380,4 +358,51 @@ void SNiagaraSimCacheView::OnBufferChanged()
 	UpdateColumns(true);
 }
 
+bool SNiagaraSimCacheView::GetShouldGenerateWidget(FName Name)
+{
+	TArray<FString> FilterArray;
+
+	// There are two filtering methods, from the tree view and from the spreadsheet view.
+	// The spreadsheet view uses a filter string input by the user. The tree view uses component item selection.
+	const bool bUseComponentFilter = SimCacheViewModel->IsComponentFilterActive();
+
+	FilterArray = bUseComponentFilter ? SimCacheViewModel->GetComponentFilters() : StringFilterArray;
+
+	// If the component filter is active, always apply the filter. If using the string filter, only apply it if it's not empty.
+	const bool bApplyFilter = bUseComponentFilter || !FilterArray.IsEmpty();
+
+	// If no filter to apply, all columns are visible.
+	if(!bApplyFilter)
+	{
+		return true;
+	}
+
+	FString ColumnName = Name.ToString();
+
+	// Always display the instance column
+	if(ColumnName.Equals(NAME_Instance.ToString()))
+	{
+		return true;
+	}
+
+	auto ComponentFilterPred = [&ColumnName](const FString& ComponentFilter)
+	{
+		return ColumnName.Equals(ComponentFilter);
+	};
+
+	auto StringFilterPred = [&ColumnName](const FString& StringFilter)
+	{
+		return ColumnName.Contains(StringFilter);
+	};
+
+	bool bPassedFilter = true;
+
+	if(bApplyFilter)
+	{
+		bPassedFilter = bUseComponentFilter ? FilterArray.ContainsByPredicate(ComponentFilterPred) : FilterArray.ContainsByPredicate(StringFilterPred);
+	}
+
+	return bPassedFilter;
+}
+
 #undef LOCTEXT_NAMESPACE
diff --git a/Engine/Plugins/FX/Niagara/Source/NiagaraEditor/Private/Widgets/SNiagaraSimCacheView.h b/Engine/Plugins/FX/Niagara/Source/NiagaraEditor/Private/Widgets/SNiagaraSimCacheView.h
--- a/Engine/Plugins/FX/Niagara/Source/NiagaraEditor/Private/Widgets/SNiagaraSimCacheView.h
+++ b/Engine/Plugins/FX/Niagara/Source/NiagaraEditor/Private/Widgets/SNiagaraSimCacheView.h
@@ -46,6 +46,7 @@ public:
 	void OnBufferChanged();
 
 private:
+	bool GetShouldGenerateWidget(FName Name);
 	void GenerateColumns();
 
 	TArray<TSharedPtr<int32>>					RowItems;
