Reintroduce FPlatformTypeLayoutParameters to the NiagaraShaderMapId to resolve conflict between win32 and linux cooks

diff --git a/Engine/Plugins/FX/Niagara/Source/NiagaraShader/Private/NiagaraShader.cpp b/Engine/Plugins/FX/Niagara/Source/NiagaraShader/Private/NiagaraShader.cpp
--- a/Engine/Plugins/FX/Niagara/Source/NiagaraShader/Private/NiagaraShader.cpp
+++ b/Engine/Plugins/FX/Niagara/Source/NiagaraShader/Private/NiagaraShader.cpp
@@ -262,6 +262,9 @@ void FNiagaraShaderMapId::AppendKeyString(FString& KeyString) const
 		KeyString.AppendChar('_');
 	}
 
+	// include the LayoutParams to differentiate between the frozen memory layout differences between platforms
+	LayoutParams.AppendKeyString(KeyString);
+
 	FName FeatureLevelName;
 	GetFeatureLevelName(FeatureLevel, FeatureLevelName);
 
