Fixed niagara editor style pointing to the wrong directory. This fixes an issue where the Niagara 'Refresh' icon wasn't found anymore in script assets

diff --git a/Engine/Plugins/FX/Niagara/Source/NiagaraEditor/Private/NiagaraEditorStyle.cpp b/Engine/Plugins/FX/Niagara/Source/NiagaraEditor/Private/NiagaraEditorStyle.cpp
--- a/Engine/Plugins/FX/Niagara/Source/NiagaraEditor/Private/NiagaraEditorStyle.cpp
+++ b/Engine/Plugins/FX/Niagara/Source/NiagaraEditor/Private/NiagaraEditorStyle.cpp
@@ -791,7 +791,7 @@ void FNiagaraEditorStyle::InitSimCacheEditor()
 FNiagaraEditorStyle::FNiagaraEditorStyle() : FSlateStyleSet("NiagaraEditorStyle")
 {
 	SetContentRoot(FPaths::EnginePluginsDir() / TEXT("FX/Niagara/Content/Slate"));
-	SetCoreContentRoot(FPaths::EngineContentDir() / TEXT("Editor/Slate"));
+	SetCoreContentRoot(FPaths::EngineContentDir() / TEXT("Slate"));
 
 	InitStats();
 	InitAssetPicker();
diff --git a/Engine/Plugins/FX/Niagara/Source/NiagaraEditorWidgets/Private/NiagaraEditorWidgetsStyle.cpp b/Engine/Plugins/FX/Niagara/Source/NiagaraEditorWidgets/Private/NiagaraEditorWidgetsStyle.cpp
--- a/Engine/Plugins/FX/Niagara/Source/NiagaraEditorWidgets/Private/NiagaraEditorWidgetsStyle.cpp
+++ b/Engine/Plugins/FX/Niagara/Source/NiagaraEditorWidgets/Private/NiagaraEditorWidgetsStyle.cpp
@@ -9,7 +9,6 @@
 #include "Styling/SlateTypes.h"
 #include "Styling/CoreStyle.h"
 #include "Settings/EditorStyleSettings.h"
-#include "Styling/StarshipCoreStyle.h"
 #include "Styling/AppStyle.h"
 #include "Styling/StyleColors.h"
 
@@ -65,7 +64,7 @@ void FNiagaraEditorWidgetsStyle::ReinitializeStyle()
 FNiagaraEditorWidgetsStyle::FNiagaraEditorWidgetsStyle() : FSlateStyleSet("NiagaraEditorWidgetsStyle")
 {
 	SetContentRoot(FPaths::EnginePluginsDir() / TEXT("FX/Niagara/Content/Slate"));
-	SetCoreContentRoot(FPaths::EngineContentDir() / TEXT("Editor/Slate"));
+	SetCoreContentRoot(FPaths::EngineContentDir() / TEXT("Slate"));
 	
 	const FTextBlockStyle NormalText = FAppStyle::Get().GetWidgetStyle<FTextBlockStyle>("NormalText");
 	const FEditableTextBoxStyle NormalEditableTextBox = FAppStyle::Get().GetWidgetStyle<FEditableTextBoxStyle>("NormalEditableTextBox");
