Update graph cache to incorporate the top level change ids to know if the cache is valid.  Ensures that emitters (post merge) will not try to use an out of date set of cached static variables.  Also includes two optimizations to try to offset the cost (especially for merging) of having to collect static variables more often:
-map builder can now be in 'static variable only' mode where we skip recursion into functions if the function graph doesn't have any static variables
-defer updating the override pins during merging to after all the modules have been merged

diff --git a/Engine/Plugins/FX/Niagara/Source/Niagara/Classes/NiagaraScript.h b/Engine/Plugins/FX/Niagara/Source/Niagara/Classes/NiagaraScript.h
--- a/Engine/Plugins/FX/Niagara/Source/Niagara/Classes/NiagaraScript.h
+++ b/Engine/Plugins/FX/Niagara/Source/Niagara/Classes/NiagaraScript.h
@@ -585,8 +585,14 @@ struct TStructOpsTypeTraits<FNiagaraVMExecutableData> : public TStructOpsTypeTra
 
 struct NIAGARA_API FNiagaraGraphCachedDataBase
 {
-	virtual ~FNiagaraGraphCachedDataBase() {}
+	virtual ~FNiagaraGraphCachedDataBase() = default;
 	virtual void GetStaticVariables(TArray<FNiagaraVariable>& OutVars) {}
+
+	virtual bool IsValidForSystem(const UNiagaraSystem* InSystem) const { return false; };
+	virtual void SetSourceSystem(const UNiagaraSystem* InSystem) {};
+
+	virtual bool IsValidForEmitter(const FVersionedNiagaraEmitterData* InEmitterData) const { return false; };
+	virtual void SetSourceEmitter(const FVersionedNiagaraEmitterData* InEmitterData) {};
 };
 
 /** Struct containing all of the data that can be different between different script versions.*/
diff --git a/Engine/Plugins/FX/Niagara/Source/Niagara/Private/NiagaraEmitter.cpp b/Engine/Plugins/FX/Niagara/Source/Niagara/Private/NiagaraEmitter.cpp
--- a/Engine/Plugins/FX/Niagara/Source/Niagara/Private/NiagaraEmitter.cpp
+++ b/Engine/Plugins/FX/Niagara/Source/Niagara/Private/NiagaraEmitter.cpp
@@ -1395,7 +1395,7 @@ TArray<UNiagaraEditorParametersAdapterBase*> UNiagaraEmitter::GetEditorOnlyParam
 const TSharedPtr<FNiagaraGraphCachedDataBase, ESPMode::ThreadSafe>& UNiagaraEmitter::GetCachedTraversalData(const FGuid& EmitterVersion) const
 {
 	const FVersionedNiagaraEmitterData* EmitterData = GetEmitterData(EmitterVersion);
-	if (EmitterData->CachedTraversalData.IsValid())
+	if (EmitterData->CachedTraversalData.IsValid() && EmitterData->CachedTraversalData->IsValidForEmitter(EmitterData))
 	{
 		return EmitterData->CachedTraversalData;
 	}
diff --git a/Engine/Plugins/FX/Niagara/Source/Niagara/Private/NiagaraSystem.cpp b/Engine/Plugins/FX/Niagara/Source/Niagara/Private/NiagaraSystem.cpp
--- a/Engine/Plugins/FX/Niagara/Source/Niagara/Private/NiagaraSystem.cpp
+++ b/Engine/Plugins/FX/Niagara/Source/Niagara/Private/NiagaraSystem.cpp
@@ -2925,7 +2925,7 @@ void UNiagaraSystem::PrepareRapidIterationParametersForCompilation()
 
 const TSharedPtr<FNiagaraGraphCachedDataBase, ESPMode::ThreadSafe>& UNiagaraSystem::GetCachedTraversalData() const
 {
-	if (CachedTraversalData.IsValid())
+	if (CachedTraversalData.IsValid() && CachedTraversalData->IsValidForSystem(this))
 		return CachedTraversalData;
 
 	INiagaraModule& NiagaraModule = FModuleManager::Get().LoadModuleChecked<INiagaraModule>(TEXT("Niagara"));
diff --git a/Engine/Plugins/FX/Niagara/Source/NiagaraEditor/Private/NiagaraCompiler.cpp b/Engine/Plugins/FX/Niagara/Source/NiagaraEditor/Private/NiagaraCompiler.cpp
--- a/Engine/Plugins/FX/Niagara/Source/NiagaraEditor/Private/NiagaraCompiler.cpp
+++ b/Engine/Plugins/FX/Niagara/Source/NiagaraEditor/Private/NiagaraCompiler.cpp
@@ -1475,6 +1475,7 @@ TSharedPtr<FNiagaraGraphCachedDataBase, ESPMode::ThreadSafe> FNiagaraEditorModul
 
 	if (System)
 	{
+		CachedPtr->SetSourceSystem(System);
 		ScriptSource = CastChecked<UNiagaraScriptSource>(System->GetSystemSpawnScript()->GetLatestSource());
 		ConstantResolver = FCompileConstantResolver(System, ENiagaraScriptUsage::SystemSpawnScript);
 		System->GatherStaticVariables(StaticVariablesFromSystem, StaticVariablesFromSystemAndEmitters);
@@ -1486,6 +1487,9 @@ TSharedPtr<FNiagaraGraphCachedDataBase, ESPMode::ThreadSafe> FNiagaraEditorModul
 	}
 	else if (Emitter)
 	{
+		const FVersionedNiagaraEmitterData* EmitterData = Emitter->GetEmitterData(Version);
+
+		CachedPtr->SetSourceEmitter(EmitterData);
 		UNiagaraSystem* SysParent = Cast<UNiagaraSystem>(Emitter->GetOuter());
 		if (SysParent)
 		{
@@ -1504,7 +1508,6 @@ TSharedPtr<FNiagaraGraphCachedDataBase, ESPMode::ThreadSafe> FNiagaraEditorModul
 			}*/
 		}
 		SrcUniqueEmitterName = Emitter->GetUniqueEmitterName();
-		const FVersionedNiagaraEmitterData* EmitterData = Emitter->GetEmitterData(Version);
 		EmitterData->GatherStaticVariables(SrcStaticVariables);
 
 		ScriptSource = CastChecked<UNiagaraScriptSource>(EmitterData->GraphSource);
@@ -1572,6 +1575,7 @@ TSharedPtr<FNiagaraGraphCachedDataBase, ESPMode::ThreadSafe> FNiagaraEditorModul
 			Builder.BeginTranslation(TranslationName);
 			Builder.BeginUsage(FoundOutputNode->GetUsage(), SimStageName);
 			Builder.EnableScriptAllowList(true, FoundOutputNode->GetUsage());
+			Builder.IncludeStaticVariablesOnly();
 			Builder.BuildParameterMaps(FoundOutputNode, true);
 			Builder.EndUsage();
 
diff --git a/Engine/Plugins/FX/Niagara/Source/NiagaraEditor/Private/NiagaraGraph.cpp b/Engine/Plugins/FX/Niagara/Source/NiagaraEditor/Private/NiagaraGraph.cpp
--- a/Engine/Plugins/FX/Niagara/Source/NiagaraEditor/Private/NiagaraGraph.cpp
+++ b/Engine/Plugins/FX/Niagara/Source/NiagaraEditor/Private/NiagaraGraph.cpp
@@ -4133,6 +4133,44 @@ void UNiagaraGraph::SetIsStaticSwitch(const FNiagaraVariable& Variable, bool InV
 	}
 }
 
+bool UNiagaraGraph::ReferencesStaticVariable(FNiagaraStaticVariableSearchContext& SearchContext) const
+{
+	FNiagaraStaticVariableSearchContext::FGraphKey GraphKey(this, ChangeId);
+	if (const bool* bHasReferencePtr = SearchContext.CachedResults.Find(GraphKey))
+	{
+		return *bHasReferencePtr;
+	}
+
+	for (const auto& ScriptVariableEntry : VariableToScriptVariable)
+	{
+		if (ScriptVariableEntry.Key.GetType().IsStatic())
+		{
+			SearchContext.CachedResults.Add(GraphKey, true);
+			return true;
+		}
+	}
+
+	if (SearchContext.bIncludeReferencedGraphs)
+	{
+		for (const UEdGraphNode* GraphNode : Nodes)
+		{
+			if (const UNiagaraNodeFunctionCall* FunctionCall = Cast<const UNiagaraNodeFunctionCall>(GraphNode))
+			{
+				if (const UNiagaraGraph* CalledGraph = FunctionCall->GetCalledGraph())
+				{
+					if (CalledGraph->ReferencesStaticVariable(SearchContext))
+					{
+						SearchContext.CachedResults.Add(GraphKey, true);
+						return true;
+					}
+				}
+			}
+		}
+	}
+
+	SearchContext.CachedResults.Add(GraphKey, false);
+	return false;
+}
 
 #undef NIAGARA_SCOPE_CYCLE_COUNTER
 #undef LOCTEXT_NAMESPACE
diff --git a/Engine/Plugins/FX/Niagara/Source/NiagaraEditor/Private/NiagaraNodeFunctionCall.cpp b/Engine/Plugins/FX/Niagara/Source/NiagaraEditor/Private/NiagaraNodeFunctionCall.cpp
--- a/Engine/Plugins/FX/Niagara/Source/NiagaraEditor/Private/NiagaraNodeFunctionCall.cpp
+++ b/Engine/Plugins/FX/Niagara/Source/NiagaraEditor/Private/NiagaraNodeFunctionCall.cpp
@@ -1287,9 +1287,9 @@ void UNiagaraNodeFunctionCall::BuildParameterMapHistory(FNiagaraParameterMapHist
 			NodeIdx = OutHistory.BeginNodeVisitation(ParamMapIdx, this);
 		}
 
-		// check if we should be recursing deeper into the graph
-		const bool DoDepthTraversal = (OutHistory.MaxGraphDepthTraversal == INDEX_NONE || OutHistory.CurrentGraphDepth < OutHistory.MaxGraphDepthTraversal);
+		const bool DoDepthTraversal = OutHistory.ShouldProcessDepthTraversal(FunctionGraph);
 
+		// check if we should be recursing deeper into the graph
 		if (DoDepthTraversal)
 		{
 			++OutHistory.CurrentGraphDepth;
@@ -1394,7 +1394,7 @@ void UNiagaraNodeFunctionCall::BuildParameterMapHistory(FNiagaraParameterMapHist
 	}
 }
 
-void UNiagaraNodeFunctionCall::ChangeScriptVersion(FGuid NewScriptVersion, const FNiagaraScriptVersionUpgradeContext& UpgradeContext, bool bShowNotesInStack)
+void UNiagaraNodeFunctionCall::ChangeScriptVersion(FGuid NewScriptVersion, const FNiagaraScriptVersionUpgradeContext& UpgradeContext, bool bShowNotesInStack, bool bDeferOverridePinUpdate)
 {
 	bool bPreviousVersionValid = !InvalidScriptVersionReference.IsValid();
 	if (NewScriptVersion == SelectedScriptVersion && bPreviousVersionValid)
@@ -1424,6 +1424,16 @@ void UNiagaraNodeFunctionCall::ChangeScriptVersion(FGuid NewScriptVersion, const
 	PreviousScriptVersion = bShowNotesInStack ? SelectedScriptVersion : NewScriptVersion;
 	SelectedScriptVersion = NewScriptVersion;
 
+	if (!bDeferOverridePinUpdate)
+	{
+		UpdateOverridePins(UpgradeContext);
+	}
+
+	MarkNodeRequiresSynchronization(__FUNCTION__, true);
+}
+
+void UNiagaraNodeFunctionCall::UpdateOverridePins(const FNiagaraScriptVersionUpgradeContext& UpgradeContext)
+{
 	if (FunctionScript)
 	{
 		// Automatically remove old inputs so it does not show a bunch of warnings to the user
@@ -1434,40 +1444,42 @@ void UNiagaraNodeFunctionCall::ChangeScriptVersion(FGuid NewScriptVersion, const
 		{
 			OverrideNode->Modify();
 			OverrideNode->GetInputPins(OverridePins);
-			
-			TArray<FNiagaraVariable> ModuleInputVariables;
-			FNiagaraStackGraphUtilities::GetStackFunctionInputs(*this, ModuleInputVariables, UpgradeContext.ConstantResolver, FNiagaraStackGraphUtilities::ENiagaraGetStackFunctionInputPinsOptions::ModuleInputsOnly);
-			for (const FNiagaraVariable& InputVariable : ModuleInputVariables)
-			{
-				FunctionInputNames.Add(FNiagaraParameterHandle(InputVariable.GetName()).GetName(), InputVariable.GetType());
-			}
-		}
-		for (UEdGraphPin* OverridePin : OverridePins)
-		{
-			if (!FNiagaraStackGraphUtilities::IsOverridePinForFunction(*OverridePin, *this))
-			{
-				continue;
-			}
-			FName InputName = FNiagaraParameterHandle(OverridePin->PinName).GetName();
-			FNiagaraTypeDefinition* InputType = FunctionInputNames.Find(InputName);
-			FNiagaraTypeDefinition OverridePinType = UEdGraphSchema_Niagara::PinTypeToTypeDefinition(OverridePin->PinType);
-			if (InputType != nullptr && *InputType != OverridePinType)
-			{
-				// looks like the type of the module input changed - we'll change the override pin type as well
-				OverridePin->Modify();
-				OverridePin->PinType = UEdGraphSchema_Niagara::TypeDefinitionToPinType(*InputType);
-			}
-			else if (InputType == nullptr)
+
+			if (!OverridePins.IsEmpty())
 			{
-				// the input doesn't exist any more, delete the override pin
-				TArray<TWeakObjectPtr<UNiagaraDataInterface>> RemovedDataObjects;
-				FNiagaraStackGraphUtilities::RemoveNodesForStackFunctionInputOverridePin(*OverridePin, RemovedDataObjects);
-				OverrideNode->RemovePin(OverridePin);
+				TArray<FNiagaraVariable> ModuleInputVariables;
+				FNiagaraStackGraphUtilities::GetStackFunctionInputs(*this, ModuleInputVariables, UpgradeContext.ConstantResolver, FNiagaraStackGraphUtilities::ENiagaraGetStackFunctionInputPinsOptions::ModuleInputsOnly);
+				for (const FNiagaraVariable& InputVariable : ModuleInputVariables)
+				{
+					FunctionInputNames.Add(FNiagaraParameterHandle(InputVariable.GetName()).GetName(), InputVariable.GetType());
+				}
+
+				for (UEdGraphPin* OverridePin : OverridePins)
+				{
+					if (!FNiagaraStackGraphUtilities::IsOverridePinForFunction(*OverridePin, *this))
+					{
+						continue;
+					}
+					FName InputName = FNiagaraParameterHandle(OverridePin->PinName).GetName();
+					FNiagaraTypeDefinition* InputType = FunctionInputNames.Find(InputName);
+					FNiagaraTypeDefinition OverridePinType = UEdGraphSchema_Niagara::PinTypeToTypeDefinition(OverridePin->PinType);
+					if (InputType != nullptr && *InputType != OverridePinType)
+					{
+						// looks like the type of the module input changed - we'll change the override pin type as well
+						OverridePin->Modify();
+						OverridePin->PinType = UEdGraphSchema_Niagara::TypeDefinitionToPinType(*InputType);
+					}
+					else if (InputType == nullptr)
+					{
+						// the input doesn't exist any more, delete the override pin
+						TArray<TWeakObjectPtr<UNiagaraDataInterface>> RemovedDataObjects;
+						FNiagaraStackGraphUtilities::RemoveNodesForStackFunctionInputOverridePin(*OverridePin, RemovedDataObjects);
+						OverrideNode->RemovePin(OverridePin);
+					}
+				}
 			}
 		}
 	}
-
-	MarkNodeRequiresSynchronization(__FUNCTION__, true);
 }
 
 UEdGraphPin* UNiagaraNodeFunctionCall::FindParameterMapDefaultValuePin(const FName VariableName, ENiagaraScriptUsage InParentUsage, FCompileConstantResolver ConstantResolver) const
diff --git a/Engine/Plugins/FX/Niagara/Source/NiagaraEditor/Private/NiagaraParameterMapHistory.cpp b/Engine/Plugins/FX/Niagara/Source/NiagaraEditor/Private/NiagaraParameterMapHistory.cpp
--- a/Engine/Plugins/FX/Niagara/Source/NiagaraEditor/Private/NiagaraParameterMapHistory.cpp
+++ b/Engine/Plugins/FX/Niagara/Source/NiagaraEditor/Private/NiagaraParameterMapHistory.cpp
@@ -5,6 +5,7 @@
 #include "NiagaraEditorCommon.h"
 #include "NiagaraHlslTranslator.h"
 #include "NiagaraSystem.h"
+#include "NiagaraSystemImpl.h"
 #include "NiagaraGraph.h"
 #include "EdGraphSchema_Niagara.h"
 #include "NiagaraNode.h"
@@ -2128,6 +2129,23 @@ int32 FNiagaraParameterMapHistoryBuilder::FindStaticVariable(const FNiagaraVaria
 	return FoundOverrideIdx;
 }
 
+bool FNiagaraParameterMapHistoryBuilder::ShouldProcessDepthTraversal(const UNiagaraGraph* Graph)
+{
+	// first check if there are limits setup for how deep we should look
+	if (MaxGraphDepthTraversal == INDEX_NONE || CurrentGraphDepth < MaxGraphDepthTraversal)
+	{
+		// next, if we only care about static variables, check if the graph we're going to go into actually includes any
+		if (StaticVariableSearchContext.bCollectingStaticVariablesOnly)
+		{
+			return Graph->ReferencesStaticVariable(StaticVariableSearchContext);
+		}
+
+		return true;
+	}
+
+	return false;
+}
+
 void FNiagaraParameterMapHistoryBuilder::SetConstantByStaticVariable(int32& OutValue, const FNiagaraVariable& Var)
 {
 	
@@ -2457,4 +2475,65 @@ int32 FNiagaraParameterMapHistoryWithMetaDataBuilder::AddVariableToHistory(FNiag
 	return History.AddVariable(InVar, InAliasedVar, ModuleName, InPin, BlankMetaData);
 }
 
+namespace NiagaraGraphCachedBuiltHistoryImpl
+{
+
+template<typename T>
+static void CollectScriptChangeIds(const T* InObject, TArray<FGuid>& ChangeIds)
+{
+	if (InObject)
+	{
+		InObject->ForEachScript([&ChangeIds](const UNiagaraScript* Script)
+		{
+			ChangeIds.Add(Script->GetBaseChangeID());
+		});
+	}
+}
+
+};
+
+bool FNiagaraGraphCachedBuiltHistory::IsValidForSystem(const UNiagaraSystem* InSystem) const
+{
+	if (!InSystem)
+	{
+		return false;
+	}
+
+	TArray<FGuid> TestChangeIds;
+	NiagaraGraphCachedBuiltHistoryImpl::CollectScriptChangeIds(InSystem, TestChangeIds);
+	return TestChangeIds == SourceAssetChangeIds;
+}
+
+void FNiagaraGraphCachedBuiltHistory::SetSourceSystem(const UNiagaraSystem* InSystem)
+{
+	SourceAssetChangeIds.Reset();
+
+	if (InSystem)
+	{
+		NiagaraGraphCachedBuiltHistoryImpl::CollectScriptChangeIds(InSystem, SourceAssetChangeIds);
+	}
+}
+
+bool FNiagaraGraphCachedBuiltHistory::IsValidForEmitter(const FVersionedNiagaraEmitterData* InEmitterData) const
+{
+	if (!InEmitterData)
+	{
+		return false;
+	}
+
+	TArray<FGuid> TestChangeIds;
+	NiagaraGraphCachedBuiltHistoryImpl::CollectScriptChangeIds(InEmitterData, TestChangeIds);
+	return TestChangeIds == SourceAssetChangeIds;
+}
+
+void FNiagaraGraphCachedBuiltHistory::SetSourceEmitter(const FVersionedNiagaraEmitterData* InEmitterData)
+{
+	SourceAssetChangeIds.Reset();
+
+	if (InEmitterData)
+	{
+		NiagaraGraphCachedBuiltHistoryImpl::CollectScriptChangeIds(InEmitterData, SourceAssetChangeIds);
+	}
+}
+
 #undef LOCTEXT_NAMESPACE
diff --git a/Engine/Plugins/FX/Niagara/Source/NiagaraEditor/Private/NiagaraScriptMergeManager.cpp b/Engine/Plugins/FX/Niagara/Source/NiagaraEditor/Private/NiagaraScriptMergeManager.cpp
--- a/Engine/Plugins/FX/Niagara/Source/NiagaraEditor/Private/NiagaraScriptMergeManager.cpp
+++ b/Engine/Plugins/FX/Niagara/Source/NiagaraEditor/Private/NiagaraScriptMergeManager.cpp
@@ -1487,9 +1487,19 @@ void FNiagaraScriptMergeManager::UpdateModuleVersions(const FVersionedNiagaraEmi
 				};
 				UpgradeContext.ApplyClipboardCallback = [ModuleItem](UNiagaraClipboardContent* ClipboardContent, FText& OutWarning) { ModuleItem->Paste(ClipboardContent, OutWarning); };
 			}
-			ChangedModule->GetFunctionCallNode()->ChangeScriptVersion(NewScriptVersion, UpgradeContext, true);
+			ChangedModule->GetFunctionCallNode()->ChangeScriptVersion(NewScriptVersion, UpgradeContext, true /*bShowNotesInStack*/, true /*bDeferOverridePinUpdate*/);
 			ChangedModule->GetFunctionCallNode()->RefreshFromExternalChanges();
 		}
+
+		// now that we've applied the changes to all of the changed modules we go back through them and ensure their override nodes are up to date
+		// this is done after the fact to avoid constantly regenerating the static variables that may be at play
+		for (TSharedRef<FNiagaraStackFunctionMergeAdapter>& ChangedModule : ChangedVersionOtherModules)
+		{
+			FNiagaraScriptVersionUpgradeContext UpgradeContext;
+			UpgradeContext.ConstantResolver = FCompileConstantResolver(Instance, FNiagaraStackGraphUtilities::GetOutputNodeUsage(*ChangedModule->GetFunctionCallNode()));
+
+			ChangedModule->GetFunctionCallNode()->UpdateOverridePins(UpgradeContext);
+		}
 	}
 
 	// now that we're done with our transient System we need to reset it so that if it does get pulled in (via TObjectIterator as an example)
diff --git a/Engine/Plugins/FX/Niagara/Source/NiagaraEditor/Public/NiagaraGraph.h b/Engine/Plugins/FX/Niagara/Source/NiagaraEditor/Public/NiagaraGraph.h
--- a/Engine/Plugins/FX/Niagara/Source/NiagaraEditor/Public/NiagaraGraph.h
+++ b/Engine/Plugins/FX/Niagara/Source/NiagaraEditor/Public/NiagaraGraph.h
@@ -13,6 +13,7 @@ class UNiagaraParameterDefinitions;
 class UNiagaraScriptVariable;
 struct FSynchronizeWithParameterDefinitionsArgs;
 struct FNiagaraScriptVariableData;
+struct FNiagaraStaticVariableSearchContext;
 
 /** This is the type of action that occurred on a given Niagara graph. Note that this should follow from EEdGraphActionType, leaving some slop for growth. */
 enum ENiagaraGraphActionType
@@ -430,6 +431,8 @@ class UNiagaraGraph : public UEdGraph
 
 	void SetIsStaticSwitch(const FNiagaraVariable& Variable, bool InValue);
 
+	bool ReferencesStaticVariable(FNiagaraStaticVariableSearchContext& SearchContext) const;
+
 protected:
 	void RebuildNumericCache();
 	bool bNeedNumericCacheRebuilt;
diff --git a/Engine/Plugins/FX/Niagara/Source/NiagaraEditor/Public/NiagaraNodeFunctionCall.h b/Engine/Plugins/FX/Niagara/Source/NiagaraEditor/Public/NiagaraNodeFunctionCall.h
--- a/Engine/Plugins/FX/Niagara/Source/NiagaraEditor/Public/NiagaraNodeFunctionCall.h
+++ b/Engine/Plugins/FX/Niagara/Source/NiagaraEditor/Public/NiagaraNodeFunctionCall.h
@@ -136,7 +136,9 @@ public:
 
 	virtual void BuildParameterMapHistory(FNiagaraParameterMapHistoryBuilder& OutHistory, bool bRecursive = true, bool bFilterForCompilation = true) const override;
 
-	NIAGARAEDITOR_API void ChangeScriptVersion(FGuid NewScriptVersion, const FNiagaraScriptVersionUpgradeContext& UpgradeContext, bool bShowNotesInStack = false);
+	/** if bDeferOverridePinUpdate is true then it is the caller's responsibility to call UpdateOverridePins() */
+	NIAGARAEDITOR_API void ChangeScriptVersion(FGuid NewScriptVersion, const FNiagaraScriptVersionUpgradeContext& UpgradeContext, bool bShowNotesInStack = false, bool bDeferOverridePinUpdate = false);
+	NIAGARAEDITOR_API void UpdateOverridePins(const FNiagaraScriptVersionUpgradeContext& UpgradeContext);
 
 	FString GetFunctionName() const { return FunctionDisplayName; }
 	NIAGARAEDITOR_API UNiagaraGraph* GetCalledGraph() const;
diff --git a/Engine/Plugins/FX/Niagara/Source/NiagaraEditor/Public/NiagaraParameterMapHistory.h b/Engine/Plugins/FX/Niagara/Source/NiagaraEditor/Public/NiagaraParameterMapHistory.h
--- a/Engine/Plugins/FX/Niagara/Source/NiagaraEditor/Public/NiagaraParameterMapHistory.h
+++ b/Engine/Plugins/FX/Niagara/Source/NiagaraEditor/Public/NiagaraParameterMapHistory.h
@@ -115,6 +115,16 @@ private:
 	TArray<FString> FriendlyPath;
 };
 
+struct FNiagaraStaticVariableSearchContext
+{
+	/** Whether we are running a limited history run where we only care about collecting static variables */
+	bool bCollectingStaticVariablesOnly = false;
+
+	bool bIncludeReferencedGraphs = true;
+
+	using FGraphKey = TTuple<TObjectKey<UNiagaraGraph>, FGuid>;
+	TMap<FGraphKey, bool> CachedResults;
+};
 
 /** Traverses a Niagara node graph to identify the variables that have been written and read from a parameter map. 
 * 	This class is meant to aid in UI and compilation of the graph. There are several main script types and each one interacts
@@ -539,6 +549,14 @@ public:
 		if (ContextuallyVisitedNodes.Num() > 0)
 			OutVistedNodes.Append(ContextuallyVisitedNodes.Last());
 	}
+
+	void IncludeStaticVariablesOnly()
+	{
+		StaticVariableSearchContext.bCollectingStaticVariablesOnly = true;
+	}
+
+	bool ShouldProcessDepthTraversal(const UNiagaraGraph* Graph);
+
 protected:
 	/** Internal helper function to decide whether or not we should use this static variable when merging with other
 		graph traversal static variables, like Emitter graphs merging with System graphs. */
@@ -599,6 +617,8 @@ protected:
 	bool bIgnoreStaticRapidIterationParameters = false;
 
 	TArray<FNiagaraVariable> EncounterableExternalVariables;
+
+	FNiagaraStaticVariableSearchContext StaticVariableSearchContext;
 };
 
 class FNiagaraParameterMapHistoryWithMetaDataBuilder : public FNiagaraParameterMapHistoryBuilder
@@ -615,8 +635,15 @@ public:
 
 struct FNiagaraGraphCachedBuiltHistory : public FNiagaraGraphCachedDataBase
 {
-	virtual ~FNiagaraGraphCachedBuiltHistory() {};
-	virtual void GetStaticVariables(TArray<FNiagaraVariable>& OutVars) { OutVars = StaticVariables; }
+	virtual ~FNiagaraGraphCachedBuiltHistory() = default;
+	virtual void GetStaticVariables(TArray<FNiagaraVariable>& OutVars) override { OutVars = StaticVariables; }
+
+	virtual bool IsValidForSystem(const UNiagaraSystem* InSystem) const override;
+	virtual void SetSourceSystem(const UNiagaraSystem* InSystem) override;
+
+	virtual bool IsValidForEmitter(const FVersionedNiagaraEmitterData* InEmitterData) const override;
+	virtual void SetSourceEmitter(const FVersionedNiagaraEmitterData* InEmitterData) override;
 
 	TArray<FNiagaraVariable> StaticVariables;
+	TArray<FGuid> SourceAssetChangeIds;
 };
\ No newline at end of file
