Fix memory leak of FSceneView allocated in CreateAndInitSingleView() from FNiagaraGpuComputeDispatch::ProcessPendingTicksFlush()

* FSceneViewFamilyContext is used to clean-up views for a view family

diff --git a/Engine/Plugins/FX/Niagara/Source/Niagara/Private/NiagaraGpuComputeDispatch.cpp b/Engine/Plugins/FX/Niagara/Source/Niagara/Private/NiagaraGpuComputeDispatch.cpp
--- a/Engine/Plugins/FX/Niagara/Source/Niagara/Private/NiagaraGpuComputeDispatch.cpp
+++ b/Engine/Plugins/FX/Niagara/Source/Niagara/Private/NiagaraGpuComputeDispatch.cpp
@@ -396,7 +396,7 @@ void FNiagaraGpuComputeDispatch::ProcessPendingTicksFlush(FRHICommandListImmedia
 
 			// Make a temporary ViewInfo
 			//-TODO: We could gather some more information here perhaps?
-			FSceneViewFamily ViewFamily(
+			FSceneViewFamilyContext ViewFamily(
 				FSceneViewFamily::ConstructionValues(nullptr, nullptr, FEngineShowFlags(ESFIM_Game))
 				.SetTime(CachedViewInitOptions.GameTime)
 				.SetGammaCorrection(CachedViewInitOptions.GammaCorrection)
