Removed details panel search box in details panel in NPC editor
Added custom search box that adds filtering for parameters in NPC editor
Also fixed the scroll bar not appearing in the Niagara Parameter Collections

diff --git a/Engine/Plugins/FX/Niagara/Source/NiagaraEditor/Private/Toolkits/NiagaraParameterCollectionToolkit.cpp b/Engine/Plugins/FX/Niagara/Source/NiagaraEditor/Private/Toolkits/NiagaraParameterCollectionToolkit.cpp
--- a/Engine/Plugins/FX/Niagara/Source/NiagaraEditor/Private/Toolkits/NiagaraParameterCollectionToolkit.cpp
+++ b/Engine/Plugins/FX/Niagara/Source/NiagaraEditor/Private/Toolkits/NiagaraParameterCollectionToolkit.cpp
@@ -27,6 +27,7 @@
 
 #include "PropertyEditorModule.h"
 #include "IDetailsView.h"
+#include "Widgets/Input/SSearchBox.h"
 
 #define LOCTEXT_NAMESPACE "NiagaraParameterCollectionEditor"
 
@@ -156,17 +157,15 @@ TSharedRef<SDockTab> FNiagaraParameterCollectionToolkit::SpawnTab_Main(const FSp
 
 	TSharedRef<SVerticalBox> Contents = SNew(SVerticalBox);
 
-
-	FDetailsViewArgs DetailArgs;
-	DetailArgs.bAllowSearch = false;
-
 	FPropertyEditorModule& PropertyEditorModule = FModuleManager::GetModuleChecked<FPropertyEditorModule>("PropertyEditor");
 	FDetailsViewArgs DetailsViewArgs;
 	DetailsViewArgs.NameAreaSettings = FDetailsViewArgs::HideNameArea;
 	DetailsViewArgs.bHideSelectionTip = true;
 	DetailsViewArgs.NotifyHook = ParameterCollectionViewModel.Get();
+	DetailsViewArgs.bAllowSearch = false;
+	
 	TSharedRef<IDetailsView> DetailsView = PropertyEditorModule.CreateDetailView(DetailsViewArgs);
-
+	
 	if (Instance->IsDefaultInstance())
 	{
 		DetailsView->SetObject(Collection);
@@ -177,28 +176,40 @@ TSharedRef<SDockTab> FNiagaraParameterCollectionToolkit::SpawnTab_Main(const FSp
 	}
 
 	Contents->AddSlot()
-		.AutoHeight()
-		.Padding(FMargin(0.0f, 2.0f))
-		[
-			DetailsView
-		];
+	.AutoHeight()
+	.Padding(FMargin(0.0f, 2.0f))
+	[
+		SNew(SSearchBox)
+		.OnTextChanged(this, &FNiagaraParameterCollectionToolkit::OnSearchTextChanged)
+	];
+	
+	Contents->AddSlot()
+	.AutoHeight()
+	.Padding(FMargin(0.0f, 2.0f))
+	[
+		DetailsView
+	];
 
 	Contents->AddSlot()
-		.AutoHeight()
-		.Padding(FMargin(0.0f, 2.0f))
-		[
-			ParameterCollection.ToSharedRef()
-		];
+	.Padding(FMargin(0.0f, 2.0f))
+	[
+		ParameterCollection.ToSharedRef()
+	];
 	
 	TSharedRef<SDockTab> SpawnedTab =
-		SNew(SDockTab)
-		[
-			Contents
-		];
+	SNew(SDockTab)
+	[
+		Contents
+	];
 
 	return SpawnedTab;
 }
 
+void FNiagaraParameterCollectionToolkit::OnSearchTextChanged(const FText& InSearchText)
+{
+	ParameterCollectionViewModel->UpdateParameterSelectionFromSearch(InSearchText);
+}
+
 void FNiagaraParameterCollectionToolkit::SetupCommands()
 {
 	// 	GetToolkitCommands()->MapAction(
diff --git a/Engine/Plugins/FX/Niagara/Source/NiagaraEditor/Private/Toolkits/NiagaraParameterCollectionToolkit.h b/Engine/Plugins/FX/Niagara/Source/NiagaraEditor/Private/Toolkits/NiagaraParameterCollectionToolkit.h
--- a/Engine/Plugins/FX/Niagara/Source/NiagaraEditor/Private/Toolkits/NiagaraParameterCollectionToolkit.h
+++ b/Engine/Plugins/FX/Niagara/Source/NiagaraEditor/Private/Toolkits/NiagaraParameterCollectionToolkit.h
@@ -46,6 +46,8 @@ private:
 
 	TSharedRef<SDockTab> SpawnTab_Main(const FSpawnTabArgs& Args);
 
+	void OnSearchTextChanged(const FText& InSearchText);
+	
 	/** Builds the toolbar widget */
 	void ExtendToolbar();
 	void SetupCommands();
diff --git a/Engine/Plugins/FX/Niagara/Source/NiagaraEditor/Private/ViewModels/NiagaraParameterCollectionAssetViewModel.cpp b/Engine/Plugins/FX/Niagara/Source/NiagaraEditor/Private/ViewModels/NiagaraParameterCollectionAssetViewModel.cpp
--- a/Engine/Plugins/FX/Niagara/Source/NiagaraEditor/Private/ViewModels/NiagaraParameterCollectionAssetViewModel.cpp
+++ b/Engine/Plugins/FX/Niagara/Source/NiagaraEditor/Private/ViewModels/NiagaraParameterCollectionAssetViewModel.cpp
@@ -164,6 +164,29 @@ void FNiagaraParameterCollectionAssetViewModel::UpdateOpenInstances()
 	}
 }
 
+void FNiagaraParameterCollectionAssetViewModel::UpdateParameterSelectionFromSearch(const FText& InSearchText)
+{
+	GetSelection().ClearSelectedObjects();
+
+	SearchText = InSearchText;
+
+	RefreshParameterViewModels();
+	
+	if(!SearchText.IsEmpty())
+	{
+		TArray<TSharedRef<INiagaraParameterViewModel>> MatchingParameters;
+		for(TSharedRef<INiagaraParameterViewModel> Parameter : ParameterViewModels)
+		{
+			if(Parameter->GetName().ToString().Contains(SearchText.ToString()))
+			{
+				MatchingParameters.Add(Parameter);
+			}
+		}
+
+		GetSelection().SetSelectedObjects(MatchingParameters);
+	}
+}
+
 void FNiagaraParameterCollectionAssetViewModel::RemoveParameter(FNiagaraVariable& Parameter)
 {
 	FScopedTransaction ScopedTransaction(LOCTEXT("RemoveNPCParameter", "Remove Parameter"));
@@ -264,10 +287,35 @@ void FNiagaraParameterCollectionAssetViewModel::RefreshParameterViewModels()
 
 	ParameterViewModels.Empty();
 
+	bool bSearchTextIsEmpty = SearchText.IsEmpty();
+
+	TArray<FString> SearchTerms;
+	if(bSearchTextIsEmpty == false)
+	{
+		SearchText.ToString().ParseIntoArray(SearchTerms, TEXT(" "));
+	}
 	for (int32 i = 0; i < Collection->GetParameters().Num(); ++i)
 	{
 		FNiagaraVariable& Var = Collection->GetParameters()[i];
 
+		if(bSearchTextIsEmpty == false)
+		{
+			bool bMatchFound = false;
+			for(const FString& SearchTerm : SearchTerms)
+			{
+				if(Var.GetName().ToString().Contains(SearchTerm))
+				{
+					bMatchFound = true;
+					break;
+				}
+			}
+
+			if(!bMatchFound)
+			{
+				continue;
+			}
+		}
+		
 		TSharedPtr<FNiagaraCollectionParameterViewModel> ParameterViewModel = MakeShareable(new FNiagaraCollectionParameterViewModel(Var, Instance, ParameterEditMode));
 
 		ParameterViewModel->OnNameChanged().AddRaw(this, &FNiagaraParameterCollectionAssetViewModel::OnParameterNameChanged, Var);
diff --git a/Engine/Plugins/FX/Niagara/Source/NiagaraEditor/Private/ViewModels/NiagaraParameterCollectionAssetViewModel.h b/Engine/Plugins/FX/Niagara/Source/NiagaraEditor/Private/ViewModels/NiagaraParameterCollectionAssetViewModel.h
--- a/Engine/Plugins/FX/Niagara/Source/NiagaraEditor/Private/ViewModels/NiagaraParameterCollectionAssetViewModel.h
+++ b/Engine/Plugins/FX/Niagara/Source/NiagaraEditor/Private/ViewModels/NiagaraParameterCollectionAssetViewModel.h
@@ -36,6 +36,8 @@ public:
 
 	void UpdateOpenInstances();
 
+	void UpdateParameterSelectionFromSearch(const FText& SearchText);
+	
 	//~ FNotifyHook interface
 	virtual void NotifyPreChange(FProperty* PropertyAboutToChange)override;
 	virtual void NotifyPostChange(const FPropertyChangedEvent& PropertyChangedEvent, FProperty* PropertyThatChanged)override;
@@ -53,7 +55,7 @@ public:
 
 	/** Sets the tooltip overrides on all parameters.*/
 	void SetAllParametersTooltipOverrides(const FText& Override);
-
+	
 	//~ FEditorUndoClient
 	void PostUndo(bool bSuccess);
 	void PostRedo(bool bSuccess) { PostUndo(bSuccess); }
@@ -62,7 +64,7 @@ protected:
 	//~ FNiagaraParameterCollectionViewModel interface.
 	virtual bool SupportsType(const FNiagaraTypeDefinition& Type) const override;
 
-private:
+private:	
 	/** Handles when the name on a parameter changes. */
 	void OnParameterNameChanged(FName OldName, FName NewName, FNiagaraVariable ParameterVariable);
 
@@ -89,6 +91,8 @@ private:
 	UNiagaraParameterCollection* Collection;
 	UNiagaraParameterCollectionInstance* Instance;
 
+	FText SearchText;
+	
 	TNiagaraViewModelManager<UNiagaraParameterCollection, FNiagaraParameterCollectionAssetViewModel>::Handle RegisteredHandle;
 
 	FDelegateHandle ExternalChangeHandle;
