Speculative safety code around world managers and their deletion.

diff --git a/Engine/Plugins/FX/Niagara/Source/Niagara/Private/NiagaraWorldManager.cpp b/Engine/Plugins/FX/Niagara/Source/Niagara/Private/NiagaraWorldManager.cpp
--- a/Engine/Plugins/FX/Niagara/Source/Niagara/Private/NiagaraWorldManager.cpp
+++ b/Engine/Plugins/FX/Niagara/Source/Niagara/Private/NiagaraWorldManager.cpp
@@ -416,9 +416,15 @@ void FNiagaraWorldManager::OnShutdown()
 	check(WorldManagers.Num() == 0);
 	for (TPair<UWorld*, FNiagaraWorldManager*> Pair : WorldManagers)
 	{
-		delete Pair.Value;
+		FNiagaraWorldManager* WorldMan = Pair.Value;
+		if(ensure(WorldMan))
+		{
+			delete WorldMan;
+		}
+		
 		Pair.Value = nullptr;
 	}
+	WorldManagers.Empty();
 }
 
 void FNiagaraWorldManager::AddReferencedObjects(FReferenceCollector& Collector)
@@ -749,7 +755,11 @@ void FNiagaraWorldManager::OnWorldCleanup(UWorld* World, bool bSessionEnded, boo
 	FNiagaraWorldManager** Manager = WorldManagers.Find(World);
 	if (Manager)
 	{
-		(*Manager)->OnWorldCleanup(bSessionEnded, bCleanupResources);
+		FNiagaraWorldManager* WorldMan = *Manager;
+		if (ensure(WorldMan))
+		{
+			WorldMan->OnWorldCleanup(bSessionEnded, bCleanupResources);
+		}
 	}
 }
 
@@ -759,7 +769,11 @@ void FNiagaraWorldManager::OnPostWorldCleanup(UWorld* World, bool bSessionEnded,
 	FNiagaraWorldManager** Manager = WorldManagers.Find(World);
 	if (Manager)
 	{
-		(*Manager)->OnPostWorldCleanup(bSessionEnded, bCleanupResources);
+		FNiagaraWorldManager* WorldMan = *Manager;
+		if (ensure(WorldMan))
+		{
+			WorldMan->OnPostWorldCleanup(bSessionEnded, bCleanupResources);
+		}
 	}
 }
 
@@ -768,7 +782,11 @@ void FNiagaraWorldManager::OnPreWorldFinishDestroy(UWorld* World)
 	FNiagaraWorldManager** Manager = WorldManagers.Find(World);
 	if (Manager)
 	{
-		delete (*Manager);
+		FNiagaraWorldManager* WorldMan = *Manager;
+		if(ensure(WorldMan))
+		{
+			delete WorldMan;
+		}
 		WorldManagers.Remove(World);
 	}
 }
@@ -778,7 +796,11 @@ void FNiagaraWorldManager::OnWorldBeginTearDown(UWorld* World)
 	FNiagaraWorldManager** Manager = WorldManagers.Find(World);
 	if (Manager)
 	{
-		(*Manager)->OnWorldBeginTearDown();
+		FNiagaraWorldManager* WorldMan = *Manager;
+		if(ensure(WorldMan))
+		{
+			WorldMan->OnWorldBeginTearDown();
+		}		
 	}
 // 	FNiagaraWorldManager** Manager = WorldManagers.Find(World);
 // 	if (Manager)
