Fixed emitter assets not dirtying when disabling a sim stage

diff --git a/Engine/Plugins/FX/Niagara/Source/Niagara/Private/NiagaraSimulationStageBase.cpp b/Engine/Plugins/FX/Niagara/Source/Niagara/Private/NiagaraSimulationStageBase.cpp
--- a/Engine/Plugins/FX/Niagara/Source/Niagara/Private/NiagaraSimulationStageBase.cpp
+++ b/Engine/Plugins/FX/Niagara/Source/Niagara/Private/NiagaraSimulationStageBase.cpp
@@ -39,6 +39,7 @@ void UNiagaraSimulationStageBase::SetEnabled(bool bInEnabled)
 	{
 		bEnabled = bInEnabled;
 		RequestRecompile();
+		OnChanged().Broadcast();
 	}
 }
 
@@ -46,7 +47,7 @@ void UNiagaraSimulationStageBase::RequestRecompile()
 {
 	if (!bEnabled)
 	{
-		Script->InvalidateCompileResults(TEXT("SimulatoinStage changed."));
+		Script->InvalidateCompileResults(TEXT("SimulationStage changed."));
 	}
 
 	FVersionedNiagaraEmitter OuterEmitter = GetOuterEmitter();
