Fixed dragdrop ops of different types than FNiagaraHierarchyDragDropOps from being used in the hierarchy editor. This fixes crashes

diff --git a/Engine/Plugins/FX/Niagara/Source/NiagaraEditor/Private/ViewModels/HierarchyEditor/NiagaraHierarchyViewModelBase.cpp b/Engine/Plugins/FX/Niagara/Source/NiagaraEditor/Private/ViewModels/HierarchyEditor/NiagaraHierarchyViewModelBase.cpp
--- a/Engine/Plugins/FX/Niagara/Source/NiagaraEditor/Private/ViewModels/HierarchyEditor/NiagaraHierarchyViewModelBase.cpp
+++ b/Engine/Plugins/FX/Niagara/Source/NiagaraEditor/Private/ViewModels/HierarchyEditor/NiagaraHierarchyViewModelBase.cpp
@@ -197,8 +197,9 @@ TOptional<EItemDropZone> FNiagaraHierarchyCategoryViewModel::OnCanAcceptDropInte
 {
 	bool bAllowDrop = false;
 
-	if(TSharedPtr<FNiagaraHierarchyDragDropOp> HierarchyDragDropOp = StaticCastSharedPtr<FNiagaraHierarchyDragDropOp>(DragDropOp))
+	if(DragDropOp->IsOfType<FNiagaraHierarchyDragDropOp>())
 	{
+		TSharedPtr<FNiagaraHierarchyDragDropOp> HierarchyDragDropOp = StaticCastSharedPtr<FNiagaraHierarchyDragDropOp>(DragDropOp);
 		TSharedPtr<FNiagaraHierarchyItemViewModelBase> SourceDropItem = HierarchyDragDropOp->GetDraggedItem().Pin();
 		TSharedPtr<FNiagaraHierarchyItemViewModelBase> TargetDropItem = AsShared();
 		
@@ -691,8 +692,9 @@ TOptional<EItemDropZone> FNiagaraHierarchyRootViewModel::OnCanAcceptDropInternal
 {
 	bool bAllowDrop = false;
 
-	if(TSharedPtr<FNiagaraHierarchyDragDropOp> HierarchyDragDropOp = StaticCastSharedPtr<FNiagaraHierarchyDragDropOp>(DragDropOp))
+	if(DragDropOp->IsOfType<FNiagaraHierarchyDragDropOp>())
 	{
+		TSharedPtr<FNiagaraHierarchyDragDropOp> HierarchyDragDropOp = StaticCastSharedPtr<FNiagaraHierarchyDragDropOp>(DragDropOp);
 		TSharedPtr<FNiagaraHierarchyItemViewModelBase> SourceDropItem = HierarchyDragDropOp->GetDraggedItem().Pin();
 		TSharedPtr<FNiagaraHierarchyItemViewModelBase> TargetDropItem = AsShared();
 
@@ -948,8 +950,9 @@ TOptional<EItemDropZone> FNiagaraHierarchyItemViewModel::OnCanAcceptDropInternal
 {
 	bool bAllowDrop = false;
 
-	if(TSharedPtr<FNiagaraHierarchyDragDropOp> HierarchyDragDropOp = StaticCastSharedPtr<FNiagaraHierarchyDragDropOp>(DragDropOp))
+	if(DragDropOp->IsOfType<FNiagaraHierarchyDragDropOp>())
 	{
+		TSharedPtr<FNiagaraHierarchyDragDropOp> HierarchyDragDropOp = StaticCastSharedPtr<FNiagaraHierarchyDragDropOp>(DragDropOp);
 		TSharedPtr<FNiagaraHierarchyItemViewModelBase> SourceDropItem = HierarchyDragDropOp->GetDraggedItem().Pin();
 		TSharedPtr<FNiagaraHierarchyItemViewModelBase> TargetDropItem = AsShared();
 
