Fixed a crash when adding a new user parameter while a custom section is selected.
We now swap to the All section when a new user parameter is added.

diff --git a/Engine/Plugins/FX/Niagara/Source/NiagaraEditor/Private/Customizations/NiagaraComponentDetails.cpp b/Engine/Plugins/FX/Niagara/Source/NiagaraEditor/Private/Customizations/NiagaraComponentDetails.cpp
--- a/Engine/Plugins/FX/Niagara/Source/NiagaraEditor/Private/Customizations/NiagaraComponentDetails.cpp
+++ b/Engine/Plugins/FX/Niagara/Source/NiagaraEditor/Private/Customizations/NiagaraComponentDetails.cpp
@@ -1033,7 +1033,7 @@ void FNiagaraSystemUserParameterBuilder::GenerateChildContent(IDetailChildrenBui
 			if(ensure(UserParameterPanelViewModel.IsValid()))
 			{
 				UserParameterPanelViewModel->OnRefreshRequested().BindSP(this, &FNiagaraSystemUserParameterBuilder::Rebuild);
-				UserParameterPanelViewModel->OnParameterAdded().BindSP(this, &FNiagaraSystemUserParameterBuilder::RequestRename);
+				UserParameterPanelViewModel->OnParameterAdded().BindSP(this, &FNiagaraSystemUserParameterBuilder::OnParameterAdded);
 			}
 			
 			bDelegatesInitialized = true;
@@ -1199,10 +1199,26 @@ void FNiagaraSystemUserParameterBuilder::DeleteParameter(FNiagaraVariable UserPa
 	System->HandleVariableRemoved(UserParameter, true);
 }
 
+void FNiagaraSystemUserParameterBuilder::SelectAllSection()
+{
+	if(ActiveSection != nullptr)
+	{
+		ActiveSection = nullptr;
+		OnRebuildChildren.Execute();
+	}
+}
+
+void FNiagaraSystemUserParameterBuilder::OnParameterAdded(FNiagaraVariable UserParameter)
+{
+	// when adding a new parameter it defaults to the 'All' section, so we select it in order to visualize it before renaming
+	SelectAllSection();
+	RequestRename(UserParameter);
+}
+
 void FNiagaraSystemUserParameterBuilder::RequestRename(FNiagaraVariable UserParameter)
 {
 	SelectedParameter = UserParameter;
-
+	
 	// we add a timer for next frame, as requesting a rename on a parameter that has just been created will not preselect the entire text. Waiting a frame fixes that.
 	UserParamToWidgetMap[UserParameter]->RegisterActiveTimer(0.f, FWidgetActiveTimerDelegate::CreateLambda([this, UserParameter](double CurrentTime, float DeltaTime)
 	{
diff --git a/Engine/Plugins/FX/Niagara/Source/NiagaraEditor/Private/Customizations/NiagaraComponentDetails.h b/Engine/Plugins/FX/Niagara/Source/NiagaraEditor/Private/Customizations/NiagaraComponentDetails.h
--- a/Engine/Plugins/FX/Niagara/Source/NiagaraEditor/Private/Customizations/NiagaraComponentDetails.h
+++ b/Engine/Plugins/FX/Niagara/Source/NiagaraEditor/Private/Customizations/NiagaraComponentDetails.h
@@ -216,6 +216,8 @@ private:
 
 	FReply SummonHierarchyEditor();
 	void DeleteParameter(FNiagaraVariable UserParameter) const;
+	void SelectAllSection();
+	void OnParameterAdded(FNiagaraVariable UserParameter);
 	void RequestRename(FNiagaraVariable UserParameter);
 	void RenameParameter(FNiagaraVariable UserParameter, FName NewName);
 
