Fixed a crash when finalizing empty sim caches

diff --git a/Engine/Plugins/FX/Niagara/Source/Niagara/Private/NiagaraSimCache.cpp b/Engine/Plugins/FX/Niagara/Source/Niagara/Private/NiagaraSimCache.cpp
--- a/Engine/Plugins/FX/Niagara/Source/Niagara/Private/NiagaraSimCache.cpp
+++ b/Engine/Plugins/FX/Niagara/Source/Niagara/Private/NiagaraSimCache.cpp
@@ -432,7 +432,7 @@ bool UNiagaraSimCache::EndWrite()
 	}
 
 	// If we allow interpolation we need to build our mapping table from Current -> Next
-	if (CreateParameters.bAllowInterpolation)
+	if (CreateParameters.bAllowInterpolation && CacheFrames.Num() > 0)
 	{
 		const auto& CreateInterpMapping =
 			[this](const FNiagaraSimCacheDataBuffersLayout& BufferLayout, FNiagaraSimCacheDataBuffers& BufferA, const FNiagaraSimCacheDataBuffers& BufferB)
