Fixed Niagara crash when creating a new asset

diff --git a/Engine/Plugins/FX/Niagara/Source/Niagara/Private/NiagaraRenderer.cpp b/Engine/Plugins/FX/Niagara/Source/Niagara/Private/NiagaraRenderer.cpp
--- a/Engine/Plugins/FX/Niagara/Source/Niagara/Private/NiagaraRenderer.cpp
+++ b/Engine/Plugins/FX/Niagara/Source/Niagara/Private/NiagaraRenderer.cpp
@@ -394,7 +394,10 @@ void FNiagaraRenderer::Initialize(const UNiagaraRendererProperties* InProps, con
 				if (NiagaraShaderScript)
 				{
 					FNiagaraShaderRef NiagaraShader = NiagaraShaderScript->GetShaderGameThread(0);
-					bNeedsPartialDepthTexture = (NiagaraShader->MiscUsageBitMask & uint16(ENiagaraScriptMiscUsageMask::UsesPartialDepthCollisionQuery)) > 0;
+					if (NiagaraShader.IsValid())
+					{
+						bNeedsPartialDepthTexture = (NiagaraShader->MiscUsageBitMask & uint16(ENiagaraScriptMiscUsageMask::UsesPartialDepthCollisionQuery)) > 0;
+					}
 				}
 			}
 		}
