Made Niagara partial collision query function be experimental for now.

diff --git a/Engine/Plugins/FX/Niagara/Source/Niagara/Private/NiagaraDataInterfaceCollisionQuery.cpp b/Engine/Plugins/FX/Niagara/Source/Niagara/Private/NiagaraDataInterfaceCollisionQuery.cpp
--- a/Engine/Plugins/FX/Niagara/Source/Niagara/Private/NiagaraDataInterfaceCollisionQuery.cpp
+++ b/Engine/Plugins/FX/Niagara/Source/Niagara/Private/NiagaraDataInterfaceCollisionQuery.cpp
@@ -237,7 +237,9 @@ void UNiagaraDataInterfaceCollisionQuery::GetFunctions(TArray<FNiagaraFunctionSi
 #if WITH_EDITORONLY_DATA
 		SigPartialDepth.FunctionVersion = FNiagaraCollisionDIFunctionVersion::LatestVersion;
 		SigPartialDepth.Description = LOCTEXT("ScenePartialDepthSignatureDescription", "Projects a given world position to view space and then queries the partial depth buffer (opaque emitter using this function are not in this depth buffer) with that position.");
+		SigPartialDepth.ExperimentalMessage = LOCTEXT("ScenePartialDepthSignatureEpxMessage", "Careful: using this node involve a scene depth copy that might slow down the base pass on some platform.");
 #endif
+		SigPartialDepth.bExperimental = true;
 		SigPartialDepth.Inputs.Add(FNiagaraVariable(FNiagaraTypeDefinition(GetClass()), TEXT("CollisionQuery")));
 		SigPartialDepth.AddInput(FNiagaraVariable(FNiagaraTypeDefinition::GetPositionDef(), TEXT("DepthSamplePosWorld")), DepthSamplePosWorldDescription);
 		SigPartialDepth.AddOutput(FNiagaraVariable(FNiagaraTypeDefinition::GetFloatDef(), TEXT("SceneDepth")), SceneDepthDescription);
