Reverted the content directoy fixup as that introduced other issues & fixed the original issue directly

diff --git a/Engine/Plugins/FX/Niagara/Source/NiagaraEditor/Private/NiagaraEditorStyle.cpp b/Engine/Plugins/FX/Niagara/Source/NiagaraEditor/Private/NiagaraEditorStyle.cpp
--- a/Engine/Plugins/FX/Niagara/Source/NiagaraEditor/Private/NiagaraEditorStyle.cpp
+++ b/Engine/Plugins/FX/Niagara/Source/NiagaraEditor/Private/NiagaraEditorStyle.cpp
@@ -373,7 +373,6 @@ void FNiagaraEditorStyle::InitSelectedEmitter()
 
 void FNiagaraEditorStyle::InitToolbarIcons()
 {
-	Set("NiagaraEditor.Refresh", new CORE_IMAGE_BRUSH_SVG("Starship/Common/Update", Icon20x20));
 	Set("NiagaraEditor.ApplyScratchPadChanges", new IMAGE_BRUSH_SVG("Icons/Commands/ApplyScratch", Icon40x40));
 	Set("NiagaraEditor.OverviewNode.IsolatedColor", FLinearColor::Yellow);
 	Set("NiagaraEditor.OverviewNode.NotIsolatedColor", FLinearColor::Transparent);
@@ -791,7 +790,7 @@ void FNiagaraEditorStyle::InitSimCacheEditor()
 FNiagaraEditorStyle::FNiagaraEditorStyle() : FSlateStyleSet("NiagaraEditorStyle")
 {
 	SetContentRoot(FPaths::EnginePluginsDir() / TEXT("FX/Niagara/Content/Slate"));
-	SetCoreContentRoot(FPaths::EngineContentDir() / TEXT("Slate"));
+	SetCoreContentRoot(FPaths::EngineContentDir() / TEXT("Editor/Slate"));
 
 	InitStats();
 	InitAssetPicker();
diff --git a/Engine/Plugins/FX/Niagara/Source/NiagaraEditor/Private/Toolkits/NiagaraScriptToolkit.cpp b/Engine/Plugins/FX/Niagara/Source/NiagaraEditor/Private/Toolkits/NiagaraScriptToolkit.cpp
--- a/Engine/Plugins/FX/Niagara/Source/NiagaraEditor/Private/Toolkits/NiagaraScriptToolkit.cpp
+++ b/Engine/Plugins/FX/Niagara/Source/NiagaraEditor/Private/Toolkits/NiagaraScriptToolkit.cpp
@@ -710,7 +710,7 @@ FText FNiagaraScriptToolkit::GetCompileStatusTooltip() const
 
 FSlateIcon FNiagaraScriptToolkit::GetRefreshStatusImage() const
 {
-	return FSlateIcon(FNiagaraEditorStyle::Get().GetStyleSetName(), "NiagaraEditor.Refresh");
+	return FSlateIcon(FAppStyle::Get().GetStyleSetName(), "FontEditor.Update");
 }
 
 FText FNiagaraScriptToolkit::GetRefreshStatusTooltip() const
diff --git a/Engine/Plugins/FX/Niagara/Source/NiagaraEditorWidgets/Private/NiagaraEditorWidgetsStyle.cpp b/Engine/Plugins/FX/Niagara/Source/NiagaraEditorWidgets/Private/NiagaraEditorWidgetsStyle.cpp
--- a/Engine/Plugins/FX/Niagara/Source/NiagaraEditorWidgets/Private/NiagaraEditorWidgetsStyle.cpp
+++ b/Engine/Plugins/FX/Niagara/Source/NiagaraEditorWidgets/Private/NiagaraEditorWidgetsStyle.cpp
@@ -64,7 +64,7 @@ void FNiagaraEditorWidgetsStyle::ReinitializeStyle()
 FNiagaraEditorWidgetsStyle::FNiagaraEditorWidgetsStyle() : FSlateStyleSet("NiagaraEditorWidgetsStyle")
 {
 	SetContentRoot(FPaths::EnginePluginsDir() / TEXT("FX/Niagara/Content/Slate"));
-	SetCoreContentRoot(FPaths::EngineContentDir() / TEXT("Slate"));
+	SetCoreContentRoot(FPaths::EngineContentDir() / TEXT("Editor/Slate"));
 	
 	const FTextBlockStyle NormalText = FAppStyle::Get().GetWidgetStyle<FTextBlockStyle>("NormalText");
 	const FEditableTextBoxStyle NormalEditableTextBox = FAppStyle::Get().GetWidgetStyle<FEditableTextBoxStyle>("NormalEditableTextBox");
