Fix C++20 compile error. Make FNiagaraHierarchyIdentity's comparison operator const to address ambiguity issues.

diff --git a/Engine/Plugins/FX/Niagara/Source/NiagaraEditor/Public/ViewModels/HierarchyEditor/NiagaraHierarchyViewModelBase.h b/Engine/Plugins/FX/Niagara/Source/NiagaraEditor/Public/ViewModels/HierarchyEditor/NiagaraHierarchyViewModelBase.h
--- a/Engine/Plugins/FX/Niagara/Source/NiagaraEditor/Public/ViewModels/HierarchyEditor/NiagaraHierarchyViewModelBase.h
+++ b/Engine/Plugins/FX/Niagara/Source/NiagaraEditor/Public/ViewModels/HierarchyEditor/NiagaraHierarchyViewModelBase.h
@@ -25,7 +25,7 @@ struct FNiagaraHierarchyIdentity
 	UPROPERTY()
 	TArray<FName> Names;
 
-	bool operator==(const FNiagaraHierarchyIdentity& OtherIdentity)
+	bool operator==(const FNiagaraHierarchyIdentity& OtherIdentity) const
 	{
 		if(Guids.Num() != OtherIdentity.Guids.Num() || Names.Num() != OtherIdentity.Names.Num())
 		{
