Fixed niagara function call node errors and warnings

diff --git a/Engine/Plugins/FX/Niagara/Source/NiagaraEditor/Private/Widgets/SNiagaraGraphNode.cpp b/Engine/Plugins/FX/Niagara/Source/NiagaraEditor/Private/Widgets/SNiagaraGraphNode.cpp
--- a/Engine/Plugins/FX/Niagara/Source/NiagaraEditor/Private/Widgets/SNiagaraGraphNode.cpp
+++ b/Engine/Plugins/FX/Niagara/Source/NiagaraEditor/Private/Widgets/SNiagaraGraphNode.cpp
@@ -62,6 +62,7 @@ void SNiagaraGraphNode::CreateOutputSideAddButton(TSharedPtr<SVerticalBox> Outpu
 void SNiagaraGraphNode::UpdateErrorInfo()
 {
 	const UNiagaraEditorSettings* NiagaraEditorSettings = GetDefault<UNiagaraEditorSettings>();
+	ErrorMsg.Reset();
 	if (NiagaraNode.IsValid())
 	{
 		if (NiagaraEditorSettings->IsAllowedClass(GraphNode->GetClass()) == false || 
@@ -71,7 +72,8 @@ void SNiagaraGraphNode::UpdateErrorInfo()
 			ErrorColor = FAppStyle::GetColor("ErrorReporting.BackgroundColor");
 		}
 	}
-	else
+
+	if(ErrorMsg.IsEmpty())
 	{
 		SGraphNode::UpdateErrorInfo();
 	}
