Fixed missing fwd declare

diff --git a/Engine/Plugins/FX/Niagara/Source/NiagaraEditor/Public/NiagaraParameterMapHistory.h b/Engine/Plugins/FX/Niagara/Source/NiagaraEditor/Public/NiagaraParameterMapHistory.h
--- a/Engine/Plugins/FX/Niagara/Source/NiagaraEditor/Public/NiagaraParameterMapHistory.h
+++ b/Engine/Plugins/FX/Niagara/Source/NiagaraEditor/Public/NiagaraParameterMapHistory.h
@@ -14,6 +14,7 @@ class UEdGraphPin;
 class UEdGraphNode;
 class UNiagaraParameterCollection;
 class FHlslNiagaraTranslator;
+class UNiagaraGraph;
 
 class FCompileConstantResolver
 {
