- Add automatic mode to sprite facing / alignment
- These will change between the old mode & custom depending on if custom attributes are bound

diff --git a/Engine/Plugins/FX/Niagara/Source/Niagara/Private/NiagaraRendererSprites.cpp b/Engine/Plugins/FX/Niagara/Source/Niagara/Private/NiagaraRendererSprites.cpp
--- a/Engine/Plugins/FX/Niagara/Source/Niagara/Private/NiagaraRendererSprites.cpp
+++ b/Engine/Plugins/FX/Niagara/Source/Niagara/Private/NiagaraRendererSprites.cpp
@@ -168,6 +168,18 @@ FNiagaraRendererSprites::FNiagaraRendererSprites(ERHIFeatureLevel::Type FeatureL
 			VFBoundOffsetsInParamStore[i] = INDEX_NONE;
 		}
 	}
+
+	TConstArrayView<FNiagaraRendererVariableInfo> VFVariables = RendererLayoutWithoutCustomSort->GetVFVariables_GameThread();
+	if (Alignment == ENiagaraSpriteAlignment::Automatic)
+	{
+		const int32 RegisterIndex = SourceMode == ENiagaraRendererSourceDataMode::Particles ? VFVariables[ENiagaraSpriteVFLayout::Alignment].GetGPUOffset() : VFBoundOffsetsInParamStore[ENiagaraSpriteVFLayout::Alignment];
+		Alignment = RegisterIndex == INDEX_NONE ? ENiagaraSpriteAlignment::Unaligned : ENiagaraSpriteAlignment::CustomAlignment;
+	}
+	if (FacingMode == ENiagaraSpriteFacingMode::Automatic)
+	{
+		const int32 RegisterIndex = SourceMode == ENiagaraRendererSourceDataMode::Particles ? VFVariables[ENiagaraSpriteVFLayout::Facing].GetGPUOffset() : VFBoundOffsetsInParamStore[ENiagaraSpriteVFLayout::Facing];
+		FacingMode = RegisterIndex == INDEX_NONE ? ENiagaraSpriteFacingMode::FaceCamera : ENiagaraSpriteFacingMode::CustomFacingVector;
+	}
 }
 
 FNiagaraRendererSprites::~FNiagaraRendererSprites()
diff --git a/Engine/Plugins/FX/Niagara/Source/Niagara/Private/NiagaraSpriteRendererProperties.cpp b/Engine/Plugins/FX/Niagara/Source/Niagara/Private/NiagaraSpriteRendererProperties.cpp
--- a/Engine/Plugins/FX/Niagara/Source/Niagara/Private/NiagaraSpriteRendererProperties.cpp
+++ b/Engine/Plugins/FX/Niagara/Source/Niagara/Private/NiagaraSpriteRendererProperties.cpp
@@ -16,6 +16,7 @@
 #include "Internationalization/Internationalization.h"
 #include "Materials/MaterialInstanceConstant.h"
 #include "Modules/ModuleManager.h"
+#include "UObject/UE5MainStreamObjectVersion.h"
 
 #include UE_INLINE_GENERATED_CPP_BY_NAME(NiagaraSpriteRendererProperties)
 
@@ -54,8 +55,8 @@ UNiagaraSpriteRendererProperties::UNiagaraSpriteRendererProperties()
 	: Material(nullptr)
 	, SourceMode(ENiagaraRendererSourceDataMode::Particles)
 	, MaterialUserParamBinding(FNiagaraTypeDefinition(UMaterialInterface::StaticClass()))
-	, Alignment(ENiagaraSpriteAlignment::Unaligned)
-	, FacingMode(ENiagaraSpriteFacingMode::FaceCamera)
+	, Alignment(ENiagaraSpriteAlignment::Automatic)
+	, FacingMode(ENiagaraSpriteFacingMode::Automatic)
 	, PivotInUVSpace(0.5f, 0.5f)
 	, SortMode(ENiagaraSortMode::None)
 	, SubImageSize(1.0f, 1.0f)
@@ -201,13 +202,21 @@ void UNiagaraSpriteRendererProperties::Serialize(FStructuredArchive::FRecord Rec
 {
 	FArchive& Ar = Record.GetUnderlyingArchive();
 	Ar.UsingCustomVersion(FNiagaraCustomVersion::GUID);
+	Ar.UsingCustomVersion(FUE5MainStreamObjectVersion::GUID);
 	const int32 NiagaraVersion = Ar.CustomVer(FNiagaraCustomVersion::GUID);
+	const int32 UE5MainVersion = Ar.CustomVer(FUE5MainStreamObjectVersion::GUID);
 
 	if (Ar.IsLoading() && (NiagaraVersion < FNiagaraCustomVersion::DisableSortingByDefault))
 	{
 		SortMode = ENiagaraSortMode::ViewDistance;
 	}
 
+	if (Ar.IsLoading() && (UE5MainVersion < FUE5MainStreamObjectVersion::NiagaraSpriteRendererFacingAlignmentAutoDefault))
+	{
+		Alignment = ENiagaraSpriteAlignment::Unaligned;
+		FacingMode = ENiagaraSpriteFacingMode::FaceCamera;
+	}
+
 	// MIC will replace the main material during serialize
 	// Be careful if adding code that looks at the material to make sure you get the correct one
 	{
diff --git a/Engine/Plugins/FX/Niagara/Source/Niagara/Public/NiagaraRendererProperties.h b/Engine/Plugins/FX/Niagara/Source/Niagara/Public/NiagaraRendererProperties.h
--- a/Engine/Plugins/FX/Niagara/Source/Niagara/Public/NiagaraRendererProperties.h
+++ b/Engine/Plugins/FX/Niagara/Source/Niagara/Public/NiagaraRendererProperties.h
@@ -140,6 +140,7 @@ struct NIAGARA_API FNiagaraRendererLayout
 	bool SetVariableFromBinding(const FNiagaraDataSetCompiledData* CompiledData, const FNiagaraVariableAttributeBinding& VariableBinding, int32 VFVarOffset);
 	void Finalize();
 
+	TConstArrayView<FNiagaraRendererVariableInfo> GetVFVariables_GameThread() const { check(IsInGameThread() || IsInParallelGameThread()); return MakeArrayView(VFVariables_GT); }
 	TConstArrayView<FNiagaraRendererVariableInfo> GetVFVariables_RenderThread() const { check(IsInRenderingThread()); return MakeArrayView(VFVariables_RT); }
 	int32 GetTotalFloatComponents_RenderThread() const { check(IsInRenderingThread()); return TotalFloatComponents_RT; }
 	int32 GetTotalHalfComponents_RenderThread() const { check(IsInRenderingThread()); return TotalHalfComponents_RT; }
diff --git a/Engine/Plugins/FX/Niagara/Source/Niagara/Public/NiagaraSpriteRendererProperties.h b/Engine/Plugins/FX/Niagara/Source/Niagara/Public/NiagaraSpriteRendererProperties.h
--- a/Engine/Plugins/FX/Niagara/Source/Niagara/Public/NiagaraSpriteRendererProperties.h
+++ b/Engine/Plugins/FX/Niagara/Source/Niagara/Public/NiagaraSpriteRendererProperties.h
@@ -21,7 +21,10 @@ enum class ENiagaraSpriteAlignment : uint8
 	/** Imagine the particle texture having an arrow pointing up, this mode makes the arrow point in the direction of the Particles.Velocity attribute. FacingMode is ignored unless CustomFacingVector is set.*/
 	VelocityAligned,
 	/** Imagine the particle texture having an arrow pointing up, this mode makes the arrow point towards the axis defined by the "Particles.SpriteAlignment" attribute. FacingMode is ignored unless CustomFacingVector is set. If the "Particles.SpriteAlignment" attribute is missing, this falls back to Unaligned mode.*/
-	CustomAlignment
+	CustomAlignment,
+
+	/** Automatically select between Unaligned & CustomAlignment depending on if SpriteAlignment Binding is valid. */
+	Automatic
 };
 
 
@@ -38,7 +41,10 @@ enum class ENiagaraSpriteFacingMode : uint8
 	/** Faces the camera position, but is not dependent on the camera rotation.  This method produces more stable particles under camera rotation. Uses the up axis of (0,0,1).*/
 	FaceCameraPosition,
 	/** Blends between FaceCamera and FaceCameraPosition.*/
-	FaceCameraDistanceBlend
+	FaceCameraDistanceBlend,
+
+	/** Automatically select between FaceCamera & CustomFacingVector depending on if SpriteFacing binding is valid. */
+	Automatic
 };
 
 UENUM()
diff --git a/Engine/Source/Runtime/Core/Public/UObject/UE5MainStreamObjectVersion.h b/Engine/Source/Runtime/Core/Public/UObject/UE5MainStreamObjectVersion.h
--- a/Engine/Source/Runtime/Core/Public/UObject/UE5MainStreamObjectVersion.h
+++ b/Engine/Source/Runtime/Core/Public/UObject/UE5MainStreamObjectVersion.h
@@ -340,6 +340,9 @@ struct CORE_API FUE5MainStreamObjectVersion
 		// Smart names removed from animation systems
 		AnimationRemoveSmartNames,
 
+		// Change the default for facing & alignment to be automatic
+		NiagaraSpriteRendererFacingAlignmentAutoDefault,
+
 		// -----<new versions can be added above this line>-------------------------------------------------
 		VersionPlusOne,
 		LatestVersion = VersionPlusOne - 1
