- Handle null component when getting the LWC converter for arrays

#rnx

diff --git a/Engine/Plugins/FX/Niagara/Source/Niagara/Private/NiagaraDataInterfaceArrayFunctionLibrary.cpp b/Engine/Plugins/FX/Niagara/Source/Niagara/Private/NiagaraDataInterfaceArrayFunctionLibrary.cpp
--- a/Engine/Plugins/FX/Niagara/Source/Niagara/Private/NiagaraDataInterfaceArrayFunctionLibrary.cpp
+++ b/Engine/Plugins/FX/Niagara/Source/Niagara/Private/NiagaraDataInterfaceArrayFunctionLibrary.cpp
@@ -112,6 +112,11 @@ TArrayType GetNiagaraArrayValue(UNiagaraComponent* NiagaraComponent, FName Overr
 
 FNiagaraLWCConverter GetLWCConverter(UNiagaraComponent* NiagaraComponent)
 {
+	if (NiagaraComponent == nullptr)
+	{
+		return FNiagaraLWCConverter();
+	}
+
 	FNiagaraSystemInstanceControllerPtr SystemInstanceController = NiagaraComponent->GetSystemInstanceController();
 	if (!SystemInstanceController.IsValid())
 	{
