- Add Decal Sort Order which auto binds to Particles.UniqueID
- This fixes flickering issues when we have overlapping decals

#rnx

diff --git a/Engine/Plugins/FX/Niagara/Source/Niagara/Private/NiagaraDecalRendererProperties.cpp b/Engine/Plugins/FX/Niagara/Source/Niagara/Private/NiagaraDecalRendererProperties.cpp
--- a/Engine/Plugins/FX/Niagara/Source/Niagara/Private/NiagaraDecalRendererProperties.cpp
+++ b/Engine/Plugins/FX/Niagara/Source/Niagara/Private/NiagaraDecalRendererProperties.cpp
@@ -67,6 +67,7 @@ namespace NiagaraDecalRendererPropertiesLocal
 		Props->DecalOrientationBinding.Setup(GetDecalOrientationVariable(), GetDecalOrientationVariable());
 		Props->DecalSizeBinding.Setup(GetDecalSizeVariable(), GetDecalSizeVariable());
 		Props->DecalFadeBinding.Setup(GetDecalFadeVariable(), GetDecalFadeVariable());
+		Props->DecalSortOrderBinding = FNiagaraConstants::GetAttributeDefaultBinding(SYS_PARAM_PARTICLES_UNIQUE_ID);
 		Props->DecalColorBinding = FNiagaraConstants::GetAttributeDefaultBinding(SYS_PARAM_PARTICLES_COLOR);
 		Props->DecalVisibleBinding.Setup(GetDecalVisibleVariable(), GetDecalVisibleVariable());
 		Props->RendererVisibilityTagBinding = FNiagaraConstants::GetAttributeDefaultBinding(SYS_PARAM_PARTICLES_VISIBILITY_TAG);
@@ -85,6 +86,7 @@ UNiagaraDecalRendererProperties::UNiagaraDecalRendererProperties()
 	AttributeBindings.Add(&DecalOrientationBinding);
 	AttributeBindings.Add(&DecalSizeBinding);
 	AttributeBindings.Add(&DecalFadeBinding);
+	AttributeBindings.Add(&DecalSortOrderBinding);
 	AttributeBindings.Add(&DecalColorBinding);
 	AttributeBindings.Add(&DecalVisibleBinding);
 	AttributeBindings.Add(&RendererVisibilityTagBinding);
@@ -244,6 +246,7 @@ void UNiagaraDecalRendererProperties::CacheFromCompiledData(const FNiagaraDataSe
 	DecalOrientationDataSetAccessor.Init(CompiledData, DecalOrientationBinding.GetDataSetBindableVariable().GetName());
 	DecalSizeDataSetAccessor.Init(CompiledData, DecalSizeBinding.GetDataSetBindableVariable().GetName());
 	DecalFadeDataSetAccessor.Init(CompiledData, DecalFadeBinding.GetDataSetBindableVariable().GetName());
+	DecalSortOrderDataSetAccessor.Init(CompiledData, DecalSortOrderBinding.GetDataSetBindableVariable().GetName());
 	DecalColorDataSetAccessor.Init(CompiledData, DecalColorBinding.GetDataSetBindableVariable().GetName());
 	DecalVisibleAccessor.Init(CompiledData, DecalVisibleBinding.GetDataSetBindableVariable().GetName());
 	RendererVisibilityTagAccessor.Init(CompiledData, RendererVisibilityTagBinding.GetDataSetBindableVariable().GetName());
@@ -297,6 +300,7 @@ const TArray<FNiagaraVariable>& UNiagaraDecalRendererProperties::GetOptionalAttr
 		Attrs.Add(GetDecalSizeVariable());
 		Attrs.Add(GetDecalOrientationVariable());
 		Attrs.Add(GetDecalFadeVariable());
+		Attrs.Add(SYS_PARAM_PARTICLES_UNIQUE_ID);
 		Attrs.Add(SYS_PARAM_PARTICLES_VISIBILITY_TAG);
 	}
 	return Attrs;
diff --git a/Engine/Plugins/FX/Niagara/Source/Niagara/Private/NiagaraRendererDecals.cpp b/Engine/Plugins/FX/Niagara/Source/Niagara/Private/NiagaraRendererDecals.cpp
--- a/Engine/Plugins/FX/Niagara/Source/Niagara/Private/NiagaraRendererDecals.cpp
+++ b/Engine/Plugins/FX/Niagara/Source/Niagara/Private/NiagaraRendererDecals.cpp
@@ -159,6 +159,7 @@ FNiagaraDynamicDataBase* FNiagaraRendererDecals::GenerateDynamicData(const FNiag
 	const FVector3f DefaultSize = UNiagaraDecalRendererProperties::GetDefaultDecalSize();
 	const FNiagaraBool DefaultVisible = UNiagaraDecalRendererProperties::GetDefaultDecalVisible();
 	const float DefaultFade = UNiagaraDecalRendererProperties::GetDefaultDecalFade();
+	const int32 DefaultSortOrder = UNiagaraDecalRendererProperties::GetDefaultDecalSortOrder();
 
 	// Check for Material Update
 	UMaterialInterface* Material = BaseMaterials_GT.Num() ? BaseMaterials_GT[0] : nullptr;
@@ -188,6 +189,7 @@ FNiagaraDynamicDataBase* FNiagaraRendererDecals::GenerateDynamicData(const FNiag
 		const auto SizeReader = RendererProperties->DecalSizeDataSetAccessor.GetReader(DataSet);
 		const auto ColorReader = RendererProperties->DecalColorDataSetAccessor.GetReader(DataSet);
 		const auto FadeReader = RendererProperties->DecalFadeDataSetAccessor.GetReader(DataSet);
+		const auto SortOrderReader = RendererProperties->DecalSortOrderDataSetAccessor.GetReader(DataSet);
 		const auto VisibleReader = RendererProperties->DecalVisibleAccessor.GetReader(DataSet);
 		const auto VisTagReader = RendererProperties->RendererVisibilityTagAccessor.GetReader(DataSet);
 
@@ -209,6 +211,7 @@ FNiagaraDynamicDataBase* FNiagaraRendererDecals::GenerateDynamicData(const FNiag
 			const FVector Size = FVector(SizeReader.GetSafe(ParticleIndex, DefaultSize) * 0.5f);
 			const FLinearColor DecalColor = ColorReader.GetSafe(ParticleIndex, FLinearColor::White);
 			const float Fade = FadeReader.GetSafe(ParticleIndex, DefaultFade);
+			const int32 SortOrder = SortOrderReader.GetSafe(ParticleIndex, DefaultSortOrder);
 
 			// Create Update Parameters
 			FDeferredDecalUpdateParams& UpdateParams = DecalUpdates.AddDefaulted_GetRef();
@@ -232,6 +235,7 @@ FNiagaraDynamicDataBase* FNiagaraRendererDecals::GenerateDynamicData(const FNiag
 			UpdateParams.FadeDuration = 1.0f;
 			UpdateParams.FadeScreenSize = RendererProperties->DecalScreenSizeFade;
 			UpdateParams.DecalColor = DecalColor;
+			UpdateParams.SortOrder = SortOrder;
 		}
 	}
 	// Emitter source mode
@@ -250,6 +254,7 @@ FNiagaraDynamicDataBase* FNiagaraRendererDecals::GenerateDynamicData(const FNiag
 			const FVector Size = FVector(ParameterStore.GetParameterValueOrDefault(RendererProperties->DecalSizeBinding.GetParamMapBindableVariable(), DefaultSize) * 0.5f);
 			const FLinearColor DecalColor = ParameterStore.GetParameterValueOrDefault(RendererProperties->DecalColorBinding.GetParamMapBindableVariable(), FLinearColor::White);
 			const float Fade = ParameterStore.GetParameterValueOrDefault(RendererProperties->DecalFadeBinding.GetParamMapBindableVariable(), DefaultFade);
+			const int32 SortOrder = ParameterStore.GetParameterValueOrDefault(RendererProperties->DecalSortOrderBinding.GetParamMapBindableVariable(), DefaultSortOrder);
 
 			// Create Update Parameters
 			FDeferredDecalUpdateParams& UpdateParams = DecalUpdates.AddDefaulted_GetRef();
@@ -273,6 +278,7 @@ FNiagaraDynamicDataBase* FNiagaraRendererDecals::GenerateDynamicData(const FNiag
 			UpdateParams.FadeDuration = 1.0f;
 			UpdateParams.FadeScreenSize = RendererProperties->DecalScreenSizeFade;
 			UpdateParams.DecalColor = DecalColor;
+			UpdateParams.SortOrder = SortOrder;
 		}
 	}
 
diff --git a/Engine/Plugins/FX/Niagara/Source/Niagara/Public/NiagaraDecalRendererProperties.h b/Engine/Plugins/FX/Niagara/Source/Niagara/Public/NiagaraDecalRendererProperties.h
--- a/Engine/Plugins/FX/Niagara/Source/Niagara/Public/NiagaraDecalRendererProperties.h
+++ b/Engine/Plugins/FX/Niagara/Source/Niagara/Public/NiagaraDecalRendererProperties.h
@@ -63,6 +63,7 @@ public:
 	static const FQuat4f GetDefaultOrientation() { return FRotator3f(-90.0f, 0.0f, 90.0f).Quaternion(); }
 	static const FVector3f GetDefaultDecalSize() { return FVector3f(50.0f, 50.0f, 50.0f); }
 	static const float GetDefaultDecalFade() { return 1.0f; }
+	static const int32 GetDefaultDecalSortOrder() { return 0; }
 	static const FNiagaraBool GetDefaultDecalVisible() { return FNiagaraBool(true); }
 
 	/** What material to use for the decal. */
@@ -106,6 +107,10 @@ public:
 	UPROPERTY(EditAnywhere, Category = "Bindings")
 	FNiagaraVariableAttributeBinding DecalFadeBinding;
 
+	/** Sort Order binding for the decal, used to ensure decals to not flicker when overlapping. */
+	UPROPERTY(EditAnywhere, Category = "Bindings")
+	FNiagaraVariableAttributeBinding DecalSortOrderBinding;
+
 	/** Color binding for the decal, value can be queried using the Decal Color material node. */
 	UPROPERTY(EditAnywhere, Category = "Bindings")
 	FNiagaraVariableAttributeBinding DecalColorBinding;
@@ -126,6 +131,7 @@ public:
 	FNiagaraDataSetAccessor<FQuat4f>			DecalOrientationDataSetAccessor;
 	FNiagaraDataSetAccessor<FVector3f>			DecalSizeDataSetAccessor;
 	FNiagaraDataSetAccessor<float>				DecalFadeDataSetAccessor;
+	FNiagaraDataSetAccessor<int32>				DecalSortOrderDataSetAccessor;
 	FNiagaraDataSetAccessor<FLinearColor>		DecalColorDataSetAccessor;
 	FNiagaraDataSetAccessor<FNiagaraBool>		DecalVisibleAccessor;
 	FNiagaraDataSetAccessor<int32>				RendererVisibilityTagAccessor;
