- Fix crash if the Niagara system completes while in ticking multiple frames in desired age mode

diff --git a/Engine/Plugins/FX/Niagara/Source/Niagara/Private/NiagaraComponent.cpp b/Engine/Plugins/FX/Niagara/Source/Niagara/Private/NiagaraComponent.cpp
--- a/Engine/Plugins/FX/Niagara/Source/Niagara/Private/NiagaraComponent.cpp
+++ b/Engine/Plugins/FX/Niagara/Source/Niagara/Private/NiagaraComponent.cpp
@@ -837,6 +837,10 @@ void UNiagaraComponent::TickComponent(float DeltaSeconds, enum ELevelTick TickTy
 							// Cannot do multiple tick off the game thread here without additional work. So we pass in null for the completion event which will force GT execution.
 							// If this becomes a perf problem I can add a new path for the tick code to handle multiple ticks.
 							SystemInstanceController->ManualTick(TickDelta, nullptr);
+							if (IsComplete())
+							{
+								break;
+							}
 							--TicksToProcess;
 
 							// This limits the amount of time we will consume seeking forward
