Stopped considering input pins for get nodes as parameter reference

This fixes an issue where a pin that was only referenced in a get node would count 2 references, which would stop it from being renamed properly from pin as it would duplicate a new parameter

diff --git a/Engine/Plugins/FX/Niagara/Source/NiagaraEditor/Private/NiagaraGraph.cpp b/Engine/Plugins/FX/Niagara/Source/NiagaraEditor/Private/NiagaraGraph.cpp
--- a/Engine/Plugins/FX/Niagara/Source/NiagaraEditor/Private/NiagaraGraph.cpp
+++ b/Engine/Plugins/FX/Niagara/Source/NiagaraEditor/Private/NiagaraGraph.cpp
@@ -3843,7 +3843,19 @@ void UNiagaraGraph::RefreshParameterReferences() const
 			{
 				if (WriteEvent.Pin->PinType.PinSubCategory == UNiagaraNodeParameterMapBase::ParameterPinSubCategory)
 				{
-					AddParameterReference(Parameter, WriteEvent.Pin);
+					bool bAddReference = true;
+
+					// we exclude default pins from adding references as they shouldn't count. Originally they weren't marked as parameter pins so this wasn't necessary.
+					if(WriteEvent.Pin->GetOwningNode()->IsA<UNiagaraNodeParameterMapGet>() && WriteEvent.Pin->Direction == EGPD_Input)
+					{
+						bAddReference = false;
+					}
+
+					if(bAddReference)
+					{
+						AddParameterReference(Parameter, WriteEvent.Pin);
+					}
+					
 					HandledParameterMapPins.Add(WriteEvent.Pin);
 				}
 			}
@@ -3852,7 +3864,7 @@ void UNiagaraGraph::RefreshParameterReferences() const
 			{
 				if (ReadHistory.ReadPin.Pin->PinType.PinSubCategory == UNiagaraNodeParameterMapBase::ParameterPinSubCategory)
 				{
-					AddParameterReference(Parameter, ReadHistory.ReadPin.Pin);
+					AddParameterReference(Parameter, ReadHistory.ReadPin.Pin);					
 					HandledParameterMapPins.Add(ReadHistory.ReadPin.Pin);
 				}
 			}
