Refactored a temporary array of transient stack entries that could get GC'd anytime and cause various crashes.
These stack entries managed stack entries for comment nodes, so now we have a proper comment collection stack entry.
Removed all the previously hacked-in code

diff --git a/Engine/Plugins/FX/Niagara/Source/NiagaraEditor/Private/ViewModels/NiagaraOverviewGraphViewModel.cpp b/Engine/Plugins/FX/Niagara/Source/NiagaraEditor/Private/ViewModels/NiagaraOverviewGraphViewModel.cpp
--- a/Engine/Plugins/FX/Niagara/Source/NiagaraEditor/Private/ViewModels/NiagaraOverviewGraphViewModel.cpp
+++ b/Engine/Plugins/FX/Niagara/Source/NiagaraEditor/Private/ViewModels/NiagaraOverviewGraphViewModel.cpp
@@ -19,6 +19,7 @@
 #include "ViewModels/Stack/NiagaraStackObject.h"
 #include "ViewModels/Stack/NiagaraStackEntry.h"
 #include "NiagaraStackEditorData.h"
+#include "ViewModels/Stack/NiagaraStackRoot.h"
 
 #define LOCTEXT_NAMESPACE "NiagaraOverviewGraphViewModel"
 
@@ -45,14 +46,6 @@ FNiagaraOverviewGraphViewModel::~FNiagaraOverviewGraphViewModel()
 			SystemViewModelPinned->GetSelectionViewModel()->OnEntrySelectionChanged().RemoveAll(this);
 		}
 	}
-
-	// Clean up stale entries
-	for (TWeakObjectPtr<UNiagaraStackObject>& TempStackObj : TempEntries)
-	{
-		if (TempStackObj.IsValid())
-			TempStackObj->Finalize();
-	}
-	TempEntries.Empty();
 }
 
 void FNiagaraOverviewGraphViewModel::Initialize(TSharedRef<FNiagaraSystemViewModel> InSystemViewModel)
@@ -466,63 +459,39 @@ void FNiagaraOverviewGraphViewModel::GraphSelectionChanged()
 		TGuardValue<bool> UpdateGuard(bUpdatingSystemSelectionFromGraph, true);
 		bool bClearCurrentSelection = FSlateApplication::Get().GetModifierKeys().IsControlDown() == false;
 
+		UNiagaraStackEntry* RootEntry = GetSystemViewModel()->GetSystemStackViewModel()->GetRootEntry();
+		UNiagaraStackCommentCollection* CommentCollection = nullptr;
+
+		// we refresh the comment collection here to sync it with the current node selection
+		if(UNiagaraStackRoot* AsRoot = Cast<UNiagaraStackRoot>(RootEntry))
+		{
+			CommentCollection = AsRoot->GetCommentCollection();
+			CommentCollection->RefreshChildren();
+		}
+		
 		TArray<UNiagaraStackEntry*> EntriesToSelect;
 		TArray<UNiagaraStackEntry*> EntriesToDeselect;
 
 		TArray<FGuid> SelectedGuids;
-		TArray<TWeakObjectPtr<UNiagaraStackObject>> NewTempEntries;
-		if (!bClearCurrentSelection)
-			NewTempEntries = TempEntries;
-		TArray<TWeakObjectPtr<UNiagaraStackObject>> ClearTempEntries = TempEntries;
+		TArray<TObjectPtr<UNiagaraStackObject>> NewTempEntries;
 
 		for (UObject* SelectedObject : NodeSelection->GetSelectedObjects())
 		{
-			UNiagaraOverviewNode* SelectedOverviewNode = Cast<UNiagaraOverviewNode>(SelectedObject);
-			if (SelectedOverviewNode != nullptr)
+			if (UNiagaraOverviewNode* SelectedOverviewNode = Cast<UNiagaraOverviewNode>(SelectedObject))
 			{
 				SelectedGuids.AddUnique(SelectedOverviewNode->GetEmitterHandleGuid());
 			}
-			else
+			// we find the stack object mapped to the comment node we have selected in the graph to select it in the stack;
+			// this should always exist at this point as we have synced above
+			else if(UEdGraphNode_Comment* CommentNode = Cast<UEdGraphNode_Comment>(SelectedObject))
 			{
-				// Add other selected nodes as generic object entries
-				// See if we already have this object selected, if so, reuse
-				UNiagaraStackObject* StackObject = nullptr;
-				for (TWeakObjectPtr<UNiagaraStackObject>& TempStackObj : TempEntries)
-				{
-					if (TempStackObj.IsValid() && TempStackObj->GetObject() == SelectedObject)
-					{
-						StackObject = TempStackObj.Get();
-						break;
-					}
-				}
-
-				// If not already in existence, create a temp one
-				if (!StackObject)
+				if (UNiagaraStackObject* CommentStackObject = CommentCollection->FindStackObjectForCommentNode(CommentNode))
 				{
-					StackObject = NewObject<UNiagaraStackObject>(GetTransientPackage());
-					UNiagaraStackEditorData* EditorData = NewObject< UNiagaraStackEditorData >(GetTransientPackage());
-					UNiagaraStackEntry::FRequiredEntryData RequiredEntryData(GetSystemViewModel(), nullptr, TEXT("Custom"), NAME_None, *EditorData);
-					bool bIsTopLevelObject = true;
-					StackObject->Initialize(RequiredEntryData, SelectedObject, bIsTopLevelObject, TEXT(""));
-					NewTempEntries.AddUnique(StackObject);
-				}
-
-				if (StackObject)
-				{
-					EntriesToSelect.Add(StackObject);
-					ClearTempEntries.Remove(StackObject);
+					EntriesToSelect.Add(CommentStackObject);
 				}
 			}
 		}
 
-		// Keep track of previously created temp entries.
-		TempEntries = NewTempEntries;
-		for (TWeakObjectPtr<UNiagaraStackObject>& TempStackObj : ClearTempEntries)
-		{
-			if (TempStackObj.IsValid())
-				EntriesToDeselect.Add(TempStackObj.Get());
-		}
-
 		UNiagaraStackEntry* SystemRootEntry = GetSystemViewModel()->GetSystemStackViewModel()->GetRootEntry();
 		if (SelectedGuids.Contains(FGuid()))
 		{
@@ -547,13 +516,6 @@ void FNiagaraOverviewGraphViewModel::GraphSelectionChanged()
 		}
 
 		GetSystemViewModel()->GetSelectionViewModel()->UpdateSelectedEntries(EntriesToSelect, EntriesToDeselect, bClearCurrentSelection);
-
-		// Clean up stale entries
-		for (TWeakObjectPtr<UNiagaraStackObject>& TempStackObj : ClearTempEntries)
-		{
-			if (TempStackObj.IsValid())
-				TempStackObj->Finalize();
-		}
 	}
 }
 
diff --git a/Engine/Plugins/FX/Niagara/Source/NiagaraEditor/Private/ViewModels/Stack/NiagaraStackCommentCollection.cpp b/Engine/Plugins/FX/Niagara/Source/NiagaraEditor/Private/ViewModels/Stack/NiagaraStackCommentCollection.cpp
new file mode 100644
--- /dev/null
+++ b/Engine/Plugins/FX/Niagara/Source/NiagaraEditor/Private/ViewModels/Stack/NiagaraStackCommentCollection.cpp
@@ -0,0 +1,58 @@
+// Copyright Epic Games, Inc. All Rights Reserved.
+
+#include "ViewModels/Stack/NiagaraStackCommentCollection.h"
+
+#include "EdGraphNode_Comment.h"
+#include "NiagaraObjectSelection.h"
+#include "ViewModels/NiagaraOverviewGraphViewModel.h"
+#include "ViewModels/NiagaraSystemViewModel.h"
+#include "ViewModels/Stack/NiagaraStackObject.h"
+#include "ViewModels/Stack/NiagaraStackViewModel.h"
+
+void UNiagaraStackCommentCollection::Initialize(FRequiredEntryData InRequiredEntryData)
+{
+	const FString StackKey = TEXT("CommentCollection");
+	Super::Initialize(InRequiredEntryData, StackKey);
+}
+
+UNiagaraStackObject* UNiagaraStackCommentCollection::FindStackObjectForCommentNode(UEdGraphNode_Comment* CommentNode) const
+{
+	TArray<UNiagaraStackObject*> CommentStackEntries;
+	GetUnfilteredChildrenOfType(CommentStackEntries);
+
+	UNiagaraStackObject** FoundObject = CommentStackEntries.FindByPredicate([CommentNode](UNiagaraStackObject* CandidateStackObject)
+	{
+		return CandidateStackObject->GetObject() == CommentNode;
+	});
+
+	if(FoundObject)
+	{
+		return *FoundObject;
+	}
+
+	return nullptr;
+}
+
+void UNiagaraStackCommentCollection::RefreshChildrenInternal(const TArray<UNiagaraStackEntry*>& CurrentChildren, TArray<UNiagaraStackEntry*>& NewChildren, TArray<FStackIssue>& NewIssues)
+{
+	const TSet<UObject*>& SelectedNodes = GetSystemViewModel()->GetOverviewGraphViewModel()->GetNodeSelection()->GetSelectedObjects();
+
+	for(UObject* SelectedNode : SelectedNodes)
+	{
+		if(UEdGraphNode_Comment* CommentNode = Cast<UEdGraphNode_Comment>(SelectedNode))
+		{
+			UNiagaraStackObject* StackObject = FindCurrentChildOfTypeByPredicate<UNiagaraStackObject>(CurrentChildren, [&](UNiagaraStackObject* CandidateStackObject)
+			{
+				return CandidateStackObject->GetObject() == CommentNode;
+			});
+			
+			if(StackObject == nullptr)
+			{
+			    StackObject = NewObject<UNiagaraStackObject>(this);
+			    StackObject->Initialize(CreateDefaultChildRequiredData(), CommentNode, true, GetStackEditorDataKey(), nullptr);
+			}
+			
+			NewChildren.Add(StackObject);
+		}
+	}
+}
diff --git a/Engine/Plugins/FX/Niagara/Source/NiagaraEditor/Private/ViewModels/Stack/NiagaraStackRoot.cpp b/Engine/Plugins/FX/Niagara/Source/NiagaraEditor/Private/ViewModels/Stack/NiagaraStackRoot.cpp
--- a/Engine/Plugins/FX/Niagara/Source/NiagaraEditor/Private/ViewModels/Stack/NiagaraStackRoot.cpp
+++ b/Engine/Plugins/FX/Niagara/Source/NiagaraEditor/Private/ViewModels/Stack/NiagaraStackRoot.cpp
@@ -116,6 +116,16 @@ void UNiagaraStackRoot::RefreshChildrenInternal(const TArray<UNiagaraStackEntry*
 		SystemUpdateGroup->Initialize(RequiredEntryData, DisplayName, ToolTip, GetSystemViewModel()->GetSystemScriptViewModel().ToSharedRef(), ENiagaraScriptUsage::SystemUpdateScript);
 	}
 
+	// we only allow comments in systems right now, not emitters
+	if(bIncludeSystemInformation && CommentCollection == nullptr)
+	{
+		CommentCollection = NewObject<UNiagaraStackCommentCollection>(this);
+		FRequiredEntryData RequiredEntryData(GetSystemViewModel(), GetEmitterViewModel(),
+			NAME_None, NAME_None,
+			GetSystemViewModel()->GetEditorData().GetStackEditorData());
+		CommentCollection->Initialize(RequiredEntryData);
+	}
+	
 	bool bShouldShowSummaryView = GetEmitterViewModel()? GetEmitterViewModel()->GetOrCreateEditorData().ShouldShowSummaryView() : false;
 
 	// we clear references to entries we don't want to display in summary view.
@@ -217,7 +227,8 @@ void UNiagaraStackRoot::RefreshChildrenInternal(const TArray<UNiagaraStackEntry*
 		NewChildren.Add(SystemUserParametersGroup);
 		NewChildren.Add(SystemSpawnGroup);
 		NewChildren.Add(SystemUpdateGroup);
-	}
+		NewChildren.Add(CommentCollection);
+	}	
 
 	if (bIncludeEmitterInformation)
 	{
diff --git a/Engine/Plugins/FX/Niagara/Source/NiagaraEditor/Public/ViewModels/NiagaraOverviewGraphViewModel.h b/Engine/Plugins/FX/Niagara/Source/NiagaraEditor/Public/ViewModels/NiagaraOverviewGraphViewModel.h
--- a/Engine/Plugins/FX/Niagara/Source/NiagaraEditor/Public/ViewModels/NiagaraOverviewGraphViewModel.h
+++ b/Engine/Plugins/FX/Niagara/Source/NiagaraEditor/Public/ViewModels/NiagaraOverviewGraphViewModel.h
@@ -96,8 +96,6 @@ private:
 	/** The set of nodes objects currently selected in the graph. */
 	TSharedRef<FNiagaraObjectSelection> NodeSelection;
 
-	TArray<TWeakObjectPtr<class UNiagaraStackObject> > TempEntries;
-
 	bool bUpdatingSystemSelectionFromGraph;
 
 	bool bUpdatingGraphSelectionFromSystem;
diff --git a/Engine/Plugins/FX/Niagara/Source/NiagaraEditor/Public/ViewModels/Stack/NiagaraStackCommentCollection.h b/Engine/Plugins/FX/Niagara/Source/NiagaraEditor/Public/ViewModels/Stack/NiagaraStackCommentCollection.h
new file mode 100644
--- /dev/null
+++ b/Engine/Plugins/FX/Niagara/Source/NiagaraEditor/Public/ViewModels/Stack/NiagaraStackCommentCollection.h
@@ -0,0 +1,27 @@
+// Copyright Epic Games, Inc. All Rights Reserved.
+
+#pragma once
+
+#include "EdGraphNode_Comment.h"
+#include "NiagaraStackObject.h"
+#include "ViewModels/Stack/NiagaraStackEntry.h"
+#include "NiagaraStackCommentCollection.generated.h"
+
+UCLASS()
+class NIAGARAEDITOR_API UNiagaraStackCommentCollection : public UNiagaraStackEntry
+{
+	GENERATED_BODY()
+
+public:
+	UNiagaraStackCommentCollection() {}
+
+	void Initialize(FRequiredEntryData InRequiredEntryData);
+
+	virtual bool GetShouldShowInOverview() const override { return false; }
+	virtual bool GetShouldShowInStack() const override { return false; }
+
+	UNiagaraStackObject* FindStackObjectForCommentNode(UEdGraphNode_Comment* CommentNode) const;
+	
+protected:
+	virtual void RefreshChildrenInternal(const TArray<UNiagaraStackEntry*>& CurrentChildren, TArray<UNiagaraStackEntry*>& NewChildren, TArray<FStackIssue>& NewIssues) override;
+};
\ No newline at end of file
diff --git a/Engine/Plugins/FX/Niagara/Source/NiagaraEditor/Public/ViewModels/Stack/NiagaraStackRoot.h b/Engine/Plugins/FX/Niagara/Source/NiagaraEditor/Public/ViewModels/Stack/NiagaraStackRoot.h
--- a/Engine/Plugins/FX/Niagara/Source/NiagaraEditor/Public/ViewModels/Stack/NiagaraStackRoot.h
+++ b/Engine/Plugins/FX/Niagara/Source/NiagaraEditor/Public/ViewModels/Stack/NiagaraStackRoot.h
@@ -2,6 +2,7 @@
 
 #pragma once
 
+#include "NiagaraStackCommentCollection.h"
 #include "ViewModels/Stack/NiagaraStackEntry.h"
 #include "NiagaraStackRoot.generated.h"
 
@@ -31,6 +32,12 @@ public:
 	{
 		return RenderGroup;
 	}
+
+	UNiagaraStackCommentCollection* GetCommentCollection() const
+	{
+		return CommentCollection;
+	}
+	
 protected:
 	virtual void RefreshChildrenInternal(const TArray<UNiagaraStackEntry*>& CurrentChildren, TArray<UNiagaraStackEntry*>& NewChildren, TArray<FStackIssue>& NewIssues) override;
 
@@ -72,6 +79,9 @@ private:
 	UPROPERTY()
 	TObjectPtr<UNiagaraStackRenderItemGroup> RenderGroup;
 
+	UPROPERTY()
+	TObjectPtr<UNiagaraStackCommentCollection> CommentCollection;
+	
 	UPROPERTY()
 	TObjectPtr<UNiagaraStackSummaryViewCollapseButton> SummaryCollapseButton;
 
