- Add debug NDI proxy ref counting
- Allows us to validate that the contract of proxies outliving all RT commands is true

#rnx

diff --git a/Engine/Plugins/FX/Niagara/Source/Niagara/Classes/NiagaraDataInterface.h b/Engine/Plugins/FX/Niagara/Source/Niagara/Classes/NiagaraDataInterface.h
--- a/Engine/Plugins/FX/Niagara/Source/Niagara/Classes/NiagaraDataInterface.h
+++ b/Engine/Plugins/FX/Niagara/Source/Niagara/Classes/NiagaraDataInterface.h
@@ -50,6 +50,10 @@ struct FNDITransformHandler
 	FORCEINLINE void TransformRotation(FQuat4d& Q1, const FQuat4d& Q2) const { Q1 = Q2 * Q1; }
 };
 
+// FNiagaraDataInterfaceProxy should always outlive any ticks, etc, that are on the rendering thread.
+// Enabling this will incur cost but will validate that this contract is not broken
+#define NIAGARA_VALIDATE_NDIPROXY_REFS	!UE_BUILD_SHIPPING
+
 //////////////////////////////////////////////////////////////////////////
 // Some helper classes allowing neat, init time binding of templated vm external functions.
 
@@ -330,7 +334,12 @@ protected:
 struct FNiagaraDataInterfaceProxy
 {
 	FNiagaraDataInterfaceProxy() {}
-	virtual ~FNiagaraDataInterfaceProxy() {/*check(IsInRenderingThread());*/}
+	virtual ~FNiagaraDataInterfaceProxy()
+	{
+#if NIAGARA_VALIDATE_NDIPROXY_REFS
+		checkf(ProxyTickRefs == 0, TEXT("NiagaraDataInterfaceProxy(%p) is still referenced but being destroyed."));
+#endif
+	}
 
 	virtual int32 PerInstanceDataPassedToRenderThreadSize() const = 0;
 	virtual void ConsumePerInstanceDataFromGameThread(void* PerInstanceData, const FNiagaraSystemInstanceID& Instance) { check(false); }
@@ -351,6 +360,10 @@ struct FNiagaraDataInterfaceProxy
 	virtual void FinalizePostStage(FRDGBuilder& GraphBuilder, const FNiagaraGpuComputeDispatchInterface& ComputeDispatchInterface) {}
 
 	virtual FNiagaraDataInterfaceProxyRW* AsIterationProxy() { return nullptr; }
+
+#if NIAGARA_VALIDATE_NDIPROXY_REFS
+	std::atomic<int32> ProxyTickRefs;
+#endif
 };
 
 //////////////////////////////////////////////////////////////////////////
diff --git a/Engine/Plugins/FX/Niagara/Source/Niagara/Classes/NiagaraGPUSystemTick.h b/Engine/Plugins/FX/Niagara/Source/Niagara/Classes/NiagaraGPUSystemTick.h
--- a/Engine/Plugins/FX/Niagara/Source/Niagara/Classes/NiagaraGPUSystemTick.h
+++ b/Engine/Plugins/FX/Niagara/Source/Niagara/Classes/NiagaraGPUSystemTick.h
@@ -9,6 +9,9 @@ struct FNiagaraComputeDataInterfaceInstanceData
 {
 	UE_NONCOPYABLE(FNiagaraComputeDataInterfaceInstanceData);
 	FNiagaraComputeDataInterfaceInstanceData() {}
+#if NIAGARA_VALIDATE_NDIPROXY_REFS
+	~FNiagaraComputeDataInterfaceInstanceData();
+#endif
 
 	void* PerInstanceDataForRT = nullptr;
 	TMap<FNiagaraDataInterfaceProxy*, int32> InterfaceProxiesToOffsets;
diff --git a/Engine/Plugins/FX/Niagara/Source/Niagara/Private/NiagaraGPUSystemTick.cpp b/Engine/Plugins/FX/Niagara/Source/Niagara/Private/NiagaraGPUSystemTick.cpp
--- a/Engine/Plugins/FX/Niagara/Source/Niagara/Private/NiagaraGPUSystemTick.cpp
+++ b/Engine/Plugins/FX/Niagara/Source/Niagara/Private/NiagaraGPUSystemTick.cpp
@@ -5,6 +5,16 @@
 #include "NiagaraSystemInstance.h"
 #include "NiagaraSystem.h"
 
+#if NIAGARA_VALIDATE_NDIPROXY_REFS
+FNiagaraComputeDataInterfaceInstanceData::~FNiagaraComputeDataInterfaceInstanceData()
+{
+	for (auto ProxyIt=InterfaceProxiesToOffsets.CreateIterator(); ProxyIt; ++ProxyIt)
+	{
+		--ProxyIt->Key->ProxyTickRefs;
+	}
+}
+#endif
+
 void FNiagaraGPUSystemTick::Init(FNiagaraSystemInstance* InSystemInstance)
 {
 	ensure(InSystemInstance != nullptr);
@@ -49,6 +59,9 @@ void FNiagaraGPUSystemTick::Init(FNiagaraSystemInstance* InSystemInstance)
 			Interface->ProvidePerInstanceDataForRenderThread(InstanceDataBase, PerInstanceData, SystemInstanceID);
 
 			// @todo rethink this. So ugly.
+		#if NIAGARA_VALIDATE_NDIPROXY_REFS
+			++Proxy->ProxyTickRefs;
+		#endif
 			DIInstanceData->InterfaceProxiesToOffsets.Add(Proxy, RunningOffset);
 
 			InstanceDataBase += RTDataSize;
