- Allow RT data interfaces to override user parameter RTs even if not used by a valid GPU emitter

diff --git a/Engine/Plugins/FX/Niagara/Source/Niagara/Private/NiagaraDataInterfaceRenderTarget2D.cpp b/Engine/Plugins/FX/Niagara/Source/Niagara/Private/NiagaraDataInterfaceRenderTarget2D.cpp
--- a/Engine/Plugins/FX/Niagara/Source/Niagara/Private/NiagaraDataInterfaceRenderTarget2D.cpp
+++ b/Engine/Plugins/FX/Niagara/Source/Niagara/Private/NiagaraDataInterfaceRenderTarget2D.cpp
@@ -450,15 +450,18 @@ bool UNiagaraDataInterfaceRenderTarget2D::InitPerInstanceData(void* PerInstanceD
 	FRenderTarget2DRWInstanceData_GameThread* InstanceData = new (PerInstanceData) FRenderTarget2DRWInstanceData_GameThread();
 	SystemInstancesToProxyData_GT.Emplace(SystemInstance->GetId(), InstanceData);
 
-	if (NiagaraDataInterfaceRenderTargetCommon::GIgnoreCookedOut && !IsUsedWithGPUEmitter())
-	{
-		return true;
-	}
+	//-TEMP: Until we prune data interface on cook this will avoid consuming memory
+	const bool bValidGpuDataInterface = NiagaraDataInterfaceRenderTargetCommon::GIgnoreCookedOut == 0 || IsUsedWithGPUEmitter();
 
 	ETextureRenderTargetFormat RenderTargetFormat;
 	if ( NiagaraDataInterfaceRenderTargetCommon::GetRenderTargetFormat(bOverrideFormat, OverrideRenderTargetFormat, RenderTargetFormat) == false )
 	{
-		return false;
+		if (bValidGpuDataInterface)
+		{
+			UE_LOG(LogNiagara, Warning, TEXT("NDIRT2D failed to find a render target format that supports UAV store"));
+			return false;
+		}
+		return true;
 	}
 
 	InstanceData->Size.X = FMath::Clamp<int>(int(float(Size.X) * NiagaraDataInterfaceRenderTargetCommon::GResolutionMultiplier), 1, GMaxTextureDimensions);
@@ -697,13 +700,10 @@ bool UNiagaraDataInterfaceRenderTarget2D::PerInstanceTickPostSimulate(void* PerI
 #endif
 
 	//-TEMP: Until we prune data interface on cook this will avoid consuming memory
-	if (NiagaraDataInterfaceRenderTargetCommon::GIgnoreCookedOut && !IsUsedWithGPUEmitter())
-	{
-		return false;
-	}
+	const bool bValidGpuDataInterface = NiagaraDataInterfaceRenderTargetCommon::GIgnoreCookedOut == 0 || IsUsedWithGPUEmitter();
 
 	// Do we need to create a new texture?
-	if (!bInheritUserParameterSettings && (InstanceData->TargetTexture == nullptr))
+	if (bValidGpuDataInterface && !bInheritUserParameterSettings && (InstanceData->TargetTexture == nullptr))
 	{
 		InstanceData->TargetTexture = NewObject<UTextureRenderTarget2D>(this);
 		InstanceData->TargetTexture->bCanCreateUAV = true;
@@ -739,7 +739,7 @@ bool UNiagaraDataInterfaceRenderTarget2D::PerInstanceTickPostSimulate(void* PerI
 	}
 
 	//-TODO: We could avoid updating each frame if we cache the resource pointer or a serial number
-	bool bUpdateRT = true;
+	const bool bUpdateRT = true && bValidGpuDataInterface;
 	if (bUpdateRT)
 	{
 		FNiagaraDataInterfaceProxyRenderTarget2DProxy* RT_Proxy = GetProxyAs<FNiagaraDataInterfaceProxyRenderTarget2DProxy>();
diff --git a/Engine/Plugins/FX/Niagara/Source/Niagara/Private/NiagaraDataInterfaceRenderTarget2DArray.cpp b/Engine/Plugins/FX/Niagara/Source/Niagara/Private/NiagaraDataInterfaceRenderTarget2DArray.cpp
--- a/Engine/Plugins/FX/Niagara/Source/Niagara/Private/NiagaraDataInterfaceRenderTarget2DArray.cpp
+++ b/Engine/Plugins/FX/Niagara/Source/Niagara/Private/NiagaraDataInterfaceRenderTarget2DArray.cpp
@@ -447,15 +447,18 @@ bool UNiagaraDataInterfaceRenderTarget2DArray::InitPerInstanceData(void* PerInst
 
 	FRenderTarget2DArrayRWInstanceData_GameThread* InstanceData = new (PerInstanceData) FRenderTarget2DArrayRWInstanceData_GameThread();
 
-	if (NiagaraDataInterfaceRenderTargetCommon::GIgnoreCookedOut && !IsUsedWithGPUEmitter())
-	{
-		return true;
-	}
+	//-TEMP: Until we prune data interface on cook this will avoid consuming memory
+	const bool bValidGpuDataInterface = NiagaraDataInterfaceRenderTargetCommon::GIgnoreCookedOut == 0 || IsUsedWithGPUEmitter();
 
 	ETextureRenderTargetFormat RenderTargetFormat;
 	if (NiagaraDataInterfaceRenderTargetCommon::GetRenderTargetFormat(bOverrideFormat, OverrideRenderTargetFormat, RenderTargetFormat) == false)
 	{
-		return false;
+		if (bValidGpuDataInterface)
+		{
+			UE_LOG(LogNiagara, Warning, TEXT("NDIRT2DArray failed to find a render target format that supports UAV store"));
+			return false;
+		}
+		return true;
 	}
 
 	InstanceData->Size.X = FMath::Clamp<int>(int(float(Size.X) * NiagaraDataInterfaceRenderTargetCommon::GResolutionMultiplier), 1, GMaxTextureDimensions);
@@ -661,12 +664,9 @@ bool UNiagaraDataInterfaceRenderTarget2DArray::PerInstanceTickPostSimulate(void*
 #endif
 
 	//-TEMP: Until we prune data interface on cook this will avoid consuming memory
-	if (NiagaraDataInterfaceRenderTargetCommon::GIgnoreCookedOut && !IsUsedWithGPUEmitter())
-	{
-		return false;
-	}
+	const bool bValidGpuDataInterface = NiagaraDataInterfaceRenderTargetCommon::GIgnoreCookedOut == 0 || IsUsedWithGPUEmitter();
 
-	if (!bInheritUserParameterSettings && (InstanceData->TargetTexture == nullptr))
+	if (bValidGpuDataInterface && !bInheritUserParameterSettings && (InstanceData->TargetTexture == nullptr))
 	{
 		InstanceData->TargetTexture = NewObject<UTextureRenderTarget2DArray>(this);
 		InstanceData->TargetTexture->bCanCreateUAV = true;
@@ -700,7 +700,7 @@ bool UNiagaraDataInterfaceRenderTarget2DArray::PerInstanceTickPostSimulate(void*
 	}
 
 	//-TODO: We could avoid updating each frame if we cache the resource pointer or a serial number
-	bool bUpdateRT = true;
+	const bool bUpdateRT = true && bValidGpuDataInterface;
 	if (bUpdateRT)
 	{
 		FNiagaraDataInterfaceProxyRenderTarget2DArrayProxy* RT_Proxy = GetProxyAs<FNiagaraDataInterfaceProxyRenderTarget2DArrayProxy>();
diff --git a/Engine/Plugins/FX/Niagara/Source/Niagara/Private/NiagaraDataInterfaceRenderTargetCommon.cpp b/Engine/Plugins/FX/Niagara/Source/Niagara/Private/NiagaraDataInterfaceRenderTargetCommon.cpp
--- a/Engine/Plugins/FX/Niagara/Source/Niagara/Private/NiagaraDataInterfaceRenderTargetCommon.cpp
+++ b/Engine/Plugins/FX/Niagara/Source/Niagara/Private/NiagaraDataInterfaceRenderTargetCommon.cpp
@@ -102,8 +102,6 @@ namespace NiagaraDataInterfaceRenderTargetCommon
 			}
 			if (PreviousFormat == OutRenderTargetFormat)
 			{
-				// This is fatal as we failed to find any format that supports typed UAV stores
-				UE_LOG(LogNiagara, Warning, TEXT("Failed to find a render target format that supports UAV store"));
 				return false;
 			}
 		}
diff --git a/Engine/Plugins/FX/Niagara/Source/Niagara/Private/NiagaraDataInterfaceRenderTargetCube.cpp b/Engine/Plugins/FX/Niagara/Source/Niagara/Private/NiagaraDataInterfaceRenderTargetCube.cpp
--- a/Engine/Plugins/FX/Niagara/Source/Niagara/Private/NiagaraDataInterfaceRenderTargetCube.cpp
+++ b/Engine/Plugins/FX/Niagara/Source/Niagara/Private/NiagaraDataInterfaceRenderTargetCube.cpp
@@ -412,15 +412,18 @@ bool UNiagaraDataInterfaceRenderTargetCube::InitPerInstanceData(void* PerInstanc
 
 	FRenderTargetCubeRWInstanceData_GameThread* InstanceData = new (PerInstanceData) FRenderTargetCubeRWInstanceData_GameThread();
 
-	if (NiagaraDataInterfaceRenderTargetCommon::GIgnoreCookedOut && !IsUsedWithGPUEmitter())
-	{
-		return true;
-	}
+	//-TEMP: Until we prune data interface on cook this will avoid consuming memory
+	const bool bValidGpuDataInterface = NiagaraDataInterfaceRenderTargetCommon::GIgnoreCookedOut == 0 || IsUsedWithGPUEmitter();
 
 	ETextureRenderTargetFormat RenderTargetFormat;
 	if (NiagaraDataInterfaceRenderTargetCommon::GetRenderTargetFormat(bOverrideFormat, OverrideRenderTargetFormat, RenderTargetFormat) == false)
 	{
-		return false;
+		if (bValidGpuDataInterface)
+		{
+			UE_LOG(LogNiagara, Warning, TEXT("NDIRTCube failed to find a render target format that supports UAV store"));
+			return false;
+		}
+		return true;
 	}
 
 	InstanceData->Size = FMath::Clamp<int>(int(float(Size) * NiagaraDataInterfaceRenderTargetCommon::GResolutionMultiplier), 1, GMaxCubeTextureDimensions);
@@ -611,13 +614,10 @@ bool UNiagaraDataInterfaceRenderTargetCube::PerInstanceTickPostSimulate(void* Pe
 #endif
 
 	//-TEMP: Until we prune data interface on cook this will avoid consuming memory
-	if (NiagaraDataInterfaceRenderTargetCommon::GIgnoreCookedOut && !IsUsedWithGPUEmitter())
-	{
-		return false;
-	}
+	const bool bValidGpuDataInterface = NiagaraDataInterfaceRenderTargetCommon::GIgnoreCookedOut == 0 || IsUsedWithGPUEmitter();
 
 	// Do we need to create a new texture?
-	if (!bInheritUserParameterSettings && (InstanceData->TargetTexture == nullptr))
+	if (bValidGpuDataInterface && !bInheritUserParameterSettings && (InstanceData->TargetTexture == nullptr))
 	{
 		InstanceData->TargetTexture = NewObject<UTextureRenderTargetCube>(this);
 		InstanceData->TargetTexture->bCanCreateUAV = true;
@@ -651,7 +651,7 @@ bool UNiagaraDataInterfaceRenderTargetCube::PerInstanceTickPostSimulate(void* Pe
 	}
 
 	//-TODO: We could avoid updating each frame if we cache the resource pointer or a serial number
-	bool bUpdateRT = true;
+	const bool bUpdateRT = true && bValidGpuDataInterface;
 	if (bUpdateRT)
 	{
 		FNiagaraDataInterfaceProxyRenderTargetCubeProxy* RT_Proxy = GetProxyAs<FNiagaraDataInterfaceProxyRenderTargetCubeProxy>();
diff --git a/Engine/Plugins/FX/Niagara/Source/Niagara/Private/NiagaraDataInterfaceRenderTargetVolume.cpp b/Engine/Plugins/FX/Niagara/Source/Niagara/Private/NiagaraDataInterfaceRenderTargetVolume.cpp
--- a/Engine/Plugins/FX/Niagara/Source/Niagara/Private/NiagaraDataInterfaceRenderTargetVolume.cpp
+++ b/Engine/Plugins/FX/Niagara/Source/Niagara/Private/NiagaraDataInterfaceRenderTargetVolume.cpp
@@ -483,15 +483,18 @@ bool UNiagaraDataInterfaceRenderTargetVolume::InitPerInstanceData(void* PerInsta
 
 	FRenderTargetVolumeRWInstanceData_GameThread* InstanceData = new (PerInstanceData) FRenderTargetVolumeRWInstanceData_GameThread();
 
-	if (NiagaraDataInterfaceRenderTargetCommon::GIgnoreCookedOut && !IsUsedWithGPUEmitter())
-	{
-		return true;
-	}
+	//-TEMP: Until we prune data interface on cook this will avoid consuming memory
+	const bool bValidGpuDataInterface = NiagaraDataInterfaceRenderTargetCommon::GIgnoreCookedOut == 0 || IsUsedWithGPUEmitter();
 
 	ETextureRenderTargetFormat RenderTargetFormat;
 	if (NiagaraDataInterfaceRenderTargetCommon::GetRenderTargetFormat(bOverrideFormat, OverrideRenderTargetFormat, RenderTargetFormat) == false)
 	{
-		return false;
+		if (bValidGpuDataInterface)
+		{
+			UE_LOG(LogNiagara, Warning, TEXT("NDIRTVolume failed to find a render target format that supports UAV store"));
+			return false;
+		}
+		return true;
 	}
 
 	InstanceData->Size.X = FMath::Clamp<int>(int(float(Size.X) * NiagaraDataInterfaceRenderTargetCommon::GResolutionMultiplier), 1, GMaxVolumeTextureDimensions);
@@ -980,13 +983,10 @@ bool UNiagaraDataInterfaceRenderTargetVolume::PerInstanceTickPostSimulate(void*
 #endif
 
 	//-TEMP: Until we prune data interface on cook this will avoid consuming memory
-	if (NiagaraDataInterfaceRenderTargetCommon::GIgnoreCookedOut && !IsUsedWithGPUEmitter())
-	{
-		return false;
-	}
+	const bool bValidGpuDataInterface = NiagaraDataInterfaceRenderTargetCommon::GIgnoreCookedOut == 0 || IsUsedWithGPUEmitter();
 
 	// Do we need to create a new texture?
-	if (!bInheritUserParameterSettings && (InstanceData->TargetTexture == nullptr))
+	if (bValidGpuDataInterface && !bInheritUserParameterSettings && (InstanceData->TargetTexture == nullptr))
 	{
 		InstanceData->TargetTexture = NewObject<UTextureRenderTargetVolume>(this);
 		InstanceData->TargetTexture->bCanCreateUAV = true;
@@ -1020,7 +1020,7 @@ bool UNiagaraDataInterfaceRenderTargetVolume::PerInstanceTickPostSimulate(void*
 	}
 
 	//-TODO: We could avoid updating each frame if we cache the resource pointer or a serial number
-	bool bUpdateRT = true;
+	const bool bUpdateRT = true && bValidGpuDataInterface;
 	if (bUpdateRT)
 	{
 		FNiagaraDataInterfaceProxyRenderTargetVolumeProxy* RT_Proxy = GetProxyAs<FNiagaraDataInterfaceProxyRenderTargetVolumeProxy>();
