Fixed SItemSelector causing a selection refresh when Refreshing when attempting to preserve selection.
This fixes an issue where if you are searching the parameters panel, click on the i.e. 3rd result (not the first), modify a value, a Refresh to the parameter panel would cause the first item to be selected.
Additionally, we never actually cleared the current selection before restoring preserved selection which could lead to selection bugs as well

diff --git a/Engine/Plugins/FX/Niagara/Source/NiagaraEditor/Private/Toolkits/SystemToolkitModes/NiagaraSystemToolkitModeBase.cpp b/Engine/Plugins/FX/Niagara/Source/NiagaraEditor/Private/Toolkits/SystemToolkitModes/NiagaraSystemToolkitModeBase.cpp
--- a/Engine/Plugins/FX/Niagara/Source/NiagaraEditor/Private/Toolkits/SystemToolkitModes/NiagaraSystemToolkitModeBase.cpp
+++ b/Engine/Plugins/FX/Niagara/Source/NiagaraEditor/Private/Toolkits/SystemToolkitModes/NiagaraSystemToolkitModeBase.cpp
@@ -155,8 +155,12 @@ void FNiagaraSystemToolkitModeBase::OnParameterPanelViewModelExternalSelectionCh
 			TArray<UObject*> SelectedNodes = ParamSel->GetSelectedObjects().Array();
 			if (SelectedNodes.Num() != 0)
 			{
-				ObjectSelection->SetSelectedObject(SelectedNodes[0]);
-				ObjectSelectionSubHeaderText = LOCTEXT("ParamPanelSel", "Parameter Panel Selection");
+				const TSet<UObject*>& CurrentlySelectedObjects = ObjectSelection->GetSelectedObjects();
+				if(CurrentlySelectedObjects.Array() != SelectedNodes)
+				{
+					ObjectSelection->SetSelectedObject(SelectedNodes[0]);
+					ObjectSelectionSubHeaderText = LOCTEXT("ParamPanelSel", "Parameter Panel Selection");
+				}
 			}
 		}
 	}
diff --git a/Engine/Plugins/FX/Niagara/Source/NiagaraEditor/Public/Widgets/SItemSelector.h b/Engine/Plugins/FX/Niagara/Source/NiagaraEditor/Public/Widgets/SItemSelector.h
--- a/Engine/Plugins/FX/Niagara/Source/NiagaraEditor/Public/Widgets/SItemSelector.h
+++ b/Engine/Plugins/FX/Niagara/Source/NiagaraEditor/Public/Widgets/SItemSelector.h
@@ -1337,7 +1337,11 @@ private:
 			DefaultCategoryPaths = InDefaultCategoryPaths;
 			GetRootItems();
 
-			UpdateSuggestedItem();
+			// as the update to the suggested item will cause a selection, we want to avoid this when merely refreshing
+			if(bPreserveSelectionOnRefresh == false)
+			{
+				UpdateSuggestedItem();
+			}
 
 			TArray<TSharedRef<FItemSelectorItemViewModel>> Children;
 			// Set expansion state if necessary.
@@ -1375,6 +1379,7 @@ private:
 					GetChildrenRecursive(Children);
 				}
 
+				ItemTreeRef->ClearSelection();
 				for (const TSharedRef<FItemSelectorItemViewModel>& ChildItemViewModel : Children)
 				{
 					if (ChildItemViewModel->GetType() == EItemSelectorItemViewModelType::Item)
