- Add uint variant for clear counts and use with array clear
- Fixes a validation issue due to buffer being a uint vs int

#rnx

diff --git a/Engine/Plugins/FX/Niagara/Shaders/Private/NiagaraClearCounts.usf b/Engine/Plugins/FX/Niagara/Shaders/Private/NiagaraClearCounts.usf
--- a/Engine/Plugins/FX/Niagara/Shaders/Private/NiagaraClearCounts.usf
+++ b/Engine/Plugins/FX/Niagara/Shaders/Private/NiagaraClearCounts.usf
@@ -8,7 +8,7 @@ RWBuffer<int>	CountBuffer;
 Buffer<int>		CountsAndValuesBuffer;
 uint			NumCounts;
 
-[numthreads(64, 1, 1)]
+[numthreads(THREADGROUP_SIZE, 1, 1)]
 void ClearCountsIntCS(uint3 DispatchThreadId : SV_DispatchThreadID)
 {
 	uint Index = DispatchThreadId.x;
@@ -21,3 +21,23 @@ void ClearCountsIntCS(uint3 DispatchThreadId : SV_DispatchThreadID)
 }
 
 #endif
+
+#if NIAGARA_CLEAR_COUNTS_UINT_CS
+
+RWBuffer<uint>	CountBuffer;
+Buffer<uint>	CountsAndValuesBuffer;
+uint			NumCounts;
+
+[numthreads(THREADGROUP_SIZE, 1, 1)]
+void ClearCountsUIntCS(uint3 DispatchThreadId : SV_DispatchThreadID)
+{
+	uint Index = DispatchThreadId.x;
+	if(Index < NumCounts)
+	{
+		uint ClearIndex = CountsAndValuesBuffer[Index * 2];
+		uint ClearValue = CountsAndValuesBuffer[(Index * 2) + 1];
+		CountBuffer[ClearIndex] = ClearValue;
+	}
+}
+
+#endif
diff --git a/Engine/Plugins/FX/Niagara/Source/Niagara/Classes/NiagaraDataInterfaceArrayImpl.h b/Engine/Plugins/FX/Niagara/Source/Niagara/Classes/NiagaraDataInterfaceArrayImpl.h
--- a/Engine/Plugins/FX/Niagara/Source/Niagara/Classes/NiagaraDataInterfaceArrayImpl.h
+++ b/Engine/Plugins/FX/Niagara/Source/Niagara/Classes/NiagaraDataInterfaceArrayImpl.h
@@ -290,9 +290,9 @@ struct FNDIArrayInstanceData_RenderThread
 			const FNiagaraGPUInstanceCountManager& CounterManager = ComputeInterface->GetGPUInstanceCounterManager();
 			const FRWBuffer& CountBuffer = CounterManager.GetInstanceCountBuffer();
 			
-			const TPair<uint32, int32> DataToClear(CountOffset, InArrayData.Num());
+			const TPair<uint32, uint32> DataToClear(CountOffset, InArrayData.Num());
 			RHICmdList.Transition(FRHITransitionInfo(CountBuffer.UAV, FNiagaraGPUInstanceCountManager::kCountBufferDefaultState, ERHIAccess::UAVCompute));
-			NiagaraClearCounts::ClearCountsInt(RHICmdList, CountBuffer.UAV, MakeArrayView(&DataToClear, 1));
+			NiagaraClearCounts::ClearCountsUInt(RHICmdList, CountBuffer.UAV, MakeArrayView(&DataToClear, 1));
 			RHICmdList.Transition(FRHITransitionInfo(CountBuffer.UAV, ERHIAccess::UAVCompute, FNiagaraGPUInstanceCountManager::kCountBufferDefaultState));
 		}
 	}
diff --git a/Engine/Plugins/FX/Niagara/Source/NiagaraShader/Private/NiagaraClearCounts.cpp b/Engine/Plugins/FX/Niagara/Source/NiagaraShader/Private/NiagaraClearCounts.cpp
--- a/Engine/Plugins/FX/Niagara/Source/NiagaraShader/Private/NiagaraClearCounts.cpp
+++ b/Engine/Plugins/FX/Niagara/Source/NiagaraShader/Private/NiagaraClearCounts.cpp
@@ -31,7 +31,29 @@ class FNiagaraClearCountsIntCS : public FGlobalShader
 	END_SHADER_PARAMETER_STRUCT()
 };
 
+class FNiagaraClearCountsUIntCS : public FGlobalShader
+{
+	DECLARE_GLOBAL_SHADER(FNiagaraClearCountsUIntCS);
+	SHADER_USE_PARAMETER_STRUCT(FNiagaraClearCountsUIntCS, FGlobalShader);
+
+	static constexpr uint32 ThreadGroupSize = 64;
+
+	static void ModifyCompilationEnvironment(const FGlobalShaderPermutationParameters& Parameters, FShaderCompilerEnvironment& OutEnvironment)
+	{
+		FGlobalShader::ModifyCompilationEnvironment(Parameters, OutEnvironment);
+		OutEnvironment.SetDefine(TEXT("NIAGARA_CLEAR_COUNTS_UINT_CS"), 1);
+		OutEnvironment.SetDefine(TEXT("THREADGROUP_SIZE"), ThreadGroupSize);
+	}
+
+	BEGIN_SHADER_PARAMETER_STRUCT(FParameters, )
+		SHADER_PARAMETER_RDG_BUFFER_UAV(RWBuffer<uint>,	CountBuffer)
+		SHADER_PARAMETER_RDG_BUFFER_SRV(Buffer<uint>,	CountsAndValuesBuffer)
+		SHADER_PARAMETER(uint32,						NumCounts)
+	END_SHADER_PARAMETER_STRUCT()
+};
+
 IMPLEMENT_GLOBAL_SHADER(FNiagaraClearCountsIntCS, "/Plugin/FX/Niagara/Private/NiagaraClearCounts.usf", "ClearCountsIntCS", SF_Compute);
+IMPLEMENT_GLOBAL_SHADER(FNiagaraClearCountsUIntCS, "/Plugin/FX/Niagara/Private/NiagaraClearCounts.usf", "ClearCountsUIntCS", SF_Compute);
 
 void NiagaraClearCounts::ClearCountsInt(FRDGBuilder& GraphBuilder, FRDGBufferUAVRef UAV, TConstArrayView<TPair<uint32, int32>> IndexAndValueArray)
 {
@@ -57,6 +79,30 @@ void NiagaraClearCounts::ClearCountsInt(FRDGBuilder& GraphBuilder, FRDGBufferUAV
 	);
 }
 
+void NiagaraClearCounts::ClearCountsUInt(FRDGBuilder& GraphBuilder, FRDGBufferUAVRef UAV, TConstArrayView<TPair<uint32, uint32>> IndexAndValueArray)
+{
+	const int32 NumInts = IndexAndValueArray.Num() * 2;
+	FRDGBufferRef CountsAndValuesBuffer = GraphBuilder.CreateBuffer(FRDGBufferDesc::CreateBufferDesc(sizeof(int32), NumInts), TEXT("NiagaraClearCountsInt::CountsAndValuesBuffer"));
+	GraphBuilder.QueueBufferUpload(CountsAndValuesBuffer, IndexAndValueArray.GetData(), sizeof(int32) * NumInts);
+
+	TShaderMapRef<FNiagaraClearCountsUIntCS> ComputeShader(GetGlobalShaderMap(GMaxRHIFeatureLevel));
+	const uint32 NumThreadGroups = FMath::DivideAndRoundUp<uint32>(IndexAndValueArray.Num(), FNiagaraClearCountsUIntCS::ThreadGroupSize);
+
+	FNiagaraClearCountsUIntCS::FParameters* Parameters = GraphBuilder.AllocParameters<FNiagaraClearCountsUIntCS::FParameters>();
+	Parameters->CountBuffer = UAV;
+	Parameters->CountsAndValuesBuffer = GraphBuilder.CreateSRV(CountsAndValuesBuffer, PF_R32_UINT);
+	Parameters->NumCounts = IndexAndValueArray.Num();
+
+	FComputeShaderUtils::AddPass(
+		GraphBuilder,
+		RDG_EVENT_NAME("NiagaraClearCountsUInt"),
+		ERDGPassFlags::Compute,
+		ComputeShader,
+		Parameters,
+		FIntVector(NumThreadGroups, 1, 1)
+	);
+}
+
 //////////////////////////////////////////////////////////////////////////
 
 class FNiagaraClearCountsIntCS_Legacy : public FGlobalShader
@@ -75,12 +121,34 @@ class FNiagaraClearCountsIntCS_Legacy : public FGlobalShader
 
 	BEGIN_SHADER_PARAMETER_STRUCT(FParameters, )
 		SHADER_PARAMETER_UAV(RWBuffer<int>, CountBuffer)
-		SHADER_PARAMETER_SRV(Buffer<int>, CountsAndValuesBuffer)
-		SHADER_PARAMETER(uint32, NumCounts)
+		SHADER_PARAMETER_SRV(Buffer<int>,	CountsAndValuesBuffer)
+		SHADER_PARAMETER(uint32,			NumCounts)
+	END_SHADER_PARAMETER_STRUCT()
+};
+
+class FNiagaraClearCountsUIntCS_Legacy : public FGlobalShader
+{
+	DECLARE_GLOBAL_SHADER(FNiagaraClearCountsUIntCS_Legacy);
+	SHADER_USE_PARAMETER_STRUCT(FNiagaraClearCountsUIntCS_Legacy, FGlobalShader);
+
+	static constexpr uint32 ThreadGroupSize = 64;
+
+	static void ModifyCompilationEnvironment(const FGlobalShaderPermutationParameters& Parameters, FShaderCompilerEnvironment& OutEnvironment)
+	{
+		FGlobalShader::ModifyCompilationEnvironment(Parameters, OutEnvironment);
+		OutEnvironment.SetDefine(TEXT("NIAGARA_CLEAR_COUNTS_UINT_CS"), 1);
+		OutEnvironment.SetDefine(TEXT("THREADGROUP_SIZE"), ThreadGroupSize);
+	}
+
+	BEGIN_SHADER_PARAMETER_STRUCT(FParameters, )
+		SHADER_PARAMETER_UAV(RWBuffer<uint>,	CountBuffer)
+		SHADER_PARAMETER_SRV(Buffer<uint>,		CountsAndValuesBuffer)
+		SHADER_PARAMETER(uint32,				NumCounts)
 	END_SHADER_PARAMETER_STRUCT()
 };
 
 IMPLEMENT_GLOBAL_SHADER(FNiagaraClearCountsIntCS_Legacy, "/Plugin/FX/Niagara/Private/NiagaraClearCounts.usf", "ClearCountsIntCS", SF_Compute);
+IMPLEMENT_GLOBAL_SHADER(FNiagaraClearCountsUIntCS_Legacy, "/Plugin/FX/Niagara/Private/NiagaraClearCounts.usf", "ClearCountsUIntCS", SF_Compute);
 
 void NiagaraClearCounts::ClearCountsInt(FRHICommandList& RHICmdList, FRHIUnorderedAccessView* UAV, TConstArrayView<TPair<uint32, int32>> IndexAndValueArray)
 {
@@ -112,3 +180,34 @@ void NiagaraClearCounts::ClearCountsInt(FRHICommandList& RHICmdList, FRHIUnorder
 	UnsetShaderUAVs(RHICmdList, ComputeShader, ShaderRHI);
 	RHICmdList.Transition(FRHITransitionInfo(UAV, ERHIAccess::UAVCompute, ERHIAccess::UAVCompute));
 }
+
+void NiagaraClearCounts::ClearCountsUInt(FRHICommandList& RHICmdList, FRHIUnorderedAccessView* UAV, TConstArrayView<TPair<uint32, uint32>> IndexAndValueArray)
+{
+	check(UAV != nullptr);
+	check(IndexAndValueArray.Num() > 0);
+
+	FReadBuffer IndexAndValueBuffer;
+	IndexAndValueBuffer.Initialize(TEXT("NiagaraClearCounts"), sizeof(int32), IndexAndValueArray.Num() * 2, PF_R32_UINT);
+	{
+		const uint32 BufferSize = IndexAndValueArray.Num() * sizeof(int32) * 2;
+		void* GPUMemory = RHILockBuffer(IndexAndValueBuffer.Buffer, 0, BufferSize, RLM_WriteOnly);
+		FMemory::Memcpy(GPUMemory, IndexAndValueArray.GetData(), BufferSize);
+		RHIUnlockBuffer(IndexAndValueBuffer.Buffer);
+	}
+
+	FNiagaraClearCountsUIntCS_Legacy::FParameters ShaderParameters;
+	ShaderParameters.CountBuffer = UAV;
+	ShaderParameters.CountsAndValuesBuffer = IndexAndValueBuffer.SRV;
+	ShaderParameters.NumCounts = IndexAndValueArray.Num();
+
+	TShaderMapRef<FNiagaraClearCountsUIntCS_Legacy> ComputeShader(GetGlobalShaderMap(GMaxRHIFeatureLevel));
+	FRHIComputeShader* ShaderRHI = ComputeShader.GetComputeShader();
+	const uint32 NumThreadGroups = FMath::DivideAndRoundUp<uint32>(IndexAndValueArray.Num(), FNiagaraClearCountsUIntCS_Legacy::ThreadGroupSize);
+
+	RHICmdList.Transition(FRHITransitionInfo(UAV, ERHIAccess::Unknown, ERHIAccess::UAVCompute));
+	SetComputePipelineState(RHICmdList, ShaderRHI);
+	SetShaderParameters(RHICmdList, ComputeShader, ShaderRHI, ShaderParameters);
+	RHICmdList.DispatchComputeShader(NumThreadGroups, 1, 1);
+	UnsetShaderUAVs(RHICmdList, ComputeShader, ShaderRHI);
+	RHICmdList.Transition(FRHITransitionInfo(UAV, ERHIAccess::UAVCompute, ERHIAccess::UAVCompute));
+}
diff --git a/Engine/Plugins/FX/Niagara/Source/NiagaraShader/Public/NiagaraClearCounts.h b/Engine/Plugins/FX/Niagara/Source/NiagaraShader/Public/NiagaraClearCounts.h
--- a/Engine/Plugins/FX/Niagara/Source/NiagaraShader/Public/NiagaraClearCounts.h
+++ b/Engine/Plugins/FX/Niagara/Source/NiagaraShader/Public/NiagaraClearCounts.h
@@ -12,9 +12,11 @@ class FRHIUnorderedAccessView;
 namespace NiagaraClearCounts
 {
 	NIAGARASHADER_API void ClearCountsInt(FRDGBuilder& GraphBuilder, FRDGBufferUAVRef UAV, TConstArrayView<TPair<uint32, int32>> IndexAndValueArray);
+	NIAGARASHADER_API void ClearCountsUInt(FRDGBuilder& GraphBuilder, FRDGBufferUAVRef UAV, TConstArrayView<TPair<uint32, uint32>> IndexAndValueArray);
 
 	//-TODO:RDG: Deprecated
 	NIAGARASHADER_API void ClearCountsInt(FRHICommandList& RHICmdList, FRHIUnorderedAccessView* UAV, TConstArrayView<TPair<uint32, int32>> IndexAndValueArray);
+	NIAGARASHADER_API void ClearCountsUInt(FRHICommandList& RHICmdList, FRHIUnorderedAccessView* UAV, TConstArrayView<TPair<uint32, uint32>> IndexAndValueArray);
 }
 
 #if UE_ENABLE_INCLUDE_ORDER_DEPRECATED_IN_5_3
