- Allow fx.NiagaraAllowGPUParticles to be set by scalability to enable preview mode to reflect the changes

diff --git a/Engine/Plugins/FX/Niagara/Source/Niagara/Private/NiagaraCommon.cpp b/Engine/Plugins/FX/Niagara/Source/Niagara/Private/NiagaraCommon.cpp
--- a/Engine/Plugins/FX/Niagara/Source/Niagara/Private/NiagaraCommon.cpp
+++ b/Engine/Plugins/FX/Niagara/Source/Niagara/Private/NiagaraCommon.cpp
@@ -38,7 +38,7 @@ FAutoConsoleVariableRef CVarAllowGPUParticles(
 	TEXT("fx.NiagaraAllowGPUParticles"),
 	GNiagaraAllowGPUParticles,
 	TEXT("If true, allow the usage of GPU particles for Niagara."),
-	ECVF_Default);
+	ECVF_Scalability | ECVF_Default);
 
 int32 GNiagaraGPUCulling = 1;
 FAutoConsoleVariableRef CVarNiagaraGPUCulling(
