- Add temporary code to help track down an infrequent issue where the emitter destination data becomes nullptr during emitter execution

#rnx

diff --git a/Engine/Plugins/FX/Niagara/Source/Niagara/Private/NiagaraEmitterInstance.cpp b/Engine/Plugins/FX/Niagara/Source/Niagara/Private/NiagaraEmitterInstance.cpp
--- a/Engine/Plugins/FX/Niagara/Source/Niagara/Private/NiagaraEmitterInstance.cpp
+++ b/Engine/Plugins/FX/Niagara/Source/Niagara/Private/NiagaraEmitterInstance.cpp
@@ -1690,6 +1690,23 @@ void FNiagaraEmitterInstance::Tick(float DeltaSeconds)
 		}
 	}
 	
+	// Temporary code to track down an intermitent issue where we had particles but now have none and destination has become invalid.
+	if (Data.GetDestinationData() == nullptr)
+	{
+		UE_LOG(LogNiagara, Error, TEXT("NiagaraEmitterInstance: DestinationData is nullptr."));
+		UE_LOG(LogNiagara, Error, TEXT("ExecutionState(%d) TickCount(%d) TotalSpawnedParticles(%d) ParentSystemInstance(%p)"), ExecutionState, TickCount, TotalSpawnedParticles, ParentSystemInstance);
+		if (ParentSystemInstance)
+		{
+			USceneComponent* AttachComponent = ParentSystemInstance->GetAttachComponent();
+			UNiagaraSystem* NiagaraSystem = ParentSystemInstance->GetSystem();
+			UE_LOG(LogNiagara, Error, TEXT("SystemInstanceState(%d) AttachComponent(%p=%s) NiagaraSystem(%p=%s)"),
+				ParentSystemInstance->SystemInstanceState,
+				AttachComponent, *GetNameSafe(AttachComponent),
+				NiagaraSystem, *GetNameSafe(NiagaraSystem)
+			);
+		}
+	}
+
 	uint32 EventSpawnStart = Data.GetDestinationDataChecked().GetNumInstances();
 	const int32 NumBeforeSpawn = Data.GetDestinationDataChecked().GetNumInstances();
 	uint32 TotalActualEventSpawns = 0;
