- Temporarily disable ref checking while I investigate issues

#rnx

diff --git a/Engine/Plugins/FX/Niagara/Source/Niagara/Classes/NiagaraDataInterface.h b/Engine/Plugins/FX/Niagara/Source/Niagara/Classes/NiagaraDataInterface.h
--- a/Engine/Plugins/FX/Niagara/Source/Niagara/Classes/NiagaraDataInterface.h
+++ b/Engine/Plugins/FX/Niagara/Source/Niagara/Classes/NiagaraDataInterface.h
@@ -52,7 +52,7 @@ struct FNDITransformHandler
 
 // FNiagaraDataInterfaceProxy should always outlive any ticks, etc, that are on the rendering thread.
 // Enabling this will incur cost but will validate that this contract is not broken
-#define NIAGARA_VALIDATE_NDIPROXY_REFS	!UE_BUILD_SHIPPING
+#define NIAGARA_VALIDATE_NDIPROXY_REFS	0//!UE_BUILD_SHIPPING
 
 //////////////////////////////////////////////////////////////////////////
 // Some helper classes allowing neat, init time binding of templated vm external functions.
