Fixed SA issue

diff --git a/Engine/Plugins/FX/Niagara/Source/NiagaraEditor/Private/ViewModels/NiagaraOverviewGraphViewModel.cpp b/Engine/Plugins/FX/Niagara/Source/NiagaraEditor/Private/ViewModels/NiagaraOverviewGraphViewModel.cpp
--- a/Engine/Plugins/FX/Niagara/Source/NiagaraEditor/Private/ViewModels/NiagaraOverviewGraphViewModel.cpp
+++ b/Engine/Plugins/FX/Niagara/Source/NiagaraEditor/Private/ViewModels/NiagaraOverviewGraphViewModel.cpp
@@ -466,7 +466,10 @@ void FNiagaraOverviewGraphViewModel::GraphSelectionChanged()
 		if(UNiagaraStackRoot* AsRoot = Cast<UNiagaraStackRoot>(RootEntry))
 		{
 			CommentCollection = AsRoot->GetCommentCollection();
-			CommentCollection->RefreshChildren();
+			if(CommentCollection != nullptr)
+			{
+				CommentCollection->RefreshChildren();
+			}
 		}
 		
 		TArray<UNiagaraStackEntry*> EntriesToSelect;
@@ -485,6 +488,11 @@ void FNiagaraOverviewGraphViewModel::GraphSelectionChanged()
 			// this should always exist at this point as we have synced above
 			else if(UEdGraphNode_Comment* CommentNode = Cast<UEdGraphNode_Comment>(SelectedObject))
 			{
+				if(CommentCollection == nullptr)
+				{
+					continue;
+				}
+				
 				if (UNiagaraStackObject* CommentStackObject = CommentCollection->FindStackObjectForCommentNode(CommentNode))
 				{
 					EntriesToSelect.Add(CommentStackObject);
