- Add a targeted GPU deny list that allows you to match against OS + OSVersion + CPU + GPU

#rnx

diff --git a/Engine/Plugins/FX/Niagara/Source/Niagara/Private/NiagaraComponentSettings.cpp b/Engine/Plugins/FX/Niagara/Source/Niagara/Private/NiagaraComponentSettings.cpp
--- a/Engine/Plugins/FX/Niagara/Source/Niagara/Private/NiagaraComponentSettings.cpp
+++ b/Engine/Plugins/FX/Niagara/Source/Niagara/Private/NiagaraComponentSettings.cpp
@@ -52,6 +52,7 @@ namespace FNiagaraComponentSettings
 	FString									GpuDataInterfaceDenyListString;
 	FString									GpuRHIDenyListString;
 	FString									GpuRHIAdapterDenyListString;
+	FString									GpuDenyListString;
 
 	static bool ParseIntoSet(const FString& StringList, TSet<FName>& OutSet)
 	{
@@ -150,14 +151,23 @@ namespace FNiagaraComponentSettings
 	static FAutoConsoleVariableRef CVarNiagaraGpuRHIDenyList(
 		TEXT("fx.Niagara.SetGpuRHIDenyList"),
 		GpuRHIDenyListString,
-		TEXT("Set Gpu RHI deny list to use, comma separated and uses wildcards, i.e. (*MyRHI*) would exclude anything that contais MyRHI"),
+		TEXT("Set Gpu RHI deny list to use, comma separated and uses wildcards, i.e. (*MyRHI*) would exclude anything that contains MyRHI"),
 		ECVF_Scalability | ECVF_Default
 	);
 
 	static FAutoConsoleVariableRef CVarNiagaraGpuRHIAdapterDenyList(
 		TEXT("fx.Niagara.SetGpuRHIAdapterDenyList"),
 		GpuRHIAdapterDenyListString,
-		TEXT("Set Gpu RHI Adapter deny list to use, comma separated and uses wildcards, i.e. (*MyGpu*) would exclude anything that contais MyGpu"),
+		TEXT("Set Gpu RHI Adapter deny list to use, comma separated and uses wildcards, i.e. (*MyGpu*) would exclude anything that contains MyGpu"),
+		ECVF_Scalability | ECVF_Default
+	);
+
+	static FAutoConsoleVariableRef CVarNiagaraGpuDenyList(
+		TEXT("fx.Niagara.SetGpuDenyList"),
+		GpuDenyListString,
+		TEXT("Set Gpu deny list to use, more targetted than to allow comparing OS,OSVersion,CPU,GPU.\n")
+		TEXT("Format is OSLabel,OSVersion,CPU,GPU| blank entries are assumed to auto pass matching.\n")
+		TEXT("For example, =\",,MyCpu,MyGpu+MyOS,,,\" would match MyCpu & MyGpu or MyOS."),
 		ECVF_Scalability | ECVF_Default
 	);
 
@@ -181,7 +191,7 @@ namespace FNiagaraComponentSettings
 			}
 		}
 
-		if (GpuRHIAdapterDenyListString.Len() > 0)
+		if (bShouldAllowGpuEmitters && GpuRHIAdapterDenyListString.Len() > 0)
 		{
 			TArray<FString> BanNames;
 			GpuRHIAdapterDenyListString.ParseIntoArray(BanNames, TEXT(","));
@@ -192,6 +202,47 @@ namespace FNiagaraComponentSettings
 			}
 		}
 
+		if (bShouldAllowGpuEmitters && GpuDenyListString.Len() > 0)
+		{
+			UE_LOG(LogNiagara, Warning, TEXT("GpuDenyListString = %s"), *GpuDenyListString);
+
+			TArray<FString> BanList;
+			GpuDenyListString.ParseIntoArray(BanList, TEXT("+"));
+
+			FString CategoryData[4];
+			FPlatformMisc::GetOSVersions(CategoryData[0], CategoryData[1]);
+			CategoryData[2] = FPlatformMisc::GetCPUBrand();
+			CategoryData[3] = FPlatformMisc::GetPrimaryGPUBrand();
+
+			TArray<FString> BanCategoryStrings;
+			for ( const FString& Ban : BanList )
+			{
+				Ban.ParseIntoArray(BanCategoryStrings, TEXT(","), false);
+				if (BanCategoryStrings.Num() == UE_ARRAY_COUNT(CategoryData))
+				{
+					bool bMatches = true;
+					for (int32 i=0; i < UE_ARRAY_COUNT(CategoryData); ++i)
+					{
+						bMatches = BanCategoryStrings[i].IsEmpty() || CategoryData[i].MatchesWildcard(BanCategoryStrings[i]);
+						if (!BanCategoryStrings[i].IsEmpty())
+						{
+							UE_LOG(LogNiagara, Warning, TEXT("Matching %s=%s == %d"), *CategoryData[i], *BanCategoryStrings[i], bMatches);
+						}
+
+						if (bMatches == false)
+						{
+							break;
+						}
+					}
+					if (bMatches)
+					{
+						bShouldAllowGpuEmitters = false;
+						break;
+					}
+				}
+			}
+		}
+
 		if (bShouldAllowGpuEmitters != bAllowGpuEmitters)
 		{
 			bAllowGpuEmitters = bShouldAllowGpuEmitters;
