Fixed loss of parameter references when updating an input parameter's name in a module.
The 'BoundPinNames' array was getting updated with an incorrect parameter handle, which made the lookup for subsequent renames break

diff --git a/Engine/Plugins/FX/Niagara/Source/NiagaraEditor/Private/ViewModels/Stack/NiagaraStackGraphUtilities.cpp b/Engine/Plugins/FX/Niagara/Source/NiagaraEditor/Private/ViewModels/Stack/NiagaraStackGraphUtilities.cpp
--- a/Engine/Plugins/FX/Niagara/Source/NiagaraEditor/Private/ViewModels/Stack/NiagaraStackGraphUtilities.cpp
+++ b/Engine/Plugins/FX/Niagara/Source/NiagaraEditor/Private/ViewModels/Stack/NiagaraStackGraphUtilities.cpp
@@ -3724,7 +3724,7 @@ void FNiagaraStackGraphUtilities::SynchronizeReferencingMapPinsWithFunctionCall(
 		return;
 	}
 
-	TSharedPtr<TMap<FGuid, FNiagaraVariable>> VariableIdToVariableMap;
+	TSharedPtr<TMap<FGuid, FNiagaraVariable>> AssetVariableIdToVariableMap;
 	const UEdGraphSchema_Niagara* NiagaraSchema = GetDefault<UEdGraphSchema_Niagara>();
 	
 	// Update the override parameter map set nodes which are used to set the function inputs.
@@ -3736,29 +3736,37 @@ void FNiagaraStackGraphUtilities::SynchronizeReferencingMapPinsWithFunctionCall(
 	}
 	if (OverridePins.Num() > 0)
 	{
-		VariableIdToVariableMap = GetVariableIdToVariableMap(InFunctionCallNode);
+		// this contains the current up to date metadata of the asset graph this function call points to
+		AssetVariableIdToVariableMap = GetVariableIdToVariableMap(InFunctionCallNode);
 		for (UEdGraphPin* OverridePin : OverridePins)
 		{
-			FNiagaraVariable OverrideVariable = NiagaraSchema->PinToNiagaraVariable(OverridePin);
-			FNiagaraParameterHandle OverrideHandle(OverrideVariable.GetName());
-			TArray<FGuid> BoundGuids = InFunctionCallNode.GetBoundPinGuidsByName(OverrideHandle.GetParameterHandleString());
+			FNiagaraVariable CurrentPinOverrideVariable = NiagaraSchema->PinToNiagaraVariable(OverridePin);
+			FNiagaraParameterHandle AliasedOverridePinHandle(CurrentPinOverrideVariable.GetName());
+			// we look up the guids associated with the override pin whose name may be out of date
+			TArray<FGuid> BoundGuids = InFunctionCallNode.GetBoundPinGuidsByName(AliasedOverridePinHandle.GetParameterHandleString());
 			for (FGuid BoundGuid : BoundGuids)
 			{
-				FNiagaraVariable* BoundVariable = VariableIdToVariableMap->Find(BoundGuid);
-				if (BoundVariable != nullptr)
+				// since the guid stayed the same, we can now look up the current parameter that might have a different name
+				FNiagaraVariable* CurrentAssetParameter = AssetVariableIdToVariableMap->Find(BoundGuid);
+				if (CurrentAssetParameter != nullptr)
 				{
-					FNiagaraParameterHandle BoundVariableHandle(BoundVariable->GetName());
-					if (BoundVariableHandle.IsModuleHandle())
+					FNiagaraParameterHandle CurrentAssetParameterHandle(CurrentAssetParameter->GetName());
+					if (CurrentAssetParameterHandle.IsModuleHandle())
 					{
-						if (OverrideHandle.GetName() != BoundVariableHandle.GetName())
+						FNiagaraVariable UnaliasedCurrentPinOverrideVariable = FNiagaraUtilities::ResolveAliases(CurrentPinOverrideVariable, FNiagaraAliasContext().ChangeModuleNameToModule(AliasedOverridePinHandle.GetName().ToString()));
+						FNiagaraParameterHandle UnaliasedCurrentPinOverrideVariableHandle = FNiagaraParameterHandle(UnaliasedCurrentPinOverrideVariable.GetName());
+						
+						// update the override pin if the parameter has changed
+						if (UnaliasedCurrentPinOverrideVariableHandle.GetName() != CurrentAssetParameterHandle.GetName())
 						{
-							FNiagaraParameterHandle UpdatedOverrideHandle = FNiagaraParameterHandle::CreateAliasedModuleParameterHandle(BoundVariable->GetName(), &InFunctionCallNode);
+							FNiagaraVariable UpdatedAliasedOverride = FNiagaraUtilities::ResolveAliases(*CurrentAssetParameter, FNiagaraAliasContext().ChangeModuleToModuleName(InFunctionCallNode.GetFunctionName()));
+							FNiagaraParameterHandle UpdatedOverrideHandle = FNiagaraParameterHandle(UpdatedAliasedOverride.GetName());
 							OverridePin->PinName = UpdatedOverrideHandle.GetParameterHandleString();
-							InFunctionCallNode.UpdateInputNameBinding(BoundGuid, BoundVariableHandle.GetParameterHandleString());
+							InFunctionCallNode.UpdateInputNameBinding(BoundGuid, UpdatedOverrideHandle.GetParameterHandleString());
 						}
-						if (OverrideVariable.GetType() != BoundVariable->GetType())
+						if (CurrentPinOverrideVariable.GetType() != CurrentAssetParameter->GetType())
 						{
-							FEdGraphPinType NewPinType = NiagaraSchema->TypeDefinitionToPinType(BoundVariable->GetType());
+							FEdGraphPinType NewPinType = NiagaraSchema->TypeDefinitionToPinType(CurrentAssetParameter->GetType());
 							OverridePin->PinType = NewPinType;
 							if (OverridePin->LinkedTo.Num() == 1 &&
 								OverridePin->LinkedTo[0]->GetOwningNode() != nullptr &&
@@ -3777,11 +3785,11 @@ void FNiagaraStackGraphUtilities::SynchronizeReferencingMapPinsWithFunctionCall(
 	}
 
 	// Update linked module outputs or module attributes referenced by downstream parameter map gets.
-	auto SynchronizeForStackLinkedModuleParameterPins = [&NiagaraSchema, &InFunctionCallNode, &VariableIdToVariableMap](const TArray<UEdGraphPin*>& StackLinkedModuleParameterPins)
+	auto SynchronizeForStackLinkedModuleParameterPins = [&NiagaraSchema, &InFunctionCallNode, &AssetVariableIdToVariableMap](const TArray<UEdGraphPin*>& StackLinkedModuleParameterPins)
 	{
-		if (VariableIdToVariableMap.IsValid() == false)
+		if (AssetVariableIdToVariableMap.IsValid() == false)
 		{
-			VariableIdToVariableMap = GetVariableIdToVariableMap(InFunctionCallNode);
+			AssetVariableIdToVariableMap = GetVariableIdToVariableMap(InFunctionCallNode);
 		}
 		for (UEdGraphPin* StackLinkedModuleParameterPin : StackLinkedModuleParameterPins)
 		{
@@ -3793,7 +3801,7 @@ void FNiagaraStackGraphUtilities::SynchronizeReferencingMapPinsWithFunctionCall(
 				TArray<FGuid> BoundGuids = TargetFunctionCallNode->GetBoundPinGuidsByName(LinkedModuleParameterHandle.GetParameterHandleString());
 				for (const FGuid& BoundGuid : BoundGuids)
 				{
-					FNiagaraVariable* BoundVariable = VariableIdToVariableMap->Find(BoundGuid);
+					FNiagaraVariable* BoundVariable = AssetVariableIdToVariableMap->Find(BoundGuid);
 					if (BoundVariable != nullptr)
 					{
 						FNiagaraVariable BoundVariableResolvedName = FNiagaraUtilities::ResolveAliases(*BoundVariable,FNiagaraAliasContext().ChangeModuleToModuleName(InFunctionCallNode.GetFunctionName()));
