Changed 'Selection' or 'Selected Emitters' in niagara system editor to 'Details', including the icon, to be more consistent with the engine.

diff --git a/Engine/Plugins/FX/Niagara/Source/NiagaraEditor/Private/Toolkits/SystemToolkitModes/NiagaraSystemToolkitModeBase.cpp b/Engine/Plugins/FX/Niagara/Source/NiagaraEditor/Private/Toolkits/SystemToolkitModes/NiagaraSystemToolkitModeBase.cpp
--- a/Engine/Plugins/FX/Niagara/Source/NiagaraEditor/Private/Toolkits/SystemToolkitModes/NiagaraSystemToolkitModeBase.cpp
+++ b/Engine/Plugins/FX/Niagara/Source/NiagaraEditor/Private/Toolkits/SystemToolkitModes/NiagaraSystemToolkitModeBase.cpp
@@ -48,10 +48,9 @@ const FName FNiagaraSystemToolkitModeBase::ViewportTabID(TEXT("NiagaraSystemEdit
 const FName FNiagaraSystemToolkitModeBase::CurveEditorTabID(TEXT("NiagaraSystemEditor_CurveEditor"));
 const FName FNiagaraSystemToolkitModeBase::SequencerTabID(TEXT("NiagaraSystemEditor_Sequencer"));
 const FName FNiagaraSystemToolkitModeBase::SystemScriptTabID(TEXT("NiagaraSystemEditor_SystemScript"));
-const FName FNiagaraSystemToolkitModeBase::SystemDetailsTabID(TEXT("NiagaraSystemEditor_SystemDetails"));
 const FName FNiagaraSystemToolkitModeBase::SystemParametersTabID(TEXT("NiagaraSystemEditor_SystemParameters"));
 const FName FNiagaraSystemToolkitModeBase::SystemParameterDefinitionsTabID(TEXT("NiagaraSystemEditor_SystemParameterDefinitions"));
-const FName FNiagaraSystemToolkitModeBase::SelectedEmitterStackTabID(TEXT("NiagaraSystemEditor_SelectedEmitterStack"));
+const FName FNiagaraSystemToolkitModeBase::DetailsTabID(TEXT("NiagaraSystemEditor_Details"));
 const FName FNiagaraSystemToolkitModeBase::SelectedEmitterGraphTabID(TEXT("NiagaraSystemEditor_SelectedEmitterGraph"));
 const FName FNiagaraSystemToolkitModeBase::DebugSpreadsheetTabID(TEXT("NiagaraSystemEditor_DebugAttributeSpreadsheet"));
 const FName FNiagaraSystemToolkitModeBase::DebugCaptureTabID(TEXT("NiagaraSystemEditor_DebugCacheCapture"));
@@ -261,10 +260,10 @@ void FNiagaraSystemToolkitModeBase::RegisterTabFactories(TSharedPtr<FTabManager>
 // 		.SetDisplayName(LOCTEXT("SystemParameterDefinitions", "Parameter Definitions"))
 // 		.SetGroup(WorkspaceMenuCategory.ToSharedRef());
 
-	InTabManager->RegisterTabSpawner(SelectedEmitterStackTabID, FOnSpawnTab::CreateSP(this, &FNiagaraSystemToolkitModeBase::SpawnTab_SelectedEmitterStack))
-		.SetDisplayName(LOCTEXT("SelectedEmitterStacks", "Selected Emitters"))
+	InTabManager->RegisterTabSpawner(DetailsTabID, FOnSpawnTab::CreateSP(this, &FNiagaraSystemToolkitModeBase::SpawnTab_Details))
+		.SetDisplayName(LOCTEXT("Details", "Details"))
 		.SetGroup(WorkspaceMenuCategory.ToSharedRef())
-		.SetIcon(FSlateIcon(FNiagaraEditorStyle::Get().GetStyleSetName(), "Tab.VisualEffects"));
+		.SetIcon(FSlateIcon(FAppStyle::GetAppStyleSetName(), "Icons.Details"));
 
 	InTabManager->RegisterTabSpawner(SelectedEmitterGraphTabID, FOnSpawnTab::CreateSP(this, &FNiagaraSystemToolkitModeBase::SpawnTab_SelectedEmitterGraph))
 		.SetDisplayName(LOCTEXT("SelectedEmitterGraph", "Selected Emitter Graph"))
@@ -587,13 +586,13 @@ TSharedRef<SDockTab> FNiagaraSystemToolkitModeBase::SpawnTab_SystemParameterDefi
 	return SpawnedTab;
 }
 
-TSharedRef<SDockTab> FNiagaraSystemToolkitModeBase::SpawnTab_SelectedEmitterStack(const FSpawnTabArgs& Args)
+TSharedRef<SDockTab> FNiagaraSystemToolkitModeBase::SpawnTab_Details(const FSpawnTabArgs& Args)
 {
-	check(Args.GetTabId().TabType == SelectedEmitterStackTabID);
+	check(Args.GetTabId().TabType == DetailsTabID);
 
 	FNiagaraEditorModule& NiagaraEditorModule = FModuleManager::LoadModuleChecked<FNiagaraEditorModule>("NiagaraEditor");
 	TSharedRef<SDockTab> SpawnedTab = SNew(SDockTab)
-		.Label(LOCTEXT("SystemOverviewSelection", "Selection"))
+		.Label(LOCTEXT("SystemOverviewDetails", "Details"))
 		[
 			SNew(SWidgetSwitcher)
 			.WidgetIndex(this, &FNiagaraSystemToolkitModeBase::GetActiveSelectionDetailsIndex)
diff --git a/Engine/Plugins/FX/Niagara/Source/NiagaraEditor/Private/Toolkits/SystemToolkitModes/NiagaraSystemToolkitModeBase.h b/Engine/Plugins/FX/Niagara/Source/NiagaraEditor/Private/Toolkits/SystemToolkitModes/NiagaraSystemToolkitModeBase.h
--- a/Engine/Plugins/FX/Niagara/Source/NiagaraEditor/Private/Toolkits/SystemToolkitModes/NiagaraSystemToolkitModeBase.h
+++ b/Engine/Plugins/FX/Niagara/Source/NiagaraEditor/Private/Toolkits/SystemToolkitModes/NiagaraSystemToolkitModeBase.h
@@ -32,7 +32,7 @@ private:
 	TSharedRef<SDockTab> SpawnTab_SystemScript(const FSpawnTabArgs& Args);
 	TSharedRef<SDockTab> SpawnTab_SystemParameters(const FSpawnTabArgs& Args);
 	TSharedRef<SDockTab> SpawnTab_SystemParameterDefinitions(const FSpawnTabArgs& Args);
-	TSharedRef<SDockTab> SpawnTab_SelectedEmitterStack(const FSpawnTabArgs& Args);
+	TSharedRef<SDockTab> SpawnTab_Details(const FSpawnTabArgs& Args);
 	TSharedRef<SDockTab> SpawnTab_SelectedEmitterGraph(const FSpawnTabArgs& Args);
 	TSharedRef<SDockTab> SpawnTab_DebugSpreadsheet(const FSpawnTabArgs& Args);
 	TSharedRef<SDockTab> SpawnTab_DebugCapture(const FSpawnTabArgs& Args);
@@ -73,10 +73,9 @@ public:
 	static const FName CurveEditorTabID;
 	static const FName SequencerTabID;
 	static const FName SystemScriptTabID;
-	static const FName SystemDetailsTabID;
 	static const FName SystemParametersTabID;
 	static const FName SystemParameterDefinitionsTabID;
-	static const FName SelectedEmitterStackTabID;
+	static const FName DetailsTabID;
 	static const FName SelectedEmitterGraphTabID;
 	static const FName DebugSpreadsheetTabID;
 	static const FName DebugCaptureTabID;
diff --git a/Engine/Plugins/FX/Niagara/Source/NiagaraEditor/Private/Toolkits/SystemToolkitModes/NiagaraSystemToolkitMode_Default.cpp b/Engine/Plugins/FX/Niagara/Source/NiagaraEditor/Private/Toolkits/SystemToolkitModes/NiagaraSystemToolkitMode_Default.cpp
--- a/Engine/Plugins/FX/Niagara/Source/NiagaraEditor/Private/Toolkits/SystemToolkitModes/NiagaraSystemToolkitMode_Default.cpp
+++ b/Engine/Plugins/FX/Niagara/Source/NiagaraEditor/Private/Toolkits/SystemToolkitModes/NiagaraSystemToolkitMode_Default.cpp
@@ -17,7 +17,7 @@
 
 FNiagaraSystemToolkitMode_Default::FNiagaraSystemToolkitMode_Default(TWeakPtr<FNiagaraSystemToolkit> InSystemToolkit) : FNiagaraSystemToolkitModeBase(FNiagaraSystemToolkit::DefaultModeName, InSystemToolkit)
 {
-	TabLayout = FTabManager::NewLayout("Standalone_Niagara_System_Layout_v28")
+	TabLayout = FTabManager::NewLayout("Standalone_Niagara_System_Layout_v29")
 		->AddArea
 		(
 			FTabManager::NewPrimaryArea()->SetOrientation(Orient_Vertical)
@@ -71,10 +71,9 @@ FNiagaraSystemToolkitMode_Default::FNiagaraSystemToolkitMode_Default(TWeakPtr<FN
 					// Right
 					FTabManager::NewStack()
 					->SetSizeCoefficient(0.33f)
-					->AddTab(SelectedEmitterStackTabID, ETabState::OpenedTab)
+					->AddTab(DetailsTabID, ETabState::OpenedTab)
 					->AddTab(SelectedEmitterGraphTabID, ETabState::ClosedTab)
 					->AddTab(SystemScriptTabID, ETabState::ClosedTab)
-					->AddTab(SystemDetailsTabID, ETabState::ClosedTab)
 					->AddTab(DebugSpreadsheetTabID, ETabState::ClosedTab)
 					->AddTab(PreviewSettingsTabId, ETabState::ClosedTab)
 					->AddTab(GeneratedCodeTabID, ETabState::ClosedTab)
diff --git a/Engine/Plugins/FX/Niagara/Source/NiagaraEditor/Private/Toolkits/SystemToolkitModes/NiagaraSystemToolkitMode_Scalability.cpp b/Engine/Plugins/FX/Niagara/Source/NiagaraEditor/Private/Toolkits/SystemToolkitModes/NiagaraSystemToolkitMode_Scalability.cpp
--- a/Engine/Plugins/FX/Niagara/Source/NiagaraEditor/Private/Toolkits/SystemToolkitModes/NiagaraSystemToolkitMode_Scalability.cpp
+++ b/Engine/Plugins/FX/Niagara/Source/NiagaraEditor/Private/Toolkits/SystemToolkitModes/NiagaraSystemToolkitMode_Scalability.cpp
@@ -250,7 +250,7 @@ TSharedRef<FTabManager::FLayout> FNiagaraSystemToolkitMode_Scalability::ChooseTa
 {
 	if(SystemToolkit.Pin()->GetSystemViewModel()->GetEditMode() == ENiagaraSystemViewModelEditMode::SystemAsset)
 	{
-		return FTabManager::NewLayout("Standalone_Niagara_System_Layout_Scalability_v30")
+		return FTabManager::NewLayout("Standalone_Niagara_System_Layout_Scalability_v31")
 		->AddArea
 		(
 			FTabManager::NewPrimaryArea()->SetOrientation(Orient_Vertical)
@@ -317,10 +317,9 @@ TSharedRef<FTabManager::FLayout> FNiagaraSystemToolkitMode_Scalability::ChooseTa
 						FTabManager::NewStack()
 						->SetSizeCoefficient(0.3f)
 						->AddTab(ResolvedScalabilityTabID, ETabState::OpenedTab)
-						->AddTab(SelectedEmitterStackTabID, ETabState::OpenedTab)
+						->AddTab(DetailsTabID, ETabState::OpenedTab)
 						->AddTab(SelectedEmitterGraphTabID, ETabState::ClosedTab)
 						->AddTab(SystemScriptTabID, ETabState::ClosedTab)
-						->AddTab(SystemDetailsTabID, ETabState::ClosedTab)
 						->AddTab(DebugSpreadsheetTabID, ETabState::ClosedTab)
 						->AddTab(PreviewSettingsTabId, ETabState::ClosedTab)
 						->AddTab(GeneratedCodeTabID, ETabState::ClosedTab)
@@ -389,10 +388,9 @@ TSharedRef<FTabManager::FLayout> FNiagaraSystemToolkitMode_Scalability::ChooseTa
 					(
 						FTabManager::NewStack()
 						->SetSizeCoefficient(0.65f)
-						->AddTab(SelectedEmitterStackTabID, ETabState::OpenedTab)
+						->AddTab(DetailsTabID, ETabState::OpenedTab)
 						->AddTab(SelectedEmitterGraphTabID, ETabState::ClosedTab)
 						->AddTab(SystemScriptTabID, ETabState::ClosedTab)
-						->AddTab(SystemDetailsTabID, ETabState::ClosedTab)
 						->AddTab(DebugSpreadsheetTabID, ETabState::ClosedTab)
 						->AddTab(PreviewSettingsTabId, ETabState::ClosedTab)
 						->AddTab(GeneratedCodeTabID, ETabState::ClosedTab)
