Added an optional Guid to the parameter map builder that will skip all emitter nodes that don't match the given guid.
Useful for UI operations that are supposed to gather system parameters & THIS emitter's parameters. Without this, going through the system scripts will also gather parameters of other emitters.

diff --git a/Engine/Plugins/FX/Niagara/Source/NiagaraEditor/Private/Customizations/NiagaraTypeCustomizations.cpp b/Engine/Plugins/FX/Niagara/Source/NiagaraEditor/Private/Customizations/NiagaraTypeCustomizations.cpp
--- a/Engine/Plugins/FX/Niagara/Source/NiagaraEditor/Private/Customizations/NiagaraTypeCustomizations.cpp
+++ b/Engine/Plugins/FX/Niagara/Source/NiagaraEditor/Private/Customizations/NiagaraTypeCustomizations.cpp
@@ -41,6 +41,7 @@
 #include "NiagaraSimulationStageBase.h"
 #include "NiagaraSystem.h"
 #include "NiagaraTypes.h"
+#include "ViewModels/NiagaraEmitterHandleViewModel.h"
 
 #include "Widgets/SNiagaraParameterMenu.h"
 #include "Widgets/Input/SSuggestionTextBox.h"
@@ -153,8 +154,6 @@ TArray<FNiagaraVariableBase> FNiagaraStackAssetAction_VarBind::FindVariables(con
 		}
 	}
 	
-
-
 	for (const FNiagaraParameterMapHistory& History : Histories)
 	{
 		for (const FNiagaraVariable& Var : History.Variables)
@@ -271,14 +270,43 @@ TArray<FName> FNiagaraVariableAttributeBindingCustomization::GetNames(const FVer
 	{
 		return Names;
 	}
+	
+	TArray<UNiagaraGraph*> GraphsToTraverse;
+	
+	GraphsToTraverse.Add(Cast<UNiagaraScriptSource>(BaseEmitter.GetEmitterData()->GraphSource)->NodeGraph);
+	UNiagaraSystem* System = BaseEmitter.Emitter->GetTypedOuter<UNiagaraSystem>();
+	GraphsToTraverse.Add(Cast<UNiagaraScriptSource>(System->GetSystemUpdateScript()->GetLatestSource())->NodeGraph);
+
+	TSet<FNiagaraVariableBase> Vars;
+	TArray<TArray<FNiagaraVariable>> MapHistoryVars;
+	for(UNiagaraGraph* Graph : GraphsToTraverse)
+	{
+		TArray<UNiagaraNodeOutput*> OutputNodes;
+		Graph->GetNodesOfClass<UNiagaraNodeOutput>(OutputNodes);
 
-	const bool bSystem = true;
-	const bool bEmitter = true;
-	const bool bParticles = true;
-	const bool bUser = true;
-	const bool bAllowStatic = true;
-	TArray<FNiagaraVariableBase> Vars = FNiagaraStackAssetAction_VarBind::FindVariables(InEmitter, bSystem, bEmitter, bParticles, bUser, bAllowStatic);
+		for(UNiagaraNodeOutput* NodeOutput : OutputNodes)
+		{
+			FNiagaraParameterMapHistoryBuilder Builder;
+			Builder.ExclusiveEmitterHandle = EmitterHandleGuid;
+			
+			Builder.BuildParameterMaps(NodeOutput);
 
+			for(FNiagaraParameterMapHistory& History : Builder.Histories)
+			{
+				MapHistoryVars.Add(History.Variables);
+			}
+		}
+	}
+
+	for(TArray<FNiagaraVariable>& MapHistoryVariables : MapHistoryVars)
+	{
+		for(FNiagaraVariable& MapHistoryVariable : MapHistoryVariables)
+		{
+			Vars.Add(MapHistoryVariable);
+		}
+	}
+
+	bool bAllowStatic = true;
 	for (const FNiagaraVariableBase& Var : Vars)
 	{
 		if (!bAllowStatic && Var.GetType() != TargetVariableBinding->GetType() )
@@ -451,6 +479,10 @@ void FNiagaraVariableAttributeBindingCustomization::CustomizeHeader(TSharedRef<I
 
 		if (BaseEmitter.Emitter)
 		{
+			UNiagaraSystem* System = BaseEmitter.Emitter->GetTypedOuter<UNiagaraSystem>();
+			TSharedPtr<FNiagaraSystemViewModel> SystemViewModel = TNiagaraViewModelManager<UNiagaraSystem, FNiagaraSystemViewModel>::GetExistingViewModelForObject(System);
+			EmitterHandleGuid = SystemViewModel->GetEmitterHandleViewModelForEmitter(BaseEmitter)->GetId();
+			
 			TargetVariableBinding = (FNiagaraVariableAttributeBinding*)PropertyHandle->GetValueBaseAddress((uint8*)Objects[0]);
 			DefaultVariableBinding = (FNiagaraVariableAttributeBinding*)PropertyHandle->GetValueBaseAddress((uint8*)Objects[0]->GetClass()->GetDefaultObject());
 				
diff --git a/Engine/Plugins/FX/Niagara/Source/NiagaraEditor/Private/Customizations/NiagaraTypeCustomizations.h b/Engine/Plugins/FX/Niagara/Source/NiagaraEditor/Private/Customizations/NiagaraTypeCustomizations.h
--- a/Engine/Plugins/FX/Niagara/Source/NiagaraEditor/Private/Customizations/NiagaraTypeCustomizations.h
+++ b/Engine/Plugins/FX/Niagara/Source/NiagaraEditor/Private/Customizations/NiagaraTypeCustomizations.h
@@ -190,6 +190,8 @@ private:
 	class UNiagaraSimulationStageBase* SimulationStage = nullptr;
 	struct FNiagaraVariableAttributeBinding* TargetVariableBinding = nullptr;
 	const struct FNiagaraVariableAttributeBinding* DefaultVariableBinding = nullptr;
+	/** The emitter handle guid this binding is used in. Used to gather available parameters for the system & this emitter (not all emitters) */
+	FGuid EmitterHandleGuid;
 };
 
 class FNiagaraUserParameterBindingCustomization : public IPropertyTypeCustomization
diff --git a/Engine/Plugins/FX/Niagara/Source/NiagaraEditor/Private/NiagaraNodeEmitter.cpp b/Engine/Plugins/FX/Niagara/Source/NiagaraEditor/Private/NiagaraNodeEmitter.cpp
--- a/Engine/Plugins/FX/Niagara/Source/NiagaraEditor/Private/NiagaraNodeEmitter.cpp
+++ b/Engine/Plugins/FX/Niagara/Source/NiagaraEditor/Private/NiagaraNodeEmitter.cpp
@@ -298,7 +298,7 @@ void UNiagaraNodeEmitter::BuildParameterMapHistory(FNiagaraParameterMapHistoryBu
 {
 	Super::BuildParameterMapHistory(OutHistory, bRecursive, bFilterForCompilation);
 
-	if (!IsNodeEnabled() && OutHistory.GetIgnoreDisabled())
+	if ((!IsNodeEnabled() && OutHistory.GetIgnoreDisabled()) || (OutHistory.ExclusiveEmitterHandle.IsSet() && (OutHistory.ExclusiveEmitterHandle != GetEmitterHandleId())))
 	{
 		RouteParameterMapAroundMe(OutHistory, bRecursive);
 		return;
diff --git a/Engine/Plugins/FX/Niagara/Source/NiagaraEditor/Public/NiagaraParameterMapHistory.h b/Engine/Plugins/FX/Niagara/Source/NiagaraEditor/Public/NiagaraParameterMapHistory.h
--- a/Engine/Plugins/FX/Niagara/Source/NiagaraEditor/Public/NiagaraParameterMapHistory.h
+++ b/Engine/Plugins/FX/Niagara/Source/NiagaraEditor/Public/NiagaraParameterMapHistory.h
@@ -533,6 +533,12 @@ public:
 
 	FCompileConstantResolver ConstantResolver;
 
+	/** An optional emitter handle guid you can register. This will skip parameter map building of all emitters that don't match this one, if set.
+	 * Reason for this is to be able to build parameter map history for a system, while excluding certain emitters.
+	 * Useful for emitter context UI that should retrieve system parameters & emitter parameters, but only emitter parameters of the given emitter.
+	 */
+	TOptional<FGuid> ExclusiveEmitterHandle;
+
 	bool bShouldBuildSubHistories = true;
 
 	int32 MaxGraphDepthTraversal = INDEX_NONE;
