Fixed pin coloring for constant value make nodes

diff --git a/Engine/Plugins/FX/Niagara/Source/NiagaraEditor/Private/Widgets/SNiagaraGraphNodeConvert.cpp b/Engine/Plugins/FX/Niagara/Source/NiagaraEditor/Private/Widgets/SNiagaraGraphNodeConvert.cpp
--- a/Engine/Plugins/FX/Niagara/Source/NiagaraEditor/Private/Widgets/SNiagaraGraphNodeConvert.cpp
+++ b/Engine/Plugins/FX/Niagara/Source/NiagaraEditor/Private/Widgets/SNiagaraGraphNodeConvert.cpp
@@ -131,15 +131,6 @@ void SNiagaraGraphNodeConvert::UpdateGraphNode()
 	}
 	SNiagaraGraphNode::UpdateGraphNode();
 	
-	if (ConvertNode && ConvertNode->IsLocalConstantValue())
-	{
-		// hide input pins when the node is just displaying a constant value
-		for (TSharedRef<SGraphPin> InPin : InputPins)
-		{
-			InPin->SetPinColorModifier(ConvertNode->IsWiringShown() ? FLinearColor::White : FLinearColor::Transparent);
-		}
-	}
-
 	// set visibility of add pins
 	if (InputPins.Num() > 0 && OutputPins.Num() > 0 && ConvertNode)
 	{
