Stopped considering input pins for get nodes as parameter reference

This fixes an issue where a pin that was only referenced in a get node would count 2 references, which would stop it from being renamed properly from pin as it would duplicate a new parameter

diff --git a/Engine/Plugins/FX/Niagara/Source/NiagaraEditor/Private/NiagaraGraph.cpp b/Engine/Plugins/FX/Niagara/Source/NiagaraEditor/Private/NiagaraGraph.cpp
--- a/Engine/Plugins/FX/Niagara/Source/NiagaraEditor/Private/NiagaraGraph.cpp
+++ b/Engine/Plugins/FX/Niagara/Source/NiagaraEditor/Private/NiagaraGraph.cpp
@@ -3846,7 +3846,7 @@ void UNiagaraGraph::RefreshParameterReferences() const
 					bool bAddReference = true;
 
 					// we exclude default pins from adding references as they shouldn't count. Originally they weren't marked as parameter pins so this wasn't necessary.
-					if(WriteEvent.Pin->GetOwningNode()->IsA<UNiagaraNodeParameterMapGet>() && WriteEvent.Pin->Direction == EGPD_Input)
+					if(WriteEvent.Pin->Direction == EGPD_Input && WriteEvent.Pin->GetOwningNode()->IsA<UNiagaraNodeParameterMapGet>())
 					{
 						bAddReference = false;
 					}
