Fixed scratch pad parameter panel bug where if a used function had the same parameter as the script asset that owned it, it would override the asset's own parameter. UI wise it would look fine due to same name + type, but selecting it wouldn't work as the guid wouldn't match.

- Removed sub graphs adding their script variables to be used for the script parameter panel
- Refactored into using the graph's script variables directly instead of using the reference map. Reference map is now only used for reference count, not for item creation.

diff --git a/Engine/Plugins/FX/Niagara/Source/NiagaraEditor/Private/ViewModels/NiagaraParameterPanelViewModel.cpp b/Engine/Plugins/FX/Niagara/Source/NiagaraEditor/Private/ViewModels/NiagaraParameterPanelViewModel.cpp
--- a/Engine/Plugins/FX/Niagara/Source/NiagaraEditor/Private/ViewModels/NiagaraParameterPanelViewModel.cpp
+++ b/Engine/Plugins/FX/Niagara/Source/NiagaraEditor/Private/ViewModels/NiagaraParameterPanelViewModel.cpp
@@ -3101,86 +3101,61 @@ TArray<FNiagaraParameterPanelItem> FNiagaraScriptToolkitParameterPanelViewModel:
 	TMap<FNiagaraVariable, FNiagaraParameterPanelItem> VisitedParameterToItemMap;
 	TArray<FNiagaraVariable> VisitedInvalidParameters;
 
-	TArray<UNiagaraGraph*> Graphs = GetEditableGraphsConst();
-
-	// For scripts we use the reference maps cached in the graph to collect parameters.
-	for (const UNiagaraGraph* Graph : Graphs)
+	TArray<UNiagaraGraph*> EditableGraphs = GetEditableGraphsConst();
+	for(UNiagaraGraph* Graph : EditableGraphs)
 	{
-		TMap<FNiagaraVariable, TObjectPtr<UNiagaraScriptVariable>> ParameterToScriptVariableMap;
-		ParameterToScriptVariableMap.Append(Graph->GetAllMetaData());
-		// Collect all subgraphs to get their UNiagaraScriptVariables to resolve metadata for parameters in the parameter reference map.
-		TSet<UNiagaraGraph*> SubGraphs;
-		TArray<UNiagaraNodeFunctionCall*> FunctionCallNodes;
-		Graph->GetNodesOfClass<UNiagaraNodeFunctionCall>(FunctionCallNodes);
-		for (UNiagaraNodeFunctionCall* FunctionCallNode : FunctionCallNodes)
-		{
-			UNiagaraScriptSource* FunctionScriptSource = FunctionCallNode->GetFunctionScriptSource();
-			if (FunctionScriptSource)
-			{
-				ParameterToScriptVariableMap.Append(FunctionScriptSource->NodeGraph->GetAllMetaData());
-			}
-		}
+		const TMap<FNiagaraVariable, TObjectPtr<UNiagaraScriptVariable>>& ParameterMetaData = Graph->GetAllMetaData();
+		const TMap<FNiagaraVariable, FNiagaraGraphParameterReferenceCollection>& ReferenceCollection = Graph->GetParameterReferenceMap();
 
-		for (const TPair<FNiagaraVariable, FNiagaraGraphParameterReferenceCollection>& ParameterElement : Graph->GetParameterReferenceMap())
+		for(const auto& ParameterMetaDataPair : ParameterMetaData)
 		{
-			const FNiagaraVariable& Var = ParameterElement.Key;
+			const FNiagaraVariable Parameter = ParameterMetaDataPair.Key;
+			const UNiagaraScriptVariable* ScriptVariable = ParameterMetaDataPair.Value;
+
 			// If this variable has already been visited and does not have a valid namespace then skip it.
-			if (VisitedInvalidParameters.Contains(Var))
+			if (VisitedInvalidParameters.Contains(Parameter))
 			{
 				continue;
 			}
 
-			if (FNiagaraParameterPanelItem* ItemPtr = VisitedParameterToItemMap.Find(Var))
+			// -Now make sure the variable namespace is in a valid category. If not, skip it.
+			FNiagaraNamespaceMetadata CandidateNamespaceMetaData;
+			if (ScriptVariable->GetIsStaticSwitch())
 			{
-				// This variable has already been registered, increment the reference count.
-				ItemPtr->ReferenceCount += ParameterElement.Value.ParameterReferences.Num();
+				CandidateNamespaceMetaData = FNiagaraEditorUtilities::GetNamespaceMetaDataForId(FNiagaraEditorGuids::StaticSwitchNamespaceMetaDataGuid);
 			}
 			else
 			{
-				// This variable has not been registered, prepare the FNiagaraParameterPanelItem.
-				// -First lookup the script variable.
-				TObjectPtr<UNiagaraScriptVariable> const* ScriptVarPtr = ParameterToScriptVariableMap.Find(Var);
-				TObjectPtr<UNiagaraScriptVariable> ScriptVar = ScriptVarPtr != nullptr ? *ScriptVarPtr : nullptr;
-				if (!ScriptVar)
-				{
-					// Create a new UNiagaraScriptVariable to represent this parameter for the lifetime of the ParameterPanelViewModel.
-					ScriptVar = NewObject<UNiagaraScriptVariable>(GetTransientPackage());
-					ScriptVar->AddToRoot();
-					ScriptVar->Init(Var, FNiagaraVariableMetaData());
-					TransientParameterToScriptVarMap.Add(Var, ScriptVar);
-				}
-
-				// -Now make sure the variable namespace is in a valid category. If not, skip it.
-				FNiagaraNamespaceMetadata CandidateNamespaceMetaData;
-				if (ScriptVar->GetIsStaticSwitch())
-				{
-					CandidateNamespaceMetaData = FNiagaraEditorUtilities::GetNamespaceMetaDataForId(FNiagaraEditorGuids::StaticSwitchNamespaceMetaDataGuid);
-				}
-				else
-				{
-					CandidateNamespaceMetaData = FNiagaraEditorUtilities::GetNamespaceMetaDataForVariableName(Var.GetName());
-				}
-
-				if (CachedCurrentCategories.Contains(FNiagaraParameterPanelCategory(CandidateNamespaceMetaData)) == false)
-				{
-					VisitedInvalidParameters.Add(Var);
-					continue;
-				}
+				CandidateNamespaceMetaData = FNiagaraEditorUtilities::GetNamespaceMetaDataForVariableName(Parameter.GetName());
+			}
 
-				FNiagaraParameterPanelItem Item = FNiagaraParameterPanelItem();
-				Item.ScriptVariable = ScriptVar;
-				Item.NamespaceMetaData = CandidateNamespaceMetaData;
-				Item.bExternallyReferenced = false;
-				Item.bSourcedFromCustomStackContext = false;
+			if (CachedCurrentCategories.Contains(FNiagaraParameterPanelCategory(CandidateNamespaceMetaData)) == false)
+			{
+				VisitedInvalidParameters.Add(Parameter);
+				continue;
+			}
 
-				// Determine whether the item is name aliasing a parameter definition's parameter.
-				Item.DefinitionMatchState = FNiagaraParameterDefinitionsUtilities::GetDefinitionMatchStateForParameter(ScriptVar->Variable);
+			FNiagaraParameterPanelItem Item = FNiagaraParameterPanelItem();
+			Item.ScriptVariable = ScriptVariable;
+			Item.NamespaceMetaData = CandidateNamespaceMetaData;
+			Item.bExternallyReferenced = false;
+			Item.bSourcedFromCustomStackContext = false;
 
-				// -Increment the reference count.
-				Item.ReferenceCount += ParameterElement.Value.ParameterReferences.Num();
+			// Determine whether the item is name aliasing a parameter definition's parameter.
+			Item.DefinitionMatchState = FNiagaraParameterDefinitionsUtilities::GetDefinitionMatchStateForParameter(Parameter);
 
-				VisitedParameterToItemMap.Add(Var, Item);
+			// Set the reference count. It is possible a parameter has no entry in the reference collection,
+			// for example a static switch that was added via propagation but then removed again will have no references. There might be other exceptions, too.
+			if(ReferenceCollection.Contains(Parameter))
+			{
+				Item.ReferenceCount = ReferenceCollection[Parameter].ParameterReferences.Num();
 			}
+			else
+			{
+				Item.ReferenceCount = 0;
+			}
+
+			VisitedParameterToItemMap.Add(Parameter, Item);
 		}
 	}
 
