Moved PreSave to be called before UpdateOriginalEmitter.
This fixes the issue where PreSave will trigger compilation which in turn updates the changeID if editing an emitter, causing the Apply action to be available right after saving.

diff --git a/Engine/Plugins/FX/Niagara/Source/NiagaraEditor/Private/ViewModels/NiagaraSystemViewModel.cpp b/Engine/Plugins/FX/Niagara/Source/NiagaraEditor/Private/ViewModels/NiagaraSystemViewModel.cpp
--- a/Engine/Plugins/FX/Niagara/Source/NiagaraEditor/Private/ViewModels/NiagaraSystemViewModel.cpp
+++ b/Engine/Plugins/FX/Niagara/Source/NiagaraEditor/Private/ViewModels/NiagaraSystemViewModel.cpp
@@ -1029,6 +1029,7 @@ void FNiagaraSystemViewModel::NotifyPreSave()
 		System->WaitForCompilationComplete(true);
 	}
 
+	// we want to save without compile for edit turned on for best results, but make sure to turn in back on again after save is done
 	if (bSupportCompileForEdit)
 	{
 		check(System->bCompileForEdit);
@@ -1040,6 +1041,7 @@ void FNiagaraSystemViewModel::NotifyPreSave()
 
 void FNiagaraSystemViewModel::NotifyPostSave()
 {
+	// we compile for edit again after having it turned off for PreSave. This should fetch DDC data and must not update the ChangeID
 	if (bSupportCompileForEdit)
 	{
 		check(!System->bCompileForEdit);
