Changed ensure to a log warning to prevent crashing the content worker on invalid asset registry entries

diff --git a/Engine/Plugins/FX/Niagara/Source/Niagara/Private/NiagaraParameterDefinitionsSubscriber.cpp b/Engine/Plugins/FX/Niagara/Source/Niagara/Private/NiagaraParameterDefinitionsSubscriber.cpp
--- a/Engine/Plugins/FX/Niagara/Source/Niagara/Private/NiagaraParameterDefinitionsSubscriber.cpp
+++ b/Engine/Plugins/FX/Niagara/Source/Niagara/Private/NiagaraParameterDefinitionsSubscriber.cpp
@@ -186,7 +186,7 @@ void INiagaraParameterDefinitionsSubscriber::SynchronizeWithParameterDefinitions
 			UNiagaraParameterDefinitionsBase* ParameterDefinitions = Cast<UNiagaraParameterDefinitionsBase>(ParameterDefinitionsAssetDatum.GetAsset());
 			if (ParameterDefinitions == nullptr)
 			{
-				ensureMsgf(false, TEXT("Failed to load parameter definition from asset registry!"));
+				UE_LOG(LogNiagara, Warning, TEXT("Failed to load parameter definition %s from asset registry!"), *ParameterDefinitionsAssetDatum.GetSoftObjectPath().ToString());
 				continue;
 			}
 			AllDefinitions.Add(ParameterDefinitions);
