- Add a context for data interface debug hud drawing

#rnx

diff --git a/Engine/Plugins/FX/Niagara/Source/Niagara/Classes/NDISkeletalMeshCommon.h b/Engine/Plugins/FX/Niagara/Source/Niagara/Classes/NDISkeletalMeshCommon.h
--- a/Engine/Plugins/FX/Niagara/Source/Niagara/Classes/NDISkeletalMeshCommon.h
+++ b/Engine/Plugins/FX/Niagara/Source/Niagara/Classes/NDISkeletalMeshCommon.h
@@ -39,7 +39,7 @@ struct FSkeletalMeshAccessorHelper
 	FSkeletalMeshAccessorHelper operator=(FSkeletalMeshAccessorHelper&&) = delete;
 
 	template<typename FilterMode, typename AreaWeightingMode>
-	FORCEINLINE void Init(FNDISkeletalMesh_InstanceData* InstData)
+	FORCEINLINE void Init(const FNDISkeletalMesh_InstanceData* InstData)
 	{
 		Comp = Cast<USkeletalMeshComponent>(InstData->SceneComponent.Get());
 		Mesh = InstData->SkeletalMesh.Get();
@@ -78,10 +78,10 @@ struct FSkeletalMeshAccessorHelper
 };
 
 template<>
-void FSkeletalMeshAccessorHelper::Init<TNDISkelMesh_FilterModeSingle, TNDISkelMesh_AreaWeightingOff>(FNDISkeletalMesh_InstanceData* InstData);
+void FSkeletalMeshAccessorHelper::Init<TNDISkelMesh_FilterModeSingle, TNDISkelMesh_AreaWeightingOff>(const FNDISkeletalMesh_InstanceData* InstData);
 
 template<>
-void FSkeletalMeshAccessorHelper::Init<TNDISkelMesh_FilterModeSingle, TNDISkelMesh_AreaWeightingOn>(FNDISkeletalMesh_InstanceData* InstData);
+void FSkeletalMeshAccessorHelper::Init<TNDISkelMesh_FilterModeSingle, TNDISkelMesh_AreaWeightingOn>(const FNDISkeletalMesh_InstanceData* InstData);
 
 //////////////////////////////////////////////////////////////////////////
 
diff --git a/Engine/Plugins/FX/Niagara/Source/Niagara/Classes/NiagaraDataInterface.h b/Engine/Plugins/FX/Niagara/Source/Niagara/Classes/NiagaraDataInterface.h
--- a/Engine/Plugins/FX/Niagara/Source/Niagara/Classes/NiagaraDataInterface.h
+++ b/Engine/Plugins/FX/Niagara/Source/Niagara/Classes/NiagaraDataInterface.h
@@ -249,6 +249,34 @@ private:
 
 //////////////////////////////////////////////////////////////////////////
 
+#if WITH_NIAGARA_DEBUGGER
+struct NIAGARA_API FNDIDrawDebugHudContext
+{
+	FNDIDrawDebugHudContext(bool bInVerbose, UWorld* InWorld, UCanvas* InCanvas, FNiagaraSystemInstance* InSystemInstance)
+		: bVerbose(bInVerbose)
+		, World(InWorld)
+		, Canvas(InCanvas)
+		, SystemInstance(InSystemInstance)
+	{
+	}
+
+	bool IsVerbose() const { return bVerbose; }
+	UWorld* GetWorld() const { return World; }
+	UCanvas* GetCanvas() const { return Canvas; }
+	const FNiagaraSystemInstance* GetSystemInstance() const { return SystemInstance; }
+	FString& GetOutputString() { return OutputString; }
+
+protected:
+	bool					bVerbose;
+	UWorld*					World;
+	UCanvas*				Canvas;
+	FNiagaraSystemInstance*	SystemInstance;
+	FString					OutputString;
+};
+#endif //WITH_NIAGARA_DEBUGGER
+
+//////////////////////////////////////////////////////////////////////////
+
 struct NIAGARA_API FNDIGpuComputeContext
 {
 	FNDIGpuComputeContext(FRDGBuilder& InGraphBuilder, const FNiagaraGpuComputeDispatchInterface& InComputeDispatchInterface)
@@ -647,7 +675,16 @@ public:
 	should be kept light to avoid polluting the display.
 	You can also use the Canvas to draw additional information based on verbosity
 	*/
-	virtual void DrawDebugHud(UCanvas* Canvas, FNiagaraSystemInstance* SystemInstance, FString& VariableDataString, bool bVerbose) const {};
+	virtual void DrawDebugHud(FNDIDrawDebugHudContext& DebugHudContext) const
+	{
+	PRAGMA_DISABLE_DEPRECATION_WARNINGS
+		DrawDebugHud(DebugHudContext.GetCanvas(), const_cast<FNiagaraSystemInstance*>(DebugHudContext.GetSystemInstance()), DebugHudContext.GetOutputString(), DebugHudContext.IsVerbose());
+	PRAGMA_ENABLE_DEPRECATION_WARNINGS
+	}
+
+	// Deprecated method will be removed in a future version
+	UE_DEPRECATED(5.3, "Please update to the context based DrawDebugHud as this will be removed in a future version.")
+	virtual void DrawDebugHud(UCanvas* Canvas, FNiagaraSystemInstance* SystemInstance, FString& VariableDataString, bool bVerbose) const {}
 #endif
 
 #if WITH_EDITOR	
diff --git a/Engine/Plugins/FX/Niagara/Source/Niagara/Classes/NiagaraDataInterfaceArray.h b/Engine/Plugins/FX/Niagara/Source/Niagara/Classes/NiagaraDataInterfaceArray.h
--- a/Engine/Plugins/FX/Niagara/Source/Niagara/Classes/NiagaraDataInterfaceArray.h
+++ b/Engine/Plugins/FX/Niagara/Source/Niagara/Classes/NiagaraDataInterfaceArray.h
@@ -58,7 +58,7 @@ struct INDIArrayProxyBase : public FNiagaraDataInterfaceProxyRW
 	virtual bool UpgradeFunctionCall(FNiagaraFunctionSignature& FunctionSignature) const = 0;
 #endif
 #if WITH_NIAGARA_DEBUGGER
-	virtual void DrawDebugHud(UCanvas* Canvas, FNiagaraSystemInstance* SystemInstance, FString& VariableDataString, bool bVerbose) const = 0;
+	virtual void DrawDebugHud(FNDIDrawDebugHudContext& DebugHudContext) const = 0;
 #endif
 	virtual bool CopyToInternal(INDIArrayProxyBase* Destination) const = 0;
 	virtual bool Equals(const INDIArrayProxyBase* Other) const = 0;
@@ -95,7 +95,7 @@ public:
 	virtual bool CanExecuteOnTarget(ENiagaraSimTarget Target) const override { return true; }
 
 #if WITH_NIAGARA_DEBUGGER
-	virtual void DrawDebugHud(UCanvas* Canvas, FNiagaraSystemInstance* SystemInstance, FString& VariableDataString, bool bVerbose) const { GetProxyAs<INDIArrayProxyBase>()->DrawDebugHud(Canvas, SystemInstance, VariableDataString, bVerbose); }
+	virtual void DrawDebugHud(FNDIDrawDebugHudContext& DebugHudContext) const { GetProxyAs<INDIArrayProxyBase>()->DrawDebugHud(DebugHudContext); }
 #endif
 
 	virtual bool CopyToInternal(UNiagaraDataInterface* Destination) const;
diff --git a/Engine/Plugins/FX/Niagara/Source/Niagara/Classes/NiagaraDataInterfaceArrayImpl.h b/Engine/Plugins/FX/Niagara/Source/Niagara/Classes/NiagaraDataInterfaceArrayImpl.h
--- a/Engine/Plugins/FX/Niagara/Source/Niagara/Classes/NiagaraDataInterfaceArrayImpl.h
+++ b/Engine/Plugins/FX/Niagara/Source/Niagara/Classes/NiagaraDataInterfaceArrayImpl.h
@@ -1015,9 +1015,9 @@ struct FNDIArrayProxyImpl : public INDIArrayProxyBase
 	}
 #endif
 #if WITH_NIAGARA_DEBUGGER
-	virtual void DrawDebugHud(UCanvas* Canvas, FNiagaraSystemInstance* SystemInstance, FString& VariableDataString, bool bVerbose) const override
+	virtual void DrawDebugHud(FNDIDrawDebugHudContext& DebugHudContext) const override
 	{
-		FNDIArrayInstanceData_GameThread<TArrayType>* InstanceData = PerInstanceData_GameThread.FindRef(SystemInstance->GetId());
+		FNDIArrayInstanceData_GameThread<TArrayType>* InstanceData = PerInstanceData_GameThread.FindRef(DebugHudContext.GetSystemInstance()->GetId());
 		if (InstanceData == nullptr )
 		{
 			return;
@@ -1039,7 +1039,7 @@ struct FNDIArrayProxyImpl : public INDIArrayProxyBase
 			CpuValuesString.Append(TEXT(", ..."));
 		}
 
-		VariableDataString = FString::Printf(
+		DebugHudContext.GetOutputString().Appendf(
 			TEXT("Type(%s) CpuLength(%d) CpuValues(%s)"),
 			*FNDIArrayImplHelper<TArrayType>::GetTypeDefinition().GetName(),
 			ArrayData.GetArray().Num(),
diff --git a/Engine/Plugins/FX/Niagara/Source/Niagara/Classes/NiagaraDataInterfaceSkeletalMesh.h b/Engine/Plugins/FX/Niagara/Source/Niagara/Classes/NiagaraDataInterfaceSkeletalMesh.h
--- a/Engine/Plugins/FX/Niagara/Source/Niagara/Classes/NiagaraDataInterfaceSkeletalMesh.h
+++ b/Engine/Plugins/FX/Niagara/Source/Niagara/Classes/NiagaraDataInterfaceSkeletalMesh.h
@@ -655,7 +655,7 @@ struct FNDISkeletalMesh_InstanceData
 
 	FORCEINLINE int32 GetLODIndex()const { return CachedLODIdx; }
 
-	FORCEINLINE_DEBUGGABLE const FSkinWeightVertexBuffer* GetSkinWeights()
+	FORCEINLINE_DEBUGGABLE const FSkinWeightVertexBuffer* GetSkinWeights() const
 	{
 		USkeletalMeshComponent* SkelComp = Cast<USkeletalMeshComponent>(SceneComponent.Get());
 		if (SkelComp != nullptr && SkelComp->GetSkeletalMeshAsset() != nullptr)
@@ -776,7 +776,7 @@ public:
 	virtual bool Equals(const UNiagaraDataInterface* Other) const override;
 	virtual bool CanExecuteOnTarget(ENiagaraSimTarget Target)const override { return true; }
 #if WITH_NIAGARA_DEBUGGER
-	virtual void DrawDebugHud(UCanvas* Canvas, FNiagaraSystemInstance* SystemInstance, FString& VariableDataString, bool bVerbose) const override;
+	virtual void DrawDebugHud(FNDIDrawDebugHudContext& DebugHudContext) const override;
 #endif
 #if WITH_EDITOR
 	virtual void GetFeedback(UNiagaraSystem* Asset, UNiagaraComponent* Component, TArray<FNiagaraDataInterfaceError>& OutErrors,
diff --git a/Engine/Plugins/FX/Niagara/Source/Niagara/Internal/DataInterface/NiagaraDataInterfaceStaticMesh.h b/Engine/Plugins/FX/Niagara/Source/Niagara/Internal/DataInterface/NiagaraDataInterfaceStaticMesh.h
--- a/Engine/Plugins/FX/Niagara/Source/Niagara/Internal/DataInterface/NiagaraDataInterfaceStaticMesh.h
+++ b/Engine/Plugins/FX/Niagara/Source/Niagara/Internal/DataInterface/NiagaraDataInterfaceStaticMesh.h
@@ -168,7 +168,7 @@ protected:
 	virtual bool CopyToInternal(UNiagaraDataInterface* Destination) const override;
 public:
 #if WITH_NIAGARA_DEBUGGER
-	virtual void DrawDebugHud(UCanvas* Canvas, FNiagaraSystemInstance* SystemInstance, FString& VariableDataString, bool bVerbose) const override;
+	virtual void DrawDebugHud(FNDIDrawDebugHudContext& DebugHudContext) const override;
 #endif
 #if WITH_EDITOR
 	virtual void GetFeedback(UNiagaraSystem* Asset, UNiagaraComponent* Component, TArray<FNiagaraDataInterfaceError>& OutErrors,
diff --git a/Engine/Plugins/FX/Niagara/Source/Niagara/Private/DataInterface/NiagaraDataInterfaceStaticMesh.cpp b/Engine/Plugins/FX/Niagara/Source/Niagara/Private/DataInterface/NiagaraDataInterfaceStaticMesh.cpp
--- a/Engine/Plugins/FX/Niagara/Source/Niagara/Private/DataInterface/NiagaraDataInterfaceStaticMesh.cpp
+++ b/Engine/Plugins/FX/Niagara/Source/Niagara/Private/DataInterface/NiagaraDataInterfaceStaticMesh.cpp
@@ -3167,9 +3167,9 @@ bool UNiagaraDataInterfaceStaticMesh::CopyToInternal(UNiagaraDataInterface* Dest
 }
 
 #if WITH_NIAGARA_DEBUGGER
-void UNiagaraDataInterfaceStaticMesh::DrawDebugHud(UCanvas* Canvas, FNiagaraSystemInstance* SystemInstance, FString& VariableDataString, bool bVerbose) const
+void UNiagaraDataInterfaceStaticMesh::DrawDebugHud(FNDIDrawDebugHudContext& DebugHudContext) const
 {
-	NDIStaticMeshLocal::FInstanceData_GameThread* InstanceData_GT = SystemInstance->FindTypedDataInterfaceInstanceData<NDIStaticMeshLocal::FInstanceData_GameThread>(this);
+	const NDIStaticMeshLocal::FInstanceData_GameThread* InstanceData_GT = DebugHudContext.GetSystemInstance()->FindTypedDataInterfaceInstanceData<NDIStaticMeshLocal::FInstanceData_GameThread>(this);
 	if (InstanceData_GT == nullptr)
 	{
 		return;
@@ -3177,7 +3177,7 @@ void UNiagaraDataInterfaceStaticMesh::DrawDebugHud(UCanvas* Canvas, FNiagaraSyst
 
 	USceneComponent* SceneComponent = InstanceData_GT->SceneComponentWeakPtr.Get();
 	UStaticMesh* SkeletalMesh = InstanceData_GT->StaticMeshWeakPtr.Get();
-	VariableDataString = FString::Printf(TEXT("StaticMesh(%s) StaticMeshComp(%s)"), *GetNameSafe(SkeletalMesh), *GetNameSafe(SceneComponent));
+	DebugHudContext.GetOutputString().Appendf(TEXT("StaticMesh(%s) StaticMeshComp(%s)"), *GetNameSafe(SkeletalMesh), *GetNameSafe(SceneComponent));
 }
 #endif
 
diff --git a/Engine/Plugins/FX/Niagara/Source/Niagara/Private/NiagaraDataInterfaceRigidMeshCollisionQuery.cpp b/Engine/Plugins/FX/Niagara/Source/Niagara/Private/NiagaraDataInterfaceRigidMeshCollisionQuery.cpp
--- a/Engine/Plugins/FX/Niagara/Source/Niagara/Private/NiagaraDataInterfaceRigidMeshCollisionQuery.cpp
+++ b/Engine/Plugins/FX/Niagara/Source/Niagara/Private/NiagaraDataInterfaceRigidMeshCollisionQuery.cpp
@@ -978,9 +978,9 @@ UNiagaraDataInterfaceRigidMeshCollisionQuery::UNiagaraDataInterfaceRigidMeshColl
 
 #if WITH_NIAGARA_DEBUGGER
 
-void UNiagaraDataInterfaceRigidMeshCollisionQuery::DrawDebugHud(UCanvas* Canvas, FNiagaraSystemInstance* SystemInstance, FString& VariableDataString, bool bVerbose) const
+void UNiagaraDataInterfaceRigidMeshCollisionQuery::DrawDebugHud(FNDIDrawDebugHudContext& DebugHudContext) const
 {
-	FNDIRigidMeshCollisionData* InstanceData_GT = SystemInstance->FindTypedDataInterfaceInstanceData<FNDIRigidMeshCollisionData>(this);
+	const FNDIRigidMeshCollisionData* InstanceData_GT = DebugHudContext.GetSystemInstance()->FindTypedDataInterfaceInstanceData<FNDIRigidMeshCollisionData>(this);
 	if (InstanceData_GT == nullptr || !InstanceData_GT->AssetArrays.IsValid())
 	{
 		return;
@@ -992,7 +992,7 @@ void UNiagaraDataInterfaceRigidMeshCollisionQuery::DrawDebugHud(UCanvas* Canvas,
 	const uint32 SphereCount = ElementOffsets.CapsuleOffset - ElementOffsets.SphereOffset;
 	const uint32 CapsuleCount = ElementOffsets.NumElements - ElementOffsets.CapsuleOffset;
 
-	VariableDataString = FString::Printf(TEXT("Boxes(%d) Spheres(%d) Capsules(%d)"), BoxCount, SphereCount, CapsuleCount);
+	DebugHudContext.GetOutputString().Appendf(TEXT("Boxes(%d) Spheres(%d) Capsules(%d)"), BoxCount, SphereCount, CapsuleCount);
 
 	auto GetCurrentTransform = [&](int32 ElementIndex)
 	{
@@ -1013,8 +1013,10 @@ void UNiagaraDataInterfaceRigidMeshCollisionQuery::DrawDebugHud(UCanvas* Canvas,
 		return ElementMatrix.GetTransposed();
 	};
 
-	if (bVerbose)
+	if (DebugHudContext.IsVerbose())
 	{
+		UCanvas* Canvas = DebugHudContext.GetCanvas();
+
 		// the DrawDebugCanvas* functions don't reasoanbly handle the near clip plane (both in terms of clipping and in terms of
 		// objects being behind the camera); so we introduce this culling behavior to work around it
 		auto ShouldClip = [&](UCanvas* Canvas, const FMatrix& Transform, const FBoxSphereBounds& Bounds)
@@ -1063,7 +1065,7 @@ void UNiagaraDataInterfaceRigidMeshCollisionQuery::DrawDebugHud(UCanvas* Canvas,
 			const UFont* Font = GEngine->GetMediumFont();
 			Canvas->SetDrawColor(FColor::White);
 
-			auto DrawDebugActor = [&](TWeakObjectPtr<AActor>& InWeakActor, const TCHAR* ActorSourceString)
+			auto DrawDebugActor = [&](const TWeakObjectPtr<AActor>& InWeakActor, const TCHAR* ActorSourceString)
 			{
 				if (AActor* Actor = InWeakActor.Get())
 				{
@@ -1091,12 +1093,12 @@ void UNiagaraDataInterfaceRigidMeshCollisionQuery::DrawDebugHud(UCanvas* Canvas,
 				}
 			};
 
-			for (TWeakObjectPtr<AActor>& ExplicitActor : InstanceData_GT->ExplicitActors)
+			for (const TWeakObjectPtr<AActor>& ExplicitActor : InstanceData_GT->ExplicitActors)
 			{
 				DrawDebugActor(ExplicitActor, TEXT("Explicit"));
 			}
 
-			for (TWeakObjectPtr<AActor>& FoundActor : InstanceData_GT->FoundActors)
+			for (const TWeakObjectPtr<AActor>& FoundActor : InstanceData_GT->FoundActors)
 			{
 				DrawDebugActor(FoundActor, TEXT("Found"));
 			}
diff --git a/Engine/Plugins/FX/Niagara/Source/Niagara/Private/NiagaraDataInterfaceSkeletalMesh.cpp b/Engine/Plugins/FX/Niagara/Source/Niagara/Private/NiagaraDataInterfaceSkeletalMesh.cpp
--- a/Engine/Plugins/FX/Niagara/Source/Niagara/Private/NiagaraDataInterfaceSkeletalMesh.cpp
+++ b/Engine/Plugins/FX/Niagara/Source/Niagara/Private/NiagaraDataInterfaceSkeletalMesh.cpp
@@ -2465,9 +2465,9 @@ bool UNiagaraDataInterfaceSkeletalMesh::PerInstanceTick(void* PerInstanceData, F
 }
 
 #if WITH_NIAGARA_DEBUGGER
-void UNiagaraDataInterfaceSkeletalMesh::DrawDebugHud(UCanvas* Canvas, FNiagaraSystemInstance* SystemInstance, FString& VariableDataString, bool bVerbose) const
+void UNiagaraDataInterfaceSkeletalMesh::DrawDebugHud(FNDIDrawDebugHudContext& DebugHudContext) const
 {
-	FNDISkeletalMesh_InstanceData* InstanceData_GT = SystemInstance->FindTypedDataInterfaceInstanceData<FNDISkeletalMesh_InstanceData>(this);
+	const FNDISkeletalMesh_InstanceData* InstanceData_GT = DebugHudContext.GetSystemInstance()->FindTypedDataInterfaceInstanceData<FNDISkeletalMesh_InstanceData>(this);
 	if (InstanceData_GT == nullptr)
 	{
 		return;
@@ -2475,9 +2475,9 @@ void UNiagaraDataInterfaceSkeletalMesh::DrawDebugHud(UCanvas* Canvas, FNiagaraSy
 
 	USceneComponent* SceneComponent = InstanceData_GT->SceneComponent.Get();
 	USkeletalMesh* SkeletalMesh = InstanceData_GT->SkeletalMesh.Get();
-	VariableDataString = FString::Printf(TEXT("Skeleton(%s) SkelComp(%s)"), *GetNameSafe(SkeletalMesh), *GetNameSafe(SceneComponent));
+	DebugHudContext.GetOutputString().Appendf(TEXT("Skeleton(%s) SkelComp(%s)"), *GetNameSafe(SkeletalMesh), *GetNameSafe(SceneComponent));
 
-	if ( bVerbose && SceneComponent && SkeletalMesh )
+	if ( DebugHudContext.IsVerbose() && SceneComponent && SkeletalMesh)
 	{
 		FSkeletalMeshAccessorHelper MeshAccessor;
 		MeshAccessor.Init<TNDISkelMesh_FilterModeNone, TNDISkelMesh_AreaWeightingOff>(InstanceData_GT);
@@ -2486,6 +2486,7 @@ void UNiagaraDataInterfaceSkeletalMesh::DrawDebugHud(UCanvas* Canvas, FNiagaraSy
 			const TArray<FTransform3f>& BoneTransforms = MeshAccessor.SkinningData->CurrComponentTransforms();
 			const FMatrix& InstanceTransform = InstanceData_GT->Transform;
 
+			UCanvas* Canvas = DebugHudContext.GetCanvas();
 			for (const FTransform3f& BoneTransform : BoneTransforms)
 			{
 				const FVector BoneLocation = InstanceTransform.TransformPosition(FVector(BoneTransform.GetLocation()));
@@ -3540,7 +3541,7 @@ void UNiagaraDataInterfaceSkeletalMesh::VMGetPreSkinnedLocalBounds(FVectorVMExte
 //////////////////////////////////////////////////////////////////////////
 
 template<>
-void FSkeletalMeshAccessorHelper::Init<TNDISkelMesh_FilterModeSingle, TNDISkelMesh_AreaWeightingOff>(FNDISkeletalMesh_InstanceData* InstData)
+void FSkeletalMeshAccessorHelper::Init<TNDISkelMesh_FilterModeSingle, TNDISkelMesh_AreaWeightingOff>(const FNDISkeletalMesh_InstanceData* InstData)
 {
 	Comp = Cast<USkeletalMeshComponent>(InstData->SceneComponent.Get());
 	Mesh = InstData->SkeletalMesh.Get();
@@ -3569,7 +3570,7 @@ void FSkeletalMeshAccessorHelper::Init<TNDISkelMesh_FilterModeSingle, TNDISkelMe
 }
 
 template<>
-void FSkeletalMeshAccessorHelper::Init<TNDISkelMesh_FilterModeSingle, TNDISkelMesh_AreaWeightingOn>(FNDISkeletalMesh_InstanceData* InstData)
+void FSkeletalMeshAccessorHelper::Init<TNDISkelMesh_FilterModeSingle, TNDISkelMesh_AreaWeightingOn>(const FNDISkeletalMesh_InstanceData* InstData)
 {
 	Comp = Cast<USkeletalMeshComponent>(InstData->SceneComponent.Get());
 	Mesh = InstData->SkeletalMesh.Get();
diff --git a/Engine/Plugins/FX/Niagara/Source/Niagara/Private/NiagaraDebugHud.cpp b/Engine/Plugins/FX/Niagara/Source/Niagara/Private/NiagaraDebugHud.cpp
--- a/Engine/Plugins/FX/Niagara/Source/Niagara/Private/NiagaraDebugHud.cpp
+++ b/Engine/Plugins/FX/Niagara/Source/Niagara/Private/NiagaraDebugHud.cpp
@@ -2902,19 +2902,19 @@ void FNiagaraDebugHud::DrawComponents(FNiagaraWorldManager* WorldManager, UCanva
 								if (ParameterStoreVariable.Variable.IsDataInterface())
 								{
 									UNiagaraDataInterface* DataInterface = ParameterStore->GetDataInterface(ParameterStoreVariable.Variable);
-									FString DataInterfaceString;
+									FNDIDrawDebugHudContext DebugHudContext(Settings.DataInterfaceVerbosity == ENiagaraDebugHudVerbosity::Verbose, World, Canvas, SystemInstance);
 									if ( DataInterface != nullptr && Settings.DataInterfaceVerbosity != ENiagaraDebugHudVerbosity::None )
 									{
-										DataInterface->DrawDebugHud(Canvas, SystemInstance, DataInterfaceString, Settings.DataInterfaceVerbosity == ENiagaraDebugHudVerbosity::Verbose);
+										DataInterface->DrawDebugHud(DebugHudContext);
 									}
 
-									if ( DataInterfaceString.IsEmpty() )
+									if (DebugHudContext.GetOutputString().IsEmpty())
 									{
 										StringBuilder.Appendf(TEXT("%s(%s %s)"), *ParameterStoreVariable.Variable.GetName().ToString(), *GetNameSafe(ParameterStoreVariable.Variable.GetType().GetClass()) , *GetNameSafe(DataInterface));
 									}
 									else
 									{
-										StringBuilder.Appendf(TEXT("%s(%s)"), *ParameterStoreVariable.Variable.GetName().ToString(), *DataInterfaceString);
+										StringBuilder.Appendf(TEXT("%s(%s)"), *ParameterStoreVariable.Variable.GetName().ToString(), *DebugHudContext.GetOutputString());
 									}
 								}
 								else if (ParameterStoreVariable.Variable.IsUObject())
diff --git a/Engine/Plugins/FX/Niagara/Source/Niagara/Public/NiagaraDataInterfaceRigidMeshCollisionQuery.h b/Engine/Plugins/FX/Niagara/Source/Niagara/Public/NiagaraDataInterfaceRigidMeshCollisionQuery.h
--- a/Engine/Plugins/FX/Niagara/Source/Niagara/Public/NiagaraDataInterfaceRigidMeshCollisionQuery.h
+++ b/Engine/Plugins/FX/Niagara/Source/Niagara/Public/NiagaraDataInterfaceRigidMeshCollisionQuery.h
@@ -231,7 +231,7 @@ public:
 	virtual void GetVMExternalFunction(const FVMExternalFunctionBindingInfo& BindingInfo, void* InstanceData, FVMExternalFunction& OutFunc) override;
 	virtual bool CanExecuteOnTarget(ENiagaraSimTarget Target) const override { return Target == ENiagaraSimTarget::GPUComputeSim; }
 #if WITH_NIAGARA_DEBUGGER
-	virtual void DrawDebugHud(UCanvas* Canvas, FNiagaraSystemInstance* SystemInstance, FString& VariableDataString, bool bVerbose) const override;
+	virtual void DrawDebugHud(FNDIDrawDebugHudContext& DebugHudContext) const override;
 #endif
 	virtual bool InitPerInstanceData(void* PerInstanceData, FNiagaraSystemInstance* SystemInstance) override;
 	virtual void DestroyPerInstanceData(void* PerInstanceData, FNiagaraSystemInstance* SystemInstance)override;
diff --git a/Engine/Plugins/FX/Niagara/Source/Niagara/Public/NiagaraSystemInstance.h b/Engine/Plugins/FX/Niagara/Source/Niagara/Public/NiagaraSystemInstance.h
--- a/Engine/Plugins/FX/Niagara/Source/Niagara/Public/NiagaraSystemInstance.h
+++ b/Engine/Plugins/FX/Niagara/Source/Niagara/Public/NiagaraSystemInstance.h
@@ -279,7 +279,7 @@ public:
 	FNiagaraSystemInstanceID GetId() const { return ID; }
 
 	/** Returns the instance data for a particular interface for this System. */
-	FORCEINLINE void* FindDataInterfaceInstanceData(UNiagaraDataInterface* Interface) 
+	FORCEINLINE void* FindDataInterfaceInstanceData(const UNiagaraDataInterface* Interface)
 	{
 		if (auto* InstDataOffsetPair = DataInterfaceInstanceDataOffsets.FindByPredicate([&](auto& Pair){ return Pair.Key.Get() == Interface;}))
 		{
@@ -298,6 +298,10 @@ public:
 		return nullptr;
 	}
 
+	FORCEINLINE const void* FindDataInterfaceInstanceData(const UNiagaraDataInterface* Interface) const { return const_cast<FNiagaraSystemInstance*>(this)->FindDataInterfaceInstanceData(Interface); }
+	template<typename TDataType>
+	FORCEINLINE const TDataType* FindTypedDataInterfaceInstanceData(const UNiagaraDataInterface* Interface) const { return const_cast<FNiagaraSystemInstance*>(this)->FindTypedDataInterfaceInstanceData<TDataType>(Interface); }
+
 	FORCEINLINE const FNiagaraPerInstanceDIFuncInfo& GetPerInstanceDIFunction(ENiagaraSystemSimulationScript ScriptType, int32 FuncIndex)const { return PerInstanceDIFunctions[(int32)ScriptType][FuncIndex]; }
 
 	void EvaluateBoundFunction(FName FunctionName, bool& UsedOnCpu, bool& UsedOnGpu) const;
