- Fix mesh renderer emitter mode with facing or locked axis not working

diff --git a/Engine/Plugins/FX/Niagara/Source/Niagara/Private/NiagaraRendererMeshes.cpp b/Engine/Plugins/FX/Niagara/Source/Niagara/Private/NiagaraRendererMeshes.cpp
--- a/Engine/Plugins/FX/Niagara/Source/Niagara/Private/NiagaraRendererMeshes.cpp
+++ b/Engine/Plugins/FX/Niagara/Source/Niagara/Private/NiagaraRendererMeshes.cpp
@@ -949,7 +949,12 @@ void FNiagaraRendererMeshes::SetupElementForGPUScene(
 	GPUSceneRes.DynamicPrimitiveData.EnableInstanceDynamicData(bNeedsPrevTransform);
 	GPUSceneRes.DynamicPrimitiveData.SetNumInstanceCustomDataFloats(1);
 
-	if (SourceMode == ENiagaraRendererSourceDataMode::Particles)
+	// Take the GPU write path if we are particles or we have options set that would involve duplicating a lot of code from shader into here
+	const bool bWriteOnGpu = 
+		SourceMode == ENiagaraRendererSourceDataMode::Particles ||
+		FacingMode != ENiagaraMeshFacingMode::Default ||
+		bLockedAxisEnable;
+	if (bWriteOnGpu)
 	{
 		FMemory::Memzero(&GPUSceneRes.GPUWriteParams, sizeof(GPUSceneRes.GPUWriteParams));
 
