- Fix for a couple of functions exposed on other compilers

#rnx

diff --git a/Engine/Plugins/FX/Niagara/Shaders/Private/NiagaraDataInterfaceSceneCapture2DTemplate.ush b/Engine/Plugins/FX/Niagara/Shaders/Private/NiagaraDataInterfaceSceneCapture2DTemplate.ush
--- a/Engine/Plugins/FX/Niagara/Shaders/Private/NiagaraDataInterfaceSceneCapture2DTemplate.ush
+++ b/Engine/Plugins/FX/Niagara/Shaders/Private/NiagaraDataInterfaceSceneCapture2DTemplate.ush
@@ -42,7 +42,7 @@ void Deproject_{ParameterName}(in float2 UV, in float Depth, out float3 Position
 	Position = mul(float4(NDC, Depth, 1.0f), {ParameterName}_InvTextureProjMatrix).xyz;
 }
 
-void TextureProject_{ParameterName}(in float3 Position, in float MipLevel, in bool bValidateTextureBounds, in bool bValidateDepthBounds, in float4 DefaultColor, out bool OutInBounds, out float4 OutColor, out float3 UVW)
+void TextureProject_{ParameterName}(in float3 Position, in float MipLevel, in bool bValidateTextureBounds, in bool bValidateDepthBounds, in float4 DefaultColor, out bool OutInBounds, out float4 OutValue, out float3 UVW)
 {
 	Project_{ParameterName}(Position, UVW);
 
@@ -59,7 +59,7 @@ void TextureProjectDepth_{ParameterName}(in float3 Position, in float MipLevel,
 	OutInBounds = bValidateTextureBounds ? all(UVW.xy >= 0.0f && UVW.xy < 1.0f) : true;
 	if ( bPointSample )
 	{
-		const int3 Texel = int3(round(UV * {ParameterName}_TextureSize.xy), int(round(MipLevel)));
+		const int3 Texel = int3(round(UVW * {ParameterName}_TextureSize.xy), int(round(MipLevel)));
 		OutCaptureDepth = {ParameterName}_Texture.Load(Texel).r;
 	}
 	else
