- Fix CIS issue

#rnx

diff --git a/Engine/Plugins/FX/Niagara/Source/Niagara/Private/DataInterface/NiagaraDataInterfaceSceneCapture2D.cpp b/Engine/Plugins/FX/Niagara/Source/Niagara/Private/DataInterface/NiagaraDataInterfaceSceneCapture2D.cpp
--- a/Engine/Plugins/FX/Niagara/Source/Niagara/Private/DataInterface/NiagaraDataInterfaceSceneCapture2D.cpp
+++ b/Engine/Plugins/FX/Niagara/Source/Niagara/Private/DataInterface/NiagaraDataInterfaceSceneCapture2D.cpp
@@ -908,7 +908,14 @@ void UNiagaraDataInterfaceSceneCapture2D::SetShaderParameters(const FNiagaraData
 	FNDInstanceData_RenderThread* InstanceData = TextureProxy.InstanceData_RT.Find(Context.GetSystemInstanceID());
 
 	FShaderParameters* ShaderParameters		= Context.GetParameterNestedStruct<FShaderParameters>();
-	ShaderParameters->Texture				= InstanceData->TextureReferenceRHI.IsValid() ? InstanceData->TextureReferenceRHI : GWhiteTexture->TextureRHI;
+	if (InstanceData->TextureReferenceRHI.IsValid())
+	{
+		ShaderParameters->Texture			= InstanceData->TextureReferenceRHI;
+	}
+	else
+	{
+		ShaderParameters->Texture			= GWhiteTexture->TextureRHI;
+	}
 	ShaderParameters->TextureSampler		= InstanceData->SamplerStateRHI.IsValid() ? InstanceData->SamplerStateRHI.GetReference() : TStaticSamplerState<SF_Bilinear>::GetRHI();
 	ShaderParameters->TextureSize			= InstanceData->TextureSize;
 	ShaderParameters->TextureViewMatrix		= InstanceData->TextureViewMatrix;
