Fix skeletal mesh post load being called on game thread from Niagara when actual loading is still happening from the loading thread

diff --git a/Engine/Plugins/FX/Niagara/Source/Niagara/Private/NiagaraDataInterfaceSkeletalMesh.cpp b/Engine/Plugins/FX/Niagara/Source/Niagara/Private/NiagaraDataInterfaceSkeletalMesh.cpp
--- a/Engine/Plugins/FX/Niagara/Source/Niagara/Private/NiagaraDataInterfaceSkeletalMesh.cpp
+++ b/Engine/Plugins/FX/Niagara/Source/Niagara/Private/NiagaraDataInterfaceSkeletalMesh.cpp
@@ -2236,12 +2236,6 @@ void UNiagaraDataInterfaceSkeletalMesh::PostLoad()
 {
 	Super::PostLoad();
 
-#if WITH_EDITOR
-	if(USkeletalMesh* LocalPreviewMesh = PreviewMesh.Get())
-	{
-		LocalPreviewMesh->ConditionalPostLoad();
-	}
-#endif
 #if WITH_EDITORONLY_DATA
 	if (Source_DEPRECATED != nullptr)
 	{
