- Remove log warning about no attributes on None
- This isn't useful and comes from data interfaces not involved in the simulation

#rnx

diff --git a/Engine/Plugins/FX/Niagara/Source/Niagara/Private/NiagaraDataInterfaceGrid3DCollection.cpp b/Engine/Plugins/FX/Niagara/Source/Niagara/Private/NiagaraDataInterfaceGrid3DCollection.cpp
--- a/Engine/Plugins/FX/Niagara/Source/Niagara/Private/NiagaraDataInterfaceGrid3DCollection.cpp
+++ b/Engine/Plugins/FX/Niagara/Source/Niagara/Private/NiagaraDataInterfaceGrid3DCollection.cpp
@@ -3341,7 +3341,10 @@ bool UNiagaraDataInterfaceGrid3DCollection::InitPerInstanceData(void* PerInstanc
 
 	if (NumAttribChannelsFound == 0)
 	{
-		UE_LOG(LogNiagara, Warning, TEXT("Zero attributes defined on %s"), *Proxy->SourceDIName.ToString());
+		if (!Proxy->SourceDIName.IsNone())
+		{
+			UE_LOG(LogNiagara, Warning, TEXT("Zero attributes defined on %s"), *Proxy->SourceDIName.ToString());
+		}
 
 		// Push Updates to Proxy.
 		FNiagaraDataInterfaceProxyGrid3DCollectionProxy* RT_Proxy = GetProxyAs<FNiagaraDataInterfaceProxyGrid3DCollectionProxy>();
diff --git a/Engine/Plugins/FX/Niagara/Source/Niagara/Private/NiagaraDataInterfaceRasterizationGrid3D.cpp b/Engine/Plugins/FX/Niagara/Source/Niagara/Private/NiagaraDataInterfaceRasterizationGrid3D.cpp
--- a/Engine/Plugins/FX/Niagara/Source/Niagara/Private/NiagaraDataInterfaceRasterizationGrid3D.cpp
+++ b/Engine/Plugins/FX/Niagara/Source/Niagara/Private/NiagaraDataInterfaceRasterizationGrid3D.cpp
@@ -928,7 +928,7 @@ bool UNiagaraDataInterfaceRasterizationGrid3D::InitPerInstanceData(void* PerInst
 	InstanceData->NumTiles.Y = NumTilesY;
 	InstanceData->NumTiles.Z = NumTilesZ;
 
-	if (NumAttribChannelsFound == 0)
+	if (NumAttribChannelsFound == 0 && !Proxy->SourceDIName.IsNone())
 	{
 		UE_LOG(LogNiagara, Warning, TEXT("Zero attributes defined on %s"), *Proxy->SourceDIName.ToString());
 	}
