- Fix an issue where element count could be -1 leading to invalid dispatch counts
- We now also early cull custom dispatch count stages to avoid building dispatch groups for them on the RT

diff --git a/Engine/Plugins/FX/Niagara/Source/Niagara/Classes/NiagaraGPUSystemTick.h b/Engine/Plugins/FX/Niagara/Source/Niagara/Classes/NiagaraGPUSystemTick.h
--- a/Engine/Plugins/FX/Niagara/Source/Niagara/Classes/NiagaraGPUSystemTick.h
+++ b/Engine/Plugins/FX/Niagara/Source/Niagara/Classes/NiagaraGPUSystemTick.h
@@ -29,9 +29,9 @@ struct FNiagaraComputeInstanceData
 		FPerStageInfo(int32 InNumIterations, FIntVector InElementCountXYZ) : NumIterations(InNumIterations), ElementCountXYZ(InElementCountXYZ) {}
 
 		int32 NumIterations = 0;
-		FIntVector ElementCountXYZ = FIntVector::NoneValue;
+		FIntVector ElementCountXYZ = FIntVector::ZeroValue;
 
-		bool ShouldRunStage() const { return NumIterations > 0 && (ElementCountXYZ != FIntVector::ZeroValue); }
+		bool ShouldRunStage() const { return NumIterations > 0; }
 	};
 
 	FNiagaraComputeInstanceData()
diff --git a/Engine/Plugins/FX/Niagara/Source/Niagara/Private/NiagaraGPUSystemTick.cpp b/Engine/Plugins/FX/Niagara/Source/Niagara/Private/NiagaraGPUSystemTick.cpp
--- a/Engine/Plugins/FX/Niagara/Source/Niagara/Private/NiagaraGPUSystemTick.cpp
+++ b/Engine/Plugins/FX/Niagara/Source/Niagara/Private/NiagaraGPUSystemTick.cpp
@@ -213,7 +213,7 @@ void FNiagaraGPUSystemTick::Init(FNiagaraSystemInstance* InSystemInstance)
 			for ( FSimulationStageMetaData& SimStageMetaData : GPUContext->SimStageInfo )
 			{
 				int32 NumIterations = SimStageMetaData.NumIterations;
-				FIntVector ElementCountXYZ = FIntVector::NoneValue;
+				FIntVector ElementCountXYZ = FIntVector(1, 1, 1);
 				if (SimStageMetaData.ShouldRunStage(InstanceData->bResetData))
 				{
 					InstanceData->bHasMultipleStages = true;
@@ -237,35 +237,27 @@ void FNiagaraGPUSystemTick::Init(FNiagaraSystemInstance* InSystemInstance)
 					}
 					if (SimStageMetaData.IterationSourceType == ENiagaraIterationSource::DirectSet)
 					{
-						if (!SimStageMetaData.ElementCountXBinding.IsNone() && (SimStageMetaData.GpuDispatchType >= ENiagaraGpuDispatchType::OneD))
+						FNiagaraParameterStore& BoundParamStore = EmitterInstance->GetRendererBoundVariables();
+						ElementCountXYZ = FIntVector(1, 1, 1);
+						if (SimStageMetaData.GpuDispatchType >= ENiagaraGpuDispatchType::OneD)
 						{
-							FNiagaraParameterStore& BoundParamStore = EmitterInstance->GetRendererBoundVariables();
-							if (const uint8* ParameterData = BoundParamStore.GetParameterData(FNiagaraVariable(FNiagaraTypeDefinition::GetIntDef(), SimStageMetaData.ElementCountXBinding)))
-							{
-								ElementCountXYZ.X = *reinterpret_cast<const int32*>(ParameterData);
-							}
+							ElementCountXYZ.X = SimStageMetaData.ElementCountXBinding.IsNone() ? 0 : BoundParamStore.GetParameterValueOrDefault(FNiagaraVariable(FNiagaraTypeDefinition::GetIntDef(), SimStageMetaData.ElementCountXBinding), 0);
 						}
-						if (!SimStageMetaData.ElementCountYBinding.IsNone() && (SimStageMetaData.GpuDispatchType >= ENiagaraGpuDispatchType::TwoD))
+						if (SimStageMetaData.GpuDispatchType >= ENiagaraGpuDispatchType::TwoD)
 						{
-							FNiagaraParameterStore& BoundParamStore = EmitterInstance->GetRendererBoundVariables();
-							if (const uint8* ParameterData = BoundParamStore.GetParameterData(FNiagaraVariable(FNiagaraTypeDefinition::GetIntDef(), SimStageMetaData.ElementCountYBinding)))
-							{
-								ElementCountXYZ.Y = *reinterpret_cast<const int32*>(ParameterData);
-							}
+							ElementCountXYZ.Y = SimStageMetaData.ElementCountYBinding.IsNone() ? 0 : BoundParamStore.GetParameterValueOrDefault(FNiagaraVariable(FNiagaraTypeDefinition::GetIntDef(), SimStageMetaData.ElementCountYBinding), 0);
 						}
-						if (!SimStageMetaData.ElementCountZBinding.IsNone() && (SimStageMetaData.GpuDispatchType >= ENiagaraGpuDispatchType::ThreeD))
+						if (SimStageMetaData.GpuDispatchType >= ENiagaraGpuDispatchType::ThreeD)
 						{
-							FNiagaraParameterStore& BoundParamStore = EmitterInstance->GetRendererBoundVariables();
-							if (const uint8* ParameterData = BoundParamStore.GetParameterData(FNiagaraVariable(FNiagaraTypeDefinition::GetIntDef(), SimStageMetaData.ElementCountZBinding)))
-							{
-								ElementCountXYZ.Z = *reinterpret_cast<const int32*>(ParameterData);
-							}
+							ElementCountXYZ.Z = SimStageMetaData.ElementCountZBinding.IsNone() ? 0 : BoundParamStore.GetParameterValueOrDefault(FNiagaraVariable(FNiagaraTypeDefinition::GetIntDef(), SimStageMetaData.ElementCountZBinding), 0);
 						}
 
-						// make sure we don't have negatives as we use the values directly for dispatch
-						ElementCountXYZ.X = FMath::Max(0, ElementCountXYZ.X);
-						ElementCountXYZ.Y = FMath::Max(0, ElementCountXYZ.Y);
-						ElementCountXYZ.Z = FMath::Max(0, ElementCountXYZ.Z);
+						// If any of the element count values are <= 0 we can kill the stage as it won't execute
+						if (ElementCountXYZ.GetMin() <= 0)
+						{
+							ElementCountXYZ = FIntVector::ZeroValue;
+							NumIterations = 0;
+						}
 					}
 				}
 				else
