- Fix PVS issue

#rnx

diff --git a/Engine/Plugins/FX/Niagara/Source/Niagara/Private/NiagaraGPUSystemTick.cpp b/Engine/Plugins/FX/Niagara/Source/Niagara/Private/NiagaraGPUSystemTick.cpp
--- a/Engine/Plugins/FX/Niagara/Source/Niagara/Private/NiagaraGPUSystemTick.cpp
+++ b/Engine/Plugins/FX/Niagara/Source/Niagara/Private/NiagaraGPUSystemTick.cpp
@@ -239,7 +239,9 @@ void FNiagaraGPUSystemTick::Init(FNiagaraSystemInstance* InSystemInstance)
 					{
 						FNiagaraParameterStore& BoundParamStore = EmitterInstance->GetRendererBoundVariables();
 						ElementCountXYZ = FIntVector(1, 1, 1);
-						if (SimStageMetaData.GpuDispatchType >= ENiagaraGpuDispatchType::OneD)
+
+						// This will cause a PVS warning, if we add something before this enum entry add the code back in
+						//if (SimStageMetaData.GpuDispatchType >= ENiagaraGpuDispatchType::OneD)
 						{
 							ElementCountXYZ.X = SimStageMetaData.ElementCountXBinding.IsNone() ? 0 : BoundParamStore.GetParameterValueOrDefault(FNiagaraVariable(FNiagaraTypeDefinition::GetIntDef(), SimStageMetaData.ElementCountXBinding), 0);
 						}
