- Add some clarity to what WholeMeshLOD means with regards to using -1

#rnx

diff --git a/Engine/Plugins/FX/Niagara/Source/Niagara/Classes/NiagaraDataInterfaceSkeletalMesh.h b/Engine/Plugins/FX/Niagara/Source/Niagara/Classes/NiagaraDataInterfaceSkeletalMesh.h
--- a/Engine/Plugins/FX/Niagara/Source/Niagara/Classes/NiagaraDataInterfaceSkeletalMesh.h
+++ b/Engine/Plugins/FX/Niagara/Source/Niagara/Classes/NiagaraDataInterfaceSkeletalMesh.h
@@ -721,7 +721,10 @@ public:
 	UPROPERTY(EditAnywhere, Category="Mesh")
 	TArray<FName> SamplingRegions;
 
-	/** If no regions are specified, we'll sample the whole mesh at this LODIndex. -1 indicates to use the last LOD.*/
+	/**
+	If no regions are specified, we'll sample the whole mesh at this LODIndex.
+	-1 will use the lowest quality LOD available, i.e. Number of LODs - 1.
+	*/
 	UPROPERTY(EditAnywhere, Category="Mesh")
 	int32 WholeMeshLOD;
 
