Added additional call to replace user param references on emitters & fixed renderer attribute binding references not being fixed when renaming a user param
Also moved over some required small changes from the Summary View 2.0 CL

diff --git a/Engine/Plugins/FX/Niagara/Source/NiagaraEditor/Private/NiagaraEditorUtilities.cpp b/Engine/Plugins/FX/Niagara/Source/NiagaraEditor/Private/NiagaraEditorUtilities.cpp
--- a/Engine/Plugins/FX/Niagara/Source/NiagaraEditor/Private/NiagaraEditorUtilities.cpp
+++ b/Engine/Plugins/FX/Niagara/Source/NiagaraEditor/Private/NiagaraEditorUtilities.cpp
@@ -2382,6 +2382,30 @@ TArray<FNiagaraUserParameterBinding*> FNiagaraEditorUtilities::GetUserParameterB
 	return Bindings;
 }
 
+TArray<TPair<FNiagaraVariableAttributeBinding*, ENiagaraRendererSourceDataMode>> FNiagaraEditorUtilities::GetVariableAttributeBindingsForParameter(TSharedRef<FNiagaraEmitterViewModel> EmitterViewModel, FNiagaraVariable Parameter)
+{
+	TArray<TPair<FNiagaraVariableAttributeBinding*, ENiagaraRendererSourceDataMode>> AttributeBindings;
+	
+	const TArray<UNiagaraRendererProperties*>& Renderers = EmitterViewModel->GetEmitter().GetEmitterData()->GetRenderers();
+	for(UNiagaraRendererProperties* Renderer : Renderers)
+	{
+		for(TFieldIterator<FStructProperty> It(Renderer->GetClass()); It; ++It)
+		{
+			if((*It)->Struct == StaticStruct<FNiagaraVariableAttributeBinding>())
+			{				
+				FNiagaraVariableAttributeBinding* ParameterBinding = (FNiagaraVariableAttributeBinding*) (*It)->ContainerPtrToValuePtr<void>(Renderer, 0);
+				
+				if(ParameterBinding->GetName().IsEqual(Parameter.GetName()))
+				{
+					AttributeBindings.Add({ParameterBinding, Renderer->GetCurrentSourceMode()});
+				}
+			}
+		}
+	}
+
+	return AttributeBindings;
+}
+
 TObjectPtr<UNiagaraScriptVariable> FNiagaraEditorUtilities::GetScriptVariableForUserParameter(const FNiagaraVariable& UserParameter, TSharedPtr<FNiagaraSystemViewModel> SystemViewModel)
 {
 	return Cast<UNiagaraSystemEditorData>(SystemViewModel->GetSystem().GetEditorData())->FindOrAddUserScriptVariable(UserParameter, SystemViewModel->GetSystem());
@@ -2392,7 +2416,7 @@ TObjectPtr<UNiagaraScriptVariable> FNiagaraEditorUtilities::GetScriptVariableFor
 	return Cast<UNiagaraSystemEditorData>(System.GetEditorData())->FindOrAddUserScriptVariable(UserParameter, System);
 }
 
-TObjectPtr<UNiagaraScriptVariable> FNiagaraEditorUtilities::FindScriptVariableForUserParameter(const FGuid& UserParameterGuid, UNiagaraSystem& System)
+const UNiagaraScriptVariable* FNiagaraEditorUtilities::FindScriptVariableForUserParameter(const FGuid& UserParameterGuid, const UNiagaraSystem& System)
 {
 	return Cast<UNiagaraSystemEditorData>(System.GetEditorData())->FindUserScriptVariable(UserParameterGuid);
 }
@@ -2406,6 +2430,12 @@ void FNiagaraEditorUtilities::ReplaceUserParameterReferences(TSharedRef<FNiagara
 		ReferencingBinding->Parameter = NewUserParameter;
 	}
 
+	TArray<TPair<FNiagaraVariableAttributeBinding*, ENiagaraRendererSourceDataMode>> ReferencingAttributeBindings = FNiagaraEditorUtilities::GetVariableAttributeBindingsForParameter(EmitterViewModel, OldUserParameter);
+	for(const auto& ReferencingBinding : ReferencingAttributeBindings)
+	{
+		ReferencingBinding.Key->SetValue(NewUserParameter.GetName(), EmitterViewModel->GetEmitter(), ReferencingBinding.Value);
+	}
+	
 	TArray<UNiagaraNodeParameterMapGet*> ReferencingMapGetNodes = FNiagaraEditorUtilities::GetParameterMapGetNodesWithUserParameter(EmitterViewModel, OldUserParameter);
 
 	for(UNiagaraNodeParameterMapGet* MapGet : ReferencingMapGetNodes)
@@ -4031,6 +4061,125 @@ TSharedRef<SWidget> FNiagaraParameterUtilities::GetParameterWidget(FNiagaraVaria
 	return ParameterWidget;
 }
 
+TSharedRef<SWidget> FNiagaraParameterUtilities::GetParameterWidget(FNiagaraVariable Variable, FNiagaraVariableMetaData Metadata, FNiagaraParameterWidgetOptions Options)
+{
+	TSharedRef<SHorizontalBox> IconBox = SNew(SHorizontalBox);
+
+	if(Options.bShowEditConditionIcon && !Metadata.EditCondition.InputName.IsNone())
+	{
+		TArray<FString> TargetValues = Metadata.EditCondition.TargetValues;
+
+		FString TooltipString = LOCTEXT("ParameterEditConditionWidgetTooltip", "This parameter is editable in the details panel when {0} is set to: ").ToString();
+
+		if(TargetValues.Num() == 0)
+		{
+			TooltipString.Append("true");
+		}
+		else
+		{
+			if(TargetValues.Num() > 1)
+			{
+				TooltipString.Append("\n");
+			}
+			
+			for(int32 TargetValueIndex = 0; TargetValueIndex < TargetValues.Num(); TargetValueIndex++)
+			{
+				TooltipString.Append(TargetValues[TargetValueIndex]);
+				
+				if(TargetValueIndex < TargetValues.Num() - 2)
+				{
+					TooltipString.Append("or\n");
+				}
+			}
+		}
+		
+		FText TooltipText = FText::FormatOrdered(FText::FromString(TooltipString), FText::FromName(Metadata.EditCondition.InputName));
+		
+		IconBox->AddSlot()
+		.Padding(2.f)
+		.AutoWidth()
+		.HAlign(HAlign_Center)
+		.VAlign(VAlign_Center)
+		[
+			SNew(SImage)
+			.Image(FAppStyle::GetBrush("Icons.Edit"))
+			.ToolTipText(TooltipText)
+		];
+	}
+
+	if(Options.bShowVisibilityConditionIcon && !Metadata.VisibleCondition.InputName.IsNone())
+	{
+		TArray<FString> TargetValues = Metadata.VisibleCondition.TargetValues;
+
+		FString TooltipString = LOCTEXT("ParameterVisibleConditionWidgetTooltip", "This parameter becomes visible in the details panel when {0} is set to: ").ToString();
+
+		if(TargetValues.Num() == 0)
+		{
+			TooltipString.Append("true");
+		}
+		else
+		{
+			if(TargetValues.Num() > 1)
+			{
+				TooltipString.Append("\n");
+			}
+			
+			for(int32 TargetValueIndex = 0; TargetValueIndex < TargetValues.Num(); TargetValueIndex++)
+			{
+				TooltipString.Append(TargetValues[TargetValueIndex]);
+				
+				if(TargetValueIndex < TargetValues.Num() - 2)
+				{
+					TooltipString.Append("or\n");
+				}
+			}
+		}
+		
+		FText TooltipText = FText::FormatOrdered(FText::FromString(TooltipString), FText::FromName(Metadata.VisibleCondition.InputName));
+		
+		IconBox->AddSlot()
+		.Padding(2.f)
+		.AutoWidth()
+		.HAlign(HAlign_Center)
+		.VAlign(VAlign_Center)
+		[
+			SNew(SImage)
+			.Image(FAppStyle::GetBrush("Icons.Visible"))
+			.ToolTipText(TooltipText)
+		];
+	}
+
+	if(Options.bShowAdvanced && Metadata.bAdvancedDisplay)
+	{		
+		FText TooltipText = LOCTEXT("ParameterIsAdvancedWidgetTooltip", "This parameter is 'advanced' and will only show up in the details panel when 'Show Advanced' is activated.");
+		
+		IconBox->AddSlot()
+		.Padding(2.f)
+		.AutoWidth()
+		.HAlign(HAlign_Center)
+		.VAlign(VAlign_Center)
+		[
+			SNew(SImage)
+			.Image(FAppStyle::GetBrush("DetailsView.PulldownArrow.Down"))
+			.ToolTipText(TooltipText)
+		];
+	}
+	
+	return SNew(SHorizontalBox)
+	+ SHorizontalBox::Slot()
+	.AutoWidth()
+	.VAlign(VAlign_Center)
+	[
+		GetParameterWidget(Variable, Options.bAddTypeIcon, Options.bShowValue)
+	]
+	+ SHorizontalBox::Slot()
+	.HAlign(HAlign_Right)
+	.VAlign(VAlign_Center)
+	[
+		IconBox
+	];	
+}
+
 TSharedRef<SToolTip> FNiagaraParameterUtilities::GetTooltipWidget(FNiagaraVariable Variable, bool bShowValue, TSharedPtr<SWidget> AdditionalVerticalWidget,  TSharedPtr<SWidget> AdditionalHorizontalWidget)
 {
 	FNiagaraTypeDefinition Type = Variable.GetType();
diff --git a/Engine/Plugins/FX/Niagara/Source/NiagaraEditor/Private/NiagaraSystemEditorData.cpp b/Engine/Plugins/FX/Niagara/Source/NiagaraEditor/Private/NiagaraSystemEditorData.cpp
--- a/Engine/Plugins/FX/Niagara/Source/NiagaraEditor/Private/NiagaraSystemEditorData.cpp
+++ b/Engine/Plugins/FX/Niagara/Source/NiagaraEditor/Private/NiagaraSystemEditorData.cpp
@@ -400,9 +400,9 @@ UNiagaraScriptVariable* UNiagaraSystemEditorData::FindOrAddUserScriptVariable(FN
 	return nullptr;
 }
 
-UNiagaraScriptVariable* UNiagaraSystemEditorData::FindUserScriptVariable(FGuid UserParameterGuid)
+const UNiagaraScriptVariable* UNiagaraSystemEditorData::FindUserScriptVariable(FGuid UserParameterGuid) const
 {
-	TObjectPtr<UNiagaraScriptVariable>* ExistingScriptVariable = UserParameterMetaData.FindByPredicate([UserParameterGuid](const UNiagaraScriptVariable* ScriptVariable)
+	const TObjectPtr<UNiagaraScriptVariable>* ExistingScriptVariable = UserParameterMetaData.FindByPredicate([UserParameterGuid](const UNiagaraScriptVariable* ScriptVariable)
 		{
 			return ScriptVariable->Metadata.GetVariableGuid() == UserParameterGuid;
 		});
diff --git a/Engine/Plugins/FX/Niagara/Source/NiagaraEditor/Private/ViewModels/NiagaraSystemViewModel.cpp b/Engine/Plugins/FX/Niagara/Source/NiagaraEditor/Private/ViewModels/NiagaraSystemViewModel.cpp
--- a/Engine/Plugins/FX/Niagara/Source/NiagaraEditor/Private/ViewModels/NiagaraSystemViewModel.cpp
+++ b/Engine/Plugins/FX/Niagara/Source/NiagaraEditor/Private/ViewModels/NiagaraSystemViewModel.cpp
@@ -1269,7 +1269,13 @@ bool FNiagaraSystemViewModel::RenameParameter(const FNiagaraVariable TargetParam
 			GetEditorData().RenameUserScriptVariable(TargetParameter, NewName);
 			ExposedParametersStore->RenameParameter(TargetParameter, NewName);
 		}
-		
+
+		FNiagaraVariable NewVariable = TargetParameter;
+		NewVariable.SetName(NewName);
+		for(TSharedPtr<FNiagaraEmitterHandleViewModel> EmitterHandleViewModel : GetEmitterHandleViewModels())
+		{
+			FNiagaraEditorUtilities::ReplaceUserParameterReferences(EmitterHandleViewModel->GetEmitterViewModel(), TargetParameter, NewVariable);
+		}
 		bExposedParametersRename = true;
 	}
 
diff --git a/Engine/Plugins/FX/Niagara/Source/NiagaraEditor/Public/NiagaraEditorUtilities.h b/Engine/Plugins/FX/Niagara/Source/NiagaraEditor/Public/NiagaraEditorUtilities.h
--- a/Engine/Plugins/FX/Niagara/Source/NiagaraEditor/Public/NiagaraEditorUtilities.h
+++ b/Engine/Plugins/FX/Niagara/Source/NiagaraEditor/Public/NiagaraEditorUtilities.h
@@ -283,9 +283,10 @@ namespace FNiagaraEditorUtilities
 	TArray<FNiagaraVariable> GetReferencedUserParametersFromEmitter(TSharedRef<FNiagaraEmitterViewModel> EmitterViewModel);
 	TArray<UNiagaraNodeParameterMapGet*> GetParameterMapGetNodesWithUserParameter(TSharedRef<FNiagaraEmitterViewModel> EmitterViewModel, FNiagaraVariable UserParameter);
 	TArray<FNiagaraUserParameterBinding*> GetUserParameterBindingsForUserParameter(TSharedRef<FNiagaraEmitterViewModel> EmitterViewModel, FNiagaraVariable UserParameter);
+	TArray<TPair<FNiagaraVariableAttributeBinding*, ENiagaraRendererSourceDataMode>> GetVariableAttributeBindingsForParameter(TSharedRef<FNiagaraEmitterViewModel> EmitterViewModel, FNiagaraVariable Parameter);
 	NIAGARAEDITOR_API TObjectPtr<UNiagaraScriptVariable> GetScriptVariableForUserParameter(const FNiagaraVariable& UserParameter, TSharedPtr<FNiagaraSystemViewModel> SystemViewModel);
 	NIAGARAEDITOR_API TObjectPtr<UNiagaraScriptVariable> GetScriptVariableForUserParameter(const FNiagaraVariable& UserParameter, UNiagaraSystem& System);
-	NIAGARAEDITOR_API TObjectPtr<UNiagaraScriptVariable> FindScriptVariableForUserParameter(const FGuid& UserParameterGuid, UNiagaraSystem& System);
+	NIAGARAEDITOR_API const UNiagaraScriptVariable* FindScriptVariableForUserParameter(const FGuid& UserParameterGuid, const UNiagaraSystem& System);
 	void ReplaceUserParameterReferences(TSharedRef<FNiagaraEmitterViewModel> EmitterViewModel, FNiagaraVariable OldUserParameter, FNiagaraVariable NewUserParameter);
 
 	NIAGARAEDITOR_API bool AddEmitterContextMenuActions(FMenuBuilder& MenuBuilder, const TSharedPtr<FNiagaraEmitterHandleViewModel>& EmitterHandleViewModel);
@@ -461,8 +462,21 @@ namespace FNiagaraParameterUtilities
 
 	NIAGARAEDITOR_API bool TestCanRenameWithMessage(FName ParameterName, FText& OutMessage);
 
+	/** A simple utility function to generate a parameter widget. */
 	NIAGARAEDITOR_API TSharedRef<SWidget> GetParameterWidget(FNiagaraVariable Variable, bool bAddTypeIcon, bool bShowValue);
+
+	struct NIAGARAEDITOR_API FNiagaraParameterWidgetOptions
+	{
+		bool bAddTypeIcon = false;
+		bool bShowValue = false;
+		bool bShowVisibilityConditionIcon = false;
+		bool bShowEditConditionIcon = false;
+		bool bShowAdvanced = false;
+	};
 	
+	/** A more advanced utility function that also makes use of parameter metadata to generate a parameter widget */
+	NIAGARAEDITOR_API TSharedRef<SWidget> GetParameterWidget(FNiagaraVariable Variable, FNiagaraVariableMetaData MetaData, FNiagaraParameterWidgetOptions Options);
+
 	/** Creates a tooltip based on a parameter. Also shows the value, if allocated and enabled. */
 	NIAGARAEDITOR_API TSharedRef<SToolTip> GetTooltipWidget(FNiagaraVariable Variable, bool bShowValue = true, TSharedPtr<SWidget> AdditionalVerticalWidget = nullptr,  TSharedPtr<SWidget> AdditionalHorizontalWidget = nullptr);
 
diff --git a/Engine/Plugins/FX/Niagara/Source/NiagaraEditor/Public/NiagaraSystemEditorData.h b/Engine/Plugins/FX/Niagara/Source/NiagaraEditor/Public/NiagaraSystemEditorData.h
--- a/Engine/Plugins/FX/Niagara/Source/NiagaraEditor/Public/NiagaraSystemEditorData.h
+++ b/Engine/Plugins/FX/Niagara/Source/NiagaraEditor/Public/NiagaraSystemEditorData.h
@@ -98,7 +98,7 @@ public:
 	FOnUserParameterScriptVariablesSynced& OnUserParameterScriptVariablesSynced() { return OnUserParameterScriptVariablesSyncedDelegate; }
 	
 	UNiagaraScriptVariable* FindOrAddUserScriptVariable(FNiagaraVariable UserParameter, UNiagaraSystem& System);
-	UNiagaraScriptVariable* FindUserScriptVariable(FGuid UserParameterGuid);
+	const UNiagaraScriptVariable* FindUserScriptVariable(FGuid UserParameterGuid) const;
 	bool RenameUserScriptVariable(FNiagaraVariable OldVariable, FName NewName);
 	bool RemoveUserScriptVariable(FNiagaraVariable Variable);
 
