Summary View 2.0
Base Hierarchy Editor Changes (also affects User Parameters):
- API Changes (mostly streamlining)
- Editability as a native feature. View models can override 'IsEditableByUser' which will make a lock icon appear on the row with a tooltip specified by the view model of the item.
- Reworked user messaging system
---> CanDropInternal now produces an optional message that is used to surface information to the user
---> Editability uses the same approach
- Moved 'CanHaveChildren' logic from data to view model
- Item view models now know if they are part of the hierarchy or part of the source. This allows us to modify behavior (modules show inputs in source view, but not in hierarchy)
- Search box & search API for view models
- Section support for sources panel (used in Summary View 2.0)
- 'Find in Hierarchy' context action for items already embedded in the hierarchy

Summary View Hierarchy Editor:
- Shows modules & renderers as source items/categories
- Parent emitter attributes can not be edited
- Override Display Name & tooltip property for inputs. Stack uses these if specified.
- Sections for Emitter/Particle Spawn/Update, Sim Stages, Events, Renderers
- Live update in the stack
- Drag & Dropping entire objects is possible (entire module, renderer, sim stage etc.)
- Bunch of QoL improvements

Merge process for summary view:
- Support for added items in children only

Post Load for existing emitters:
- data is transferred from summary view 1.0 data to 2.0 data
- Child Inputs are transferred explicitly as new setup uses explicit children

Stack changes:
- Refactored the stack category into a base class and some child classes
- Refactored the input collection into a 'property collection' base class and more child classes.
- Added new summary view support API that lets stack entries determine individually if they can be added to stack entries or not
- Added 'Show in Summary View' context action
- Added optional  'Owner Guid' functionality to stack objects & property rows to determine if a single property can be added to summary view

Some new utility functions and changes
Some new small-ish widgets

diff --git a/Engine/Plugins/FX/Niagara/Source/Niagara/Classes/NiagaraEmitter.h b/Engine/Plugins/FX/Niagara/Source/Niagara/Classes/NiagaraEmitter.h
--- a/Engine/Plugins/FX/Niagara/Source/Niagara/Classes/NiagaraEmitter.h
+++ b/Engine/Plugins/FX/Niagara/Source/Niagara/Classes/NiagaraEmitter.h
@@ -555,6 +555,8 @@ public:
 #if WITH_EDITOR
 	DECLARE_MULTICAST_DELEGATE(FOnPropertiesChanged);
 	DECLARE_MULTICAST_DELEGATE(FOnRenderersChanged);
+	DECLARE_MULTICAST_DELEGATE(FOnSimStagesChanged);
+	DECLARE_MULTICAST_DELEGATE(FOnEventHandlersChanged);
 	DECLARE_MULTICAST_DELEGATE_OneParam(FOnEmitterCompiled, FVersionedNiagaraEmitter);
 
 	struct NIAGARA_API PrivateMemberNames
@@ -582,6 +584,8 @@ public:
 	virtual void PostEditChangeVersionedProperty(FPropertyChangedEvent& PropertyChangedEvent, const FGuid& Version);
 	NIAGARA_API FOnPropertiesChanged& OnPropertiesChanged();
 	NIAGARA_API FOnRenderersChanged& OnRenderersChanged();
+	NIAGARA_API FOnSimStagesChanged& OnSimStagesChanged();
+	NIAGARA_API FOnSimStagesChanged& OnEventHandlersChanged();
 	/** Helper method for when a rename has been detected within the graph. Covers renaming the internal renderer bindings.*/
 	NIAGARA_API void HandleVariableRenamed(const FNiagaraVariable& InOldVariable, const FNiagaraVariable& InNewVariable, bool bUpdateContexts, FGuid EmitterVersion);
 	/** Helper method for when a rename has been detected within the graph. Covers resetting the internal renderer bindings.*/
@@ -980,6 +984,8 @@ private:
 #if WITH_EDITOR
 	FOnPropertiesChanged OnPropertiesChangedDelegate;
 	FOnRenderersChanged OnRenderersChangedDelegate;
+	FOnSimStagesChanged OnSimStagesChangedDelegate;
+	FOnEventHandlersChanged OnEventHandlersChangedDelegate;
 #endif
 
 	void GenerateStatID()const;
diff --git a/Engine/Plugins/FX/Niagara/Source/Niagara/Private/NiagaraEmitter.cpp b/Engine/Plugins/FX/Niagara/Source/Niagara/Private/NiagaraEmitter.cpp
--- a/Engine/Plugins/FX/Niagara/Source/Niagara/Private/NiagaraEmitter.cpp
+++ b/Engine/Plugins/FX/Niagara/Source/Niagara/Private/NiagaraEmitter.cpp
@@ -777,6 +777,7 @@ void FVersionedNiagaraEmitterData::PostLoad(UNiagaraEmitter& Emitter, int32 Niag
 	if (EditorData != nullptr)
 	{
 		EditorData->ConditionalPostLoad();
+		EditorData->PostLoadFromOwner(&Emitter);
 	}
 	
 	if (!GPUComputeScript)
@@ -1311,6 +1312,16 @@ UNiagaraEmitter::FOnPropertiesChanged& UNiagaraEmitter::OnRenderersChanged()
 	return OnRenderersChangedDelegate;
 }
 
+UNiagaraEmitter::FOnSimStagesChanged& UNiagaraEmitter::OnSimStagesChanged()
+{
+	return OnSimStagesChangedDelegate;
+}
+
+UNiagaraEmitter::FOnSimStagesChanged& UNiagaraEmitter::OnEventHandlersChanged()
+{
+	return OnEventHandlersChangedDelegate;
+}
+
 void UNiagaraEmitter::HandleVariableRenamed(const FNiagaraVariable& InOldVariable, const FNiagaraVariable& InNewVariable, bool bUpdateContexts, FGuid EmitterVersion)
 {
 	// Rename the variable if it is in use by any renderer properties
@@ -3023,6 +3034,7 @@ void UNiagaraEmitter::AddEventHandler(FNiagaraEventScriptProperties EventHandler
 	EventHandler.Script->RapidIterationParameters.AddOnChangedHandler(
 		FNiagaraParameterStore::FOnChanged::FDelegate::CreateUObject(this, &UNiagaraEmitter::ScriptRapidIterationParameterChanged));
 	UpdateChangeId(TEXT("Event handler added"));
+	OnEventHandlersChangedDelegate.Broadcast();
 #endif
 }
 
@@ -3040,6 +3052,7 @@ void UNiagaraEmitter::RemoveEventHandlerByUsageId(FGuid EventHandlerUsageId, FGu
 #endif
 	EventHandlerScriptProps.RemoveAll(FindEventHandlerById);
 #if WITH_EDITOR
+	OnEventHandlersChangedDelegate.Broadcast();
 	UpdateChangeId(TEXT("Event handler removed"));
 #endif
 }
@@ -3060,6 +3073,7 @@ void UNiagaraEmitter::AddSimulationStage(UNiagaraSimulationStageBase* Simulation
 	SimulationStage->Script->RapidIterationParameters.AddOnChangedHandler(
 		FNiagaraParameterStore::FOnChanged::FDelegate::CreateUObject(this, &UNiagaraEmitter::ScriptRapidIterationParameterChanged));
 	UpdateChangeId(TEXT("Shader stage added"));
+	OnSimStagesChangedDelegate.Broadcast();
 #endif
 }
 
@@ -3073,6 +3087,7 @@ void UNiagaraEmitter::RemoveSimulationStage(UNiagaraSimulationStageBase* Simulat
 		SimulationStage->OnChanged().RemoveAll(this);
 		SimulationStage->Script->RapidIterationParameters.RemoveAllOnChangedHandlers(this);
 		UpdateChangeId(TEXT("Simulation stage removed"));
+		OnSimStagesChangedDelegate.Broadcast();
 	}
 #endif
 }
@@ -3092,6 +3107,7 @@ void UNiagaraEmitter::MoveSimulationStageToIndex(UNiagaraSimulationStageBase* Si
 		SimulationStages.Insert(SimulationStageToMove, AdjustedTargetIndex);
 #if WITH_EDITOR
 		UpdateChangeId("Simulation stage moved.");
+		OnSimStagesChangedDelegate.Broadcast();
 #endif
 	}
 }
diff --git a/Engine/Plugins/FX/Niagara/Source/Niagara/Public/NiagaraStackSection.h b/Engine/Plugins/FX/Niagara/Source/Niagara/Public/NiagaraStackSection.h
--- a/Engine/Plugins/FX/Niagara/Source/Niagara/Public/NiagaraStackSection.h
+++ b/Engine/Plugins/FX/Niagara/Source/Niagara/Public/NiagaraStackSection.h
@@ -27,6 +27,9 @@ struct FNiagaraStackSection
 	UPROPERTY(EditAnywhere, Category=Section)
 	TArray<FText> Categories;
 
+	UPROPERTY()
+	FText Tooltip;
+	
 	UPROPERTY(EditAnywhere, Category=Section)
 	bool bEnabled;
 };
\ No newline at end of file
diff --git a/Engine/Plugins/FX/Niagara/Source/Niagara/Public/NiagaraTypes.h b/Engine/Plugins/FX/Niagara/Source/Niagara/Public/NiagaraTypes.h
--- a/Engine/Plugins/FX/Niagara/Source/Niagara/Public/NiagaraTypes.h
+++ b/Engine/Plugins/FX/Niagara/Source/Niagara/Public/NiagaraTypes.h
@@ -683,7 +683,7 @@ struct NIAGARA_API FNiagaraVariableMetaData
 		, StaticSwitchDefaultValue_DEPRECATED(0)
 	{};
 
-	UPROPERTY(EditAnywhere, Category = "Variable", meta = (MultiLine = true, SkipForCompileHash = "true"))
+	UPROPERTY(EditAnywhere, Category = "Variable", DisplayName="Tooltip", meta = (MultiLine = true, SkipForCompileHash = "true"))
 	FText Description;
 
 	UPROPERTY(EditAnywhere, Category = "Variable", meta = (SkipForCompileHash = "true"))
diff --git a/Engine/Plugins/FX/Niagara/Source/NiagaraCore/Private/NiagaraMergeable.cpp b/Engine/Plugins/FX/Niagara/Source/NiagaraCore/Private/NiagaraMergeable.cpp
--- a/Engine/Plugins/FX/Niagara/Source/NiagaraCore/Private/NiagaraMergeable.cpp
+++ b/Engine/Plugins/FX/Niagara/Source/NiagaraCore/Private/NiagaraMergeable.cpp
@@ -49,7 +49,7 @@ UNiagaraMergeable::FOnChanged& UNiagaraMergeable::OnChanged()
 	return OnChangedDelegate;
 }
 
-FGuid UNiagaraMergeable::GetMergeId()
+FGuid UNiagaraMergeable::GetMergeId() const
 {
 	return MergeId;
 }
diff --git a/Engine/Plugins/FX/Niagara/Source/NiagaraCore/Public/NiagaraMergeable.h b/Engine/Plugins/FX/Niagara/Source/NiagaraCore/Public/NiagaraMergeable.h
--- a/Engine/Plugins/FX/Niagara/Source/NiagaraCore/Public/NiagaraMergeable.h
+++ b/Engine/Plugins/FX/Niagara/Source/NiagaraCore/Public/NiagaraMergeable.h
@@ -24,7 +24,7 @@ public:
 
 	FOnChanged& OnChanged();
 
-	FGuid GetMergeId();
+	FGuid GetMergeId() const;
 
 protected:
 	UNiagaraMergeable* StaticDuplicateWithNewMergeIdInternal(UObject* InOuter) const;
diff --git a/Engine/Plugins/FX/Niagara/Source/NiagaraEditor/Private/NiagaraEditorStyle.cpp b/Engine/Plugins/FX/Niagara/Source/NiagaraEditor/Private/NiagaraEditorStyle.cpp
--- a/Engine/Plugins/FX/Niagara/Source/NiagaraEditor/Private/NiagaraEditorStyle.cpp
+++ b/Engine/Plugins/FX/Niagara/Source/NiagaraEditor/Private/NiagaraEditorStyle.cpp
@@ -388,6 +388,19 @@ void FNiagaraEditorStyle::InitIcons()
 	Set("NiagaraEditor.Message.CustomNote", new IMAGE_BRUSH("Icons/icon_custom_note_16x", Icon16x16));
 	Set("NiagaraEditor.Module.DynamicInput", new IMAGE_BRUSH_SVG("Icons/DynamicInput", Icon16x16));
 	Set("NiagaraEditor.Module.TypeIconPill", new IMAGE_BRUSH_SVG("Icons/TypeIconPill", FVector2D(5.0f, 16.0f)));
+
+	FLinearColor EmitterIconColor(FColor(126, 87, 67));
+	FLinearColor ParticleIconColor(FColor(87, 107, 61));
+	FLinearColor RendererIconColor(FColor(134, 80, 80));
+	Set("NiagaraEditor.Emitter.SpawnIcon", new IMAGE_BRUSH("Icons/Spawn", Icon12x12, EmitterIconColor));
+	Set("NiagaraEditor.Emitter.UpdateIcon", new IMAGE_BRUSH("Icons/Update", Icon12x12, EmitterIconColor));
+	Set("NiagaraEditor.Particle.SpawnIcon", new IMAGE_BRUSH("Icons/Spawn", Icon12x12, ParticleIconColor));
+	Set("NiagaraEditor.Particle.UpdateIcon", new IMAGE_BRUSH("Icons/Update", Icon12x12, ParticleIconColor));
+	Set("NiagaraEditor.EventIcon", new IMAGE_BRUSH("Icons/Event", Icon12x12, ParticleIconColor));
+	Set("NiagaraEditor.SimulationStageIcon", new IMAGE_BRUSH("Icons/SimulationStage", Icon12x12, ParticleIconColor));
+	Set("NiagaraEditor.RenderIcon", new IMAGE_BRUSH("Icons/Render", Icon12x12, RendererIconColor));
+	
+	Set("NiagaraEditor.HierarchyEditor.RootDropIcon", new IMAGE_BRUSH_SVG("Icons/caret-down", Icon20x20));
 	
 	FSlateBrush* SlateBrush = new FSlateBrush();
 	*SlateBrush = FAppStyle::Get().GetWidgetStyle<FEditableTextBoxStyle>("NormalEditableTextBox").BackgroundImageNormal;
@@ -395,6 +408,21 @@ void FNiagaraEditorStyle::InitIcons()
 	Set("NiagaraEditor.Module.InputTypeBorder", SlateBrush);
 }
 
+void FNiagaraEditorStyle::InitTextStyles()
+{
+	const FTextBlockStyle NormalText = FAppStyle::Get().GetWidgetStyle<FTextBlockStyle>("NormalText");
+
+	FInlineEditableTextBlockStyle HierarchyEditorCategoryStyle = FAppStyle::GetWidgetStyle<FInlineEditableTextBlockStyle>("InlineEditableTextBlockStyle");
+	HierarchyEditorCategoryStyle.TextStyle = FAppStyle::Get().GetWidgetStyle<FTextBlockStyle>("DetailsView.CategoryTextStyle");
+	HierarchyEditorCategoryStyle.TextStyle.SetFontSize(10.f);
+	Set("NiagaraEditor.HierarchyEditor.CategoryTextStyle", HierarchyEditorCategoryStyle);
+
+	FSlateFontInfo SummaryViewInputOverrideFontInfo = DEFAULT_FONT("Italic", 9);
+	FTextBlockStyle SummaryViewInputOverrideTextBlockStyle = FTextBlockStyle(NormalText)
+		.SetFont(SummaryViewInputOverrideFontInfo);
+	Set("NiagaraEditor.HierarchyEditor.SummaryView.ModuleInputNameOverride", SummaryViewInputOverrideTextBlockStyle);
+}
+
 void FNiagaraEditorStyle::InitOverview()
 {
 	// Overview debug icons
@@ -803,6 +831,7 @@ FNiagaraEditorStyle::FNiagaraEditorStyle() : FSlateStyleSet("NiagaraEditorStyle"
 	InitToolbarIcons();
 	InitTabIcons();
 	InitIcons();
+	InitTextStyles();
 	InitOverview();
 	InitEmitterDetails();
 	InitAssetColors();
diff --git a/Engine/Plugins/FX/Niagara/Source/NiagaraEditor/Private/NiagaraEditorUtilities.cpp b/Engine/Plugins/FX/Niagara/Source/NiagaraEditor/Private/NiagaraEditorUtilities.cpp
--- a/Engine/Plugins/FX/Niagara/Source/NiagaraEditor/Private/NiagaraEditorUtilities.cpp
+++ b/Engine/Plugins/FX/Niagara/Source/NiagaraEditor/Private/NiagaraEditorUtilities.cpp
@@ -4164,20 +4164,55 @@ TSharedRef<SWidget> FNiagaraParameterUtilities::GetParameterWidget(FNiagaraVaria
 			.ToolTipText(TooltipText)
 		];
 	}
+
+	TSharedPtr<STextBlock> NameOverride;
+	if(Options.NameOverride.IsSet())
+	{
+		NameOverride = SNew(STextBlock)
+		.Text(Options.NameOverride.GetValue())
+		.TextStyle(&FNiagaraEditorStyle::Get().GetWidgetStyle<FTextBlockStyle>("NiagaraEditor.HierarchyEditor.SummaryView.ModuleInputNameOverride"));
+
+		if(Options.NameOverrideVisibility.IsSet())
+		{
+			NameOverride->SetVisibility(Options.NameOverrideVisibility.GetValue());
+		}
+
+		if(Options.NameOverrideTooltip.IsSet())
+		{
+			NameOverride->SetToolTipText(Options.NameOverrideTooltip.GetValue());
+		}
+	}
 	
-	return SNew(SHorizontalBox)
+	TSharedRef<SHorizontalBox> ResultWidget = SNew(SHorizontalBox)
 	+ SHorizontalBox::Slot()
 	.AutoWidth()
 	.VAlign(VAlign_Center)
 	[
-		GetParameterWidget(Variable, Options.bAddTypeIcon, Options.bShowValue)
-	]
-	+ SHorizontalBox::Slot()
-	.HAlign(HAlign_Right)
-	.VAlign(VAlign_Center)
-	[
-		IconBox
-	];	
+		FNiagaraParameterUtilities::GetParameterWidget(Variable, Options.bAddTypeIcon, Options.bShowValue)       	
+	];
+
+	if(NameOverride.IsValid())
+	{
+		ResultWidget->AddSlot()
+		.AutoWidth()
+		.VAlign(VAlign_Center)
+		.Padding(4.f, 0.f)
+		[
+			NameOverride.ToSharedRef()
+		];
+	}
+
+	if(IconBox->NumSlots() > 0)
+	{
+		ResultWidget->AddSlot()
+		.HAlign(HAlign_Right)
+		.VAlign(VAlign_Center)
+		[
+			IconBox
+		];
+	}
+
+	return ResultWidget;
 }
 
 TSharedRef<SToolTip> FNiagaraParameterUtilities::GetTooltipWidget(FNiagaraVariable Variable, bool bShowValue, TSharedPtr<SWidget> AdditionalVerticalWidget,  TSharedPtr<SWidget> AdditionalHorizontalWidget)
diff --git a/Engine/Plugins/FX/Niagara/Source/NiagaraEditor/Private/NiagaraEmitterEditorData.cpp b/Engine/Plugins/FX/Niagara/Source/NiagaraEditor/Private/NiagaraEmitterEditorData.cpp
--- a/Engine/Plugins/FX/Niagara/Source/NiagaraEditor/Private/NiagaraEmitterEditorData.cpp
+++ b/Engine/Plugins/FX/Niagara/Source/NiagaraEditor/Private/NiagaraEmitterEditorData.cpp
@@ -2,18 +2,25 @@
 
 #include "NiagaraEmitterEditorData.h"
 
+#include "NiagaraNodeFunctionCall.h"
+#include "NiagaraScriptVariable.h"
 #include "EdGraph/EdGraph.h"
 #include "NiagaraStackEditorData.h"
 #include "ScopedTransaction.h"
+#include "NiagaraNodeAssignment.h"
+#include "ViewModels/HierarchyEditor/NiagaraSummaryViewViewModel.h"
 
 #include UE_INLINE_GENERATED_CPP_BY_NAME(NiagaraEmitterEditorData)
 
-const FName UNiagaraEmitterEditorData::PrivateMemberNames::SummarySections = GET_MEMBER_NAME_CHECKED(UNiagaraEmitterEditorData, SummarySections);
+#define LOCTEXT_NAMESPACE "NiagaraEmitterEditorData"
+
+const FName UNiagaraEmitterEditorData::PrivateMemberNames::SummarySections = GET_MEMBER_NAME_CHECKED(UNiagaraEmitterEditorData, SummarySections_DEPRECATED);
 
 UNiagaraEmitterEditorData::UNiagaraEmitterEditorData(const FObjectInitializer& ObjectInitializer)
 {
 	StackEditorData = ObjectInitializer.CreateDefaultSubobject<UNiagaraStackEditorData>(this, TEXT("StackEditorData"));
-
+	SummaryViewRoot = CreateDefaultSubobject<UNiagaraHierarchyRoot>(TEXT("SummaryViewRoot"));
+	
 	if (StackEditorData != nullptr)
 	{
 		StackEditorData->OnPersistentDataChanged().AddUObject(this, &UNiagaraEmitterEditorData::StackEditorDataChanged);
@@ -23,6 +30,330 @@ UNiagaraEmitterEditorData::UNiagaraEmitterEditorData(const FObjectInitializer& O
 	PlaybackRangeMax = 10;
 }
 
+void UNiagaraEmitterEditorData::PostLoad_TransferSummaryDataToNewFormat()
+{
+	// ATTENTION
+	// It is important that any element that is created during post load has a stable identifier.
+	// This is _required_ as the merge process or just loading multiple times without saving will end up in calling this multiple times,
+	// creating new objects that supposed to represent the same object, but with different IDs.
+	// Summary View 1.0 only supported Module Inputs, Categories & Sections
+	// - Module inputs have a stable ID by nature (function call node ID + module input ID)
+	// - Sections are not stable by default (random guid when created). Here we add a stable identity using the section identifier.
+	// - Categories by default are unstable as well. So we need to find a stable ID
+	// ---> This is only required for post load, not in new user-created items, as those IDs will then remain stable across merges.
+	
+	if(SummarySections_DEPRECATED.Num() == 0 && SummaryViewFunctionInputMetadata_DEPRECATED.Num() == 0)
+	{
+		return;
+	}
+	
+	// since we lack graph context during post load, we do this workaround to find the emitter graph
+	UNiagaraEmitter* Emitter = GetTypedOuter<UNiagaraEmitter>();
+	TArray<FNiagaraAssetVersion> AssetVersions = Emitter->GetAllAvailableVersions();
+	UNiagaraGraph* EmitterGraph = nullptr;
+	for(FNiagaraAssetVersion& AssetVersion : AssetVersions)
+	{
+		FVersionedNiagaraEmitterData* EmitterData = Emitter->GetEmitterData(AssetVersion.VersionGuid);
+		if(EmitterData->GetEditorData() == this)
+		{
+			EmitterGraph = Cast<UNiagaraScriptSource>(EmitterData->GraphSource)->NodeGraph;
+		}
+	}
+
+	auto GenerateStableIdentityForCategories = [](UNiagaraHierarchyCategory* Category) -> TArray<FName>
+	{
+		TArray<FName> CategoryNameChain;
+		
+		const UNiagaraHierarchyCategory* OuterCategory = Category;
+
+		CategoryNameChain.Add(Category->GetCategoryName());
+
+		for (; OuterCategory != nullptr; OuterCategory = OuterCategory->GetTypedOuter<UNiagaraHierarchyCategory>() )
+		{
+			if(OuterCategory != nullptr)
+			{
+				CategoryNameChain.Add(OuterCategory->GetCategoryName());
+			}
+		}
+	
+		return CategoryNameChain;
+	};
+	
+	TArray<UNiagaraNodeFunctionCall*> ModuleNodes = FNiagaraStackGraphUtilities::GetAllModuleNodes(EmitterGraph);
+	
+	TMap<FName, UNiagaraHierarchySection*> Sections;
+	TMap<FName, UNiagaraHierarchyCategory*> IdentityCategoryMap;
+	TMap<FName, FName> SectionCategoryMapping;
+	if(SummarySections_DEPRECATED.Num() > 0)
+	{
+		for(FNiagaraStackSection& StackSection : SummarySections_DEPRECATED)
+		{
+			FText SectionDisplayName = StackSection.SectionDisplayName.IsEmptyOrWhitespace() ? FText::FromName(StackSection.SectionIdentifier) : StackSection.SectionDisplayName;
+			UNiagaraHierarchySection* Section = SummaryViewRoot->AddSection(SectionDisplayName);
+			FNiagaraHierarchyIdentity SectionIdentity;
+			SectionIdentity.Names.Add(StackSection.SectionIdentifier);
+			Section->SetIdentity(SectionIdentity);
+			
+			Sections.Add(StackSection.SectionIdentifier, Section);
+			
+			for(const FText& CategoryIdentity : StackSection.Categories)
+			{
+				FName CategoryName = FName(CategoryIdentity.ToString());
+				UNiagaraHierarchyCategory* HierarchyCategory = SummaryViewRoot->AddChild<UNiagaraHierarchyCategory>();
+				HierarchyCategory->SetSection(Section);
+				HierarchyCategory->SetCategoryName(CategoryName);
+				IdentityCategoryMap.Add(CategoryName, HierarchyCategory);
+				SectionCategoryMapping.Add(CategoryName, Section->GetSectionName());
+			}
+		}
+
+		SummarySections_DEPRECATED.Empty();
+	}
+
+	if(SummaryViewFunctionInputMetadata_DEPRECATED.Num() > 0)
+	{
+		// we have a custom sort order cached explicitly for summary view. This takes precedence.
+		TMap<UNiagaraHierarchyItemBase*, int32> SummarySortOrder;
+		// the original metadata defined in the script asset for a given input.
+		TMap<FNiagaraHierarchyIdentity, FNiagaraVariableMetaData> AssetVariableMetadata;
+
+		// since the order in the new summary view data is completely explicit (or managed implicitly by certain classes)
+		// we have to sort inputs during transfer first by asset order, then summary sort order.
+		for(auto It = SummaryViewFunctionInputMetadata_DEPRECATED.CreateConstIterator(); It; ++It)
+		{
+			// The input identity we are constructing from the function call node guid & the variable guid
+			FNiagaraHierarchyIdentity ParentInputIdentity;
+			TArray<FNiagaraHierarchyIdentity> ChildInputIdentities;
+
+			FName OriginalCategoryName;
+
+			UNiagaraNodeFunctionCall** MatchingFunctionCallNode = ModuleNodes.FindByPredicate([&](UNiagaraNodeFunctionCall* Candidate)
+				{
+					return It.Key().GetFunctionGuid() == Candidate->NodeGuid;
+				});
+
+			TObjectPtr<UNiagaraScriptVariable> ScriptVariable = nullptr;
+			// there should typically be a node found here. In case it's not, it seems the node was deleted but the summary entry kept. We skip this one.
+			if(MatchingFunctionCallNode != nullptr)
+			{
+				UNiagaraNodeFunctionCall* FunctionCall = *MatchingFunctionCallNode;
+				ParentInputIdentity.Guids.Add(FunctionCall->NodeGuid);
+				
+				if(UNiagaraGraph* Graph = FunctionCall->GetCalledGraph())
+				{
+					bool bIsModuleInput = It.Key().GetInputGuid().IsValid();			
+
+					// for module inputs we look up the variable guid in the asset graph
+					if(bIsModuleInput)
+					{
+						TArray<TObjectPtr<UNiagaraScriptVariable>> ScriptVariables;
+						Graph->GetAllMetaData().GenerateValueArray(ScriptVariables);
+						TObjectPtr<UNiagaraScriptVariable>* MatchingScriptVariable = ScriptVariables.FindByPredicate([&](TObjectPtr<UNiagaraScriptVariable> Candidate)
+						{
+							return Candidate->Metadata.GetVariableGuid() == It.Key().GetInputGuid();
+						});
+
+						// we should find a matching script variable. If there is none, the variable was probably deleted but the summary entry kept. We skip this one.
+						if(MatchingScriptVariable != nullptr)
+						{
+							ScriptVariable = *MatchingScriptVariable;
+							OriginalCategoryName = FName(ScriptVariable->Metadata.CategoryName.ToString());
+							ParentInputIdentity.Guids.Add(ScriptVariable->Metadata.GetVariableGuid());
+							AssetVariableMetadata.Add(ParentInputIdentity, ScriptVariable->Metadata);
+
+							// now that we found the matching parent input, we look for child inputs
+							for(UNiagaraScriptVariable* CandidateChildScriptVariable : ScriptVariables)
+							{
+								if(CandidateChildScriptVariable != nullptr && !CandidateChildScriptVariable->Metadata.ParentAttribute.IsNone())
+								{
+									if(CandidateChildScriptVariable->Metadata.ParentAttribute.IsEqual(ScriptVariable->Variable.GetName()))
+									{
+										FNiagaraHierarchyIdentity ChildIdentity;
+										ChildIdentity.Guids.Add(FunctionCall->NodeGuid);
+										ChildIdentity.Guids.Add(CandidateChildScriptVariable->Metadata.GetVariableGuid());
+										ChildInputIdentities.Add(ChildIdentity);
+										AssetVariableMetadata.Add(ChildIdentity, CandidateChildScriptVariable->Metadata);
+									}
+								}
+							}
+						}
+						else
+						{
+							continue;
+						}
+						
+					}
+					// for assignment inputs we have to match via name as the data doesn't contain a stable guid
+					else
+					{
+						TArray<FNiagaraVariable> Variables;
+						Graph->GetAllVariables(Variables);
+						
+						FNiagaraVariableBase* MatchingVariable = Variables.FindByPredicate([InputName = It.Key().GetInputName()](const FNiagaraVariable& Candidate)
+						{
+							return Candidate.GetName() == InputName;
+						});
+
+						// if we didn't find a matching variable we assume the summary entry is out of date and we skip it
+						if(MatchingVariable != nullptr)
+						{
+							ParentInputIdentity.Names.Add(It.Key().GetInputName());
+							OriginalCategoryName = FName(TEXT("Uncategorized"));
+							// we add new empty metadata here as assignment node targets don't have metadata. This will assume a base default sort order of 0.
+							AssetVariableMetadata.Add(ParentInputIdentity, FNiagaraVariableMetaData());
+						}
+						else
+						{
+							continue;
+						}
+					}					
+				}				
+			}
+			else
+			{
+				continue;
+			}
+
+			bool bIsCategorySpecified = !It.Value().Category.IsNone();
+			FName CategoryName;
+			if(bIsCategorySpecified)
+			{
+				CategoryName = It.Value().Category;
+			}
+			// if the category is not specified, we use the default category that is defined in the graph
+			else
+			{
+				CategoryName = OriginalCategoryName;
+			}
+
+			UNiagaraHierarchyCategory* Category = nullptr;
+			if(IdentityCategoryMap.Contains(CategoryName))
+			{
+				Category = IdentityCategoryMap[CategoryName];
+			}
+			else
+			{
+				Category = SummaryViewRoot->AddChild<UNiagaraHierarchyCategory>();
+				Category->SetCategoryName(CategoryName);
+				IdentityCategoryMap.Add(Category->GetCategoryName(), Category);
+			}
+
+			if(ParentInputIdentity.Guids.Num() == 2)
+			{
+				UNiagaraHierarchyModuleInput* SummaryModuleInput = Category->AddChild<UNiagaraHierarchyModuleInput>();		
+				SummaryModuleInput->SetIdentity(ParentInputIdentity);
+				
+				if(!It.Value().DisplayName.IsNone() && ScriptVariable != nullptr)
+				{
+					SummaryModuleInput->SetDisplayNameOverride(FText::FromName(It.Value().DisplayName));
+				}
+				SummarySortOrder.Add(SummaryModuleInput, It.Value().SortIndex);
+
+				for(const FNiagaraHierarchyIdentity& ChildInputIdentity : ChildInputIdentities)
+				{
+					UNiagaraHierarchyModuleInput* SummaryModuleChildInput = SummaryModuleInput->AddChild<UNiagaraHierarchyModuleInput>();		
+					SummaryModuleChildInput->SetIdentity(ChildInputIdentity);
+					SummarySortOrder.Add(SummaryModuleChildInput, INDEX_NONE);
+				}
+			}
+			else
+			{
+				UNiagaraHierarchyAssignmentInput* SummaryAssignmentInput = Category->AddChild<UNiagaraHierarchyAssignmentInput>();		
+				SummaryAssignmentInput->SetIdentity(ParentInputIdentity);
+				SummarySortOrder.Add(SummaryAssignmentInput, It.Value().SortIndex);
+			}
+		}
+
+		// now that we have created all inputs, categories and sections, we fix up the order
+		TArray<UNiagaraHierarchyCategory*> TopLevelCategories;
+		SummaryViewRoot->GetChildrenOfType(TopLevelCategories, false);
+
+		// we sort by summary sort order if specified, and use the asset order as fallback.
+		// advanced items always get put at the bottom
+		auto SortInputsPredicate = ([&](UNiagaraHierarchyItemBase& ItemA, UNiagaraHierarchyItemBase& ItemB)
+		{
+			const FNiagaraVariableMetaData& MetaDataA = AssetVariableMetadata[ItemA.GetPersistentIdentity()];
+			const FNiagaraVariableMetaData& MetaDataB = AssetVariableMetadata[ItemB.GetPersistentIdentity()];
+
+			if(MetaDataA.bAdvancedDisplay && !MetaDataB.bAdvancedDisplay)
+			{
+				return false;
+			}
+			else if(!MetaDataA.bAdvancedDisplay && MetaDataB.bAdvancedDisplay)
+			{
+				return true;
+			}
+			
+			int32 SortOrderA = SummarySortOrder[&ItemA];
+			if(SortOrderA == INDEX_NONE)
+			{
+				SortOrderA = MetaDataA.EditorSortPriority;
+			}
+
+			int32 SortOrderB = SummarySortOrder[&ItemB];
+			if(SortOrderB == INDEX_NONE)
+			{
+				SortOrderB = MetaDataB.EditorSortPriority;
+			}
+				
+			return SortOrderA < SortOrderB;
+		});
+
+		TMap<UNiagaraHierarchyItemBase*, int32> CategoryMinimumInputSortOrderMap;
+		for(UNiagaraHierarchyCategory* Category : TopLevelCategories)
+		{
+			Category->SortChildren(SortInputsPredicate, false);
+
+			TArray<UNiagaraHierarchyModuleInput*> Inputs;
+			Category->GetChildrenOfType<UNiagaraHierarchyModuleInput>(Inputs, false);
+
+			// inputs can have child inputs one layer deep, so we sort them too
+			for(UNiagaraHierarchyModuleInput* Input : Inputs)
+			{
+				Input->SortChildren(SortInputsPredicate, false);
+			}
+
+			if(Category->GetChildren().Num() > 0)
+			{
+				// since we have sorted the children by sort priority, child 0 has the lowest sort priority number (= highest priority)
+				UNiagaraHierarchyItemBase* MinimumItem = Category->GetChildren()[0];
+				int32 MinimumCategorySortOrder = SummarySortOrder[MinimumItem];
+				if(MinimumCategorySortOrder == INDEX_NONE)
+				{
+					MinimumCategorySortOrder = AssetVariableMetadata[MinimumItem->GetPersistentIdentity()].EditorSortPriority;
+				}
+				
+				CategoryMinimumInputSortOrderMap.Add(Category, MinimumCategorySortOrder);
+			}
+			else
+			{
+				CategoryMinimumInputSortOrderMap.Add(Category, 0);
+			}
+		}
+
+		auto SortCategoriesPredicate = ([&](UNiagaraHierarchyItemBase& ItemA, UNiagaraHierarchyItemBase& ItemB)
+		{
+			return CategoryMinimumInputSortOrderMap[&ItemA] < CategoryMinimumInputSortOrderMap[&ItemB];
+		});
+
+		SummaryViewRoot->SortChildren(SortCategoriesPredicate, false);
+
+		// we now perform identity fixup for categories using the chain of category names.
+		// we need to do this so the generated identity is stable instead of random.
+		TArray<UNiagaraHierarchyCategory*> AllCategories;
+		SummaryViewRoot->GetChildrenOfType(AllCategories, true);
+		
+		for(UNiagaraHierarchyCategory* Category : AllCategories)
+		{
+			FNiagaraHierarchyIdentity StableIdentity;
+			StableIdentity.Names = GenerateStableIdentityForCategories(Category);
+			Category->SetIdentity(StableIdentity);
+		}
+
+		SummaryViewFunctionInputMetadata_DEPRECATED.Empty();
+	}
+}
+
 void UNiagaraEmitterEditorData::PostLoad()
 {
 	Super::PostLoad();
@@ -32,6 +363,16 @@ void UNiagaraEmitterEditorData::PostLoad()
 		StackEditorData->OnPersistentDataChanged().AddUObject(this, &UNiagaraEmitterEditorData::StackEditorDataChanged);
 	}
 	StackEditorData->ConditionalPostLoad();
+
+	if(SummaryViewRoot == nullptr)
+	{
+		SummaryViewRoot = NewObject<UNiagaraHierarchyRoot>(this, TEXT("SummaryViewRoot"), RF_Transactional);
+	}
+}
+
+void UNiagaraEmitterEditorData::PostLoadFromOwner(UObject* InOwner)
+{
+	PostLoad_TransferSummaryDataToNewFormat();
 }
 
 #if WITH_EDITORONLY_DATA
@@ -65,33 +406,6 @@ void UNiagaraEmitterEditorData::StackEditorDataChanged()
 	OnPersistentDataChanged().Broadcast();
 }
 
-
-
-const TMap<FFunctionInputSummaryViewKey, FFunctionInputSummaryViewMetadata>& UNiagaraEmitterEditorData::GetSummaryViewMetaDataMap() const
-{
-	return SummaryViewFunctionInputMetadata;
-}
-
-TOptional<FFunctionInputSummaryViewMetadata> UNiagaraEmitterEditorData::GetSummaryViewMetaData(const FFunctionInputSummaryViewKey& Key) const
-{
-	return SummaryViewFunctionInputMetadata.Contains(Key)? SummaryViewFunctionInputMetadata[Key] : TOptional<FFunctionInputSummaryViewMetadata>();
-}
-
-void UNiagaraEmitterEditorData::SetSummaryViewMetaData(const FFunctionInputSummaryViewKey& Key, TOptional<FFunctionInputSummaryViewMetadata> NewMetadata)
-{	
-	if (NewMetadata.IsSet() == false)
-	{
-		SummaryViewFunctionInputMetadata.Remove(Key);
-	}
-	else
-	{
-		SummaryViewFunctionInputMetadata.FindOrAdd(Key) = NewMetadata.GetValue();
-	}
-	
-	OnPersistentDataChanged().Broadcast();
-	OnSummaryViewStateChangedDelegate.Broadcast();
-}
-
 void UNiagaraEmitterEditorData::SetShowSummaryView(bool bInShouldShowSummaryView)
 {
 	bShowSummaryView = bInShouldShowSummaryView;
@@ -113,12 +427,9 @@ FSimpleMulticastDelegate& UNiagaraEmitterEditorData::OnSummaryViewStateChanged()
 	return OnSummaryViewStateChangedDelegate;
 }
 
-const TArray<FNiagaraStackSection>& UNiagaraEmitterEditorData::GetSummarySections() const
+const TArray<UNiagaraHierarchySection*>& UNiagaraEmitterEditorData::GetSummarySections() const
 {
-	return SummarySections;
+	return SummaryViewRoot->GetSectionData();
 }
 
-void UNiagaraEmitterEditorData::SetSummarySections(const TArray<FNiagaraStackSection>& InSummarySections)
-{
-	SummarySections = InSummarySections;
-}
+#undef LOCTEXT_NAMESPACE
diff --git a/Engine/Plugins/FX/Niagara/Source/NiagaraEditor/Private/NiagaraGraph.cpp b/Engine/Plugins/FX/Niagara/Source/NiagaraEditor/Private/NiagaraGraph.cpp
--- a/Engine/Plugins/FX/Niagara/Source/NiagaraEditor/Private/NiagaraGraph.cpp
+++ b/Engine/Plugins/FX/Niagara/Source/NiagaraEditor/Private/NiagaraGraph.cpp
@@ -2267,6 +2267,90 @@ UNiagaraScriptVariable* UNiagaraGraph::GetScriptVariable(FName ParameterName) co
 	return nullptr;
 }
 
+UNiagaraScriptVariable* UNiagaraGraph::GetScriptVariable(FGuid VariableGuid) const
+{
+	check(!bIsForCompilationOnly);
+
+	for (auto& VariableToScriptVariableItem : VariableToScriptVariable)
+	{
+		if(VariableToScriptVariableItem.Value->Metadata.GetVariableGuid() == VariableGuid)
+		{
+			return VariableToScriptVariableItem.Value;
+		}
+	}
+	
+	return nullptr;
+}
+
+TArray<UNiagaraScriptVariable*> UNiagaraGraph::GetChildScriptVariablesForInput(FGuid VariableGuid) const
+{
+	check(!bIsForCompilationOnly);
+
+	TArray<UNiagaraScriptVariable*> ChildrenVariables;
+	
+	TOptional<FNiagaraVariable> FoundVariable;
+	for (auto& VariableToScriptVariableItem : VariableToScriptVariable)
+	{
+		if(VariableToScriptVariableItem.Value->Metadata.GetVariableGuid() == VariableGuid)
+		{
+			FoundVariable = VariableToScriptVariableItem.Key;
+		}
+	}
+
+	if(FoundVariable.IsSet())
+	{
+		for (auto& VariableToScriptVariableItem : VariableToScriptVariable)
+		{
+			if(VariableToScriptVariableItem.Value->Metadata.ParentAttribute == FoundVariable->GetName())
+			{
+				ChildrenVariables.Add(VariableToScriptVariableItem.Value);
+			}
+		}
+	}
+
+	ChildrenVariables.Sort([&](const UNiagaraScriptVariable& VariableA, const UNiagaraScriptVariable& VariableB)
+	{
+		return VariableA.Metadata.EditorSortPriority < VariableB.Metadata.EditorSortPriority;
+	});
+	
+	return ChildrenVariables;
+}
+
+TArray<FGuid> UNiagaraGraph::GetChildScriptVariableGuidsForInput(FGuid VariableGuid) const
+{
+	check(!bIsForCompilationOnly);
+
+	TArray<FGuid> ChildrenVariableGuids;
+	TMap<FGuid, int32> SortOrderMap;
+	TOptional<FNiagaraVariable> FoundVariable;
+	for (auto& VariableToScriptVariableItem : VariableToScriptVariable)
+	{
+		if(VariableToScriptVariableItem.Value->Metadata.GetVariableGuid() == VariableGuid)
+		{
+			FoundVariable = VariableToScriptVariableItem.Key;
+		}
+	}
+
+	if(FoundVariable.IsSet())
+	{
+		for (auto& VariableToScriptVariableItem : VariableToScriptVariable)
+		{
+			if(VariableToScriptVariableItem.Value->Metadata.ParentAttribute == FoundVariable->GetName())
+			{
+				ChildrenVariableGuids.Add(VariableToScriptVariableItem.Value->Metadata.GetVariableGuid());
+				SortOrderMap.Add(VariableToScriptVariableItem.Value->Metadata.GetVariableGuid(), VariableToScriptVariableItem.Value->Metadata.EditorSortPriority);
+			}
+		}
+	}
+
+	ChildrenVariableGuids.Sort([&](const FGuid& GuidA, const FGuid& GuidB)
+	{
+		return SortOrderMap[GuidA] < SortOrderMap[GuidB];
+	});
+	
+	return ChildrenVariableGuids;
+}
+
 UNiagaraScriptVariable* UNiagaraGraph::AddParameter(const FNiagaraVariable& Parameter, bool bIsStaticSwitch /*= false*/)
 {
 	check(!bIsForCompilationOnly);
diff --git a/Engine/Plugins/FX/Niagara/Source/NiagaraEditor/Private/NiagaraScriptMergeManager.cpp b/Engine/Plugins/FX/Niagara/Source/NiagaraEditor/Private/NiagaraScriptMergeManager.cpp
--- a/Engine/Plugins/FX/Niagara/Source/NiagaraEditor/Private/NiagaraScriptMergeManager.cpp
+++ b/Engine/Plugins/FX/Niagara/Source/NiagaraEditor/Private/NiagaraScriptMergeManager.cpp
@@ -45,6 +45,7 @@
 DECLARE_CYCLE_STAT(TEXT("Niagara - ScriptMergeManager - DiffEmitters"), STAT_NiagaraEditor_ScriptMergeManager_DiffEmitters, STATGROUP_NiagaraEditor);
 DECLARE_CYCLE_STAT(TEXT("Niagara - ScriptMergeManager - MergeEmitter"), STAT_NiagaraEditor_ScriptMergeManager_MergeEmitter, STATGROUP_NiagaraEditor);
 DECLARE_CYCLE_STAT(TEXT("Niagara - ScriptMergeManager - IsModuleInputDifferentFromBase"), STAT_NiagaraEditor_ScriptMergeManager_IsModuleInputDifferentFromBase, STATGROUP_NiagaraEditor);
+DECLARE_CYCLE_STAT(TEXT("Niagara - ScriptMergeManager - DoesSummaryItemExistInBase"), STAT_NiagaraEditor_ScriptMergeManager_DoesSummaryItemExistInBase, STATGROUP_NiagaraEditor)
 
 int32 GNiagaraForceFailIfPreviouslyNotSetOnMerge = 0;
 static FAutoConsoleVariableRef CVarForceErrorIfMissingDefaultOnMergeh(
@@ -1188,10 +1189,9 @@ bool FNiagaraEmitterDiffResults::IsEmpty() const
 		AddedOtherRenderers.Num() == 0 &&
 		ModifiedBaseRenderers.Num() == 0 &&
 		ModifiedOtherRenderers.Num() == 0 &&
-		RemovedInputSummaryEntries.Num() == 0 &&
-		AddedInputSummaryEntries.Num() == 0 &&
-		ModifiedInputSummaryEntries.Num() == 0 &&
-		ModifiedOtherInputSummaryEntries.Num() == 0 &&
+		AddedSummaryEntriesInOther.Num() == 0 &&
+		AddedSummarySectionsInOther.Num() == 0 &&
+		RemovedSummaryEntriesInBase.Num() == 0 && 
 		ModifiedStackEntryDisplayNames.Num() == 0 &&
 		bScratchPadModified == false &&
 		NewShouldShowSummaryViewValue.IsSet() == false;
@@ -1945,6 +1945,20 @@ bool FNiagaraScriptMergeManager::IsModuleInputDifferentFromBase(const FVersioned
 		ScriptStackDiffResults.ModifiedOtherInputOverrides.FindByPredicate(FindInputOverrideByInputName) != nullptr;
 }
 
+bool FNiagaraScriptMergeManager::DoesSummaryItemExistInBase(const FVersionedNiagaraEmitter& Emitter, FNiagaraHierarchyIdentity Identity)
+{
+	SCOPE_CYCLE_COUNTER(STAT_NiagaraEditor_ScriptMergeManager_DoesSummaryItemExistInBase);
+
+	FVersionedNiagaraEmitter BaseEmitter = Emitter.GetEmitterData()->GetParent();
+	if(BaseEmitter.GetEmitterData() != nullptr)
+	{
+		TSharedRef<FNiagaraEmitterMergeAdapter> BaseEmitterAdapter = GetEmitterMergeAdapterUsingCache(BaseEmitter);
+		return BaseEmitterAdapter->GetEditorData()->GetSummaryRoot()->FindChildWithIdentity(Identity, true) != nullptr;
+	}
+
+	return false;
+}
+
 FNiagaraScriptMergeManager::FApplyDiffResults FNiagaraScriptMergeManager::ResetModuleInputToBase(const FVersionedNiagaraEmitter& VersionedEmitter, const FVersionedNiagaraEmitter& VersionedBaseEmitter, ENiagaraScriptUsage ScriptUsage, FGuid ScriptUsageId, FGuid ModuleId, FString InputName)
 {
 	TSharedRef<FNiagaraEmitterMergeAdapter> EmitterAdapter = GetEmitterMergeAdapterUsingCache(VersionedEmitter);
@@ -2430,85 +2444,43 @@ void FNiagaraScriptMergeManager::DiffRenderers(const TArray<TSharedRef<FNiagaraR
 
 void FNiagaraScriptMergeManager::DiffEmitterSummary(const UNiagaraEmitterEditorData* BaseEditorData, const UNiagaraEmitterEditorData* OtherEditorData, FNiagaraEmitterDiffResults& DiffResults) const
 {
-	TArray<TSharedRef<FNiagaraInputSummaryMergeAdapter>> BaseSummaryViewEntries;
-	TArray<TSharedRef<FNiagaraInputSummaryMergeAdapter>> OtherSummaryViewEntries;
-
-	if (BaseEditorData)
-	{
-		for (const auto& Entry : BaseEditorData->GetSummaryViewMetaDataMap())
-		{
-			BaseSummaryViewEntries.Add(MakeShared<FNiagaraInputSummaryMergeAdapter>(Entry.Key, Entry.Value));
-		}
-	}
-
-	if (OtherEditorData)
-	{
-		for (const auto& Entry : OtherEditorData->GetSummaryViewMetaDataMap())
-		{
-			OtherSummaryViewEntries.Add(MakeShared<FNiagaraInputSummaryMergeAdapter>(Entry.Key, Entry.Value));
-		}		
-	}
+	TArray<UNiagaraHierarchyItemBase*> AddedItemsInOther;
+	TArray<UNiagaraHierarchyItemBase*> RemovedItemsInBase;
+	TArray<UNiagaraHierarchySection*> AddedSections;
 	
-	FListDiffResults<TSharedRef<FNiagaraInputSummaryMergeAdapter>> InputListDiffResults = DiffLists<TSharedRef<FNiagaraInputSummaryMergeAdapter>, FFunctionInputSummaryViewKey>(
-		BaseSummaryViewEntries,
-		OtherSummaryViewEntries,
-		[](TSharedRef<FNiagaraInputSummaryMergeAdapter> Entry) { return Entry->GetKey(); });
+	UNiagaraHierarchyRoot* BaseRoot = BaseEditorData->GetSummaryRoot();
+	UNiagaraHierarchyRoot* OtherRoot = OtherEditorData->GetSummaryRoot();	
 
-	DiffResults.RemovedInputSummaryEntries.Append(InputListDiffResults.RemovedBaseValues);
-	DiffResults.AddedInputSummaryEntries.Append(InputListDiffResults.AddedOtherValues);	
-
-	for (const FCommonValuePair<TSharedRef<FNiagaraInputSummaryMergeAdapter>>& CommonValuePair : InputListDiffResults.CommonValuePairs)
-	{
-		if (!(CommonValuePair.BaseValue->GetValue() == CommonValuePair.OtherValue->GetValue()))
-		{
-			DiffResults.ModifiedInputSummaryEntries.Add(CommonValuePair.BaseValue);
-			DiffResults.ModifiedOtherInputSummaryEntries.Add(CommonValuePair.OtherValue);
-		}
-	}
+	TArray<UNiagaraHierarchyItemBase*> BaseItems;
+	BaseRoot->GetChildrenOfType(BaseItems, true);
 
-	TArray<FNiagaraStackSection> BaseSections = BaseEditorData != nullptr ? BaseEditorData->GetSummarySections() : TArray<FNiagaraStackSection>();
-	TArray<FNiagaraStackSection> OtherSections = OtherEditorData != nullptr ? OtherEditorData->GetSummarySections() : TArray<FNiagaraStackSection>();
-	FListDiffResults<FNiagaraStackSection> SectionDiffResults = DiffLists<FNiagaraStackSection, FName>(
-		BaseSections,
-		OtherSections,
-		[](const FNiagaraStackSection& Section) { return Section.SectionIdentifier; });
+	TArray<UNiagaraHierarchyItemBase*> OtherItems;
+	OtherRoot->GetChildrenOfType(OtherItems, true);
+	
+	FListDiffResults<UNiagaraHierarchyItemBase*> SummaryItemDiff = DiffLists<UNiagaraHierarchyItemBase*, FNiagaraHierarchyIdentity>(
+		BaseItems,
+		OtherItems,
+		[](UNiagaraHierarchyItemBase* HierarchyItem) { return HierarchyItem->GetPersistentIdentity(); });
 
-	DiffResults.RemovedBaseSummarySections.Append(SectionDiffResults.RemovedBaseValues);
-	DiffResults.AddedOtherSummarySections.Append(SectionDiffResults.AddedOtherValues);
+	DiffResults.RemovedSummaryEntriesInBase.Append(SummaryItemDiff.RemovedBaseValues);
+	DiffResults.AddedSummaryEntriesInOther.Append(SummaryItemDiff.AddedOtherValues);
+	
+	TArray<UNiagaraHierarchySection*> BaseSections = BaseRoot->GetSectionData();
+	TArray<UNiagaraHierarchySection*> OtherSections = OtherRoot->GetSectionData();
 
-	for (const FCommonValuePair<FNiagaraStackSection>& CommonValuePair : SectionDiffResults.CommonValuePairs)
+	for (UNiagaraHierarchySection* OtherSection : OtherSections)
 	{
-		bool bSectionsMatch = true;
-		if (CommonValuePair.BaseValue.SectionDisplayName.CompareTo(CommonValuePair.OtherValue.SectionDisplayName) == 0 &&
-			CommonValuePair.BaseValue.bEnabled == CommonValuePair.OtherValue.bEnabled &&
-			CommonValuePair.BaseValue.Categories.Num() == CommonValuePair.OtherValue.Categories.Num())
-		{
-			for (int32 CategoryIndex = 0; CategoryIndex < CommonValuePair.BaseValue.Categories.Num(); ++CategoryIndex)
-			{
-				if (CommonValuePair.OtherValue.Categories.ContainsByPredicate([&CommonValuePair, CategoryIndex](const FText& Category)
-					{ return Category.CompareTo(CommonValuePair.BaseValue.Categories[CategoryIndex]) == 0; }) == false)
-				{
-					bSectionsMatch = false;
-					break;
-				}
-			}
-		}
-		else
+		FNiagaraHierarchyIdentity OtherSectionIdentity = OtherSection->GetPersistentIdentity();
+		if(!BaseSections.ContainsByPredicate([OtherSectionIdentity](UNiagaraHierarchySection* BaseSection)
 		{
-			bSectionsMatch = false;
-		}
-
-		if (bSectionsMatch == false)
+			return BaseSection->GetPersistentIdentity() == OtherSectionIdentity;
+		}))
 		{
-			DiffResults.ModifiedBaseSummarySections.Add(CommonValuePair.BaseValue);
-			DiffResults.ModifiedOtherSummarySections.Add(CommonValuePair.OtherValue);
+			AddedSections.Add(OtherSection);
 		}
 	}
 
-	if (BaseEditorData != nullptr && OtherEditorData != nullptr && BaseEditorData->ShouldShowSummaryView() != OtherEditorData->ShouldShowSummaryView())
-	{
-		DiffResults.NewShouldShowSummaryViewValue = OtherEditorData->ShouldShowSummaryView();
-	}	
+	DiffResults.AddedSummarySectionsInOther = AddedSections;
 }
 
 void FNiagaraScriptMergeManager::DiffScriptStacks(TSharedRef<FNiagaraScriptStackMergeAdapter> BaseScriptStackAdapter, TSharedRef<FNiagaraScriptStackMergeAdapter> OtherScriptStackAdapter, FNiagaraScriptStackDiffResults& DiffResults) const
@@ -3685,66 +3657,70 @@ FNiagaraScriptMergeManager::FApplyDiffResults FNiagaraScriptMergeManager::ApplyE
 	FVersionedNiagaraEmitterData* BaseEmitterData = BaseEmitter.GetEmitterData();
 	UNiagaraEmitterEditorData* EditorData = Cast<UNiagaraEmitterEditorData>(BaseEmitterData->GetEditorData());
 
-	if (EditorData == nullptr)
-	{
-		EditorData = NewObject<UNiagaraEmitterEditorData>(BaseEmitter.Emitter, NAME_None, RF_Transactional);
-		BaseEmitter.Emitter->Modify();
-		BaseEmitter.Emitter->SetEditorData(EditorData, BaseEmitter.Version);
-	}
-
 	check(EditorData != GetDefault<UNiagaraEmitterEditorData>());
 	check(EditorData->GetOuter() == BaseEmitter.Emitter);
 	
 	EditorData->Modify();
-	for (TSharedRef<FNiagaraInputSummaryMergeAdapter> RemovedInput : DiffResults.RemovedInputSummaryEntries)
-	{
-		// Don't support removed summary inputs for now.  The user can change the visible flag if needed.
-	}
 
-	for (TSharedRef<FNiagaraInputSummaryMergeAdapter> AddedInput : DiffResults.AddedInputSummaryEntries)
-	{
-		EditorData->SetSummaryViewMetaData(AddedInput->GetKey(), AddedInput->GetValue());
-	}
+	UNiagaraHierarchyRoot* BaseRoot = EditorData->GetSummaryRoot();
 
-	for (TSharedRef<FNiagaraInputSummaryMergeAdapter> ModifiedInput : DiffResults.ModifiedOtherInputSummaryEntries)
+	// this will copy over the section without any child elements
+	for (const UNiagaraHierarchySection* AddedOtherSection : DiffResults.AddedSummarySectionsInOther)
 	{
-		EditorData->SetSummaryViewMetaData(ModifiedInput->GetKey(), ModifiedInput->GetValue());
+		BaseRoot->DuplicateSectionFromOtherRoot(*AddedOtherSection);
 	}
 
-	TArray<FNiagaraStackSection> SummarySections = EditorData->GetSummarySections();
-	for (const FNiagaraStackSection& RemovedBaseSection : DiffResults.RemovedBaseSummarySections)
-	{
-		// Don't support removed summary sections for now.  The user can change the enabled flag if needed.
-	}
-	
-	for (const FNiagaraStackSection& AddedOtherSection : DiffResults.AddedOtherSummarySections)
+	for (UNiagaraHierarchyItemBase* AddedInput : DiffResults.AddedSummaryEntriesInOther)
 	{
-		SummarySections.Add(AddedOtherSection);
-	}
+		UNiagaraHierarchyItemBase* SummaryItemParent = AddedInput->GetTypedOuter<UNiagaraHierarchyItemBase>();
+		bool bParentIsRoot = SummaryItemParent->IsA<UNiagaraHierarchyRoot>();
+		FNiagaraHierarchyIdentity ParentIdentity = SummaryItemParent ? SummaryItemParent->GetPersistentIdentity() : FNiagaraHierarchyIdentity();
 
-	for (int32 ModifiedSectionIndex = 0; ModifiedSectionIndex < DiffResults.ModifiedBaseSummarySections.Num(); ++ModifiedSectionIndex)
-	{
-		const FNiagaraStackSection& BaseSummarySection = DiffResults.ModifiedBaseSummarySections[ModifiedSectionIndex];
-		const FNiagaraStackSection& OtherSummarySection = DiffResults.ModifiedBaseSummarySections[ModifiedSectionIndex];
-		FNiagaraStackSection* ModifiedSectionPtr = SummarySections.FindByPredicate([BaseSummarySection](const FNiagaraStackSection& SummarySection)
-			{ return SummarySection.SectionDisplayName.CompareTo(BaseSummarySection.SectionDisplayName) == 0; });
-		if (ModifiedSectionPtr != nullptr)
+		if(ParentIdentity.IsValid())
 		{
-			// Use the enabled state in the child.
-			ModifiedSectionPtr->bEnabled = OtherSummarySection.bEnabled;
-			for (FText OtherSummarySectionCategory : OtherSummarySection.Categories)
+			// if the parent is the root, we can't use the parent identity to copy over the child as the two roots are supposed to be different, so we add it directly
+			if(bParentIsRoot)
+			{
+				/** Core principal of adding new items:
+				 * Any added item should come without children, as the parent could have added the same child somewhere else
+				 * Otherwise we'd need to identify not just new items vs. old items, but also new items with old children.
+				 * The children we removed but actually belong will be added back during another iteration of the loop
+				*/
+				if(BaseRoot->FindChildWithIdentity(AddedInput->GetPersistentIdentity(), false) == nullptr)
+				{
+					UNiagaraHierarchyItemBase* NewChild = BaseRoot->CopyAndAddItemAsChild(*AddedInput);
+					NewChild->GetChildrenMutable().Empty();
+					// for root level category, we have to fixup the section objects they point to as after duplicating them they will still point to the original section
+					if(UNiagaraHierarchyCategory* AsCategory = Cast<UNiagaraHierarchyCategory>(NewChild))
+					{
+						AsCategory->FixupSectionLinkage();
+					}
+				}				
+			}
+			else
 			{
-				// Add any categories added by the child, while ignoring any removed by the child.
-				if (ModifiedSectionPtr->Categories.ContainsByPredicate([OtherSummarySectionCategory](FText& Category)
-					{ return Category.CompareTo(OtherSummarySectionCategory) == 0; }) == false)
+				// it's possible the child was already added via ownership link
+				// (i.e. if categoryA owns inputA and categoryA is copied first, inputA will be included too)
+				// we only want to add that child if it hasn't been added already
+				if(BaseRoot->FindChildWithIdentity(AddedInput->GetPersistentIdentity(), true) == nullptr)
 				{
-					ModifiedSectionPtr->Categories.Add(OtherSummarySectionCategory);
+					UNiagaraHierarchyItemBase* NewChild = BaseRoot->CopyAndAddItemUnderParentIdentity(*AddedInput, ParentIdentity);
+					// see above
+					NewChild->GetChildrenMutable().Empty();
+					
+					if(!ensure(NewChild != nullptr))
+					{
+						UE_LOG(LogNiagaraEditor, Log, TEXT("Item %s could not be added during merge process"), *AddedInput->ToString());
+					}
 				}
 			}
 		}
 	}
 
-	EditorData->SetSummarySections(SummarySections);
+	for (UNiagaraHierarchyItemBase* RemovedInputInBase : DiffResults.RemovedSummaryEntriesInBase)
+	{
+		bool bChildRemoved = BaseRoot->RemoveChildWithIdentity(RemovedInputInBase->GetPersistentIdentity(), true);
+	}
 
 	if (DiffResults.NewShouldShowSummaryViewValue.IsSet())
 	{
@@ -3762,12 +3738,6 @@ FNiagaraScriptMergeManager::FApplyDiffResults FNiagaraScriptMergeManager::ApplyS
 	if (DiffResults.ModifiedStackEntryDisplayNames.Num() > 0)
 	{
 		UNiagaraEmitterEditorData* EditorData = Cast<UNiagaraEmitterEditorData>(BaseEmitter.GetEmitterData()->GetEditorData());
-		if (EditorData == nullptr)
-		{
-			EditorData = NewObject<UNiagaraEmitterEditorData>(BaseEmitter.Emitter, NAME_None, RF_Transactional);
-			EditorData->Modify();
-			BaseEmitter.Emitter->SetEditorData(EditorData, BaseEmitter.Version);
-		}
 
 		for (auto& Pair : DiffResults.ModifiedStackEntryDisplayNames)
 		{
diff --git a/Engine/Plugins/FX/Niagara/Source/NiagaraEditor/Private/NiagaraScriptMergeManager.h b/Engine/Plugins/FX/Niagara/Source/NiagaraEditor/Private/NiagaraScriptMergeManager.h
--- a/Engine/Plugins/FX/Niagara/Source/NiagaraEditor/Private/NiagaraScriptMergeManager.h
+++ b/Engine/Plugins/FX/Niagara/Source/NiagaraEditor/Private/NiagaraScriptMergeManager.h
@@ -405,15 +405,9 @@ struct FNiagaraEmitterDiffResults
 	TArray<TSharedRef<FNiagaraRendererMergeAdapter>> ModifiedBaseRenderers;
 	TArray<TSharedRef<FNiagaraRendererMergeAdapter>> ModifiedOtherRenderers;
 
-	TArray<TSharedRef<FNiagaraInputSummaryMergeAdapter>> RemovedInputSummaryEntries;
-	TArray<TSharedRef<FNiagaraInputSummaryMergeAdapter>> AddedInputSummaryEntries;
-	TArray<TSharedRef<FNiagaraInputSummaryMergeAdapter>> ModifiedInputSummaryEntries;
-	TArray<TSharedRef<FNiagaraInputSummaryMergeAdapter>> ModifiedOtherInputSummaryEntries;
-
-	TArray<FNiagaraStackSection> RemovedBaseSummarySections;
-	TArray<FNiagaraStackSection> AddedOtherSummarySections;
-	TArray<FNiagaraStackSection> ModifiedBaseSummarySections;
-	TArray<FNiagaraStackSection> ModifiedOtherSummarySections;
+	TArray<UNiagaraHierarchyItemBase*> RemovedSummaryEntriesInBase;
+	TArray<UNiagaraHierarchyItemBase*> AddedSummaryEntriesInOther;	
+	TArray<UNiagaraHierarchySection*> AddedSummarySectionsInOther;
 
 	TOptional<bool> NewShouldShowSummaryViewValue;
 	
@@ -460,6 +454,9 @@ public:
 
 	bool IsModuleInputDifferentFromBase(const FVersionedNiagaraEmitter& Emitter, const FVersionedNiagaraEmitter& BaseEmitter, ENiagaraScriptUsage ScriptUsage, FGuid ScriptUsageId, FGuid ModuleId, FString InputName);
 
+	/** Check whether a given summary item identity exists in the parent of a given emitter. */
+	bool DoesSummaryItemExistInBase(const FVersionedNiagaraEmitter& Emitter, FNiagaraHierarchyIdentity Identity);
+	
 	FApplyDiffResults ResetModuleInputToBase(const FVersionedNiagaraEmitter& VersionedEmitter, const FVersionedNiagaraEmitter& VersionedBaseEmitter, ENiagaraScriptUsage ScriptUsage, FGuid ScriptUsageId, FGuid ModuleId, FString InputName);
 
 	bool HasBaseEventHandler(const FVersionedNiagaraEmitter& BaseEmitter, FGuid EventScriptUsageId);
diff --git a/Engine/Plugins/FX/Niagara/Source/NiagaraEditor/Private/Toolkits/SystemToolkitModes/NiagaraSystemToolkitModeBase.cpp b/Engine/Plugins/FX/Niagara/Source/NiagaraEditor/Private/Toolkits/SystemToolkitModes/NiagaraSystemToolkitModeBase.cpp
--- a/Engine/Plugins/FX/Niagara/Source/NiagaraEditor/Private/Toolkits/SystemToolkitModes/NiagaraSystemToolkitModeBase.cpp
+++ b/Engine/Plugins/FX/Niagara/Source/NiagaraEditor/Private/Toolkits/SystemToolkitModes/NiagaraSystemToolkitModeBase.cpp
@@ -3,10 +3,12 @@
 #include "NiagaraSystemToolkitModeBase.h"
 
 #include "AdvancedPreviewSceneModule.h"
+#include "ClassIconFinder.h"
 #include "NiagaraConstants.h"
 #include "NiagaraEditorModule.h"
 #include "NiagaraEditorSettings.h"
 #include "NiagaraEditorStyle.h"
+#include "NiagaraNodeAssignment.h"
 #include "ViewModels/NiagaraScriptGraphViewModel.h"
 #include "Toolkits/NiagaraSystemToolkit.h"
 #include "NiagaraScriptSource.h"
@@ -28,10 +30,13 @@
 #include "Widgets/SNiagaraGeneratedCodeView.h"
 #include "Widgets/SNiagaraHierarchy.h"
 #include "Widgets/Layout/SWidgetSwitcher.h"
+#include "Widgets/Layout/SBox.h"
 #include "Widgets/Input/SButton.h"
 #include "ViewModels/NiagaraScratchPadScriptViewModel.h"
 #include "Widgets/SNiagaraSelectedObjectsDetails.h"
+#include "Widgets/SNiagaraHierarchyModuleInput.h"
 #include "NiagaraObjectSelection.h"
+#include "NiagaraSimulationStageBase.h"
 #include "Widgets/SNiagaraDebugCaptureView.h"
 #include "Widgets/SNiagaraSimCacheView.h"
 #include "Widgets/SNiagaraSimCacheViewTimeline.h"
@@ -40,7 +45,9 @@
 #include "Customizations/NiagaraComponentDetails.h"
 #include "Styling/StyleColors.h"
 #include "ViewModels/NiagaraSystemEditorDocumentsViewModel.h"
+#include "ViewModels/HierarchyEditor/NiagaraSummaryViewViewModel.h"
 #include "ViewModels/HierarchyEditor/NiagaraUserParametersHierarchyViewModel.h"
+#include "Widgets/SNiagaraHierarchyAssignment.h"
 
 #define LOCTEXT_NAMESPACE "NiagaraSystemToolkitModeBase"
 
@@ -65,6 +72,7 @@ const FName FNiagaraSystemToolkitModeBase::VersioningTabID(TEXT("NiagaraSystemEd
 const FName FNiagaraSystemToolkitModeBase::ScratchPadScriptsTabID(TEXT("NiagaraSystemEditor_ScratchPadScripts"));
 const FName FNiagaraSystemToolkitModeBase::UserParametersTabID(TEXT("NiagaraSystemEditor_UserParameters"));
 const FName FNiagaraSystemToolkitModeBase::UserParametersHierarchyTabID(TEXT("NiagaraSystemEditor_UserParametersHierarchy"));
+const FName FNiagaraSystemToolkitModeBase::EmitterSummaryViewEditorTabID(TEXT("NiagaraSystemEditor_SummaryViewHierarchyEditor"));
 
 FNiagaraSystemToolkitModeBase::FNiagaraSystemToolkitModeBase(FName InModeName, TWeakPtr<FNiagaraSystemToolkit> InSystemToolkit) : FApplicationMode(InModeName), SystemToolkit(InSystemToolkit), SwitcherIdx(0)
 {
@@ -82,6 +90,11 @@ FNiagaraSystemToolkitModeBase::~FNiagaraSystemToolkitModeBase()
 			TSharedPtr<FNiagaraScratchPadScriptViewModel> OldScratchScriptVM = LastActiveDocumentModel.Pin();
 			OldScratchScriptVM->GetGraphViewModel()->GetNodeSelection()->OnSelectedObjectsChanged().Remove(LastSelectionUpdateDelegate);
 		}
+		
+		if(UpdateSummaryViewHandle.IsValid())
+		{
+			SystemToolkit.Pin()->GetSystemViewModel()->GetSelectionViewModel()->OnEmitterHandleIdSelectionChanged().Remove(UpdateSummaryViewHandle);
+		}
 	}
 }
 
@@ -334,6 +347,11 @@ void FNiagaraSystemToolkitModeBase::RegisterTabFactories(TSharedPtr<FTabManager>
 		.SetDisplayName(LOCTEXT("UserParametersHierarchyTab", "User Parameters Hierarchy"))
 		.SetGroup(WorkspaceMenuCategory.ToSharedRef())
 		.SetIcon(FSlateIcon(FNiagaraEditorStyle::Get().GetStyleSetName(), "Tab.UserParameterHierarchy"));
+		
+	InTabManager->RegisterTabSpawner(EmitterSummaryViewEditorTabID, FOnSpawnTab::CreateSP(this, &FNiagaraSystemToolkitModeBase::SpawnTab_SummaryViewEditor))
+		.SetDisplayName(LOCTEXT("SummaryViewEditorTitle", "Edit Summary View"))
+		.SetGroup(WorkspaceMenuCategory.ToSharedRef())
+		.SetIcon(FSlateIcon(FNiagaraEditorStyle::Get().GetStyleSetName(), "Tab.UserParameterHierarchy"));
 }
 
 int FNiagaraSystemToolkitModeBase::GetActiveSelectionDetailsIndex() const
@@ -921,7 +939,7 @@ TSharedRef<SWidget> GenerateRowContentForUserParameterHierarchyEditor(TSharedRef
 	return SNullWidget::NullWidget;
 }
 
-TSharedRef<SWidget> GenerateCustomDetailsPanelNameWidget(TSharedPtr<FNiagaraHierarchyItemViewModelBase> HierarchyItem)
+TSharedRef<SWidget> GenerateCustomDetailsPanelNameWidgetForUserParameterEditor(TSharedPtr<FNiagaraHierarchyItemViewModelBase> HierarchyItem)
 {
 	if(!HierarchyItem.IsValid())
 	{
@@ -960,11 +978,219 @@ TSharedRef<SDockTab> FNiagaraSystemToolkitModeBase::SpawnTab_UserParametersHiera
 			[
 				SNew(SNiagaraHierarchy, SystemToolkit.Pin()->GetSystemViewModel()->GetUserParametersHierarchyViewModel())
 				.OnGenerateRowContentWidget_Static(&GenerateRowContentForUserParameterHierarchyEditor, SystemToolkit.Pin()->GetSystemViewModel())
-				.OnGenerateCustomDetailsPanelNameWidget_Static(&GenerateCustomDetailsPanelNameWidget)
+				.OnGenerateCustomDetailsPanelNameWidget_Static(&GenerateCustomDetailsPanelNameWidgetForUserParameterEditor)
 			]
 		];
 	
 	return SpawnedTab;
 }
 
+TSharedRef<SWidget> GenerateRowContentForSummaryViewHierarchyEditor(TSharedRef<FNiagaraHierarchyItemViewModelBase> HierarchyItem, TSharedPtr<FNiagaraEmitterViewModel> EmitterViewModel)
+{
+	if(const UNiagaraHierarchyModuleInput* ModuleInput = Cast<UNiagaraHierarchyModuleInput>(HierarchyItem->GetData()))
+	{
+		TSharedRef<FNiagaraModuleInputViewModel> ModuleInputViewModel = StaticCastSharedRef<FNiagaraModuleInputViewModel>(HierarchyItem);
+		return SNew(SNiagaraHierarchyModuleInput, ModuleInputViewModel);
+	}
+	else if(const UNiagaraHierarchyAssignmentInput* AssignmentInput = Cast<UNiagaraHierarchyAssignmentInput>(HierarchyItem->GetData()))
+	{
+		TSharedRef<FNiagaraAssignmentInputViewModel> AssignmentInputViewModel = StaticCastSharedRef<FNiagaraAssignmentInputViewModel>(HierarchyItem);
+		TOptional<FNiagaraStackGraphUtilities::FMatchingFunctionInputData> InputData = AssignmentInputViewModel->GetInputData();
+
+		if(ensure(InputData.IsSet()))
+		{
+			return FNiagaraParameterUtilities::GetParameterWidget(FNiagaraVariable(InputData->Type, InputData->InputName), false, false);
+		}
+	}
+	else if(const UNiagaraHierarchyModule* Module = Cast<UNiagaraHierarchyModule>(HierarchyItem->GetDataMutable()))
+	{
+		TSharedPtr<FNiagaraHierarchyItemViewModelBase> ItemPtr = HierarchyItem;
+		TSharedPtr<FNiagaraFunctionViewModel> ModuleViewModel = StaticCastSharedPtr<FNiagaraFunctionViewModel>(ItemPtr);
+
+		if(UNiagaraNodeAssignment* AssignmentNode = Cast<UNiagaraNodeAssignment>(ModuleViewModel->GetFunctionCallNode()))
+		{
+			return SNew(SNiagaraHierarchyAssignment, *AssignmentNode);
+		}
+		else
+		{
+			return SNew(SNiagaraHierarchyModule, ModuleViewModel);
+		}		
+	}
+	else if(const UNiagaraHierarchyRenderer* Renderer = Cast<UNiagaraHierarchyRenderer>(HierarchyItem->GetDataMutable()))
+	{
+		TSharedPtr<FNiagaraHierarchyItemViewModelBase> ItemPtr = HierarchyItem;
+		TSharedPtr<FNiagaraHierarchyRendererViewModel> RendererViewModel = StaticCastSharedPtr<FNiagaraHierarchyRendererViewModel>(ItemPtr);
+		return SNew(SHorizontalBox)
+			+ SHorizontalBox::Slot()
+			.AutoWidth()
+			.Padding(2.f)
+			[
+				SNew(SImage).Image(FSlateIconFinder::FindIconBrushForClass(RendererViewModel->GetRendererProperties()->GetClass()))
+			]
+			+ SHorizontalBox::Slot()
+			[
+				SNew(STextBlock).Text(RendererViewModel.ToSharedRef(), &FNiagaraHierarchyRendererViewModel::ToStringAsText)
+			];
+	}
+	else if(const UNiagaraHierarchyEventHandler* EventHandler = Cast<UNiagaraHierarchyEventHandler>(HierarchyItem->GetDataMutable()))
+	{
+		TSharedPtr<FNiagaraHierarchyItemViewModelBase> ItemPtr = HierarchyItem;
+		TSharedPtr<FNiagaraHierarchyEventHandlerViewModel> EventHandlerViewModel = StaticCastSharedPtr<FNiagaraHierarchyEventHandlerViewModel>(ItemPtr);
+		return SNew(SHorizontalBox)
+			+ SHorizontalBox::Slot()
+			.AutoWidth()
+			.Padding(2.f)
+			[
+				SNew(SImage).Image(FNiagaraEditorStyle::Get().GetBrush("NiagaraEditor.EventIcon"))
+			]
+			+ SHorizontalBox::Slot()
+			[
+				SNew(STextBlock).Text(EventHandlerViewModel.ToSharedRef(), &FNiagaraHierarchyEventHandlerViewModel::ToStringAsText)
+			];
+	}
+	else if(const UNiagaraHierarchySimStage* SimStage = Cast<UNiagaraHierarchySimStage>(HierarchyItem->GetDataMutable()))
+	{
+		TSharedPtr<FNiagaraHierarchyItemViewModelBase> ItemPtr = HierarchyItem;
+		TSharedPtr<FNiagaraHierarchySimStageViewModel> SimStageViewModel = StaticCastSharedPtr<FNiagaraHierarchySimStageViewModel>(ItemPtr);
+		return SNew(SHorizontalBox)
+			+ SHorizontalBox::Slot()
+			.AutoWidth()
+			.Padding(2.f)
+			[
+				SNew(SImage).Image(FNiagaraEditorStyle::Get().GetBrush("NiagaraEditor.SimulationStageIcon"))
+			]
+			+ SHorizontalBox::Slot()
+			[
+				SNew(STextBlock).Text(SimStageViewModel.ToSharedRef(), &FNiagaraHierarchySimStagePropertiesViewModel::ToStringAsText)
+			];
+	}
+	else if(const UNiagaraHierarchySimStageProperties* SimStageProperties = Cast<UNiagaraHierarchySimStageProperties>(HierarchyItem->GetDataMutable()))
+	{
+		TSharedPtr<FNiagaraHierarchyItemViewModelBase> ItemPtr = HierarchyItem;
+		TSharedPtr<FNiagaraHierarchySimStagePropertiesViewModel> SimStagePropertiesViewModel = StaticCastSharedPtr<FNiagaraHierarchySimStagePropertiesViewModel>(ItemPtr);
+		return SNew(STextBlock).Text(SimStagePropertiesViewModel.ToSharedRef(), &FNiagaraHierarchySimStagePropertiesViewModel::ToStringAsText);
+	}
+	else if(HierarchyItem->GetDataMutable()->IsA<UNiagaraHierarchyCategory>())
+	{
+		TSharedRef<FNiagaraHierarchyCategoryViewModel> TreeViewCategory = StaticCastSharedRef<FNiagaraHierarchyCategoryViewModel>(HierarchyItem);
+		return SNew(SNiagaraHierarchyCategory, TreeViewCategory);
+	}
+	else if(const UNiagaraHierarchyItem* Item = Cast<UNiagaraHierarchyItem>(HierarchyItem->GetData()))
+	{
+		return SNew(STextBlock).Text(FText::FromString(HierarchyItem->ToString()));
+	}
+
+	return SNullWidget::NullWidget;
+}
+
+TSharedRef<SWidget> GenerateCustomDetailsPanelNameWidgetForSummaryViewEditor(TSharedPtr<FNiagaraHierarchyItemViewModelBase> HierarchyItem)
+{
+	if(!HierarchyItem.IsValid())
+	{
+		return SNew(STextBlock).Text(FText::FromString("None selected"));
+	}
+	
+	if(HierarchyItem->GetData()->IsA<UNiagaraHierarchyCategory>() || HierarchyItem->GetData()->IsA<UNiagaraHierarchySection>())
+	{
+		return SNew(SBox).Padding(2.f)
+		[
+			SNew(STextBlock)
+			.Text_Lambda([HierarchyItem]
+			{
+				return FText::FromString(HierarchyItem->ToString());
+			})
+		];
+	}
+	else if(HierarchyItem->GetData()->IsA<UNiagaraHierarchyModuleInput>())
+	{
+		TSharedPtr<FNiagaraModuleInputViewModel> ModuleInputViewModel = StaticCastSharedPtr<FNiagaraModuleInputViewModel>(HierarchyItem);
+		return SNew(SNiagaraHierarchyModuleInput, ModuleInputViewModel.ToSharedRef());
+	}
+	else if(HierarchyItem->GetData()->IsA<UNiagaraHierarchyItem>())
+	{
+		return SNew(STextBlock).Text(FText::FromString(HierarchyItem->ToString()));
+	}
+
+	return SNullWidget::NullWidget;
+}
+
+TSharedRef<SWidget> FNiagaraSystemToolkitModeBase::CreateSummaryViewWidget() const
+{
+	TArray<FGuid> SelectedEmitterHandleGuids = SystemToolkit.Pin()->GetSystemViewModel()->GetSelectionViewModel()->GetSelectedEmitterHandleIds();
+
+	TSharedPtr<SWidget> ContentWidget = SNullWidget::NullWidget;
+	if(SystemToolkit.Pin()->GetSystemViewModel()->GetEditMode() == ENiagaraSystemViewModelEditMode::EmitterAsset)
+	{
+		TSharedPtr<FNiagaraEmitterViewModel> EmitterViewModel = SystemToolkit.Pin()->GetSystemViewModel()->GetEmitterHandleViewModels()[0]->GetEmitterViewModel();
+		UNiagaraSummaryViewViewModel* SummaryViewHierarchyViewModel = EmitterViewModel->GetSummaryHierarchyViewModel();
+		
+		ContentWidget = SNew(SNiagaraHierarchy, SummaryViewHierarchyViewModel)
+		.OnGenerateRowContentWidget_Static(&GenerateRowContentForSummaryViewHierarchyEditor, EmitterViewModel)
+		.OnGenerateCustomDetailsPanelNameWidget_Static(&GenerateCustomDetailsPanelNameWidgetForSummaryViewEditor);
+	}
+	else
+	{
+		if(SelectedEmitterHandleGuids.Num() != 1)
+		{
+			ContentWidget = SNew(SBox)
+			.HAlign(HAlign_Center)
+			.VAlign(VAlign_Center)
+			[
+				SNew(STextBlock).Text(LOCTEXT("SummaryViewEditorInvalidSelection", "Please select a single emitter to display its summary view options."))
+			];
+		}
+		else
+		{		
+			TSharedPtr<FNiagaraEmitterViewModel> EmitterViewModel = SystemToolkit.Pin()->GetSystemViewModel()->GetEmitterHandleViewModelById(SelectedEmitterHandleGuids[0])->GetEmitterViewModel();
+			UNiagaraSummaryViewViewModel* SummaryViewHierarchyViewModel = EmitterViewModel->GetSummaryHierarchyViewModel();
+		
+			ContentWidget = SNew(SNiagaraHierarchy, SummaryViewHierarchyViewModel)
+			.OnGenerateRowContentWidget_Static(&GenerateRowContentForSummaryViewHierarchyEditor, EmitterViewModel)
+			.OnGenerateCustomDetailsPanelNameWidget_Static(&GenerateCustomDetailsPanelNameWidgetForSummaryViewEditor);
+		}
+	}
+
+	return ContentWidget.ToSharedRef();
+}
+
+void FNiagaraSystemToolkitModeBase::UpdateSummaryViewOnSelectionChanged() const
+{
+	if(SystemToolkit.Pin()->GetTabManager()->FindExistingLiveTab(EmitterSummaryViewEditorTabID).IsValid())
+	{
+		SummaryViewContainer->SetContent(CreateSummaryViewWidget());
+	}
+}
+
+void FNiagaraSystemToolkitModeBase::OnSummaryViewEditorClosed(TSharedRef<SDockTab> DockTab) const
+{
+	SummaryViewContainer->SetContent(SNullWidget::NullWidget);
+}
+
+TSharedRef<SDockTab> FNiagaraSystemToolkitModeBase::SpawnTab_SummaryViewEditor(const FSpawnTabArgs& Args)
+{
+	check(Args.GetTabId().TabType == EmitterSummaryViewEditorTabID);
+
+	if(UpdateSummaryViewHandle.IsValid())
+	{
+		SystemToolkit.Pin()->GetSystemViewModel()->GetSelectionViewModel()->OnEmitterHandleIdSelectionChanged().Remove(UpdateSummaryViewHandle);
+	}
+
+	if(SystemToolkit.Pin()->GetSystemViewModel()->GetEditMode() == ENiagaraSystemViewModelEditMode::SystemAsset)
+	{
+		UpdateSummaryViewHandle = SystemToolkit.Pin()->GetSystemViewModel()->GetSelectionViewModel()->OnEmitterHandleIdSelectionChanged().AddSP(this, &FNiagaraSystemToolkitModeBase::UpdateSummaryViewOnSelectionChanged);
+	}
+	
+	TSharedRef<SDockTab> SpawnedTab = SNew(SDockTab)
+		.OnTabClosed(this, &FNiagaraSystemToolkitModeBase::OnSummaryViewEditorClosed)
+		.Label(LOCTEXT("SummaryViewHierarchyTitle", "Edit Summary View"))
+		[
+			SAssignNew(SummaryViewContainer, SBox)
+			.AddMetaData<FTagMetaData>(FTagMetaData(TEXT("SummaryView")))
+		];
+
+	SummaryViewContainer->SetContent(CreateSummaryViewWidget());
+	
+	return SpawnedTab;
+}
+
 #undef LOCTEXT_NAMESPACE
