Summary View 2.0

diff --git a/Engine/Plugins/FX/Niagara/Source/NiagaraEditor/Private/ViewModels/HierarchyEditor/NiagaraHierarchyViewModelBase.cpp b/Engine/Plugins/FX/Niagara/Source/NiagaraEditor/Private/ViewModels/HierarchyEditor/NiagaraHierarchyViewModelBase.cpp
--- a/Engine/Plugins/FX/Niagara/Source/NiagaraEditor/Private/ViewModels/HierarchyEditor/NiagaraHierarchyViewModelBase.cpp
+++ b/Engine/Plugins/FX/Niagara/Source/NiagaraEditor/Private/ViewModels/HierarchyEditor/NiagaraHierarchyViewModelBase.cpp
@@ -5,7 +5,6 @@
 #include "NiagaraEmitterEditorData.h"
 #include "NiagaraNodeFunctionCall.h"
 #include "NiagaraScriptVariable.h"
-#include "NiagaraSystemEditorData.h"
 #include "SDropTarget.h"
 #include "ViewModels/NiagaraSystemViewModel.h"
 #include "ViewModels/Stack/NiagaraStackViewModel.h"
@@ -15,58 +14,112 @@
 #include "Widgets/SToolTip.h"
 #include "Widgets/Input/SCheckBox.h"
 #include "Editor.h"
+#include "IPropertyRowGenerator.h"
 #include "ScopedTransaction.h"
 #include "Framework/Commands/GenericCommands.h"
 #include "ToolMenus.h"
+#include "Widgets/Notifications/SNotificationList.h"
 
 #define LOCTEXT_NAMESPACE "NiagaraHierarchyEditor"
 
 const TArray<UNiagaraHierarchyItemBase*>& UNiagaraHierarchyItemBase::GetChildren() const
 {
-	if(CanHaveChildren())
+	return Children;
+}
+
+UNiagaraHierarchyItemBase* UNiagaraHierarchyItemBase::FindChildWithIdentity(FNiagaraHierarchyIdentity ChildIdentity, bool bSearchRecursively)
+{
+	TObjectPtr<UNiagaraHierarchyItemBase>* FoundItem = Children.FindByPredicate([ChildIdentity](UNiagaraHierarchyItemBase* Child)
 	{
-		return Children;
+		return Child->GetPersistentIdentity() == ChildIdentity;
+	});
+
+	if(FoundItem)
+	{
+		return *FoundItem;
 	}
+	
+	if(bSearchRecursively)
+	{
+		for(UNiagaraHierarchyItemBase* Child : Children)
+		{
+			UNiagaraHierarchyItemBase* FoundChild = Child->FindChildWithIdentity(ChildIdentity, bSearchRecursively);
+
+			if(FoundChild)
+			{
+				return FoundChild;
+			}
+		}
+	}	
 
-	static TArray<UNiagaraHierarchyItemBase*> EmptyChildren;
-	return EmptyChildren;
+	return nullptr;
 }
 
-void UNiagaraHierarchyItemBase::Refresh()
+UNiagaraHierarchyItemBase* UNiagaraHierarchyItemBase::CopyAndAddItemAsChild(const UNiagaraHierarchyItemBase& ItemToCopy)
 {
-	TArray<UNiagaraHierarchyItemBase*> OldChildren = Children;
+	UNiagaraHierarchyItemBase* NewChild = Cast<UNiagaraHierarchyItemBase>(StaticDuplicateObject(&ItemToCopy, this));
+	if(NewChild->GetPersistentIdentity() != ItemToCopy.GetPersistentIdentity())
+	{
+		check(false);
+	}
+	GetChildrenMutable().Add(NewChild);
 
-	RefreshDataInternal();
+	return NewChild;
+}
 
-	Children.Empty();
-	for(UNiagaraHierarchyItemBase* Item : OldChildren)
+UNiagaraHierarchyItemBase* UNiagaraHierarchyItemBase::CopyAndAddItemUnderParentIdentity(const UNiagaraHierarchyItemBase& ItemToCopy, FNiagaraHierarchyIdentity ParentIdentity)
+{
+	UNiagaraHierarchyItemBase* ParentItem = FindChildWithIdentity(ParentIdentity, true);
+
+	if(ParentItem)
 	{
-		// if an item is already finalized we can skip it
-		if(Item->IsFinalized())
+		UNiagaraHierarchyItemBase* NewChild = Cast<UNiagaraHierarchyItemBase>(StaticDuplicateObject(&ItemToCopy, ParentItem));
+		if(NewChild->GetPersistentIdentity() != ItemToCopy.GetPersistentIdentity())
 		{
-			continue;
+			check(false);
 		}
-		
-		Item->Refresh();
+		ParentItem->GetChildrenMutable().Add(NewChild);
+		return NewChild;
+	}
 
-		// the refresh can have caused the item to finalize. In that case, we don't add it back
-		if(Item->IsFinalized())
+	return nullptr;
+}
+
+bool UNiagaraHierarchyItemBase::RemoveChildWithIdentity(FNiagaraHierarchyIdentity ChildIdentity, bool bSearchRecursively)
+{	
+	int32 RemovedChildrenCount = Children.RemoveAll([ChildIdentity](UNiagaraHierarchyItemBase* Child)
+	{
+		return Child->GetPersistentIdentity() == ChildIdentity;
+	});
+
+	if(RemovedChildrenCount > 1)
+	{
+		UE_LOG(LogNiagaraEditor, Warning, TEXT("More than one child with the same identity has been found in parent %s"), *ToString());
+	}
+	
+	bool bChildrenRemoved = RemovedChildrenCount > 0;
+	
+	if(bSearchRecursively && bChildrenRemoved == false)
+	{
+		for(UNiagaraHierarchyItemBase* Child : Children)
 		{
-			continue;
+			bChildrenRemoved |= Child->RemoveChildWithIdentity(ChildIdentity, bSearchRecursively);
 		}
-		
-		Children.Add(Item);
 	}
+
+	return bChildrenRemoved;
 }
 
-void UNiagaraHierarchyItemBase::Finalize()
-{	
-	for(UNiagaraHierarchyItemBase* Child : Children)
+TArray<FNiagaraHierarchyIdentity> UNiagaraHierarchyItemBase::GetParentIdentities() const
+{
+	TArray<FNiagaraHierarchyIdentity> ParentIdentities;
+
+	for(UNiagaraHierarchyItemBase* Parent = Cast<UNiagaraHierarchyItemBase>(GetOuter()); Parent != nullptr; Parent = Cast<UNiagaraHierarchyItemBase>(Parent->GetOuter()))
 	{
-		Child->Finalize();
+		ParentIdentities.Add(Parent->GetPersistentIdentity());
 	}
-	
-	bFinalized = true;
+
+	return ParentIdentities;
 }
 
 bool UNiagaraHierarchyItemBase::Modify(bool bAlwaysMarkDirty)
@@ -93,25 +146,7 @@ void UNiagaraHierarchyItemBase::PostLoad()
 	Super::PostLoad();
 }
 
-void UNiagaraHierarchyRoot::RefreshDataInternal()
-{
-	// in addition, we update our sections
-	TArray<UNiagaraHierarchySection*> OldSections = Sections;
-
-	Sections.Empty();
-	for(UNiagaraHierarchySection* Section : OldSections)
-	{
-		// if a section was marked finalized, we don't add it back
-		if(Section->IsFinalized())
-		{
-			continue;
-		}
-
-		Sections.Add(Section);
-	}
-}
-
-UNiagaraHierarchySection* UNiagaraHierarchyRoot::AddSection(FText InNewSectionName)
+UNiagaraHierarchySection* UNiagaraHierarchyRoot::AddSection(FText InNewSectionName, int32 InsertIndex)
 {
 	TSet<FName> ExistingSectionNames;
 	
@@ -124,16 +159,48 @@ UNiagaraHierarchySection* UNiagaraHierarchyRoot::AddSection(FText InNewSectionNa
 	UNiagaraHierarchySection* NewSectionItem = NewObject<UNiagaraHierarchySection>(this);
 	NewSectionItem->SetSectionName(NewName);
 	NewSectionItem->SetFlags(RF_Transactional);
-	Sections.Insert(NewSectionItem, 0);
+	
+	if(InsertIndex == INDEX_NONE)
+	{
+		Sections.Add(NewSectionItem);
+	}
+	else
+	{
+		Sections.Insert(NewSectionItem, InsertIndex);
+	}
+	
 	return NewSectionItem;
 }
 
+UNiagaraHierarchySection* UNiagaraHierarchyRoot::FindSectionByIdentity(FNiagaraHierarchyIdentity SectionIdentity)
+{
+	for(UNiagaraHierarchySection* Section : Sections)
+	{
+		if(Section->GetPersistentIdentity() == SectionIdentity)
+		{
+			return Section;
+		}
+	}
+
+	return nullptr;
+}
+
+void UNiagaraHierarchyRoot::DuplicateSectionFromOtherRoot(const UNiagaraHierarchySection& SectionToCopy)
+{
+	if(FindSectionByIdentity(SectionToCopy.GetPersistentIdentity()) != nullptr || SectionToCopy.GetOuter() == this)
+	{
+		return;
+	}
+	
+	Sections.Add(Cast<UNiagaraHierarchySection>(StaticDuplicateObject(&SectionToCopy, this)));
+}
+
 void UNiagaraHierarchyRoot::RemoveSection(FText SectionName)
 {
-	if(ensure(Sections.ContainsByPredicate([SectionName](UNiagaraHierarchySection* Section)
+	if(Sections.ContainsByPredicate([SectionName](UNiagaraHierarchySection* Section)
 	{
 		return Section->GetSectionNameAsText().EqualTo(SectionName);
-	})))
+	}))
 	{
 		Sections.RemoveAll([SectionName](UNiagaraHierarchySection* Section)
 		{
@@ -142,6 +209,14 @@ void UNiagaraHierarchyRoot::RemoveSection(FText SectionName)
 	}
 }
 
+void UNiagaraHierarchyRoot::RemoveSectionByIdentity(FNiagaraHierarchyIdentity SectionIdentity)
+{
+	Sections.RemoveAll([SectionIdentity](UNiagaraHierarchySection* Section)
+	{
+		return Section->GetPersistentIdentity() == SectionIdentity;
+	});
+}
+
 TSet<FName> UNiagaraHierarchyRoot::GetSections() const
 {
 	TSet<FName> OutSections;
@@ -187,52 +262,71 @@ bool FNiagaraHierarchyCategoryViewModel::IsTopCategoryActive() const
 			}
 		}
 		
-		return HierarchyViewModel->IsSectionActive(Result->GetSection());
+		return HierarchyViewModel->IsHierarchySectionActive(Result->GetSection());
 	}
 
 	return false;	
 }
 
-TOptional<EItemDropZone> FNiagaraHierarchyCategoryViewModel::OnCanAcceptDropInternal(TSharedPtr<FDragDropOperation> DragDropOp, EItemDropZone ItemDropZone)
+FNiagaraHierarchyItemViewModelBase::FCanPerformActionResults FNiagaraHierarchyCategoryViewModel::CanDropOnInternal(TSharedPtr<FNiagaraHierarchyItemViewModelBase> DraggedItem, EItemDropZone ItemDropZone)
 {
-	bool bAllowDrop = false;
+	FCanPerformActionResults Results(false);
+	
+	TArray<TSharedPtr<FNiagaraHierarchyCategoryViewModel>> TargetChildrenCategories;
+	GetChildrenViewModelsForType<UNiagaraHierarchyCategory, FNiagaraHierarchyCategoryViewModel>(TargetChildrenCategories);
 
-	if(DragDropOp->IsOfType<FNiagaraHierarchyDragDropOp>())
+	// we only allow drops if some general conditions are fulfilled
+	if(DraggedItem->GetData() != GetData() &&
+		(!DraggedItem->HasParent(AsShared(), false) || ItemDropZone != EItemDropZone::OntoItem)  &&
+		!HasParent(DraggedItem, true))
 	{
-		TSharedPtr<FNiagaraHierarchyDragDropOp> HierarchyDragDropOp = StaticCastSharedPtr<FNiagaraHierarchyDragDropOp>(DragDropOp);
-		TSharedPtr<FNiagaraHierarchyItemViewModelBase> SourceDropItem = HierarchyDragDropOp->GetDraggedItem().Pin();
-		TSharedPtr<FNiagaraHierarchyItemViewModelBase> TargetDropItem = AsShared();
-		
-		TArray<TSharedPtr<FNiagaraHierarchyCategoryViewModel>> TargetChildrenCategories;
-		TargetDropItem->GetChildrenViewModelsForType<UNiagaraHierarchyCategory, FNiagaraHierarchyCategoryViewModel>(TargetChildrenCategories);
-
-		// we only allow drops if some general conditions are fulfilled
-		if(SourceDropItem->GetData() != TargetDropItem->GetData() &&
-			(!SourceDropItem->HasParent(TargetDropItem, false) || ItemDropZone != EItemDropZone::OntoItem)  &&
-			!TargetDropItem->HasParent(SourceDropItem, true))
+		// categories can be dropped on categories, but only if the resulting sibling categories have different names
+		if(DraggedItem->GetData()->IsA<UNiagaraHierarchyCategory>())
 		{
-			// categories can be dropped on categories, but only if the resulting sibling categories have different names
-			if(SourceDropItem->GetData()->IsA<UNiagaraHierarchyCategory>())
-			{
-				bAllowDrop =
-				(
-					(ItemDropZone == EItemDropZone::OntoItem && !TargetChildrenCategories.ContainsByPredicate([SourceDropItem](TSharedPtr<FNiagaraHierarchyCategoryViewModel> HierarchyCategoryViewModel)
-						{
-							return SourceDropItem->ToString() == HierarchyCategoryViewModel->ToString();
-						}))
-						||
-					(ItemDropZone != EItemDropZone::OntoItem && SourceDropItem->ToString() != TargetDropItem->ToString())
-				);
-			}
-			else if(SourceDropItem->GetData()->IsA<UNiagaraHierarchyItem>())
+			Results.bCanPerform =
+			(
+				(ItemDropZone == EItemDropZone::OntoItem && !TargetChildrenCategories.ContainsByPredicate([DraggedItem](TSharedPtr<FNiagaraHierarchyCategoryViewModel> HierarchyCategoryViewModel)
+					{
+						return DraggedItem->ToString() == HierarchyCategoryViewModel->ToString();
+					}))
+					||
+				(ItemDropZone != EItemDropZone::OntoItem && DraggedItem->ToString() != ToString())
+			);
+		}
+		else if(DraggedItem->GetData()->IsA<UNiagaraHierarchyItem>())
+		{
+			// items can generally be dropped onto categories
+			Results.bCanPerform = EItemDropZone::OntoItem == ItemDropZone;
+
+			if(Results.bCanPerform)
 			{
-				// items can generally be dropped onto categories
-				bAllowDrop = EItemDropZone::OntoItem == ItemDropZone;
+				if(DraggedItem->IsForHierarchy() == false)
+				{
+					FText Message = LOCTEXT("AddItemToCategoryDragMessage", "Add {0} to {1}");
+					Results.CanPerformMessage = FText::FormatOrdered(Message, FText::FromString(DraggedItem->ToString()), FText::FromString(ToString()));
+				}
+				else
+				{
+					FText Message = LOCTEXT("MoveItemToCategoryDragMessage", "Move {0} to {1}");
+					Results.CanPerformMessage = FText::FormatOrdered(Message, FText::FromString(DraggedItem->ToString()), FText::FromString(ToString()));
+				}
 			}
 		}
 	}
 
-	return bAllowDrop ? ItemDropZone : TOptional<EItemDropZone>();
+	return Results;
+}
+
+void UNiagaraHierarchyCategory::FixupSectionLinkage()
+{
+	UNiagaraHierarchyRoot* OwningRoot = GetTypedOuter<UNiagaraHierarchyRoot>();
+
+	if(Section != nullptr && Section->GetTypedOuter<UNiagaraHierarchyRoot>() != OwningRoot)
+	{
+		UNiagaraHierarchySection* CorrectSection = OwningRoot->FindSectionByIdentity(Section->GetPersistentIdentity());
+		ensure(CorrectSection != nullptr);
+		Section = CorrectSection;
+	}
 }
 
 void UNiagaraHierarchySection::SetSectionNameAsText(const FText& Text)
@@ -242,8 +336,6 @@ void UNiagaraHierarchySection::SetSectionNameAsText(const FText& Text)
 
 UNiagaraHierarchyViewModelBase::UNiagaraHierarchyViewModelBase()
 {
-	SourceRoot = NewObject<UNiagaraHierarchyRoot>(this, FName("HierarchySourceRoot"));
-	SourceViewModelRoot = MakeShared<FNiagaraHierarchyRootViewModel>(SourceRoot.Get(), this);
 	Commands = MakeShared<FUICommandList>();
 }
 
@@ -251,42 +343,104 @@ UNiagaraHierarchyViewModelBase::~UNiagaraHierarchyViewModelBase()
 {
 	RefreshSourceViewDelegate.Unbind();
 	RefreshHierarchyWidgetDelegate.Unbind();
-	RefreshSectionsWidgetDelegate.Unbind();
+	RefreshSectionsViewDelegate.Unbind();
 }
 
-void UNiagaraHierarchyViewModelBase::ForceFullRefresh()
+void UNiagaraHierarchyViewModelBase::Initialize()
+{		
+	HierarchyRoot = GetHierarchyRoot();
+	HierarchyRoot->SetFlags(RF_Transactional);
+
+	TArray<UNiagaraHierarchyItemBase*> AllItems;
+	HierarchyRoot->GetChildrenOfType<UNiagaraHierarchyItemBase>(AllItems, true);
+	for(UNiagaraHierarchyItemBase* Item : AllItems)
+	{
+		Item->SetFlags(RF_Transactional);
+	}
+
+	for(UNiagaraHierarchySection* Section : HierarchyRoot->GetSectionDataMutable())
+	{
+		Section->SetFlags(RF_Transactional);
+	}
+
+	UToolMenus* ToolMenus = UToolMenus::Get();
+
+	FName MenuName("NiagaraHierarchyMenu");
+	if(!ToolMenus->IsMenuRegistered(MenuName))
+	{
+		UToolMenu* HierarchyMenu = ToolMenus->RegisterMenu(MenuName, NAME_None, EMultiBoxType::Menu);
+		HierarchyMenu->AddSection("Base");
+	}
+	
+	SetupCommands();
+
+	HierarchyRootViewModel = MakeShared<FNiagaraHierarchyRootViewModel>(HierarchyRoot.Get(), this, true);
+	HierarchyRootViewModel->Initialize();
+	HierarchyRootViewModel->AddChildFilter(FNiagaraHierarchyItemViewModelBase::FOnFilterChild::CreateUObject(this, &UNiagaraHierarchyViewModelBase::FilterForHierarchySection));
+	HierarchyRootViewModel->AddChildFilter(FNiagaraHierarchyItemViewModelBase::FOnFilterChild::CreateUObject(this, &UNiagaraHierarchyViewModelBase::FilterForUncategorizedRootItemsInAllSection));
+	HierarchyRootViewModel->SyncViewModelsToData();
+	
+	SetActiveHierarchySection(nullptr);
+	
+	InitializeInternal();
+	
+	OnInitializedDelegate.ExecuteIfBound();
+}
+
+void UNiagaraHierarchyViewModelBase::Finalize()
 {	
-	PrepareSourceItems();
-	HierarchyViewModelRoot->SyncToData();
-	RefreshSourceView(true);
-	RefreshHierarchyView(true);
-	RefreshSectionsWidget();
+	HierarchyRootViewModel.Reset();
+	HierarchyRoot = nullptr;
+	
+	FinalizeInternal();
 }
 
-void UNiagaraHierarchyViewModelBase::RefreshSourceView(bool bFullRefresh) const
+void UNiagaraHierarchyViewModelBase::ForceFullRefresh()
 {
-	if(RefreshSourceViewDelegate.IsBound())
+	RefreshSourceItemsRequestedDelegate.ExecuteIfBound();
+	// todo (me) during merge at startup this can be nullptr for some reason
+	if(HierarchyRootViewModel.IsValid())
 	{
-		RefreshSourceViewDelegate.Execute(bFullRefresh);
+		HierarchyRootViewModel->SyncViewModelsToData();
 	}
+	RefreshAllViewsRequestedDelegate.ExecuteIfBound(true);
 }
 
-void UNiagaraHierarchyViewModelBase::RefreshHierarchyView(bool bFullRefresh) const
+void UNiagaraHierarchyViewModelBase::ForceFullRefreshOnTimer()
 {
-	if(RefreshHierarchyWidgetDelegate.IsBound())
-	{
-		RefreshHierarchyWidgetDelegate.Execute(bFullRefresh);
-	}
+	ensure(FullRefreshNextFrameHandle.IsValid());
+	ForceFullRefresh();
+	FullRefreshNextFrameHandle.Invalidate();
 }
 
-void UNiagaraHierarchyViewModelBase::RefreshSectionsWidget() const
+void UNiagaraHierarchyViewModelBase::RequestFullRefreshNextFrame()
 {
-	if(RefreshSectionsWidgetDelegate.IsBound())
+	if(!FullRefreshNextFrameHandle.IsValid())
 	{
-		RefreshSectionsWidgetDelegate.Execute();
+		FullRefreshNextFrameHandle = GEditor->GetTimerManager()->SetTimerForNextTick(FTimerDelegate::CreateUObject(this, &UNiagaraHierarchyViewModelBase::ForceFullRefreshOnTimer));
 	}
 }
 
+void UNiagaraHierarchyViewModelBase::RefreshAllViews(bool bFullRefresh) const
+{
+	RefreshAllViewsRequestedDelegate.ExecuteIfBound(bFullRefresh);
+}
+
+void UNiagaraHierarchyViewModelBase::RefreshSourceView(bool bFullRefresh) const
+{
+	RefreshSourceViewDelegate.ExecuteIfBound(bFullRefresh);
+}
+
+void UNiagaraHierarchyViewModelBase::RefreshHierarchyView(bool bFullRefresh) const
+{
+	RefreshHierarchyWidgetDelegate.ExecuteIfBound(bFullRefresh);
+}
+
+void UNiagaraHierarchyViewModelBase::RefreshSectionsView() const
+{
+	RefreshSectionsViewDelegate.ExecuteIfBound();
+}
+
 void UNiagaraHierarchyViewModelBase::PostUndo(bool bSuccess)
 {
 	ForceFullRefresh();
@@ -297,87 +451,57 @@ void UNiagaraHierarchyViewModelBase::PostRedo(bool bSuccess)
 	PostUndo(bSuccess);
 }
 
-void UNiagaraHierarchyViewModelBase::Initialize()
+bool UNiagaraHierarchyViewModelBase::MatchesContext(const FTransactionContext& InContext, const TArray<TPair<UObject*, FTransactionObjectEvent>>& TransactionObjectContexts) const
 {
-	if(GEditor)
+	for(const TPair<UObject*, FTransactionObjectEvent>& TransactionObjectContext : TransactionObjectContexts)
 	{
-		GEditor->RegisterForUndo(this);
+		if(TransactionObjectContext.Key->IsA<UNiagaraHierarchyItemBase>())
+		{
+			return true;
+		}
 	}
 	
-	HierarchyRoot = GetHierarchyDataRoot();
-	HierarchyRoot->SetFlags(RF_Transactional);
-
-	TArray<UNiagaraHierarchyItemBase*> AllItems;
-	HierarchyRoot->GetChildrenOfType<UNiagaraHierarchyItemBase>(AllItems, true);
-	for(UNiagaraHierarchyItemBase* Item : AllItems)
-	{
-		Item->SetFlags(RF_Transactional);
-	}
-
-	for(UNiagaraHierarchySection* Section : HierarchyRoot->GetSectionDataMutable())
-	{
-		Section->SetFlags(RF_Transactional);
-	}
-
-	UToolMenus* ToolMenus = UToolMenus::Get();
+	return false;
+}
 
-	FName MenuName("NiagaraHierarchyMenu");
-	if(!ToolMenus->IsMenuRegistered(MenuName))
+bool UNiagaraHierarchyViewModelBase::FilterForHierarchySection(TSharedPtr<const FNiagaraHierarchyItemViewModelBase> ItemViewModel) const
+{
+	if(ActiveHierarchySection.IsValid())
 	{
-		UToolMenu* HierarchyMenu = ToolMenus->RegisterMenu(MenuName, NAME_None, EMultiBoxType::Menu);
-		HierarchyMenu->AddSection("Base");
+		return GetActiveHierarchySectionData() == ItemViewModel->GetSection();
 	}
-	
-	SetupCommands();
-
-	HierarchyViewModelRoot = MakeShared<FNiagaraHierarchyRootViewModel>(HierarchyRoot.Get(), this);
-	HierarchyViewModelRoot->SyncToData();
-
-	SetActiveSection(nullptr);
-
-	PrepareSourceItems();
 
-	InitializeInternal();
+	return true;
 }
 
-void UNiagaraHierarchyViewModelBase::Finalize()
+bool UNiagaraHierarchyViewModelBase::FilterForUncategorizedRootItemsInAllSection(TSharedPtr<const FNiagaraHierarchyItemViewModelBase> ItemViewModel) const
 {
-	if(GEditor)
+	// we want to filter out all items that are directly added to the root if we aren't in the 'All' section
+	if(ActiveHierarchySection.IsValid())
 	{
-		GEditor->UnregisterForUndo(this);
+		return ItemViewModel->GetData<UNiagaraHierarchyCategory>() != nullptr;
 	}
-	
-	SourceViewModelRoot->Finalize();
-	SourceViewModelRoot.Reset();
-	HierarchyViewModelRoot.Reset();
-	SourceRoot = nullptr;
-	HierarchyRoot = nullptr;
-	
-	FinalizeInternal();
+
+	return true;
 }
 
 TSharedPtr<FNiagaraHierarchyItemViewModelBase> UNiagaraHierarchyViewModelBase::CreateViewModelForData(UNiagaraHierarchyItemBase* ItemBase, TSharedPtr<FNiagaraHierarchyItemViewModelBase> Parent)
 {
 	if(UNiagaraHierarchyItem* Item = Cast<UNiagaraHierarchyItem>(ItemBase))
 	{
-		return MakeShared<FNiagaraHierarchyItemViewModel>(Item, Parent, this);
+		return MakeShared<FNiagaraHierarchyItemViewModel>(Item, Parent, this, Parent.IsValid() ? Parent->IsForHierarchy() : false);
 	}
 	else if(UNiagaraHierarchyCategory* Category = Cast<UNiagaraHierarchyCategory>(ItemBase))
 	{
-		return MakeShared<FNiagaraHierarchyCategoryViewModel>(Category, Parent, this);
+		return MakeShared<FNiagaraHierarchyCategoryViewModel>(Category, Parent, this, Parent.IsValid() ? Parent->IsForHierarchy() : false);
 	}
 
 	return nullptr;
 }
 
-const TArray<TSharedPtr<FNiagaraHierarchyItemViewModelBase>>& UNiagaraHierarchyViewModelBase::GetSourceItems() const
-{
-	return SourceViewModelRoot->GetChildren();
-}
-
 const TArray<TSharedPtr<FNiagaraHierarchyItemViewModelBase>>& UNiagaraHierarchyViewModelBase::GetHierarchyItems() const
 {
-	return HierarchyViewModelRoot->GetFilteredChildren();
+	return HierarchyRootViewModel->GetFilteredChildren();
 }
 
 void UNiagaraHierarchyViewModelBase::OnGetChildren(TSharedPtr<FNiagaraHierarchyItemViewModelBase> Item, TArray<TSharedPtr<FNiagaraHierarchyItemViewModelBase>>& OutChildren) const
@@ -385,26 +509,26 @@ void UNiagaraHierarchyViewModelBase::OnGetChildren(TSharedPtr<FNiagaraHierarchyI
 	OutChildren.Append(Item->GetFilteredChildren());
 }
 
-void UNiagaraHierarchyViewModelBase::SetActiveSection(TSharedPtr<FNiagaraHierarchySectionViewModel> Section)
+void UNiagaraHierarchyViewModelBase::SetActiveHierarchySection(TSharedPtr<FNiagaraHierarchySectionViewModel> Section)
 {
-	ActiveSection = Section;	
+	ActiveHierarchySection = Section;	
 	RefreshHierarchyView(true);
-	OnSectionActivatedDelegate.Broadcast(Section);
+	OnHierarchySectionActivatedDelegate.ExecuteIfBound(Section);
 }
 
-TSharedPtr<FNiagaraHierarchySectionViewModel> UNiagaraHierarchyViewModelBase::GetActiveSection() const
+TSharedPtr<FNiagaraHierarchySectionViewModel> UNiagaraHierarchyViewModelBase::GetActiveHierarchySection() const
 {
-	return ActiveSection.IsValid() ? ActiveSection.Pin() : nullptr;
+	return ActiveHierarchySection.IsValid() ? ActiveHierarchySection.Pin() : nullptr;
 }
 
-UNiagaraHierarchySection* UNiagaraHierarchyViewModelBase::GetActiveSectionData() const
+UNiagaraHierarchySection* UNiagaraHierarchyViewModelBase::GetActiveHierarchySectionData() const
 {
-	return ActiveSection.IsValid() ? ActiveSection.Pin()->GetDataMutable<UNiagaraHierarchySection>() : nullptr;
+	return ActiveHierarchySection.IsValid() ? ActiveHierarchySection.Pin()->GetDataMutable<UNiagaraHierarchySection>() : nullptr;
 }
 
-bool UNiagaraHierarchyViewModelBase::IsSectionActive(const UNiagaraHierarchySection* Section) const
+bool UNiagaraHierarchyViewModelBase::IsHierarchySectionActive(const UNiagaraHierarchySection* Section) const
 {
-	return ActiveSection == nullptr || ActiveSection.Pin()->GetData() == Section;
+	return ActiveHierarchySection == nullptr || ActiveHierarchySection.Pin()->GetData() == Section;
 }
 
 FString UNiagaraHierarchyViewModelBase::OnItemToStringDebug(TSharedPtr<FNiagaraHierarchyItemViewModelBase> ItemBaseViewModel) const
@@ -412,6 +536,42 @@ FString UNiagaraHierarchyViewModelBase::OnItemToStringDebug(TSharedPtr<FNiagaraH
 	return ItemBaseViewModel->ToString();	
 }
 
+const TArray<TSharedRef<IDetailTreeNode>>& UNiagaraHierarchyViewModelBase::RequestDetailTreeNodesForObject(UObject* Object)
+{
+	if(ObjectTreeNodeMap.Contains(Object))
+	{
+		return ObjectTreeNodeMap[Object];
+	}
+
+	if(ObjectToPropertyRowGeneratorMap.Contains(Object))
+	{
+		ObjectTreeNodeMap.Add(Object, ObjectToPropertyRowGeneratorMap[Object]->GetRootTreeNodes());
+		return ObjectTreeNodeMap[Object];
+	}
+	
+	FPropertyEditorModule& PropertyEditorModule = FModuleManager::Get().LoadModuleChecked<FPropertyEditorModule>("PropertyEditor");
+	ObjectToPropertyRowGeneratorMap.Add(Object, PropertyEditorModule.CreatePropertyRowGenerator(FPropertyRowGeneratorArgs()));
+	ObjectToPropertyRowGeneratorMap[Object]->SetObjects({Object});
+	return RequestDetailTreeNodesForObject(Object);
+}
+
+FNiagaraHierarchyItemViewModelBase::~FNiagaraHierarchyItemViewModelBase()
+{
+	Children.Empty();
+	FilteredChildren.Empty();
+}
+
+UNiagaraHierarchyItemBase* FNiagaraHierarchyItemViewModelBase::AddChild(TSubclassOf<UNiagaraHierarchyItemBase> NewChildClass, FNiagaraHierarchyIdentity ChildIdentity)
+{
+	UNiagaraHierarchyItemBase* NewChild = NewObject<UNiagaraHierarchyItemBase>(GetDataMutable(), NewChildClass);
+	NewChild->SetIdentity(ChildIdentity);
+	GetDataMutable()->GetChildrenMutable().Add(NewChild);
+	
+	SyncViewModelsToData();
+	HierarchyViewModel->OnHierarchyChanged().Broadcast();
+	return NewChild;
+}
+
 void FNiagaraHierarchyItemViewModelBase::Tick(float DeltaTime)
 {
 	if(bRenamePending)
@@ -425,12 +585,35 @@ TStatId FNiagaraHierarchyItemViewModelBase::GetStatId() const
 	RETURN_QUICK_DECLARE_CYCLE_STAT(FNiagaraHierarchyItemViewModelBase, STATGROUP_Tickables);
 }
 
-void FNiagaraHierarchyItemViewModelBase::SyncToData()
+void FNiagaraHierarchyItemViewModelBase::RefreshChildrenData()
 {
-	// this will refresh all underlying data recursively
-	ItemBase->Refresh();
+	TArray<TSharedPtr<FNiagaraHierarchyItemViewModelBase>> TmpChildren = Children;
+	for(TSharedPtr<FNiagaraHierarchyItemViewModelBase> Child : TmpChildren)
+	{
+		if(Child->RepresentsExternalData() && Child->DoesExternalDataStillExist(HierarchyViewModel->GetRefreshContext()) == false)
+		{
+			Child->Delete();
+		}
+	}
+
+	RefreshChildrenDataInternal();
+
+	/** All remaining children are supposed to exist at this point, as internal data won't be removed by refreshing & external data was cleaned up already.
+	 * This will not call RefreshChildrenData on data that has just been added as no view models exist for these yet.
+	 */
+	for(TSharedPtr<FNiagaraHierarchyItemViewModelBase> Child : Children)
+	{
+		Child->RefreshChildrenData();
+	}
+}
 
+void FNiagaraHierarchyItemViewModelBase::SyncViewModelsToData()
+{	
+	// this will recursively remove all outdated external data as well as give individual view models the chance to add new data
+	RefreshChildrenData();
+	
 	// now that the data is refreshed, we can sync to the data by recycling view models & creating new ones
+	// old view models will get deleted automatically
 	TArray<TSharedPtr<FNiagaraHierarchyItemViewModelBase>> NewChildren;
 	for(UNiagaraHierarchyItemBase* Child : ItemBase->GetChildren())
 	{		
@@ -438,33 +621,39 @@ void FNiagaraHierarchyItemViewModelBase::SyncToData()
 		// if we couldn't find a view model for a data child, we create it here
 		if(ViewModelIndex == INDEX_NONE)
 		{
-			TSharedPtr<FNiagaraHierarchyItemViewModelBase> NewItem = HierarchyViewModel->CreateViewModelForData(Child, AsShared());
-
-			if(ensure(NewItem.IsValid()))
+			TSharedPtr<FNiagaraHierarchyItemViewModelBase> ChildViewModel = HierarchyViewModel->CreateViewModelForData(Child, AsShared());
+			if(ensure(ChildViewModel.IsValid()))
 			{
-				NewItem->SyncToData();
-				NewChildren.Add(NewItem);
+				ChildViewModel->Initialize();
+				ChildViewModel->SyncViewModelsToData();
+				NewChildren.Add(ChildViewModel);
 			}
 		}
-		// if we could find it, we refresh its contained view models and readd it
+		// if we could find the view model, we refresh its contained view models and readd it
 		else
 		{
-			Children[ViewModelIndex]->SyncToData();
+			Children[ViewModelIndex]->SyncViewModelsToData();
 			NewChildren.Add(Children[ViewModelIndex]);
 		}
 	}
 
-	/** Give the view models a chance to further customize the children sync process. */
-	SyncChildrenViewModelsInternal();
-	
 	Children.Empty();
 	Children.Append(NewChildren);
+	
+	for(TSharedPtr<FNiagaraHierarchyItemViewModelBase> Child : Children)
+	{
+		Child->OnChildRequestedDeletion().BindSP(this, &FNiagaraHierarchyItemViewModelBase::DeleteChild);
+		Child->GetOnSynced().BindSP(this, &FNiagaraHierarchyItemViewModelBase::PropagateOnChildSynced);
+	}
+
+	/** Give the view models a chance to further customize the children sync process. */
+	SyncViewModelsToDataInternal();	
 
 	// first we sort the data. Categories before items.
 	GetDataMutable()->GetChildrenMutable().StableSort([](const UNiagaraHierarchyItemBase& ItemA, const UNiagaraHierarchyItemBase& ItemB)
-	{
-		return ItemA.IsA<UNiagaraHierarchyCategory>() && ItemB.IsA<UNiagaraHierarchyItem>();
-	});
+		{
+			return ItemA.IsA<UNiagaraHierarchyCategory>() && ItemB.IsA<UNiagaraHierarchyItem>();
+		});
 
 	// then we sort the view models according to the data order as this is what will determine widget order created from the view models
 	Children.Sort([this](const TSharedPtr<FNiagaraHierarchyItemViewModelBase>& ItemA, const TSharedPtr<FNiagaraHierarchyItemViewModelBase>& ItemB)
@@ -475,40 +664,46 @@ void FNiagaraHierarchyItemViewModelBase::SyncToData()
 	// we refresh the filtered children here as well
 	GetFilteredChildren();
 
-	OnSyncedDelegate.Broadcast();
+	OnSyncedDelegate.ExecuteIfBound();
 }
 
 const TArray<TSharedPtr<FNiagaraHierarchyItemViewModelBase>>& FNiagaraHierarchyItemViewModelBase::GetFilteredChildren() const
 {
 	FilteredChildren.Empty();
 
-	for(TSharedPtr<FNiagaraHierarchyItemViewModelBase> Child : Children)
+	if(CanHaveChildren())
 	{
-		// we add categories only if the top most category is considered active
-		if(const UNiagaraHierarchyCategory* Category = Child->GetData<UNiagaraHierarchyCategory>())
+		for(TSharedPtr<FNiagaraHierarchyItemViewModelBase> Child : Children)
 		{
-			if(StaticCastSharedPtr<FNiagaraHierarchyCategoryViewModel>(Child)->IsTopCategoryActive())
+			bool bPassesFilter = true;
+			for(const FOnFilterChild& OnFilterChild : ChildFilters)
 			{
-				FilteredChildren.Add(Child);
+				bPassesFilter &= OnFilterChild.Execute(Child);
+
+				if(!bPassesFilter)
+				{
+					break;
+				}
 			}
-		}
-		// if we have an item directly under the root without being categorized, we only allow if if the active section is set to "All"
-		else if(Child->GetParent().IsValid() && Child->GetParent() == HierarchyViewModel->GetHierarchyViewModelRoot() && Child->GetData()->IsA<UNiagaraHierarchyItem>())
-		{
-			if(HierarchyViewModel->GetActiveSectionData() == nullptr)
+
+			if(bPassesFilter)
 			{
 				FilteredChildren.Add(Child);
 			}
 		}
-		else
-		{
-			FilteredChildren.Add(Child);
-		}
 	}
 
 	return FilteredChildren;
 }
 
+void FNiagaraHierarchyItemViewModelBase::AddChildFilter(FOnFilterChild InFilterChild)
+{
+	if(ensure(InFilterChild.IsBound()))
+	{
+		ChildFilters.Add(InFilterChild);
+	}
+}
+
 bool FNiagaraHierarchyItemViewModelBase::HasParent(TSharedPtr<FNiagaraHierarchyItemViewModelBase> ParentCandidate, bool bRecursive)
 {
 	if(Parent.IsValid())
@@ -540,51 +735,57 @@ TSharedPtr<FNiagaraHierarchyItemViewModelBase> FNiagaraHierarchyItemViewModelBas
 {
 	FText TransactionText = FText::FormatOrdered(LOCTEXT("Transaction_AddedItem", "Added new {0} to hierarchy"), FText::FromString(NewItemClass->GetName()));
 	FScopedTransaction Transaction(TransactionText);
-	HierarchyViewModel->GetHierarchyDataRoot()->Modify();
+	HierarchyViewModel->GetHierarchyRoot()->Modify();
 	
 	UNiagaraHierarchyItemBase* NewItem = NewObject<UNiagaraHierarchyItemBase>(GetDataMutable(), NewItemClass, NAME_None, RF_Transactional);
 	NewItem->Modify();
 	GetDataMutable()->GetChildrenMutable().Add(NewItem);
-	SyncToData();
+	SyncViewModelsToData();
 
 	TSharedPtr<FNiagaraHierarchyItemViewModelBase> ViewModel = FindViewModelForChild(NewItem);
 	ensure(ViewModel.IsValid());
 	return ViewModel;
 }
 
-void FNiagaraHierarchyItemViewModelBase::DuplicateToThis(TSharedPtr<FNiagaraHierarchyItemViewModelBase> ItemToDuplicate, int32 InsertIndex)
+TSharedRef<FNiagaraHierarchyItemViewModelBase> FNiagaraHierarchyItemViewModelBase::DuplicateToThis(TSharedPtr<FNiagaraHierarchyItemViewModelBase> ItemToDuplicate, int32 InsertIndex)
 {
+	UNiagaraHierarchyItemBase* NewItem = Cast<UNiagaraHierarchyItemBase>(StaticDuplicateObject(ItemToDuplicate->GetData(), GetDataMutable()));
 	if(InsertIndex == INDEX_NONE)
 	{
-		GetDataMutable()->GetChildrenMutable().Add(Cast<UNiagaraHierarchyItemBase>(StaticDuplicateObject(ItemToDuplicate->GetData(), GetDataMutable())));
+		GetDataMutable()->GetChildrenMutable().Add(NewItem);
 	}
 	else
 	{
-		GetDataMutable()->GetChildrenMutable().Insert(Cast<UNiagaraHierarchyItemBase>(StaticDuplicateObject(ItemToDuplicate->GetData(), GetDataMutable())), InsertIndex);
+		GetDataMutable()->GetChildrenMutable().Insert(NewItem, InsertIndex);
 	}
 	
-	SyncToData();
+	SyncViewModelsToData();
+
 	HierarchyViewModel->OnHierarchyChanged().Broadcast();
+	TSharedPtr<FNiagaraHierarchyItemViewModelBase> ViewModel = FindViewModelForChild(NewItem);
+	return ViewModel.ToSharedRef();
 }
 
-void FNiagaraHierarchyItemViewModelBase::ReparentToThis(TSharedPtr<FNiagaraHierarchyItemViewModelBase> ItemToMove, int32 InsertIndex)
+TSharedRef<FNiagaraHierarchyItemViewModelBase> FNiagaraHierarchyItemViewModelBase::ReparentToThis(TSharedPtr<FNiagaraHierarchyItemViewModelBase> ItemToMove, int32 InsertIndex)
 {
+	UNiagaraHierarchyItemBase* NewItem = Cast<UNiagaraHierarchyItemBase>(StaticDuplicateObject(ItemToMove->GetData(), GetDataMutable()));
 	if(InsertIndex == INDEX_NONE)
 	{
-		GetDataMutable()->GetChildrenMutable().Add(Cast<UNiagaraHierarchyItemBase>(StaticDuplicateObject(ItemToMove->GetDataMutable(), GetDataMutable())));
+		GetDataMutable()->GetChildrenMutable().Add(NewItem);
 	}
 	else
 	{
-		GetDataMutable()->GetChildrenMutable().Insert(Cast<UNiagaraHierarchyItemBase>(StaticDuplicateObject(ItemToMove->GetData(), GetDataMutable())), InsertIndex);
+		GetDataMutable()->GetChildrenMutable().Insert(NewItem, InsertIndex);
 	}
 	
-	ItemToMove->Finalize();
-	ItemToMove->Parent.Pin()->SyncToData();
-	SyncToData();
+	ItemToMove->Delete();
+	SyncViewModelsToData();
 	HierarchyViewModel->OnHierarchyChanged().Broadcast();
+	TSharedPtr<FNiagaraHierarchyItemViewModelBase> ViewModel = FindViewModelForChild(NewItem);
+	return ViewModel.ToSharedRef();
 }
 
-TSharedPtr<FNiagaraHierarchyItemViewModelBase> FNiagaraHierarchyItemViewModelBase::FindViewModelForChild(UNiagaraHierarchyItemBase* Child) const
+TSharedPtr<FNiagaraHierarchyItemViewModelBase> FNiagaraHierarchyItemViewModelBase::FindViewModelForChild(UNiagaraHierarchyItemBase* Child, bool bSearchRecursively) const
 {
 	int32 Index = FindIndexOfChild(Child);
 	if(Index != INDEX_NONE)
@@ -592,12 +793,46 @@ TSharedPtr<FNiagaraHierarchyItemViewModelBase> FNiagaraHierarchyItemViewModelBas
 		return Children[Index];
 	}
 
+	if(bSearchRecursively)
+	{
+		for(TSharedPtr<FNiagaraHierarchyItemViewModelBase> ChildViewModel : Children)
+		{
+			TSharedPtr<FNiagaraHierarchyItemViewModelBase> FoundViewModel = ChildViewModel->FindViewModelForChild(Child, bSearchRecursively);
+
+			if(FoundViewModel.IsValid())
+			{
+				return FoundViewModel;
+			}
+		}
+	}
+
 	return nullptr;
 }
 
-int32 FNiagaraHierarchyItemViewModelBase::FindIndexOfChild(TSharedPtr<FNiagaraHierarchyItemViewModelBase> Child) const
+TSharedPtr<FNiagaraHierarchyItemViewModelBase> FNiagaraHierarchyItemViewModelBase::FindViewModelForChild(FNiagaraHierarchyIdentity ChildIdentity, bool bSearchRecursively) const
 {
-	return Children.Find(Child);
+	for(TSharedPtr<FNiagaraHierarchyItemViewModelBase> Child : Children)
+	{
+		if(Child->GetData()->GetPersistentIdentity() == ChildIdentity)
+		{
+			return Child;
+		}
+	}
+
+	if(bSearchRecursively)
+	{
+		for(TSharedPtr<FNiagaraHierarchyItemViewModelBase> ChildViewModel : Children)
+		{
+			TSharedPtr<FNiagaraHierarchyItemViewModelBase> FoundViewModel = ChildViewModel->FindViewModelForChild(ChildIdentity, bSearchRecursively);
+
+			if(FoundViewModel.IsValid())
+			{
+				return FoundViewModel;
+			}
+		}
+	}
+
+	return nullptr;
 }
 
 int32 FNiagaraHierarchyItemViewModelBase::FindIndexOfChild(UNiagaraHierarchyItemBase* Child) const
@@ -620,37 +855,66 @@ int32 FNiagaraHierarchyItemViewModelBase::FindIndexOfDataChild(UNiagaraHierarchy
 
 void FNiagaraHierarchyItemViewModelBase::Delete()
 {
-	FScopedTransaction Transaction(LOCTEXT("DeletedItemTransaction", "Deleted hierarchy item"));
-	ItemBase->Modify();
+	OnChildRequestedDeletionDelegate.Execute(AsShared());
+}
 
-	if(Parent.IsValid())
+void FNiagaraHierarchyItemViewModelBase::DeleteChild(TSharedPtr<FNiagaraHierarchyItemViewModelBase> Child)
+{	
+	ensure(Child->GetParent().Pin() == AsShared());	
+	GetDataMutable()->GetChildrenMutable().Remove(Child->GetDataMutable());
+	Children.Remove(Child);
+}
+
+TOptional<EItemDropZone> FNiagaraHierarchyItemViewModelBase::OnCanRowAcceptDrop(const FDragDropEvent& DragDropEvent, EItemDropZone ItemDropZone, TSharedPtr<FNiagaraHierarchyItemViewModelBase> Item)
+{
+	if(TSharedPtr<FNiagaraHierarchyDragDropOp> DragDropOp = DragDropEvent.GetOperationAs<FNiagaraHierarchyDragDropOp>())
 	{
-		Parent.Pin()->ItemBase->Modify();
+		FCanPerformActionResults Results = CanDropOn(DragDropOp->GetDraggedItem().Pin(), ItemDropZone);
+		DragDropOp->SetDescription(Results.CanPerformMessage);
+		return Results.bCanPerform ? ItemDropZone : TOptional<EItemDropZone>();
 	}
-	
-	Finalize();
+
+	return TOptional<EItemDropZone>();
 }
 
-void FNiagaraHierarchyItemViewModelBase::Finalize()
-{	
-	for(TSharedPtr<FNiagaraHierarchyItemViewModelBase> ChildViewModel : Children)
+FReply FNiagaraHierarchyItemViewModelBase::OnDroppedOnRow(const FDragDropEvent& DragDropEvent, EItemDropZone ItemDropZone, TSharedPtr<FNiagaraHierarchyItemViewModelBase> Item)
+{
+	if(TSharedPtr<FNiagaraHierarchyDragDropOp> HierarchyDragDropOp = DragDropEvent.GetOperationAs<FNiagaraHierarchyDragDropOp>())
 	{
-		ChildViewModel->Finalize();
-	}	
+		OnDroppedOn(HierarchyDragDropOp->GetDraggedItem().Pin(), ItemDropZone);
+		return FReply::Handled();
+	}
 
-	GetDataMutable()->Finalize();
-	FinalizeInternal();
-	SyncToData();
+	return FReply::Unhandled();
+}
+
+void FNiagaraHierarchyItemViewModelBase::OnRowDragLeave(const FDragDropEvent& DragDropEvent)
+{
+	if(TSharedPtr<FNiagaraHierarchyDragDropOp> HierarchyDragDropOp = DragDropEvent.GetOperationAs<FNiagaraHierarchyDragDropOp>())
+	{
+		HierarchyDragDropOp->SetDescription(FText::GetEmpty());
+	}
+}
+
+FNiagaraHierarchyItemViewModelBase::FCanPerformActionResults FNiagaraHierarchyItemViewModelBase::CanDrag()
+{
+	FCanPerformActionResults Results = IsEditableByUser();
+	if(Results.bCanPerform == false)
+	{
+		return Results;
+	}
+
+	return CanDragInternal();
 }
 
-TOptional<EItemDropZone> FNiagaraHierarchyItemViewModelBase::OnCanAcceptDrop(const FDragDropEvent& DragDropEvent, EItemDropZone ItemDropZone, TSharedPtr<FNiagaraHierarchyItemViewModelBase> Item)
+FNiagaraHierarchyItemViewModelBase::FCanPerformActionResults FNiagaraHierarchyItemViewModelBase::CanDropOnInternal(TSharedPtr<FNiagaraHierarchyItemViewModelBase>, EItemDropZone ItemDropZone)
 {
-	return OnCanAcceptDropInternal(DragDropEvent.GetOperation(), ItemDropZone);
+	return false;
 }
 
-TOptional<EItemDropZone> FNiagaraHierarchyItemViewModelBase::OnCanAcceptDropInternal(TSharedPtr<FDragDropOperation> DragDropOp, EItemDropZone ItemDropZone)
+void FNiagaraHierarchyItemViewModelBase::PropagateOnChildSynced()
 {
-	return TOptional<EItemDropZone>();	
+	OnSyncedDelegate.ExecuteIfBound();
 }
 
 TSharedRef<FNiagaraHierarchyDragDropOp> FNiagaraHierarchyItemViewModelBase::CreateDragDropOp()
@@ -662,21 +926,33 @@ TSharedRef<FNiagaraHierarchyDragDropOp> FNiagaraHierarchyItemViewModelBase::Crea
 
 FReply FNiagaraHierarchyItemViewModelBase::OnDragDetected(const FGeometry& Geometry, const FPointerEvent& PointerEvent,	bool bIsSource)
 {
-	if(CanDrag())
+	FCanPerformActionResults CanDragResults = CanDrag();
+	if(CanDragResults == true)
 	{
 		// if the drag is coming from source, we check if any of the hierarchy data already contains that item and we don't start a drag drop in that case
 		if(bIsSource)
 		{
-			TArray<UNiagaraHierarchyItemBase*> AllItems;
-			GetHierarchyViewModel()->GetHierarchyDataRoot()->GetChildrenOfType<UNiagaraHierarchyItemBase>(AllItems, true);
+			TArray<TSharedPtr<FNiagaraHierarchyItemViewModelBase>> AllChildren;
+			GetChildrenViewModelsForType<UNiagaraHierarchyItemBase, FNiagaraHierarchyItemViewModelBase>(AllChildren, true);
 
-			for(UNiagaraHierarchyItemBase* Item : AllItems)
+			bool bCanDrag = GetHierarchyViewModel()->GetHierarchyRootViewModel()->FindViewModelForChild(GetData()->GetPersistentIdentity(), true) == nullptr;			
+
+			if(bCanDrag)
 			{
-				if(GetData()->GetPersistentIdentity() == Item->GetPersistentIdentity())
+				for(TSharedPtr<FNiagaraHierarchyItemViewModelBase>& ItemViewModel : AllChildren)
 				{
-					return FReply::Unhandled();
+					if(GetHierarchyViewModel()->GetHierarchyRootViewModel()->FindViewModelForChild(ItemViewModel->GetData()->GetPersistentIdentity(), true) != nullptr)
+					{
+						bCanDrag = false;
+						break;
+					}
 				}
 			}
+			
+			if(bCanDrag == false)
+			{
+				return FReply::Unhandled();
+			}
 		}
 		
 		TSharedRef<FNiagaraHierarchyDragDropOp> HierarchyDragDropOp = CreateDragDropOp();
@@ -684,116 +960,179 @@ FReply FNiagaraHierarchyItemViewModelBase::OnDragDetected(const FGeometry& Geome
 
 		return FReply::Handled().BeginDragDrop(HierarchyDragDropOp);			
 	}
+	else
+	{
+		// if we can't drag and have a message, we show it as a slate notification
+		if(CanDragResults.CanPerformMessage.IsEmpty() == false)
+		{
+			FNotificationInfo CantDragInfo(CanDragResults.CanPerformMessage);
+			FSlateNotificationManager::Get().AddNotification(CantDragInfo);
+		}
+	}
 		
 	return FReply::Unhandled();
 }
 
-TOptional<EItemDropZone> FNiagaraHierarchyRootViewModel::OnCanAcceptDropInternal(TSharedPtr<FDragDropOperation> DragDropOp, EItemDropZone ItemDropZone)
+FNiagaraHierarchyRootViewModel::~FNiagaraHierarchyRootViewModel()
 {
-	bool bAllowDrop = false;
-
-	if(DragDropOp->IsOfType<FNiagaraHierarchyDragDropOp>())
-	{
-		TSharedPtr<FNiagaraHierarchyDragDropOp> HierarchyDragDropOp = StaticCastSharedPtr<FNiagaraHierarchyDragDropOp>(DragDropOp);
-		TSharedPtr<FNiagaraHierarchyItemViewModelBase> SourceDropItem = HierarchyDragDropOp->GetDraggedItem().Pin();
-		TSharedPtr<FNiagaraHierarchyItemViewModelBase> TargetDropItem = AsShared();
-
-		// we only allow drops if some general conditions are fulfilled
-		if(SourceDropItem->GetData() != TargetDropItem->GetData() &&
-			(!SourceDropItem->HasParent(TargetDropItem, false) || ItemDropZone != EItemDropZone::OntoItem)  &&
-			!TargetDropItem->HasParent(SourceDropItem, true))
-		{
-			bAllowDrop = 
-				// items can be dropped onto the root directly if the section is set to "All"
-				(SourceDropItem->GetData()->IsA<UNiagaraHierarchyItem>() && HierarchyViewModel->GetActiveSectionData() == nullptr)
-					||
-				// categories can be dropped onto the root always
-				(SourceDropItem->GetData()->IsA<UNiagaraHierarchyCategory>());
-		}
-		
-	}
+	
+}
 
-	return bAllowDrop ? ItemDropZone : TOptional<EItemDropZone>();
+void FNiagaraHierarchyRootViewModel::Initialize()
+{
+	GetOnSynced().BindSP(this, &FNiagaraHierarchyRootViewModel::PropagateOnSynced);
 }
 
-FReply FNiagaraHierarchyRootViewModel::OnDroppedOn(const FDragDropEvent& DragDropEvent, EItemDropZone ItemDropZone, TSharedPtr<FNiagaraHierarchyItemViewModelBase> Item)
+FNiagaraHierarchyItemViewModelBase::FCanPerformActionResults FNiagaraHierarchyRootViewModel::CanDropOnInternal(TSharedPtr<FNiagaraHierarchyItemViewModelBase> DraggedItem, EItemDropZone ItemDropZone)
 {
-	if(DragDropEvent.GetOperation()->IsOfType<FNiagaraHierarchyDragDropOp>())
+	FCanPerformActionResults Results(false);
+	
+	// we only allow drops if some general conditions are fulfilled
+	if(DraggedItem->GetData() != GetData() &&
+		(!DraggedItem->HasParent(AsShared(), false) || ItemDropZone != EItemDropZone::OntoItem)  &&
+		!HasParent(DraggedItem, true))
 	{
-		FScopedTransaction Transaction(LOCTEXT("Transaction_OnDropOnRoot", "Dropped item on root"));
-		HierarchyViewModel->GetHierarchyDataRoot()->Modify();
-		
-		TSharedPtr<FNiagaraHierarchyDragDropOp> HierarchyDragDropOp = DragDropEvent.GetOperationAs<FNiagaraHierarchyDragDropOp>();
-		TWeakPtr<FNiagaraHierarchyItemViewModelBase> DraggedItem = HierarchyDragDropOp->GetDraggedItem();
-
-		if(DraggedItem.Pin()->GetDataMutable()->IsA<UNiagaraHierarchyItem>() || DraggedItem.Pin()->GetDataMutable()->IsA<UNiagaraHierarchyCategory>())
+		Results.bCanPerform = 
+			// items can be dropped onto the root directly if the section is set to "All"
+			(DraggedItem->GetData()->IsA<UNiagaraHierarchyItem>() && HierarchyViewModel->GetActiveHierarchySectionData() == nullptr)
+				||
+			// categories can be dropped onto the root always
+			(DraggedItem->GetData()->IsA<UNiagaraHierarchyCategory>());
+
+		if(Results.bCanPerform)
 		{
-			if(HierarchyDragDropOp->GetIsFromSourceList())
+			if(DraggedItem->IsForHierarchy() == false)
 			{
-				DuplicateToThis(DraggedItem.Pin());
+				FText Message = LOCTEXT("CanDropSourceItemOnRootDragMessage", "Add {0} to the hierarchy root.");
+				Message = FText::FormatOrdered(Message, FText::FromString(DraggedItem->ToString()));
+				Results.CanPerformMessage = Message;
 			}
 			else
 			{
-				ReparentToThis(DraggedItem.Pin());
-			}
-
-			HierarchyViewModel->RefreshHierarchyView();
+				FText Message = LOCTEXT("CanDropHierarchyItemOnRootDragMessage", "Move {0} to the hierarchy root.");
+				Message = FText::FormatOrdered(Message, FText::FromString(DraggedItem->ToString()));
+				Results.CanPerformMessage = Message;
+			}			
+		}
+		else
+		{
+			FText Message = LOCTEXT("CantDropHierarchyItemOnRootDragMessage", "Can not add {0} to the hierarchy root directly. Custom sections only accept categories at the root level!");
+			Message = FText::FormatOrdered(Message, FText::FromString(DraggedItem->ToString()));
+			Results.CanPerformMessage = Message;
 		}
 	}
+	
+	return Results;
+}
 
-	return FReply::Unhandled();
+void FNiagaraHierarchyRootViewModel::OnDroppedOnInternal(TSharedPtr<FNiagaraHierarchyItemViewModelBase> DroppedItem, EItemDropZone ItemDropZone)
+{
+	FScopedTransaction Transaction(LOCTEXT("Transaction_OnDropOnRoot", "Dropped item on root"));
+	HierarchyViewModel->GetHierarchyRoot()->Modify();
+
+	if(DroppedItem->GetDataMutable()->IsA<UNiagaraHierarchyItem>() || DroppedItem->GetDataMutable()->IsA<UNiagaraHierarchyCategory>())
+	{
+		TSharedPtr<FNiagaraHierarchyItemViewModelBase> NewViewModel;
+		// we duplicate the item if the dragged item is from source
+		if(DroppedItem->IsForHierarchy() == false)
+		{
+			NewViewModel = DuplicateToThis(DroppedItem);
+		}
+		else
+		{
+			NewViewModel = ReparentToThis(DroppedItem);
+		}
+
+		if(UNiagaraHierarchyCategory* AsCategory = Cast<UNiagaraHierarchyCategory>(NewViewModel->GetDataMutable()))
+		{
+			AsCategory->SetSection(HierarchyViewModel->GetActiveHierarchySectionData());
+		}
+
+		HierarchyViewModel->RefreshHierarchyView();
+	}
 }
 
-TSharedPtr<FNiagaraHierarchySectionViewModel> FNiagaraHierarchyRootViewModel::AddNewSection()
+TSharedPtr<FNiagaraHierarchySectionViewModel> FNiagaraHierarchyRootViewModel::AddSection()
 {
 	FScopedTransaction ScopedTransaction(LOCTEXT("NewSectionAdded","Added Section"));
-	HierarchyViewModel->GetHierarchyDataRoot()->Modify();
+	HierarchyViewModel->GetHierarchyRoot()->Modify();
 	
-	UNiagaraHierarchySection* SectionData = Cast<UNiagaraHierarchyRoot>(ItemBase)->AddSection(LOCTEXT("NiagaraHierarchyEditorDefaultNewSectionName", "Section"));
+	UNiagaraHierarchySection* SectionData = Cast<UNiagaraHierarchyRoot>(ItemBase)->AddSection(LOCTEXT("NiagaraHierarchyEditorDefaultNewSectionName", "Section"), 0);
 	SectionData->Modify();
-	TSharedPtr<FNiagaraHierarchySectionViewModel> NewSectionViewModel = MakeShared<FNiagaraHierarchySectionViewModel>(SectionData, StaticCastSharedRef<FNiagaraHierarchyRootViewModel>(AsShared()), HierarchyViewModel);
+	TSharedPtr<FNiagaraHierarchySectionViewModel> NewSectionViewModel = MakeShared<FNiagaraHierarchySectionViewModel>(SectionData, StaticCastSharedRef<FNiagaraHierarchyRootViewModel>(AsShared()), HierarchyViewModel, bIsForHierarchy);
 	SectionViewModels.Add(NewSectionViewModel);
-	SyncToData();
-	HierarchyViewModel->SetActiveSection(NewSectionViewModel);
+	SyncViewModelsToData();
+	HierarchyViewModel->SetActiveHierarchySection(NewSectionViewModel);
 
+	OnSectionAddedDelegate.ExecuteIfBound(NewSectionViewModel);
+	OnSectionsChangedDelegate.ExecuteIfBound();
 	return NewSectionViewModel;
 }
 
-void FNiagaraHierarchyRootViewModel::SyncChildrenViewModelsInternal()
+void FNiagaraHierarchyRootViewModel::DeleteSection(TSharedPtr<FNiagaraHierarchyItemViewModelBase> InSectionViewModel)
+{
+	TSharedPtr<FNiagaraHierarchySectionViewModel> SectionViewModel = StaticCastSharedPtr<FNiagaraHierarchySectionViewModel>(InSectionViewModel);
+	GetDataMutable<UNiagaraHierarchyRoot>()->GetSectionDataMutable().Remove(SectionViewModel->GetDataMutable<UNiagaraHierarchySection>());
+	SectionViewModels.Remove(SectionViewModel);
+
+	OnSectionDeletedDelegate.ExecuteIfBound(SectionViewModel);
+	OnSectionsChangedDelegate.ExecuteIfBound();
+}
+
+void FNiagaraHierarchyRootViewModel::PropagateOnSynced()
+{
+	OnSyncPropagatedDelegate.ExecuteIfBound();
+}
+
+void FNiagaraHierarchyRootViewModel::SyncViewModelsToDataInternal()
 {
 	const UNiagaraHierarchyRoot* RootData = GetData<UNiagaraHierarchyRoot>();
 
-	TArray<TSharedPtr<FNiagaraHierarchySectionViewModel>> SectionsToRemove;
+	TArray<TSharedPtr<FNiagaraHierarchySectionViewModel>> NewSectionViewModels;
+	TArray<TSharedPtr<FNiagaraHierarchySectionViewModel>> SectionViewModelsToDelete;
+	
 	for(TSharedPtr<FNiagaraHierarchySectionViewModel> SectionViewModel : SectionViewModels)
 	{
-		if(SectionViewModel->GetDataMutable()->IsFinalized() || !RootData->GetSectionData().Contains(SectionViewModel->GetData()))
+		if(!RootData->GetSectionData().Contains(SectionViewModel->GetData()))
 		{
-			SectionsToRemove.Add(SectionViewModel);
+			SectionViewModelsToDelete.Add(SectionViewModel);
 		}
 	}
 
-	for (TSharedPtr<FNiagaraHierarchySectionViewModel> SectionViewModel : SectionsToRemove)
+	for (TSharedPtr<FNiagaraHierarchySectionViewModel> SectionViewModel : SectionViewModelsToDelete)
 	{
-		SectionViewModel->Finalize();
-		SectionViewModels.Remove(SectionViewModel);
+		SectionViewModel->Delete();
 	}
 	
 	for(UNiagaraHierarchySection* Section : RootData->GetSectionData())
 	{
-		if(Section->IsFinalized())
-		{
-			continue;
-		}
-		
-		if(!SectionViewModels.ContainsByPredicate([Section](TSharedPtr<FNiagaraHierarchySectionViewModel> SectionViewModel)
+		TSharedPtr<FNiagaraHierarchySectionViewModel>* SectionViewModelPtr = SectionViewModels.FindByPredicate([Section](TSharedPtr<FNiagaraHierarchySectionViewModel> SectionViewModel)
 		{
 			return SectionViewModel->GetData() == Section;
-		}))
+		});
+
+		TSharedPtr<FNiagaraHierarchySectionViewModel> SectionViewModel = nullptr;
+
+		if(SectionViewModelPtr)
+		{
+			SectionViewModel = *SectionViewModelPtr;
+		}
+
+		if(SectionViewModel == nullptr)
 		{
-			TSharedPtr<FNiagaraHierarchySectionViewModel> NewSectionViewModel = MakeShared<FNiagaraHierarchySectionViewModel>(Section, StaticCastSharedRef<FNiagaraHierarchyRootViewModel>(AsShared()), HierarchyViewModel); 
-			NewSectionViewModel->SyncToData();
-			SectionViewModels.Add(NewSectionViewModel);
+			SectionViewModel = MakeShared<FNiagaraHierarchySectionViewModel>(Section, StaticCastSharedRef<FNiagaraHierarchyRootViewModel>(AsShared()), HierarchyViewModel, bIsForHierarchy);
+			SectionViewModel->SyncViewModelsToData();;
 		}
+		
+		NewSectionViewModels.Add(SectionViewModel);
+	}
+
+	SectionViewModels.Empty();
+	SectionViewModels.Append(NewSectionViewModels);
+
+	for(TSharedPtr<FNiagaraHierarchySectionViewModel> SectionViewModel : SectionViewModels)
+	{
+		SectionViewModel->OnChildRequestedDeletion().BindSP(this, &FNiagaraHierarchyRootViewModel::DeleteSection);
 	}
 
 	SectionViewModels.Sort([this](const TSharedPtr<FNiagaraHierarchySectionViewModel>& ItemA, const TSharedPtr<FNiagaraHierarchySectionViewModel>& ItemB)
@@ -805,54 +1144,80 @@ void FNiagaraHierarchyRootViewModel::SyncChildrenViewModelsInternal()
 		});
 }
 
-TOptional<EItemDropZone> FNiagaraHierarchySectionViewModel::OnCanAcceptDropInternal(TSharedPtr<FDragDropOperation> DragDropOp, EItemDropZone ItemDropZone)
+void FNiagaraHierarchySectionViewModel::PopulateDynamicContextMenuSection(FToolMenuSection& DynamicSection)
 {
-	bool bAllowDrop = false;
+	//DynamicSection.AddMenuEntry()
+}
 
-	if(DragDropOp->IsOfType<FNiagaraHierarchyDragDropOp>())
+FNiagaraHierarchyItemViewModelBase::FCanPerformActionResults FNiagaraHierarchySectionViewModel::CanDropOnInternal(TSharedPtr<FNiagaraHierarchyItemViewModelBase> DraggedItem, EItemDropZone ItemDropZone)
+{
+	FCanPerformActionResults Results(false);
+	// we don't allow dropping onto source sections and we don't specify a message as the sections aren't going to light up as valid drop targets
+	if(IsForHierarchy() == false)
 	{
-		TSharedPtr<FNiagaraHierarchyDragDropOp> HierarchyDragDropOp = StaticCastSharedPtr<FNiagaraHierarchyDragDropOp>(DragDropOp);
-		TWeakPtr<FNiagaraHierarchyItemViewModelBase> DraggedItem = HierarchyDragDropOp->GetDraggedItem();
-
-		if(const UNiagaraHierarchyCategory* Category = Cast<UNiagaraHierarchyCategory>(DraggedItem.Pin()->GetData()))
+		return false;
+	}
+	
+	if(const UNiagaraHierarchyCategory* Category = Cast<UNiagaraHierarchyCategory>(DraggedItem->GetData()))
+	{
+		if(ItemDropZone == EItemDropZone::OntoItem)
 		{
-			if(ItemDropZone == EItemDropZone::OntoItem)
-			{
-				bAllowDrop = GetData() != Category->GetSection();
-			}
+			FText Message = LOCTEXT("DropCategoryOnSectionDragMessage", "Add {0} to section {1}");
+			Message = FText::FormatOrdered(Message, FText::FromString(DraggedItem->ToString()), FText::FromString(ToString()));
+			
+			Results.CanPerformMessage = Message;
+			Results.bCanPerform = GetData() != Category->GetSection();
 		}
-		else if(UNiagaraHierarchySection* Section = Cast<UNiagaraHierarchySection>(DraggedItem.Pin()->GetDataMutable()))
+	}
+	else if(UNiagaraHierarchySection* Section = Cast<UNiagaraHierarchySection>(DraggedItem->GetDataMutable()))
+	{
+		const bool bSameSection = GetData() == Section;
+
+		if(ItemDropZone == EItemDropZone::OntoItem && bSameSection == false)
 		{
-			if(ItemDropZone != EItemDropZone::OntoItem)
-			{
-				int32 DraggedItemIndex = GetHierarchyViewModel()->GetHierarchyDataRoot()->GetSectionIndex(Section);
-				bool bSameSection = GetData() == Section;
+			FText Message = LOCTEXT("CantDropSectionOnSectionDragMessage", "Can't drop section on a section.");
+			Results.CanPerformMessage = Message;
+			Results.bCanPerform = false;
+			return Results;
+		}
+		
+		int32 DraggedItemIndex = GetHierarchyViewModel()->GetHierarchyRoot()->GetSectionIndex(Section);
+		int32 InsertionIndex = GetHierarchyViewModel()->GetHierarchyRoot()->GetSectionIndex(GetDataMutable<UNiagaraHierarchySection>());
+		// we add 1 to the insertion index if it's below an item because we either want to insert at the current index to place the item above, or at current+1 for below
+		InsertionIndex += ItemDropZone == EItemDropZone::AboveItem ? -1 : 1;
 
-				int32 InsertionIndex = GetHierarchyViewModel()->GetHierarchyDataRoot()->GetSectionIndex(GetDataMutable<UNiagaraHierarchySection>());
-				// we add 1 to the insertion index if it's below an item because we either want to insert at the current index to place the item above, or at current+1 for below
-				InsertionIndex += ItemDropZone == EItemDropZone::AboveItem ? -1 : 1;
+		Results.bCanPerform = !bSameSection && DraggedItemIndex != InsertionIndex;
 
-				bAllowDrop = !bSameSection && DraggedItemIndex != InsertionIndex;			
+		if(Results.bCanPerform)
+		{
+			if(ItemDropZone == EItemDropZone::AboveItem)
+			{
+				FText Message = LOCTEXT("MoveSectionLeftDragMessage", "Move section {0} to the left");
+				Message = FText::FormatOrdered(Message, FText::FromString(DraggedItem->ToString()));
+				Results.CanPerformMessage = Message;
+			}
+			else if(ItemDropZone == EItemDropZone::BelowItem)
+			{
+				FText Message = LOCTEXT("MoveSectionRightDragMessage", "Move section {0} to the right");
+				Message = FText::FormatOrdered(Message, FText::FromString(DraggedItem->ToString()));
+				Results.CanPerformMessage = Message;
 			}
 		}
-	}
-	 
+	}	 
 
-	return bAllowDrop ? ItemDropZone : TOptional<EItemDropZone>();
+	return Results;
 }
 
-FReply FNiagaraHierarchySectionViewModel::OnDroppedOn(const FDragDropEvent& DragDropEvent, EItemDropZone ItemDropZone, TSharedPtr<FNiagaraHierarchyItemViewModelBase> Item)
+void FNiagaraHierarchySectionViewModel::OnDroppedOnInternal(TSharedPtr<FNiagaraHierarchyItemViewModelBase> DroppedItem, EItemDropZone ItemDropZone)
 {
-	if(DragDropEvent.GetOperation()->IsOfType<FNiagaraSectionDragDropOp>())
+	if(DroppedItem->GetData()->IsA<UNiagaraHierarchySection>())
 	{
-		TSharedPtr<FNiagaraSectionDragDropOp> SectionDragDropOp = DragDropEvent.GetOperationAs<FNiagaraSectionDragDropOp>();
-		TWeakPtr<FNiagaraHierarchySectionViewModel> DraggedItem = SectionDragDropOp->GetDraggedSection();
-		UNiagaraHierarchySection* DraggedSectionData = DraggedItem.Pin()->GetDataMutable<UNiagaraHierarchySection>();
+		UNiagaraHierarchySection* DraggedSectionData = DroppedItem->GetDataMutable<UNiagaraHierarchySection>();
 
-		int32 IndexOfThis = HierarchyViewModel->GetHierarchyDataRoot()->GetSectionData().Find(GetDataMutable<UNiagaraHierarchySection>());
-		int32 DraggedSectionIndex = HierarchyViewModel->GetHierarchyDataRoot()->GetSectionData().Find(DraggedItem.Pin()->GetDataMutable<UNiagaraHierarchySection>());
+		int32 IndexOfThis = HierarchyViewModel->GetHierarchyRoot()->GetSectionData().Find(GetDataMutable<UNiagaraHierarchySection>());
+		int32 DraggedSectionIndex = HierarchyViewModel->GetHierarchyRoot()->GetSectionData().Find(DraggedSectionData);
 
-		TArray<UNiagaraHierarchySection*>& SectionData = HierarchyViewModel->GetHierarchyDataRoot()->GetSectionDataMutable();
+		TArray<UNiagaraHierarchySection*>& SectionData = HierarchyViewModel->GetHierarchyRoot()->GetSectionDataMutable();
 		int32 Count = SectionData.Num();
 
 		bool bDropSucceeded = false;
@@ -885,20 +1250,14 @@ FReply FNiagaraHierarchySectionViewModel::OnDroppedOn(const FDragDropEvent& Drag
 		{
 			HierarchyViewModel->ForceFullRefresh();
 			HierarchyViewModel->OnHierarchyChanged().Broadcast();
-			return FReply::Handled();
 		}
-
-		return FReply::Unhandled();
 	}
-	else if(DragDropEvent.GetOperation()->IsOfType<FNiagaraHierarchyDragDropOp>())
+	else
 	{		
-		TSharedPtr<FNiagaraHierarchyDragDropOp> HierarchyDragDropOp = DragDropEvent.GetOperationAs<FNiagaraHierarchyDragDropOp>();
-		TWeakPtr<FNiagaraHierarchyItemViewModelBase> DraggedItem = HierarchyDragDropOp->GetDraggedItem();
-
-		if(UNiagaraHierarchyCategory* HierarchyCategory = DraggedItem.Pin()->GetDataMutable<UNiagaraHierarchyCategory>())
+		if(UNiagaraHierarchyCategory* HierarchyCategory = DroppedItem->GetDataMutable<UNiagaraHierarchyCategory>())
 		{
 			FScopedTransaction Transaction(LOCTEXT("Transaction_OnSectionDrop", "Moved category to section"));
-			HierarchyViewModel->GetHierarchyDataRoot()->Modify();
+			HierarchyViewModel->GetHierarchyRoot()->Modify();
 			
 			HierarchyCategory->SetSection(GetDataMutable<UNiagaraHierarchySection>());
 
@@ -911,31 +1270,27 @@ FReply FNiagaraHierarchySectionViewModel::OnDroppedOn(const FDragDropEvent& Drag
 			}
 
 			// we only need to reparent if the parent isn't already the root. This stops unnecessary reordering
-			if(DraggedItem.Pin()->GetParent() != HierarchyViewModel->GetHierarchyViewModelRoot())
+			if(DroppedItem->GetParent() != HierarchyViewModel->GetHierarchyRootViewModel())
 			{
-				HierarchyViewModel->GetHierarchyViewModelRoot()->ReparentToThis(DraggedItem.Pin());
+				HierarchyViewModel->GetHierarchyRootViewModel()->ReparentToThis(DroppedItem);
 			}
 			
 			HierarchyViewModel->RefreshHierarchyView();
 			HierarchyViewModel->OnHierarchyChanged().Broadcast();
-
-			return FReply::Handled();
 		}
 	}
-
-	return FReply::Unhandled();
 }
 
 void FNiagaraHierarchySectionViewModel::FinalizeInternal()
 {
-	if(HierarchyViewModel->GetActiveSection() == AsShared())
+	if(HierarchyViewModel->GetActiveHierarchySection() == AsShared())
 	{
-		HierarchyViewModel->SetActiveSection(nullptr);
+		HierarchyViewModel->SetActiveHierarchySection(nullptr);
 	}
 
 	// we make sure to reset all categories' section entry that were referencing this section
 	TArray<UNiagaraHierarchyCategory*> AllCategories;
-	HierarchyViewModel->GetHierarchyDataRoot()->GetChildrenOfType<UNiagaraHierarchyCategory>(AllCategories, true);
+	HierarchyViewModel->GetHierarchyRoot()->GetChildrenOfType<UNiagaraHierarchyCategory>(AllCategories, true);
 
 	for(UNiagaraHierarchyCategory* Category : AllCategories)
 	{
@@ -946,104 +1301,78 @@ void FNiagaraHierarchySectionViewModel::FinalizeInternal()
 	}
 }
 
-TOptional<EItemDropZone> FNiagaraHierarchyItemViewModel::OnCanAcceptDropInternal(TSharedPtr<FDragDropOperation> DragDropOp, EItemDropZone ItemDropZone)
+::FNiagaraHierarchyItemViewModelBase::FCanPerformActionResults FNiagaraHierarchyItemViewModel::CanDropOnInternal(TSharedPtr<FNiagaraHierarchyItemViewModelBase> DraggedItem, EItemDropZone ItemDropZone)
 {
 	bool bAllowDrop = false;
+	
+	TSharedPtr<FNiagaraHierarchyItemViewModelBase> SourceDropItem = DraggedItem;
+	TSharedPtr<FNiagaraHierarchyItemViewModelBase> TargetDropItem = AsShared();
 
-	if(DragDropOp->IsOfType<FNiagaraHierarchyDragDropOp>())
+	// we only allow drops if some general conditions are fulfilled
+	if(SourceDropItem->GetData() != TargetDropItem->GetData() &&
+		(!SourceDropItem->HasParent(TargetDropItem, false) || ItemDropZone != EItemDropZone::OntoItem)  &&
+		!TargetDropItem->HasParent(SourceDropItem, true))
 	{
-		TSharedPtr<FNiagaraHierarchyDragDropOp> HierarchyDragDropOp = StaticCastSharedPtr<FNiagaraHierarchyDragDropOp>(DragDropOp);
-		TSharedPtr<FNiagaraHierarchyItemViewModelBase> SourceDropItem = HierarchyDragDropOp->GetDraggedItem().Pin();
-		TSharedPtr<FNiagaraHierarchyItemViewModelBase> TargetDropItem = AsShared();
-
-		// we only allow drops if some general conditions are fulfilled
-		if(SourceDropItem->GetData() != TargetDropItem->GetData() &&
-			(!SourceDropItem->HasParent(TargetDropItem, false) || ItemDropZone != EItemDropZone::OntoItem)  &&
-			!TargetDropItem->HasParent(SourceDropItem, true))
-		{
-			// items can be generally be dropped above/below other items
-			bAllowDrop = (SourceDropItem->GetData()->IsA<UNiagaraHierarchyItem>() && ItemDropZone != EItemDropZone::OntoItem);
-		}
+		// items can be generally be dropped above/below other items
+		bAllowDrop = (SourceDropItem->GetData()->IsA<UNiagaraHierarchyItem>() && ItemDropZone != EItemDropZone::OntoItem);
 	}
 
-	return bAllowDrop ? ItemDropZone : TOptional<EItemDropZone>();
+	return bAllowDrop;
 }
 
-FReply FNiagaraHierarchyItemViewModel::OnDroppedOn(const FDragDropEvent& DragDropEvent, EItemDropZone ItemDropZone, TSharedPtr<FNiagaraHierarchyItemViewModelBase> Item)
+void FNiagaraHierarchyItemViewModel::OnDroppedOnInternal(TSharedPtr<FNiagaraHierarchyItemViewModelBase> DroppedItem, EItemDropZone ItemDropZone)
 {
 	FScopedTransaction Transaction(LOCTEXT("Transaction_MovedItem", "Moved an item in the hierarchy"));
-	HierarchyViewModel->GetHierarchyDataRoot()->Modify();
+	HierarchyViewModel->GetHierarchyRoot()->Modify();
 	
-	if(DragDropEvent.GetOperation()->IsOfType<FNiagaraHierarchyDragDropOp>())
+	bool bDropSucceeded = false;
+	if(ItemDropZone == EItemDropZone::AboveItem)
 	{
-		TSharedPtr<FNiagaraHierarchyDragDropOp> HierarchyDragDropOp = DragDropEvent.GetOperationAs<FNiagaraHierarchyDragDropOp>();
-		bool bDropSucceeded = false;
-		if(ItemDropZone == EItemDropZone::AboveItem)
-		{
-			int32 IndexOfThis = Parent.Pin()->FindIndexOfDataChild(AsShared());
-
-			if(HierarchyDragDropOp->GetIsFromSourceList())
-			{
-				Parent.Pin()->DuplicateToThis(HierarchyDragDropOp->GetDraggedItem().Pin(), FMath::Max(IndexOfThis, 0));
-			}
-			else
-			{
-				Parent.Pin()->ReparentToThis(HierarchyDragDropOp->GetDraggedItem().Pin(), FMath::Max(IndexOfThis, 0));
-			}
+		int32 IndexOfThis = Parent.Pin()->FindIndexOfDataChild(AsShared());
 
-			bDropSucceeded = true;
-		}
-		else if(ItemDropZone == EItemDropZone::BelowItem)
+		if(DroppedItem->IsForHierarchy() == false)
 		{
-			int32 IndexOfThis = Parent.Pin()->FindIndexOfDataChild(AsShared());
-			Parent.Pin()->ReparentToThis(HierarchyDragDropOp->GetDraggedItem().Pin(), FMath::Min(IndexOfThis+1, Parent.Pin()->GetChildren().Num()));
-
-			bDropSucceeded = true;
+			Parent.Pin()->DuplicateToThis(DroppedItem, FMath::Max(IndexOfThis, 0));
 		}
-
-		if(bDropSucceeded)
+		else
 		{
-			HierarchyViewModel->RefreshHierarchyView();
-			HierarchyViewModel->RefreshSourceView();
-			return FReply::Handled();
+			Parent.Pin()->ReparentToThis(DroppedItem, FMath::Max(IndexOfThis, 0));
 		}
+
+		bDropSucceeded = true;
 	}
+	else if(ItemDropZone == EItemDropZone::BelowItem)
+	{
+		int32 IndexOfThis = Parent.Pin()->FindIndexOfDataChild(AsShared());
+		Parent.Pin()->ReparentToThis(DroppedItem, FMath::Min(IndexOfThis+1, Parent.Pin()->GetChildren().Num()));
 
-	return FReply::Unhandled();
+		bDropSucceeded = true;
+	}
+
+	if(bDropSucceeded)
+	{
+		HierarchyViewModel->RefreshHierarchyView();
+		HierarchyViewModel->RefreshSourceView();
+	}
+	else
+	{
+		Transaction.Cancel();
+	}
 }
 
-FReply FNiagaraHierarchyCategoryViewModel::OnDroppedOn(const FDragDropEvent& DragDropEvent, EItemDropZone ItemDropZone, TSharedPtr<FNiagaraHierarchyItemViewModelBase> Item)
+void FNiagaraHierarchyCategoryViewModel::OnDroppedOnInternal(TSharedPtr<FNiagaraHierarchyItemViewModelBase> DroppedItem, EItemDropZone ItemDropZone)
 {
-	if(DragDropEvent.GetOperation()->IsOfType<FNiagaraHierarchyDragDropOp>())
+	FScopedTransaction Transaction(LOCTEXT("Transaction_OnCategoryDrop", "Dropped item on/above/below category"));
+	HierarchyViewModel->GetHierarchyRoot()->Modify();
+	
+	if(UNiagaraHierarchyCategory* Category = DroppedItem->GetDataMutable<UNiagaraHierarchyCategory>())
 	{
-		FScopedTransaction Transaction(LOCTEXT("Transaction_OnCategoryDrop", "Dropped item on/above/below category"));
-		HierarchyViewModel->GetHierarchyDataRoot()->Modify();
-		
-		TSharedPtr<FNiagaraHierarchyDragDropOp> HierarchyDragDropOp = DragDropEvent.GetOperationAs<FNiagaraHierarchyDragDropOp>();
-		TSharedPtr<FNiagaraHierarchyItemViewModelBase> DraggedItem = HierarchyDragDropOp->GetDraggedItem().Pin();
-
-		if(UNiagaraHierarchyCategory* Category = DraggedItem->GetDataMutable<UNiagaraHierarchyCategory>())
+		// if we are dragging a category above/below another category and the new parent is going to be the root, we update its section to the active section
+		if(ItemDropZone != EItemDropZone::OntoItem)
 		{
-			// if we are dragging a category above/below another category and the new parent is going to be the root, we update its section to the active section
-			if(ItemDropZone != EItemDropZone::OntoItem)
-			{
-				if(Parent.IsValid() && Parent == HierarchyViewModel->GetHierarchyViewModelRoot())
-				{
-					Category->SetSection(HierarchyViewModel->GetActiveSectionData());
-
-					// we null out any sections for all contained categories
-					TArray<UNiagaraHierarchyCategory*> AllChildCategories;
-					Category->GetChildrenOfType<UNiagaraHierarchyCategory>(AllChildCategories, true);
-					for(UNiagaraHierarchyCategory* ChildCategory : AllChildCategories)
-					{
-						ChildCategory->SetSection(nullptr);
-					}
-				}				
-			}
-			// if we are dragging a category onto another category, we null out its section instead
-			else
+			if(Parent.IsValid() && Parent == HierarchyViewModel->GetHierarchyRootViewModel())
 			{
-				Category->SetSection(nullptr);
+				Category->SetSection(HierarchyViewModel->GetActiveHierarchySectionData());
 
 				// we null out any sections for all contained categories
 				TArray<UNiagaraHierarchyCategory*> AllChildCategories;
@@ -1052,59 +1381,64 @@ FReply FNiagaraHierarchyCategoryViewModel::OnDroppedOn(const FDragDropEvent& Dra
 				{
 					ChildCategory->SetSection(nullptr);
 				}
-			}			
+			}				
 		}
-
-		// the actual moving of the item happens here
-		if(ItemDropZone == EItemDropZone::OntoItem)
+		// if we are dragging a category onto another category, we null out its section instead
+		else
 		{
-			if(HierarchyDragDropOp->GetIsFromSourceList())
-			{
-				DuplicateToThis(HierarchyDragDropOp->GetDraggedItem().Pin());
-			}
-			else
+			Category->SetSection(nullptr);
+
+			// we null out any sections for all contained categories
+			TArray<UNiagaraHierarchyCategory*> AllChildCategories;
+			Category->GetChildrenOfType<UNiagaraHierarchyCategory>(AllChildCategories, true);
+			for(UNiagaraHierarchyCategory* ChildCategory : AllChildCategories)
 			{
-				ReparentToThis(HierarchyDragDropOp->GetDraggedItem().Pin());
+				ChildCategory->SetSection(nullptr);
 			}
+		}			
+	}
+
+	// the actual moving of the item happens here
+	if(ItemDropZone == EItemDropZone::OntoItem)
+	{
+		if(DroppedItem->IsForHierarchy() == false)
+		{
+			DuplicateToThis(DroppedItem);
 		}
-		else if(ItemDropZone == EItemDropZone::AboveItem)
+		else
 		{
-			int32 IndexOfThis = Parent.Pin()->FindIndexOfDataChild(AsShared());
-			if(HierarchyDragDropOp->GetIsFromSourceList())
-			{
-				Parent.Pin()->DuplicateToThis(HierarchyDragDropOp->GetDraggedItem().Pin(), FMath::Max(IndexOfThis, 0));
-			}
-			else
-			{				
-				Parent.Pin()->ReparentToThis(HierarchyDragDropOp->GetDraggedItem().Pin(), FMath::Max(IndexOfThis, 0));
-			}
+			ReparentToThis(DroppedItem);
 		}
-		else if(ItemDropZone == EItemDropZone::BelowItem)
+	}
+	else if(ItemDropZone == EItemDropZone::AboveItem)
+	{
+		int32 IndexOfThis = Parent.Pin()->FindIndexOfDataChild(AsShared());
+		if(DroppedItem->IsForHierarchy() == false)
 		{
-			int32 IndexOfThis = Parent.Pin()->FindIndexOfDataChild(AsShared());
-			if(HierarchyDragDropOp->GetIsFromSourceList())
-			{
-				Parent.Pin()->DuplicateToThis(HierarchyDragDropOp->GetDraggedItem().Pin(), FMath::Min(IndexOfThis+1, Parent.Pin()->GetChildren().Num()));
-			}
-			else
-			{
-				Parent.Pin()->ReparentToThis(HierarchyDragDropOp->GetDraggedItem().Pin(), FMath::Min(IndexOfThis+1, Parent.Pin()->GetChildren().Num()));
-			}
+			Parent.Pin()->DuplicateToThis(DroppedItem, FMath::Max(IndexOfThis, 0));
+		}
+		else
+		{				
+			Parent.Pin()->ReparentToThis(DroppedItem, FMath::Max(IndexOfThis, 0));
+		}
+	}
+	else if(ItemDropZone == EItemDropZone::BelowItem)
+	{
+		int32 IndexOfThis = Parent.Pin()->FindIndexOfDataChild(AsShared());
+		if(DroppedItem->IsForHierarchy() == false)
+		{
+			Parent.Pin()->DuplicateToThis(DroppedItem, FMath::Min(IndexOfThis+1, Parent.Pin()->GetChildren().Num()));
+		}
+		else
+		{
+			Parent.Pin()->ReparentToThis(DroppedItem, FMath::Min(IndexOfThis+1, Parent.Pin()->GetChildren().Num()));
 		}
-		
-		// todo (me)
-		// refreshing both views whenever we drag something seems wasteful
-		HierarchyViewModel->RefreshHierarchyView();
-		HierarchyViewModel->RefreshSourceView();
-		return FReply::Handled();
 	}
-
-	return FReply::Unhandled();
 }
 
 void UNiagaraHierarchyViewModelBase::AddCategory() const
 {
-	TSharedPtr<FNiagaraHierarchyItemViewModelBase> ViewModel = HierarchyViewModelRoot->AddNewItem(UNiagaraHierarchyCategory::StaticClass());
+	TSharedPtr<FNiagaraHierarchyItemViewModelBase> ViewModel = HierarchyRootViewModel->AddNewItem(UNiagaraHierarchyCategory::StaticClass());
 
 	if(UNiagaraHierarchyCategory* Category = ViewModel->GetDataMutable<UNiagaraHierarchyCategory>())
 	{
@@ -1120,22 +1454,96 @@ void UNiagaraHierarchyViewModelBase::AddCategory() const
 		Category->SetCategoryName(FNiagaraUtilities::GetUniqueName(FName("New Category"), CategoryNames));
 		
 		// we only set the section property if the current section isn't set to "All"
-		Category->SetSection(GetActiveSectionData());
+		Category->SetSection(GetActiveHierarchySectionData());
 		
 		RefreshHierarchyView();
 
-		OnItemAddedDelegate.Broadcast(ViewModel);
+		OnItemAddedDelegate.ExecuteIfBound(ViewModel);
 		OnHierarchyChangedDelegate.Broadcast();
 	}
 }
 
 void UNiagaraHierarchyViewModelBase::AddSection() const
 {
-	TSharedPtr<FNiagaraHierarchySectionViewModel> HierarchySectionViewModel = HierarchyViewModelRoot->AddNewSection();
-	OnItemAddedDelegate.Broadcast(HierarchySectionViewModel);
+	TSharedPtr<FNiagaraHierarchySectionViewModel> HierarchySectionViewModel = HierarchyRootViewModel->AddSection();
+	OnItemAddedDelegate.ExecuteIfBound(HierarchySectionViewModel);
 	OnHierarchyChangedDelegate.Broadcast();
 }
 
+void UNiagaraHierarchyViewModelBase::DeleteItemWithIdentity(FNiagaraHierarchyIdentity Identity)
+{
+	FScopedTransaction Transaction(LOCTEXT("Transaction_DeleteItem", "Deleted hierarchy item"));
+	HierarchyRoot->Modify();
+
+	bool bItemDeleted = false;
+	if(TSharedPtr<FNiagaraHierarchyItemViewModelBase> ViewModel = HierarchyRootViewModel->FindViewModelForChild(Identity, true))
+	{
+		if(ViewModel->CanDelete())
+		{
+			ViewModel->Delete();
+			bItemDeleted = true;
+		}
+	}
+	
+	TArray<TSharedPtr<FNiagaraHierarchySectionViewModel>> SectionViewModels = HierarchyRootViewModel->GetSectionViewModels();
+	for(TSharedPtr<FNiagaraHierarchySectionViewModel> SectionViewModel : SectionViewModels)
+	{
+		if(SectionViewModel->GetData()->GetPersistentIdentity() == Identity && SectionViewModel->CanDelete())
+		{
+			SectionViewModel->Delete();
+			bItemDeleted = true;
+		}
+	}
+
+	if(bItemDeleted)
+	{
+		HierarchyRootViewModel->SyncViewModelsToData();
+		OnHierarchyChangedDelegate.Broadcast();
+	}
+}
+
+void UNiagaraHierarchyViewModelBase::DeleteItemsWithIdentities(TArray<FNiagaraHierarchyIdentity> Identities)
+{
+	FScopedTransaction Transaction(LOCTEXT("Transaction_DeleteMultipleItems", "Deleted hierarchy items"));
+	HierarchyRoot->Modify();
+
+	bool bAnyItemsDeleted = false;
+	for(FNiagaraHierarchyIdentity& Identity : Identities)
+	{
+		if(TSharedPtr<FNiagaraHierarchyItemViewModelBase> ViewModel = HierarchyRootViewModel->FindViewModelForChild(Identity, true))
+		{
+			if(ViewModel->CanDelete())
+			{
+				ViewModel->Delete();
+				bAnyItemsDeleted = true;
+				continue;
+			}
+		}
+
+		TArray<TSharedPtr<FNiagaraHierarchySectionViewModel>> SectionViewModels = HierarchyRootViewModel->GetSectionViewModels();
+		for(TSharedPtr<FNiagaraHierarchySectionViewModel> SectionViewModel : SectionViewModels)
+		{
+			if(SectionViewModel->GetData()->GetPersistentIdentity() == Identity && SectionViewModel->CanDelete())
+			{
+				SectionViewModel->Delete();
+				bAnyItemsDeleted = true;
+				continue;
+			}
+		}
+	}
+
+	if(bAnyItemsDeleted)
+	{
+		HierarchyRootViewModel->SyncViewModelsToData();
+		OnHierarchyChangedDelegate.Broadcast();
+	}
+}
+
+void UNiagaraHierarchyViewModelBase::NavigateToItemInHierarchy(const FNiagaraHierarchyIdentity& NiagaraHierarchyIdentity)
+{
+	OnNavigateToItemInHierarchyRequestedDelegate.ExecuteIfBound(NiagaraHierarchyIdentity);
+}
+
 TSharedPtr<SWidget> FNiagaraHierarchyDragDropOp::GetDefaultDecorator() const
 {
 	TSharedRef<SWidget> CustomDecorator = CreateCustomDecorator();
