Summary View 2.0

diff --git a/Engine/Plugins/FX/Niagara/Source/NiagaraEditor/Private/ViewModels/HierarchyEditor/NiagaraSummaryViewViewModel.cpp b/Engine/Plugins/FX/Niagara/Source/NiagaraEditor/Private/ViewModels/HierarchyEditor/NiagaraSummaryViewViewModel.cpp
new file mode 100644
--- /dev/null
+++ b/Engine/Plugins/FX/Niagara/Source/NiagaraEditor/Private/ViewModels/HierarchyEditor/NiagaraSummaryViewViewModel.cpp
@@ -0,0 +1,1928 @@
+// Copyright Epic Games, Inc. All Rights Reserved.
+
+#include "ViewModels/HierarchyEditor/NiagaraSummaryViewViewModel.h"
+#include "NiagaraConstants.h"
+#include "NiagaraEditorStyle.h"
+#include "NiagaraEmitterEditorData.h"
+#include "NiagaraScriptMergeManager.h"
+#include "NiagaraScriptVariable.h"
+#include "NiagaraSystem.h"
+#include "NiagaraSimulationStageBase.h"
+#include "SDropTarget.h"
+#include "ViewModels/NiagaraEmitterViewModel.h"
+#include "ViewModels/NiagaraScriptGraphViewModel.h"
+#include "ViewModels/NiagaraScriptViewModel.h"
+#include "ViewModels/NiagaraSystemViewModel.h"
+#include "Widgets/Text/SInlineEditableTextBlock.h"
+
+#define LOCTEXT_NAMESPACE "NiagaraSummaryViewHierarchyEditor"
+
+bool GetIsFromBaseEmitter(const FVersionedNiagaraEmitter& Emitter, FNiagaraHierarchyIdentity SummaryItemIdentity)
+{
+	TSharedRef<FNiagaraScriptMergeManager> MergeManager = FNiagaraScriptMergeManager::Get();
+	return MergeManager->DoesSummaryItemExistInBase(Emitter, SummaryItemIdentity);
+}
+
+void UNiagaraHierarchyModule::Initialize(const UNiagaraNodeFunctionCall& InModuleNode)
+{
+	FNiagaraHierarchyIdentity ModuleIdentity;
+	ModuleIdentity.Guids.Add(InModuleNode.NodeGuid);
+	SetIdentity(ModuleIdentity);
+}
+
+void UNiagaraHierarchyModuleInput::Initialize(const UNiagaraNodeFunctionCall& InModuleNode, FGuid InputGuid)
+{
+	FNiagaraHierarchyIdentity InputIdentity;
+	InputIdentity.Guids.Add(InModuleNode.NodeGuid);
+	InputIdentity.Guids.Add(InputGuid);
+	SetIdentity(InputIdentity);
+}
+
+void UNiagaraHierarchyAssignmentInput::Initialize(const UNiagaraNodeAssignment& AssignmentNode, FName AssignmentTarget)
+{
+	FNiagaraHierarchyIdentity InputIdentity;
+	InputIdentity.Guids.Add(AssignmentNode.NodeGuid);
+	InputIdentity.Names.Add(AssignmentTarget);
+	SetIdentity(InputIdentity);
+}
+
+void UNiagaraHierarchyRenderer::Initialize(const UNiagaraRendererProperties& Renderer)
+{
+	FNiagaraHierarchyIdentity RendererIdentity;
+	RendererIdentity.Guids.Add(Renderer.GetMergeId());
+	SetIdentity(RendererIdentity);
+}
+
+void UNiagaraHierarchyEventHandler::Initialize(const FNiagaraEventScriptProperties& EventHandler)
+{
+	FNiagaraHierarchyIdentity EventHandlerIdentity;
+	EventHandlerIdentity.Guids.Add(EventHandler.Script->GetUsageId());
+	SetIdentity(EventHandlerIdentity);
+}
+
+void UNiagaraHierarchyEventHandlerProperties::Initialize(const FNiagaraEventScriptProperties& EventHandler)
+{
+	SetIdentity(MakeIdentity(EventHandler));
+}
+
+FNiagaraHierarchyIdentity UNiagaraHierarchyEventHandlerProperties::MakeIdentity(const FNiagaraEventScriptProperties& EventHandler)
+{
+	FNiagaraHierarchyIdentity Identity;
+	Identity.Guids.Add(EventHandler.Script->GetUsageId());
+	Identity.Names.Add(TEXT("Category"));
+	Identity.Names.Add(TEXT("Properties"));
+	return Identity;
+}
+
+void UNiagaraHierarchySimStage::Initialize(const UNiagaraSimulationStageBase& SimStage)
+{
+	FNiagaraHierarchyIdentity SimStageIdentity;
+	SimStageIdentity.Guids.Add(SimStage.GetMergeId());
+	SetIdentity(SimStageIdentity);
+}
+
+void UNiagaraHierarchySimStageProperties::Initialize(const UNiagaraSimulationStageBase& SimStage)
+{
+	SetIdentity(MakeIdentity(SimStage));
+}
+
+FNiagaraHierarchyIdentity UNiagaraHierarchySimStageProperties::MakeIdentity(const UNiagaraSimulationStageBase& SimStage)
+{
+	FNiagaraHierarchyIdentity SimStagePropertiesIdentity;
+	SimStagePropertiesIdentity.Guids.Add(SimStage.GetMergeId());
+	SimStagePropertiesIdentity.Names.Add(FName("Category"));
+	SimStagePropertiesIdentity.Names.Add("Properties");
+	return SimStagePropertiesIdentity;
+}
+
+void UNiagaraHierarchyObjectProperty::Initialize(FGuid ObjectGuid, FString PropertyName)
+{
+	FNiagaraHierarchyIdentity PropertyIdentity;
+	PropertyIdentity.Guids.Add(ObjectGuid);
+	PropertyIdentity.Names.Add(FName(PropertyName));
+	SetIdentity(PropertyIdentity);
+}
+
+void UNiagaraSummaryViewViewModel::Initialize(TSharedRef<FNiagaraEmitterViewModel> EmitterViewModel)
+{
+	EmitterViewModelWeak = EmitterViewModel;
+	EmitterViewModel->OnScriptGraphChanged().AddUObject(this, &UNiagaraSummaryViewViewModel::OnScriptGraphChanged);
+	EmitterViewModel->GetEmitter().Emitter->OnRenderersChanged().AddUObject(this, &UNiagaraSummaryViewViewModel::OnRenderersChanged);
+	EmitterViewModel->GetEmitter().Emitter->OnSimStagesChanged().AddUObject(this, &UNiagaraSummaryViewViewModel::OnSimStagesChanged);
+	EmitterViewModel->GetEmitter().Emitter->OnEventHandlersChanged().AddUObject(this, &UNiagaraSummaryViewViewModel::OnEventHandlersChanged);
+
+	UNiagaraHierarchyViewModelBase::Initialize();
+}
+
+void UNiagaraSummaryViewViewModel::FinalizeInternal()
+{
+	GetEmitterViewModel()->OnScriptGraphChanged().RemoveAll(this);
+
+	if(GetEmitterViewModel()->GetEmitter().Emitter != nullptr)
+	{
+		GetEmitterViewModel()->GetEmitter().Emitter->OnRenderersChanged().RemoveAll(this);
+		GetEmitterViewModel()->GetEmitter().Emitter->OnSimStagesChanged().RemoveAll(this);
+		GetEmitterViewModel()->GetEmitter().Emitter->OnEventHandlersChanged().RemoveAll(this);
+	}
+}
+
+TSharedRef<FNiagaraEmitterViewModel> UNiagaraSummaryViewViewModel::GetEmitterViewModel() const
+{
+	TSharedPtr<FNiagaraEmitterViewModel> EmitterViewModel = EmitterViewModelWeak.Pin();
+	checkf(EmitterViewModel.IsValid(), TEXT("Emitter view model destroyed before summary hierarchy view model."));
+	return EmitterViewModel.ToSharedRef();
+}
+
+TWeakObjectPtr<UNiagaraNodeFunctionCall> FNiagaraFunctionViewModel::GetFunctionCallNode() const
+{
+	TWeakObjectPtr<UNiagaraHierarchyViewModelBase> ViewModel = GetHierarchyViewModel();
+	UNiagaraSummaryViewViewModel* SummaryViewModel = CastChecked<UNiagaraSummaryViewViewModel>(ViewModel.Get());
+	return SummaryViewModel->GetFunctionCallNode(GetData()->GetPersistentIdentity().Guids[0]);
+}
+
+void FNiagaraFunctionViewModel::OnScriptApplied(UNiagaraScript* NiagaraScript, FGuid Guid)
+{
+	if(GetFunctionCallNode()->FunctionScript == NiagaraScript)
+	{
+		RefreshChildrenInputs(true);
+		SyncViewModelsToData();
+	}
+}
+
+void FNiagaraFunctionViewModel::ClearCache() const
+{
+	TWeakObjectPtr<UNiagaraHierarchyViewModelBase> ViewModel = GetHierarchyViewModel();
+	UNiagaraSummaryViewViewModel* SummaryViewModel = CastChecked<UNiagaraSummaryViewViewModel>(ViewModel.Get());
+	SummaryViewModel->ClearFunctionCallNodeCache(GetData()->GetPersistentIdentity().Guids[0]);
+}
+
+FString FNiagaraFunctionViewModel::ToString() const
+{
+	if(GetFunctionCallNode().IsValid())
+	{
+		return GetFunctionCallNode()->GetNodeTitle(ENodeTitleType::ListView).ToString();
+	}
+
+	return TEXT("Unknown");
+}
+
+FNiagaraHierarchyItemViewModelBase::FCanPerformActionResults FNiagaraFunctionViewModel::IsEditableByUser()
+{
+	if(bIsDynamicInput)
+	{
+		FCanPerformActionResults CanEditResults(false);
+		CanEditResults.CanPerformMessage = LOCTEXT("DynamicInputCantBeDragged", "You can not drag entire Dynamic Inputs. Either drag the entire module input, or individual inputs of the Dynamic Input");
+		return CanEditResults;
+	}
+	
+	FCanPerformActionResults CanEditResults(IsFromBaseEmitter() == false);
+	CanEditResults.CanPerformMessage = CanEditResults.bCanPerform == false ? LOCTEXT("ModuleIsFromBaseEmitter", "This module was added in the parent emitter and can not be edited.") : FText::GetEmpty();
+	return CanEditResults;
+}
+
+bool FNiagaraFunctionViewModel::IsFromBaseEmitter() const
+{
+	if(!IsFromBaseEmitterCache.IsSet())
+	{
+		UNiagaraSummaryViewViewModel* ViewModel = Cast<UNiagaraSummaryViewViewModel>(GetHierarchyViewModel());
+		IsFromBaseEmitterCache = GetIsFromBaseEmitter(ViewModel->GetEmitterViewModel()->GetEmitter(), GetData()->GetPersistentIdentity());
+	}
+
+	return IsFromBaseEmitterCache.GetValue();
+}
+
+void FNiagaraFunctionViewModel::Initialize()
+{
+	if(GetFunctionCallNode().IsValid())
+	{
+		OnScriptAppliedHandle = FNiagaraEditorModule::Get().OnScriptApplied().AddLambda([this](UNiagaraScript* Script, FGuid ScriptVersion)
+		{
+			if(GetFunctionCallNode()->FunctionScript == Script)
+			{
+				SyncViewModelsToData();
+			}
+		});
+
+		// determine whether this represents a dynamic input or a module by checking if the output pin of this node is a parameter map.
+		UEdGraphPin* OutputPin = GetFunctionCallNode()->GetOutputPin(0);
+		bIsDynamicInput = UEdGraphSchema_Niagara::PinToTypeDefinition(OutputPin) != FNiagaraTypeDefinition::GetParameterMapDef();
+	}
+}
+
+FNiagaraFunctionViewModel::~FNiagaraFunctionViewModel()
+{
+	if(OnScriptAppliedHandle.IsValid())
+	{
+		FNiagaraEditorModule::Get().OnScriptApplied().Remove(OnScriptAppliedHandle);
+		OnScriptAppliedHandle.Reset();
+	}
+}
+
+void FNiagaraFunctionViewModel::RefreshChildrenDataInternal()
+{
+	RefreshChildrenInputs(false);
+}
+
+void FNiagaraFunctionViewModel::RefreshChildrenInputs(bool bClearCache) const
+{
+	TWeakObjectPtr<UNiagaraNodeFunctionCall> FunctionNodeWeak = GetFunctionCallNode();
+	if(FunctionNodeWeak.IsValid())
+	{
+		UNiagaraNodeFunctionCall* FunctionNode = FunctionNodeWeak.Get();
+		UNiagaraNodeAssignment* AsAssignmentNode = Cast<UNiagaraNodeAssignment>(FunctionNode);
+		
+		if(UNiagaraGraph* AssetGraph = FunctionNode->GetCalledGraph())
+		{
+			// if it's not an assignment node, it's a module node
+			if(AsAssignmentNode == nullptr)
+			{
+				TArray<FNiagaraVariable> Variables;
+				AssetGraph->GetAllVariables(Variables);
+
+				TMap<FGuid, FNiagaraVariable> VariableGuidMap;
+				TMap<FGuid, FNiagaraVariableMetaData> VariableGuidMetadataMap;
+				for(const FNiagaraVariable& Variable : Variables)
+				{
+					// we create an input for most top level static switches & module inputs
+					bool bIsModuleInput = Variable.IsInNameSpace(FNiagaraConstants::ModuleNamespaceString);
+					TOptional<bool> bIsStaticSwitchInputOptional = AssetGraph->IsStaticSwitch(Variable);
+					if(!bIsModuleInput && !bIsStaticSwitchInputOptional.Get(false))
+					{
+						continue;
+					}
+					
+					TOptional<FNiagaraVariableMetaData> VariableMetaData = AssetGraph->GetMetaData(Variable);
+					// we don't show inline edit condition attributes
+					if(VariableMetaData->bInlineEditConditionToggle == false)
+					{
+						VariableGuidMap.Add(VariableMetaData->GetVariableGuid(), Variable);
+						VariableGuidMetadataMap.Add(VariableMetaData->GetVariableGuid(), VariableMetaData.GetValue());
+					}
+				}
+				
+				TArray<FGuid> VariableGuids;
+				VariableGuidMap.GenerateKeyArray(VariableGuids);
+				VariableGuids.Sort([&](const FGuid& GuidA, const FGuid& GuidB)
+				{
+					if(VariableGuidMetadataMap[GuidA].bAdvancedDisplay != VariableGuidMetadataMap[GuidB].bAdvancedDisplay)
+					{
+						if(VariableGuidMetadataMap[GuidA].bAdvancedDisplay)
+						{
+							return false;
+						}
+						else
+						{
+							return true;
+						}
+					}
+					
+					return VariableGuidMetadataMap[GuidA].EditorSortPriority < VariableGuidMetadataMap[GuidB].EditorSortPriority;
+				});
+				
+				for(const FGuid& VariableGuid : VariableGuids)
+				{
+					FNiagaraHierarchyIdentity SearchedChildIdentity;
+					SearchedChildIdentity.Guids.Add(FunctionNode->NodeGuid);
+					SearchedChildIdentity.Guids.Add(VariableGuid);
+					const bool bChildExists = GetData()->GetChildren().ContainsByPredicate([SearchedChildIdentity](UNiagaraHierarchyItemBase* CandidateChild)
+					{
+						return CandidateChild->GetPersistentIdentity() == SearchedChildIdentity;
+					});
+				
+					if(bChildExists == false)
+					{
+						UNiagaraHierarchyModuleInput* ModuleInput = GetDataMutable()->AddChild<UNiagaraHierarchyModuleInput>();
+						ModuleInput->Initialize(*FunctionNode, VariableGuid);
+					}
+				}
+			}
+			else
+			{
+				UNiagaraNodeAssignment* AssignmentNode = Cast<UNiagaraNodeAssignment>(FunctionNode);
+
+				for(const FNiagaraVariable& Variable : AssignmentNode->GetAssignmentTargets())
+				{
+					const bool bChildExists = GetDataMutable()->GetChildrenMutable().ContainsByPredicate([Variable, bClearCache, this](UNiagaraHierarchyItemBase* Candidate)
+					{
+						if(UNiagaraHierarchyAssignmentInput* AssignmentInput = Cast<UNiagaraHierarchyAssignmentInput>(Candidate))
+						{
+							return Variable.GetName() == AssignmentInput->GetPersistentIdentity().Names[0];
+						}
+
+						return false;						
+					});
+
+					if(bChildExists == false)
+					{
+						UNiagaraHierarchyAssignmentInput* AssignmentInput = GetDataMutable()->AddChild<UNiagaraHierarchyAssignmentInput>();
+						AssignmentInput->Initialize(*AsAssignmentNode, Variable.GetName());
+					}
+				}
+			}
+		}
+	}
+}
+
+FNiagaraHierarchyItemViewModelBase::FCanPerformActionResults FNiagaraFunctionViewModel::CanDropOnInternal(TSharedPtr<FNiagaraHierarchyItemViewModelBase> NiagaraHierarchyItemViewModelBase, EItemDropZone ItemDropZone)
+{
+	if(IsEditableByUser().bCanPerform == false)
+	{
+		return false;
+	}
+	
+	// we don't allow any items to be added directly onto the module as it's self managing
+	if(ItemDropZone == EItemDropZone::OntoItem)
+	{
+		FCanPerformActionResults Results(false);
+		Results.CanPerformMessage = LOCTEXT("CanDropOnModuleDragMessage", "You can not add any items to a module directly. Please create a category, which can contain arbitrary items.");
+		return Results;
+	}
+	
+	return FNiagaraHierarchyItemViewModel::CanDropOnInternal(NiagaraHierarchyItemViewModelBase, ItemDropZone);
+}
+
+const UNiagaraHierarchySection* FNiagaraFunctionViewModel::GetSectionInternal() const
+{
+	if(bIsForHierarchy)
+	{
+		return nullptr;
+	}
+
+	return Section.IsValid() ? Section.Get() : nullptr;
+}
+
+FNiagaraHierarchyItemViewModelBase::FCanPerformActionResults FNiagaraModuleInputViewModel::CanDropOnInternal(TSharedPtr<FNiagaraHierarchyItemViewModelBase> DraggedItem, EItemDropZone ItemDropZone)
+{
+	// if the input isn't editable, we don't allow any drops on/above/below the item.
+	// Even though it technically works, the merge process will only re-add the item at the end and not preserve order so there is no point in allowing dropping above/below
+	if(IsEditableByUser().bCanPerform == false)
+	{
+		return false;
+	}
+
+	FCanPerformActionResults AllowDrop(false);
+	
+	TSharedPtr<FNiagaraHierarchyItemViewModelBase> TargetDropItem = AsShared();
+
+	// we only allow drops if some general conditions are fulfilled
+	if(DraggedItem->GetData() != TargetDropItem->GetData() &&
+		(!DraggedItem->HasParent(TargetDropItem, false) || ItemDropZone != EItemDropZone::OntoItem)  &&
+		!TargetDropItem->HasParent(DraggedItem, true))
+	{
+		if(ItemDropZone == EItemDropZone::OntoItem)
+		{
+			// if the current input doesn't have a parent input, we allow dropping other inputs onto it
+			if(DraggedItem->GetData()->IsA<UNiagaraHierarchyModuleInput>() && TargetDropItem->GetData()->IsA<UNiagaraHierarchyModuleInput>() && TargetDropItem->GetParent().Pin()->GetData<UNiagaraHierarchyModuleInput>() == nullptr)
+			{
+				if(DraggedItem->GetData()->GetChildren().Num() > 0)
+				{
+					FText BaseMessage = LOCTEXT("DroppingInputOnInputWillEmptyChildren", "Input {0} has child inputs. Dropping the input here will remove these children as we only allow nested inputs one level deep.");
+					AllowDrop.CanPerformMessage = FText::FormatOrdered(BaseMessage, DraggedItem->ToStringAsText());
+					AllowDrop.bCanPerform = true;
+				}
+				else
+				{
+					FText BaseMessage = LOCTEXT("DroppingInputOnInputNestedChild", "This will nest input {0} under input {1}");
+					AllowDrop.CanPerformMessage = FText::FormatOrdered(BaseMessage, DraggedItem->ToStringAsText(), TargetDropItem->ToStringAsText());
+					AllowDrop.bCanPerform = true;
+				}
+			}
+		}
+		else
+		{
+			// if the dragged item is an input, we generally allow above/below, even for nested child inputs
+			if(DraggedItem->GetData()->IsA<UNiagaraHierarchyModuleInput>())
+			{
+				AllowDrop.bCanPerform = true;
+			}
+			else
+			{
+				// we use default logic only if there is no parent input. Nested children are not allowed to contain anything but other inputs.
+				if(TargetDropItem->GetParent().Pin()->GetData<UNiagaraHierarchyModuleInput>() == nullptr)
+				{
+					AllowDrop = FNiagaraHierarchyItemViewModel::CanDropOnInternal(DraggedItem, ItemDropZone);
+				}
+			}
+		}
+	}
+
+	return AllowDrop;
+}
+
+void FNiagaraModuleInputViewModel::OnDroppedOnInternal(TSharedPtr<FNiagaraHierarchyItemViewModelBase> DroppedItem, EItemDropZone ItemDropZone)
+{
+	if(ItemDropZone != EItemDropZone::OntoItem)
+	{
+		return FNiagaraHierarchyItemViewModel::OnDroppedOnInternal(DroppedItem, ItemDropZone);
+	}
+	else
+	{
+		FScopedTransaction Transaction(LOCTEXT("Transaction_AddedChildInput", "Added child input"));
+		HierarchyViewModel->GetHierarchyRoot()->Modify();
+
+		// we empty out the children as technically you can drag a parent input onto another input now. We don't take these child-child inputs with us as we only allow child inputs 1 layer deep
+		if(DroppedItem->IsForHierarchy() == false)
+		{
+			TSharedRef<FNiagaraHierarchyItemViewModelBase> AddedItemViewModel = DuplicateToThis(DroppedItem);
+			AddedItemViewModel->GetChildrenMutable().Empty();
+			AddedItemViewModel->SyncViewModelsToData();
+		}
+		else
+		{
+			TSharedRef<FNiagaraHierarchyItemViewModelBase> ReparentedViewModel = ReparentToThis(DroppedItem);
+			ReparentedViewModel->GetChildrenMutable().Empty();
+			ReparentedViewModel->SyncViewModelsToData();
+		}
+
+		HierarchyViewModel->RefreshHierarchyView();
+		HierarchyViewModel->RefreshSourceView();
+	}
+}
+
+void FNiagaraModuleInputViewModel::PopulateDynamicContextMenuSection(FToolMenuSection& DynamicSection)
+{
+	FUIAction Action;
+	Action.ExecuteAction = FExecuteAction::CreateSP(this, &FNiagaraModuleInputViewModel::AddNativeChildrenInputs);
+	Action.CanExecuteAction = FCanExecuteAction::CreateSP(this, &FNiagaraModuleInputViewModel::CanAddNativeChildrenInputs);
+	Action.IsActionVisibleDelegate = FIsActionButtonVisible::CreateSP(this, &FNiagaraModuleInputViewModel::CanAddNativeChildrenInputs);
+	
+	DynamicSection.AddMenuEntry(
+		FName("Add Children Inputs"),
+		LOCTEXT("AddChildrenInputsMenuLabel", "Add Children Inputs"),
+		LOCTEXT("AddChildrenInputsMenuTooltip", "Add children inputs of this input as child inputs."),
+		FSlateIcon(),
+		Action);
+}
+
+TWeakObjectPtr<UNiagaraNodeFunctionCall> FNiagaraModuleInputViewModel::GetModuleNode() const
+{
+	TWeakObjectPtr<UNiagaraHierarchyViewModelBase> ViewModel = GetHierarchyViewModel();
+	UNiagaraSummaryViewViewModel* SummaryViewModel = CastChecked<UNiagaraSummaryViewViewModel>(ViewModel.Get());
+	return SummaryViewModel->GetFunctionCallNode(GetData()->GetPersistentIdentity().Guids[0]);
+}
+
+TOptional<FInputData> FNiagaraModuleInputViewModel::GetInputData() const
+{
+	if(!InputDataCache.IsSet())
+	{
+		InputDataCache = FindInputDataInternal();
+	}
+
+	return InputDataCache;
+}
+
+bool FNiagaraModuleInputViewModel::CanHaveChildren() const
+{
+	// we generally allow children inputs in the source view
+	if(IsForHierarchy() == false)
+	{
+		return true;
+	}
+
+	// we allow module inputs to have children inputs one layer deep
+	if(Parent.IsValid() && Parent.Pin()->GetData()->IsA<UNiagaraHierarchyModuleInput>())
+	{
+		return false;
+	}
+
+	return true;
+}
+
+FString FNiagaraModuleInputViewModel::ToString() const
+{
+	TOptional<FInputData> InputData = GetInputData();
+	if(InputData.IsSet())
+	{
+		return InputData->InputName.ToString();
+	}
+
+	return FNiagaraHierarchyItemViewModel::ToString();
+}
+
+TArray<FString> FNiagaraModuleInputViewModel::GetSearchTerms() const
+{
+	TArray<FString> SearchTerms;
+	SearchTerms.Add(ToString());
+
+	FText DisplayNameOverride = GetData<UNiagaraHierarchyModuleInput>()->GetDisplayNameOverride();
+	if(DisplayNameOverride.IsEmpty() == false)
+	{
+		SearchTerms.Add(DisplayNameOverride.ToString());
+	}
+	
+	return SearchTerms;
+}
+
+bool FNiagaraModuleInputViewModel::IsFromBaseEmitter() const
+{
+	if(!IsFromBaseEmitterCache.IsSet())
+	{
+		UNiagaraSummaryViewViewModel* ViewModel = Cast<UNiagaraSummaryViewViewModel>(GetHierarchyViewModel());
+		IsFromBaseEmitterCache = GetIsFromBaseEmitter(ViewModel->GetEmitterViewModel()->GetEmitter(), GetData()->GetPersistentIdentity());
+	}
+
+	return IsFromBaseEmitterCache.GetValue();
+}
+
+void FNiagaraModuleInputViewModel::ClearCache() const
+{
+	InputDataCache.Reset();
+}
+
+void FNiagaraModuleInputViewModel::RefreshChildDynamicInputs(bool bClearCache)
+{
+	TOptional<FInputData> InputData = GetInputData();
+	if(InputData.IsSet())
+	{
+		UNiagaraNodeFunctionCall* DynamicInputNode = FNiagaraStackGraphUtilities::FindDynamicInputNodeForInput(*GetModuleNode().Get(), InputData->InputName);
+
+		if(DynamicInputNode)
+		{
+			FNiagaraHierarchyIdentity DynamicInputIdentity;
+			DynamicInputIdentity.Guids.Add(DynamicInputNode->NodeGuid);
+			const bool bChildExists = GetDataMutable()->GetChildrenMutable().ContainsByPredicate([DynamicInputIdentity, bClearCache](UNiagaraHierarchyItemBase* CandidateChild)
+				{
+					return CandidateChild->GetPersistentIdentity() == DynamicInputIdentity;
+				});
+				
+			if(bChildExists == false)
+			{
+				UNiagaraHierarchyModule* DynamicInputHierarchyModule = GetDataMutable()->AddChild<UNiagaraHierarchyModule>();
+				DynamicInputHierarchyModule->Initialize(*DynamicInputNode);
+			}			
+		}
+	}
+}
+
+FText FNiagaraModuleInputViewModel::GetSummaryInputNameOverride() const
+{
+	return GetData<UNiagaraHierarchyModuleInput>()->GetDisplayNameOverride();
+}
+
+FNiagaraHierarchyItemViewModelBase::FCanPerformActionResults FNiagaraModuleInputViewModel::IsEditableByUser()
+{
+	FCanPerformActionResults CanEditResults(IsFromBaseEmitter() == false);
+	CanEditResults.CanPerformMessage = CanEditResults.bCanPerform == false ? LOCTEXT("ModuleInputIsFromBaseEmitter", "This input was added in the parent emitter and can not be edited.") : FText::GetEmpty();
+
+	if(bIsForHierarchy && CanEditResults.bCanPerform == true)
+	{
+		if(Parent.Pin()->GetData()->IsA<UNiagaraHierarchyModule>())
+		{
+			CanEditResults.bCanPerform = false;
+			CanEditResults.CanPerformMessage = LOCTEXT("ModuleCanOnlyBeEditedDirectly", "This input can not be modified as it is inherent part of its parent module. Add this input separately if you want to modify it.");
+		}
+	}
+	
+	return CanEditResults;
+}
+
+void FNiagaraModuleInputViewModel::RefreshChildrenDataInternal()
+{
+	RefreshChildDynamicInputs(false);
+}
+
+TOptional<FInputData> FNiagaraModuleInputViewModel::FindInputDataInternal() const
+{
+	TWeakObjectPtr<UNiagaraHierarchyViewModelBase> ViewModel = GetHierarchyViewModel();
+	UNiagaraSummaryViewViewModel* SummaryViewModel = CastChecked<UNiagaraSummaryViewViewModel>(ViewModel.Get());
+	return SummaryViewModel->GetInputData(*GetData<UNiagaraHierarchyModuleInput>());
+}
+
+void FNiagaraModuleInputViewModel::AddNativeChildrenInputs()
+{
+	if(IsForHierarchy() == false)
+	{
+		return;
+	}
+	
+	if(GetModuleNode().IsValid() && GetInputData().IsSet())
+	{
+		TArray<FNiagaraHierarchyIdentity> ChildIdentities = GetNativeChildInputIdentities();
+		TMap<FNiagaraHierarchyIdentity, int32> ChildSortOrderMap;
+
+		if(UNiagaraGraph* Graph = GetModuleNode()->GetCalledGraph())
+		{
+			TArray<FNiagaraVariable> Variables;
+			Graph->GetAllVariables(Variables);
+
+			for(const FNiagaraVariable& Variable : Variables)
+			{
+				if(UNiagaraScriptVariable* ScriptVariable = Graph->GetScriptVariable(Variable))
+				{
+					if(!ScriptVariable->Metadata.ParentAttribute.IsNone() && ScriptVariable->Metadata.ParentAttribute.IsEqual(GetInputData()->InputName))
+					{
+						FNiagaraHierarchyIdentity ChildIdentity;
+						ChildIdentity.Guids.Add(GetData()->GetPersistentIdentity().Guids[0]);
+						ChildIdentity.Guids.Add(ScriptVariable->Metadata.GetVariableGuid());
+						ChildSortOrderMap.Add(ChildIdentity, ScriptVariable->Metadata.EditorSortPriority);
+					}
+				}
+			}
+		}
+
+		for(const FNiagaraHierarchyIdentity& ChildIdentity : ChildIdentities)
+		{
+			if(FindViewModelForChild(ChildIdentity, false) == nullptr)
+			{
+				TSharedPtr<FNiagaraHierarchyItemViewModelBase> ViewModel = HierarchyViewModel->GetHierarchyRootViewModel()->FindViewModelForChild(ChildIdentity, true);
+				if(ViewModel.IsValid())
+				{
+					ReparentToThis(ViewModel);
+				}
+				else
+				{
+					UNiagaraHierarchyModuleInput* ModuleInput = GetDataMutable()->AddChild<UNiagaraHierarchyModuleInput>();
+					ModuleInput->Initialize(*GetModuleNode().Get(), ChildIdentity.Guids[1]);
+				}
+			}
+		}
+
+		SyncViewModelsToData();
+		
+		TArray<TSharedPtr<FNiagaraModuleInputViewModel>> ChildInputs;
+		GetChildrenViewModelsForType<UNiagaraHierarchyModuleInput, FNiagaraModuleInputViewModel>(ChildInputs);
+
+		auto SortChildrenInputs = [&](UNiagaraHierarchyItemBase& ItemA, UNiagaraHierarchyItemBase& ItemB)
+		{
+			int32 SortOrderA = ChildSortOrderMap[ItemA.GetPersistentIdentity()];
+			int32 SortOrderB = ChildSortOrderMap[ItemB.GetPersistentIdentity()];
+
+			return SortOrderA < SortOrderB;
+		};
+		
+		GetDataMutable()->SortChildren(SortChildrenInputs, false);
+		SyncViewModelsToData();
+		HierarchyViewModel->OnHierarchyChanged().Broadcast();
+	}
+}
+
+bool FNiagaraModuleInputViewModel::CanAddNativeChildrenInputs() const
+{
+	if(IsForHierarchy() == false)
+	{
+		return false;
+	}
+	
+	return GetNativeChildInputIdentities().Num() > 0;
+}
+
+TArray<FNiagaraHierarchyIdentity> FNiagaraModuleInputViewModel::GetNativeChildInputIdentities() const
+{
+	TArray<FNiagaraHierarchyIdentity> ChildIdentities;
+	if(GetModuleNode().IsValid())
+	{
+		if(UNiagaraGraph* Graph = GetModuleNode()->GetCalledGraph())
+		{
+			TArray<FNiagaraVariable> Variables;
+			Graph->GetAllVariables(Variables);
+
+			for(const FNiagaraVariable& Variable : Variables)
+			{
+				if(UNiagaraScriptVariable* ScriptVariable = Graph->GetScriptVariable(Variable))
+				{
+					if(!ScriptVariable->Metadata.ParentAttribute.IsNone() && ScriptVariable->Metadata.ParentAttribute.IsEqual(GetInputData()->InputName))
+					{
+						FNiagaraHierarchyIdentity ChildIdentity;
+						ChildIdentity.Guids.Add(GetData()->GetPersistentIdentity().Guids[0]);
+						ChildIdentity.Guids.Add(ScriptVariable->Metadata.GetVariableGuid());
+						ChildIdentities.Add(ChildIdentity);
+					}
+				}
+			}
+		}
+	}
+
+	return ChildIdentities;
+}
+
+FNiagaraHierarchyItemViewModelBase::FCanPerformActionResults FNiagaraAssignmentInputViewModel::CanDropOnInternal(TSharedPtr<FNiagaraHierarchyItemViewModelBase> DraggedItem, EItemDropZone ItemDropZone)
+{
+	// if the input isn't editable, we don't allow any drops on/above/below the item.
+	// Even though it technically works, the merge process will only re-add the item at the end and not preserve order.
+	if(IsEditableByUser().bCanPerform == false)
+	{
+		return false;
+	}
+	
+	return FNiagaraHierarchyItemViewModel::CanDropOnInternal(DraggedItem, ItemDropZone);
+}
+
+TWeakObjectPtr<UNiagaraNodeAssignment> FNiagaraAssignmentInputViewModel::GetAssignmentNode() const
+{
+	TWeakObjectPtr<UNiagaraHierarchyViewModelBase> ViewModel = GetHierarchyViewModel();
+	UNiagaraSummaryViewViewModel* SummaryViewModel = CastChecked<UNiagaraSummaryViewViewModel>(ViewModel.Get());
+	return Cast<UNiagaraNodeAssignment>(SummaryViewModel->GetFunctionCallNode(GetData()->GetPersistentIdentity().Guids[0]));
+}
+
+TOptional<FNiagaraStackGraphUtilities::FMatchingFunctionInputData> FNiagaraAssignmentInputViewModel::GetInputData() const
+{
+	if(!InputDataCache.IsSet())
+	{
+		InputDataCache = FindInputDataInternal();
+	}
+
+	return InputDataCache;
+}
+
+FString FNiagaraAssignmentInputViewModel::ToString() const
+{
+	TOptional<FNiagaraStackGraphUtilities::FMatchingFunctionInputData> InputData = GetInputData();
+	if(GetInputData().IsSet())
+	{
+		return InputData->InputName.ToString();
+	}
+
+	return FNiagaraHierarchyItemViewModel::ToString();
+}
+
+TArray<FString> FNiagaraAssignmentInputViewModel::GetSearchTerms() const
+{
+	TArray<FString> SearchTerms;
+	SearchTerms.Add(ToString());	
+	return SearchTerms;
+}
+
+bool FNiagaraAssignmentInputViewModel::IsFromBaseEmitter() const
+{
+	if(!IsFromBaseEmitterCache.IsSet())
+	{
+		UNiagaraSummaryViewViewModel* ViewModel = Cast<UNiagaraSummaryViewViewModel>(GetHierarchyViewModel());
+		IsFromBaseEmitterCache = GetIsFromBaseEmitter(ViewModel->GetEmitterViewModel()->GetEmitter(), GetData()->GetPersistentIdentity());
+	}
+
+	return IsFromBaseEmitterCache.GetValue();
+}
+
+void FNiagaraAssignmentInputViewModel::ClearCache() const
+{
+	InputDataCache.Reset();
+	TWeakObjectPtr<UNiagaraHierarchyViewModelBase> ViewModel = GetHierarchyViewModel();
+	UNiagaraSummaryViewViewModel* SummaryViewModel = CastChecked<UNiagaraSummaryViewViewModel>(ViewModel.Get());
+	SummaryViewModel->ClearFunctionCallNodeCache(GetData()->GetPersistentIdentity().Guids[0]);
+}
+
+FNiagaraHierarchyItemViewModelBase::FCanPerformActionResults FNiagaraAssignmentInputViewModel::IsEditableByUser()
+{
+	FCanPerformActionResults CanEditResults(IsFromBaseEmitter() == false);
+	CanEditResults.CanPerformMessage = CanEditResults.bCanPerform == false ? LOCTEXT("ModuleInputIsFromBaseEmitter", "This input was added in the parent emitter and can not be edited.") : FText::GetEmpty();
+
+	if(bIsForHierarchy && CanEditResults.bCanPerform == true)
+	{
+		if(Parent.Pin()->GetData()->IsA<UNiagaraHierarchyModule>())
+		{
+			CanEditResults.bCanPerform = false;
+			CanEditResults.CanPerformMessage = LOCTEXT("ModuleCanOnlyBeEditedDirectly", "This input can not be modified as it is inherent part of its parent module. Add this input separately if you want to modify it.");
+		}
+	}
+	
+	return CanEditResults;
+}
+
+TOptional<FNiagaraStackGraphUtilities::FMatchingFunctionInputData> FNiagaraAssignmentInputViewModel::FindInputDataInternal() const
+{
+	TWeakObjectPtr<UNiagaraHierarchyViewModelBase> ViewModel = GetHierarchyViewModel();
+	UNiagaraSummaryViewViewModel* SummaryViewModel = CastChecked<UNiagaraSummaryViewViewModel>(ViewModel.Get());
+	
+	if(GetAssignmentNode().IsValid())
+	{
+		return FNiagaraStackGraphUtilities::FindAssignmentInputData(*GetAssignmentNode().Get(), GetData()->GetPersistentIdentity().Names[0], SummaryViewModel->GetEmitterViewModel());
+	}
+
+	return TOptional<FNiagaraStackGraphUtilities::FMatchingFunctionInputData>();
+}
+
+bool FNiagaraHierarchySummaryCategoryViewModel::IsFromBaseEmitter() const
+{
+	if(!IsFromBaseEmitterCache.IsSet())
+	{
+		UNiagaraSummaryViewViewModel* ViewModel = Cast<UNiagaraSummaryViewViewModel>(GetHierarchyViewModel());
+		IsFromBaseEmitterCache = GetIsFromBaseEmitter(ViewModel->GetEmitterViewModel()->GetEmitter(), GetData()->GetPersistentIdentity());
+	}
+
+	return IsFromBaseEmitterCache.GetValue();
+}
+
+FNiagaraHierarchyItemViewModelBase::FCanPerformActionResults FNiagaraHierarchySummaryCategoryViewModel::IsEditableByUser()
+{
+	FCanPerformActionResults CanEditResults(IsFromBaseEmitter() == false);
+	CanEditResults.CanPerformMessage = CanEditResults.bCanPerform == false ? LOCTEXT("CategoryIsFromBaseEmitter", "This category was added in the parent emitter and can not be edited. You can add new items.") : FText::GetEmpty();
+	return CanEditResults;
+}
+
+bool FNiagaraHierarchyPropertyViewModel::IsFromBaseEmitter() const
+{
+	if(!IsFromBaseEmitterCache.IsSet())
+	{
+		UNiagaraSummaryViewViewModel* ViewModel = Cast<UNiagaraSummaryViewViewModel>(GetHierarchyViewModel());
+		IsFromBaseEmitterCache = GetIsFromBaseEmitter(ViewModel->GetEmitterViewModel()->GetEmitter(), GetData()->GetPersistentIdentity());
+	}
+
+	return IsFromBaseEmitterCache.GetValue();
+}
+
+bool FNiagaraHierarchyPropertyViewModel::DoesExternalDataStillExist(const UNiagaraHierarchyDataRefreshContext* Context) const
+{
+	UNiagaraSummaryViewViewModel* ViewModel = Cast<UNiagaraSummaryViewViewModel>(GetHierarchyViewModel());
+	TMap<FGuid, UObject*> PropertyObjectMap = ViewModel->GetObjectsForProperties();
+	return PropertyObjectMap.Contains(GetData()->GetPersistentIdentity().Guids[0]);
+}
+
+FNiagaraHierarchyItemViewModelBase::FCanPerformActionResults FNiagaraHierarchyPropertyViewModel::IsEditableByUser()
+{
+	FCanPerformActionResults CanEditResults(IsFromBaseEmitter() == false);
+	CanEditResults.CanPerformMessage = CanEditResults.bCanPerform == false ? LOCTEXT("ObjectPropertyIsFromBaseEmitter", "This property was added in the parent emitter and can not be edited.") : FText::GetEmpty();
+	return CanEditResults;
+}
+
+FString FNiagaraHierarchyRendererViewModel::ToString() const
+{
+	UNiagaraRendererProperties* RendererProperties = GetRendererProperties();
+	return RendererProperties != nullptr ? RendererProperties->GetWidgetDisplayName().ToString() : FString(); 
+}
+
+bool FNiagaraHierarchyRendererViewModel::IsFromBaseEmitter() const
+{
+	if(!IsFromBaseEmitterCache.IsSet())
+	{
+		UNiagaraSummaryViewViewModel* ViewModel = Cast<UNiagaraSummaryViewViewModel>(GetHierarchyViewModel());
+		IsFromBaseEmitterCache = GetIsFromBaseEmitter(ViewModel->GetEmitterViewModel()->GetEmitter(), GetData()->GetPersistentIdentity());
+	}
+
+	return IsFromBaseEmitterCache.GetValue();
+}
+
+void FNiagaraHierarchyRendererViewModel::RefreshChildrenDataInternal()
+{
+	TArray<UNiagaraHierarchyItemBase*> NewChildren;
+	for (TFieldIterator<FProperty> PropertyIterator(GetRendererProperties()->GetClass(), EFieldIteratorFlags::ExcludeSuper); PropertyIterator; ++PropertyIterator)
+	{
+		if(PropertyIterator->HasAnyPropertyFlags(CPF_Edit))
+		{
+			FString PropertyName = (*PropertyIterator)->GetName();
+
+			FNiagaraHierarchyIdentity PropertyIdentity;
+			PropertyIdentity.Guids.Add(GetRendererProperties()->GetMergeId());
+			PropertyIdentity.Names.Add(FName(PropertyName));
+			
+			UNiagaraHierarchyItemBase** FoundItem = GetDataMutable()->GetChildrenMutable().FindByPredicate([PropertyIdentity](UNiagaraHierarchyItemBase* Candidate)
+			{
+				return Candidate->GetPersistentIdentity() == PropertyIdentity;
+			});
+
+			UNiagaraHierarchyObjectProperty* RendererProperty = nullptr;
+			if(FoundItem == nullptr)
+			{
+				RendererProperty = GetDataMutable()->AddChild<UNiagaraHierarchyObjectProperty>();
+				RendererProperty->Initialize(GetRendererProperties()->GetMergeId(), PropertyName);
+			}
+			else
+			{
+				RendererProperty = CastChecked<UNiagaraHierarchyObjectProperty>(*FoundItem);
+			}
+
+			NewChildren.Add(RendererProperty);
+		}
+	}
+
+	GetDataMutable()->GetChildrenMutable().Empty();
+	GetDataMutable()->GetChildrenMutable().Append(NewChildren);
+}
+
+FNiagaraHierarchyItemViewModelBase::FCanPerformActionResults FNiagaraHierarchyRendererViewModel::IsEditableByUser()
+{
+	FCanPerformActionResults CanEditResults(IsFromBaseEmitter() == false);
+	CanEditResults.CanPerformMessage = CanEditResults.bCanPerform == false ? LOCTEXT("RendererIsFromBaseEmitter", "This renderer was added in the parent emitter and can not be edited.") : FText::GetEmpty();
+	return CanEditResults;
+}
+
+const UNiagaraHierarchySection* FNiagaraHierarchyRendererViewModel::GetSectionInternal() const
+{
+	return Section.IsValid() ? Section.Get() : nullptr;
+}
+
+FString FNiagaraHierarchyEventHandlerViewModel::ToString() const
+{
+	FNiagaraEventScriptProperties* ScriptProperties = GetEventScriptProperties();
+	if(ScriptProperties != nullptr)
+	{
+		return ScriptProperties->SourceEventName.ToString();
+	}
+
+	return FString();
+}
+
+FNiagaraEventScriptProperties* FNiagaraHierarchyEventHandlerViewModel::GetEventScriptProperties() const
+{
+	TWeakObjectPtr<UNiagaraHierarchyViewModelBase> ViewModel = GetHierarchyViewModel();
+	UNiagaraSummaryViewViewModel* SummaryViewModel = CastChecked<UNiagaraSummaryViewViewModel>(ViewModel.Get());
+	FGuid UsageID = GetData()->GetPersistentIdentity().Guids[0];
+	
+	for(FNiagaraEventScriptProperties& ScriptProperties : SummaryViewModel->GetEmitterViewModel()->GetEmitter().GetEmitterData()->EventHandlerScriptProps)
+	{
+		if(ScriptProperties.Script->GetUsageId() == UsageID)
+		{
+			return &ScriptProperties;
+		}
+	}
+	
+	return nullptr;
+}
+
+bool FNiagaraHierarchyEventHandlerViewModel::IsFromBaseEmitter() const
+{
+	if(!IsFromBaseEmitterCache.IsSet())
+	{
+		UNiagaraSummaryViewViewModel* ViewModel = Cast<UNiagaraSummaryViewViewModel>(GetHierarchyViewModel());
+		IsFromBaseEmitterCache = GetIsFromBaseEmitter(ViewModel->GetEmitterViewModel()->GetEmitter(), GetData()->GetPersistentIdentity());
+	}
+
+	return IsFromBaseEmitterCache.GetValue();
+}
+
+void FNiagaraHierarchyEventHandlerViewModel::RefreshChildrenDataInternal()
+{
+	TArray<UNiagaraHierarchyItemBase*> NewChildren;
+
+	// First we add the properties item
+	FNiagaraHierarchyIdentity PropertiesIdentity = UNiagaraHierarchyEventHandlerProperties::MakeIdentity(*GetEventScriptProperties());	
+
+	UNiagaraHierarchyItemBase** FoundProperties = GetDataMutable()->GetChildrenMutable().FindByPredicate([PropertiesIdentity](UNiagaraHierarchyItemBase* Candidate)
+	{
+		return Candidate->GetPersistentIdentity() == PropertiesIdentity;
+	});
+
+	UNiagaraHierarchyEventHandlerProperties* PropertiesCategory = nullptr;
+	if(FoundProperties == nullptr)
+	{
+		PropertiesCategory = GetDataMutable()->AddChild<UNiagaraHierarchyEventHandlerProperties>();
+		PropertiesCategory->Initialize(*GetEventScriptProperties());
+	}
+	else
+	{
+		PropertiesCategory = CastChecked<UNiagaraHierarchyEventHandlerProperties>(*FoundProperties);
+	}
+
+	NewChildren.Add(PropertiesCategory);
+
+	// Then we go through all modules of that sim stage
+	UNiagaraSummaryViewViewModel* ViewModel = Cast<UNiagaraSummaryViewViewModel>(GetHierarchyViewModel());
+	TArray<UNiagaraNodeFunctionCall*> EventHandlerModules = FNiagaraStackGraphUtilities::FindModuleNodesForEventHandler(*GetEventScriptProperties(), ViewModel->GetEmitterViewModel());
+	
+	for(UNiagaraNodeFunctionCall* EventHandlerModule : EventHandlerModules)
+	{
+		UNiagaraHierarchyModule* HierarchyEventHandlerModule = nullptr;
+		FNiagaraHierarchyIdentity ModuleIdentity;
+		ModuleIdentity.Guids.Add(EventHandlerModule->NodeGuid);
+		UNiagaraHierarchyItemBase** FoundHierarchySimStageModule = GetDataMutable()->GetChildrenMutable().FindByPredicate([ModuleIdentity](UNiagaraHierarchyItemBase* Candidate)
+		{
+			return Candidate->GetPersistentIdentity() == ModuleIdentity;
+		});
+		
+		if(FoundHierarchySimStageModule == nullptr)		
+		{
+			HierarchyEventHandlerModule = GetDataMutable()->AddChild<UNiagaraHierarchyModule>();
+			HierarchyEventHandlerModule->Initialize(*EventHandlerModule);
+		}
+		else
+		{
+			HierarchyEventHandlerModule = CastChecked<UNiagaraHierarchyModule>(*FoundHierarchySimStageModule);
+		}
+		
+		NewChildren.Add(HierarchyEventHandlerModule);
+	}
+	
+	GetDataMutable()->GetChildrenMutable().Empty();
+	GetDataMutable()->GetChildrenMutable().Append(NewChildren);
+}
+
+FNiagaraHierarchyItemViewModelBase::FCanPerformActionResults FNiagaraHierarchyEventHandlerViewModel::IsEditableByUser()
+{
+	FCanPerformActionResults CanEditResults(IsFromBaseEmitter() == false);
+	CanEditResults.CanPerformMessage = CanEditResults.bCanPerform == false ? LOCTEXT("EventHandlerIsFromBaseEmitter", "This event handler was added in the parent emitter and can not be edited.") : FText::GetEmpty();
+	return CanEditResults;
+}
+
+const UNiagaraHierarchySection* FNiagaraHierarchyEventHandlerViewModel::GetSectionInternal() const
+{
+	if(Section.IsValid())
+	{
+		return Section.Get();
+	}
+
+	return nullptr;
+}
+
+UNiagaraRendererProperties* FNiagaraHierarchyRendererViewModel::GetRendererProperties() const
+{
+	TWeakObjectPtr<UNiagaraHierarchyViewModelBase> ViewModel = GetHierarchyViewModel();
+	UNiagaraSummaryViewViewModel* SummaryViewModel = CastChecked<UNiagaraSummaryViewViewModel>(ViewModel.Get());
+	FGuid RendererGuid = GetData()->GetPersistentIdentity().Guids[0];
+
+	UNiagaraRendererProperties* MatchingRenderer = nullptr;
+	const TArray<UNiagaraRendererProperties*> RendererProperties = SummaryViewModel->GetEmitterViewModel()->GetEmitter().GetEmitterData()->GetRenderers();
+	for(UNiagaraRendererProperties* Renderer : RendererProperties)
+	{
+		if(Renderer->GetMergeId() == RendererGuid)
+		{
+			MatchingRenderer = Renderer;
+			break;
+		}
+	}
+	
+	return MatchingRenderer;
+}
+
+FString FNiagaraHierarchyEventHandlerPropertiesViewModel::ToString() const
+{
+	return GetEventScriptProperties()->SourceEventName.ToString().Append(" Properties");
+}
+
+FNiagaraEventScriptProperties* FNiagaraHierarchyEventHandlerPropertiesViewModel::GetEventScriptProperties() const
+{
+	TWeakObjectPtr<UNiagaraHierarchyViewModelBase> ViewModel = GetHierarchyViewModel();
+	UNiagaraSummaryViewViewModel* SummaryViewModel = CastChecked<UNiagaraSummaryViewViewModel>(ViewModel.Get());
+	FGuid UsageID = GetData()->GetPersistentIdentity().Guids[0];
+	
+	for(FNiagaraEventScriptProperties& ScriptProperties : SummaryViewModel->GetEmitterViewModel()->GetEmitter().GetEmitterData()->EventHandlerScriptProps)
+	{
+		if(ScriptProperties.Script->GetUsageId() == UsageID)
+		{
+			return &ScriptProperties;
+		}
+	}
+	
+	return nullptr;
+}
+
+bool FNiagaraHierarchyEventHandlerPropertiesViewModel::IsFromBaseEmitter() const
+{
+	if(!IsFromBaseEmitterCache.IsSet())
+	{
+		UNiagaraSummaryViewViewModel* ViewModel = Cast<UNiagaraSummaryViewViewModel>(GetHierarchyViewModel());
+		IsFromBaseEmitterCache = GetIsFromBaseEmitter(ViewModel->GetEmitterViewModel()->GetEmitter(), GetData()->GetPersistentIdentity());
+	}
+
+	return IsFromBaseEmitterCache.GetValue();
+}
+
+void FNiagaraHierarchyEventHandlerPropertiesViewModel::RefreshChildrenDataInternal()
+{
+	TArray<UNiagaraHierarchyItemBase*> NewPropertiesChildren;
+
+	// todo (me) while this works, the stack needs to access the correct FStructOnScope that points to the FNiagaraEventScriptProperties
+	// That can be made to work correctly, but the EventScriptProperties are heavily customized and introduce UI issues
+	// Potentially solvable by registering the same customization but skipping for now
+	
+	// for (TFieldIterator<FProperty> PropertyIterator(StaticStruct<FNiagaraEventScriptProperties>(), EFieldIteratorFlags::IncludeSuper); PropertyIterator; ++PropertyIterator)
+	// {
+	// 	if(PropertyIterator->HasAnyPropertyFlags(CPF_Edit))
+	// 	{
+	// 		FString PropertyName = (*PropertyIterator)->GetName();
+	//
+	// 		FNiagaraHierarchyIdentity PropertyIdentity;
+	// 		PropertyIdentity.Guids.Add(GetEventScriptProperties()->Script->GetUsageId());
+	// 		PropertyIdentity.Names.Add(FName(PropertyName));
+	// 		
+	// 		UNiagaraHierarchyItemBase** FoundPropertyItem = GetDataMutable()->GetChildrenMutable().FindByPredicate([PropertyIdentity](UNiagaraHierarchyItemBase* Candidate)
+	// 		{
+	// 			return Candidate->GetPersistentIdentity() == PropertyIdentity;
+	// 		});
+	//
+	// 		UNiagaraHierarchyObjectProperty* EventHandlerProperty = nullptr;
+	// 		if(FoundPropertyItem == nullptr)
+	// 		{
+	// 			EventHandlerProperty = GetDataMutable()->AddChild<UNiagaraHierarchyObjectProperty>();
+	// 			EventHandlerProperty->Initialize(GetEventScriptProperties()->Script->GetUsageId(), PropertyName);
+	// 		}
+	// 		else
+	// 		{
+	// 			EventHandlerProperty = CastChecked<UNiagaraHierarchyObjectProperty>(*FoundPropertyItem);
+	// 		}
+	//
+	// 		NewPropertiesChildren.Add(EventHandlerProperty);
+	// 	}
+	// }
+
+	GetDataMutable()->GetChildrenMutable().Empty();
+	GetDataMutable()->GetChildrenMutable().Append(NewPropertiesChildren);
+}
+
+FNiagaraHierarchyItemViewModelBase::FCanPerformActionResults FNiagaraHierarchyEventHandlerPropertiesViewModel::IsEditableByUser()
+{
+	FCanPerformActionResults CanEditResults(IsFromBaseEmitter() == false);
+	CanEditResults.CanPerformMessage = CanEditResults.bCanPerform == false ? LOCTEXT("EventHandlerPropertiesIsFromBaseEmitter", "This property item was added in the parent emitter and can not be edited.") : FText::GetEmpty();
+	return CanEditResults;
+}
+
+FString FNiagaraHierarchySimStageViewModel::ToString() const
+{
+	UNiagaraSimulationStageBase* SimStage = GetSimStage();
+	return SimStage != nullptr ? SimStage->SimulationStageName.ToString() : FString(); 
+}
+
+UNiagaraSimulationStageBase* FNiagaraHierarchySimStageViewModel::GetSimStage() const
+{
+	TWeakObjectPtr<UNiagaraHierarchyViewModelBase> ViewModel = GetHierarchyViewModel();
+	UNiagaraSummaryViewViewModel* SummaryViewModel = CastChecked<UNiagaraSummaryViewViewModel>(ViewModel.Get());
+	FGuid SimStageGuid = GetData()->GetPersistentIdentity().Guids[0];
+
+	UNiagaraSimulationStageBase* MatchingSimStage = nullptr;
+	const TArray<UNiagaraSimulationStageBase*> SimStages = SummaryViewModel->GetEmitterViewModel()->GetEmitter().GetEmitterData()->GetSimulationStages();
+	for(UNiagaraSimulationStageBase* SimStage : SimStages)
+	{
+		if(SimStage->GetMergeId() == SimStageGuid)
+		{
+			MatchingSimStage = SimStage;
+			break;
+		}
+	}
+	
+	return MatchingSimStage;
+}
+
+bool FNiagaraHierarchySimStageViewModel::IsFromBaseEmitter() const
+{
+	if(!IsFromBaseEmitterCache.IsSet())
+	{
+		UNiagaraSummaryViewViewModel* ViewModel = Cast<UNiagaraSummaryViewViewModel>(GetHierarchyViewModel());
+		IsFromBaseEmitterCache = GetIsFromBaseEmitter(ViewModel->GetEmitterViewModel()->GetEmitter(), GetData()->GetPersistentIdentity());
+	}
+
+	return IsFromBaseEmitterCache.GetValue();
+}
+
+void FNiagaraHierarchySimStageViewModel::RefreshChildrenDataInternal()
+{
+	TArray<UNiagaraHierarchyItemBase*> NewChildren;
+
+	// First we add the properties item
+	FNiagaraHierarchyIdentity PropertiesIdentity;
+	PropertiesIdentity.Guids.Add(GetSimStage()->GetMergeId());
+	PropertiesIdentity.Names.Add(FName("Category"));
+	PropertiesIdentity.Names.Add(FName("Properties"));
+
+	UNiagaraHierarchyItemBase** FoundProperties = GetDataMutable()->GetChildrenMutable().FindByPredicate([PropertiesIdentity](UNiagaraHierarchyItemBase* Candidate)
+	{
+		return Candidate->GetPersistentIdentity() == PropertiesIdentity;
+	});
+
+	UNiagaraHierarchySimStageProperties* PropertiesCategory = nullptr;
+	if(FoundProperties == nullptr)
+	{
+		PropertiesCategory = GetDataMutable()->AddChild<UNiagaraHierarchySimStageProperties>();
+		PropertiesCategory->Initialize(*GetSimStage());
+	}
+	else
+	{
+		PropertiesCategory = CastChecked<UNiagaraHierarchySimStageProperties>(*FoundProperties);
+	}
+
+	NewChildren.Add(PropertiesCategory);
+
+	// Then we go through all modules of that sim stage
+	UNiagaraSummaryViewViewModel* ViewModel = Cast<UNiagaraSummaryViewViewModel>(GetHierarchyViewModel());
+	TArray<UNiagaraNodeFunctionCall*> SimStageModules = FNiagaraStackGraphUtilities::FindModuleNodesForSimulationStage(*GetSimStage(), ViewModel->GetEmitterViewModel());
+
+	for(UNiagaraNodeFunctionCall* SimStageModule : SimStageModules)
+	{
+		UNiagaraHierarchyModule* HierarchySimStageModule = nullptr;
+		FNiagaraHierarchyIdentity ModuleIdentity;
+		ModuleIdentity.Guids.Add(SimStageModule->NodeGuid);
+		UNiagaraHierarchyItemBase** FoundHierarchySimStageModule = GetDataMutable()->GetChildrenMutable().FindByPredicate([ModuleIdentity](UNiagaraHierarchyItemBase* Candidate)
+		{
+			return Candidate->GetPersistentIdentity() == ModuleIdentity;
+		});
+		
+		if(FoundHierarchySimStageModule == nullptr)		
+		{
+			HierarchySimStageModule = GetDataMutable()->AddChild<UNiagaraHierarchyModule>();
+			HierarchySimStageModule->Initialize(*SimStageModule);
+		}
+		else
+		{
+			HierarchySimStageModule = CastChecked<UNiagaraHierarchyModule>(*FoundHierarchySimStageModule);
+		}
+		
+		NewChildren.Add(HierarchySimStageModule);
+	}
+	
+	GetDataMutable()->GetChildrenMutable().Empty();
+	GetDataMutable()->GetChildrenMutable().Append(NewChildren);
+}
+
+FNiagaraHierarchyItemViewModelBase::FCanPerformActionResults FNiagaraHierarchySimStageViewModel::IsEditableByUser()
+{
+	FCanPerformActionResults CanEditResults(IsFromBaseEmitter() == false);
+	CanEditResults.CanPerformMessage = CanEditResults.bCanPerform == false ? LOCTEXT("SimStageIsFromBaseEmitter", "This simulation stage was added in the parent emitter and can not be edited.") : FText::GetEmpty();
+	return CanEditResults;
+}
+
+const UNiagaraHierarchySection* FNiagaraHierarchySimStageViewModel::GetSectionInternal() const
+{
+	return Section.IsValid() ? Section.Get() : nullptr;
+}
+
+FString FNiagaraHierarchySimStagePropertiesViewModel::ToString() const
+{
+	return GetSimStage()->SimulationStageName.ToString().Append(" Properties");
+}
+
+UNiagaraSimulationStageBase* FNiagaraHierarchySimStagePropertiesViewModel::GetSimStage() const
+{
+	TWeakObjectPtr<UNiagaraHierarchyViewModelBase> ViewModel = GetHierarchyViewModel();
+	UNiagaraSummaryViewViewModel* SummaryViewModel = CastChecked<UNiagaraSummaryViewViewModel>(ViewModel.Get());
+	FGuid SimStageGuid = GetData()->GetPersistentIdentity().Guids[0];
+
+	UNiagaraSimulationStageBase* MatchingSimStage = nullptr;
+	const TArray<UNiagaraSimulationStageBase*> SimStages = SummaryViewModel->GetEmitterViewModel()->GetEmitter().GetEmitterData()->GetSimulationStages();
+	for(UNiagaraSimulationStageBase* SimStage : SimStages)
+	{
+		if(SimStage->GetMergeId() == SimStageGuid)
+		{
+			MatchingSimStage = SimStage;
+			break;
+		}
+	}
+	
+	return MatchingSimStage;
+}
+
+bool FNiagaraHierarchySimStagePropertiesViewModel::IsFromBaseEmitter() const
+{
+	if(!IsFromBaseEmitterCache.IsSet())
+	{
+		UNiagaraSummaryViewViewModel* ViewModel = Cast<UNiagaraSummaryViewViewModel>(GetHierarchyViewModel());
+		IsFromBaseEmitterCache = GetIsFromBaseEmitter(ViewModel->GetEmitterViewModel()->GetEmitter(), GetData()->GetPersistentIdentity());
+	}
+
+	return IsFromBaseEmitterCache.GetValue();
+}
+
+void FNiagaraHierarchySimStagePropertiesViewModel::RefreshChildrenDataInternal()
+{
+	TArray<UNiagaraHierarchyItemBase*> NewPropertiesChildren;
+
+	for (TFieldIterator<FProperty> PropertyIterator(GetSimStage()->GetClass(), EFieldIteratorFlags::IncludeSuper); PropertyIterator; ++PropertyIterator)
+	{
+		if(PropertyIterator->HasAnyPropertyFlags(CPF_Edit))
+		{
+			FString PropertyName = (*PropertyIterator)->GetName();
+
+			FNiagaraHierarchyIdentity PropertyIdentity;
+			PropertyIdentity.Guids.Add(GetSimStage()->GetMergeId());
+			PropertyIdentity.Names.Add(FName(PropertyName));
+			
+			UNiagaraHierarchyItemBase** FoundPropertyItem = GetDataMutable()->GetChildrenMutable().FindByPredicate([PropertyIdentity](UNiagaraHierarchyItemBase* Candidate)
+			{
+				return Candidate->GetPersistentIdentity() == PropertyIdentity;
+			});
+
+			UNiagaraHierarchyObjectProperty* SimStageProperty = nullptr;
+			if(FoundPropertyItem == nullptr)
+			{
+				SimStageProperty = GetDataMutable()->AddChild<UNiagaraHierarchyObjectProperty>();
+				SimStageProperty->Initialize(GetSimStage()->GetMergeId(), PropertyName);
+			}
+			else
+			{
+				SimStageProperty = CastChecked<UNiagaraHierarchyObjectProperty>(*FoundPropertyItem);
+			}
+
+			NewPropertiesChildren.Add(SimStageProperty);
+		}
+	}
+
+	GetDataMutable()->GetChildrenMutable().Empty();
+	GetDataMutable()->GetChildrenMutable().Append(NewPropertiesChildren);
+}
+
+FNiagaraHierarchyItemViewModelBase::FCanPerformActionResults FNiagaraHierarchySimStagePropertiesViewModel::IsEditableByUser()
+{
+	FCanPerformActionResults CanEditResults(IsFromBaseEmitter() == false);
+	CanEditResults.CanPerformMessage = CanEditResults.bCanPerform == false ? LOCTEXT("RendererPropertiesIsFromBaseEmitter", "This renderer's properties were added in the parent emitter and can not be edited.") : FText::GetEmpty();
+	return CanEditResults;
+}
+
+TSharedRef<SWidget> FNiagaraHierarchyInputParameterHierarchyDragDropOp::CreateCustomDecorator() const
+{
+	TSharedPtr<FNiagaraModuleInputViewModel> InputViewModel = StaticCastSharedPtr<FNiagaraModuleInputViewModel>(DraggedItem.Pin());
+	TOptional<FInputData> InputData = InputViewModel->GetInputData();
+	return FNiagaraParameterUtilities::GetParameterWidget(FNiagaraVariable(InputData->Type, InputData->InputName), false, false);
+}
+
+FString FNiagaraHierarchyPropertyViewModel::ToString() const
+{
+	return GetData()->GetPersistentIdentity().Names[0].ToString();
+}
+
+UNiagaraHierarchyRoot* UNiagaraSummaryViewViewModel::GetHierarchyRoot() const
+{
+	UNiagaraHierarchyRoot* RootItem = GetEmitterViewModel()->GetEditorData().GetSummaryRoot();
+
+	ensure(RootItem != nullptr);
+	return RootItem;
+}
+
+TSharedPtr<FNiagaraHierarchyItemViewModelBase> UNiagaraSummaryViewViewModel::CreateViewModelForData(UNiagaraHierarchyItemBase* ItemBase, TSharedPtr<FNiagaraHierarchyItemViewModelBase> Parent)
+{
+	bool bIsForHierarchy = Parent.IsValid() ? Parent->IsForHierarchy() : false;
+	
+	if(UNiagaraHierarchyModuleInput* SummaryViewItem = Cast<UNiagaraHierarchyModuleInput>(ItemBase))
+	{
+		return MakeShared<FNiagaraModuleInputViewModel>(SummaryViewItem, Parent, this, bIsForHierarchy);
+	}
+	else if(UNiagaraHierarchyModule* Module = Cast<UNiagaraHierarchyModule>(ItemBase))
+	{
+		return MakeShared<FNiagaraFunctionViewModel>(Module, Parent, this, bIsForHierarchy);
+	}
+	else if(UNiagaraHierarchyRenderer* Renderer = Cast<UNiagaraHierarchyRenderer>(ItemBase))
+	{
+		return MakeShared<FNiagaraHierarchyRendererViewModel>(Renderer, Parent, this, bIsForHierarchy);
+	}
+	else if(UNiagaraHierarchyEventHandler* EventHandler = Cast<UNiagaraHierarchyEventHandler>(ItemBase))
+	{
+		return MakeShared<FNiagaraHierarchyEventHandlerViewModel>(EventHandler, Parent, this, bIsForHierarchy);
+	}
+	else if(UNiagaraHierarchyEventHandlerProperties* EventHandlerProperties = Cast<UNiagaraHierarchyEventHandlerProperties>(ItemBase))
+	{
+		return MakeShared<FNiagaraHierarchyEventHandlerPropertiesViewModel>(EventHandlerProperties, Parent, this, bIsForHierarchy);
+	}
+	else if(UNiagaraHierarchySimStage* SimStage = Cast<UNiagaraHierarchySimStage>(ItemBase))
+	{
+		return MakeShared<FNiagaraHierarchySimStageViewModel>(SimStage, Parent, this, bIsForHierarchy);
+	}
+	else if(UNiagaraHierarchySimStageProperties* SimStageProperties = Cast<UNiagaraHierarchySimStageProperties>(ItemBase))
+	{
+		return MakeShared<FNiagaraHierarchySimStagePropertiesViewModel>(SimStageProperties, Parent, this, bIsForHierarchy);
+	}
+	else if(UNiagaraHierarchyObjectProperty* ObjectProperty = Cast<UNiagaraHierarchyObjectProperty>(ItemBase))
+	{
+		return MakeShared<FNiagaraHierarchyPropertyViewModel>(ObjectProperty, Parent, this, bIsForHierarchy);
+	}
+	else if(UNiagaraHierarchyAssignmentInput* AssignmentInput = Cast<UNiagaraHierarchyAssignmentInput>(ItemBase))
+	{
+		return MakeShared<FNiagaraAssignmentInputViewModel>(AssignmentInput, Parent, this, bIsForHierarchy);
+	}
+	else if(UNiagaraHierarchyCategory* Category = Cast<UNiagaraHierarchyCategory>(ItemBase))
+	{
+		return MakeShared<FNiagaraHierarchySummaryCategoryViewModel>(Category, Parent, this, bIsForHierarchy);
+	}
+
+	check(false);
+	return nullptr;
+}
+
+void UNiagaraSummaryViewViewModel::PrepareSourceItems(UNiagaraHierarchyRoot* SourceRoot, TSharedPtr<FNiagaraHierarchyRootViewModel> SourceRootViewModel)
+{
+	TSharedPtr<FNiagaraEmitterViewModel> EmitterViewModel = EmitterViewModelWeak.Pin();
+
+	TArray<UNiagaraHierarchyItemBase*> NewItems;
+	TArray<UNiagaraHierarchySection*> NewSections;
+	
+	FunctionCallCache.Empty();
+		
+	// we keep track of all category items (modules etc.) belonging to a certain usage so we can set the sections later 
+	TMap<ENiagaraScriptUsage, TArray<UNiagaraHierarchyModule*>> UsageMap;
+	// we keep track of renderers & sim stages here for the same reasons
+	TArray<UNiagaraHierarchyRenderer*> HierarchyRenderers;
+	TArray<UNiagaraHierarchySimStage*> HierarchySimStages;
+	TArray<UNiagaraHierarchyEventHandler*> HierarchyEventHandlers;
+
+	TArray<UNiagaraNodeFunctionCall*> SimStageModules = FNiagaraStackGraphUtilities::GetAllSimStagesModuleNodes(EmitterViewModel.ToSharedRef());
+	TArray<UNiagaraNodeFunctionCall*> EventHandlerModules = FNiagaraStackGraphUtilities::GetAllEventHandlerModuleNodes(EmitterViewModel.ToSharedRef());
+
+	// We create hierarchy modules here. We attempt to maintain as many previous elements as possible in order to maintain UI state
+	TArray<UNiagaraNodeFunctionCall*> ModuleNodes = FNiagaraStackGraphUtilities::GetAllModuleNodes(EmitterViewModel.ToSharedRef());
+	for(UNiagaraNodeFunctionCall* ModuleNode : ModuleNodes)
+	{
+		// we skip over sim stage modules here as we want to add them to their respective sim stage group items instead
+		if(SimStageModules.Contains(ModuleNode) || EventHandlerModules.Contains(ModuleNode))
+		{
+			continue;
+		}
+		
+		UNiagaraHierarchyModule* HierarchyModule = nullptr;
+		FNiagaraHierarchyIdentity ModuleIdentity;
+		ModuleIdentity.Guids.Add(ModuleNode->NodeGuid);
+		UNiagaraHierarchyItemBase** FoundItem = SourceRoot->GetChildrenMutable().FindByPredicate([ModuleIdentity](UNiagaraHierarchyItemBase* ItemBase)
+		{
+			return ItemBase->GetPersistentIdentity() == ModuleIdentity;
+		});
+		
+		if(FoundItem == nullptr)		
+		{
+			HierarchyModule = SourceRoot->AddChild<UNiagaraHierarchyModule>();
+			HierarchyModule->Initialize(*ModuleNode);
+		}
+		else
+		{
+			HierarchyModule = CastChecked<UNiagaraHierarchyModule>(*FoundItem);
+		}
+		
+		NewItems.Add(HierarchyModule);
+		UsageMap.FindOrAdd(FNiagaraStackGraphUtilities::GetOutputNodeUsage(*ModuleNode)).Add(HierarchyModule);
+	}
+
+	const TArray<FNiagaraEventScriptProperties> ScriptProperties = EmitterViewModel->GetEmitter().GetEmitterData()->GetEventHandlers();
+	for(const FNiagaraEventScriptProperties& ScriptPropertiesItem : ScriptProperties)
+	{
+		UNiagaraHierarchyEventHandler* HierarchyEventHandler = nullptr;
+		FNiagaraHierarchyIdentity EventHandlerIdentity;
+		EventHandlerIdentity.Guids.Add(ScriptPropertiesItem.Script->GetUsageId());
+		EventHandlerIdentity.Guids.Add(ScriptPropertiesItem.SourceEmitterID);
+
+		UNiagaraHierarchyItemBase** FoundItem = SourceRoot->GetChildrenMutable().FindByPredicate([EventHandlerIdentity](UNiagaraHierarchyItemBase* ItemBase)
+		{
+			return ItemBase->GetPersistentIdentity() == EventHandlerIdentity;
+		});
+		
+		if(FoundItem == nullptr)		
+		{
+			HierarchyEventHandler = SourceRoot->AddChild<UNiagaraHierarchyEventHandler>();
+			HierarchyEventHandler->Initialize(ScriptPropertiesItem);
+		}
+		else
+		{
+			HierarchyEventHandler = CastChecked<UNiagaraHierarchyEventHandler>(*FoundItem);
+		}
+		
+		NewItems.Add(HierarchyEventHandler);
+		HierarchyEventHandlers.Add(HierarchyEventHandler);
+	}
+
+	// We add sim stages here
+	const TArray<UNiagaraSimulationStageBase*>& SimStages = EmitterViewModel->GetEmitter().GetEmitterData()->GetSimulationStages();
+	for(UNiagaraSimulationStageBase* SimStage : SimStages)
+	{
+		UNiagaraHierarchySimStage* HierarchySimStage = nullptr;
+		FNiagaraHierarchyIdentity SimStageID;
+		SimStageID.Guids.Add(SimStage->GetMergeId());
+
+		UNiagaraHierarchyItemBase** FoundItem = SourceRoot->GetChildrenMutable().FindByPredicate([SimStageID](UNiagaraHierarchyItemBase* ItemBase)
+		{
+			return ItemBase->GetPersistentIdentity() == SimStageID;
+		});
+		
+		if(FoundItem == nullptr)		
+		{
+			HierarchySimStage = SourceRoot->AddChild<UNiagaraHierarchySimStage>();
+			HierarchySimStage->Initialize(*SimStage);
+		}
+		else
+		{
+			HierarchySimStage = CastChecked<UNiagaraHierarchySimStage>(*FoundItem);
+		}
+		
+		NewItems.Add(HierarchySimStage);
+		HierarchySimStages.Add(HierarchySimStage);
+	}
+	
+	// We create hierarchy renderers here
+	const TArray<UNiagaraRendererProperties*>& Renderers = EmitterViewModel->GetEmitter().GetEmitterData()->GetRenderers();
+	for(UNiagaraRendererProperties* RendererProperties : Renderers)
+	{
+		UNiagaraHierarchyRenderer* HierarchyRenderer = nullptr;
+		FNiagaraHierarchyIdentity RendererIdentity;
+		RendererIdentity.Guids.Add(RendererProperties->GetMergeId());
+
+		UNiagaraHierarchyItemBase** FoundItem = SourceRoot->GetChildrenMutable().FindByPredicate([RendererIdentity](UNiagaraHierarchyItemBase* ItemBase)
+		{
+			return ItemBase->GetPersistentIdentity() == RendererIdentity;
+		});
+		
+		if(FoundItem == nullptr)		
+		{
+			HierarchyRenderer = SourceRoot->AddChild<UNiagaraHierarchyRenderer>();
+			HierarchyRenderer->Initialize(*RendererProperties);
+		}
+		else
+		{
+			HierarchyRenderer = CastChecked<UNiagaraHierarchyRenderer>(*FoundItem);
+		}
+		
+		NewItems.Add(HierarchyRenderer);
+		HierarchyRenderers.Add(HierarchyRenderer);
+	}
+
+	SourceRoot->GetChildrenMutable().Empty();
+	SourceRoot->GetChildrenMutable().Append(NewItems);
+
+	// we force a sync so we can access the section data for the source items
+	SourceRootViewModel->SyncViewModelsToData();
+	
+	// Now we create a section for each usage case that has at least one element and link up the respective sections
+	UNiagaraHierarchySection* EmitterSpawnSection = nullptr;
+	if(UsageMap.Contains(ENiagaraScriptUsage::EmitterSpawnScript))
+	{
+		FText SectionName = FText::FromString("Emitter Spawn");
+		UNiagaraHierarchySection** FoundHierarchySection = SourceRootViewModel->GetDataMutable<UNiagaraHierarchyRoot>()->GetSectionDataMutable().FindByPredicate([SectionName](UNiagaraHierarchySection* Candidate)
+		{
+			return Candidate->GetSectionNameAsText().EqualTo(SectionName);
+		});
+		
+		if(FoundHierarchySection == nullptr)
+		{
+			EmitterSpawnSection = SourceRootViewModel->GetDataMutable<UNiagaraHierarchyRoot>()->AddSection(SectionName);
+		}
+		else
+		{
+			EmitterSpawnSection = *FoundHierarchySection;
+		}
+
+		NewSections.Add(EmitterSpawnSection);
+
+		for(UNiagaraHierarchyModule* Module : UsageMap[ENiagaraScriptUsage::EmitterSpawnScript])
+		{
+			TSharedPtr<FNiagaraFunctionViewModel> ModuleViewModel = StaticCastSharedPtr<FNiagaraFunctionViewModel>(SourceRootViewModel->FindViewModelForChild(Module));
+			ModuleViewModel->SetSection(*EmitterSpawnSection);
+		}
+	}
+
+	UNiagaraHierarchySection* EmitterUpdateSection = nullptr;
+	if(UsageMap.Contains(ENiagaraScriptUsage::EmitterUpdateScript))
+	{
+		FText SectionName = FText::FromString("Emitter Update");
+		UNiagaraHierarchySection** FoundHierarchySection = SourceRootViewModel->GetDataMutable<UNiagaraHierarchyRoot>()->GetSectionDataMutable().FindByPredicate([SectionName](UNiagaraHierarchySection* Candidate)
+		{
+			return Candidate->GetSectionNameAsText().EqualTo(SectionName);
+		});
+		
+		if(FoundHierarchySection == nullptr)
+		{
+			EmitterUpdateSection = SourceRootViewModel->GetDataMutable<UNiagaraHierarchyRoot>()->AddSection(SectionName);
+		}
+		else
+		{
+			EmitterUpdateSection = *FoundHierarchySection;
+		}
+
+		NewSections.Add(EmitterUpdateSection);
+
+		for(UNiagaraHierarchyModule* Module : UsageMap[ENiagaraScriptUsage::EmitterUpdateScript])
+		{
+			TSharedPtr<FNiagaraFunctionViewModel> ModuleViewModel = StaticCastSharedPtr<FNiagaraFunctionViewModel>(SourceRootViewModel->FindViewModelForChild(Module));
+			ModuleViewModel->SetSection(*EmitterUpdateSection);
+		}
+	}
+
+	UNiagaraHierarchySection* ParticleSpawnSection = nullptr;
+	if(UsageMap.Contains(ENiagaraScriptUsage::ParticleSpawnScript))
+	{
+		FText SectionName = FText::FromString("Particle Spawn");
+		UNiagaraHierarchySection** FoundHierarchySection = SourceRootViewModel->GetDataMutable<UNiagaraHierarchyRoot>()->GetSectionDataMutable().FindByPredicate([SectionName](UNiagaraHierarchySection* Candidate)
+		{
+			return Candidate->GetSectionNameAsText().EqualTo(SectionName);
+		});
+		
+		if(FoundHierarchySection == nullptr)
+		{
+			ParticleSpawnSection = SourceRootViewModel->GetDataMutable<UNiagaraHierarchyRoot>()->AddSection(SectionName);
+		}
+		else
+		{
+			ParticleSpawnSection = *FoundHierarchySection;
+		}
+
+		NewSections.Add(ParticleSpawnSection);
+
+		for(UNiagaraHierarchyModule* Module : UsageMap[ENiagaraScriptUsage::ParticleSpawnScript])
+		{
+			TSharedPtr<FNiagaraFunctionViewModel> ModuleViewModel = StaticCastSharedPtr<FNiagaraFunctionViewModel>(SourceRootViewModel->FindViewModelForChild(Module));
+			ModuleViewModel->SetSection(*ParticleSpawnSection);
+		}
+	}
+
+	UNiagaraHierarchySection* ParticleUpdateSection = nullptr;
+	if(UsageMap.Contains(ENiagaraScriptUsage::ParticleUpdateScript))
+	{
+		FText SectionName = FText::FromString("Particle Update");
+		UNiagaraHierarchySection** FoundHierarchySection = SourceRootViewModel->GetDataMutable<UNiagaraHierarchyRoot>()->GetSectionDataMutable().FindByPredicate([SectionName](UNiagaraHierarchySection* Candidate)
+		{
+			return Candidate->GetSectionNameAsText().EqualTo(SectionName);
+		});
+		
+		if(FoundHierarchySection == nullptr)
+		{
+			ParticleUpdateSection = SourceRootViewModel->GetDataMutable<UNiagaraHierarchyRoot>()->AddSection(SectionName);
+		}
+		else
+		{
+			ParticleUpdateSection = *FoundHierarchySection;
+		}
+
+		NewSections.Add(ParticleUpdateSection);
+
+		for(UNiagaraHierarchyModule* Module : UsageMap[ENiagaraScriptUsage::ParticleUpdateScript])
+		{
+			TSharedPtr<FNiagaraFunctionViewModel> ModuleViewModel = StaticCastSharedPtr<FNiagaraFunctionViewModel>(SourceRootViewModel->FindViewModelForChild(Module));
+			ModuleViewModel->SetSection(*ParticleUpdateSection);
+		}
+	}
+
+	UNiagaraHierarchySection* EventHandlerSection = nullptr;
+	if(HierarchyEventHandlers.Num() > 0)
+	{
+		FText SectionName = FText::FromString("Events");
+		UNiagaraHierarchySection** FoundHierarchySection = SourceRootViewModel->GetDataMutable<UNiagaraHierarchyRoot>()->GetSectionDataMutable().FindByPredicate([SectionName](UNiagaraHierarchySection* Candidate)
+		{
+			return Candidate->GetSectionNameAsText().EqualTo(SectionName);
+		});
+		
+		if(FoundHierarchySection == nullptr)
+		{
+			EventHandlerSection = SourceRootViewModel->GetDataMutable<UNiagaraHierarchyRoot>()->AddSection(SectionName);
+		}
+		else
+		{
+			EventHandlerSection = *FoundHierarchySection;
+		}
+
+		NewSections.Add(EventHandlerSection);
+
+		for(UNiagaraHierarchyEventHandler* EventHandler : HierarchyEventHandlers)
+		{
+			TSharedPtr<FNiagaraHierarchyEventHandlerViewModel> EventHandlerViewModel = StaticCastSharedPtr<FNiagaraHierarchyEventHandlerViewModel>(SourceRootViewModel->FindViewModelForChild(EventHandler));
+			EventHandlerViewModel->SetSection(*EventHandlerSection);
+		}
+	}
+
+	UNiagaraHierarchySection* SimulationStagesSection = nullptr;
+	if(HierarchySimStages.Num() > 0)
+	{
+		FText SectionName = FText::FromString("Sim Stages");
+		UNiagaraHierarchySection** FoundHierarchySection = SourceRootViewModel->GetDataMutable<UNiagaraHierarchyRoot>()->GetSectionDataMutable().FindByPredicate([SectionName](UNiagaraHierarchySection* Candidate)
+		{
+			return Candidate->GetSectionNameAsText().EqualTo(SectionName);
+		});
+		
+		if(FoundHierarchySection == nullptr)
+		{
+			SimulationStagesSection = SourceRootViewModel->GetDataMutable<UNiagaraHierarchyRoot>()->AddSection(SectionName);
+		}
+		else
+		{
+			SimulationStagesSection = *FoundHierarchySection;
+		}
+
+		NewSections.Add(SimulationStagesSection);
+
+		for(UNiagaraHierarchySimStage* SimStage : HierarchySimStages)
+		{
+			TSharedPtr<FNiagaraHierarchySimStageViewModel> SimStageViewModel = StaticCastSharedPtr<FNiagaraHierarchySimStageViewModel>(SourceRootViewModel->FindViewModelForChild(SimStage));
+			SimStageViewModel->SetSection(*SimulationStagesSection);
+		}
+	}
+
+	UNiagaraHierarchySection* RenderersSection = nullptr; 
+	if(HierarchyRenderers.Num() > 0)
+	{
+		FText SectionName = FText::FromString("Renderers");
+		UNiagaraHierarchySection** FoundHierarchySection = SourceRootViewModel->GetDataMutable<UNiagaraHierarchyRoot>()->GetSectionDataMutable().FindByPredicate([SectionName](UNiagaraHierarchySection* Candidate)
+		{
+			return Candidate->GetSectionNameAsText().EqualTo(SectionName);
+		});
+		
+		if(FoundHierarchySection == nullptr)
+		{
+			RenderersSection = SourceRootViewModel->GetDataMutable<UNiagaraHierarchyRoot>()->AddSection(SectionName);
+		}
+		else
+		{
+			RenderersSection = *FoundHierarchySection;
+		}
+
+		NewSections.Add(RenderersSection);
+		
+		for(UNiagaraHierarchyRenderer* Renderer : HierarchyRenderers)
+		{
+			TSharedPtr<FNiagaraHierarchyRendererViewModel> RendererViewModel = StaticCastSharedPtr<FNiagaraHierarchyRendererViewModel>(SourceRootViewModel->FindViewModelForChild(Renderer));
+			RendererViewModel->SetSection(*RenderersSection);
+		}
+	}
+
+	// this will implicitly sort the sections as well as get rid of outdated ones
+	SourceRoot->GetSectionDataMutable().Empty();
+	SourceRoot->GetSectionDataMutable().Append(NewSections);
+
+	// force a sync so we have the view models for the sections available
+	SourceRootViewModel->SyncViewModelsToData();
+
+	for(TSharedPtr<FNiagaraHierarchySectionViewModel> SectionViewModel : SourceRootViewModel->GetSectionViewModels())
+	{
+		if(EmitterSpawnSection != nullptr && SectionViewModel->GetData() == EmitterSpawnSection)
+		{
+			SectionViewModel->SetSectionImage(FNiagaraEditorStyle::Get().GetBrush("NiagaraEditor.Emitter.SpawnIcon"));
+		}
+
+		if(EmitterUpdateSection != nullptr && SectionViewModel->GetData() == EmitterUpdateSection)
+		{
+			SectionViewModel->SetSectionImage(FNiagaraEditorStyle::Get().GetBrush("NiagaraEditor.Emitter.UpdateIcon"));
+		}
+
+		if(ParticleSpawnSection != nullptr && SectionViewModel->GetData() == ParticleSpawnSection)
+		{
+			SectionViewModel->SetSectionImage(FNiagaraEditorStyle::Get().GetBrush("NiagaraEditor.Particle.SpawnIcon"));
+		}
+
+		if(ParticleUpdateSection != nullptr && SectionViewModel->GetData() == ParticleUpdateSection)
+		{
+			SectionViewModel->SetSectionImage(FNiagaraEditorStyle::Get().GetBrush("NiagaraEditor.Particle.UpdateIcon"));
+		}
+
+		if(EventHandlerSection != nullptr && SectionViewModel->GetData() == EventHandlerSection)
+		{
+			SectionViewModel->SetSectionImage(FNiagaraEditorStyle::Get().GetBrush("NiagaraEditor.EventIcon"));
+		}
+
+		if(SimulationStagesSection != nullptr && SectionViewModel->GetData() == SimulationStagesSection)
+		{
+			SectionViewModel->SetSectionImage(FNiagaraEditorStyle::Get().GetBrush("NiagaraEditor.SimulationStageIcon"));
+		}
+
+		if(RenderersSection != nullptr && SectionViewModel->GetData() == RenderersSection)
+		{
+			SectionViewModel->SetSectionImage(FNiagaraEditorStyle::Get().GetBrush("NiagaraEditor.RenderIcon"));
+		}
+	}
+}
+
+void UNiagaraSummaryViewViewModel::SetupCommands()
+{
+	// no custom commands yet
+}
+
+TSharedRef<FNiagaraHierarchyDragDropOp> UNiagaraSummaryViewViewModel::CreateDragDropOp(TSharedRef<FNiagaraHierarchyItemViewModelBase> Item)
+{
+	if(UNiagaraHierarchyCategory* HierarchyCategory = Cast<UNiagaraHierarchyCategory>(Item->GetDataMutable()))
+	{
+		TSharedRef<FNiagaraHierarchyDragDropOp> CategoryDragDropOp = MakeShared<FNiagaraHierarchyDragDropOp>(Item);
+		CategoryDragDropOp->SetAdditionalLabel(FText::FromString(Item->ToString()));
+		CategoryDragDropOp->Construct();
+		return CategoryDragDropOp;
+	}
+	else if(UNiagaraHierarchyModuleInput* ModuleInput = Cast<UNiagaraHierarchyModuleInput>(Item->GetDataMutable()))
+	{
+		TSharedPtr<FNiagaraModuleInputViewModel> ModuleInputViewModel = StaticCastSharedRef<FNiagaraModuleInputViewModel>(Item);
+		TSharedRef<FNiagaraHierarchyDragDropOp> ModuleInputDragDropOp = MakeShared<FNiagaraHierarchyInputParameterHierarchyDragDropOp>(ModuleInputViewModel);
+		ModuleInputDragDropOp->Construct();
+		return ModuleInputDragDropOp;
+	}
+	else if(UNiagaraHierarchyItem* ItemData = Cast<UNiagaraHierarchyItem>(Item->GetDataMutable()))
+	{
+		TSharedRef<FNiagaraHierarchyDragDropOp> ObjectPropertyDragDropOp = MakeShared<FNiagaraHierarchyDragDropOp>(Item);
+		ObjectPropertyDragDropOp->SetAdditionalLabel(FText::FromString(Item->ToString()));
+		ObjectPropertyDragDropOp->Construct();
+		return ObjectPropertyDragDropOp;
+	}
+
+	check(false);
+	return MakeShared<FNiagaraHierarchyDragDropOp>(nullptr);
+}
+
+TArray<TTuple<UClass*, FOnGetDetailCustomizationInstance>> UNiagaraSummaryViewViewModel::GetInstanceCustomizations()
+{
+	return {};
+}
+
+TMap<FGuid, UObject*> UNiagaraSummaryViewViewModel::GetObjectsForProperties()
+{
+	TMap<FGuid, UObject*> GuidToObjectMap;
+	for(UNiagaraRendererProperties* RendererProperties : GetEmitterViewModel()->GetEmitter().GetEmitterData()->GetRenderers())
+	{
+		GuidToObjectMap.Add(RendererProperties->GetMergeId(), RendererProperties);
+	}
+
+	for(UNiagaraSimulationStageBase* SimulationStage : GetEmitterViewModel()->GetEmitter().GetEmitterData()->GetSimulationStages())
+	{
+		GuidToObjectMap.Add(SimulationStage->GetMergeId(), SimulationStage);
+	}
+
+	return GuidToObjectMap;
+}
+
+void UNiagaraSummaryViewViewModel::OnScriptGraphChanged(const FEdGraphEditAction& Action, const UNiagaraScript& Script)
+{
+	// since OnScriptGraphChanged will be called many times in a row, we avoid refreshing too much by requesting a full refresh for next frame instead
+	// it will request a single refresh regardless of how often this is called until the refresh is done
+	if((Action.Action & EEdGraphActionType::GRAPHACTION_RemoveNode) != 0)
+	{
+ 		for(const UEdGraphNode* RemovedNode : Action.Nodes)
+		{
+			if(FunctionCallCache.Contains(RemovedNode->NodeGuid))
+			{
+				FunctionCallCache.Remove(RemovedNode->NodeGuid);
+			}
+		}
+	}
+	
+	RequestFullRefreshNextFrame();
+}
+
+void UNiagaraSummaryViewViewModel::OnRenderersChanged()
+{
+	RequestFullRefreshNextFrame();
+}
+
+void UNiagaraSummaryViewViewModel::OnSimStagesChanged()
+{
+	RequestFullRefreshNextFrame();
+}
+
+void UNiagaraSummaryViewViewModel::OnEventHandlersChanged()
+{
+	RequestFullRefreshNextFrame();
+}
+
+UNiagaraNodeFunctionCall* UNiagaraSummaryViewViewModel::GetFunctionCallNode(const FGuid& NodeIdentity)
+{
+	if(FunctionCallCache.Contains(NodeIdentity))
+	{
+		if(FunctionCallCache[NodeIdentity].IsValid())
+		{
+			return FunctionCallCache[NodeIdentity].Get();
+		}
+		else
+		{
+			FunctionCallCache.Remove(NodeIdentity);
+		}
+	}
+
+	if(UNiagaraNodeFunctionCall* FoundFunctionCall = FNiagaraStackGraphUtilities::FindFunctionCallNode(NodeIdentity, GetEmitterViewModel()))
+	{
+		FunctionCallCache.Add(NodeIdentity, FoundFunctionCall);
+		return FoundFunctionCall;
+	}
+
+	return nullptr;
+}
+
+void UNiagaraSummaryViewViewModel::ClearFunctionCallNodeCache(const FGuid& NodeIdentity)
+{
+	if(FunctionCallCache.Contains(NodeIdentity))
+	{
+		FunctionCallCache.Remove(NodeIdentity);
+	}
+}
+
+TOptional<FInputData> UNiagaraSummaryViewViewModel::GetInputData(const UNiagaraHierarchyModuleInput& Input)
+{
+	if(UNiagaraNodeFunctionCall* FunctionCall = GetFunctionCallNode(Input.GetPersistentIdentity().Guids[0]))
+	{
+		if(UNiagaraGraph* CalledGraph = FunctionCall->GetCalledGraph())
+		{
+			if(UNiagaraScriptVariable* MatchingScriptVariable = CalledGraph->GetScriptVariable(Input.GetPersistentIdentity().Guids[1]))
+			{
+				FInputData InputData;
+				InputData.InputName = MatchingScriptVariable->Variable.GetName();
+				InputData.Type = MatchingScriptVariable->Variable.GetType();
+				InputData.MetaData = MatchingScriptVariable->Metadata;
+				InputData.bIsStatic = FNiagaraParameterHandle(InputData.InputName).IsModuleHandle() ? false : true;
+				InputData.FunctionCallNode = FunctionCall;
+				InputData.ChildrenInputGuids = CalledGraph->GetChildScriptVariableGuidsForInput(MatchingScriptVariable->Metadata.GetVariableGuid());
+				return InputData;
+			}			
+		}
+	}
+
+	return TOptional<FInputData>();
+}
+
+void SNiagaraHierarchyModule::Construct(const FArguments& InArgs, TSharedPtr<FNiagaraFunctionViewModel> InModuleViewModel)
+{
+	ModuleViewModel = InModuleViewModel;
+	
+	ChildSlot
+	[
+		SNew(SHorizontalBox)
+		+ SHorizontalBox::Slot()
+		.AutoWidth()
+		.Padding(2.f)
+		[
+			SNew(SImage)
+			.Image(FNiagaraEditorStyle::Get().GetBrush("NiagaraEditor.Module.DynamicInput"))
+			.ToolTipText(LOCTEXT("DynamicInputIconTooltip", "Dynamic Inputs can not be dragged directly into the hierarchy. Please use the entire module input or individual inputs beneath."))
+			.Visibility(InModuleViewModel->IsDynamicInput() ? EVisibility::Visible : EVisibility::Collapsed)
+		]
+		+ SHorizontalBox::Slot()
+		[
+			SAssignNew(InlineEditableTextBlock, SInlineEditableTextBlock)
+			.Text(this, &SNiagaraHierarchyModule::GetModuleDisplayName)
+			.IsReadOnly(true)
+		]
+	];
+}
+
+FText SNiagaraHierarchyModule::GetModuleDisplayName() const
+{
+	if(ModuleViewModel.IsValid())
+	{
+		if(ModuleViewModel.Pin()->GetFunctionCallNode().IsValid())
+		{
+			return ModuleViewModel.Pin()->GetFunctionCallNode()->GetNodeTitle(ENodeTitleType::ListView);
+		}
+	}
+
+	return FText::GetEmpty();
+}
+
+#undef LOCTEXT_NAMESPACE
diff --git a/Engine/Plugins/FX/Niagara/Source/NiagaraEditor/Private/ViewModels/HierarchyEditor/NiagaraUserParametersHierarchyViewModel.cpp b/Engine/Plugins/FX/Niagara/Source/NiagaraEditor/Private/ViewModels/HierarchyEditor/NiagaraUserParametersHierarchyViewModel.cpp
--- a/Engine/Plugins/FX/Niagara/Source/NiagaraEditor/Private/ViewModels/HierarchyEditor/NiagaraUserParametersHierarchyViewModel.cpp
+++ b/Engine/Plugins/FX/Niagara/Source/NiagaraEditor/Private/ViewModels/HierarchyEditor/NiagaraUserParametersHierarchyViewModel.cpp
@@ -10,11 +10,6 @@
 #include "SDropTarget.h"
 #include "Customizations/NiagaraScriptVariableCustomization.h"
 #include "ViewModels/NiagaraSystemViewModel.h"
-#include "ViewModels/Stack/NiagaraStackViewModel.h"
-#include "Widgets/SNiagaraParameterName.h"
-#include "Widgets/Layout/SWrapBox.h"
-#include "Widgets/Views/STableRow.h"
-#include "Widgets/SToolTip.h"
 
 #define LOCTEXT_NAMESPACE "NiagaraUserParametersHierarchyEditor"
 
@@ -22,15 +17,6 @@ void UNiagaraHierarchyUserParameter::Initialize(UNiagaraScriptVariable& InUserPa
 {
 	UserParameterScriptVariable = &InUserParameterScriptVariable;
 	SetIdentity(FNiagaraHierarchyIdentity({InUserParameterScriptVariable.Metadata.GetVariableGuid()}, {}));
-	System = &InSystem;
-}
-
-void UNiagaraHierarchyUserParameter::RefreshDataInternal()
-{	
-	if(System == nullptr || FNiagaraEditorUtilities::FindScriptVariableForUserParameter(GetPersistentIdentity().Guids[0], *System) == nullptr)
-	{
-		Finalize();
-	}
 }
 
 UObject* FNiagaraHierarchyUserParameterViewModel::GetDataForEditing()
@@ -43,6 +29,13 @@ UObject* FNiagaraHierarchyUserParameterViewModel::GetDataForEditing()
 	return ScriptVariable.Get();
 }
 
+bool FNiagaraHierarchyUserParameterViewModel::DoesExternalDataStillExist(const UNiagaraHierarchyDataRefreshContext* Context) const
+{
+	const UNiagaraHierarchyUserParameterRefreshContext* UserParameterRefreshContext = CastChecked<UNiagaraHierarchyUserParameterRefreshContext>(Context);
+	const UNiagaraSystem* System = UserParameterRefreshContext->GetSystem();
+	return FNiagaraEditorUtilities::FindScriptVariableForUserParameter(GetData()->GetPersistentIdentity().Guids[0], *System) != nullptr;
+}
+
 TSharedRef<FNiagaraSystemViewModel> UNiagaraUserParametersHierarchyViewModel::GetSystemViewModel() const
 {
 	TSharedPtr<FNiagaraSystemViewModel> SystemViewModelPinned = SystemViewModelWeak.Pin();
@@ -55,6 +48,11 @@ void UNiagaraUserParametersHierarchyViewModel::Initialize(TSharedRef<FNiagaraSys
 	SystemViewModelWeak = InSystemViewModel;
 	UNiagaraHierarchyViewModelBase::Initialize();
 
+	UNiagaraHierarchyUserParameterRefreshContext* UserParameterRefreshContext = NewObject<UNiagaraHierarchyUserParameterRefreshContext>(this);
+	UserParameterRefreshContext->SetSystem(&InSystemViewModel->GetSystem());
+
+	SetRefreshContext(UserParameterRefreshContext);
+	
 	UNiagaraSystemEditorData& SystemEditorData = InSystemViewModel->GetEditorData();
 	SystemEditorData.OnUserParameterScriptVariablesSynced().AddUObject(this, &UNiagaraUserParametersHierarchyViewModel::ForceFullRefresh);
 }
@@ -75,7 +73,7 @@ TSharedRef<SWidget> FNiagaraUserParameterHierarchyDragDropOp::CreateCustomDecora
 	return FNiagaraParameterUtilities::GetParameterWidget(GetUserParameter(), true, false);
 }
 
-UNiagaraHierarchyRoot* UNiagaraUserParametersHierarchyViewModel::GetHierarchyDataRoot() const
+UNiagaraHierarchyRoot* UNiagaraUserParametersHierarchyViewModel::GetHierarchyRoot() const
 {
 	UNiagaraHierarchyRoot* RootItem = GetSystemViewModel()->GetEditorData().UserParameterHierarchy;
 
@@ -87,18 +85,18 @@ TSharedPtr<FNiagaraHierarchyItemViewModelBase> UNiagaraUserParametersHierarchyVi
 {
 	if(UNiagaraHierarchyUserParameter* UserParameter = Cast<UNiagaraHierarchyUserParameter>(ItemBase))
 	{
-		return MakeShared<FNiagaraHierarchyUserParameterViewModel>(UserParameter, Parent, this);
+		return MakeShared<FNiagaraHierarchyUserParameterViewModel>(UserParameter, Parent, this, Parent.IsValid() ? Parent->IsForHierarchy() : false);
 	}
 	else if(UNiagaraHierarchyCategory* Category = Cast<UNiagaraHierarchyCategory>(ItemBase))
 	{
-		return MakeShared<FNiagaraHierarchyCategoryViewModel>(Category, Parent, this);
+		return MakeShared<FNiagaraHierarchyCategoryViewModel>(Category, Parent, this, Parent.IsValid() ? Parent->IsForHierarchy() : false);
 	}
 
 	check(false);
 	return nullptr;
 }
 
-void UNiagaraUserParametersHierarchyViewModel::PrepareSourceItems()
+void UNiagaraUserParametersHierarchyViewModel::PrepareSourceItems(UNiagaraHierarchyRoot* SourceRoot, TSharedPtr<FNiagaraHierarchyRootViewModel>)
 {	
 	TArray<FNiagaraVariable> UserParameters;
 	GetSystemViewModel()->GetSystem().GetExposedParameters().GetUserParameters(UserParameters);
@@ -127,9 +125,6 @@ void UNiagaraUserParametersHierarchyViewModel::PrepareSourceItems()
 		UserParameterHierarchyObject->Initialize(*ScriptVariable.Get(), GetSystemViewModel()->GetSystem());
 		SourceRoot->GetChildrenMutable().Add(UserParameterHierarchyObject);
 	}
-
-	// after we synced the data, we sync their view models that we display in the UI
-	SourceViewModelRoot->SyncToData();
 }
 
 void UNiagaraUserParametersHierarchyViewModel::SetupCommands()
diff --git a/Engine/Plugins/FX/Niagara/Source/NiagaraEditor/Private/ViewModels/NiagaraEmitterViewModel.cpp b/Engine/Plugins/FX/Niagara/Source/NiagaraEditor/Private/ViewModels/NiagaraEmitterViewModel.cpp
--- a/Engine/Plugins/FX/Niagara/Source/NiagaraEditor/Private/ViewModels/NiagaraEmitterViewModel.cpp
+++ b/Engine/Plugins/FX/Niagara/Source/NiagaraEditor/Private/ViewModels/NiagaraEmitterViewModel.cpp
@@ -23,6 +23,7 @@
 #include "Framework/Docking/TabManager.h"
 #include "Widgets/SWindow.h"
 #include "ViewModels/NiagaraEmitterHandleViewModel.h"
+#include "ViewModels/HierarchyEditor/NiagaraSummaryViewViewModel.h"
 
 #define LOCTEXT_NAMESPACE "EmitterEditorViewModel"
 
@@ -40,6 +41,7 @@ namespace NiagaraCommands
 FNiagaraEmitterViewModel::FNiagaraEmitterViewModel(bool bInIsForDataProcessingOnly)
 	: EmitterWeakPtr(FVersionedNiagaraEmitterWeakPtr(nullptr, FGuid()))
 	, SharedScriptViewModel(MakeShareable(new FNiagaraScriptViewModel(LOCTEXT("SharedDisplayName", "Graph"), ENiagaraParameterEditMode::EditAll, bInIsForDataProcessingOnly)))
+	, SummaryViewHierarchyViewModel(NewObject<UNiagaraSummaryViewViewModel>())
 	, bUpdatingSelectionInternally(false)
 	, ExecutionStateEnum(StaticEnum<ENiagaraExecutionState>())
 {	
@@ -64,6 +66,17 @@ void FNiagaraEmitterViewModel::Cleanup()
 		SharedScriptViewModel.Reset();
 	}
 
+	if(SummaryViewHierarchyViewModel != nullptr)
+	{
+		SummaryViewHierarchyViewModel->Finalize();
+		SummaryViewHierarchyViewModel = nullptr;
+	}
+
+	if(EmitterWeakPtr.IsValid())
+	{
+		GetOrCreateEditorData().OnPersistentDataChanged().RemoveAll(this);
+	}
+
 	RemoveScriptEventHandlers();
 }
 
@@ -71,13 +84,17 @@ bool FNiagaraEmitterViewModel::Initialize(const FVersionedNiagaraEmitter& InEmit
 {
 	SetEmitter(InEmitter);
 	SetSimulation(InSimulation);
-	bSummaryIsInEditMode = false;
+	SummaryViewHierarchyViewModel->Initialize(AsShared());
+	SummaryViewHierarchyViewModel->OnHierarchyChanged().AddSP(this, &FNiagaraEmitterViewModel::OnSummaryViewHierarchyChanged);
+	SummaryViewHierarchyViewModel->OnHierarchyPropertiesChanged().AddSP(this, &FNiagaraEmitterViewModel::OnSummaryViewHierarchyChanged);
 	return true;
 }
 
 void FNiagaraEmitterViewModel::Reset()
 {
 	SetEmitter(FVersionedNiagaraEmitter());
+	SummaryViewHierarchyViewModel->OnHierarchyChanged().RemoveAll(this);
+	SummaryViewHierarchyViewModel->OnHierarchyPropertiesChanged().RemoveAll(this);
 	SetSimulation(nullptr);
 }
 
@@ -302,6 +319,11 @@ TSharedRef<FNiagaraScriptViewModel> FNiagaraEmitterViewModel::GetSharedScriptVie
 	return SharedScriptViewModel.ToSharedRef();
 }
 
+UNiagaraSummaryViewViewModel* FNiagaraEmitterViewModel::GetSummaryHierarchyViewModel()
+{
+	return SummaryViewHierarchyViewModel;
+}
+
 const UNiagaraEmitterEditorData& FNiagaraEmitterViewModel::GetEditorData() const
 {
 	check(EmitterWeakPtr.Emitter.IsValid());
@@ -582,4 +604,22 @@ void FNiagaraEmitterViewModel::OnEmitterPropertiesChanged()
 	}
 }
 
+void FNiagaraEmitterViewModel::OnSummaryViewHierarchyChanged()
+{
+	GetOrCreateEditorData().OnPersistentDataChanged().Broadcast();
+}
+
+FString FNiagaraEmitterViewModel::GetReferencerName() const
+{
+	return TEXT("EmitterViewModel");
+}
+
+void FNiagaraEmitterViewModel::AddReferencedObjects(FReferenceCollector& Collector)
+{
+	if(SummaryViewHierarchyViewModel != nullptr)
+	{
+		Collector.AddReferencedObject(SummaryViewHierarchyViewModel);
+	}
+}
+
 #undef LOCTEXT_NAMESPACE
diff --git a/Engine/Plugins/FX/Niagara/Source/NiagaraEditor/Private/ViewModels/Stack/NiagaraStackEmitterSettingsGroup.cpp b/Engine/Plugins/FX/Niagara/Source/NiagaraEditor/Private/ViewModels/Stack/NiagaraStackEmitterSettingsGroup.cpp
--- a/Engine/Plugins/FX/Niagara/Source/NiagaraEditor/Private/ViewModels/Stack/NiagaraStackEmitterSettingsGroup.cpp
+++ b/Engine/Plugins/FX/Niagara/Source/NiagaraEditor/Private/ViewModels/Stack/NiagaraStackEmitterSettingsGroup.cpp
@@ -16,10 +16,12 @@
 #include "PropertyHandle.h"
 #include "ScopedTransaction.h"
 #include "ViewModels/NiagaraEmitterHandleViewModel.h"
-#include "ViewModels/NiagaraSystemSelectionViewModel.h"
 #include "ViewModels/Stack/NiagaraStackViewModel.h"
 #include "Styling/AppStyle.h"
 #include "IDetailTreeNode.h"
+#include "NiagaraNodeFunctionCall.h"
+#include "Toolkits/SystemToolkitModes/NiagaraSystemToolkitModeBase.h"
+#include "Widgets/SNiagaraHierarchy.h"
 
 #include UE_INLINE_GENERATED_CPP_BY_NAME(NiagaraStackEmitterSettingsGroup)
 
@@ -325,26 +327,9 @@ FText UNiagaraStackEmitterSummaryItem::GetIconText() const
 
 void UNiagaraStackEmitterSummaryItem::RefreshChildrenInternal(const TArray<UNiagaraStackEntry*>& CurrentChildren, TArray<UNiagaraStackEntry*>& NewChildren, TArray<FStackIssue>& NewIssues)
 {
-	if (GetEmitterViewModel()->GetSummaryIsInEditMode())
-	{
-		if (SummaryEditorData == nullptr)
-		{
-			SummaryEditorData = NewObject<UNiagaraStackObject>(this);
-			bool bIsTopLevelObject = true;
-			SummaryEditorData->Initialize(CreateDefaultChildRequiredData(), &GetEmitterViewModel()->GetOrCreateEditorData(), bIsTopLevelObject, GetStackEditorDataKey());
-			SummaryEditorData->SetOnSelectRootNodes(UNiagaraStackObject::FOnSelectRootNodes::CreateUObject(this,
-				&UNiagaraStackEmitterSummaryItem::SelectSummaryNodesFromEmitterEditorDataRootNodes));
-		}
-		NewChildren.Add(SummaryEditorData);
-	}
-	else
-	{
-		SummaryEditorData = nullptr;
-	}
-
 	if (FilteredObject == nullptr)
 	{
-		FilteredObject = NewObject<UNiagaraStackSummaryViewObject>(this);
+		FilteredObject = NewObject<UNiagaraStackSummaryViewCollection>(this);
 		FRequiredEntryData RequiredEntryData(GetSystemViewModel(), GetEmitterViewModel(), FExecutionCategoryNames::Emitter, NAME_None, GetStackEditorData());
 		FilteredObject->Initialize(RequiredEntryData, Emitter, GetStackEditorDataKey());
 	}
@@ -387,21 +372,32 @@ void UNiagaraStackEmitterSummaryItem::SelectSummaryNodesFromEmitterEditorDataRoo
 	}
 }
 
-bool UNiagaraStackEmitterSummaryItem::GetEditModeIsActive() const
-{ 
-	return GetEmitterViewModel().IsValid() && GetEmitterViewModel()->GetSummaryIsInEditMode();
+void UNiagaraStackEmitterSummaryItem::OnEditButtonClicked()
+{
+	if (GetEmitterViewModel().IsValid())
+	{
+		GetSystemViewModel()->FocusTab(FNiagaraSystemToolkitModeBase::EmitterSummaryViewEditorTabID, true);
+	}
 }
 
-void UNiagaraStackEmitterSummaryItem::SetEditModeIsActive(bool bInEditModeIsActive)
+TOptional<FText> UNiagaraStackEmitterSummaryItem::GetEditModeButtonText() const
 {
-	if (GetEmitterViewModel().IsValid())
+	return LOCTEXT("EditSummaryViewButtonLabel", "Edit Summary");
+}
+
+TOptional<FText> UNiagaraStackEmitterSummaryItem::GetEditModeButtonTooltip() const
+{
+	return LOCTEXT("EditSummaryViewButtonTooltip", "Summons the Summary Editor that lets you define what inputs, renderers and properties (and more) should be displayed when looking at the Emitter Summary.");
+}
+
+EVisibility UNiagaraStackEmitterSummaryItem::IsEditButtonVisible() const
+{
+	if(GetEmitterViewModel().IsValid())
 	{
-		if (GetEmitterViewModel()->GetSummaryIsInEditMode() != bInEditModeIsActive)
-		{
-			GetEmitterViewModel()->SetSummaryIsInEditMode(bInEditModeIsActive);
-			RefreshChildren();
-		}
+		return GetEmitterViewModel()->GetEditorData().ShouldShowSummaryView() ? EVisibility::Collapsed : EVisibility::Visible;
 	}
+
+	return EVisibility::Collapsed;
 }
 
 UNiagaraStackEmitterSummaryGroup::UNiagaraStackEmitterSummaryGroup()
@@ -426,9 +422,6 @@ FText UNiagaraStackEmitterSummaryGroup::GetIconText() const
 	return FText::FromString(FString(TEXT("\xf0ca")/* fa-list-ul */));
 }
 
-
-
-
 void UNiagaraStackSummaryViewCollapseButton::Initialize(FRequiredEntryData InRequiredEntryData)
 {
 	Super::Initialize(InRequiredEntryData, TEXT("ShowAdvanced"));
diff --git a/Engine/Plugins/FX/Niagara/Source/NiagaraEditor/Private/ViewModels/Stack/NiagaraStackEntry.cpp b/Engine/Plugins/FX/Niagara/Source/NiagaraEditor/Private/ViewModels/Stack/NiagaraStackEntry.cpp
--- a/Engine/Plugins/FX/Niagara/Source/NiagaraEditor/Private/ViewModels/Stack/NiagaraStackEntry.cpp
+++ b/Engine/Plugins/FX/Niagara/Source/NiagaraEditor/Private/ViewModels/Stack/NiagaraStackEntry.cpp
@@ -11,6 +11,8 @@
 #include "NiagaraScriptMergeManager.h"
 #include "Misc/SecureHash.h"
 #include "ScopedTransaction.h"
+#include "ViewModels/HierarchyEditor/NiagaraHierarchyViewModelBase.h"
+#include "ViewModels/HierarchyEditor/NiagaraSummaryViewViewModel.h"
 
 #include UE_INLINE_GENERATED_CPP_BY_NAME(NiagaraStackEntry)
 
@@ -1090,6 +1092,91 @@ void UNiagaraStackEntry::OnRenamed(FText NewName)
 	}
 }
 
+FNiagaraHierarchyIdentity UNiagaraStackEntry::DetermineSummaryIdentity() const
+{
+	return FNiagaraHierarchyIdentity();
+}
+
+bool UNiagaraStackEntry::IsInSummaryView() const
+{
+	if(GetEmitterViewModel().IsValid())
+	{
+		if(SupportsSummaryView())
+		{
+			FNiagaraHierarchyIdentity Identity = DetermineSummaryIdentity();
+			if(Identity.IsValid())
+			{
+				return GetEmitterViewModel()->GetSummaryHierarchyViewModel()->GetHierarchyRootViewModel()->FindViewModelForChild(Identity, true) != nullptr;
+			}
+		}
+	}
+
+	return false;
+}
+
+bool UNiagaraStackEntry::IsAnyParentInSummaryView() const
+{
+	if(GetEmitterViewModel().IsValid())
+	{
+		for(UNiagaraStackEntry* Parent = Cast<UNiagaraStackEntry>(GetOuter()); Parent != nullptr; Parent = Cast<UNiagaraStackEntry>(Parent->GetOuter()))
+		{
+			if(Parent->IsInSummaryView())
+			{
+				return true;
+			}
+		}
+	}
+
+	return false;
+}
+
+bool UNiagaraStackEntry::IsAnyChildInSummaryView(bool bRecursive) const
+{
+	if(GetEmitterViewModel().IsValid())
+	{
+		TArray<UNiagaraStackEntry*> AllChildren;
+		GetUnfilteredChildren(AllChildren);
+		for(UNiagaraStackEntry* Child : AllChildren)
+		{
+			if(Child->IsInSummaryView())
+			{
+				return true;
+			}
+		}
+
+		if(bRecursive)
+		{
+			for(UNiagaraStackEntry* Child : AllChildren)
+			{
+				if(Child->IsAnyChildInSummaryView(bRecursive))
+				{
+					return true;
+				}
+			}			
+		}		
+	}
+
+	return false;
+}
+
+bool UNiagaraStackEntry::ExistsInParentEmitterSummary() const
+{
+	if(GetEmitterViewModel().IsValid())
+	{
+		if(GetEmitterViewModel()->GetParentEmitter().Emitter != nullptr)
+		{
+			FVersionedNiagaraEmitter BaseEmitter = GetEmitterViewModel()->GetParentEmitter();
+			FNiagaraHierarchyIdentity Identity = DetermineSummaryIdentity();
+			if(BaseEmitter.GetEmitterData() != nullptr && Identity.IsValid())
+			{
+				return Cast<UNiagaraEmitterEditorData>(BaseEmitter.GetEmitterData()->GetEditorData())->GetSummaryRoot()->FindChildWithIdentity(Identity, true) != nullptr;
+			}
+		}
+	}
+
+	return false;
+}
+
 void UNiagaraStackSpacer::Initialize(FRequiredEntryData InRequiredEntryData, float InSpacerHeight, TAttribute<bool> InShouldShowInStack, FString InOwningStackItemEditorDataKey)
 {
 	Super::Initialize(InRequiredEntryData, InOwningStackItemEditorDataKey + TEXT("Spacer"));
diff --git a/Engine/Plugins/FX/Niagara/Source/NiagaraEditor/Private/ViewModels/Stack/NiagaraStackEventScriptItemGroup.cpp b/Engine/Plugins/FX/Niagara/Source/NiagaraEditor/Private/ViewModels/Stack/NiagaraStackEventScriptItemGroup.cpp
--- a/Engine/Plugins/FX/Niagara/Source/NiagaraEditor/Private/ViewModels/Stack/NiagaraStackEventScriptItemGroup.cpp
+++ b/Engine/Plugins/FX/Niagara/Source/NiagaraEditor/Private/ViewModels/Stack/NiagaraStackEventScriptItemGroup.cpp
@@ -162,6 +162,15 @@ void UNiagaraStackEventHandlerPropertiesItem::RefreshChildrenInternal(const TArr
 	Super::RefreshChildrenInternal(CurrentChildren, NewChildren, NewIssues);
 }
 
+FNiagaraHierarchyIdentity UNiagaraStackEventHandlerPropertiesItem::DetermineSummaryIdentity() const
+{
+	FNiagaraHierarchyIdentity Identity;
+	Identity.Guids.Add(GetEventScriptUsageId());
+	Identity.Names.Add("Category");
+	Identity.Names.Add("Properties");
+	return Identity;
+}
+
 void UNiagaraStackEventHandlerPropertiesItem::EventHandlerPropertiesChanged()
 {
 	bCanResetToBaseCache.Reset();
@@ -352,6 +361,13 @@ FText UNiagaraStackEventScriptItemGroup::GetInheritanceMessage() const
 	return LOCTEXT("EventGroupInheritanceMessage", "This event handler is inherited from a parent emitter.  Inherited\nevent handlers can only be deleted while editing the parent emitter.");
 }
 
+FNiagaraHierarchyIdentity UNiagaraStackEventScriptItemGroup::DetermineSummaryIdentity() const
+{
+	FNiagaraHierarchyIdentity Identity;
+	Identity.Guids.Add(GetScriptUsageId());
+	return Identity;
+}
+
 bool UNiagaraStackEventScriptItemGroup::HasBaseEventHandler() const
 {
 	if (bHasBaseEventHandlerCache.IsSet() == false)
diff --git a/Engine/Plugins/FX/Niagara/Source/NiagaraEditor/Private/ViewModels/Stack/NiagaraStackFunctionInput.cpp b/Engine/Plugins/FX/Niagara/Source/NiagaraEditor/Private/ViewModels/Stack/NiagaraStackFunctionInput.cpp
--- a/Engine/Plugins/FX/Niagara/Source/NiagaraEditor/Private/ViewModels/Stack/NiagaraStackFunctionInput.cpp
+++ b/Engine/Plugins/FX/Niagara/Source/NiagaraEditor/Private/ViewModels/Stack/NiagaraStackFunctionInput.cpp
@@ -239,7 +239,7 @@ EUnit UNiagaraStackFunctionInput::GetInputDisplayUnit() const
 
 FText UNiagaraStackFunctionInput::GetTooltipText() const
 {
-	FText Description = InputMetaData.IsSet() ? InputMetaData->Description : FText::GetEmpty();
+	FText Description = SummaryViewTooltipOverride.IsSet() && !SummaryViewTooltipOverride->Get().IsEmptyOrWhitespace() ? SummaryViewTooltipOverride->Get() : InputMetaData.IsSet() ? InputMetaData->Description : FText::GetEmpty();
 	return FNiagaraEditorUtilities::FormatVariableDescription(Description, GetDisplayName(), InputType.GetNameText());
 }
 
@@ -380,6 +380,36 @@ bool UNiagaraStackFunctionInput::HasOverridenContent() const
 	return CanReset();
 }
 
+bool UNiagaraStackFunctionInput::SupportsSummaryView() const
+{
+	if(GetInputMetaData().IsSet())
+	{
+		if(GetInputMetaData()->bInlineEditConditionToggle || !GetInputMetaData()->ParentAttribute.IsNone())
+		{
+			return false;
+		}
+	}
+
+	return true;
+}
+
+FNiagaraHierarchyIdentity UNiagaraStackFunctionInput::DetermineSummaryIdentity() const
+{
+	FNiagaraHierarchyIdentity Identity;
+	Identity.Guids.Add(GetInputFunctionCallNode().NodeGuid);
+	if(Cast<UNiagaraNodeAssignment>(&GetInputFunctionCallNode()) != nullptr)
+	{
+		Identity.Names.Add(GetInputParameterHandle().GetParameterHandleString());
+	}
+	else
+	{
+		FNiagaraVariable ThisInput(GetInputType(), GetInputParameterHandle().GetParameterHandleString());
+		Identity.Guids.Add(GetInputFunctionCallNode().GetCalledGraph()->GetScriptVariableGuid(ThisInput).GetValue());
+	}
+
+	return Identity;
+}
+
 TArray<UNiagaraStackFunctionInput*> UNiagaraStackFunctionInput::GetChildInputs() const
 {
 	TArray<UNiagaraStackFunctionInputCollection*> DynamicInputCollections;
@@ -461,11 +491,16 @@ FText UNiagaraStackFunctionInput::GetCollapsedStateText() const
 	return CollapsedTextCache.GetValue();
 }
 
-void UNiagaraStackFunctionInput::SetSummaryViewDisiplayName(TOptional<FText> InDisplayName)
+void UNiagaraStackFunctionInput::SetSummaryViewDisplayName(TAttribute<FText> InDisplayName)
 {
 	SummaryViewDisplayNameOverride = InDisplayName;
 }
 
+void UNiagaraStackFunctionInput::SetSummaryViewTooltip(TAttribute<FText> InTooltipOverride)
+{
+	SummaryViewTooltipOverride = InTooltipOverride;
+}
+
 FText UNiagaraStackFunctionInput::GetValueToolTip() const
 {
 	if (IsFinalized())
@@ -1238,7 +1273,7 @@ void UNiagaraStackFunctionInput::RefreshFromMetaData(TArray<FStackIssue>& NewIss
 
 FText UNiagaraStackFunctionInput::GetDisplayName() const
 {
-	return SummaryViewDisplayNameOverride.IsSet()? SummaryViewDisplayNameOverride.GetValue() : DisplayNameOverride.IsSet() ? DisplayNameOverride.GetValue() : DisplayName;
+	return SummaryViewDisplayNameOverride.IsSet() && !SummaryViewDisplayNameOverride->Get().IsEmptyOrWhitespace() ? SummaryViewDisplayNameOverride->Get() : DisplayNameOverride.IsSet() ? DisplayNameOverride.GetValue() : DisplayName;
 }
 
 const TArray<FNiagaraParameterHandle>& UNiagaraStackFunctionInput::GetInputParameterHandlePath() const
@@ -2755,17 +2790,6 @@ void UNiagaraStackFunctionInput::GetSearchItems(TArray<FStackSearchItem>& Search
 	}
 }
 
-TOptional<FNiagaraVariableMetaData> UNiagaraStackFunctionInput::GetMetadata() const
-{
-	return InputMetaData;
-}
-
-TOptional<FGuid> UNiagaraStackFunctionInput::GetMetadataGuid() const
-{
-	return InputMetaData.IsSet() ? InputMetaData->GetVariableGuid() : TOptional<FGuid>();
-}
-
-
 const UNiagaraStackFunctionInput::FCollectedUsageData& UNiagaraStackFunctionInput::GetCollectedUsageData() const
 {
 	if (CachedCollectedUsageData.IsSet() == false)
