Summary View 2.0

diff --git a/Engine/Plugins/FX/Niagara/Source/NiagaraEditor/Private/ViewModels/Stack/NiagaraStackFunctionInputCollection.cpp b/Engine/Plugins/FX/Niagara/Source/NiagaraEditor/Private/ViewModels/Stack/NiagaraStackFunctionInputCollection.cpp
--- a/Engine/Plugins/FX/Niagara/Source/NiagaraEditor/Private/ViewModels/Stack/NiagaraStackFunctionInputCollection.cpp
+++ b/Engine/Plugins/FX/Niagara/Source/NiagaraEditor/Private/ViewModels/Stack/NiagaraStackFunctionInputCollection.cpp
@@ -20,14 +20,18 @@
 #include "ViewModels/Stack/NiagaraStackGraphUtilities.h"
 #include "ViewModels/Stack/NiagaraStackInputCategory.h"
 #include "NiagaraEmitterEditorData.h"
+#include "ViewModels/NiagaraScriptGraphViewModel.h"
+#include "NiagaraStackEditorData.h"
 
 #include UE_INLINE_GENERATED_CPP_BY_NAME(NiagaraStackFunctionInputCollection)
 
 #define LOCTEXT_NAMESPACE "UNiagaraStackFunctionInputCollection"
 
-FText UNiagaraStackFunctionInputCollectionBase::UncategorizedName = LOCTEXT("Uncategorized", "Uncategorized");
+FText UNiagaraStackValueCollection::UncategorizedName = LOCTEXT("Uncategorized", "Uncategorized");
 
-FText UNiagaraStackFunctionInputCollectionBase::AllSectionName = LOCTEXT("All", "All");
+FText UNiagaraStackValueCollection::AllSectionName = LOCTEXT("All", "All");
+
+using namespace FNiagaraStackGraphUtilities;
 
 static FText GetUserFriendlyFunctionName(UNiagaraNodeFunctionCall* Node)
 {
@@ -39,195 +43,19 @@ static FText GetUserFriendlyFunctionName(UNiagaraNodeFunctionCall* Node)
 	return FText::FromString(Node->GetFunctionName());
 }
 
-static TOptional<FFunctionInputSummaryViewKey> GetSummaryViewInputKeyForFunctionInput(UNiagaraNodeFunctionCall* FunctionCall, const FNiagaraVariable& InputVariable, const TOptional<FNiagaraVariableMetaData>& InputMetaData)
-{
-	FGuid NodeGuid = FunctionCall->NodeGuid;
-	
-	if (FunctionCall->IsA<UNiagaraNodeAssignment>())
-	{
-		return FFunctionInputSummaryViewKey(NodeGuid, InputVariable.GetName());
-	}
-	
-	if (InputMetaData.IsSet())
-	{
-		return FFunctionInputSummaryViewKey(NodeGuid, InputMetaData->GetVariableGuid());
-	}
-	
-	return TOptional<FFunctionInputSummaryViewKey>();
-}
-
-static bool ShouldShowInSummaryView(const FVersionedNiagaraEmitterData* EmitterData, UNiagaraNodeFunctionCall* FunctionCall, const FNiagaraVariable& InputVariable, const TOptional<FNiagaraVariableMetaData>& InputMetaData, bool bIsSummaryInEditMode)
-{
-	const UNiagaraEmitterEditorData* EditorData = EmitterData? Cast<UNiagaraEmitterEditorData>(EmitterData->GetEditorData()) : nullptr;
-	if (EditorData)
-	{
-		TOptional<FFunctionInputSummaryViewKey> SummaryViewKey = GetSummaryViewInputKeyForFunctionInput(FunctionCall, InputVariable, InputMetaData);
-		if (SummaryViewKey.IsSet())
-		{
-			TOptional<FFunctionInputSummaryViewMetadata> SummaryViewMetadata = EditorData->GetSummaryViewMetaData(SummaryViewKey.GetValue());
-			return SummaryViewMetadata.IsSet() && (bIsSummaryInEditMode || SummaryViewMetadata->bVisible);
-		}
-	}
-	return false;
-}
-
-UNiagaraStackFunctionInputCollection::UNiagaraStackFunctionInputCollection()
-	: ModuleNode(nullptr)
-	, InputFunctionCallNode(nullptr)
-{
-}
-
-UNiagaraNodeFunctionCall* UNiagaraStackFunctionInputCollection::GetModuleNode() const
-{
-	return ModuleNode;
-}
-
-UNiagaraNodeFunctionCall* UNiagaraStackFunctionInputCollection::GetInputFunctionCallNode() const
-{
-	return InputFunctionCallNode;
-}
-
-void UNiagaraStackFunctionInputCollection::Initialize(
-	FRequiredEntryData InRequiredEntryData,
-	UNiagaraNodeFunctionCall& InModuleNode,
-	UNiagaraNodeFunctionCall& InInputFunctionCallNode,
-	FString InOwnerStackItemEditorDataKey)
-{
-	checkf(ModuleNode == nullptr && InputFunctionCallNode == nullptr, TEXT("Can not set the node more than once."));
-	FString InputCollectionStackEditorDataKey = FString::Printf(TEXT("%s-Inputs"), *InInputFunctionCallNode.NodeGuid.ToString(EGuidFormats::DigitsWithHyphens));
-	Super::Initialize(InRequiredEntryData, InOwnerStackItemEditorDataKey, InputCollectionStackEditorDataKey);
-	ModuleNode = &InModuleNode;
-	InputFunctionCallNode = &InInputFunctionCallNode;
-	InputFunctionCallNode->OnInputsChanged().AddUObject(this, &UNiagaraStackFunctionInputCollection::OnFunctionInputsChanged);
-
-	FNiagaraEditorModule::Get().OnScriptApplied().AddUObject(this, &UNiagaraStackFunctionInputCollection::OnScriptApplied);
-}
-
-void UNiagaraStackFunctionInputCollection::FinalizeInternal()
-{
-	InputFunctionCallNode->OnInputsChanged().RemoveAll(this);
-	
-	FNiagaraEditorModule::Get().OnScriptApplied().RemoveAll(this);
-	
-	Super::FinalizeInternal();
-}
-
-FText UNiagaraStackFunctionInputCollection::GetDisplayName() const
-{
-	return LOCTEXT("InputCollectionDisplayName", "Inputs");
-}
-
-bool UNiagaraStackFunctionInputCollection::GetIsEnabled() const
-{
-	return InputFunctionCallNode->GetDesiredEnabledState() == ENodeEnabledState::Enabled;
-}
-
-void UNiagaraStackFunctionInputCollection::ToClipboardFunctionInputs(UObject* InOuter, TArray<const UNiagaraClipboardFunctionInput*>& OutClipboardFunctionInputs) const
-{
-	TArray<UNiagaraStackInputCategory*> ChildCategories;
-	GetUnfilteredChildrenOfType(ChildCategories);
-	for (UNiagaraStackInputCategory* ChildCategory : ChildCategories)
-	{
-		ChildCategory->ToClipboardFunctionInputs(InOuter, OutClipboardFunctionInputs);
-	}
-}
-
-void UNiagaraStackFunctionInputCollection::SetValuesFromClipboardFunctionInputs(const TArray<const UNiagaraClipboardFunctionInput*>& ClipboardFunctionInputs)
-{
-	TArray<const UNiagaraClipboardFunctionInput*> StaticSwitchInputs;
-	TArray<const UNiagaraClipboardFunctionInput*> StandardInputs;
-
-	TArray<UNiagaraStackInputCategory*> ChildCategories;
-	GetUnfilteredChildrenOfType(ChildCategories);
-
-	// Set static switches first so that other inputs will be available to set.
-	for (UNiagaraStackInputCategory* ChildCategory : ChildCategories)
-	{
-		ChildCategory->SetStaticSwitchValuesFromClipboardFunctionInputs(ClipboardFunctionInputs, *this);
-	}
-
-	RefreshChildren();
-	for (UNiagaraStackInputCategory* ChildCategory : ChildCategories)
-	{
-		ChildCategory->SetStandardValuesFromClipboardFunctionInputs(ClipboardFunctionInputs);
-	}
-}
-
-void UNiagaraStackFunctionInputCollection::GetChildInputs(TArray<UNiagaraStackFunctionInput*>& OutResult) const
-{
-	TArray<UNiagaraStackInputCategory*> ChildCategories;
-	GetUnfilteredChildrenOfType(ChildCategories);
-	for (UNiagaraStackInputCategory* ChildCategory : ChildCategories)
-	{
-		ChildCategory->GetUnfilteredChildrenOfType(OutResult);
-	}
-}
-
-void UNiagaraStackFunctionInputCollection::GetCustomFilteredChildInputs(TArray<UNiagaraStackFunctionInput*>& OutResult, const TArray<FOnFilterChild>& CustomFilters) const
-{
-	TArray<UNiagaraStackInputCategory*> ChildCategories;
-	GetUnfilteredChildrenOfType(ChildCategories);
-	for (UNiagaraStackInputCategory* ChildCategory : ChildCategories)
-	{
-		ChildCategory->GetCustomFilteredChildrenOfType(OutResult, CustomFilters);
-	}
-}
-
-TArray<UNiagaraStackFunctionInput*> UNiagaraStackFunctionInputCollection::GetInlineParameterInputs() const
-{
-	TArray<UNiagaraStackFunctionInput*> OutArray;
-	
-	TArray<UNiagaraStackFunctionInput*> FunctionInputs;
-	TArray<FOnFilterChild> CustomChildFilters;
-	CustomChildFilters.Add(FOnFilterChild::CreateUObject(this, &UNiagaraStackFunctionInput::FilterHiddenChildren));
-	GetCustomFilteredChildInputs(FunctionInputs, CustomChildFilters);
-
-	for(UNiagaraStackFunctionInput* FunctionInput : FunctionInputs)
-	{
-		if(FunctionInput->ShouldDisplayInline())
-		{
-			OutArray.Add(FunctionInput);
-		}
-	}
-
-	return OutArray;
-}
-
-void UNiagaraStackFunctionInputCollection::OnScriptApplied(UNiagaraScript* NiagaraScript, FGuid Guid)
-{
-	if(InputFunctionCallNode->FunctionScript == NiagaraScript)
-	{
-		RefreshChildren();
-	}
-}
-
-void UNiagaraStackFunctionInputCollection::RefreshChildrenInternal(const TArray<UNiagaraStackEntry*>& CurrentChildren, TArray<UNiagaraStackEntry*>& NewChildren, TArray<FStackIssue>& NewIssues)
-{
-	RefreshChildrenForFunctionCall(ModuleNode, InputFunctionCallNode, CurrentChildren, NewChildren, NewIssues, false);
-	Super::RefreshChildrenInternal(CurrentChildren, NewChildren, NewIssues);
-}
-
-void UNiagaraStackFunctionInputCollection::GetSectionsInternal(TArray<FNiagaraStackSection>& OutStackSections) const
-{
-	if (InputFunctionCallNode->GetCalledUsage() == ENiagaraScriptUsage::Module)
-	{
-		OutStackSections.Append(InputFunctionCallNode->GetScriptData()->InputSections);
-	}
-}
-
-void UNiagaraStackFunctionInputCollectionBase::Initialize(FRequiredEntryData InRequiredEntryData, FString InOwningStackItemEditorDataKey, FString InStackEditorDataKey)
+void UNiagaraStackValueCollection::Initialize(FRequiredEntryData InRequiredEntryData, FString InOwningStackItemEditorDataKey, FString InStackEditorDataKey)
 {
 	Super::Initialize(InRequiredEntryData, InOwningStackItemEditorDataKey, InStackEditorDataKey);
-	AddChildFilter(FOnFilterChild::CreateUObject(this, &UNiagaraStackFunctionInputCollectionBase::FilterByActiveSection));
+	AddChildFilter(FOnFilterChild::CreateUObject(this, &UNiagaraStackFunctionInputCollection::FilterByActiveSection));
 	ActiveSectionCache = GetStackEditorData().GetStackEntryActiveSection(GetStackEditorDataKey(), AllSectionName);
 }
 
-void UNiagaraStackFunctionInputCollectionBase::SetShouldDisplayLabel(bool bInShouldDisplayLabel)
+void UNiagaraStackValueCollection::SetShouldDisplayLabel(bool bInShouldDisplayLabel)
 {
 	bShouldDisplayLabel = bInShouldDisplayLabel;
 }
 
-const TArray<FText>& UNiagaraStackFunctionInputCollectionBase::GetSections() const
+const TArray<FText>& UNiagaraStackValueCollection::GetSections() const
 {
 	if (SectionsCache.IsSet() == false)
 	{
@@ -236,7 +64,7 @@ const TArray<FText>& UNiagaraStackFunctionInputCollectionBase::GetSections() con
 	return SectionsCache.GetValue();
 }
 
-FText UNiagaraStackFunctionInputCollectionBase::GetActiveSection() const
+FText UNiagaraStackValueCollection::GetActiveSection() const
 {
 	if (ActiveSectionCache.IsSet() == false)
 	{
@@ -245,24 +73,34 @@ FText UNiagaraStackFunctionInputCollectionBase::GetActiveSection() const
 	return ActiveSectionCache.GetValue();
 }
 
-void UNiagaraStackFunctionInputCollectionBase::SetActiveSection(FText InActiveSection)
+void UNiagaraStackValueCollection::SetActiveSection(FText InActiveSection)
 {
 	ActiveSectionCache = InActiveSection;
 	GetStackEditorData().SetStackEntryActiveSection(GetStackEditorDataKey(), InActiveSection);
 	RefreshFilteredChildren();
 }
 
-bool UNiagaraStackFunctionInputCollectionBase::GetCanExpand() const
+FText UNiagaraStackValueCollection::GetTooltipForSection(FString Section) const
+{
+	if(SectionToTooltipMapCache.IsSet() && SectionToTooltipMapCache->Contains(Section))
+	{
+		return SectionToTooltipMapCache.GetValue()[Section];
+	}
+
+	return FText::GetEmpty();
+}
+
+bool UNiagaraStackValueCollection::GetCanExpand() const
 {
 	return bShouldDisplayLabel;
 }
 
-bool UNiagaraStackFunctionInputCollectionBase::GetShouldShowInStack() const
+bool UNiagaraStackValueCollection::GetShouldShowInStack() const
 {
 	return GetSections().Num() > 0 || bShouldDisplayLabel;
 }
 
-void UNiagaraStackFunctionInputCollectionBase::RefreshChildrenInternal(const TArray<UNiagaraStackEntry*>& CurrentChildren, TArray<UNiagaraStackEntry*>& NewChildren, TArray<FStackIssue>& NewIssues)
+void UNiagaraStackValueCollection::RefreshChildrenInternal(const TArray<UNiagaraStackEntry*>& CurrentChildren, TArray<UNiagaraStackEntry*>& NewChildren, TArray<FStackIssue>& NewIssues)
 {
 	Super::RefreshChildrenInternal(CurrentChildren, NewChildren, NewIssues);
 	LastActiveSection = ActiveSectionCache.IsSet() ? ActiveSectionCache.GetValue() : AllSectionName;
@@ -271,24 +109,22 @@ void UNiagaraStackFunctionInputCollectionBase::RefreshChildrenInternal(const TAr
 	ActiveSectionCache.Reset();
 }
 
-int32 UNiagaraStackFunctionInputCollectionBase::GetChildIndentLevel() const
+int32 UNiagaraStackValueCollection::GetChildIndentLevel() const
 {
 	// We don't want the child categories to be indented.
 	return GetIndentLevel();
 }
 
-void UNiagaraStackFunctionInputCollectionBase::RefreshChildrenForFunctionCall(UNiagaraNodeFunctionCall* ModuleNode, UNiagaraNodeFunctionCall* InputFunctionCallNode, 
-	const TArray<UNiagaraStackEntry*>& CurrentChildren, TArray<UNiagaraStackEntry*>& NewChildren, TArray<FStackIssue>& NewIssues, bool bShouldApplySummaryFilter)
+void UNiagaraStackFunctionInputCollection::RefreshChildrenForFunctionCall(const TArray<UNiagaraStackEntry*>& CurrentChildren, TArray<UNiagaraStackEntry*>& NewChildren, TArray<FStackIssue>& NewIssues, bool bShouldApplySummaryFilter)
 {	
 	FFunctionCallNodesState State;
-	AppendInputsForFunctionCall(State, ModuleNode, InputFunctionCallNode, NewIssues, bShouldApplySummaryFilter);
+	AppendInputsForFunctionCall(State, NewIssues, bShouldApplySummaryFilter);
 	ApplyAllFunctionInputsToChildren(State, CurrentChildren, NewChildren, NewIssues, bShouldApplySummaryFilter);
 }
 
-void UNiagaraStackFunctionInputCollectionBase::AppendInputsForFunctionCall(FFunctionCallNodesState& State, UNiagaraNodeFunctionCall* ModuleNode, UNiagaraNodeFunctionCall* InputFunctionCallNode, TArray<FStackIssue>& NewIssues, bool bShouldApplySummaryFilter)
+void UNiagaraStackFunctionInputCollection::AppendInputsForFunctionCall(FFunctionCallNodesState& State, TArray<FStackIssue>& NewIssues, bool bShouldApplySummaryFilter)
 {
 	FVersionedNiagaraEmitter Emitter = GetEmitterViewModel().IsValid() ? GetEmitterViewModel()->GetEmitter() : FVersionedNiagaraEmitter();
-	FVersionedNiagaraEmitterData* EmitterData = Emitter.GetEmitterData();
 
 	TSet<FNiagaraVariable> HiddenVariables;
 	TArray<FNiagaraVariable> InputVariables;
@@ -314,44 +150,6 @@ void UNiagaraStackFunctionInputCollectionBase::AppendInputsForFunctionCall(FFunc
 
 	UNiagaraGraph* InputFunctionGraph = InputFunctionCallNode->GetCalledGraph();
 	
-	
-	const auto GetSummaryViewState = [&](const UNiagaraNodeFunctionCall* FunctionCall, const FNiagaraVariable& InputVariable,
-		const TOptional<FNiagaraVariableMetaData>& InputMetaData, FText& InputCategory, TOptional<FText>& DisplayName, int32& EditorSortPriority ) -> bool
-	{
-		const UNiagaraEmitterEditorData* EditorData = EmitterData ? Cast<UNiagaraEmitterEditorData>(EmitterData->GetEditorData()) : nullptr;
-
-		const TOptional<FFunctionInputSummaryViewKey> SummaryViewKey = GetSummaryViewInputKeyForFunctionInput(InputFunctionCallNode, InputVariable, InputMetaData);
-		if (EditorData && bShouldApplySummaryFilter && SummaryViewKey.IsSet())
-		{
-			TOptional<FFunctionInputSummaryViewMetadata> SummaryViewData = EditorData->GetSummaryViewMetaData(SummaryViewKey.GetValue());
-			if (SummaryViewData.IsSet() && SummaryViewData->Category != NAME_None)
-			{
-				InputCategory = FText::FromName(SummaryViewData->Category);
-			}
-			else
-			{
-				if (!InputCategory.EqualTo(UncategorizedName))
-				{
-					InputCategory = FText::Format(LOCTEXT("EmitterSummaryNodeCategory", "{0} - {1}"), FText::FromString(*InputFunctionCallNode->GetFunctionName()), InputCategory);
-				}
-				else if (FunctionCall->IsA<UNiagaraNodeAssignment>())
-				{
-					InputCategory = LOCTEXT("EmitterSummaryNodeCategoryAssignment", "Set Parameters");									
-				}
-				else
-				{					
-					InputCategory = FText::FromString(*InputFunctionCallNode->GetFunctionName());
-				}				
-			}
-			EditorSortPriority = SummaryViewData.IsSet() && SummaryViewData->SortIndex != INDEX_NONE ? SummaryViewData->SortIndex : EditorSortPriority;
-			DisplayName = SummaryViewData.IsSet() && SummaryViewData->DisplayName != NAME_None ? FText::FromName(SummaryViewData->DisplayName) : TOptional<FText>();		
-		}
-		
-		return ShouldShowInSummaryView(EmitterData, InputFunctionCallNode, InputVariable, InputMetaData, GetEmitterViewModel().IsValid() && GetEmitterViewModel()->GetSummaryIsInEditMode());
-	};
-
-
-	
 	// Gather input data
 	for (const FNiagaraVariable& InputVariable : InputVariables)
 	{
@@ -382,10 +180,9 @@ void UNiagaraStackFunctionInputCollectionBase::AppendInputsForFunctionCall(FFunc
 
 		int32 EditorSortPriority = InputMetaData.IsSet() ? InputMetaData->EditorSortPriority : 0;
 		TOptional<FText> DisplayName;
-
-		bool bShouldShowInSummary = GetSummaryViewState(InputFunctionCallNode, InputVariable, InputMetaData, InputCategory, DisplayName, EditorSortPriority);			
+		
 		bool bIsInputHidden = HiddenVariables.Contains(InputVariable);
-		FInputData InputData = { InputVariable, EditorSortPriority, DisplayName, InputCategory, false, bIsInputHidden, bShouldShowInSummary, ModuleNode, InputFunctionCallNode };
+		FInputData InputData = { InputVariable, EditorSortPriority, DisplayName, InputCategory, false, bIsInputHidden, ModuleNode, InputFunctionCallNode };
 		int32 Index = State.InputDataCollection.Add(InputData);
 
 		if (InputVariable.GetType().IsStatic())
@@ -450,10 +247,9 @@ void UNiagaraStackFunctionInputCollectionBase::AppendInputsForFunctionCall(FFunc
 
 		int32 EditorSortPriority = InputMetaData.IsSet() ? InputMetaData->EditorSortPriority : 0;
 		TOptional<FText> DisplayName;
-						
-		bool bShouldShowInSummary = GetSummaryViewState(InputFunctionCallNode, InputVariable, InputMetaData, InputCategory, DisplayName, EditorSortPriority);		
+		
 		bool bIsInputHidden = HiddenSwitchPins.Contains(InputPin);
-		FInputData InputData = { InputVariable, EditorSortPriority, DisplayName, InputCategory, true, bIsInputHidden, bShouldShowInSummary, ModuleNode, InputFunctionCallNode };
+		FInputData InputData = { InputVariable, EditorSortPriority, DisplayName, InputCategory, true, bIsInputHidden, ModuleNode, InputFunctionCallNode };
 		int32 Index = State.InputDataCollection.Add(InputData);
 
 		// set up the data for the parent-child mapping
@@ -476,10 +272,10 @@ void UNiagaraStackFunctionInputCollectionBase::AppendInputsForFunctionCall(FFunc
 		}
 	}
 
-	RefreshIssues(InputFunctionCallNode, DuplicateInputNames, ValidAliasedInputNames, InputsWithInvalidTypes, StaticSwitchInputs, NewIssues);
+	RefreshIssues(DuplicateInputNames, ValidAliasedInputNames, InputsWithInvalidTypes, StaticSwitchInputs, NewIssues);
 }
 
-void UNiagaraStackFunctionInputCollectionBase::ApplyAllFunctionInputsToChildren(FFunctionCallNodesState& State, const TArray<UNiagaraStackEntry*>& CurrentChildren, TArray<UNiagaraStackEntry*>& NewChildren, TArray<FStackIssue>& NewIssues, bool bShouldApplySummaryFilter)
+void UNiagaraStackFunctionInputCollection::ApplyAllFunctionInputsToChildren(FFunctionCallNodesState& State, const TArray<UNiagaraStackEntry*>& CurrentChildren, TArray<UNiagaraStackEntry*>& NewChildren, TArray<FStackIssue>& NewIssues, bool bShouldApplySummaryFilter)
 {
 	// resolve the parent/child relationships
 	for (auto& Entry : State.ParentMapping)
@@ -508,31 +304,6 @@ void UNiagaraStackFunctionInputCollectionBase::ApplyAllFunctionInputsToChildren(
 		}
 	}
 	
-	if (bShouldApplySummaryFilter)
-	{
-		static TFunction<void(FInputData*)> PropagateSummaryViewShowToChildren = [](FInputData* Input)
-		{		
-			const bool bShouldShowChildren = !Input->bIsHidden && Input->bShouldShowInSummary;
-
-			Input->bIsHidden = !bShouldShowChildren;
-
-			for (FInputData* Child : Input->Children)
-			{
-				Child->bShouldShowInSummary = bShouldShowChildren;
-				PropagateSummaryViewShowToChildren(Child);
-			}		
-		};
-
-		// Propagate summary visibility to children
-		for (FInputData& InputData : State.InputDataCollection)
-		{
-			if (!InputData.bIsChild)
-			{
-				PropagateSummaryViewShowToChildren(&InputData);
-			}
-		}
-	}
-	
 	auto SortPredicate = [](const FInputData& A, const FInputData& B)
 	{
 		// keep the uncategorized attributes first
@@ -580,10 +351,7 @@ void UNiagaraStackFunctionInputCollectionBase::ApplyAllFunctionInputsToChildren(
 	}
 }
 
-
-
-void UNiagaraStackFunctionInputCollectionBase::RefreshIssues(UNiagaraNodeFunctionCall* InputFunctionCallNode, const TArray<FName>& DuplicateInputNames, 
-	const TArray<FName>& ValidAliasedInputNames, const TArray<FNiagaraVariable>& InputsWithInvalidTypes, const TMap<FName, UEdGraphPin*>& StaticSwitchInputs, TArray<FStackIssue>& NewIssues)
+void UNiagaraStackFunctionInputCollection::RefreshIssues(const TArray<FName>& DuplicateInputNames, const TArray<FName>& ValidAliasedInputNames, const TArray<FNiagaraVariable>& InputsWithInvalidTypes, const TMap<FName, UEdGraphPin*>& StaticSwitchInputs, TArray<FStackIssue>& NewIssues)
 {
 	if (!GetIsEnabled())
 	{
@@ -733,12 +501,12 @@ void UNiagaraStackFunctionInputCollectionBase::RefreshIssues(UNiagaraNodeFunctio
 	}
 }
 
-void UNiagaraStackFunctionInputCollectionBase::OnFunctionInputsChanged()
+void UNiagaraStackFunctionInputCollection::OnFunctionInputsChanged()
 {
 	RefreshChildren();
 }
 
-void UNiagaraStackFunctionInputCollectionBase::AddInvalidChildStackIssue(FName PinName, TArray<FStackIssue>& OutIssues)
+void UNiagaraStackFunctionInputCollection::AddInvalidChildStackIssue(FName PinName, TArray<FStackIssue>& OutIssues)
 {
 	FStackIssue InvalidHierarchyWarning(
 		EStackIssueSeverity::Warning,
@@ -748,17 +516,17 @@ void UNiagaraStackFunctionInputCollectionBase::AddInvalidChildStackIssue(FName P
 	OutIssues.Add(InvalidHierarchyWarning);
 }
 
-void UNiagaraStackFunctionInputCollectionBase::AddInputToCategory(const FInputData& InputData, const TArray<UNiagaraStackEntry*>& CurrentChildren, TArray<UNiagaraStackEntry*>& NewChildren)
+void UNiagaraStackFunctionInputCollection::AddInputToCategory(const FInputData& InputData, const TArray<UNiagaraStackEntry*>& CurrentChildren, TArray<UNiagaraStackEntry*>& NewChildren)
 {	
 	// Try to find an existing category in the already processed children.
 	UNiagaraStackInputCategory* InputCategory = FindCurrentChildOfTypeByPredicate<UNiagaraStackInputCategory>(NewChildren,
-		[&](UNiagaraStackInputCategory* CurrentCategory) { return CurrentCategory->GetCategoryName().CompareTo(InputData.Category) == 0; });
+		[&](UNiagaraStackInputCategory* CurrentCategory) { return CurrentCategory->GetDisplayName().CompareTo(InputData.Category) == 0; });
 
 	if (InputCategory == nullptr)
 	{
 		// If we haven't added any children to this category yet see if there is one that can be reused from the current children.
 		InputCategory = FindCurrentChildOfTypeByPredicate<UNiagaraStackInputCategory>(CurrentChildren,
-			[&](UNiagaraStackInputCategory* CurrentCategory) { return CurrentCategory->GetCategoryName().CompareTo(InputData.Category) == 0; });
+			[&](UNiagaraStackInputCategory* CurrentCategory) { return CurrentCategory->GetDisplayName().CompareTo(InputData.Category) == 0; });
 		if (InputCategory == nullptr)
 		{
 			// If we don't have a current child for this category make a new one.
@@ -783,7 +551,46 @@ void UNiagaraStackFunctionInputCollectionBase::AddInputToCategory(const FInputDa
 	InputCategory->AddInput(InputData.ModuleNode, InputData.InputFunctionCallNode, InputData.InputVariable.GetName(), InputData.InputVariable.GetType(), InputData.bIsStatic ? EStackParameterBehavior::Static : EStackParameterBehavior::Dynamic, InputData.DisplayName, InputData.bIsHidden, InputData.bIsChild);
 }
 
-bool UNiagaraStackFunctionInputCollectionBase::FilterByActiveSection(const UNiagaraStackEntry& Child) const
+UNiagaraStackEntry::FStackIssueFix UNiagaraStackFunctionInputCollection::GetNodeRemovalFix(UEdGraphPin* PinToRemove, FText FixDescription)
+{
+	return FStackIssueFix(
+		FixDescription,
+		UNiagaraStackEntry::FStackIssueFixDelegate::CreateLambda([=]()
+			{
+				FScopedTransaction ScopedTransaction(FixDescription);
+				TArray<TWeakObjectPtr<UNiagaraDataInterface>> RemovedDataObjects;
+				FNiagaraStackGraphUtilities::RemoveNodesForStackFunctionInputOverridePin(*PinToRemove, RemovedDataObjects);
+				TArray<UObject*> RemovedObjects;
+				for (TWeakObjectPtr<UNiagaraDataInterface> RemovedDataObject : RemovedDataObjects)
+				{
+					if (RemovedDataObject.IsValid())
+					{
+						RemovedObjects.Add(RemovedDataObject.Get());
+					}
+				}
+				OnDataObjectModified().Broadcast(RemovedObjects, ENiagaraDataObjectChange::Removed);
+				PinToRemove->GetOwningNode()->RemovePin(PinToRemove);
+			}));
+}
+
+UNiagaraStackEntry::FStackIssueFix UNiagaraStackFunctionInputCollection::GetResetPinFix(UEdGraphPin* PinToReset, FText FixDescription)
+{
+	return FStackIssueFix(
+		FixDescription,
+		UNiagaraStackEntry::FStackIssueFixDelegate::CreateLambda([=]()
+			{
+				FScopedTransaction ScopedTransaction(FixDescription);
+				const UEdGraphSchema_Niagara* NiagaraSchema = GetDefault<UEdGraphSchema_Niagara>();
+				UNiagaraNode* OwningNiagaraNode = Cast<UNiagaraNode>(PinToReset->GetOwningNode());
+				NiagaraSchema->ResetPinToAutogeneratedDefaultValue(PinToReset);
+				if (OwningNiagaraNode != nullptr)
+				{
+					OwningNiagaraNode->MarkNodeRequiresSynchronization("Pin reset to default value.", true);
+				}
+			}));
+}
+
+bool UNiagaraStackValueCollection::FilterByActiveSection(const UNiagaraStackEntry& Child) const
 {
 	const TArray<FText>& Sections = GetSections();
 	FText ActiveSection = GetActiveSection();
@@ -793,15 +600,16 @@ bool UNiagaraStackFunctionInputCollectionBase::FilterByActiveSection(const UNiag
 	}
 
 	TArray<FText>* ActiveCategoryNames = SectionToCategoryMapCache->Find(ActiveSection.ToString());
-	const UNiagaraStackInputCategory* ChildCategory = Cast<UNiagaraStackInputCategory>(&Child);
+	const UNiagaraStackCategory* ChildCategory = Cast<UNiagaraStackCategory>(&Child);
 	return ChildCategory == nullptr || ActiveCategoryNames == nullptr || ActiveCategoryNames->ContainsByPredicate(
 		[ChildCategory](const FText& ActiveCategoryName) { return ActiveCategoryName.EqualTo(ChildCategory->GetDisplayName()); });
 }
 
-void UNiagaraStackFunctionInputCollectionBase::UpdateCachedSectionData() const
+void UNiagaraStackValueCollection::UpdateCachedSectionData() const
 {
 	TArray<FText> Sections;
 	TMap<FString, TArray<FText>> SectionToCategoryMap;
+	TMap<FString, FText> SectionToTooltipMap;
 	TArray<FNiagaraStackSection> StackSections;
 	GetSectionsInternal(StackSections);
 
@@ -809,10 +617,10 @@ void UNiagaraStackFunctionInputCollectionBase::UpdateCachedSectionData() const
 	{
 		// Get the current list of categories.
 		TArray<FText> CategoryNames;
-		TArray<UNiagaraStackInputCategory*> ChildCategories;
+		TArray<UNiagaraStackCategory*> ChildCategories;
 		GetUnfilteredChildrenOfType(ChildCategories);
 		
-		for (UNiagaraStackInputCategory* ChildCategory : ChildCategories)
+		for (UNiagaraStackCategory* ChildCategory : ChildCategories)
 		{
 			if (ChildCategory->GetShouldShowInStack())
 			{
@@ -836,21 +644,25 @@ void UNiagaraStackFunctionInputCollectionBase::UpdateCachedSectionData() const
 				Sections.Add(StackSection.SectionDisplayName);
 				SectionToCategoryMap.Add(StackSection.SectionDisplayName.ToString(), ContainedCategories);
 			}
+
+			SectionToTooltipMap.Add(StackSection.SectionDisplayName.ToString(), StackSection.Tooltip);
 		}
 
 		Sections.Add(AllSectionName);
 		SectionToCategoryMap.Add(AllSectionName.ToString(), CategoryNames);
-
+		
 		if (Sections.Num() == 1)
 		{
 			// If there is only one section, it's the "All" section which is not useful.
 			Sections.Empty();
 			SectionToCategoryMap.Empty();
+			SectionToTooltipMap.Empty();
 		}
 	}
 
 	SectionsCache = Sections;
 	SectionToCategoryMapCache = SectionToCategoryMap;
+	SectionToTooltipMapCache = SectionToTooltipMap;
 	FText LastActiveSectionLocal = LastActiveSection;
 	if (Sections.ContainsByPredicate([LastActiveSectionLocal](FText Section) { return Section.EqualTo(LastActiveSectionLocal); }))
 	{
@@ -862,44 +674,148 @@ void UNiagaraStackFunctionInputCollectionBase::UpdateCachedSectionData() const
 	}
 }
 
-UNiagaraStackEntry::FStackIssueFix UNiagaraStackFunctionInputCollectionBase::GetNodeRemovalFix(UEdGraphPin* PinToRemove, FText FixDescription)
+UNiagaraStackFunctionInputCollection::UNiagaraStackFunctionInputCollection()
+	: ModuleNode(nullptr)
+	, InputFunctionCallNode(nullptr)
 {
-	return FStackIssueFix(
-		FixDescription,
-		UNiagaraStackEntry::FStackIssueFixDelegate::CreateLambda([=]()
-			{
-				FScopedTransaction ScopedTransaction(FixDescription);
-				TArray<TWeakObjectPtr<UNiagaraDataInterface>> RemovedDataObjects;
-				FNiagaraStackGraphUtilities::RemoveNodesForStackFunctionInputOverridePin(*PinToRemove, RemovedDataObjects);
-				TArray<UObject*> RemovedObjects;
-				for (TWeakObjectPtr<UNiagaraDataInterface> RemovedDataObject : RemovedDataObjects)
-				{
-					if (RemovedDataObject.IsValid())
-					{
-						RemovedObjects.Add(RemovedDataObject.Get());
-					}
-				}
-				OnDataObjectModified().Broadcast(RemovedObjects, ENiagaraDataObjectChange::Removed);
-				PinToRemove->GetOwningNode()->RemovePin(PinToRemove);
-			}));
 }
 
-UNiagaraStackEntry::FStackIssueFix UNiagaraStackFunctionInputCollectionBase::GetResetPinFix(UEdGraphPin* PinToReset, FText FixDescription)
+UNiagaraNodeFunctionCall* UNiagaraStackFunctionInputCollection::GetModuleNode() const
 {
-	return FStackIssueFix(
-		FixDescription,
-		UNiagaraStackEntry::FStackIssueFixDelegate::CreateLambda([=]()
-			{
-				FScopedTransaction ScopedTransaction(FixDescription);
-				const UEdGraphSchema_Niagara* NiagaraSchema = GetDefault<UEdGraphSchema_Niagara>();
-				UNiagaraNode* OwningNiagaraNode = Cast<UNiagaraNode>(PinToReset->GetOwningNode());
-				NiagaraSchema->ResetPinToAutogeneratedDefaultValue(PinToReset);
-				if (OwningNiagaraNode != nullptr)
-				{
-					OwningNiagaraNode->MarkNodeRequiresSynchronization("Pin reset to default value.", true);
-				}
-			}));
+	return ModuleNode;
 }
 
-#undef LOCTEXT_NAMESPACE
+UNiagaraNodeFunctionCall* UNiagaraStackFunctionInputCollection::GetInputFunctionCallNode() const
+{
+	return InputFunctionCallNode;
+}
+
+void UNiagaraStackFunctionInputCollection::Initialize(
+	FRequiredEntryData InRequiredEntryData,
+	UNiagaraNodeFunctionCall& InModuleNode,
+	UNiagaraNodeFunctionCall& InInputFunctionCallNode,
+	FString InOwnerStackItemEditorDataKey)
+{
+	checkf(ModuleNode == nullptr && InputFunctionCallNode == nullptr, TEXT("Can not set the node more than once."));
+	FString InputCollectionStackEditorDataKey = FString::Printf(TEXT("%s-Inputs"), *InInputFunctionCallNode.NodeGuid.ToString(EGuidFormats::DigitsWithHyphens));
+	Super::Initialize(InRequiredEntryData, InOwnerStackItemEditorDataKey, InputCollectionStackEditorDataKey);
+	ModuleNode = &InModuleNode;
+	InputFunctionCallNode = &InInputFunctionCallNode;
+	InputFunctionCallNode->OnInputsChanged().AddUObject(this, &UNiagaraStackFunctionInputCollection::OnFunctionInputsChanged);
+
+	FNiagaraEditorModule::Get().OnScriptApplied().AddUObject(this, &UNiagaraStackFunctionInputCollection::OnScriptApplied);
+}
+
+void UNiagaraStackFunctionInputCollection::FinalizeInternal()
+{
+	InputFunctionCallNode->OnInputsChanged().RemoveAll(this);
+	
+	FNiagaraEditorModule::Get().OnScriptApplied().RemoveAll(this);
+	
+	Super::FinalizeInternal();
+}
 
+FText UNiagaraStackFunctionInputCollection::GetDisplayName() const
+{
+	return LOCTEXT("InputCollectionDisplayName", "Inputs");
+}
+
+bool UNiagaraStackFunctionInputCollection::GetIsEnabled() const
+{
+	return InputFunctionCallNode->GetDesiredEnabledState() == ENodeEnabledState::Enabled;
+}
+
+void UNiagaraStackFunctionInputCollection::ToClipboardFunctionInputs(UObject* InOuter, TArray<const UNiagaraClipboardFunctionInput*>& OutClipboardFunctionInputs) const
+{
+	TArray<UNiagaraStackInputCategory*> ChildCategories;
+	GetUnfilteredChildrenOfType(ChildCategories);
+	for (UNiagaraStackInputCategory* ChildCategory : ChildCategories)
+	{
+		ChildCategory->ToClipboardFunctionInputs(InOuter, OutClipboardFunctionInputs);
+	}
+}
+
+void UNiagaraStackFunctionInputCollection::SetValuesFromClipboardFunctionInputs(const TArray<const UNiagaraClipboardFunctionInput*>& ClipboardFunctionInputs)
+{
+	TArray<const UNiagaraClipboardFunctionInput*> StaticSwitchInputs;
+	TArray<const UNiagaraClipboardFunctionInput*> StandardInputs;
+
+	TArray<UNiagaraStackInputCategory*> ChildCategories;
+	GetUnfilteredChildrenOfType(ChildCategories);
+
+	// Set static switches first so that other inputs will be available to set.
+	for (UNiagaraStackInputCategory* ChildCategory : ChildCategories)
+	{
+		ChildCategory->SetStaticSwitchValuesFromClipboardFunctionInputs(ClipboardFunctionInputs, *this);
+	}
+
+	RefreshChildren();
+	for (UNiagaraStackInputCategory* ChildCategory : ChildCategories)
+	{
+		ChildCategory->SetStandardValuesFromClipboardFunctionInputs(ClipboardFunctionInputs);
+	}
+}
+
+void UNiagaraStackFunctionInputCollection::GetChildInputs(TArray<UNiagaraStackFunctionInput*>& OutResult) const
+{
+	TArray<UNiagaraStackInputCategory*> ChildCategories;
+	GetUnfilteredChildrenOfType(ChildCategories);
+	for (UNiagaraStackInputCategory* ChildCategory : ChildCategories)
+	{
+		ChildCategory->GetUnfilteredChildrenOfType(OutResult);
+	}
+}
+
+void UNiagaraStackFunctionInputCollection::GetCustomFilteredChildInputs(TArray<UNiagaraStackFunctionInput*>& OutResult, const TArray<FOnFilterChild>& CustomFilters) const
+{
+	TArray<UNiagaraStackInputCategory*> ChildCategories;
+	GetUnfilteredChildrenOfType(ChildCategories);
+	for (UNiagaraStackInputCategory* ChildCategory : ChildCategories)
+	{
+		ChildCategory->GetCustomFilteredChildrenOfType(OutResult, CustomFilters);
+	}
+}
+
+TArray<UNiagaraStackFunctionInput*> UNiagaraStackFunctionInputCollection::GetInlineParameterInputs() const
+{
+	TArray<UNiagaraStackFunctionInput*> OutArray;
+	
+	TArray<UNiagaraStackFunctionInput*> FunctionInputs;
+	TArray<FOnFilterChild> CustomChildFilters;
+	CustomChildFilters.Add(FOnFilterChild::CreateUObject(this, &UNiagaraStackFunctionInput::FilterHiddenChildren));
+	GetCustomFilteredChildInputs(FunctionInputs, CustomChildFilters);
+
+	for(UNiagaraStackFunctionInput* FunctionInput : FunctionInputs)
+	{
+		if(FunctionInput->ShouldDisplayInline())
+		{
+			OutArray.Add(FunctionInput);
+		}
+	}
+
+	return OutArray;
+}
+
+void UNiagaraStackFunctionInputCollection::OnScriptApplied(UNiagaraScript* NiagaraScript, FGuid Guid)
+{
+	if(InputFunctionCallNode->FunctionScript == NiagaraScript)
+	{
+		RefreshChildren();
+	}
+}
+
+void UNiagaraStackFunctionInputCollection::RefreshChildrenInternal(const TArray<UNiagaraStackEntry*>& CurrentChildren, TArray<UNiagaraStackEntry*>& NewChildren, TArray<FStackIssue>& NewIssues)
+{
+	RefreshChildrenForFunctionCall(CurrentChildren, NewChildren, NewIssues, false);
+	Super::RefreshChildrenInternal(CurrentChildren, NewChildren, NewIssues);
+}
+
+void UNiagaraStackFunctionInputCollection::GetSectionsInternal(TArray<FNiagaraStackSection>& OutStackSections) const
+{
+	if (InputFunctionCallNode->GetCalledUsage() == ENiagaraScriptUsage::Module)
+	{
+		OutStackSections.Append(InputFunctionCallNode->GetScriptData()->InputSections);
+	}
+}
+
+#undef LOCTEXT_NAMESPACE
diff --git a/Engine/Plugins/FX/Niagara/Source/NiagaraEditor/Private/ViewModels/Stack/NiagaraStackGraphUtilities.cpp b/Engine/Plugins/FX/Niagara/Source/NiagaraEditor/Private/ViewModels/Stack/NiagaraStackGraphUtilities.cpp
--- a/Engine/Plugins/FX/Niagara/Source/NiagaraEditor/Private/ViewModels/Stack/NiagaraStackGraphUtilities.cpp
+++ b/Engine/Plugins/FX/Niagara/Source/NiagaraEditor/Private/ViewModels/Stack/NiagaraStackGraphUtilities.cpp
@@ -45,8 +45,13 @@
 #include "ViewModels/NiagaraScriptGraphViewModel.h"
 #include "NiagaraSettings.h"
 #include "ViewModels/HierarchyEditor/NiagaraHierarchyViewModelBase.h"
+#include "ViewModels/HierarchyEditor/NiagaraSummaryViewViewModel.h"
 
 DECLARE_CYCLE_STAT(TEXT("Niagara - StackGraphUtilities - RelayoutGraph"), STAT_NiagaraEditor_StackGraphUtilities_RelayoutGraph, STATGROUP_NiagaraEditor);
+DECLARE_CYCLE_STAT(TEXT("Niagara - StackGraphUtilities - GatherInputRelationsForStack"), STAT_NiagaraEditor_StackGraphUtilities_GatherInputRelationsForStack, STATGROUP_NiagaraEditor);
+DECLARE_CYCLE_STAT(TEXT("Niagara - StackGraphUtilities - FindInputData"), STAT_NiagaraEditor_StackGraphUtilities_FindInputData, STATGROUP_NiagaraEditor);
+DECLARE_CYCLE_STAT(TEXT("Niagara - StackGraphUtilities - FindModuleInputData"), STAT_NiagaraEditor_StackGraphUtilities_FindModuleInputData, STATGROUP_NiagaraEditor);
+DECLARE_CYCLE_STAT(TEXT("Niagara - StackGraphUtilities - FindAssignmentInputData"), STAT_NiagaraEditor_StackGraphUtilities_FindAssignmentInputData, STATGROUP_NiagaraEditor);
 
 #define LOCTEXT_NAMESPACE "NiagaraStackGraphUtilities"
 
@@ -315,6 +320,50 @@ TArray<UNiagaraNodeFunctionCall*> FNiagaraStackGraphUtilities::GetAllModuleNodes
 	return ModuleNodes;
 }
 
+TArray<UNiagaraNodeFunctionCall*> FNiagaraStackGraphUtilities::GetAllSimStagesModuleNodes(TSharedRef<FNiagaraEmitterViewModel> EmitterViewModel)
+{
+	TArray<UNiagaraNodeFunctionCall*> ModuleNodes;
+	TArray<UNiagaraNodeOutput*> OutputNodes = GetAllOutputNodes(EmitterViewModel);
+
+	for(UNiagaraNodeOutput* OutputNode : OutputNodes)
+	{
+		if(OutputNode->GetUsage() == ENiagaraScriptUsage::ParticleSimulationStageScript)
+		{
+			TArray<UNiagaraNodeFunctionCall*> FunctionCalls;
+			FNiagaraStackGraphUtilities::GetOrderedModuleNodes(*OutputNode, FunctionCalls);
+			
+			for(UNiagaraNodeFunctionCall* FunctionCall : FunctionCalls)
+			{
+				ModuleNodes.Add(FunctionCall);
+			}
+		}
+	}
+
+	return ModuleNodes;
+}
+
+TArray<UNiagaraNodeFunctionCall*> FNiagaraStackGraphUtilities::GetAllEventHandlerModuleNodes(TSharedRef<FNiagaraEmitterViewModel> EmitterViewModel)
+{
+	TArray<UNiagaraNodeFunctionCall*> ModuleNodes;
+	TArray<UNiagaraNodeOutput*> OutputNodes = GetAllOutputNodes(EmitterViewModel);
+
+	for(UNiagaraNodeOutput* OutputNode : OutputNodes)
+	{
+		if(OutputNode->GetUsage() == ENiagaraScriptUsage::ParticleEventScript)
+		{
+			TArray<UNiagaraNodeFunctionCall*> FunctionCalls;
+			FNiagaraStackGraphUtilities::GetOrderedModuleNodes(*OutputNode, FunctionCalls);
+			
+			for(UNiagaraNodeFunctionCall* FunctionCall : FunctionCalls)
+			{
+				ModuleNodes.Add(FunctionCall);
+			}
+		}
+	}
+
+	return ModuleNodes;
+}
+
 UNiagaraNodeFunctionCall* FNiagaraStackGraphUtilities::GetPreviousModuleNode(UNiagaraNodeFunctionCall& CurrentNode)
 {
 	UNiagaraNodeOutput* OutputNode = GetEmitterOutputNodeForStackNode(CurrentNode);
@@ -969,6 +1018,165 @@ void FNiagaraStackGraphUtilities::GetStackFunctionInputs(const UNiagaraNodeFunct
 	GetStackFunctionInputs(FunctionCallNode, OutInputVariables, EmptyResolver, Options, bIgnoreDisabled);
 }
 
+void FNiagaraStackGraphUtilities::GatherInputRelationsForStack(FInputDataCollection& State, TSharedRef<FNiagaraEmitterViewModel> EmitterViewModel)
+{
+	SCOPE_CYCLE_COUNTER(STAT_NiagaraEditor_StackGraphUtilities_GatherInputRelationsForStack);
+
+	FText UncategorizedName = UNiagaraStackValueCollection::UncategorizedName;
+	TSet<const UEdGraphPin*> HiddenPins;
+    TArray<const UEdGraphPin*> InputPins;
+	
+	UNiagaraHierarchyRoot* SummaryRoot = EmitterViewModel->GetEditorData().GetSummaryRoot();
+
+	TArray<UNiagaraNodeOutput*> OutputNodes = GetAllOutputNodes(EmitterViewModel);
+
+	TSharedRef<FNiagaraScriptGraphViewModel> GraphViewModel = EmitterViewModel->GetSharedScriptViewModel()->GetGraphViewModel();
+
+	TArray<UNiagaraNodeFunctionCall*> FunctionCallNodes;
+	GraphViewModel->GetGraph()->GetNodesOfClass<UNiagaraNodeFunctionCall>(FunctionCallNodes);
+	
+	// first we create a map so it's easy to access all relevant module & dynamic input nodes
+	TMap<FGuid, UNiagaraNodeFunctionCall*> NodeGuidToModuleNodeMap;
+	for(UNiagaraNodeFunctionCall* FunctionCallNode : FunctionCallNodes)
+	{
+		NodeGuidToModuleNodeMap.Add(FunctionCallNode->NodeGuid, FunctionCallNode);
+	}
+
+	// the inputs found in the root are the ones we want to display
+	TArray<UNiagaraHierarchyModuleInput*> ModuleInputs;
+	SummaryRoot->GetChildrenOfType(ModuleInputs, true);
+	
+	// now we create a map from summary input to the function node it (modules, dynamic inputs, assignment nodes etc.) maps to
+	TMap<UNiagaraHierarchyModuleInput*, UNiagaraNodeFunctionCall*> HierarchyInputToFunctionCallNodeMap;
+	TMap<UNiagaraHierarchyModuleInput*, bool> HierarchyInputToStaticSwitchMap;
+	for(UNiagaraHierarchyModuleInput* Input : ModuleInputs)
+	{
+		HierarchyInputToFunctionCallNodeMap.Add(Input, NodeGuidToModuleNodeMap[Input->GetPersistentIdentity().Guids[0]]); 
+	}
+
+	// then we create a map from summary input to its actual input script variable. This does not include children inputs yet as the root only contains parent level inputs.
+	TMap<UNiagaraHierarchyModuleInput*, TObjectPtr<UNiagaraScriptVariable>> HierarchyInputToScriptVariableMap;
+	TMap<UNiagaraHierarchyModuleInput*, FNiagaraVariable> HierarchyInputToAssignmentMap;
+	for(UNiagaraHierarchyModuleInput* Input : ModuleInputs)
+	{
+		UNiagaraGraph* GraphWithInput = HierarchyInputToFunctionCallNodeMap[Input]->GetCalledGraph();
+		ensure(GraphWithInput != nullptr);
+
+		bool bIsModuleInput = Input->GetPersistentIdentity().Guids.Num() == 2;
+
+		// the input guid is only valid for module inputs, not assignment targets
+		if(bIsModuleInput)
+		{
+			TMap<FNiagaraVariable, TObjectPtr<UNiagaraScriptVariable>> MetaDataMap = GraphWithInput->GetAllMetaData();
+			TArray<TObjectPtr<UNiagaraScriptVariable>> ScriptVariables;
+			MetaDataMap.GenerateValueArray(ScriptVariables);
+
+			FGuid InputGuid = Input->GetPersistentIdentity().Guids[1];
+			TObjectPtr<UNiagaraScriptVariable>* MatchingScriptVariable = ScriptVariables.FindByPredicate([InputGuid](TObjectPtr<UNiagaraScriptVariable> ScriptVariableCandidate)
+			{
+				return ScriptVariableCandidate->Metadata.GetVariableGuid() == InputGuid;
+			});
+
+			TOptional<bool> bIsStaticSwitch = GraphWithInput->IsStaticSwitch((*MatchingScriptVariable)->Variable);
+			if(ensure(MatchingScriptVariable != nullptr))
+			{
+				HierarchyInputToScriptVariableMap.Add(Input, *MatchingScriptVariable);
+				HierarchyInputToStaticSwitchMap.Add(Input, bIsStaticSwitch.GetValue());
+			}
+		}
+		// otherwise we have to use the name
+		else
+		{
+			TArray<FNiagaraVariable> Variables;
+			GraphWithInput->GetAllVariables(Variables);
+
+			FName InputName = Input->GetPersistentIdentity().Names[0];
+			FNiagaraVariableBase* MatchingVariable = Variables.FindByPredicate([InputName](FNiagaraVariable& Candidate)
+			{
+				return Candidate.GetName() == InputName;
+			});
+
+			if(ensure(MatchingVariable != nullptr))
+			{
+				HierarchyInputToAssignmentMap.Add(Input, *MatchingVariable);
+			}
+		}
+	}
+
+	TMap<UNiagaraHierarchyModuleInput*, TArray<FGuid>> HierarchyInputToChildrenGuidMap;
+	TMap<FGuid, TObjectPtr<UNiagaraScriptVariable>> ChildrenGuidToScriptVariablesMap;
+	// then we go through the child inputs and add them to the maps as well
+	for(UNiagaraHierarchyModuleInput* Input : ModuleInputs)
+	{
+		UNiagaraGraph* AssetGraph = HierarchyInputToFunctionCallNodeMap[Input]->GetCalledGraph();
+		ensure(AssetGraph != nullptr);
+
+		TArray<FNiagaraVariable> Variables;
+		AssetGraph->GetAllVariables(Variables);
+
+		FNiagaraVariable ThisInput;
+		if(Input->GetPersistentIdentity().Guids.Num() == 2)
+		{
+			ThisInput = HierarchyInputToScriptVariableMap[Input]->Variable;
+		}
+		else
+		{
+			FNiagaraVariable* FoundTargetVariable = Variables.FindByPredicate([InputName = Input->GetPersistentIdentity().Names[0]](FNiagaraVariable& CandidateVariable)
+			{
+				return CandidateVariable.GetName() == InputName;	
+			});
+
+			if(FoundTargetVariable)
+			{
+				ThisInput = *FoundTargetVariable;
+			}
+		}
+
+		ensure(ThisInput.IsValid());
+		
+		TArray<FNiagaraVariable> ChildrenInputs;
+		FNiagaraParameterHandle ThisInputHandle(ThisInput.GetName());
+
+		TMap<FGuid, int32> ChildrenGuidSortOrderMap;
+		TArray<FGuid> ChildrenInputGuids;
+		for(const FNiagaraVariable& Variable : Variables)
+		{
+			bool bIsModuleInput = Variable.IsInNameSpace(FNiagaraConstants::ModuleNamespaceString);
+			TOptional<bool> bIsStaticSwitchInputOptional = AssetGraph->IsStaticSwitch(Variable);
+			if(!bIsModuleInput && !bIsStaticSwitchInputOptional.Get(false))
+			{
+				continue;
+			}
+
+			TOptional<FNiagaraVariableMetaData> VariableMetaData = AssetGraph->GetMetaData(Variable);
+			if(VariableMetaData.IsSet())
+			{
+				if(VariableMetaData.GetValue().ParentAttribute == ThisInputHandle.GetName())
+				{
+					HierarchyInputToChildrenGuidMap.FindOrAdd(Input).Add(VariableMetaData.GetValue().GetVariableGuid());
+					ChildrenGuidToScriptVariablesMap.Add(VariableMetaData->GetVariableGuid(), AssetGraph->GetScriptVariable(Variable));
+					ChildrenGuidSortOrderMap.Add(VariableMetaData->GetVariableGuid(), VariableMetaData->EditorSortPriority);
+				}
+			}
+		}
+
+		// we sort children by sort order specified in the metadata
+		if(HierarchyInputToChildrenGuidMap.Contains(Input))
+		{
+			HierarchyInputToChildrenGuidMap[Input].Sort([&ChildrenGuidSortOrderMap](const FGuid& GuidA, const FGuid& GuidB)
+			{
+				return ChildrenGuidSortOrderMap[GuidA] < ChildrenGuidSortOrderMap[GuidB];
+			});
+		}
+	}
+
+	State.NodeGuidToModuleNodeMap = NodeGuidToModuleNodeMap;
+	State.HierarchyInputToAssignmentMap = HierarchyInputToAssignmentMap;
+	State.HierarchyInputToScriptVariableMap = HierarchyInputToScriptVariableMap;
+	State.HierarchyInputToChildrenGuidMap = HierarchyInputToChildrenGuidMap;
+	State.ChildrenGuidToScriptVariablesMap = ChildrenGuidToScriptVariablesMap;
+}
+
 void FNiagaraStackGraphUtilities::GetStackFunctionInputPinsWithoutCache(
 	const UNiagaraNodeFunctionCall& FunctionCallNode,
 	TConstArrayView<FNiagaraVariable> StaticVars,
@@ -1053,51 +1261,163 @@ UNiagaraNodeFunctionCall* FNiagaraStackGraphUtilities::FindModuleNode(FGuid Modu
 	return nullptr;
 }
 
+UNiagaraNodeFunctionCall* FNiagaraStackGraphUtilities::FindDynamicInputNodeForInput(UNiagaraNodeFunctionCall& OwningFunctionNode, FName UnaliasedParameterName)
+{
+	FNiagaraParameterHandle InputHandle(UnaliasedParameterName);
+	FNiagaraParameterHandle AliasedHandle = FNiagaraParameterHandle::CreateAliasedModuleParameterHandle(InputHandle, &OwningFunctionNode);
+
+	if(UEdGraphPin* OverridePin = FNiagaraStackGraphUtilities::GetStackFunctionInputOverridePin(OwningFunctionNode, AliasedHandle))
+	{
+		if(OverridePin->LinkedTo.Num() == 1 && OverridePin->LinkedTo[0] != nullptr && OverridePin->LinkedTo[0]->GetOwningNode() != nullptr)
+		{
+			UEdGraphNode* LinkedNode = OverridePin->LinkedTo[0]->GetOwningNode();
+			return Cast<UNiagaraNodeFunctionCall>(LinkedNode);
+		}
+	}
+
+	return nullptr;
+}
+
+UNiagaraNodeFunctionCall* FNiagaraStackGraphUtilities::FindFunctionCallNode(FGuid FunctionCallGuid,	TSharedRef<FNiagaraEmitterViewModel> EmitterViewModel)
+{
+	TSharedRef<FNiagaraScriptGraphViewModel> GraphViewModel = EmitterViewModel->GetSharedScriptViewModel()->GetGraphViewModel();
+
+	TArray<UNiagaraNodeFunctionCall*> FunctionCalls;
+	GraphViewModel->GetGraph()->GetNodesOfClass<UNiagaraNodeFunctionCall>(FunctionCalls);
+
+	for(UNiagaraNodeFunctionCall* FunctionCall : FunctionCalls)
+	{
+		if(FunctionCall->NodeGuid == FunctionCallGuid)
+		{
+			return FunctionCall;
+		}
+	}
+
+	return nullptr;
+}
+
+TArray<UNiagaraNodeFunctionCall*> FNiagaraStackGraphUtilities::FindModuleNodesForSimulationStage(UNiagaraSimulationStageBase& SimStage, TSharedRef<FNiagaraEmitterViewModel> EmitterViewModel)
+{
+	TArray<UNiagaraNodeOutput*> OutputNodes = GetAllOutputNodes(EmitterViewModel);
+	for(UNiagaraNodeOutput* Output : OutputNodes)
+	{
+		if(Output->GetUsage() == ENiagaraScriptUsage::ParticleSimulationStageScript && Output->GetUsageId() == SimStage.Script->GetUsageId())
+		{
+			TArray<UNiagaraNodeFunctionCall*> SimStageModules;
+			GetOrderedModuleNodes(*Output, SimStageModules);
+			return SimStageModules;
+		}
+	}
+
+	return {};
+}
+
+TArray<UNiagaraNodeFunctionCall*> FNiagaraStackGraphUtilities::FindModuleNodesForEventHandler(FNiagaraEventScriptProperties& EventScriptProperties, TSharedRef<FNiagaraEmitterViewModel> EmitterViewModel)
+{
+	TArray<UNiagaraNodeOutput*> OutputNodes = GetAllOutputNodes(EmitterViewModel);
+	for(UNiagaraNodeOutput* Output : OutputNodes)
+	{
+		if(Output->GetUsage() == ENiagaraScriptUsage::ParticleEventScript && Output->GetUsageId() == EventScriptProperties.Script->GetUsageId())
+		{
+			TArray<UNiagaraNodeFunctionCall*> EventHandlerModuleNodes;
+			GetOrderedModuleNodes(*Output, EventHandlerModuleNodes);
+			return EventHandlerModuleNodes;
+		}
+	}
+
+	return {};
+}
+
+UNiagaraNodeAssignment* FNiagaraStackGraphUtilities::FindAssignmentNode(FGuid AssignmentNodeGuid,	TSharedRef<FNiagaraEmitterViewModel> EmitterViewModel)
+{
+	UNiagaraNodeFunctionCall* AssignmentNode = FindModuleNode(AssignmentNodeGuid, EmitterViewModel);
+
+	if(AssignmentNode)
+	{
+		return CastChecked<UNiagaraNodeAssignment>(AssignmentNode);
+	}
+
+	return nullptr;
+}
+
 TOptional<FNiagaraStackGraphUtilities::FMatchingFunctionInputData> FNiagaraStackGraphUtilities::FindInputData(FNiagaraHierarchyIdentity ModuleInputIdentity, TSharedRef<FNiagaraEmitterViewModel> EmitterViewModel)
 {
 	FMatchingFunctionInputData Result;
 	
-	TArray<UNiagaraNodeOutput*> OutputNodes = GetAllOutputNodes(EmitterViewModel);
-	
-	UNiagaraNodeFunctionCall* MatchingFunctionCall = FindModuleNode(ModuleInputIdentity.Guids[0], EmitterViewModel);
+	UNiagaraNodeFunctionCall* MatchingFunctionCall = FindFunctionCallNode(ModuleInputIdentity.Guids[0], EmitterViewModel);
 	if(MatchingFunctionCall)
 	{
-		Result.FunctionCallNode = MatchingFunctionCall;
+		return FindInputData(*MatchingFunctionCall, ModuleInputIdentity, EmitterViewModel);
+	}					
+	
+	return TOptional<FMatchingFunctionInputData>();
+}
+
+TOptional<FNiagaraStackGraphUtilities::FMatchingFunctionInputData> FNiagaraStackGraphUtilities::FindInputData(const UNiagaraNodeFunctionCall& FunctionCallNode, FNiagaraHierarchyIdentity InputIdentity, TSharedRef<FNiagaraEmitterViewModel> EmitterViewModel)
+{
+	SCOPE_CYCLE_COUNTER(STAT_NiagaraEditor_StackGraphUtilities_FindInputData)
+	
+	FMatchingFunctionInputData Result;
+	Result.FunctionCallNode = const_cast<UNiagaraNodeFunctionCall*>(&FunctionCallNode);
+	
+	if(const UNiagaraNodeAssignment* AssignmentNode = Cast<UNiagaraNodeAssignment>(&FunctionCallNode))
+	{
+		TOptional<FMatchingFunctionInputData> InputData = FindAssignmentInputData(*AssignmentNode, InputIdentity.Names[0], EmitterViewModel);
+
+		if(InputData.IsSet())
+		{
+			Result.InputName = InputData->InputName;
+			Result.Type = InputData->Type;
+			Result.bIsHidden = InputData->bIsHidden;
+			Result.bIsStatic = InputData->bIsStatic;
+			Result.MetaData = InputData->MetaData;
+			Result.ChildrenInputGuids = InputData->ChildrenInputGuids;
+			Result.HiddenChildrenInputGuids = InputData->HiddenChildrenInputGuids;
+
+			return Result;
+		}
 	}
-					
-	if(MatchingFunctionCall)
+	else
 	{
-		TOptional<FMatchingFunctionInputData> InputData = FindInputData(*MatchingFunctionCall, ModuleInputIdentity.Guids[1], EmitterViewModel);
+		TOptional<FMatchingFunctionInputData> InputData = FindModuleInputData(FunctionCallNode, InputIdentity.Guids[1], EmitterViewModel);
 
 		if(InputData.IsSet())
 		{
 			Result.InputName = InputData->InputName;
 			Result.Type = InputData->Type;
+			Result.bIsHidden = InputData->bIsHidden;
+			Result.bIsStatic = InputData->bIsStatic;
 			Result.MetaData = InputData->MetaData;
 			Result.ChildrenInputGuids = InputData->ChildrenInputGuids;
+			Result.HiddenChildrenInputGuids = InputData->HiddenChildrenInputGuids;
 
 			return Result;
 		}
 	}
-	
+
 	return TOptional<FMatchingFunctionInputData>();
 }
 
-TOptional<FNiagaraStackGraphUtilities::FMatchingFunctionInputData> FNiagaraStackGraphUtilities::FindInputData(const UNiagaraNodeFunctionCall& FunctionCallNode, FGuid InputGuid, TSharedRef<FNiagaraEmitterViewModel> EmitterViewModel, bool bIncludeChildrenInputs)
+TOptional<FNiagaraStackGraphUtilities::FMatchingFunctionInputData> FNiagaraStackGraphUtilities::FindAssignmentInputData(const UNiagaraNodeAssignment& AssignmentNode, FName VariableName, TSharedRef<FNiagaraEmitterViewModel> EmitterViewModel)
 {
+	SCOPE_CYCLE_COUNTER(STAT_NiagaraEditor_StackGraphUtilities_FindAssignmentInputData);
+	
 	FMatchingFunctionInputData Result;
-	Result.FunctionCallNode = const_cast<UNiagaraNodeFunctionCall*>(&FunctionCallNode);
+	UNiagaraNodeAssignment* AssignmentNodeNonConst = const_cast<UNiagaraNodeAssignment*>(&AssignmentNode);
+	Result.FunctionCallNode = AssignmentNodeNonConst;
 	
-	FCompileConstantResolver ConstantResolver(EmitterViewModel->GetEmitter(), GetOutputNodeUsage(FunctionCallNode));
+	FCompileConstantResolver ConstantResolver(EmitterViewModel->GetEmitter(), GetOutputNodeUsage(AssignmentNode));
 
 	// our goal is to find the correct variable's metadata from the graph
 	TOptional<FNiagaraVariableMetaData> MetaData;
 				
 	// first, we check the normal module inputs
 	TArray<FNiagaraVariable> InputVariables;
+	TSet<FNiagaraVariable> HiddenVariables;
+
 	FNiagaraStackGraphUtilities::ENiagaraGetStackFunctionInputPinsOptions Options = FNiagaraStackGraphUtilities::ENiagaraGetStackFunctionInputPinsOptions::AllInputs;
-	FNiagaraStackGraphUtilities::GetStackFunctionInputs(FunctionCallNode, InputVariables, ConstantResolver, Options, true);
-				
+	FNiagaraStackGraphUtilities::GetStackFunctionInputs(AssignmentNode, InputVariables, HiddenVariables, ConstantResolver, Options, false);
+	
 	for(const FNiagaraVariable& InputVariable : InputVariables)
 	{
 		if (InputVariable.GetType().IsValid() == false)
@@ -1105,19 +1425,105 @@ TOptional<FNiagaraStackGraphUtilities::FMatchingFunctionInputData> FNiagaraStack
 			continue;
 		}
 
-		MetaData = FunctionCallNode.GetCalledGraph()->GetMetaData(InputVariable);
-		if(MetaData.IsSet() && MetaData.GetValue().GetVariableGuid() == InputGuid)
+		// the input variables we are iterating on have the Module prefix so we add it if necessary
+		FNiagaraParameterHandle SearchedParameterHandle(VariableName);
+		if(!SearchedParameterHandle.IsModuleHandle())
 		{
-			Result.InputName = InputVariable.GetName();
+			FString VariableNameWithModulePrefix = FNiagaraConstants::ModuleNamespaceString + ".";
+			VariableNameWithModulePrefix.Append(VariableName.ToString());
+			SearchedParameterHandle = FNiagaraParameterHandle(FName(VariableNameWithModulePrefix));
+		}
+		
+		if(InputVariable.GetName().IsEqual(SearchedParameterHandle.GetParameterHandleString()))
+		{
+			// we use the name without the prefix as that's what's stored in the hierarchy data
+			if(SearchedParameterHandle.IsModuleHandle())
+			{
+				FString VariableNameWithoutPrefix = SearchedParameterHandle.GetParameterHandleString().ToString().RightChop(FNiagaraConstants::ModuleNamespaceString.Len() + 1);
+				SearchedParameterHandle = FNiagaraParameterHandle(FName(VariableNameWithoutPrefix));
+			}
+			
+			Result.InputName = SearchedParameterHandle.GetParameterHandleString();
 			Result.Type = InputVariable.GetType();
-			Result.MetaData = MetaData.GetValue();
+			Result.MetaData = FNiagaraVariableMetaData();
+			Result.bIsStatic = false;
+			Result.bIsHidden = HiddenVariables.Contains(InputVariable);
+			
+			return Result;
+		}
+	}
+
+	return TOptional<FNiagaraStackGraphUtilities::FMatchingFunctionInputData>();
+}
+
+TOptional<FNiagaraStackGraphUtilities::FMatchingFunctionInputData> FNiagaraStackGraphUtilities::FindModuleInputData(const UNiagaraNodeFunctionCall& FunctionCallNode, FGuid VariableGuid, TSharedRef<FNiagaraEmitterViewModel> EmitterViewModel, bool bIncludeChildrenInputs, TMap<FInputDataCacheKey, FMatchingFunctionInputData>* OptionalCache)
+{
+	SCOPE_CYCLE_COUNTER(STAT_NiagaraEditor_StackGraphUtilities_FindModuleInputData)
+
+	if(OptionalCache != nullptr)
+	{
+		FInputDataCacheKey CacheKey(FunctionCallNode.NodeGuid, VariableGuid);
+		if(OptionalCache->Contains(CacheKey))
+		{
+			return (*OptionalCache)[CacheKey];
+		}
+	}
+	
+	FCompileConstantResolver ConstantResolver(EmitterViewModel->GetEmitter(), GetOutputNodeUsage(FunctionCallNode));
+
+	// our goal is to find the correct variable's metadata from the graph
+	TOptional<FNiagaraVariableMetaData> MetaData;
+				
+	// first, we check the normal module inputs
+	TArray<FNiagaraVariable> InputVariables;
+	TSet<FNiagaraVariable> HiddenVariables;
+
+	FNiagaraStackGraphUtilities::ENiagaraGetStackFunctionInputPinsOptions Options = FNiagaraStackGraphUtilities::ENiagaraGetStackFunctionInputPinsOptions::ModuleInputsOnly;
+	FNiagaraStackGraphUtilities::GetStackFunctionInputs(FunctionCallNode, InputVariables, HiddenVariables, ConstantResolver, Options, false);
+	
+	for(const FNiagaraVariable& InputVariable : InputVariables)
+	{
+		if (InputVariable.GetType().IsValid() == false)
+		{
+			continue;
+		}
+		
+		MetaData = FunctionCallNode.GetCalledGraph()->GetMetaData(InputVariable);
+
+		if(MetaData.IsSet())
+		{
+			FMatchingFunctionInputData PotentialMatchingInputData;
+			PotentialMatchingInputData.FunctionCallNode = const_cast<UNiagaraNodeFunctionCall*>(&FunctionCallNode);
+			PotentialMatchingInputData.InputName = InputVariable.GetName();
+			PotentialMatchingInputData.Type = InputVariable.GetType();
+			PotentialMatchingInputData.MetaData = MetaData.GetValue();
+			PotentialMatchingInputData.bIsStatic = false;
+			PotentialMatchingInputData.bIsHidden = HiddenVariables.Contains(InputVariable);
 
 			if(bIncludeChildrenInputs)
 			{
-				Result.ChildrenInputGuids = GetChildrenInputGuids(FunctionCallNode, Result.InputName);				
+				PotentialMatchingInputData.ChildrenInputGuids = GetChildrenInputGuids(FunctionCallNode, PotentialMatchingInputData.InputName);
+
+				// we only care about hidden children if the parent input isn't hidden
+				if(!PotentialMatchingInputData.bIsHidden)
+				{
+					PotentialMatchingInputData.HiddenChildrenInputGuids = GetHiddenChildrenInputGuids(FunctionCallNode, PotentialMatchingInputData.InputName, EmitterViewModel).Array();
+				}
 			}
-			
-			return Result;
+
+			if(OptionalCache != nullptr)
+			{
+				OptionalCache->Add(FInputDataCacheKey(FunctionCallNode.NodeGuid, MetaData->GetVariableGuid()), PotentialMatchingInputData);
+			}
+
+			// if we specified an optional cache, we want to keep on creating the cache data even after we've found the result to avoid redundant calls
+			if(OptionalCache == nullptr)
+			{
+				if(MetaData.GetValue().GetVariableGuid() == VariableGuid)
+				{
+					return PotentialMatchingInputData;
+				}
+			}			
 		}
 	}
 
@@ -1127,13 +1533,7 @@ TOptional<FNiagaraStackGraphUtilities::FMatchingFunctionInputData> FNiagaraStack
 				
 	FNiagaraStackGraphUtilities::GetStackFunctionStaticSwitchPins(FunctionCallNode, SwitchPins, HiddenSwitchPins, ConstantResolver);
 	for (UEdGraphPin* InputPin : SwitchPins)
-	{
-		// The static switch pin names to not contain the module namespace, as they are not part of the parameter maps.
-		// We add it here only to check for name clashes with actual module parameters.
-		FString ModuleName = PARAM_MAP_MODULE_STR;
-		InputPin->PinName.AppendString(ModuleName);
-		FName SwitchPinName(*ModuleName);
-
+	{		
 		FNiagaraVariable InputVariable = UEdGraphSchema_Niagara::PinToNiagaraVariable(InputPin);
 		if (InputVariable.GetType().IsValid() == false)
 		{
@@ -1143,20 +1543,52 @@ TOptional<FNiagaraStackGraphUtilities::FMatchingFunctionInputData> FNiagaraStack
 		if (FunctionCallNode.GetCalledGraph() != nullptr)
 		{
 			MetaData = FunctionCallNode.GetCalledGraph()->GetMetaData(InputVariable);
-			// if metadata is set and the 2nd guid matches the variable, we found the input that our HierarchyModuleInput represents
-			if(MetaData.IsSet() && MetaData.GetValue().GetVariableGuid() == InputGuid)
+
+			if(MetaData.IsSet())
 			{
-				Result.InputName = SwitchPinName;
-				Result.Type = InputVariable.GetType();
-				Result.MetaData = MetaData.GetValue();
+				FMatchingFunctionInputData PotentialMatchingInputData;
+				PotentialMatchingInputData.FunctionCallNode = const_cast<UNiagaraNodeFunctionCall*>(&FunctionCallNode);
+				PotentialMatchingInputData.InputName = InputPin->PinName;
+				PotentialMatchingInputData.Type = InputVariable.GetType();
+				PotentialMatchingInputData.MetaData = MetaData.GetValue();
+				PotentialMatchingInputData.bIsStatic = true;
+				PotentialMatchingInputData.bIsHidden = HiddenSwitchPins.Contains(InputPin);
 
 				if(bIncludeChildrenInputs)
 				{
-					Result.ChildrenInputGuids = GetChildrenInputGuids(FunctionCallNode, Result.InputName);				
+					PotentialMatchingInputData.ChildrenInputGuids = GetChildrenInputGuids(FunctionCallNode, PotentialMatchingInputData.InputName);
+
+					// we only care about hidden children if the parent input isn't hidden
+					if(!PotentialMatchingInputData.bIsHidden)
+					{
+						PotentialMatchingInputData.HiddenChildrenInputGuids = GetHiddenChildrenInputGuids(FunctionCallNode, PotentialMatchingInputData.InputName, EmitterViewModel).Array();
+					}
 				}
 				
-				return Result;
-			}
+				if(OptionalCache != nullptr)
+				{
+					OptionalCache->Add(FInputDataCacheKey(FunctionCallNode.NodeGuid, MetaData->GetVariableGuid()), PotentialMatchingInputData);
+				}
+
+				if(OptionalCache == nullptr)
+				{
+					// if metadata is set and the 2nd guid matches the variable, we found the input that our HierarchyModuleInput represents
+					if(MetaData.GetValue().GetVariableGuid() == VariableGuid)
+					{
+						return PotentialMatchingInputData;
+					}
+				}
+			}			
+		}
+	}
+
+	// if we specified a cache, we didn't have an early return. That means when we reach this code, the entire data should be cached now and if it's not, that means the data doesn't exist.
+	if(OptionalCache != nullptr)
+	{
+		FInputDataCacheKey InputDataCacheKey(FunctionCallNode.NodeGuid, VariableGuid);
+		if(OptionalCache->Contains(InputDataCacheKey))
+		{
+			return (*OptionalCache)[InputDataCacheKey];
 		}
 	}
 	
@@ -1182,6 +1614,45 @@ TArray<FGuid> FNiagaraStackGraphUtilities::GetChildrenInputGuids(const UNiagaraN
 				Result.Add(MetaData->GetVariableGuid());
 			}
 		}
+		else
+		{
+			UE_LOG(LogNiagaraEditor, Log, TEXT("Couldn't find metadata for variable"));
+		}
+	}
+
+	return Result;
+}
+
+TSet<FGuid> FNiagaraStackGraphUtilities::GetHiddenChildrenInputGuids(const UNiagaraNodeFunctionCall& FunctionCallNode, FName ParentInput, TSharedRef<FNiagaraEmitterViewModel> EmitterViewModel)
+{
+	TSet<FGuid> Result;
+	
+	FCompileConstantResolver ConstantResolver(EmitterViewModel->GetEmitter(), GetOutputNodeUsage(FunctionCallNode));
+
+	TArray<FGuid> ChildrenInputGuids = GetChildrenInputGuids(FunctionCallNode, ParentInput);
+	// first, we check the normal module inputs
+	TArray<FNiagaraVariable> InputVariables;
+	TSet<FNiagaraVariable> HiddenVariables;
+	FNiagaraStackGraphUtilities::ENiagaraGetStackFunctionInputPinsOptions Options = FNiagaraStackGraphUtilities::ENiagaraGetStackFunctionInputPinsOptions::AllInputs;
+	FNiagaraStackGraphUtilities::GetStackFunctionInputs(FunctionCallNode, InputVariables, HiddenVariables, ConstantResolver, Options, true);
+
+	TArray<FGuid> ChildrenGuids = GetChildrenInputGuids(FunctionCallNode, ParentInput);
+	
+	for(const FNiagaraVariable& Variable : HiddenVariables.Array())
+	{
+		TOptional<FNiagaraVariableMetaData> MetaData = FunctionCallNode.GetCalledGraph()->GetMetaData(Variable);
+
+		if(MetaData.IsSet())
+		{
+			if(ChildrenGuids.Contains(MetaData->GetVariableGuid()))
+			{
+				Result.Add(MetaData->GetVariableGuid());
+			}
+		}
+		else
+		{
+			UE_LOG(LogNiagaraEditor, Log, TEXT("Couldn't find metadata for variable"));
+		}
 	}
 
 	return Result;
diff --git a/Engine/Plugins/FX/Niagara/Source/NiagaraEditor/Private/ViewModels/Stack/NiagaraStackInputCategory.cpp b/Engine/Plugins/FX/Niagara/Source/NiagaraEditor/Private/ViewModels/Stack/NiagaraStackInputCategory.cpp
--- a/Engine/Plugins/FX/Niagara/Source/NiagaraEditor/Private/ViewModels/Stack/NiagaraStackInputCategory.cpp
+++ b/Engine/Plugins/FX/Niagara/Source/NiagaraEditor/Private/ViewModels/Stack/NiagaraStackInputCategory.cpp
@@ -1,68 +1,74 @@
 // Copyright Epic Games, Inc. All Rights Reserved.
 
 #include "ViewModels/Stack/NiagaraStackInputCategory.h"
+
+#include "IDetailTreeNode.h"
+#include "IPropertyRowGenerator.h"
 #include "ViewModels/Stack/NiagaraStackFunctionInput.h"
 #include "NiagaraNodeFunctionCall.h"
 #include "NiagaraClipboard.h"
+#include "NiagaraConstants.h"
+#include "NiagaraSimulationStageBase.h"
+#include "ViewModels/NiagaraEmitterViewModel.h"
+#include "ViewModels/HierarchyEditor/NiagaraSummaryViewViewModel.h"
+#include "ViewModels/Stack/NiagaraStackEventScriptItemGroup.h"
 #include "ViewModels/Stack/NiagaraStackFunctionInputCollection.h"
+#include "ViewModels/Stack/NiagaraStackModuleItem.h"
+#include "ViewModels/Stack/NiagaraStackPropertyRow.h"
+#include "ViewModels/Stack/NiagaraStackRendererItem.h"
+#include "ViewModels/Stack/NiagaraStackSimulationStageGroup.h"
 
 #include UE_INLINE_GENERATED_CPP_BY_NAME(NiagaraStackInputCategory)
 
-void UNiagaraStackInputCategory::Initialize(
-	FRequiredEntryData InRequiredEntryData,
-	FString InputCategoryStackEditorDataKey,
-	FText InCategoryName,
-	bool bInIsTopLevelCategory,
-	FString InOwnerStackItemEditorDataKey)
+using namespace FNiagaraStackGraphUtilities;
+
+void UNiagaraStackCategory::Initialize(FRequiredEntryData InRequiredEntryData, FString InOwningStackItemEditorDataKey, FString InStackEditorDataKey)
 {
-	bool bCategoryIsAdvanced = false;
-	Super::Initialize(InRequiredEntryData, InOwnerStackItemEditorDataKey, InputCategoryStackEditorDataKey);
-	CategoryName = InCategoryName;
-	bShouldShowInStack = true;
-	bIsTopLevelCategory = bInIsTopLevelCategory;
-	CategorySpacer = nullptr;
-	
+	Super::Initialize(InRequiredEntryData, InOwningStackItemEditorDataKey, InStackEditorDataKey);
 	AddChildFilter(FOnFilterChild::CreateUObject(this, &UNiagaraStackInputCategory::FilterForVisibleCondition));
 	AddChildFilter(FOnFilterChild::CreateUObject(this, &UNiagaraStackInputCategory::FilterForIsInlineEditConditionToggle));
+	bShouldShowInStack = true;
+	CategorySpacer = nullptr;
 }
 
-const FText& UNiagaraStackInputCategory::GetCategoryName() const
+FText UNiagaraStackCategory::GetDisplayName() const
 {
-	return CategoryName;
+	return FText::GetEmpty();
 }
 
-void UNiagaraStackInputCategory::ResetInputs()
+bool UNiagaraStackCategory::GetShouldShowInStack() const
 {
-	Inputs.Empty();
+	// Categories may be empty if their children have all been hidden due to visible filters or advanced display.
+	// in the case where all children have been hidden, don't show the category in the stack.
+	TArray<UNiagaraStackEntry*> CurrentFilteredChildren;
+	GetFilteredChildren(CurrentFilteredChildren);
+	int32 EmptyCount = CategorySpacer == nullptr ? 0 : 1;
+	return bShouldShowInStack && CurrentFilteredChildren.Num() > EmptyCount;
 }
 
-void UNiagaraStackInputCategory::AddInput(UNiagaraNodeFunctionCall* InModuleNode, UNiagaraNodeFunctionCall* InInputFunctionCallNode, FName InInputParameterHandle, FNiagaraTypeDefinition InInputType, EStackParameterBehavior InParameterBehavior, TOptional<FText> InOptionalDisplayName, bool bIsInputHidden, bool bIsChildInput)
+UNiagaraStackEntry::EStackRowStyle UNiagaraStackCategory::GetStackRowStyle() const
 {
-	Inputs.Add({ InModuleNode, InInputFunctionCallNode, InInputParameterHandle, InInputType, InParameterBehavior, InOptionalDisplayName, bIsInputHidden, bIsChildInput });
+	return IsTopLevelCategory() ? EStackRowStyle::ItemCategory : EStackRowStyle::ItemSubCategory;
 }
 
-void UNiagaraStackInputCategory::RefreshChildrenInternal(const TArray<UNiagaraStackEntry*>& CurrentChildren, TArray<UNiagaraStackEntry*>& NewChildren, TArray<FStackIssue>& NewIssues)
+void UNiagaraStackCategory::GetSearchItems(TArray<FStackSearchItem>& SearchItems) const
 {
-	for (FInputParameterHandleAndType& Input : Inputs)
+	// Don't return search results if we're not being shown in the stack, otherwise we'll generate search results which can't be navigated to.
+	if (GetShouldShowInStack())
 	{
-		UNiagaraStackFunctionInput* InputChild = FindCurrentChildOfTypeByPredicate<UNiagaraStackFunctionInput>(CurrentChildren, [&](UNiagaraStackFunctionInput* CurrentInput) 
-		{ 
-			return CurrentInput->GetInputParameterHandle() == Input.ParameterHandle && CurrentInput->GetInputType() == Input.Type && CurrentInput->GetInputFunctionCallInitialScript() == Input.InputFunctionCallNode->FunctionScript;
-		});
-
-		if (InputChild == nullptr)
-		{
-			InputChild = NewObject<UNiagaraStackFunctionInput>(this);
-			InputChild->Initialize(CreateDefaultChildRequiredData(), *Input.ModuleNode, *Input.InputFunctionCallNode,
-				Input.ParameterHandle, Input.Type, Input.ParameterBehavior, GetOwnerStackItemEditorDataKey());
-		}
-		InputChild->SetIsHidden(Input.bIsHidden);
-		InputChild->SetSemanticChild(Input.bIsChildInput);
-		InputChild->SetSummaryViewDisiplayName(Input.DisplayName);
-		NewChildren.Add(InputChild);
+		Super::GetSearchItems(SearchItems);
 	}
+}
 
-	if (bIsTopLevelCategory)
+int32 UNiagaraStackCategory::GetChildIndentLevel() const
+{
+	// We want to keep inputs under a top level category at the same indent level as the category.
+	return IsTopLevelCategory() ? GetIndentLevel() : Super::GetChildIndentLevel();
+}
+
+void UNiagaraStackCategory::RefreshChildrenInternal(const TArray<UNiagaraStackEntry*>& CurrentChildren,	TArray<UNiagaraStackEntry*>& NewChildren, TArray<FStackIssue>& NewIssues)
+{
+	if (IsTopLevelCategory())
 	{
 		if (CategorySpacer == nullptr)
 		{
@@ -75,30 +81,29 @@
 	}
 }
 
-int32 UNiagaraStackInputCategory::GetChildIndentLevel() const
+bool UNiagaraStackCategory::FilterForVisibleCondition(const UNiagaraStackEntry& Child) const
 {
-	// We want to keep inputs under a top level category at the same indent level as the category.
-	return bIsTopLevelCategory ? GetIndentLevel() : Super::GetChildIndentLevel();
+	const UNiagaraStackFunctionInput* StackFunctionInputChild = Cast<UNiagaraStackFunctionInput>(&Child);
+	return StackFunctionInputChild == nullptr || StackFunctionInputChild->GetShouldPassFilterForVisibleCondition();
 }
 
-FText UNiagaraStackInputCategory::GetDisplayName() const
+bool UNiagaraStackCategory::FilterForIsInlineEditConditionToggle(const UNiagaraStackEntry& Child) const
 {
-	return CategoryName;
+	const UNiagaraStackFunctionInput* StackFunctionInputChild = Cast<UNiagaraStackFunctionInput>(&Child);
+	return StackFunctionInputChild == nullptr || StackFunctionInputChild->GetIsInlineEditConditionToggle() == false;
 }
 
-bool UNiagaraStackInputCategory::GetShouldShowInStack() const
+void UNiagaraStackInputCategory::Initialize(
+	FRequiredEntryData InRequiredEntryData,
+	FString InputCategoryStackEditorDataKey,
+	FText InCategoryName,
+	bool bInIsTopLevelCategory,
+	FString InOwnerStackItemEditorDataKey)
 {
-	// Categories may be empty if their children have all been hidden due to visible filters or advanced display.
-	// in the case where all children have been hidden, don't show the category in the stack.
-	TArray<UNiagaraStackEntry*> CurrentFilteredChildren;
-	GetFilteredChildren(CurrentFilteredChildren);
-	int32 EmptyCount = CategorySpacer == nullptr ? 0 : 1;
-	return bShouldShowInStack && CurrentFilteredChildren.Num() > EmptyCount;
+	bool bCategoryIsAdvanced = false;
+	Super::Initialize(InRequiredEntryData, InOwnerStackItemEditorDataKey, InputCategoryStackEditorDataKey);
+	CategoryName = InCategoryName;
+	bIsTopLevelCategory = bInIsTopLevelCategory;
 }
 
-UNiagaraStackEntry::EStackRowStyle UNiagaraStackInputCategory::GetStackRowStyle() const
-{
-	return bIsTopLevelCategory ? EStackRowStyle::ItemCategory : EStackRowStyle::ItemSubCategory;
-}
-
 bool UNiagaraStackInputCategory::GetIsEnabled() const
@@ -114,13 +119,42 @@
 	return false;
 }
 
-void UNiagaraStackInputCategory::GetSearchItems(TArray<FStackSearchItem>& SearchItems) const
+FText UNiagaraStackInputCategory::GetDisplayName() const
 {
-	// Don't return search results if we're not being shown in the stack, otherwise we'll generate search results which can't be navigated to.
-	if (GetShouldShowInStack())
+	return CategoryName;
+}
+
+void UNiagaraStackInputCategory::ResetInputs()
+{
+	Inputs.Empty();
+}
+
+void UNiagaraStackInputCategory::AddInput(UNiagaraNodeFunctionCall* InModuleNode, UNiagaraNodeFunctionCall* InInputFunctionCallNode, FName InInputParameterHandle, FNiagaraTypeDefinition InInputType, EStackParameterBehavior InParameterBehavior, TOptional<FText> InOptionalDisplayName, bool bIsInputHidden, bool bIsChildInput)
+{
+	Inputs.Add({ InModuleNode, InInputFunctionCallNode, InInputParameterHandle, InInputType, InParameterBehavior, InOptionalDisplayName, bIsInputHidden, bIsChildInput });
+}
+
+void UNiagaraStackInputCategory::RefreshChildrenInternal(const TArray<UNiagaraStackEntry*>& CurrentChildren, TArray<UNiagaraStackEntry*>& NewChildren, TArray<FStackIssue>& NewIssues)
+{
+	for (FInputParameterHandleAndType& Input : Inputs)
 	{
-		Super::GetSearchItems(SearchItems);
+		UNiagaraStackFunctionInput* InputChild = FindCurrentChildOfTypeByPredicate<UNiagaraStackFunctionInput>(CurrentChildren, [&](UNiagaraStackFunctionInput* CurrentInput) 
+		{ 
+			return CurrentInput->GetInputParameterHandle() == Input.ParameterHandle && CurrentInput->GetInputType() == Input.Type && CurrentInput->GetInputFunctionCallInitialScript() == Input.InputFunctionCallNode->FunctionScript;
+		});
+
+		if (InputChild == nullptr)
+		{
+			InputChild = NewObject<UNiagaraStackFunctionInput>(this);
+			InputChild->Initialize(CreateDefaultChildRequiredData(), *Input.ModuleNode, *Input.InputFunctionCallNode,
+				Input.ParameterHandle, Input.Type, Input.ParameterBehavior, GetOwnerStackItemEditorDataKey());
+		}
+		InputChild->SetIsHidden(Input.bIsHidden);
+		InputChild->SetSemanticChild(Input.bIsChildInput);
+		NewChildren.Add(InputChild);
 	}
+
+	Super::RefreshChildrenInternal(CurrentChildren, NewChildren, NewIssues);
 }
 
 void UNiagaraStackInputCategory::SetShouldShowInStack(bool bInShouldShowInStack)
@@ -175,16 +209,427 @@
 	SetValuesFromFunctionInputsInternal(this, ClipboardFunctionInputs, [](UNiagaraStackFunctionInput* ChildInput) { return ChildInput->IsStaticParameter() == false; }, nullptr);
 }
 
-bool UNiagaraStackInputCategory::FilterForVisibleCondition(const UNiagaraStackEntry& Child) const
+void UNiagaraStackSummaryCategory::Initialize(FRequiredEntryData InRequiredEntryData, TSharedPtr<FNiagaraHierarchyCategoryViewModel> InCategoryViewModel, FString InOwnerStackItemEditorDataKey)
 {
-	const UNiagaraStackFunctionInput* StackFunctionInputChild = Cast<UNiagaraStackFunctionInput>(&Child);
-	return StackFunctionInputChild == nullptr || StackFunctionInputChild->GetShouldPassFilterForVisibleCondition();
+	CategoryViewModelWeakPtr = InCategoryViewModel;
+	
+	FString EditorDataKey = FString::Printf(TEXT("%s-%s"), *InOwnerStackItemEditorDataKey, *InCategoryViewModel->GetCategoryName().ToString());
+	Super::Initialize(InRequiredEntryData, InOwnerStackItemEditorDataKey, EditorDataKey);
 }
 
-bool UNiagaraStackInputCategory::FilterForIsInlineEditConditionToggle(const UNiagaraStackEntry& Child) const
+FText UNiagaraStackSummaryCategory::GetDisplayName() const
 {
-	const UNiagaraStackFunctionInput* StackFunctionInputChild = Cast<UNiagaraStackFunctionInput>(&Child);
-	return StackFunctionInputChild == nullptr || StackFunctionInputChild->GetIsInlineEditConditionToggle() == false;
+	return CategoryViewModelWeakPtr.Pin()->GetCategoryName();
 }
 
+void UNiagaraStackSummaryCategory::RefreshChildrenInternal(const TArray<UNiagaraStackEntry*>& CurrentChildren, TArray<UNiagaraStackEntry*>& NewChildren, TArray<FStackIssue>& NewIssues)
+{
+	FInputDataCollection State;
+	GatherInputRelationsForStack(State, GetEmitterViewModel().ToSharedRef());
 
+	TArray<TSharedPtr<FNiagaraHierarchyItemViewModelBase>> ChildrenViewModels;
+	CategoryViewModelWeakPtr.Pin()->GetChildrenViewModelsForType<UNiagaraHierarchyItemBase, FNiagaraHierarchyItemViewModelBase>(ChildrenViewModels, false);
+
+	TArray<TSharedPtr<FNiagaraHierarchyItemViewModelBase>> AllChildrenViewModels;
+	CategoryViewModelWeakPtr.Pin()->GetChildrenViewModelsForType<UNiagaraHierarchyItemBase, FNiagaraHierarchyItemViewModelBase>(AllChildrenViewModels, true);
+	
+	// first we gather all function call nodes so we can create cache for them instead of looking it up for each input individually
+	TSet<UNiagaraNodeFunctionCall*> UsedFunctionCallNodes;
+	for(TSharedPtr<FNiagaraHierarchyItemViewModelBase> HierarchyViewModel : AllChildrenViewModels)
+	{
+		UNiagaraHierarchyItemBase* Data = HierarchyViewModel->GetDataMutable();
+		if(UNiagaraHierarchyModuleInput* ModuleInput = Cast<UNiagaraHierarchyModuleInput>(Data))
+		{
+			UsedFunctionCallNodes.Add(State.NodeGuidToModuleNodeMap[ModuleInput->GetPersistentIdentity().Guids[0]]);
+		}
+		else if(UNiagaraHierarchyAssignmentInput* AssignmentInput = Cast<UNiagaraHierarchyAssignmentInput>(Data))
+		{
+			UsedFunctionCallNodes.Add(State.NodeGuidToModuleNodeMap[AssignmentInput->GetPersistentIdentity().Guids[0]]);
+		}
+		else if(UNiagaraHierarchyModule* Module = Cast<UNiagaraHierarchyModule>(Data))
+		{
+			UsedFunctionCallNodes.Add(State.NodeGuidToModuleNodeMap[Module->GetPersistentIdentity().Guids[0]]);
+		}
+	}
+	
+	TMap<FGuid, TSet<FNiagaraVariable>> FunctionCallToHiddenVariablesMap;
+	for(UNiagaraNodeFunctionCall* FunctionCall : UsedFunctionCallNodes)
+	{
+		TArray<FNiagaraVariable> InputVariables;
+		TSet<FNiagaraVariable> HiddenVariables;
+		FCompileConstantResolver Resolver(GetEmitterViewModel()->GetEmitter(), FNiagaraStackGraphUtilities::GetOutputNodeUsage(*FunctionCall));
+		FNiagaraStackGraphUtilities::GetStackFunctionInputs(*FunctionCall, InputVariables, HiddenVariables, Resolver, ENiagaraGetStackFunctionInputPinsOptions::ModuleInputsOnly, true);
+		FunctionCallToHiddenVariablesMap.Add(FunctionCall->NodeGuid, HiddenVariables);
+	}
+	
+	for(TSharedPtr<FNiagaraHierarchyItemViewModelBase> HierarchyViewModel : ChildrenViewModels)
+	{
+		UNiagaraHierarchyItemBase* Data = HierarchyViewModel->GetDataMutable();
+
+		if(UNiagaraHierarchyModuleInput* ModuleInput = Cast<UNiagaraHierarchyModuleInput>(Data))
+		{
+			if(GetEmitterViewModel().IsValid())
+			{
+				UNiagaraNodeFunctionCall* OwningFunctionCallNode = State.NodeGuidToModuleNodeMap[ModuleInput->GetPersistentIdentity().Guids[0]];
+				
+				TSharedPtr<FNiagaraModuleInputViewModel> ModuleInputViewModel = StaticCastSharedPtr<FNiagaraModuleInputViewModel>(HierarchyViewModel);
+				TOptional<FInputData> ModuleBaseInputData = ModuleInputViewModel->GetInputData();
+				
+				if(ModuleBaseInputData.IsSet())
+				{
+					UNiagaraStackFunctionInput* TopLevelInput = FindCurrentChildOfTypeByPredicate<UNiagaraStackFunctionInput>(CurrentChildren, [&](UNiagaraStackFunctionInput* CurrentInput) 
+					{
+						return CurrentInput->GetInputParameterHandle() == ModuleBaseInputData->InputName && CurrentInput->GetInputType() == ModuleBaseInputData->Type
+							&& CurrentInput->GetInputFunctionCallInitialScript() == OwningFunctionCallNode->FunctionScript
+							// we are checking for node guid as we could have 2x the same input from 2 identical modules
+							&& CurrentInput->GetInputFunctionCallNode().NodeGuid == OwningFunctionCallNode->NodeGuid;
+					});
+	
+					if(TopLevelInput == nullptr)
+					{
+						TopLevelInput = NewObject<UNiagaraStackFunctionInput>(this);
+						TopLevelInput->Initialize(CreateDefaultChildRequiredData(), *OwningFunctionCallNode, *OwningFunctionCallNode, ModuleBaseInputData->InputName, ModuleBaseInputData->Type, ModuleBaseInputData->bIsStatic ? EStackParameterBehavior::Static : EStackParameterBehavior::Dynamic, GetOwnerStackItemEditorDataKey());
+						TAttribute<FText> SummaryViewDisplayNameOverride;
+						SummaryViewDisplayNameOverride.BindUObject(ModuleInput, &UNiagaraHierarchyModuleInput::GetDisplayNameOverride);
+						TopLevelInput->SetSummaryViewDisplayName(SummaryViewDisplayNameOverride);
+						TAttribute<FText> SummaryViewTooltipOverride;
+						SummaryViewTooltipOverride.BindUObject(ModuleInput, &UNiagaraHierarchyModuleInput::GetTooltipOverride);
+						TopLevelInput->SetSummaryViewTooltip(SummaryViewTooltipOverride);
+					}
+
+					FNiagaraVariable DisplayedVariable(TopLevelInput->GetInputType(), TopLevelInput->GetInputParameterHandle().GetParameterHandleString());
+					bool bIsTopLevelHidden = FunctionCallToHiddenVariablesMap[ModuleInput->GetPersistentIdentity().Guids[0]].Contains(DisplayedVariable);
+					TopLevelInput->SetIsHidden(bIsTopLevelHidden);
+					
+					NewChildren.Add(TopLevelInput);
+
+					TArray<UNiagaraHierarchyModuleInput*> ChildInputs;
+					ModuleInput->GetChildrenOfType<UNiagaraHierarchyModuleInput>(ChildInputs);
+					
+					for(UNiagaraHierarchyModuleInput* ChildInput : ChildInputs)
+					{
+						TSharedPtr<FNiagaraModuleInputViewModel> ChildInputViewModel = StaticCastSharedPtr<FNiagaraModuleInputViewModel>(ModuleInputViewModel->FindViewModelForChild(ChildInput, false));
+						TOptional<FInputData> ChildInputData = ChildInputViewModel->GetInputData();
+
+						if(ChildInputData.IsSet())
+						{
+							UNiagaraNodeFunctionCall* ChildFunctionNode = ChildInputData->FunctionCallNode;
+							UNiagaraStackFunctionInput* StackChildInput = FindCurrentChildOfTypeByPredicate<UNiagaraStackFunctionInput>(CurrentChildren, [&](UNiagaraStackFunctionInput* CurrentInput) 
+							{
+								return CurrentInput->GetInputParameterHandle() == ChildInputData->InputName && CurrentInput->GetInputType() == ChildInputData->Type
+									&& CurrentInput->GetInputFunctionCallInitialScript() == ChildFunctionNode->FunctionScript
+									// we are checking for node guid as we could have 2x the same input from 2 identical modules
+									&& CurrentInput->GetInputFunctionCallNode().NodeGuid == ChildFunctionNode->NodeGuid;
+							});
+							
+							if(StackChildInput == nullptr)
+							{
+								StackChildInput = NewObject<UNiagaraStackFunctionInput>(this);
+								StackChildInput->Initialize(CreateDefaultChildRequiredData(), *ChildInputData->FunctionCallNode, *ChildInputData->FunctionCallNode, ChildInputData->InputName, ChildInputData->Type, ChildInputData->bIsStatic ? EStackParameterBehavior::Static : EStackParameterBehavior::Dynamic, GetOwnerStackItemEditorDataKey());
+								TAttribute<FText> SummaryViewDisplayNameOverride;
+								SummaryViewDisplayNameOverride.BindUObject(ChildInput, &UNiagaraHierarchyModuleInput::GetDisplayNameOverride);
+								StackChildInput->SetSummaryViewDisplayName(SummaryViewDisplayNameOverride);
+								TAttribute<FText> SummaryViewTooltipOverride;
+								SummaryViewTooltipOverride.BindUObject(ChildInput, &UNiagaraHierarchyModuleInput::GetTooltipOverride);
+								StackChildInput->SetSummaryViewTooltip(SummaryViewTooltipOverride);
+								
+								StackChildInput->SetSemanticChild(true);
+							}
+
+							FNiagaraVariable DisplayedChildVariable(StackChildInput->GetInputType(), StackChildInput->GetInputParameterHandle().GetParameterHandleString());
+							bool bIsChildInputHidden = FunctionCallToHiddenVariablesMap[ChildInput->GetPersistentIdentity().Guids[0]].Contains(DisplayedChildVariable);
+							StackChildInput->SetIsHidden(bIsChildInputHidden);
+							
+							NewChildren.Add(StackChildInput);
+						}
+					}
+					
+					/** Automatically add children inputs. Disabled as user is managing this. */
+					// if(State.HierarchyInputToChildrenGuidMap.Contains(ModuleInput))
+					// {
+					// 	// children guids are already sorted by sort order
+					// 	TArray<FGuid> ChildrenInputs = State.HierarchyInputToChildrenGuidMap[ModuleInput];
+					//
+					// 	for(FGuid ChildrenInputGuid : ChildrenInputs)
+					// 	{
+					// 		if(UNiagaraScriptVariable* ChildScriptVariable = State.ChildrenGuidToScriptVariablesMap[ChildrenInputGuid])
+					// 		{
+					// 			FNiagaraVariable ChildVariable = ChildScriptVariable->Variable;
+					//
+					// 			// we generally don't show inline edit toggles as the managed items will display an inline checkbox instead
+					// 			if(ChildScriptVariable->Metadata.bInlineEditConditionToggle)
+					// 			{
+					// 				continue;
+					// 			}
+					// 			
+					// 			UNiagaraStackFunctionInput* InputChild = FindCurrentChildOfTypeByPredicate<UNiagaraStackFunctionInput>(CurrentChildren, [&](UNiagaraStackFunctionInput* CurrentInput) 
+					// 			{
+					// 				return CurrentInput->GetInputParameterHandle() == ChildScriptVariable->Variable.GetName() && CurrentInput->GetInputType() == ChildScriptVariable->Variable.GetType() && CurrentInput->GetInputFunctionCallNode().NodeGuid == ModuleBaseInputData->FunctionCallNode->NodeGuid;
+					// 			});
+					//
+					// 			TOptional<bool> bIsStatic = OwningFunctionCallNode->GetCalledGraph()->IsStaticSwitch(ChildVariable);
+					// 			if(InputChild == nullptr)
+					// 			{
+					// 				InputChild = NewObject<UNiagaraStackFunctionInput>(this);
+					// 				InputChild->Initialize(CreateDefaultChildRequiredData(), *ModuleBaseInputData->FunctionCallNode, *ModuleBaseInputData->FunctionCallNode, ChildScriptVariable->Variable.GetName(), ChildScriptVariable->Variable.GetType(), bIsStatic.GetValue() ? EStackParameterBehavior::Static : EStackParameterBehavior::Dynamic, GetOwnerStackItemEditorDataKey());						
+					// 				InputChild->SetSemanticChild(true);
+					// 			}
+					//
+					// 			// we update the hidden flag every time
+					// 			bool bIsChildHidden = FunctionCallToHiddenVariablesMap[ModuleInput->GetPersistentIdentity().Guids[0]].Contains(ChildVariable);
+					// 			InputChild->SetIsHidden(bIsTopLevelHidden || bIsChildHidden);
+					// 			NewChildren.Add(InputChild);								
+					// 		}
+					// 	}
+					// }					
+				}
+			}
+		}
+		else if(UNiagaraHierarchyAssignmentInput* AssignmentInput = Cast<UNiagaraHierarchyAssignmentInput>(Data))
+		{
+			TSharedPtr<FNiagaraAssignmentInputViewModel> AssignmentInputViewModel = StaticCastSharedPtr<FNiagaraAssignmentInputViewModel>(HierarchyViewModel);
+			TOptional<FMatchingFunctionInputData> InputData = AssignmentInputViewModel->GetInputData();
+			if(InputData.IsSet())
+			{
+				FString VariableNameWithModulePrefix = FNiagaraConstants::ModuleNamespaceString + ".";
+				VariableNameWithModulePrefix.Append(InputData->InputName.ToString());
+				
+				UNiagaraStackFunctionInput* TopLevelInput = FindCurrentChildOfTypeByPredicate<UNiagaraStackFunctionInput>(CurrentChildren, [&](UNiagaraStackFunctionInput* CurrentInput) 
+				{
+					return CurrentInput->GetInputParameterHandle().GetName().ToString() == VariableNameWithModulePrefix && CurrentInput->GetInputType() == InputData->Type && &CurrentInput->GetInputFunctionCallNode() == InputData->FunctionCallNode;
+				});
+	
+				if(TopLevelInput == nullptr)
+				{
+					TopLevelInput = NewObject<UNiagaraStackFunctionInput>(this);
+					TopLevelInput->Initialize(CreateDefaultChildRequiredData(), *InputData->FunctionCallNode, *InputData->FunctionCallNode, FName(VariableNameWithModulePrefix), InputData->Type, InputData->bIsStatic ? EStackParameterBehavior::Static : EStackParameterBehavior::Dynamic, GetOwnerStackItemEditorDataKey());
+					TAttribute<FText> SummaryViewTooltipOverride;
+					SummaryViewTooltipOverride.BindUObject(AssignmentInput, &UNiagaraHierarchyAssignmentInput::GetTooltipOverride);
+					TopLevelInput->SetSummaryViewTooltip(SummaryViewTooltipOverride);
+				}
+				
+				TopLevelInput->SetIsHidden(InputData->bIsHidden);
+				
+				NewChildren.Add(TopLevelInput);
+			}
+		}
+		else if(UNiagaraHierarchyModule* SummaryModule = Cast<UNiagaraHierarchyModule>(Data))
+		{
+			if(UNiagaraNodeFunctionCall* MatchingFunctionCall = FindFunctionCallNode(SummaryModule->GetPersistentIdentity().Guids[0], GetEmitterViewModel().ToSharedRef()))
+			{
+				UNiagaraStackModuleItem* Module = FindCurrentChildOfTypeByPredicate<UNiagaraStackModuleItem>(CurrentChildren,
+				[&](UNiagaraStackModuleItem* CurrentModule) { return &CurrentModule->GetModuleNode() == MatchingFunctionCall; });
+		
+				if (Module == nullptr)
+				{
+					Module = NewObject<UNiagaraStackModuleItem>(this);
+					Module->Initialize(CreateDefaultChildRequiredData(), nullptr, *MatchingFunctionCall);
+				}
+			
+				NewChildren.Add(Module);	
+			}		
+		}
+		else if(UNiagaraHierarchyEventHandler* EventHandler = Cast<UNiagaraHierarchyEventHandler>(Data))
+		{
+			const TArray<FNiagaraEventScriptProperties>& EventScriptProperties = GetEmitterViewModel()->GetEmitter().GetEmitterData()->GetEventHandlers();
+			const FNiagaraEventScriptProperties* EventScriptPropertiesItem = EventScriptProperties.FindByPredicate([EventHandlerIdentity = EventHandler->GetPersistentIdentity()](const FNiagaraEventScriptProperties& Candidate)
+			{
+				return Candidate.Script->GetUsageId() == EventHandlerIdentity.Guids[0];
+			});
+
+			if(EventScriptPropertiesItem != nullptr)
+			{
+				UNiagaraStackEventScriptItemGroup* StackEventGroup = FindCurrentChildOfTypeByPredicate<UNiagaraStackEventScriptItemGroup>(CurrentChildren,
+				[&](UNiagaraStackEventScriptItemGroup* CurrentEventProperties) { return CurrentEventProperties->GetScriptUsageId() == (*EventScriptPropertiesItem).Script->GetUsageId() && CurrentEventProperties->GetEventSourceEmitterId() == (*EventScriptPropertiesItem).SourceEmitterID; });
+		
+				if (StackEventGroup == nullptr)
+				{
+					StackEventGroup = NewObject<UNiagaraStackEventScriptItemGroup>(this);
+					StackEventGroup->Initialize(CreateDefaultChildRequiredData(), GetEmitterViewModel()->GetSharedScriptViewModel(), ENiagaraScriptUsage::ParticleEventScript, (*EventScriptPropertiesItem).Script->GetUsageId(), (*EventScriptPropertiesItem).SourceEmitterID);
+					StackEventGroup->SetOnModifiedEventHandlers(UNiagaraStackEventScriptItemGroup::FOnModifiedEventHandlers::CreateUObject(this, &UNiagaraStackEntry::RefreshChildren));
+
+				}
+			
+				NewChildren.Add(StackEventGroup);	
+			}
+		}
+		else if(UNiagaraHierarchyEventHandlerProperties* EventHandlerProperties = Cast<UNiagaraHierarchyEventHandlerProperties>(Data))
+		{
+			const TArray<FNiagaraEventScriptProperties>& EventScriptProperties = GetEmitterViewModel()->GetEmitter().GetEmitterData()->GetEventHandlers();
+			const FNiagaraEventScriptProperties* EventScriptPropertiesItem = EventScriptProperties.FindByPredicate([EventHandlerIdentity = EventHandlerProperties->GetPersistentIdentity()](const FNiagaraEventScriptProperties& Candidate)
+			{
+				FNiagaraHierarchyIdentity CandidateIdentity = UNiagaraHierarchyEventHandlerProperties::MakeIdentity(Candidate);
+				return CandidateIdentity == EventHandlerIdentity;
+			});
+
+			if(EventScriptPropertiesItem != nullptr)
+			{
+				UNiagaraStackEventHandlerPropertiesItem* StackEventProperties = FindCurrentChildOfTypeByPredicate<UNiagaraStackEventHandlerPropertiesItem>(CurrentChildren,
+				[&](UNiagaraStackEventHandlerPropertiesItem* CurrentEventProperties) { return CurrentEventProperties->GetEventScriptUsageId() == EventScriptPropertiesItem->Script->GetUsageId(); });
+		
+				if (StackEventProperties == nullptr)
+				{
+					StackEventProperties = NewObject<UNiagaraStackEventHandlerPropertiesItem>(this);
+					StackEventProperties->Initialize(CreateDefaultChildRequiredData(), (*EventScriptPropertiesItem).Script->GetUsageId());
+				}
+			
+				NewChildren.Add(StackEventProperties);	
+			}
+		}
+		else if(UNiagaraHierarchyRenderer* SummaryRenderer = Cast<UNiagaraHierarchyRenderer>(Data))
+		{
+			TArray<UNiagaraRendererProperties*> RendererProperties = GetEmitterViewModel()->GetEmitter().GetEmitterData()->GetRenderers();
+			UNiagaraRendererProperties** MatchingRendererProperties = RendererProperties.FindByPredicate([RendererIdentity = SummaryRenderer->GetPersistentIdentity().Guids[0]](UNiagaraRendererProperties* Candidate)
+			{
+				return Candidate->GetMergeId() == RendererIdentity;
+			});
+
+			if(MatchingRendererProperties != nullptr)
+			{
+				UNiagaraStackRendererItem* StackRenderer = FindCurrentChildOfTypeByPredicate<UNiagaraStackRendererItem>(CurrentChildren,
+				[&](UNiagaraStackRendererItem* CurrentRenderer) { return CurrentRenderer->GetRendererProperties()->GetMergeId() == (*MatchingRendererProperties)->GetMergeId(); });
+		
+				if (StackRenderer == nullptr)
+				{
+					StackRenderer = NewObject<UNiagaraStackRendererItem>(this);
+					StackRenderer->Initialize(CreateDefaultChildRequiredData(), *MatchingRendererProperties);
+				}
+			
+				NewChildren.Add(StackRenderer);	
+			}
+		}
+		else if(UNiagaraHierarchySimStage* SummarySimStage = Cast<UNiagaraHierarchySimStage>(Data))
+		{
+			TArray<UNiagaraSimulationStageBase*> SimStages = GetEmitterViewModel()->GetEmitter().GetEmitterData()->GetSimulationStages();
+			UNiagaraSimulationStageBase** MatchingSimStage = SimStages.FindByPredicate([RendererIdentity = SummarySimStage->GetPersistentIdentity().Guids[0]](UNiagaraSimulationStageBase* Candidate)
+			{
+				return Candidate->GetMergeId() == RendererIdentity;
+			});
+
+			if(MatchingSimStage != nullptr)
+			{
+				UNiagaraStackSimulationStageGroup* SimStageGroup = FindCurrentChildOfTypeByPredicate<UNiagaraStackSimulationStageGroup>(CurrentChildren,
+				[&](UNiagaraStackSimulationStageGroup* CurrentSimStageItem)
+				{
+					if(CurrentSimStageItem->GetSimulationStage())
+					{
+						return CurrentSimStageItem->GetSimulationStage()->GetMergeId() == (*MatchingSimStage)->GetMergeId();
+					}
+
+					return false;
+				});
+		
+				if (SimStageGroup == nullptr)
+				{
+					SimStageGroup = NewObject<UNiagaraStackSimulationStageGroup>(this);
+					FRequiredEntryData RequiredEntryData(GetSystemViewModel(), GetEmitterViewModel(),
+						FExecutionCategoryNames::Particle, FExecutionSubcategoryNames::SimulationStage,
+						GetEmitterViewModel()->GetEditorData().GetStackEditorData());
+					SimStageGroup->Initialize(RequiredEntryData, GetEmitterViewModel()->GetSharedScriptViewModel(), *MatchingSimStage);
+					SimStageGroup->SetOnModifiedSimulationStages(UNiagaraStackSimulationStageGroup::FOnModifiedSimulationStages::CreateUObject(this, &UNiagaraStackEntry::RefreshChildren));
+				}
+			
+				NewChildren.Add(SimStageGroup);	
+			}
+		}
+		else if(UNiagaraHierarchySimStageProperties* SummarySimStageProperties = Cast<UNiagaraHierarchySimStageProperties>(Data))
+		{
+			TArray<UNiagaraSimulationStageBase*> SimStages = GetEmitterViewModel()->GetEmitter().GetEmitterData()->GetSimulationStages();
+			UNiagaraSimulationStageBase** MatchingSimStage = SimStages.FindByPredicate([RendererIdentity = SummarySimStageProperties->GetPersistentIdentity().Guids[0]](UNiagaraSimulationStageBase* Candidate)
+			{
+				return Candidate->GetMergeId() == RendererIdentity;
+			});
+
+			if(MatchingSimStage != nullptr)
+			{
+				UNiagaraStackSimulationStagePropertiesItem* SimStageProperties = FindCurrentChildOfTypeByPredicate<UNiagaraStackSimulationStagePropertiesItem>(CurrentChildren,
+				[&](UNiagaraStackSimulationStagePropertiesItem* CurrentSimStageItem)
+				{
+					if(CurrentSimStageItem->GetSimulationStage().IsValid())
+					{
+						return CurrentSimStageItem->GetSimulationStage()->GetMergeId() == (*MatchingSimStage)->GetMergeId();
+					}
+
+					return false;
+				});
+		
+				if (SimStageProperties == nullptr)
+				{
+					SimStageProperties = NewObject<UNiagaraStackSimulationStagePropertiesItem>(this);
+					SimStageProperties->Initialize(CreateDefaultChildRequiredData(), *MatchingSimStage);
+				}
+			
+				NewChildren.Add(SimStageProperties);	
+			}
+		}
+		else if(UNiagaraHierarchyCategory* HierarchyCategory = Cast<UNiagaraHierarchyCategory>(Data))
+		{			
+			TSharedPtr<FNiagaraHierarchyCategoryViewModel> CategoryViewModel = StaticCastSharedPtr<FNiagaraHierarchyCategoryViewModel>(HierarchyViewModel);
+			
+			UNiagaraStackSummaryCategory* StackCategory = FindCurrentChildOfTypeByPredicate<UNiagaraStackSummaryCategory>(CurrentChildren,
+			[&](UNiagaraStackSummaryCategory* StackCategoryCandidate)
+			{
+				if(StackCategoryCandidate->GetHierarchyCategory().IsValid())
+				{
+					return StackCategoryCandidate->GetHierarchyCategory().Pin()->GetData() == CategoryViewModel->GetData();
+				}
+
+				return false;
+			});
+	
+			if (StackCategory == nullptr)
+			{
+				StackCategory = NewObject<UNiagaraStackSummaryCategory>(this);
+				StackCategory->Initialize(CreateDefaultChildRequiredData(), CategoryViewModel, GetOwnerStackItemEditorDataKey());
+			}
+		
+			NewChildren.Add(StackCategory);	
+		}
+		else if(const UNiagaraHierarchyObjectProperty* ObjectProperty = Cast<UNiagaraHierarchyObjectProperty>(Data))
+		{
+			TMap<FGuid, UObject*> ObjectsForProperties = GetEmitterViewModel()->GetSummaryHierarchyViewModel()->GetObjectsForProperties();
+			
+			FGuid OwnerGuid = ObjectProperty->GetPersistentIdentity().Guids.Num() > 0 ? ObjectProperty->GetPersistentIdentity().Guids[0] : FGuid();
+			
+			if(ObjectsForProperties.Contains(OwnerGuid))
+			{
+				UObject* Object = ObjectsForProperties[ObjectProperty->GetPersistentIdentity().Guids[0]];
+				TArray<TSharedRef<IDetailTreeNode>> RootNodes = GetEmitterViewModel()->GetSummaryHierarchyViewModel()->RequestDetailTreeNodesForObject(Object);
+				
+				for(TSharedRef<IDetailTreeNode>& RootNode : RootNodes)
+				{
+					TArray<TSharedRef<IDetailTreeNode>> ChildrenNodes;
+					RootNode->GetChildren(ChildrenNodes);
+
+					for(TSharedRef<IDetailTreeNode> ChildNode : ChildrenNodes)
+					{
+						if(ChildNode->GetNodeName() == ObjectProperty->GetPersistentIdentity().Names[0])
+						{
+							UNiagaraStackPropertyRow* PropertyRow = NewObject<UNiagaraStackPropertyRow>(this);
+							PropertyRow->Initialize(CreateDefaultChildRequiredData(), ChildNode, false, GetOwnerStackItemEditorDataKey(), GetOwnerStackItemEditorDataKey(), nullptr);
+							PropertyRow->SetOwnerGuid(ObjectProperty->GetPersistentIdentity().Guids[0]);
+							NewChildren.Add(PropertyRow);
+						}
+					}
+				}
+			}
+		}
+	}
+
+	Super::RefreshChildrenInternal(CurrentChildren, NewChildren, NewIssues);
+}
+
+bool UNiagaraStackSummaryCategory::IsTopLevelCategory() const
+{
+	return CategoryViewModelWeakPtr.Pin()->GetData()->GetOuter()->IsA<UNiagaraHierarchyRoot>();
+}
+
+FText UNiagaraStackSummaryCategory::GetTooltipText() const
+{
+	return Cast<UNiagaraHierarchyCategory>(CategoryViewModelWeakPtr.Pin()->GetData())->GetTooltip();
+}
+
+int32 UNiagaraStackSummaryCategory::GetChildIndentLevel() const
+{
+	return UNiagaraStackEntry::GetChildIndentLevel();
+}
diff --git a/Engine/Plugins/FX/Niagara/Source/NiagaraEditor/Private/ViewModels/Stack/NiagaraStackModuleItem.cpp b/Engine/Plugins/FX/Niagara/Source/NiagaraEditor/Private/ViewModels/Stack/NiagaraStackModuleItem.cpp
--- a/Engine/Plugins/FX/Niagara/Source/NiagaraEditor/Private/ViewModels/Stack/NiagaraStackModuleItem.cpp
+++ b/Engine/Plugins/FX/Niagara/Source/NiagaraEditor/Private/ViewModels/Stack/NiagaraStackModuleItem.cpp
@@ -1657,6 +1657,13 @@ FText UNiagaraStackModuleItem::GetInheritanceMessage() const
 	return LOCTEXT("ModuleItemInheritanceMessage", "This module is inherited from a parent emitter.  Inherited modules\ncan only be moved, deleted, and versioned while editing the parent emitter.");
 }
 
+FNiagaraHierarchyIdentity UNiagaraStackModuleItem::DetermineSummaryIdentity() const
+{
+	FNiagaraHierarchyIdentity Identity;
+	Identity.Guids.Add(GetModuleNode().NodeGuid);
+	return Identity;
+}
+
 bool UNiagaraStackModuleItem::IsScratchModule() const
 {
 	if (bIsScratchModuleCache.IsSet() == false)
