Summary View 2.0

diff --git a/Engine/Plugins/FX/Niagara/Source/NiagaraEditor/Private/ViewModels/Stack/NiagaraStackObject.cpp b/Engine/Plugins/FX/Niagara/Source/NiagaraEditor/Private/ViewModels/Stack/NiagaraStackObject.cpp
--- a/Engine/Plugins/FX/Niagara/Source/NiagaraEditor/Private/ViewModels/Stack/NiagaraStackObject.cpp
+++ b/Engine/Plugins/FX/Niagara/Source/NiagaraEditor/Private/ViewModels/Stack/NiagaraStackObject.cpp
@@ -15,6 +15,7 @@
 #include "IDetailTreeNode.h"
 
 #include "NiagaraPlatformSet.h"
+#include "ViewModels/HierarchyEditor/NiagaraHierarchyViewModelBase.h"
 #include "ViewModels/Stack/NiagaraStackObjectIssueGenerator.h"
 
 #include UE_INLINE_GENERATED_CPP_BY_NAME(NiagaraStackObject)
@@ -87,6 +88,13 @@ UObject* UNiagaraStackObject::GetDisplayedObject() const
 	return WeakObject.Get();
 }
 
+FNiagaraHierarchyIdentity UNiagaraStackObject::DetermineSummaryIdentity() const
+{
+	FNiagaraHierarchyIdentity Identity;
+	Identity.Guids.Add(ObjectGuid.GetValue());
+	return Identity;
+}
+
 void UNiagaraStackObject::FinalizeInternal()
 {
 	if (PropertyRowGenerator.IsValid())
@@ -296,6 +304,7 @@ void UNiagaraStackObject::RefreshChildrenInternal(const TArray<UNiagaraStackEntr
 			{
 				ChildRow = NewObject<UNiagaraStackPropertyRow>(this);
 				ChildRow->Initialize(CreateDefaultChildRequiredData(), RootTreeNode, bIsTopLevelObject, GetOwnerStackItemEditorDataKey(), GetOwnerStackItemEditorDataKey(), OwningNiagaraNode);
+				ChildRow->SetOwnerGuid(ObjectGuid);
 			}
 
 			NewChildren.Add(ChildRow);
diff --git a/Engine/Plugins/FX/Niagara/Source/NiagaraEditor/Private/ViewModels/Stack/NiagaraStackPropertyRow.cpp b/Engine/Plugins/FX/Niagara/Source/NiagaraEditor/Private/ViewModels/Stack/NiagaraStackPropertyRow.cpp
--- a/Engine/Plugins/FX/Niagara/Source/NiagaraEditor/Private/ViewModels/Stack/NiagaraStackPropertyRow.cpp
+++ b/Engine/Plugins/FX/Niagara/Source/NiagaraEditor/Private/ViewModels/Stack/NiagaraStackPropertyRow.cpp
@@ -7,6 +7,7 @@
 #include "PropertyHandle.h"
 #include "IDetailPropertyRow.h"
 #include "ScopedTransaction.h"
+#include "ViewModels/HierarchyEditor/NiagaraHierarchyViewModelBase.h"
 
 #include UE_INLINE_GENERATED_CPP_BY_NAME(NiagaraStackPropertyRow)
 
@@ -104,6 +105,7 @@ void UNiagaraStackPropertyRow::RefreshChildrenInternal(const TArray<UNiagaraStac
 			bool bChildIsTopLevelProperty = false;
 			ChildRow = NewObject<UNiagaraStackPropertyRow>(this);
 			ChildRow->Initialize(CreateDefaultChildRequiredData(), NodeChild, bChildIsTopLevelProperty, GetOwnerStackItemEditorDataKey(), GetStackEditorDataKey(), OwningNiagaraNode);
+			ChildRow->SetOwnerGuid(OwnerGuid);
 		}
 
 		NewChildren.Add(ChildRow);
@@ -228,3 +230,15 @@ TOptional<UNiagaraStackEntry::FDropRequestResponse> UNiagaraStackPropertyRow::Dr
 	return CanDropResponse;
 }
 
+bool UNiagaraStackPropertyRow::SupportsSummaryView() const
+{
+	return OwnerGuid.IsSet() && OwnerGuid->IsValid() && DetailTreeNode->GetNodeType() == EDetailNodeType::Item;
+}
+
+FNiagaraHierarchyIdentity UNiagaraStackPropertyRow::DetermineSummaryIdentity() const
+{
+	FNiagaraHierarchyIdentity Identity;
+	Identity.Guids.Add(OwnerGuid.GetValue());
+	Identity.Names.Add(DetailTreeNode->GetNodeName());
+	return Identity;
+}
diff --git a/Engine/Plugins/FX/Niagara/Source/NiagaraEditor/Private/ViewModels/Stack/NiagaraStackRendererItem.cpp b/Engine/Plugins/FX/Niagara/Source/NiagaraEditor/Private/ViewModels/Stack/NiagaraStackRendererItem.cpp
--- a/Engine/Plugins/FX/Niagara/Source/NiagaraEditor/Private/ViewModels/Stack/NiagaraStackRendererItem.cpp
+++ b/Engine/Plugins/FX/Niagara/Source/NiagaraEditor/Private/ViewModels/Stack/NiagaraStackRendererItem.cpp
@@ -254,6 +254,11 @@ UNiagaraRendererProperties* UNiagaraStackRendererItem::GetRendererProperties()
 	return RendererProperties.Get();
 }
 
+const UNiagaraRendererProperties* UNiagaraStackRendererItem::GetRendererProperties() const
+{
+	return RendererProperties.Get();
+}
+
 FText UNiagaraStackRendererItem::GetDisplayName() const
 {
 	if(DisplayNameCache.IsSet() == false)
@@ -388,6 +393,13 @@ FText UNiagaraStackRendererItem::GetInheritanceMessage() const
 	return LOCTEXT("RendererItemInheritanceMessage", "This renderer is inherited from a parent emitter.  Inherited\nrenderers can only be deleted while editing the parent emitter.");
 }
 
+FNiagaraHierarchyIdentity UNiagaraStackRendererItem::DetermineSummaryIdentity() const
+{
+	FNiagaraHierarchyIdentity Identity;
+	Identity.Guids.Add(GetRendererProperties()->GetMergeId());
+	return Identity;
+}
+
 bool UNiagaraStackRendererItem::HasBaseRenderer() const
 {
 	if (HasBaseEmitter())
@@ -489,6 +501,7 @@ void UNiagaraStackRendererItem::RefreshChildrenInternal(const TArray<UNiagaraSta
 		RendererObject = NewObject<UNiagaraStackObject>(this);
 		bool bIsTopLevelObject = true;
 		RendererObject->Initialize(CreateDefaultChildRequiredData(), RendererProperties.Get(), bIsTopLevelObject, GetStackEditorDataKey());
+		RendererObject->SetObjectGuid(GetRendererProperties()->GetMergeId());
 	}
 
 	NewChildren.Add(RendererObject);
diff --git a/Engine/Plugins/FX/Niagara/Source/NiagaraEditor/Private/ViewModels/Stack/NiagaraStackSimulationStageGroup.cpp b/Engine/Plugins/FX/Niagara/Source/NiagaraEditor/Private/ViewModels/Stack/NiagaraStackSimulationStageGroup.cpp
--- a/Engine/Plugins/FX/Niagara/Source/NiagaraEditor/Private/ViewModels/Stack/NiagaraStackSimulationStageGroup.cpp
+++ b/Engine/Plugins/FX/Niagara/Source/NiagaraEditor/Private/ViewModels/Stack/NiagaraStackSimulationStageGroup.cpp
@@ -19,6 +19,7 @@
 
 #include "Internationalization/Internationalization.h"
 #include "ScopedTransaction.h"
+#include "ViewModels/HierarchyEditor/NiagaraSummaryViewViewModel.h"
 
 #include UE_INLINE_GENERATED_CPP_BY_NAME(NiagaraStackSimulationStageGroup)
 
@@ -103,6 +104,7 @@ void UNiagaraStackSimulationStagePropertiesItem::RefreshChildrenInternal(const T
 		SimulationStageObject = NewObject<UNiagaraStackObject>(this);
 		bool bIsTopLevelObject = true;
 		SimulationStageObject->Initialize(CreateDefaultChildRequiredData(), SimulationStage.Get(), bIsTopLevelObject, GetStackEditorDataKey());
+		SimulationStageObject->SetObjectGuid(SimulationStage->GetMergeId());
 	}
 
 	if ( SimulationStage.IsValid() )
@@ -149,6 +151,11 @@ void UNiagaraStackSimulationStagePropertiesItem::RefreshChildrenInternal(const T
 	Super::RefreshChildrenInternal(CurrentChildren, NewChildren, NewIssues);
 }
 
+FNiagaraHierarchyIdentity UNiagaraStackSimulationStagePropertiesItem::DetermineSummaryIdentity() const
+{
+	return UNiagaraHierarchySimStageProperties::MakeIdentity(*GetSimulationStage().Get());
+}
+
 void UNiagaraStackSimulationStagePropertiesItem::SimulationStagePropertiesChanged()
 {
 	FVersionedNiagaraEmitter VersionedEmitter = GetEmitterViewModel()->GetEmitter();
@@ -417,6 +424,13 @@ FText UNiagaraStackSimulationStageGroup::GetInheritanceMessage() const
 	return LOCTEXT("SimulationStageGroupInheritanceMessage", "This simulation stage is inherited from a parent emitter.  Inherited\nsimulation stages can only be deleted while editing the parent emitter.");
 }
 
+FNiagaraHierarchyIdentity UNiagaraStackSimulationStageGroup::DetermineSummaryIdentity() const
+{
+	FNiagaraHierarchyIdentity Identity;
+	Identity.Guids.Add(GetSimulationStage()->GetMergeId());
+	return Identity;
+}
+
 void UNiagaraStackSimulationStageGroup::SimulationStagePropertiesChanged()
 {
 	SetDisplayName(FText::FromName(SimulationStage->SimulationStageName));
diff --git a/Engine/Plugins/FX/Niagara/Source/NiagaraEditor/Private/ViewModels/Stack/NiagaraStackSummaryViewInputCollection.cpp b/Engine/Plugins/FX/Niagara/Source/NiagaraEditor/Private/ViewModels/Stack/NiagaraStackSummaryViewInputCollection.cpp
--- a/Engine/Plugins/FX/Niagara/Source/NiagaraEditor/Private/ViewModels/Stack/NiagaraStackSummaryViewInputCollection.cpp
+++ b/Engine/Plugins/FX/Niagara/Source/NiagaraEditor/Private/ViewModels/Stack/NiagaraStackSummaryViewInputCollection.cpp
@@ -1,92 +1,475 @@
 // Copyright Epic Games, Inc. All Rights Reserved.
 
 #include "ViewModels/Stack/NiagaraStackSummaryViewInputCollection.h"
-#include "EdGraphSchema_Niagara.h"
+
+#include "IDetailTreeNode.h"
 #include "NiagaraClipboard.h"
-#include "NiagaraGraph.h"
+#include "NiagaraConstants.h"
 #include "NiagaraNodeFunctionCall.h"
 #include "NiagaraNodeOutput.h"
+#include "NiagaraRendererProperties.h"
 #include "NiagaraSimulationStageBase.h"
 #include "ViewModels/NiagaraEmitterViewModel.h"
 #include "ViewModels/NiagaraSystemViewModel.h"
 #include "ViewModels/NiagaraScriptViewModel.h"
-#include "ViewModels/NiagaraScriptGraphViewModel.h"
-#include "ViewModels/Stack/NiagaraStackFunctionInput.h"
+#include "ViewModels/HierarchyEditor/NiagaraSummaryViewViewModel.h"
+#include "ViewModels/Stack/NiagaraStackEventScriptItemGroup.h"
 #include "ViewModels/Stack/NiagaraStackGraphUtilities.h"
 #include "ViewModels/Stack/NiagaraStackFunctionInputCollection.h"
+#include "ViewModels/Stack/NiagaraStackInputCategory.h"
+#include "ViewModels/Stack/NiagaraStackModuleItem.h"
+#include "ViewModels/Stack/NiagaraStackPropertyRow.h"
+#include "ViewModels/Stack/NiagaraStackRendererItem.h"
 #include "ViewModels/Stack/NiagaraStackSimulationStageGroup.h"
 
 #include UE_INLINE_GENERATED_CPP_BY_NAME(NiagaraStackSummaryViewInputCollection)
 
-
 #define LOCTEXT_NAMESPACE "UNiagaraStackSummaryViewObject"
 
+using namespace FNiagaraStackGraphUtilities;
 
-UNiagaraStackSummaryViewObject::UNiagaraStackSummaryViewObject()
-{
-}
-
-void UNiagaraStackSummaryViewObject::Initialize(FRequiredEntryData InRequiredEntryData, FVersionedNiagaraEmitterWeakPtr InEmitter, FString InOwningStackItemEditorDataKey)
+void UNiagaraStackSummaryViewCollection::Initialize(FRequiredEntryData InRequiredEntryData,FVersionedNiagaraEmitterWeakPtr InEmitter, FString InOwningStackItemEditorDataKey)
 {
 	checkf(Emitter.Emitter == nullptr, TEXT("Can only initialize once."));
 	FString ObjectStackEditorDataKey = FString::Printf(TEXT("%s-FilteredView"), *InOwningStackItemEditorDataKey);
 	Super::Initialize(InRequiredEntryData, InOwningStackItemEditorDataKey, ObjectStackEditorDataKey);
 
 	Emitter = InEmitter;
-	GetEmitterViewModel()->GetOrCreateEditorData().OnSummaryViewStateChanged().AddUObject(this, &UNiagaraStackSummaryViewObject::OnViewStateChanged);
+	GetEmitterViewModel()->GetOrCreateEditorData().OnSummaryViewStateChanged().AddUObject(this, &UNiagaraStackSummaryViewCollection::OnViewStateChanged);
+	GetEmitterViewModel()->GetSummaryHierarchyViewModel()->OnHierarchyChanged().AddUObject(this, &UNiagaraStackSummaryViewCollection::OnViewStateChanged);
+	GetEmitterViewModel()->GetSummaryHierarchyViewModel()->OnHierarchyPropertiesChanged().AddUObject(this, &UNiagaraStackSummaryViewCollection::OnViewStateChanged);
 }
 
-void UNiagaraStackSummaryViewObject::FinalizeInternal()
+void UNiagaraStackSummaryViewCollection::FinalizeInternal()
 {
-	if (GetEmitterViewModel() && GetEmitterViewModel()->GetEmitter().GetEmitterData())
+	TSharedPtr<FNiagaraEmitterViewModel> EmitterViewModelPtr = GetEmitterViewModel();
+	if (EmitterViewModelPtr.IsValid() && GetEmitterViewModel()->GetEmitter().GetEmitterData())
 	{
 		GetEmitterViewModel()->GetOrCreateEditorData().OnSummaryViewStateChanged().RemoveAll(this);
 	}
+
+	if(EmitterViewModelPtr.IsValid() && EmitterViewModelPtr->GetSummaryHierarchyViewModel() != nullptr)
+	{
+		EmitterViewModelPtr->GetSummaryHierarchyViewModel()->OnHierarchyChanged().RemoveAll(this);
+		EmitterViewModelPtr->GetSummaryHierarchyViewModel()->OnHierarchyPropertiesChanged().RemoveAll(this);
+	}
+
 	Super::FinalizeInternal();
 }
 
-FText UNiagaraStackSummaryViewObject::GetDisplayName() const
+FText UNiagaraStackSummaryViewCollection::GetDisplayName() const
 {
 	return LOCTEXT("FilteredInputCollectionDisplayName", "Filtered Inputs");
 }
 
-bool UNiagaraStackSummaryViewObject::GetIsEnabled() const
+bool UNiagaraStackSummaryViewCollection::GetIsEnabled() const
 {
 	return true;
 }
 
-
-void UNiagaraStackSummaryViewObject::RefreshChildrenInternal(const TArray<UNiagaraStackEntry*>& CurrentChildren, TArray<UNiagaraStackEntry*>& NewChildren, TArray<FStackIssue>& NewIssues)
-{	
+void UNiagaraStackSummaryViewCollection::RefreshChildrenInternal(const TArray<UNiagaraStackEntry*>& CurrentChildren, TArray<UNiagaraStackEntry*>& NewChildren, TArray<FStackIssue>& NewIssues)
+{
 	TSharedPtr<FNiagaraEmitterViewModel> ViewModel = GetEmitterViewModel();
 	TSharedPtr<FNiagaraScriptViewModel> ScriptViewModelPinned = ViewModel->GetSharedScriptViewModel();
 	checkf(ScriptViewModelPinned.IsValid(), TEXT("Can not refresh children when the script view model has been deleted."));
 
-	FFunctionCallNodesState State;
+	const UNiagaraHierarchyRoot* Root = ViewModel->GetOrCreateEditorData().GetSummaryRoot();
+	TSharedPtr<FNiagaraHierarchyRootViewModel> RootViewModel = ViewModel->GetSummaryHierarchyViewModel()->GetHierarchyRootViewModel();
 	
-	AppendEmitterCategory(State, ScriptViewModelPinned, ENiagaraScriptUsage::EmitterSpawnScript, FGuid(), NewIssues);
-	AppendEmitterCategory(State, ScriptViewModelPinned, ENiagaraScriptUsage::EmitterUpdateScript, FGuid(), NewIssues);
+	FNiagaraStackGraphUtilities::FInputDataCollection State;
+	FNiagaraStackGraphUtilities::GatherInputRelationsForStack(State, GetEmitterViewModel().ToSharedRef());
 	
-	AppendEmitterCategory(State, ScriptViewModelPinned, ENiagaraScriptUsage::ParticleSpawnScript, FGuid(), NewIssues);
-	AppendEmitterCategory(State, ScriptViewModelPinned, ENiagaraScriptUsage::ParticleUpdateScript, FGuid(), NewIssues);
+	TArray<TSharedPtr<FNiagaraHierarchyItemViewModelBase>> ChildrenViewModels;
+	RootViewModel->GetChildrenViewModelsForType<UNiagaraHierarchyItemBase, FNiagaraHierarchyItemViewModelBase>(ChildrenViewModels, false);
 
-	for (const FNiagaraEventScriptProperties& EventScriptProperties : ViewModel->GetEmitter().GetEmitterData()->GetEventHandlers())
+	TArray<TSharedPtr<FNiagaraHierarchyItemViewModelBase>> AllChildrenViewModels;
+	RootViewModel->GetChildrenViewModelsForType<UNiagaraHierarchyItemBase, FNiagaraHierarchyItemViewModelBase>(AllChildrenViewModels, true);
+
+	// first we gather all function call nodes so we can create cache for them instead of looking it up for each input individually
+	TSet<UNiagaraNodeFunctionCall*> UsedFunctionCallNodes;
+	for(TSharedPtr<FNiagaraHierarchyItemViewModelBase> HierarchyViewModel : AllChildrenViewModels)
 	{
-		AppendEmitterCategory(State, ScriptViewModelPinned, ENiagaraScriptUsage::ParticleEventScript, EventScriptProperties.Script->GetUsageId(), NewIssues);
+		UNiagaraHierarchyItemBase* Data = HierarchyViewModel->GetDataMutable();
+		if(UNiagaraHierarchyModuleInput* ModuleInput = Cast<UNiagaraHierarchyModuleInput>(Data))
+		{
+			UsedFunctionCallNodes.Add(State.NodeGuidToModuleNodeMap[ModuleInput->GetPersistentIdentity().Guids[0]]);
+		}
+		else if(UNiagaraHierarchyAssignmentInput* AssignmentInput = Cast<UNiagaraHierarchyAssignmentInput>(Data))
+		{
+			UsedFunctionCallNodes.Add(State.NodeGuidToModuleNodeMap[AssignmentInput->GetPersistentIdentity().Guids[0]]);
+		}
+		else if(UNiagaraHierarchyModule* Module = Cast<UNiagaraHierarchyModule>(Data))
+		{
+			UsedFunctionCallNodes.Add(State.NodeGuidToModuleNodeMap[Module->GetPersistentIdentity().Guids[0]]);
+		}
 	}
 	
-	for (UNiagaraSimulationStageBase* SimulationStage : GetEmitterViewModel()->GetEmitter().GetEmitterData()->GetSimulationStages())
+	TMap<FGuid, TSet<FNiagaraVariable>> FunctionCallToHiddenVariablesMap;
+	for(UNiagaraNodeFunctionCall* FunctionCall : UsedFunctionCallNodes)
 	{
-		AppendEmitterCategory(State, ScriptViewModelPinned, ENiagaraScriptUsage::ParticleSimulationStageScript, SimulationStage->Script->GetUsageId(), NewIssues);
-	}	
+		TArray<FNiagaraVariable> InputVariables;
+		TSet<FNiagaraVariable> HiddenVariables;
+		FCompileConstantResolver Resolver(GetEmitterViewModel()->GetEmitter(), FNiagaraStackGraphUtilities::GetOutputNodeUsage(*FunctionCall));
+		FNiagaraStackGraphUtilities::GetStackFunctionInputs(*FunctionCall, InputVariables, HiddenVariables, Resolver, ENiagaraGetStackFunctionInputPinsOptions::ModuleInputsOnly, true);
+		FunctionCallToHiddenVariablesMap.Add(FunctionCall->NodeGuid, HiddenVariables);
+	}
 	
-	ApplyAllFunctionInputsToChildren(State, CurrentChildren, NewChildren, NewIssues, true);
+	for(TSharedPtr<FNiagaraHierarchyItemViewModelBase> HierarchyViewModel : ChildrenViewModels)
+	{
+		UNiagaraHierarchyItemBase* Data = HierarchyViewModel->GetDataMutable();
+
+		if(UNiagaraHierarchyModuleInput* ModuleInput = Cast<UNiagaraHierarchyModuleInput>(Data))
+		{
+			if(GetEmitterViewModel().IsValid())
+			{
+				UNiagaraNodeFunctionCall* OwningFunctionCallNode = State.NodeGuidToModuleNodeMap[ModuleInput->GetPersistentIdentity().Guids[0]];
+				
+				TSharedPtr<FNiagaraModuleInputViewModel> ModuleInputViewModel = StaticCastSharedPtr<FNiagaraModuleInputViewModel>(HierarchyViewModel);
+				TOptional<FInputData> ModuleBaseInputData = ModuleInputViewModel->GetInputData();
+				
+				if(ModuleBaseInputData.IsSet())
+				{
+					UNiagaraStackFunctionInput* TopLevelInput = FindCurrentChildOfTypeByPredicate<UNiagaraStackFunctionInput>(CurrentChildren, [&](UNiagaraStackFunctionInput* CurrentInput) 
+					{
+						return CurrentInput->GetInputParameterHandle() == ModuleBaseInputData->InputName && CurrentInput->GetInputType() == ModuleBaseInputData->Type
+							&& CurrentInput->GetInputFunctionCallInitialScript() == OwningFunctionCallNode->FunctionScript
+							// we are checking for node guid as we could have 2x the same input from 2 identical modules
+							&& CurrentInput->GetInputFunctionCallNode().NodeGuid == OwningFunctionCallNode->NodeGuid;
+					});
+	
+					if(TopLevelInput == nullptr)
+					{
+						TopLevelInput = NewObject<UNiagaraStackFunctionInput>(this);
+						TopLevelInput->Initialize(CreateDefaultChildRequiredData(), *OwningFunctionCallNode, *OwningFunctionCallNode, ModuleBaseInputData->InputName, ModuleBaseInputData->Type, ModuleBaseInputData->bIsStatic ? EStackParameterBehavior::Static : EStackParameterBehavior::Dynamic, GetOwnerStackItemEditorDataKey());
+						TAttribute<FText> SummaryViewDisplayNameOverride;
+						SummaryViewDisplayNameOverride.BindUObject(ModuleInput, &UNiagaraHierarchyModuleInput::GetDisplayNameOverride);
+						TopLevelInput->SetSummaryViewDisplayName(SummaryViewDisplayNameOverride);
+						TAttribute<FText> SummaryViewTooltipOverride;
+						SummaryViewTooltipOverride.BindUObject(ModuleInput, &UNiagaraHierarchyModuleInput::GetTooltipOverride);
+						TopLevelInput->SetSummaryViewTooltip(SummaryViewTooltipOverride);
+					}
+
+					FNiagaraVariable DisplayedVariable(TopLevelInput->GetInputType(), TopLevelInput->GetInputParameterHandle().GetParameterHandleString());
+					bool bIsTopLevelHidden = FunctionCallToHiddenVariablesMap[ModuleInput->GetPersistentIdentity().Guids[0]].Contains(DisplayedVariable);
+					TopLevelInput->SetIsHidden(bIsTopLevelHidden);
+					
+					NewChildren.Add(TopLevelInput);
+
+					TArray<UNiagaraHierarchyModuleInput*> ChildInputs;
+					ModuleInput->GetChildrenOfType<UNiagaraHierarchyModuleInput>(ChildInputs);
+
+					for(UNiagaraHierarchyModuleInput* ChildInput : ChildInputs)
+					{
+						TSharedPtr<FNiagaraModuleInputViewModel> ChildInputViewModel = StaticCastSharedPtr<FNiagaraModuleInputViewModel>(ModuleInputViewModel->FindViewModelForChild(ChildInput, false));
+						TOptional<FInputData> ChildInputData = ChildInputViewModel->GetInputData();
+
+						if(ChildInputData.IsSet())
+						{
+							UNiagaraNodeFunctionCall* ChildFunctionNode = ChildInputData->FunctionCallNode;
+							UNiagaraStackFunctionInput* StackChildInput = FindCurrentChildOfTypeByPredicate<UNiagaraStackFunctionInput>(CurrentChildren, [&](UNiagaraStackFunctionInput* CurrentInput) 
+							{
+								return CurrentInput->GetInputParameterHandle() == ChildInputData->InputName && CurrentInput->GetInputType() == ChildInputData->Type
+									&& CurrentInput->GetInputFunctionCallInitialScript() == ChildFunctionNode->FunctionScript
+									// we are checking for node guid as we could have 2x the same input from 2 identical modules
+									&& CurrentInput->GetInputFunctionCallNode().NodeGuid == ChildFunctionNode->NodeGuid;
+							});
+							
+							if(StackChildInput == nullptr)
+							{
+								StackChildInput = NewObject<UNiagaraStackFunctionInput>(this);
+								StackChildInput->Initialize(CreateDefaultChildRequiredData(), *ChildInputData->FunctionCallNode, *ChildInputData->FunctionCallNode, ChildInputData->InputName, ChildInputData->Type, ChildInputData->bIsStatic ? EStackParameterBehavior::Static : EStackParameterBehavior::Dynamic, GetOwnerStackItemEditorDataKey());
+								TAttribute<FText> SummaryViewDisplayNameOverride;
+								SummaryViewDisplayNameOverride.BindUObject(ChildInput, &UNiagaraHierarchyModuleInput::GetDisplayNameOverride);
+								StackChildInput->SetSummaryViewDisplayName(SummaryViewDisplayNameOverride);
+								TAttribute<FText> SummaryViewTooltipOverride;
+								SummaryViewTooltipOverride.BindUObject(ChildInput, &UNiagaraHierarchyModuleInput::GetTooltipOverride);
+								StackChildInput->SetSummaryViewTooltip(SummaryViewTooltipOverride);
+								
+								StackChildInput->SetSemanticChild(true);
+							}
+
+							FNiagaraVariable DisplayedChildVariable(StackChildInput->GetInputType(), StackChildInput->GetInputParameterHandle().GetParameterHandleString());
+							bool bIsChildInputHidden = FunctionCallToHiddenVariablesMap[ChildInput->GetPersistentIdentity().Guids[0]].Contains(DisplayedChildVariable);
+							StackChildInput->SetIsHidden(bIsChildInputHidden);
+							
+							NewChildren.Add(StackChildInput);
+						}
+					}
+					
+					/** Add children inputs automatically. Disabled as we want children inputs to be managed by the user for summary view. */
+					// if(State.HierarchyInputToChildrenGuidMap.Contains(ModuleInput))
+					// {
+					// 	// children guids are already sorted by sort order
+					// 	TArray<FGuid> ChildrenInputs = State.HierarchyInputToChildrenGuidMap[ModuleInput];
+					//
+					// 	for(FGuid ChildrenInputGuid : ChildrenInputs)
+					// 	{
+					// 		if(UNiagaraScriptVariable* ChildScriptVariable = State.ChildrenGuidToScriptVariablesMap[ChildrenInputGuid])
+					// 		{
+					// 			FNiagaraVariable ChildVariable = ChildScriptVariable->Variable;
+					// 			
+					// 			// we generally don't show inline edit toggles as the managed items will display an inline checkbox instead
+					// 			if(ChildScriptVariable->Metadata.bInlineEditConditionToggle)
+					// 			{
+					// 				continue;
+					// 			}
+					// 			
+					// 			UNiagaraStackFunctionInput* InputChild = FindCurrentChildOfTypeByPredicate<UNiagaraStackFunctionInput>(CurrentChildren, [&](UNiagaraStackFunctionInput* CurrentInput) 
+					// 			{
+					// 				return CurrentInput->GetInputParameterHandle() == ChildScriptVariable->Variable.GetName() && CurrentInput->GetInputType() == ChildScriptVariable->Variable.GetType() && CurrentInput->GetInputFunctionCallNode().NodeGuid == ModuleBaseInputData->FunctionCallNode->NodeGuid;
+					// 			});
+					//
+					// 			TOptional<bool> bIsStatic = OwningFunctionCallNode->GetCalledGraph()->IsStaticSwitch(ChildVariable);
+					// 			if(InputChild == nullptr)
+					// 			{
+					// 				InputChild = NewObject<UNiagaraStackFunctionInput>(this);
+					// 				InputChild->Initialize(CreateDefaultChildRequiredData(), *ModuleBaseInputData->FunctionCallNode, *ModuleBaseInputData->FunctionCallNode, ChildScriptVariable->Variable.GetName(), ChildScriptVariable->Variable.GetType(), bIsStatic.GetValue() ? EStackParameterBehavior::Static : EStackParameterBehavior::Dynamic, GetOwnerStackItemEditorDataKey());						
+					// 				InputChild->SetSemanticChild(true);
+					// 			}
+					//
+					// 			// we update the hidden flag every time
+					// 			bool bIsChildHidden = FunctionCallToHiddenVariablesMap[ModuleInput->GetPersistentIdentity().Guids[0]].Contains(ChildVariable);
+					// 			InputChild->SetIsHidden(bIsTopLevelHidden || bIsChildHidden);
+					// 			NewChildren.Add(InputChild);								
+					// 		}
+					// 	}
+					// }					
+				}
+			}
+		}
+		else if(UNiagaraHierarchyAssignmentInput* AssignmentInput = Cast<UNiagaraHierarchyAssignmentInput>(Data))
+		{
+			TSharedPtr<FNiagaraAssignmentInputViewModel> AssignmentInputViewModel = StaticCastSharedPtr<FNiagaraAssignmentInputViewModel>(HierarchyViewModel);
+			TOptional<FMatchingFunctionInputData> InputData = AssignmentInputViewModel->GetInputData();
+			if(InputData.IsSet())
+			{
+				FString VariableNameWithModulePrefix = FNiagaraConstants::ModuleNamespaceString + ".";
+				VariableNameWithModulePrefix.Append(InputData->InputName.ToString());
+				
+				UNiagaraStackFunctionInput* TopLevelInput = FindCurrentChildOfTypeByPredicate<UNiagaraStackFunctionInput>(CurrentChildren, [&](UNiagaraStackFunctionInput* CurrentInput) 
+				{
+					return CurrentInput->GetInputParameterHandle().GetName().ToString() == VariableNameWithModulePrefix && CurrentInput->GetInputType() == InputData->Type && &CurrentInput->GetInputFunctionCallNode() == InputData->FunctionCallNode;
+				});
+	
+				if(TopLevelInput == nullptr)
+				{
+					TopLevelInput = NewObject<UNiagaraStackFunctionInput>(this);
+					TopLevelInput->Initialize(CreateDefaultChildRequiredData(), *InputData->FunctionCallNode, *InputData->FunctionCallNode, FName(VariableNameWithModulePrefix), InputData->Type, InputData->bIsStatic ? EStackParameterBehavior::Static : EStackParameterBehavior::Dynamic, GetOwnerStackItemEditorDataKey());
+					TAttribute<FText> SummaryViewTooltipOverride;
+					SummaryViewTooltipOverride.BindUObject(AssignmentInput, &UNiagaraHierarchyAssignmentInput::GetTooltipOverride);
+					TopLevelInput->SetSummaryViewTooltip(SummaryViewTooltipOverride);
+				}
+				
+				TopLevelInput->SetIsHidden(InputData->bIsHidden);
+				
+				NewChildren.Add(TopLevelInput);
+			}
+		}
+		else if(UNiagaraHierarchyModule* SummaryModule = Cast<UNiagaraHierarchyModule>(Data))
+		{
+			if(UNiagaraNodeFunctionCall* MatchingFunctionCall = FNiagaraStackGraphUtilities::FindFunctionCallNode(SummaryModule->GetPersistentIdentity().Guids[0], GetEmitterViewModel().ToSharedRef()))
+			{
+				UNiagaraStackModuleItem* Module = FindCurrentChildOfTypeByPredicate<UNiagaraStackModuleItem>(CurrentChildren,
+				[&](UNiagaraStackModuleItem* CurrentModule) { return &CurrentModule->GetModuleNode() == MatchingFunctionCall; });
+		
+				if (Module == nullptr)
+				{
+					Module = NewObject<UNiagaraStackModuleItem>(this);
+					Module->Initialize(CreateDefaultChildRequiredData(), nullptr, *MatchingFunctionCall);
+				}
+				
+				NewChildren.Add(Module);	
+			}		
+		}
+		else if(UNiagaraHierarchyRenderer* SummaryRenderer = Cast<UNiagaraHierarchyRenderer>(Data))
+		{
+			TArray<UNiagaraRendererProperties*> RendererProperties = GetEmitterViewModel()->GetEmitter().GetEmitterData()->GetRenderers();
+			UNiagaraRendererProperties** MatchingRendererProperties = RendererProperties.FindByPredicate([RendererIdentity = SummaryRenderer->GetPersistentIdentity().Guids[0]](UNiagaraRendererProperties* Candidate)
+			{
+				return Candidate->GetMergeId() == RendererIdentity;
+			});
+
+			if(MatchingRendererProperties != nullptr)
+			{
+				UNiagaraStackRendererItem* StackRenderer = FindCurrentChildOfTypeByPredicate<UNiagaraStackRendererItem>(CurrentChildren,
+				[&](UNiagaraStackRendererItem* CurrentRenderer) { return CurrentRenderer->GetRendererProperties()->GetMergeId() == (*MatchingRendererProperties)->GetMergeId(); });
+		
+				if (StackRenderer == nullptr)
+				{
+					StackRenderer = NewObject<UNiagaraStackRendererItem>(this);
+					StackRenderer->Initialize(CreateDefaultChildRequiredData(), *MatchingRendererProperties);
+				}
+			
+				NewChildren.Add(StackRenderer);	
+			}
+		}
+		else if(UNiagaraHierarchyEventHandler* EventHandler = Cast<UNiagaraHierarchyEventHandler>(Data))
+		{
+			const TArray<FNiagaraEventScriptProperties>& EventScriptProperties = GetEmitterViewModel()->GetEmitter().GetEmitterData()->GetEventHandlers();
+			const FNiagaraEventScriptProperties* EventScriptPropertiesItem = EventScriptProperties.FindByPredicate([EventHandlerIdentity = EventHandler->GetPersistentIdentity()](const FNiagaraEventScriptProperties& Candidate)
+			{
+				return Candidate.Script->GetUsageId() == EventHandlerIdentity.Guids[0];
+			});
+
+			if(EventScriptPropertiesItem != nullptr)
+			{
+				UNiagaraStackEventScriptItemGroup* StackEventGroup = FindCurrentChildOfTypeByPredicate<UNiagaraStackEventScriptItemGroup>(CurrentChildren,
+				[&](UNiagaraStackEventScriptItemGroup* CurrentEventProperties) { return CurrentEventProperties->GetScriptUsageId() == (*EventScriptPropertiesItem).Script->GetUsageId() && CurrentEventProperties->GetEventSourceEmitterId() == (*EventScriptPropertiesItem).SourceEmitterID; });
+		
+				if (StackEventGroup == nullptr)
+				{
+					StackEventGroup = NewObject<UNiagaraStackEventScriptItemGroup>(this);
+					StackEventGroup->Initialize(CreateDefaultChildRequiredData(), GetEmitterViewModel()->GetSharedScriptViewModel(), ENiagaraScriptUsage::ParticleEventScript, (*EventScriptPropertiesItem).Script->GetUsageId(), (*EventScriptPropertiesItem).SourceEmitterID);
+					StackEventGroup->SetOnModifiedEventHandlers(UNiagaraStackEventScriptItemGroup::FOnModifiedEventHandlers::CreateUObject(this, &UNiagaraStackEntry::RefreshChildren));
+				}
+			
+				NewChildren.Add(StackEventGroup);	
+			}
+		}
+		else if(UNiagaraHierarchyEventHandlerProperties* EventHandlerProperties = Cast<UNiagaraHierarchyEventHandlerProperties>(Data))
+		{
+			const TArray<FNiagaraEventScriptProperties>& EventScriptProperties = GetEmitterViewModel()->GetEmitter().GetEmitterData()->GetEventHandlers();
+			const FNiagaraEventScriptProperties* EventScriptPropertiesItem = EventScriptProperties.FindByPredicate([EventHandlerIdentity = EventHandlerProperties->GetPersistentIdentity()](const FNiagaraEventScriptProperties& Candidate)
+			{
+				FNiagaraHierarchyIdentity CandidateIdentity = UNiagaraHierarchyEventHandlerProperties::MakeIdentity(Candidate);
+				return CandidateIdentity == EventHandlerIdentity;
+			});
+
+			if(EventScriptPropertiesItem != nullptr)
+			{
+				UNiagaraStackEventHandlerPropertiesItem* StackEventProperties = FindCurrentChildOfTypeByPredicate<UNiagaraStackEventHandlerPropertiesItem>(CurrentChildren,
+				[&](UNiagaraStackEventHandlerPropertiesItem* CurrentEventProperties) { return CurrentEventProperties->GetEventScriptUsageId() == EventScriptPropertiesItem->Script->GetUsageId(); });
+		
+				if (StackEventProperties == nullptr)
+				{
+					StackEventProperties = NewObject<UNiagaraStackEventHandlerPropertiesItem>(this);
+					StackEventProperties->Initialize(CreateDefaultChildRequiredData(), (*EventScriptPropertiesItem).Script->GetUsageId());
+				}
+			
+				NewChildren.Add(StackEventProperties);	
+			}
+		}	
+		else if(UNiagaraHierarchySimStage* SummarySimStage = Cast<UNiagaraHierarchySimStage>(Data))
+		{
+			TArray<UNiagaraSimulationStageBase*> SimStages = GetEmitterViewModel()->GetEmitter().GetEmitterData()->GetSimulationStages();
+			UNiagaraSimulationStageBase** MatchingSimStage = SimStages.FindByPredicate([RendererIdentity = SummarySimStage->GetPersistentIdentity().Guids[0]](UNiagaraSimulationStageBase* Candidate)
+			{
+				return Candidate->GetMergeId() == RendererIdentity;
+			});
+
+			if(MatchingSimStage != nullptr)
+			{
+				UNiagaraStackSimulationStageGroup* SimStageGroup = FindCurrentChildOfTypeByPredicate<UNiagaraStackSimulationStageGroup>(CurrentChildren,
+				[&](UNiagaraStackSimulationStageGroup* CurrentSimStageItem)
+				{
+					if(CurrentSimStageItem->GetSimulationStage())
+					{
+						return CurrentSimStageItem->GetSimulationStage()->GetMergeId() == (*MatchingSimStage)->GetMergeId();
+					}
+
+					return false;
+				});
+		
+				if (SimStageGroup == nullptr)
+				{
+					SimStageGroup = NewObject<UNiagaraStackSimulationStageGroup>(this);
+					FRequiredEntryData RequiredEntryData(GetSystemViewModel(), GetEmitterViewModel(),
+						FExecutionCategoryNames::Particle, FExecutionSubcategoryNames::SimulationStage,
+						GetEmitterViewModel()->GetEditorData().GetStackEditorData());
+					SimStageGroup->Initialize(RequiredEntryData, GetEmitterViewModel()->GetSharedScriptViewModel(), *MatchingSimStage);
+					SimStageGroup->SetOnModifiedSimulationStages(UNiagaraStackSimulationStageGroup::FOnModifiedSimulationStages::CreateUObject(this, &UNiagaraStackEntry::RefreshChildren));
+				}
+			
+				NewChildren.Add(SimStageGroup);	
+			}
+		}
+		else if(UNiagaraHierarchySimStageProperties* SummarySimStageProperties = Cast<UNiagaraHierarchySimStageProperties>(Data))
+		{
+			TArray<UNiagaraSimulationStageBase*> SimStages = GetEmitterViewModel()->GetEmitter().GetEmitterData()->GetSimulationStages();
+			UNiagaraSimulationStageBase** MatchingSimStage = SimStages.FindByPredicate([RendererIdentity = SummarySimStageProperties->GetPersistentIdentity().Guids[0]](UNiagaraSimulationStageBase* Candidate)
+			{
+				return Candidate->GetMergeId() == RendererIdentity;
+			});
+
+			if(MatchingSimStage != nullptr)
+			{
+				UNiagaraStackSimulationStagePropertiesItem* SimStageProperties = FindCurrentChildOfTypeByPredicate<UNiagaraStackSimulationStagePropertiesItem>(CurrentChildren,
+				[&](UNiagaraStackSimulationStagePropertiesItem* CurrentSimStageItem)
+				{
+					if(CurrentSimStageItem->GetSimulationStage().IsValid())
+					{
+						return CurrentSimStageItem->GetSimulationStage()->GetMergeId() == (*MatchingSimStage)->GetMergeId();
+					}
+
+					return false;
+				});
+		
+				if (SimStageProperties == nullptr)
+				{
+					SimStageProperties = NewObject<UNiagaraStackSimulationStagePropertiesItem>(this);
+					SimStageProperties->Initialize(CreateDefaultChildRequiredData(), *MatchingSimStage);
+				}
+			
+				NewChildren.Add(SimStageProperties);	
+			}
+		}
+		else if(UNiagaraHierarchyCategory* HierarchyCategory = Cast<UNiagaraHierarchyCategory>(Data))
+		{			
+			TSharedPtr<FNiagaraHierarchyCategoryViewModel> CategoryViewModel = StaticCastSharedPtr<FNiagaraHierarchyCategoryViewModel>(HierarchyViewModel);
+			
+			UNiagaraStackSummaryCategory* StackCategory = FindCurrentChildOfTypeByPredicate<UNiagaraStackSummaryCategory>(CurrentChildren,
+			[&](UNiagaraStackSummaryCategory* StackCategoryCandidate)
+			{
+				if(StackCategoryCandidate->GetHierarchyCategory().IsValid())
+				{
+					return StackCategoryCandidate->GetHierarchyCategory().Pin()->GetData() == CategoryViewModel->GetData();
+				}
+
+				return false;
+			});
+	
+			if (StackCategory == nullptr)
+			{
+				StackCategory = NewObject<UNiagaraStackSummaryCategory>(this);
+				StackCategory->Initialize(CreateDefaultChildRequiredData(), CategoryViewModel, GetOwnerStackItemEditorDataKey());
+			}
+		
+			NewChildren.Add(StackCategory);	
+		}
+		else if(const UNiagaraHierarchyObjectProperty* ObjectProperty = Cast<UNiagaraHierarchyObjectProperty>(Data))
+		{
+			TMap<FGuid, UObject*> ObjectsForProperties = GetEmitterViewModel()->GetSummaryHierarchyViewModel()->GetObjectsForProperties();
+			FGuid ObjectGuid = ObjectProperty->GetPersistentIdentity().Guids.Num() > 0 ? ObjectProperty->GetPersistentIdentity().Guids[0] : FGuid();
+			
+			if(ObjectsForProperties.Contains(ObjectGuid))
+			{
+				UObject* Object = ObjectsForProperties[ObjectProperty->GetPersistentIdentity().Guids[0]];
+				TArray<TSharedRef<IDetailTreeNode>> RootNodes = GetEmitterViewModel()->GetSummaryHierarchyViewModel()->RequestDetailTreeNodesForObject(Object);
+				
+				for(TSharedRef<IDetailTreeNode>& RootNode : RootNodes)
+				{
+					TArray<TSharedRef<IDetailTreeNode>> ChildrenNodes;
+					RootNode->GetChildren(ChildrenNodes);
+
+					for(TSharedRef<IDetailTreeNode> ChildNode : ChildrenNodes)
+					{
+						if(ChildNode->GetNodeName() == ObjectProperty->GetPersistentIdentity().Names[0])
+						{
+							UNiagaraStackPropertyRow* PropertyRow = NewObject<UNiagaraStackPropertyRow>(this);
+							PropertyRow->Initialize(CreateDefaultChildRequiredData(), ChildNode, false, GetOwnerStackItemEditorDataKey(), GetOwnerStackItemEditorDataKey(), nullptr);
+							NewChildren.Add(PropertyRow);
+						}
+					}
+				}
+			}
+		}
+	}
 	
 	Super::RefreshChildrenInternal(CurrentChildren, NewChildren, NewIssues);
 
 	if (NewChildren.Num() == 0)
 	{
-		FText EmptyAssignmentNodeMessageText = LOCTEXT("EmptySummaryNodeMessage", "No Parameters in Emitter Summary.\n\nTo add a parameter use the context menu for an available parameter in the stack. Once in the summary view you can use the drop down menu from the icon next to a parameter to set the DisplayName, SortIndex, and Category.");
+		FText EmptyAssignmentNodeMessageText = LOCTEXT("EmptySummaryNodeMessage", "No Parameters in Emitter Summary.\n\nYou can click on the 'Edit' button in the top right to summon the Summary View Editor. It lets you add categories, modules or individual module inputs as well as renderer properties via drag & drop.");
 		UNiagaraStackItemTextContent* EmtpySummaryMessage = FindCurrentChildOfTypeByPredicate<UNiagaraStackItemTextContent>(CurrentChildren,
 			[&](UNiagaraStackItemTextContent* CurrentStackItemTextContent) { return CurrentStackItemTextContent->GetDisplayName().IdenticalTo(EmptyAssignmentNodeMessageText); });
 		
@@ -99,48 +482,73 @@ void UNiagaraStackSummaryViewObject::RefreshChildrenInternal(const TArray<UNiaga
 	}
 }
 
-void UNiagaraStackSummaryViewObject::AppendEmitterCategory(FFunctionCallNodesState& State, TSharedPtr<FNiagaraScriptViewModel> ScriptViewModelPinned, ENiagaraScriptUsage ScriptUsage, FGuid ScriptUsageId, TArray<FStackIssue>& NewIssues)
+void UNiagaraStackSummaryViewCollection::GetSectionsInternal(TArray<FNiagaraStackSection>& OutStackSections) const
 {
-	UNiagaraGraph* Graph = ScriptViewModelPinned->GetGraphViewModel()->GetGraph();
-	FText ErrorMessage;
-	if (FNiagaraStackGraphUtilities::ValidateGraphForOutput(*Graph, ScriptUsage, ScriptUsageId, ErrorMessage) == true)
-	{
-		UNiagaraNodeOutput* MatchingOutputNode = Graph->FindEquivalentOutputNode(ScriptUsage, ScriptUsageId);
+	const TArray<UNiagaraHierarchySection*>& HierarchySections = GetHierarchySections();
 
-		TArray<UNiagaraNodeFunctionCall*> ModuleNodes;
-		FNiagaraStackGraphUtilities::GetOrderedModuleNodes(*MatchingOutputNode, ModuleNodes);
+	TMap<const UNiagaraHierarchySection*, TArray<UNiagaraHierarchyCategory*>> SectionToCategoryMap;
 
-		bool bIsFirst = true;
-		for (UNiagaraNodeFunctionCall* ModuleNode : ModuleNodes)
+	TSharedPtr<FNiagaraEmitterViewModel> ViewModel = GetEmitterViewModel();
+	const UNiagaraHierarchyRoot* Root = ViewModel->GetOrCreateEditorData().GetSummaryRoot();
+	
+	TArray<UNiagaraHierarchyCategory*> HierarchyCategories;
+	Root->GetChildrenOfType<UNiagaraHierarchyCategory>(HierarchyCategories, false);
+	
+	for(UNiagaraHierarchyCategory* TopLevelHierarchyCategories : HierarchyCategories)
+	{
+		SectionToCategoryMap.FindOrAdd(TopLevelHierarchyCategories->GetSection()).Add(TopLevelHierarchyCategories);
+	}
+	
+	for(const UNiagaraHierarchySection* HierarchySection : HierarchySections)
+	{
+		FNiagaraStackSection Section;
+		Section.SectionDisplayName = HierarchySection->GetSectionNameAsText();
+		Section.Tooltip = HierarchySection->GetTooltip();
+		
+		if(SectionToCategoryMap.Contains(HierarchySection))
 		{
-			if (ModuleNode && ModuleNode->IsNodeEnabled() && (ModuleNode->HasValidScriptAndGraph() || ModuleNode->Signature.IsValid()))
+			for(UNiagaraHierarchyCategory* Category : SectionToCategoryMap[HierarchySection])
 			{
-				AppendInputsForFunctionCall(State, ModuleNode, ModuleNode, NewIssues, true);
+				Section.Categories.Add(FText::FromString(Category->ToString()));
 			}
 		}
+		
+		OutStackSections.Add(Section);
 	}
 }
 
-void UNiagaraStackSummaryViewObject::PostRefreshChildrenInternal()
+const TArray<UNiagaraHierarchySection*>& UNiagaraStackSummaryViewCollection::GetHierarchySections() const
 {
-	Super::PostRefreshChildrenInternal();
+	if (GetEmitterViewModel().IsValid())
+	{
+		return GetEmitterViewModel()->GetOrCreateEditorData().GetSummarySections();
+	}
+
+	static TArray<UNiagaraHierarchySection*> Dummy;
+	return Dummy;
 }
 
-void UNiagaraStackSummaryViewObject::GetSectionsInternal(TArray<FNiagaraStackSection>& OutStackSections) const
+bool UNiagaraStackSummaryViewCollection::FilterByActiveSection(const UNiagaraStackEntry& Child) const
 {
-	if (GetEmitterViewModel().IsValid())
+	bool bIsAllowed = Super::FilterByActiveSection(Child);
+
+	// we additionally filter out stack function inputs directly under the root if we are not in the 'All' section
+	FText ActiveSection = GetActiveSection();
+	if (ActiveSection.IdenticalTo(AllSectionName) == false && Child.IsA<UNiagaraStackFunctionInput>())
 	{
-		OutStackSections.Append(GetEmitterViewModel()->GetOrCreateEditorData().GetSummarySections());
+		bIsAllowed = false;
 	}
+	
+	return bIsAllowed;
 }
 
-void UNiagaraStackSummaryViewObject::OnViewStateChanged()
+void UNiagaraStackSummaryViewCollection::OnViewStateChanged()
 {
 	if (!IsFinalized())
 	{
+		UpdateCachedSectionData();
 		RefreshChildren();
 	}
 }
 
-
 #undef LOCTEXT_NAMESPACE
