Summary View 2.0

diff --git a/Engine/Plugins/FX/Niagara/Source/NiagaraEditor/Private/Widgets/SNiagaraHierarchy.cpp b/Engine/Plugins/FX/Niagara/Source/NiagaraEditor/Private/Widgets/SNiagaraHierarchy.cpp
--- a/Engine/Plugins/FX/Niagara/Source/NiagaraEditor/Private/Widgets/SNiagaraHierarchy.cpp
+++ b/Engine/Plugins/FX/Niagara/Source/NiagaraEditor/Private/Widgets/SNiagaraHierarchy.cpp
@@ -13,25 +13,25 @@
 #include "Widgets/Images/SImage.h"
 #include "Widgets/Text/SInlineEditableTextBlock.h"
 #include "Widgets/Text/STextBlock.h"
-#include "Framework/MultiBox/MultiBoxBuilder.h"
 #include "ScopedTransaction.h"
 #include "Modules/ModuleManager.h"
 #include "PropertyEditorModule.h"
-#include "Framework/Commands/GenericCommands.h"
 #include "Framework/Application/MenuStack.h"
 #include "Framework/Application/SlateApplication.h"
 #include "ToolMenus.h"
+#include "Widgets/Input/SSearchBox.h"
 
 #define LOCTEXT_NAMESPACE "NiagaraHierarchy"
 
-TSharedRef<SWidget> SummonContextMenu(TArray<TSharedPtr<FNiagaraHierarchyItemViewModelBase>> Items, bool bFromHierarchy)
+TSharedRef<SWidget> SummonContextMenu(TArray<TSharedPtr<FNiagaraHierarchyItemViewModelBase>> Items, bool bFromHierarchy, TSharedPtr<const SWidget> HierarchyWidget)
 {
 	UNiagaraHierarchyMenuContext* MenuContextObject = NewObject<UNiagaraHierarchyMenuContext>();
 	MenuContextObject->Items = Items;
 	MenuContextObject->bFromHierarchy = bFromHierarchy;
 
 	FToolMenuContext MenuContext(MenuContextObject);
-	MenuContext.AppendCommandList(Items[0]->GetHierarchyViewModel()->GetCommands());
+	TWeakObjectPtr<UNiagaraHierarchyViewModelBase> ViewModel = Items[0]->GetHierarchyViewModel();
+	MenuContext.AppendCommandList(ViewModel->GetCommands());
 		
 	UToolMenu* Menu = UToolMenus::Get()->GenerateMenu("NiagaraHierarchyMenu", MenuContext);
 	FToolMenuSection* BaseSection = Menu->FindSection("Base");
@@ -57,6 +57,24 @@ TSharedRef<SWidget> SummonContextMenu(TArray<TSharedPtr<FNiagaraHierarchyItemVie
 		{
 			BaseSection->AddMenuEntry(FGenericCommands::Get().Delete);
 		}
+
+		if(Items.Num() == 1 && bFromHierarchy == false)
+		{
+			if(UNiagaraHierarchyItemBase* FoundItem = ViewModel->GetHierarchyRoot()->FindChildWithIdentity(Items[0]->GetData()->GetPersistentIdentity(), true))
+			{
+				FExecuteAction ExecuteAction = FExecuteAction::CreateLambda([Items, ViewModel, HierarchyWidget, Identity = FoundItem->GetPersistentIdentity()]()
+				{
+					TSharedPtr<FNiagaraHierarchyItemViewModelBase> RespectiveHierarchyItem = ViewModel->GetHierarchyRootViewModel()->FindViewModelForChild(Identity, true);
+					if(RespectiveHierarchyItem.IsValid())
+					{
+						TSharedPtr<const SNiagaraHierarchy> CastHierarchyWidget = StaticCastSharedPtr<const SNiagaraHierarchy>(HierarchyWidget);
+						CastHierarchyWidget->NavigateToHierarchyItem(RespectiveHierarchyItem);
+					}
+				});
+				FToolUIActionChoice Action(ExecuteAction);
+				BaseSection->AddMenuEntry(FName("FindInHierarchy"), LOCTEXT("FindInHierarchyLabel", "Find in Hierarchy"), FText::GetEmpty(), FSlateIcon(), Action);
+			}			
+		}
 	}
 
 	if(Items.Num() == 1)
@@ -69,13 +87,13 @@ TSharedRef<SWidget> SummonContextMenu(TArray<TSharedPtr<FNiagaraHierarchyItemVie
 	return MenuWidget;
 }
 
-void SNiagaraHierarchyCategory::Construct(const FArguments& InArgs,	TSharedPtr<FNiagaraHierarchyCategoryViewModel> InCategory)
+void SNiagaraHierarchyCategory::Construct(const FArguments& InArgs,	TSharedPtr<FNiagaraHierarchyCategoryViewModel> InCategoryViewModel)
 {
-	Category = InCategory;
-
-	InCategory->GetOnRequestRename().BindSP(this, &SNiagaraHierarchyCategory::EnterEditingMode);
+	CategoryViewModel = InCategoryViewModel;
+	
+	InCategoryViewModel->GetOnRequestRename().BindSP(this, &SNiagaraHierarchyCategory::EnterEditingMode);
 
-	const UNiagaraHierarchyCategory* CategoryData = InCategory->GetData<UNiagaraHierarchyCategory>();
+	const UNiagaraHierarchyCategory* CategoryData = InCategoryViewModel->GetData<UNiagaraHierarchyCategory>();
 	SetToolTipText(TAttribute<FText>::CreateLambda([CategoryData]
 	{
 		return CategoryData->GetTooltip();
@@ -84,7 +102,7 @@ void SNiagaraHierarchyCategory::Construct(const FArguments& InArgs,	TSharedPtr<F
 	ChildSlot
 	[
 		SAssignNew(InlineEditableTextBlock, SInlineEditableTextBlock)
-		.Style(&FNiagaraEditorStyle::Get().GetWidgetStyle<FInlineEditableTextBlockStyle>("NiagaraEditor.HierarchyEditor.Category"))
+		.Style(&FNiagaraEditorStyle::Get().GetWidgetStyle<FInlineEditableTextBlockStyle>("NiagaraEditor.HierarchyEditor.CategoryTextStyle"))
 		.Text(this, &SNiagaraHierarchyCategory::GetCategoryText)
 		.OnTextCommitted(this, &SNiagaraHierarchyCategory::OnRenameCategory)
 		.OnVerifyTextChanged(this, &SNiagaraHierarchyCategory::OnVerifyCategoryRename)
@@ -94,20 +112,23 @@ void SNiagaraHierarchyCategory::Construct(const FArguments& InArgs,	TSharedPtr<F
 
 void SNiagaraHierarchyCategory::EnterEditingMode() const
 {
-	InlineEditableTextBlock->EnterEditingMode();
+	if(CategoryViewModel.Pin()->CanRename())
+	{
+		InlineEditableTextBlock->EnterEditingMode();
+	}
 }
 
 bool SNiagaraHierarchyCategory::OnVerifyCategoryRename(const FText& NewName, FText& OutTooltip) const
 {
-	TArray<TSharedPtr<FNiagaraHierarchyCategoryViewModel>> Categories;
-	Category.Pin()->GetParent().Pin()->GetChildrenViewModelsForType<UNiagaraHierarchyCategory, FNiagaraHierarchyCategoryViewModel>(Categories);
+	TArray<TSharedPtr<FNiagaraHierarchyCategoryViewModel>> SiblingCategoryViewModels;
+	CategoryViewModel.Pin()->GetParent().Pin()->GetChildrenViewModelsForType<UNiagaraHierarchyCategory, FNiagaraHierarchyCategoryViewModel>(SiblingCategoryViewModels);
 
 	if(GetCategoryText().ToString() != NewName.ToString())
 	{
 		TSet<FString> CategoryNames;
-		for(const auto& CategoryViewModel : Categories)
+		for(const auto& SiblingCategoryViewModel : SiblingCategoryViewModels)
 		{
-			CategoryNames.Add(Cast<UNiagaraHierarchyCategory>(CategoryViewModel->GetDataMutable())->GetCategoryName().ToString());
+			CategoryNames.Add(Cast<UNiagaraHierarchyCategory>(SiblingCategoryViewModel->GetDataMutable())->GetCategoryName().ToString());
 		}
 
 		if(CategoryNames.Contains(NewName.ToString()))
@@ -122,15 +143,15 @@ bool SNiagaraHierarchyCategory::OnVerifyCategoryRename(const FText& NewName, FTe
 
 FText SNiagaraHierarchyCategory::GetCategoryText() const
 {
-	return FText::FromName(Cast<UNiagaraHierarchyCategory>(Category.Pin()->GetData())->GetCategoryName());
+	return FText::FromString(CategoryViewModel.Pin()->ToString());
 }
 
 void SNiagaraHierarchyCategory::OnRenameCategory(const FText& NewText, ETextCommit::Type) const
 {
 	FScopedTransaction Transaction(LOCTEXT("Transaction_Rename_Category", "Renamed hierarchy category"));
-	Category.Pin()->GetHierarchyViewModel()->GetHierarchyDataRoot()->Modify();
+	CategoryViewModel.Pin()->GetHierarchyViewModel()->GetHierarchyRoot()->Modify();
 	
-	Cast<UNiagaraHierarchyCategory>(Category.Pin()->GetDataMutable())->SetCategoryName(FName(NewText.ToString()));
+	CategoryViewModel.Pin()->Rename(FName(NewText.ToString()));
 }
 
 void SNiagaraHierarchySection::Construct(const FArguments& InArgs, TSharedPtr<FNiagaraHierarchySectionViewModel> InSection, TWeakObjectPtr<UNiagaraHierarchyViewModelBase> InHierarchyViewModel)
@@ -138,7 +159,9 @@ void SNiagaraHierarchySection::Construct(const FArguments& InArgs, TSharedPtr<FN
 	SectionViewModel = InSection;
 	HierarchyViewModel = InHierarchyViewModel;
 
+	IsSectionActive = InArgs._IsSectionActive;
 	OnSectionActivatedDelegate = InArgs._OnSectionActivated;
+	bForbidDropOn = InArgs._bForbidDropOn;
 	
 	if(SectionViewModel != nullptr)
 	{
@@ -172,34 +195,44 @@ void SNiagaraHierarchySection::Construct(const FArguments& InArgs, TSharedPtr<FN
 			+ SHorizontalBox::Slot()
 			.AutoWidth()
 			[
-				SNew(SNiagaraSectionDragDropTarget)
+				SNew(SNiagaraSectionDragDropTarget, SectionViewModel, EItemDropZone::AboveItem)
 				.DropTargetArgs(LeftDropTargetArgs)
 			]
 			+ SHorizontalBox::Slot()
 			[
-				SNew(SNiagaraSectionDragDropTarget)
+				SNew(SNiagaraSectionDragDropTarget, SectionViewModel, EItemDropZone::OntoItem)
 				.DropTargetArgs(OntoDropTargetArgs
 				.Content()
 				[
-					SNew(SCheckBox)
-					.Visibility(EVisibility::HitTestInvisible)
+					SAssignNew(CheckBox, SCheckBox)
+					//.Visibility(EVisibility::HitTestInvisible)
 					.Style(FAppStyle::Get(), "DetailsView.SectionButton")
 					.OnCheckStateChanged(this, &SNiagaraHierarchySection::OnSectionCheckChanged)
 					.IsChecked(this, &SNiagaraHierarchySection::GetSectionCheckState)
 					.Padding(FMargin(8.f, 4.f))
 					[
-						SAssignNew(InlineEditableTextBlock, SInlineEditableTextBlock)
-						.IsSelected(InArgs._IsSelected)
-						.Text(this, &SNiagaraHierarchySection::GetText)
-						.OnTextCommitted(this, &SNiagaraHierarchySection::OnRenameSection)
-						.OnVerifyTextChanged(this, &SNiagaraHierarchySection::OnVerifySectionRename)
+						SNew(SHorizontalBox)
+						+ SHorizontalBox::Slot()
+						.AutoWidth()
+						.Padding(2.f)
+						[
+							SNew(SImage).Image(SectionViewModel->GetSectionImage() != nullptr ? SectionViewModel->GetSectionImage() : FStyleDefaults::GetNoBrush())
+						]
+						+ SHorizontalBox::Slot()
+						[
+							SAssignNew(InlineEditableTextBlock, SInlineEditableTextBlock)
+							.Text(this, &SNiagaraHierarchySection::GetText)
+							.OnTextCommitted(this, &SNiagaraHierarchySection::OnRenameSection)
+							.OnVerifyTextChanged(this, &SNiagaraHierarchySection::OnVerifySectionRename)
+							.IsSelected(this, &SNiagaraHierarchySection::IsSectionSelected)
+						]
 					]
 				])
 			]
 			+ SHorizontalBox::Slot()
 			.AutoWidth()
 			[
-				SNew(SNiagaraSectionDragDropTarget)
+				SNew(SNiagaraSectionDragDropTarget, SectionViewModel, EItemDropZone::BelowItem)
 				.DropTargetArgs(RightDropTargetArgs)				
 			]
 		];		
@@ -234,7 +267,10 @@ SNiagaraHierarchySection::~SNiagaraHierarchySection()
 
 void SNiagaraHierarchySection::EnterEditingMode() const
 {
-	InlineEditableTextBlock->EnterEditingMode();
+	if(SectionViewModel.IsValid() && SectionViewModel->CanRename())
+	{
+		InlineEditableTextBlock->EnterEditingMode();
+	}
 }
 
 TSharedPtr<FNiagaraHierarchySectionViewModel> SNiagaraHierarchySection::GetSectionViewModel()
@@ -244,19 +280,25 @@ TSharedPtr<FNiagaraHierarchySectionViewModel> SNiagaraHierarchySection::GetSecti
 
 bool SNiagaraHierarchySection::OnCanAcceptDrop(TSharedPtr<FDragDropOperation> DragDropOperation, EItemDropZone ItemDropZone) const
 {
-	if(SectionViewModel.IsValid())
+	if(bForbidDropOn)
 	{
-		return SectionViewModel->OnCanAcceptDropInternal(DragDropOperation, ItemDropZone).IsSet();
+		return false;
 	}
-	// for the All section which has no valid view model, we simply do a check if the sections of the dragged categories are different
-	else
+
+	if(DragDropOperation->IsOfType<FNiagaraHierarchyDragDropOp>())
 	{
+		TSharedPtr<FNiagaraHierarchyDragDropOp> HierarchyDragDropOp = StaticCastSharedPtr<FNiagaraHierarchyDragDropOp>(DragDropOperation);
+		TSharedPtr<FNiagaraHierarchyItemViewModelBase> DraggedItem = HierarchyDragDropOp->GetDraggedItem().Pin();
+		
+		if(SectionViewModel.IsValid())
+		{
+			return SectionViewModel->CanDropOn(DraggedItem, ItemDropZone).bCanPerform;
+		}
+		
+		// for the All section which has no valid view model, we simply do a check if the sections of the dragged categories are different
 		if(DragDropOperation->IsOfType<FNiagaraHierarchyDragDropOp>())
 		{
-			TSharedPtr<FNiagaraHierarchyDragDropOp> HierarchyDragDropOp = StaticCastSharedPtr<FNiagaraHierarchyDragDropOp>(DragDropOperation);
-			TWeakPtr<FNiagaraHierarchyItemViewModelBase> DraggedItem = HierarchyDragDropOp->GetDraggedItem();
-
-			if(const UNiagaraHierarchyCategory* Category = Cast<UNiagaraHierarchyCategory>(DraggedItem.Pin()->GetData()))
+			if(const UNiagaraHierarchyCategory* Category = Cast<UNiagaraHierarchyCategory>(DraggedItem->GetData()))
 			{
 				return TryGetSectionData() != Category->GetSection();
 			}
@@ -268,35 +310,36 @@ bool SNiagaraHierarchySection::OnCanAcceptDrop(TSharedPtr<FDragDropOperation> Dr
 
 FReply SNiagaraHierarchySection::OnDroppedOn(const FGeometry&, const FDragDropEvent& DragDropEvent, EItemDropZone DropZone) const
 {
-	if(SectionViewModel.IsValid())
-	{
-		return SectionViewModel->OnDroppedOn(DragDropEvent, DropZone, SectionViewModel);
-	}
-	else
+	bDraggedOn = false;
+	if(TSharedPtr<FNiagaraHierarchyDragDropOp> HierarchyDragDropOp = DragDropEvent.GetOperationAs<FNiagaraHierarchyDragDropOp>())
 	{
-		if(TSharedPtr<FNiagaraHierarchyDragDropOp> HierarchyDragDropOp = DragDropEvent.GetOperationAs<FNiagaraHierarchyDragDropOp>())
+		TSharedPtr<FNiagaraHierarchyItemViewModelBase> DraggedItem = HierarchyDragDropOp->GetDraggedItem().Pin();
+		if(SectionViewModel.IsValid())
 		{
-			if(UNiagaraHierarchyCategory* Category = Cast<UNiagaraHierarchyCategory>(HierarchyDragDropOp->GetDraggedItem().Pin()->GetDataMutable()))
-			{
-				Category->SetSection(nullptr);
-				TArray<UNiagaraHierarchyCategory*> ChildrenCategories;
-				Category->GetChildrenOfType<>(ChildrenCategories, true);
-
-				for(UNiagaraHierarchyCategory* ChildCategory : ChildrenCategories)
-				{
-					ChildCategory->SetSection(nullptr);
-				}
+			SectionViewModel->OnDroppedOn(DraggedItem, DropZone);
+			return FReply::Handled();
+		}
+		
+		if(UNiagaraHierarchyCategory* Category = Cast<UNiagaraHierarchyCategory>(DraggedItem->GetDataMutable()))
+		{
+			Category->SetSection(nullptr);
+			TArray<UNiagaraHierarchyCategory*> ChildrenCategories;
+			Category->GetChildrenOfType<>(ChildrenCategories, true);
 
-				// we only need to reparent if the parent isn't already the root. This stops unnecessary reordering
-				if(HierarchyDragDropOp->GetDraggedItem().Pin()->GetParent() != HierarchyViewModel->GetHierarchyViewModelRoot())
-				{
-					HierarchyViewModel->GetHierarchyViewModelRoot()->ReparentToThis(HierarchyDragDropOp->GetDraggedItem().Pin());
-				}
-			
-				HierarchyViewModel->RefreshHierarchyView();
+			for(UNiagaraHierarchyCategory* ChildCategory : ChildrenCategories)
+			{
+				ChildCategory->SetSection(nullptr);
+			}
 
-				return FReply::Handled();
+			// we only need to reparent if the parent isn't already the root. This stops unnecessary reordering
+			if(HierarchyDragDropOp->GetDraggedItem().Pin()->GetParent() != HierarchyViewModel->GetHierarchyRootViewModel())
+			{
+				HierarchyViewModel->GetHierarchyRootViewModel()->ReparentToThis(HierarchyDragDropOp->GetDraggedItem().Pin());
 			}
+		
+			HierarchyViewModel->RefreshHierarchyView();
+
+			return FReply::Handled();
 		}
 	}
 
@@ -324,7 +367,7 @@ FReply SNiagaraHierarchySection::OnMouseButtonDown(const FGeometry& MyGeometry,
 
 FReply SNiagaraHierarchySection::OnMouseButtonUp(const FGeometry& MyGeometry, const FPointerEvent& MouseEvent)
 {
-	if(SectionViewModel.IsValid())
+	if(SectionViewModel.IsValid() && SectionViewModel->IsForHierarchy())
 	{
 		if(MouseEvent.GetEffectingButton() == EKeys::RightMouseButton)
 		{			
@@ -332,7 +375,7 @@ FReply SNiagaraHierarchySection::OnMouseButtonUp(const FGeometry& MyGeometry, co
 			FSlateApplication::Get().PushMenu(
 				AsShared(),
 				FWidgetPath(),
-				SummonContextMenu({GetSectionViewModel()}, true),
+				SummonContextMenu({GetSectionViewModel()}, true, AsShared()),
 				FSlateApplication::Get().GetCursorPos(),
 				FPopupTransitionEffect(FPopupTransitionEffect::ContextMenu)
 			);
@@ -345,9 +388,18 @@ FReply SNiagaraHierarchySection::OnMouseButtonUp(const FGeometry& MyGeometry, co
 	return FReply::Unhandled();
 }
 
+void SNiagaraHierarchySection::OnDragEnter(const FGeometry& MyGeometry, const FDragDropEvent& DragDropEvent)
+{
+	bDraggedOn = true;
+	if(DragDropEvent.GetOperationAs<FNiagaraHierarchyDragDropOp>() && DragDropEvent.GetOperationAs<FNiagaraSectionDragDropOp>() == nullptr)
+	{
+		RegisterActiveTimer(1.0f, FWidgetActiveTimerDelegate::CreateSP(this, &SNiagaraHierarchySection::ActivateSectionIfDragging));
+	}
+}
+
 FReply SNiagaraHierarchySection::OnDragDetected(const FGeometry& MyGeometry, const FPointerEvent& MouseEvent)
 {
-	if(MouseEvent.IsMouseButtonDown(EKeys::LeftMouseButton) && SectionViewModel.IsValid())
+	if(MouseEvent.IsMouseButtonDown(EKeys::LeftMouseButton) && SectionViewModel.IsValid() && SectionViewModel->CanDrag().bCanPerform)
 	{
 		TSharedRef<FNiagaraSectionDragDropOp> SectionDragDropOp = MakeShared<FNiagaraSectionDragDropOp>(SectionViewModel);
 		SectionDragDropOp->Construct();
@@ -359,12 +411,30 @@ FReply SNiagaraHierarchySection::OnDragDetected(const FGeometry& MyGeometry, con
 
 void SNiagaraHierarchySection::OnDragLeave(const FDragDropEvent& DragDropEvent)
 {
+	bDraggedOn = false;
 	if(TSharedPtr<FNiagaraHierarchyDragDropOp> HierarchyDragDropOp = DragDropEvent.GetOperationAs<FNiagaraHierarchyDragDropOp>())
 	{
 		HierarchyDragDropOp->SetDescription(FText::GetEmpty());
 	}
 }
 
+FReply SNiagaraHierarchySection::OnKeyDown(const FGeometry& MyGeometry, const FKeyEvent& InKeyEvent)
+{
+	if(SectionViewModel.IsValid() && InKeyEvent.GetKey() == EKeys::Delete && SectionViewModel->CanDelete())
+	{
+		SectionViewModel->Delete();
+		return FReply::Handled();
+	}
+
+	return FReply::Unhandled();
+}
+
+void SNiagaraHierarchySection::OnMouseLeave(const FPointerEvent& MouseEvent)
+{
+	//bDraggedOn = false;
+	SCompoundWidget::OnMouseLeave(MouseEvent);
+}
+
 UNiagaraHierarchySection* SNiagaraHierarchySection::TryGetSectionData() const
 {
 	return SectionViewModel.IsValid() ? Cast<UNiagaraHierarchySection>(SectionViewModel->GetDataMutable()) : nullptr;
@@ -383,9 +453,9 @@ FText SNiagaraHierarchySection::GetTooltipText() const
 void SNiagaraHierarchySection::OnRenameSection(const FText& Text, ETextCommit::Type CommitType) const
 {
 	FScopedTransaction Transaction(LOCTEXT("Transaction_Rename_Section", "Renamed hierarchy section"));
-	HierarchyViewModel->GetHierarchyDataRoot()->Modify();
+	HierarchyViewModel->GetHierarchyRoot()->Modify();
 	
-	SectionViewModel->SetSectionNameAsText(Text);
+	SectionViewModel->Rename(FName(Text.ToString()));
 }
 
 bool SNiagaraHierarchySection::OnVerifySectionRename(const FText& NewName, FText& OutTooltip) const
@@ -401,7 +471,7 @@ bool SNiagaraHierarchySection::OnVerifySectionRename(const FText& NewName, FText
 		TArray<FString> SectionNames;
 
 		SectionNames.Add("All");
-		for(auto& Section : HierarchyViewModel->GetHierarchyViewModelRoot()->GetSectionViewModels())
+		for(auto& Section : HierarchyViewModel->GetHierarchyRootViewModel()->GetSectionViewModels())
 		{
 			SectionNames.Add(Section->GetSectionName().ToString());
 		}
@@ -416,9 +486,14 @@ bool SNiagaraHierarchySection::OnVerifySectionRename(const FText& NewName, FText
 	return true;
 }
 
+bool SNiagaraHierarchySection::IsSectionSelected() const
+{
+	return GetSectionCheckState() == ECheckBoxState::Checked ? true : false;
+}
+
 ECheckBoxState SNiagaraHierarchySection::GetSectionCheckState() const
 {
-	return HierarchyViewModel->GetActiveSection() == SectionViewModel ? ECheckBoxState::Checked : ECheckBoxState::Unchecked;
+	return IsSectionActive.Get();
 }
 
 void SNiagaraHierarchySection::OnSectionCheckChanged(ECheckBoxState NewState)
@@ -426,16 +501,41 @@ void SNiagaraHierarchySection::OnSectionCheckChanged(ECheckBoxState NewState)
 	OnSectionActivatedDelegate.ExecuteIfBound(SectionViewModel);
 }
 
+EActiveTimerReturnType SNiagaraHierarchySection::ActivateSectionIfDragging(double CurrentTime, float DeltaTime) const
+{
+	if(bDraggedOn && FSlateApplication::Get().GetDragDroppingContent().IsValid() && FSlateApplication::Get().GetDragDroppingContent()->IsOfType<FNiagaraHierarchyDragDropOp>())
+	{
+		if(IsSectionSelected() == false)
+		{
+			OnSectionActivatedDelegate.ExecuteIfBound(SectionViewModel);
+		}
+	}
+
+	return EActiveTimerReturnType::Stop;
+}
+
 void SNiagaraHierarchy::Construct(const FArguments& InArgs, TObjectPtr<UNiagaraHierarchyViewModelBase> InHierarchyViewModel)
 {
 	HierarchyViewModel = InHierarchyViewModel;
 	
+	SourceRoot = NewObject<UNiagaraHierarchyRoot>(GetTransientPackage(), MakeUniqueObjectName(GetTransientPackage(), UNiagaraHierarchyRoot::StaticClass()));
+	SourceRootViewModel = MakeShared<FNiagaraHierarchyRootViewModel>(SourceRoot, HierarchyViewModel, false);
+	SourceRootViewModel->Initialize();
+	SourceRootViewModel->AddChildFilter(FNiagaraHierarchyItemViewModelBase::FOnFilterChild::CreateSP(this, &SNiagaraHierarchy::FilterForSourceSection));
+	SourceRootViewModel->OnSyncPropagated().BindSP(this, &SNiagaraHierarchy::RequestRefreshSourceViewNextFrame, false);
+	SourceRootViewModel->OnSectionsChanged().BindSP(this, &SNiagaraHierarchy::RefreshSectionsView);
+
+	HierarchyViewModel->OnInitialized().BindSP(this, &SNiagaraHierarchy::Reinitialize);
+	HierarchyViewModel->OnNavigateToItemInHierarchyRequested().BindSP(this, &SNiagaraHierarchy::NavigateToHierarchyItem);
+	HierarchyViewModel->OnRefreshSourceItemsRequested().BindSP(this, &SNiagaraHierarchy::RefreshSourceItems);
+	HierarchyViewModel->OnRefreshViewRequested().BindSP(this, &SNiagaraHierarchy::RefreshAllViews);
 	HierarchyViewModel->OnRefreshSourceView().BindSP(this, &SNiagaraHierarchy::RefreshSourceView);
 	HierarchyViewModel->OnRefreshHierarchyView().BindSP(this, &SNiagaraHierarchy::RefreshHierarchyView);
-	HierarchyViewModel->OnRefreshSections().BindSP(this, &SNiagaraHierarchy::RefreshSectionsWidget);
+	HierarchyViewModel->OnRefreshSectionsView().BindSP(this, &SNiagaraHierarchy::RefreshSectionsView);
+	HierarchyViewModel->OnHierarchySectionActivated().BindSP(this, &SNiagaraHierarchy::OnHierarchySectionActivated);
+	HierarchyViewModel->OnItemAdded().BindSP(this, &SNiagaraHierarchy::OnItemAdded);
 
-	HierarchyViewModel->OnSectionActivated().AddSP(this, &SNiagaraHierarchy::OnSectionActivated);
-	HierarchyViewModel->OnItemAdded().AddSP(this, &SNiagaraHierarchy::OnItemAdded);
+	BindToHierarchyRootViewModel();
 	
 	OnGenerateRowContentWidget = InArgs._OnGenerateRowContentWidget;
 	OnGenerateCustomDetailsPanelNameWidget = InArgs._OnGenerateCustomDetailsPanelNameWidget;
@@ -450,7 +550,11 @@ void SNiagaraHierarchy::Construct(const FArguments& InArgs, TObjectPtr<UNiagaraH
 		.AutoWidth()
 		[
 			SNew(SButton)
-			.OnPressed_UObject(HierarchyViewModel.Get(), &UNiagaraHierarchyViewModelBase::AddSection)
+			.OnClicked_Lambda([this]()
+			{
+				HierarchyViewModel.Get()->AddSection();
+				return FReply::Handled();
+			})
 			.ButtonStyle(FNiagaraEditorStyle::Get(), "NiagaraEditor.HierarchyEditor.ButtonStyle")
 			[
 				SNew(SHorizontalBox)
@@ -476,7 +580,11 @@ void SNiagaraHierarchy::Construct(const FArguments& InArgs, TObjectPtr<UNiagaraH
 		.AutoWidth()
 		[
 			SNew(SButton)
-			.OnPressed_UObject(HierarchyViewModel.Get(), &UNiagaraHierarchyViewModelBase::AddCategory)
+			.OnClicked_Lambda([this]()
+			{
+				HierarchyViewModel.Get()->AddCategory();
+				return FReply::Handled();
+			})
 			.ButtonStyle(FNiagaraEditorStyle::Get(), "NiagaraEditor.HierarchyEditor.ButtonStyle")
 			[
 				SNew(SHorizontalBox)
@@ -505,13 +613,36 @@ void SNiagaraHierarchy::Construct(const FArguments& InArgs, TObjectPtr<UNiagaraH
 			SNew(SHorizontalBox)
 			+ SHorizontalBox::Slot()
 			[
-				SAssignNew(SourceTreeView, STreeView<TSharedPtr<FNiagaraHierarchyItemViewModelBase>>)
-				.TreeItemsSource(&InHierarchyViewModel->GetSourceItems())
-				.OnSelectionChanged(this, &SNiagaraHierarchy::OnSelectionChanged, false)
-				.OnGenerateRow(this, &SNiagaraHierarchy::GenerateSourceItemRow)
-				.OnGetChildren_UObject(InHierarchyViewModel.Get(), &UNiagaraHierarchyViewModelBase::OnGetChildren)
-				.OnItemToString_Debug_UObject(InHierarchyViewModel.Get(), &UNiagaraHierarchyViewModelBase::OnItemToStringDebug)
-				.OnContextMenuOpening(this, &SNiagaraHierarchy::SummonContextMenuForSelectedRows, false)
+				SNew(SVerticalBox)
+				+ SVerticalBox::Slot()
+				.AutoHeight()
+				.Padding(2.f)
+				[
+					SAssignNew(SourceSearchBox, SSearchBox)
+					.OnTextChanged(this, &SNiagaraHierarchy::OnSourceSearchTextChanged)
+					.OnTextCommitted(this, &SNiagaraHierarchy::OnSourceSearchTextCommitted)
+					.OnSearch(SSearchBox::FOnSearch::CreateSP(this, &SNiagaraHierarchy::OnSearchButtonClicked))
+					.DelayChangeNotificationsWhileTyping(true)
+					.SearchResultData(this, &SNiagaraHierarchy::GetSearchResultData)
+				]
+				+ SVerticalBox::Slot()
+				.AutoHeight()
+				.Padding(2.f)
+				[
+					SAssignNew(SourceSectionBox, SWrapBox)
+					.UseAllottedWidth(true)
+				]
+				+ SVerticalBox::Slot()
+				.Padding(6.f, 2.f)
+				[
+					SAssignNew(SourceTreeView, STreeView<TSharedPtr<FNiagaraHierarchyItemViewModelBase>>)
+					.TreeItemsSource(&GetSourceItems())
+					.OnSelectionChanged(this, &SNiagaraHierarchy::OnSelectionChanged, false)
+					.OnGenerateRow(this, &SNiagaraHierarchy::GenerateSourceItemRow)
+					.OnGetChildren_UObject(InHierarchyViewModel.Get(), &UNiagaraHierarchyViewModelBase::OnGetChildren)
+					.OnItemToString_Debug_UObject(InHierarchyViewModel.Get(), &UNiagaraHierarchyViewModelBase::OnItemToStringDebug)
+					.OnContextMenuOpening(this, &SNiagaraHierarchy::SummonContextMenuForSelectedRows, false)
+				]
 			]
 			+ SHorizontalBox::Slot()
 			.AutoWidth()
@@ -539,7 +670,7 @@ void SNiagaraHierarchy::Construct(const FArguments& InArgs, TObjectPtr<UNiagaraH
 					.AutoHeight()
 					.Padding(1.f)
 					[
-						SAssignNew(SectionBox, SWrapBox)
+						SAssignNew(HierarchySectionBox, SWrapBox)
 						.UseAllottedWidth(true)
 					]				
 				]
@@ -556,8 +687,34 @@ void SNiagaraHierarchy::Construct(const FArguments& InArgs, TObjectPtr<UNiagaraH
 							AddCategoryButton
 						]
 					]
+					// we reserve space at the bottom for root drops
 					+ SVerticalBox::Slot()
-					.AutoHeight()
+					.FillHeight(0.1f)
+					.Padding(2.f, 4.f)
+					[
+						SNew(SDropTarget)
+						.OnDropped(this, &SNiagaraHierarchy::HandleHierarchyRootDrop)
+						.OnAllowDrop(this, &SNiagaraHierarchy::OnCanDropOnRoot)
+						.OnDragEnter(this, &SNiagaraHierarchy::OnRootDragEnter)
+						.OnDragLeave(this, &SNiagaraHierarchy::OnRootDragLeave)
+						[
+							SNew(SBorder)
+							.BorderImage(FAppStyle::GetBrush("DashedBorder"))
+							.BorderBackgroundColor(FLinearColor(0.2f, 0.2f, 0.2f, 0.5f))
+							[
+								SNew(SBox)
+								.HAlign(HAlign_Center)
+								.VAlign(VAlign_Center)
+								[
+									SNew(SImage)
+									.Image(FNiagaraEditorStyle::Get().GetBrush("NiagaraEditor.HierarchyEditor.RootDropIcon"))
+									.ColorAndOpacity(this, &SNiagaraHierarchy::GetRootIconColor)
+								]
+							]
+						]
+					]
+					+ SVerticalBox::Slot()
+					.Padding(6.f, 2.f)
 					[
 						SAssignNew(HierarchyTreeView, STreeView<TSharedPtr<FNiagaraHierarchyItemViewModelBase>>)
 						.TreeItemsSource(&InHierarchyViewModel->GetHierarchyItems())
@@ -567,13 +724,6 @@ void SNiagaraHierarchy::Construct(const FArguments& InArgs, TObjectPtr<UNiagaraH
 						.OnItemToString_Debug_UObject(InHierarchyViewModel.Get(), &UNiagaraHierarchyViewModelBase::OnItemToStringDebug)
 						.OnContextMenuOpening(this, &SNiagaraHierarchy::SummonContextMenuForSelectedRows, true)
 					]
-					+ SVerticalBox::Slot()
-					.VAlign(VAlign_Fill)
-					[
-						SNew(SDropTarget)
-						.OnDropped(this, &SNiagaraHierarchy::HandleHierarchyRootDrop)
-						.OnAllowDrop(this, &SNiagaraHierarchy::CanDropOnRoot)
-					]
 				]
 			]
 			+ SHorizontalBox::Slot().Expose(DetailsPanelSlot)		
@@ -587,6 +737,8 @@ void SNiagaraHierarchy::Construct(const FArguments& InArgs, TObjectPtr<UNiagaraH
 		FDetailsViewArgs DetailsViewArgs;
 		DetailsViewArgs.NameAreaSettings = FDetailsViewArgs::ENameAreaSettings::ObjectsUseNameArea;
 		DetailsViewArgs.bShowObjectLabel = false;
+		DetailsViewArgs.bAllowSearch = false;
+		DetailsViewArgs.NotifyHook = this;
 
 		DetailsPanel = PropertyEditorModule.CreateDetailView(DetailsViewArgs);
 		
@@ -594,7 +746,9 @@ void SNiagaraHierarchy::Construct(const FArguments& InArgs, TObjectPtr<UNiagaraH
 		{
 			TSharedRef<SWidget> CustomDetailsPanelNameWidget = OnGenerateCustomDetailsPanelNameWidget.Execute(nullptr);
 			DetailsPanel->SetNameAreaCustomContent(CustomDetailsPanelNameWidget);
-		}		
+		}
+
+		DetailsPanel->SetIsPropertyEditingEnabledDelegate(FIsPropertyEditingEnabled::CreateSP(this, &SNiagaraHierarchy::IsDetailsPanelEditingAllowed));
 		
 		for(auto& Customizations : HierarchyViewModel->GetInstanceCustomizations())
 		{
@@ -603,34 +757,70 @@ void SNiagaraHierarchy::Construct(const FArguments& InArgs, TObjectPtr<UNiagaraH
 		
 		DetailsPanelSlot->AttachWidget(DetailsPanel.ToSharedRef());
 	}
-	
-	RefreshSectionsWidget();
-	HierarchyViewModel->GetHierarchyViewModelRoot()->GetOnSynced().AddSP(this, &SNiagaraHierarchy::RefreshSectionsWidget);
-	
+		
 	HierarchyViewModel->GetCommands()->MapAction(FGenericCommands::Get().Rename,
 		FExecuteAction::CreateSP(this, &SNiagaraHierarchy::RequestRenameSelectedItem),
 		FCanExecuteAction::CreateSP(this, &SNiagaraHierarchy::CanRequestRenameSelectedItem));
 
 	HierarchyViewModel->GetCommands()->MapAction(FGenericCommands::Get().Delete,
 		FExecuteAction::CreateSP(this, &SNiagaraHierarchy::DeleteSelectedItems),
-		FCanExecuteAction::CreateSP(this, &SNiagaraHierarchy::CanDeleteSelectedItems));
-
+		FCanExecuteAction::CreateSP(this, &SNiagaraHierarchy::CanDeleteSelectedItems),
+		FIsActionChecked(), FIsActionButtonVisible::CreateSP(this, &SNiagaraHierarchy::CanDeleteSelectedItems));
+	
 	HierarchyViewModel->ForceFullRefresh();
 }
 
 SNiagaraHierarchy::~SNiagaraHierarchy()
 {
-	HierarchyViewModel->GetHierarchyViewModelRoot()->GetOnSynced().RemoveAll(this);
+	SourceSearchResults.Empty();
+	FocusedSearchResult.Reset();
+	
+	ClearSourceItems();
+	
+	if(HierarchyViewModel.IsValid())
+	{
+		HierarchyViewModel->OnInitialized().Unbind();
+		HierarchyViewModel->OnRefreshSourceItemsRequested().Unbind();
+		HierarchyViewModel->OnRefreshViewRequested().Unbind();
+		HierarchyViewModel->OnRefreshSourceView().Unbind();
+		HierarchyViewModel->OnRefreshHierarchyView().Unbind();
+		HierarchyViewModel->OnRefreshSectionsView().Unbind();
+		HierarchyViewModel->OnHierarchySectionActivated().Unbind();
+		HierarchyViewModel->OnItemAdded().Unbind();
+
+		UnbindFromHierarchyRootViewModel();
+	
+		HierarchyViewModel->GetCommands()->UnmapAction(FGenericCommands::Get().Delete);
+		HierarchyViewModel->GetCommands()->UnmapAction(FGenericCommands::Get().Rename);			
+	}
+	
+	SourceRootViewModel->OnSyncPropagated().Unbind();
+	SourceRootViewModel->OnSectionsChanged().Unbind();
+	SourceRootViewModel.Reset();
+	SourceRoot->ConditionalBeginDestroy();
+	SourceRoot = nullptr;
+}
 
-	HierarchyViewModel->OnRefreshSourceView().Unbind();
-	HierarchyViewModel->OnRefreshHierarchyView().Unbind();
-	HierarchyViewModel->OnRefreshSections().Unbind();
+void SNiagaraHierarchy::RefreshSourceItems()
+{
+	HierarchyViewModel->PrepareSourceItems(SourceRoot, SourceRootViewModel);
+	SourceRootViewModel->SyncViewModelsToData();
+	RefreshSourceView(false);
+	RefreshSectionsView();
+}
 
-	HierarchyViewModel->OnSectionActivated().RemoveAll(this);
-	HierarchyViewModel->OnItemAdded().RemoveAll(this);
-	
-	HierarchyViewModel->GetCommands()->UnmapAction(FGenericCommands::Get().Delete);
-	HierarchyViewModel->GetCommands()->UnmapAction(FGenericCommands::Get().Rename);
+void SNiagaraHierarchy::RefreshAllViews(bool bFullRefresh)
+{
+	RefreshSourceView(bFullRefresh);
+	RefreshHierarchyView(bFullRefresh);
+	RefreshSectionsView();
+}
+
+void SNiagaraHierarchy::RequestRefreshAllViewsNextFrame(bool bFullRefresh)
+{
+	RequestRefreshSourceViewNextFrame(bFullRefresh);
+	RequestRefreshHierarchyViewNextFrame(bFullRefresh);
+	RequestRefreshSectionsViewNextFrame();
 }
 
 FReply SNiagaraHierarchy::OnKeyDown(const FGeometry& MyGeometry, const FKeyEvent& InKeyEvent)
@@ -661,11 +851,17 @@ TSharedPtr<SWidget> SNiagaraHierarchy::SummonContextMenuForSelectedRows(bool bFr
 		SourceTreeView->GetSelectedItems(ViewModels);
 	}
 
-	return SummonContextMenu(ViewModels, bFromHierarchy);
+	if(ViewModels.IsEmpty())
+	{
+		return nullptr;
+	}
+	
+	return SummonContextMenu(ViewModels, bFromHierarchy, AsShared());
 }
 
 void SNiagaraHierarchy::RefreshSourceView(bool bFullRefresh) const
 {
+	SourceTreeView->SetTreeItemsSource(&GetSourceItems());
 	if(bFullRefresh)
 	{
 		SourceTreeView->RebuildList();
@@ -676,6 +872,19 @@ void SNiagaraHierarchy::RefreshSourceView(bool bFullRefresh) const
 	}
 }
 
+void SNiagaraHierarchy::RequestRefreshSourceViewNextFrame(bool bFullRefresh)
+{
+	if(!RefreshSourceViewNextFrameHandle.IsValid())
+	{
+		RefreshSourceViewNextFrameHandle = RegisterActiveTimer(0.f, FWidgetActiveTimerDelegate::CreateLambda([=](double CurrentTime, float DeltaTime)
+		{
+			RefreshSourceView(bFullRefresh);
+			RefreshSourceViewNextFrameHandle.Reset();
+			return EActiveTimerReturnType::Stop;
+		}));
+	}
+}
+
 void SNiagaraHierarchy::RefreshHierarchyView(bool bFullRefresh) const
 {
 	// the top layer objects might have changed due to filtering. We need to refresh these too.
@@ -690,21 +899,76 @@ void SNiagaraHierarchy::RefreshHierarchyView(bool bFullRefresh) const
 	}
 }
 
-void SNiagaraHierarchy::RefreshSectionsWidget()
+void SNiagaraHierarchy::RequestRefreshHierarchyViewNextFrame(bool bFullRefresh)
 {
-	SectionBox->ClearChildren();
-	SectionsWidgetMap.Empty();
+	if(!RefreshHierarchyViewNextFrameHandle.IsValid())
+	{
+		RefreshHierarchyViewNextFrameHandle = RegisterActiveTimer(0.f, FWidgetActiveTimerDelegate::CreateLambda([=](double CurrentTime, float DeltaTime)
+		{
+			RefreshHierarchyView(bFullRefresh);
+			RefreshHierarchyViewNextFrameHandle.Reset();
+			return EActiveTimerReturnType::Stop;
+		}));
+	}
+}
+
+void SNiagaraHierarchy::RefreshSectionsView()
+{
+	SourceSectionBox->ClearChildren();
+	HierarchySectionBox->ClearChildren();
+
+	for (TSharedPtr<FNiagaraHierarchySectionViewModel>& SourceSection : SourceRootViewModel->GetSectionViewModels())
+	{
+		TSharedPtr<SNiagaraHierarchySection> SectionWidget = SNew(SNiagaraHierarchySection, SourceSection, HierarchyViewModel)
+		.IsSectionActive_Lambda([this, SourceSection]()
+		{
+			return GetActiveSourceSection() == SourceSection ? ECheckBoxState::Checked : ECheckBoxState::Unchecked;
+		})
+		.OnSectionActivated_Lambda([this](TSharedPtr<FNiagaraHierarchySectionViewModel> SectionViewModel)
+		{
+			SetActiveSourceSection(SectionViewModel);
+		});
+		
+		SourceSectionBox->AddSlot()
+		.HAlign(HAlign_Center)
+		.VAlign(VAlign_Center)
+		[
+			SectionWidget.ToSharedRef()		
+		];
+	}
+
+	if(SourceRootViewModel->GetSectionViewModels().Num() > 0)
+	{
+		TSharedPtr<SNiagaraHierarchySection> DefaultSourceSection = SNew(SNiagaraHierarchySection, nullptr, HierarchyViewModel)
+		.IsSectionActive_Lambda([this]()
+		{
+			return GetActiveSourceSection() == nullptr ? ECheckBoxState::Checked : ECheckBoxState::Unchecked;
+		})
+		.OnSectionActivated(this, &SNiagaraHierarchy::SetActiveSourceSection)
+		// we forbid drop on here as the 'All' sections don't have a valid view model to determine it instead
+		.bForbidDropOn(true);
+
+		SourceSectionBox->AddSlot()
+		.HAlign(HAlign_Center)
+		.VAlign(VAlign_Center)
+		[
+			DefaultSourceSection.ToSharedRef()
+		];
+	}
 	
-	for (TSharedPtr<FNiagaraHierarchySectionViewModel>& Section : HierarchyViewModel->GetHierarchyViewModelRoot()->GetSectionViewModels())
+	for (TSharedPtr<FNiagaraHierarchySectionViewModel>& HierarchySection : HierarchyViewModel->GetHierarchyRootViewModel()->GetSectionViewModels())
 	{
-		TSharedPtr<SNiagaraHierarchySection> SectionWidget = SNew(SNiagaraHierarchySection, Section, HierarchyViewModel)
+		TSharedPtr<SNiagaraHierarchySection> SectionWidget = SNew(SNiagaraHierarchySection, HierarchySection, HierarchyViewModel)
+		.IsSectionActive_Lambda([this, HierarchySection]()
+		{
+			return HierarchyViewModel->GetActiveHierarchySection() == HierarchySection ? ECheckBoxState::Checked : ECheckBoxState::Unchecked;
+		})
 		.OnSectionActivated_Lambda([this](TSharedPtr<FNiagaraHierarchySectionViewModel> SectionViewModel)
 		{
-			HierarchyViewModel->SetActiveSection(SectionViewModel);
+			HierarchyViewModel->SetActiveHierarchySection(SectionViewModel);
 		});
-		SectionsWidgetMap.Add(Section, SectionWidget);
 		
-		SectionBox->AddSlot()
+		HierarchySectionBox->AddSlot()
 		.HAlign(HAlign_Center)
 		.VAlign(VAlign_Center)
 		[
@@ -712,33 +976,65 @@ void SNiagaraHierarchy::RefreshSectionsWidget()
 		];
 	}
 	
-	TSharedPtr<SNiagaraHierarchySection> DefaultSectionWidget = SNew(SNiagaraHierarchySection, nullptr, HierarchyViewModel)
+	TSharedPtr<SNiagaraHierarchySection> DefaultHierarchySection = SNew(SNiagaraHierarchySection, nullptr, HierarchyViewModel)
+	.IsSectionActive_Lambda([this]()
+	{
+		return HierarchyViewModel->GetActiveHierarchySection() == nullptr ? ECheckBoxState::Checked : ECheckBoxState::Unchecked;
+	})
 	.OnSectionActivated_Lambda([this](TSharedPtr<FNiagaraHierarchySectionViewModel> SectionViewModel)
 	{
-		HierarchyViewModel->SetActiveSection(SectionViewModel);
+		HierarchyViewModel->SetActiveHierarchySection(SectionViewModel);
 	});
-	
-	SectionsWidgetMap.Add(nullptr, DefaultSectionWidget);
 
-	SectionBox->AddSlot()
+	HierarchySectionBox->AddSlot()
 	.HAlign(HAlign_Center)
 	.VAlign(VAlign_Center)
 	[
-		DefaultSectionWidget.ToSharedRef()
+		DefaultHierarchySection.ToSharedRef()
 	];
 }
 
-bool SNiagaraHierarchy::IsItemSelected(TSharedPtr<FNiagaraHierarchyItemViewModelBase> Item) const
+void SNiagaraHierarchy::RequestRefreshSectionsViewNextFrame()
 {
-	return HierarchyTreeView->IsItemSelected(Item);
+	if(!RefreshSectionsViewNextFrameHandle.IsValid())
+	{
+		RefreshSectionsViewNextFrameHandle = RegisterActiveTimer(0.f, FWidgetActiveTimerDelegate::CreateLambda([=](double CurrentTime, float DeltaTime)
+		{
+			RefreshSectionsView();
+			RefreshSectionsViewNextFrameHandle.Reset();
+			return EActiveTimerReturnType::Stop;
+		}));
+	}
 }
 
-void SNiagaraHierarchy::SelectObjectInDetailsPanel(UObject* Object) const
+void SNiagaraHierarchy::NavigateToHierarchyItem(FNiagaraHierarchyIdentity Identity) const
 {
-	if(DetailsPanel.IsValid())
+	if(TSharedPtr<FNiagaraHierarchyItemViewModelBase> ViewModel = HierarchyViewModel->GetHierarchyRootViewModel()->FindViewModelForChild(Identity, true))
+	{
+		NavigateToHierarchyItem(ViewModel);
+	}
+}
+
+void SNiagaraHierarchy::NavigateToHierarchyItem(TSharedPtr<FNiagaraHierarchyItemViewModelBase> Item) const
+{
+	TArray<TSharedPtr<FNiagaraHierarchyItemViewModelBase>> ParentChain;
+	for(TWeakPtr<FNiagaraHierarchyItemViewModelBase> Parent = Item->GetParent(); Parent.IsValid(); Parent = Parent.Pin()->GetParent())
 	{
-		DetailsPanel->SetObject(Object);
+		ParentChain.Add(Parent.Pin());
 	}
+
+	for(int32 ParentIndex = ParentChain.Num()-1; ParentIndex >= 0; ParentIndex--)
+	{
+		HierarchyTreeView->SetItemExpansion(ParentChain[ParentIndex], true);
+	}
+		
+	HierarchyTreeView->SetSelection(Item);
+	HierarchyTreeView->RequestScrollIntoView(Item);
+}
+
+bool SNiagaraHierarchy::IsItemSelected(TSharedPtr<FNiagaraHierarchyItemViewModelBase> Item) const
+{
+	return HierarchyTreeView->IsItemSelected(Item);
 }
 
 TSharedRef<ITableRow> SNiagaraHierarchy::GenerateSourceItemRow(TSharedPtr<FNiagaraHierarchyItemViewModelBase> HierarchyItem, const TSharedRef<STableViewBase>& TableViewBase)
@@ -746,28 +1042,148 @@ TSharedRef<ITableRow> SNiagaraHierarchy::GenerateSourceItemRow(TSharedPtr<FNiaga
 	return SNew(STableRow<TSharedPtr<FNiagaraHierarchyItemViewModelBase>>, TableViewBase)
 	.OnDragDetected(HierarchyItem.ToSharedRef(), &FNiagaraHierarchyItemViewModelBase::OnDragDetected, true)
 	[
-		OnGenerateRowContentWidget.Execute(HierarchyItem.ToSharedRef())
+		SNew(SHorizontalBox)
+		+ SHorizontalBox::Slot()
+		.AutoWidth()
+		.Padding(1.f)
+		.HAlign(HAlign_Center)
+		.VAlign(VAlign_Center)
+		[
+			SNew(SBox)
+			.HeightOverride(10.f)
+			.WidthOverride(10.f)
+			.Visibility_Lambda([HierarchyItem, Identity = HierarchyItem->GetData()->GetPersistentIdentity(), this]()
+			{
+				TArray<TSharedPtr<FNiagaraHierarchyItemViewModelBase>> AllChildren;
+				HierarchyItem->GetChildrenViewModelsForType<UNiagaraHierarchyItemBase, FNiagaraHierarchyItemViewModelBase>(AllChildren, true);
+
+				bool bCanDrag = HierarchyItem->GetHierarchyViewModel()->GetHierarchyRootViewModel()->FindViewModelForChild(HierarchyItem->GetData()->GetPersistentIdentity(), true) == nullptr;			
+
+				if(bCanDrag)
+				{
+					for(TSharedPtr<FNiagaraHierarchyItemViewModelBase>& ItemViewModel : AllChildren)
+					{
+						if(HierarchyItem->GetHierarchyViewModel()->GetHierarchyRootViewModel()->FindViewModelForChild(ItemViewModel->GetData()->GetPersistentIdentity(), true) != nullptr)
+						{
+							bCanDrag = false;
+							break;
+						}
+					}
+				}
+
+				return bCanDrag ? EVisibility::Collapsed : EVisibility::Visible;
+			})
+			[
+				SNew(SImage)
+				.Image(FAppStyle::GetBrush("Icons.Lock"))
+				.ToolTipText(LOCTEXT("CantDragItemAlreadyInHierarchyTooltip", "This item already exists within the hierarchy and can not be dragged. Drag the existing one within the hierarchy directly."))
+			]
+		]
+		+ SHorizontalBox::Slot()
+		[
+			OnGenerateRowContentWidget.Execute(HierarchyItem.ToSharedRef())
+		]
 	];
 }
 
 TSharedRef<ITableRow> SNiagaraHierarchy::GenerateHierarchyItemRow(TSharedPtr<FNiagaraHierarchyItemViewModelBase> HierarchyItem, const TSharedRef<STableViewBase>& TableViewBase)
 {
 	return SNew(STableRow<TSharedPtr<FNiagaraHierarchyItemViewModelBase>>, TableViewBase)
-	.OnAcceptDrop(HierarchyItem.ToSharedRef(), &FNiagaraHierarchyItemViewModelBase::OnDroppedOn)
-	.OnCanAcceptDrop(HierarchyItem.ToSharedRef(), &FNiagaraHierarchyItemViewModelBase::OnCanAcceptDrop)
+	.OnAcceptDrop(HierarchyItem.ToSharedRef(), &FNiagaraHierarchyItemViewModelBase::OnDroppedOnRow)
+	.OnCanAcceptDrop(HierarchyItem.ToSharedRef(), &FNiagaraHierarchyItemViewModelBase::OnCanRowAcceptDrop)
 	.OnDragDetected(HierarchyItem.ToSharedRef(), &FNiagaraHierarchyItemViewModelBase::OnDragDetected, false)
+	.OnDragLeave(HierarchyItem.ToSharedRef(), &FNiagaraHierarchyItemViewModelBase::OnRowDragLeave)
 	[
-		OnGenerateRowContentWidget.Execute(HierarchyItem.ToSharedRef())
+		SNew(SHorizontalBox)
+		+ SHorizontalBox::Slot()
+		.AutoWidth()
+		.Padding(1.f)
+		.HAlign(HAlign_Center)
+		.VAlign(VAlign_Center)
+		[
+			SNew(SBox)
+			.HeightOverride(10.f)
+			.WidthOverride(10.f)
+			.Visibility_Lambda([HierarchyItem]()
+			{
+				return HierarchyItem->IsEditableByUser().bCanPerform ? EVisibility::Collapsed : EVisibility::Visible;
+			})
+			[
+				SNew(SImage)
+				.Image(FAppStyle::GetBrush("Icons.Lock"))
+				.ToolTipText_Lambda([HierarchyItem]()
+				{
+					FNiagaraHierarchyItemViewModelBase::FCanPerformActionResults IsEditableResults = HierarchyItem->IsEditableByUser();
+					if(IsEditableResults.bCanPerform == false)
+					{
+						return IsEditableResults.CanPerformMessage;
+					}
+
+					return FText::GetEmpty();
+				})
+			]
+		]
+		+ SHorizontalBox::Slot()
+		[
+			OnGenerateRowContentWidget.Execute(HierarchyItem.ToSharedRef())
+		]
 	];
 }
 
+bool SNiagaraHierarchy::FilterForSourceSection(TSharedPtr<const FNiagaraHierarchyItemViewModelBase> ItemViewModel) const
+{
+	if(ActiveSourceSection.IsValid())
+	{
+		if(const UNiagaraHierarchySection* Section = ItemViewModel->GetSection())
+		{
+			return GetActiveSourceSectionData() == Section;
+		}
+	}
+
+	return true;
+}
+
+void SNiagaraHierarchy::Reinitialize()
+{
+	// the hierarchy root view model has been recreated if the summary view model reinitialized. Therefore we update the bindings.
+	BindToHierarchyRootViewModel();
+	RefreshSourceItems();
+	RefreshAllViews(true);
+}
+
+void SNiagaraHierarchy::BindToHierarchyRootViewModel()
+{
+	HierarchyViewModel->GetHierarchyRootViewModel()->OnSyncPropagated().BindSP(this, &SNiagaraHierarchy::RequestRefreshHierarchyViewNextFrame, false);
+	HierarchyViewModel->GetHierarchyRootViewModel()->OnSectionsChanged().BindSP(this, &SNiagaraHierarchy::RefreshSectionsView);
+	HierarchyViewModel->GetHierarchyRootViewModel()->OnSectionAdded().BindSP(this, &SNiagaraHierarchy::OnHierarchySectionAdded);
+	HierarchyViewModel->GetHierarchyRootViewModel()->OnSectionDeleted().BindSP(this, &SNiagaraHierarchy::OnHierarchySectionDeleted);
+}
+
+void SNiagaraHierarchy::UnbindFromHierarchyRootViewModel() const
+{
+	HierarchyViewModel->GetHierarchyRootViewModel()->OnSyncPropagated().Unbind();
+	HierarchyViewModel->GetHierarchyRootViewModel()->OnSectionsChanged().Unbind();
+	HierarchyViewModel->GetHierarchyRootViewModel()->OnSectionAdded().Unbind();
+	HierarchyViewModel->GetHierarchyRootViewModel()->OnSectionDeleted().Unbind();
+}
+
+const TArray<TSharedPtr<FNiagaraHierarchyItemViewModelBase>>& SNiagaraHierarchy::GetSourceItems() const
+{
+	return SourceRootViewModel->GetFilteredChildren();
+}
+
+bool SNiagaraHierarchy::IsDetailsPanelEditingAllowed() const
+{
+	return SelectedDetailsPanelItemViewModel.IsValid() && SelectedDetailsPanelItemViewModel.Pin()->IsEditableByUser().bCanPerform;
+}
+
 void SNiagaraHierarchy::RequestRenameSelectedItem()
 {
 	TArray<TSharedPtr<FNiagaraHierarchyItemViewModelBase>> SelectedItems = HierarchyTreeView->GetSelectedItems();
 
 	if(SelectedItems.Num() == 0)
 	{
-		TSharedPtr<FNiagaraHierarchySectionViewModel> ActiveSection = HierarchyViewModel->GetActiveSection();
+		TSharedPtr<FNiagaraHierarchySectionViewModel> ActiveSection = HierarchyViewModel->GetActiveHierarchySection();
 		if(ActiveSection)
 		{
 			SelectedItems = { ActiveSection };
@@ -786,7 +1202,7 @@ bool SNiagaraHierarchy::CanRequestRenameSelectedItem() const
 
 	if(SelectedItems.Num() == 0)
 	{
-		TSharedPtr<FNiagaraHierarchySectionViewModel> ActiveSection = HierarchyViewModel->GetActiveSection();
+		TSharedPtr<FNiagaraHierarchySectionViewModel> ActiveSection = HierarchyViewModel->GetActiveHierarchySection();
 		if(ActiveSection)
 		{
 			SelectedItems = { ActiveSection };
@@ -801,68 +1217,35 @@ bool SNiagaraHierarchy::CanRequestRenameSelectedItem() const
 	return false;
 }
 
-void SNiagaraHierarchy::DeleteItems(TArray<TSharedPtr<FNiagaraHierarchyItemViewModelBase>> ItemsToDelete)
+void SNiagaraHierarchy::ClearSourceItems() const
 {
-	FScopedTransaction Transaction(LOCTEXT("Transaction_DeleteSelectedItems", "Deleted selected hierarchy items"));
-	HierarchyViewModel->GetHierarchyDataRoot()->Modify();
+	SourceRoot->GetChildrenMutable().Empty();
+	SourceRoot->GetSectionDataMutable().Empty();
+	SourceRootViewModel->GetChildrenMutable().Empty();
+	SourceRootViewModel->GetSectionViewModels().Empty();
+}
 
-	bool bAnyItemsDeleted = false;
+void SNiagaraHierarchy::DeleteItems(TArray<TSharedPtr<FNiagaraHierarchyItemViewModelBase>> ItemsToDelete) const
+{
+	TArray<FNiagaraHierarchyIdentity> DeletionIdentities;
+	
 	for(TSharedPtr<FNiagaraHierarchyItemViewModelBase>& SelectedItem : ItemsToDelete)
 	{
-		if(SelectedItem->CanDelete())
-		{
-			SelectedItem->Finalize();
-			if(SelectedItem->GetParent().IsValid())
-			{
-				SelectedItem->GetParent().Pin()->SyncToData();
-			}
-
-			bAnyItemsDeleted = true;
-		}
+		DeletionIdentities.Add(SelectedItem->GetData()->GetPersistentIdentity());
 	}
 
-	if(bAnyItemsDeleted)
-	{
-		RefreshSourceView();
-		RefreshHierarchyView();
-		HierarchyViewModel->OnHierarchyChanged().Broadcast();
-	}
-	else
-	{
-		Transaction.Cancel();
-	}
+	HierarchyViewModel->DeleteItemsWithIdentities(DeletionIdentities);
 }
 
-void SNiagaraHierarchy::DeleteSelectedItems()
+void SNiagaraHierarchy::DeleteSelectedItems() const
 {	
 	TArray<TSharedPtr<FNiagaraHierarchyItemViewModelBase>> SelectedItems = HierarchyTreeView->GetSelectedItems();
-
-	if(SelectedItems.Num() == 0)
-	{
-		TSharedPtr<FNiagaraHierarchySectionViewModel> ActiveSection = HierarchyViewModel->GetActiveSection();
-		if(ActiveSection)
-		{
-			SelectedItems = { ActiveSection };
-		}
-	}
-
 	DeleteItems(SelectedItems);
 }
 
 bool SNiagaraHierarchy::CanDeleteSelectedItems() const
 {
 	TArray<TSharedPtr<FNiagaraHierarchyItemViewModelBase>> SelectedItems = HierarchyTreeView->GetSelectedItems();
-	TArray<TSharedPtr<FNiagaraHierarchyItemViewModelBase>> SelectedItemsInSource = SourceTreeView->GetSelectedItems();
-
-	// we only want to allow inferring the intent to delete a section if _nothing_ else is selected
-	if(SelectedItems.Num() == 0 && SelectedItemsInSource.Num() == 0)
-	{
-		TSharedPtr<FNiagaraHierarchySectionViewModel> ActiveSection = HierarchyViewModel->GetActiveSection();
-		if(ActiveSection)
-		{
-			SelectedItems = { ActiveSection };
-		}
-	}
 
 	if(SelectedItems.Num() > 0)
 	{
@@ -889,22 +1272,60 @@ void SNiagaraHierarchy::OnItemAdded(TSharedPtr<FNiagaraHierarchyItemViewModelBas
 	}
 	else if(AddedItem->GetData()->IsA<UNiagaraHierarchySection>())
 	{
-		RefreshSectionsWidget();
+		RefreshSectionsView();
 	}
 
 	AddedItem->RequestRenamePending();
 }
 
-void SNiagaraHierarchy::OnSectionActivated(TSharedPtr<FNiagaraHierarchySectionViewModel> Section) const
+void SNiagaraHierarchy::OnHierarchySectionActivated(TSharedPtr<FNiagaraHierarchySectionViewModel> Section)
 {
 	OnSelectionChanged(Section, ESelectInfo::Direct, true);
 }
 
+void SNiagaraHierarchy::OnSourceSectionActivated(TSharedPtr<FNiagaraHierarchySectionViewModel> Section)
+{
+	OnSelectionChanged(Section, ESelectInfo::Direct, false);
+	RunSourceSearch();
+}
+
+void SNiagaraHierarchy::OnHierarchySectionAdded(TSharedPtr<FNiagaraHierarchySectionViewModel> AddedSection)
+{
+	HierarchyViewModel->SetActiveHierarchySection(AddedSection);
+	AddedSection->RequestRenamePending();
+}
+
+void SNiagaraHierarchy::OnHierarchySectionDeleted(TSharedPtr<FNiagaraHierarchySectionViewModel> DeletedSection)
+{
+	if(HierarchyViewModel->GetActiveHierarchySection() == DeletedSection)
+	{
+		HierarchyViewModel->SetActiveHierarchySection(nullptr);
+	}
+}
+
+void SNiagaraHierarchy::SetActiveSourceSection(TSharedPtr<FNiagaraHierarchySectionViewModel> Section)
+{
+	ActiveSourceSection = Section;	
+	RefreshSourceView(true);
+	OnSourceSectionActivated(Section);
+}
+
+TSharedPtr<FNiagaraHierarchySectionViewModel> SNiagaraHierarchy::GetActiveSourceSection() const
+{
+	return ActiveSourceSection.IsValid() ? ActiveSourceSection.Pin() : nullptr;
+}
+
+UNiagaraHierarchySection* SNiagaraHierarchy::GetActiveSourceSectionData() const
+{
+	return ActiveSourceSection.IsValid() ? ActiveSourceSection.Pin()->GetDataMutable<UNiagaraHierarchySection>() : nullptr;
+}
+
 void SNiagaraHierarchy::OnSelectionChanged(TSharedPtr<FNiagaraHierarchyItemViewModelBase> HierarchyItem, ESelectInfo::Type Type, bool bFromHierarchy) const
 {
+	SelectedDetailsPanelItemViewModel.Reset();
 	if(DetailsPanel.IsValid())
 	{
-		if(HierarchyItem.IsValid())
+		if(HierarchyItem.IsValid() && HierarchyItem->IsForHierarchy())
 		{
 			// when se select a section, and the previous item selection is no longer available due to it, we would get a selection refresh next tick
 			// to wipe out the current selection. We want to avoid that, so we manually clear the selected items in that case.
@@ -927,76 +1348,311 @@ void SNiagaraHierarchy::OnSelectionChanged(TSharedPtr<FNiagaraHierarchyItemViewM
 			DataForEditing->SetFlags(RF_Transactional);
 
 			// we make sure the object we are editing is transactional
-			DetailsPanel->SetObject(DataForEditing);	
+			DetailsPanel->SetObject(DataForEditing);
+			SelectedDetailsPanelItemViewModel = HierarchyItem;
 		}
 		else
 		{
+			SelectedDetailsPanelItemViewModel.Reset();
 			DetailsPanel->SetObject(nullptr);
 		}
 	}
 	
-	if(DetailsPanel.IsValid() && OnGenerateCustomDetailsPanelNameWidget.IsBound() && HierarchyItem.IsValid())
+	if(DetailsPanel.IsValid() && OnGenerateCustomDetailsPanelNameWidget.IsBound() && SelectedDetailsPanelItemViewModel.IsValid())
 	{
-		TSharedRef<SWidget> NameWidget = OnGenerateCustomDetailsPanelNameWidget.Execute(HierarchyItem);
+		TSharedRef<SWidget> NameWidget = SNew(SHorizontalBox)
+		+ SHorizontalBox::Slot()
+		.AutoWidth()
+		.Padding(2.f)
+		.HAlign(HAlign_Center)
+		.VAlign(VAlign_Center)
+		[
+			SNew(SImage)
+			.Image(FAppStyle::GetBrush("Icons.Lock"))
+			.Visibility(SelectedDetailsPanelItemViewModel.Pin()->IsEditableByUser().bCanPerform ? EVisibility::Collapsed : EVisibility::Visible)
+			.ToolTipText(SelectedDetailsPanelItemViewModel.Pin()->IsEditableByUser().CanPerformMessage)
+		]
+		+ SHorizontalBox::Slot()
+		.AutoWidth()
+		.Padding(2.f)
+		[
+			OnGenerateCustomDetailsPanelNameWidget.Execute(SelectedDetailsPanelItemViewModel.Pin())
+		];
+		
 		DetailsPanel->SetNameAreaCustomContent(NameWidget);
 	}
+}
 
-	HierarchyViewModel->OnSelectionChanged(HierarchyItem);
+void SNiagaraHierarchy::RunSourceSearch()
+{
+	if(!SourceSearchBox->GetText().IsEmpty())
+	{
+		OnSourceSearchTextChanged(SourceSearchBox->GetText());
+	}
 }
 
-FReply SNiagaraHierarchy::HandleHierarchyRootDrop(const FGeometry& Geometry, const FDragDropEvent& DragDropEvent) const
+void SNiagaraHierarchy::OnSourceSearchTextChanged(const FText& Text)
 {
-	// we pass in nullptr because the dragged object is contained within the drag drop event
-	return HierarchyViewModel->GetHierarchyViewModelRoot()->OnDroppedOn(DragDropEvent, EItemDropZone::OntoItem, nullptr);
+	SourceSearchResults.Empty();
+	FocusedSearchResult.Reset();
+	SourceTreeView->ClearSelection();
+
+	if(!Text.IsEmpty())
+	{
+		FString TextAsString = Text.ToString();;
+		TArray<TSharedPtr<FNiagaraHierarchyItemViewModelBase>> AllChildren;
+		SourceRootViewModel->GetChildrenViewModelsForType<UNiagaraHierarchyItemBase, FNiagaraHierarchyItemViewModelBase>(AllChildren, true);
+
+		TArray<FSearchItem> SearchItems;
+		GenerateSearchItems(SourceRootViewModel.ToSharedRef(), {}, SearchItems);
+		
+		for(const FSearchItem& SearchItem : SearchItems)
+		{
+			for(const FString& SearchTerm : SearchItem.GetEntry()->GetSearchTerms())
+			{
+				if(SearchTerm.Contains(TextAsString))
+				{
+					SourceSearchResults.Add(SearchItem);
+				}
+			}
+		}
+
+		ExpandSourceSearchResults();
+		SelectNextSourceSearchResult();
+	}
+	else
+	{
+		SourceTreeView->ClearExpandedItems();
+	}
 }
 
-bool SNiagaraHierarchy::CanDropOnRoot(TSharedPtr<FDragDropOperation> DragDropOp) const
+void SNiagaraHierarchy::OnSourceSearchTextCommitted(const FText& Text, ETextCommit::Type CommitType)
 {
-	if(DragDropOp->IsOfType<FNiagaraHierarchyDragDropOp>())
+	bool bIsShiftDown = FSlateApplication::Get().GetModifierKeys().IsShiftDown();
+	if(CommitType == ETextCommit::OnEnter)
 	{
-		TSharedPtr<FNiagaraHierarchyDragDropOp> HierarchyDragDropOp = StaticCastSharedPtr<FNiagaraHierarchyDragDropOp>(DragDropOp);
-		TSharedPtr<FNiagaraHierarchyItemViewModelBase> DraggedItem = HierarchyDragDropOp->GetDraggedItem().Pin();
-		TOptional<EItemDropZone> Result = HierarchyViewModel->GetHierarchyViewModelRoot()->OnCanAcceptDropInternal(DragDropOp, EItemDropZone::OntoItem);
-		return Result.IsSet() && Result.GetValue() == EItemDropZone::OntoItem ? true : false;
+		if(bIsShiftDown == false)
+		{
+			SelectNextSourceSearchResult();
+		}
+		else
+		{
+			SelectPreviousSourceSearchResult();
+		}
 	}
+}
 
-	return false;
+void SNiagaraHierarchy::OnSearchButtonClicked(SSearchBox::SearchDirection SearchDirection)
+{
+	if(SearchDirection == SSearchBox::Next)
+	{
+		SelectNextSourceSearchResult();
+	}
+	else
+	{
+		SelectPreviousSourceSearchResult();
+	}
 }
 
-void SNiagaraSectionDragDropTarget::Construct(const FArguments& InArgs)
+void SNiagaraHierarchy::GenerateSearchItems(TSharedRef<FNiagaraHierarchyItemViewModelBase> Root, TArray<TSharedPtr<FNiagaraHierarchyItemViewModelBase>> ParentChain, TArray<FSearchItem>& OutSearchItems)
 {
-	SDropTarget::Construct(InArgs._DropTargetArgs);
+	const TArray<TSharedPtr<FNiagaraHierarchyItemViewModelBase>> FilteredChildren = Root->GetFilteredChildren();
+	ParentChain.Add(Root);
+	OutSearchItems.Add(FSearchItem{ParentChain});
+	for(TSharedPtr<FNiagaraHierarchyItemViewModelBase> Child : FilteredChildren)
+	{
+		GenerateSearchItems(Child.ToSharedRef(), ParentChain, OutSearchItems);
+	}
 }
 
-FReply SNiagaraSectionDragDropTarget::OnDragOver(const FGeometry& MyGeometry, const FDragDropEvent& DragDropEvent)
+void SNiagaraHierarchy::ExpandSourceSearchResults()
 {
-	TSharedPtr<FNiagaraHierarchyDragDropOp> HierarchyDragDropOp = DragDropEvent.GetOperationAs<FNiagaraHierarchyDragDropOp>();
+	SourceTreeView->ClearExpandedItems();
 
-	if(HierarchyDragDropOp.IsValid())
+	for(FSearchItem& SearchResult : SourceSearchResults)
 	{
-		if(AllowDrop(DragDropEvent.GetOperation()))
+		for(TSharedPtr<FNiagaraHierarchyItemViewModelBase>& EntryInPath : SearchResult.Path)
 		{
-			if(const UNiagaraHierarchySection* Section = Cast<UNiagaraHierarchySection>(HierarchyDragDropOp->GetDraggedItem().Pin()->GetData()))
-			{
-				FText Text = LOCTEXT("MoveSectionHere", "The section will be moved here.");
-				HierarchyDragDropOp->SetDescription(Text);
-			}
+			SourceTreeView->SetItemExpansion(EntryInPath, true);
+		}
+	}
+}
 
-			if(const UNiagaraHierarchyCategory* Category = Cast<UNiagaraHierarchyCategory>(HierarchyDragDropOp->GetDraggedItem().Pin()->GetData()))
-			{
-				FText Text = LOCTEXT("DropCategoryOnSection", "Move category here");
-				HierarchyDragDropOp->SetDescription(Text);
-			}
+void SNiagaraHierarchy::SelectNextSourceSearchResult()
+{
+	if(SourceSearchResults.IsEmpty())
+	{
+		return;
+	}
+	
+	if(!FocusedSearchResult.IsSet())
+	{
+		FocusedSearchResult = SourceSearchResults[0];
+	}
+	else
+	{
+		int32 CurrentSearchResultIndex = SourceSearchResults.Find(FocusedSearchResult.GetValue());
+		if(SourceSearchResults.IsValidIndex(CurrentSearchResultIndex+1))
+		{
+			FocusedSearchResult = SourceSearchResults[CurrentSearchResultIndex+1];
 		}
 		else
 		{
-			if(const UNiagaraHierarchySection* Section = Cast<UNiagaraHierarchySection>(HierarchyDragDropOp->GetDraggedItem().Pin()->GetData()))
+			FocusedSearchResult = SourceSearchResults[0];
+		}
+	}
+
+	SourceTreeView->ClearSelection();
+	SourceTreeView->RequestScrollIntoView(FocusedSearchResult.GetValue().GetEntry());
+	SourceTreeView->SetItemSelection(FocusedSearchResult.GetValue().GetEntry(), true);
+}
+
+void SNiagaraHierarchy::SelectPreviousSourceSearchResult()
+{
+	if(SourceSearchResults.IsEmpty())
+	{
+		return;
+	}
+	
+	if(!FocusedSearchResult.IsSet())
+	{
+		FocusedSearchResult = SourceSearchResults[0];
+	}
+	else
+	{
+		int32 CurrentSearchResultIndex = SourceSearchResults.Find(FocusedSearchResult.GetValue());
+		if(SourceSearchResults.IsValidIndex(CurrentSearchResultIndex-1))
+		{
+			FocusedSearchResult = SourceSearchResults[CurrentSearchResultIndex-1];
+		}
+		else
+		{
+			FocusedSearchResult = SourceSearchResults[SourceSearchResults.Num()-1];
+		}
+	}
+
+	SourceTreeView->ClearSelection();
+	SourceTreeView->RequestScrollIntoView(FocusedSearchResult.GetValue().GetEntry());
+	SourceTreeView->SetItemSelection(FocusedSearchResult.GetValue().GetEntry(), true);
+}
+
+TOptional<SSearchBox::FSearchResultData> SNiagaraHierarchy::GetSearchResultData() const
+{
+	if(SourceSearchResults.Num() > 0)
+	{
+		SSearchBox::FSearchResultData SearchResultData;
+		SearchResultData.NumSearchResults = SourceSearchResults.Num();
+
+		if(FocusedSearchResult.IsSet())
+		{
+			// we add one just to make it look nicer as this is merely for cosmetic purposes
+			SearchResultData.CurrentSearchResultIndex = SourceSearchResults.Find(FocusedSearchResult.GetValue()) + 1;
+		}
+		else
+		{
+			SearchResultData.CurrentSearchResultIndex = INDEX_NONE;
+		}
+
+		return SearchResultData;
+	}
+
+	return TOptional<SSearchBox::FSearchResultData>();
+}
+
+FReply SNiagaraHierarchy::HandleHierarchyRootDrop(const FGeometry& Geometry, const FDragDropEvent& DragDropEvent) const
+{
+	if(TSharedPtr<FNiagaraHierarchyDragDropOp> DragDropOp = DragDropEvent.GetOperationAs<FNiagaraHierarchyDragDropOp>())
+	{
+		HierarchyViewModel->GetHierarchyRootViewModel()->OnDroppedOn(DragDropOp->GetDraggedItem().Pin(), EItemDropZone::OntoItem);
+		return FReply::Handled();
+	}
+
+	return FReply::Unhandled();
+}
+
+FNiagaraHierarchyItemViewModelBase::FCanPerformActionResults SNiagaraHierarchy::CanDropOnRoot(TSharedPtr<FNiagaraHierarchyItemViewModelBase> DraggedItem) const
+{
+	return HierarchyViewModel->GetHierarchyRootViewModel()->CanDropOnInternal(DraggedItem, EItemDropZone::OntoItem);
+}
+
+bool SNiagaraHierarchy::OnCanDropOnRoot(TSharedPtr<FDragDropOperation> DragDropOperation) const
+{
+	if(DragDropOperation->IsOfType<FNiagaraHierarchyDragDropOp>())
+	{
+		TSharedPtr<FNiagaraHierarchyDragDropOp> HierarchyDragDropOp = StaticCastSharedPtr<FNiagaraHierarchyDragDropOp>(DragDropOperation);
+		TSharedPtr<FNiagaraHierarchyItemViewModelBase> DraggedItem = HierarchyDragDropOp->GetDraggedItem().Pin();
+		return CanDropOnRoot(DraggedItem).bCanPerform;
+	}
+	
+	return false;
+}
+
+void SNiagaraHierarchy::OnRootDragEnter(const FDragDropEvent& DragDropEvent) const
+{
+	if(TSharedPtr<FNiagaraHierarchyDragDropOp> HierarchyDragDropOp = DragDropEvent.GetOperationAs<FNiagaraHierarchyDragDropOp>())
+	{
+		FNiagaraHierarchyItemViewModelBase::FCanPerformActionResults CanPerformActionResults = CanDropOnRoot(HierarchyDragDropOp->GetDraggedItem().Pin());
+		HierarchyDragDropOp->SetDescription(CanPerformActionResults.CanPerformMessage);
+	}
+}
+
+void SNiagaraHierarchy::OnRootDragLeave(const FDragDropEvent& DragDropEvent) const
+{
+	if(TSharedPtr<FNiagaraHierarchyDragDropOp> HierarchyDragDropOp = DragDropEvent.GetOperationAs<FNiagaraHierarchyDragDropOp>())
+	{
+		HierarchyDragDropOp->SetDescription(FText::GetEmpty());
+	}
+}
+
+FSlateColor SNiagaraHierarchy::GetRootIconColor() const
+{
+	if(FSlateApplication::Get().IsDragDropping())
+	{
+		if(FSlateApplication::Get().GetDragDroppingContent()->IsOfType<FNiagaraHierarchyDragDropOp>())
+		{
+			TSharedPtr<FNiagaraHierarchyDragDropOp> HierarchyDragDropOp = StaticCastSharedPtr<FNiagaraHierarchyDragDropOp>(FSlateApplication::Get().GetDragDroppingContent());
+			if(CanDropOnRoot(HierarchyDragDropOp->GetDraggedItem().Pin()).bCanPerform)
 			{
-				HierarchyDragDropOp->SetDescription(LOCTEXT("SectionOnSectionNotAllowed", "Can't drop section on a section."));
+				return FLinearColor(0.8f, 0.8f, 0.8f, 0.8f);
 			}
 		}
 	}
 
+	return FLinearColor(0.2f, 0.2f, 0.2f, 0.5f);
+}
+
+FString SNiagaraHierarchy::GetReferencerName() const
+{
+	return TEXT("Niagara Hierarchy");
+}
+
+void SNiagaraHierarchy::AddReferencedObjects(FReferenceCollector& Collector)
+{
+	Collector.AddReferencedObject(SourceRoot);
+}
+
+void SNiagaraHierarchy::NotifyPostChange(const FPropertyChangedEvent& PropertyChangedEvent,	FProperty* PropertyThatChanged)
+{
+	HierarchyViewModel->OnHierarchyPropertiesChanged().Broadcast();
+}
+
+void SNiagaraSectionDragDropTarget::Construct(const FArguments& InArgs,TSharedPtr<FNiagaraHierarchySectionViewModel> InOwningSection, EItemDropZone InItemDropZone)
+{
+	OwningSection = InOwningSection;
+	DropZone = InItemDropZone;
+	
+	SDropTarget::Construct(InArgs._DropTargetArgs);
+}
+
+FReply SNiagaraSectionDragDropTarget::OnDragOver(const FGeometry& MyGeometry, const FDragDropEvent& DragDropEvent)
+{
+	if(TSharedPtr<FNiagaraHierarchyDragDropOp> HierarchyDragDropOp = DragDropEvent.GetOperationAs<FNiagaraHierarchyDragDropOp>())
+	{
+		TSharedPtr<FNiagaraHierarchyItemViewModelBase> DraggedItem = HierarchyDragDropOp->GetDraggedItem().Pin();
+		FNiagaraHierarchyItemViewModelBase::FCanPerformActionResults Results = OwningSection->CanDropOn(DraggedItem, DropZone);
+		HierarchyDragDropOp->SetDescription(Results.CanPerformMessage);
+	}
+
 	return SDropTarget::OnDragOver(MyGeometry, DragDropEvent);
 }
 
diff --git a/Engine/Plugins/FX/Niagara/Source/NiagaraEditor/Private/Widgets/SNiagaraHierarchyAssignment.h b/Engine/Plugins/FX/Niagara/Source/NiagaraEditor/Private/Widgets/SNiagaraHierarchyAssignment.h
new file mode 100644
--- /dev/null
+++ b/Engine/Plugins/FX/Niagara/Source/NiagaraEditor/Private/Widgets/SNiagaraHierarchyAssignment.h
@@ -0,0 +1,23 @@
+﻿// Copyright Epic Games, Inc. All Rights Reserved.
+
+#pragma once
+
+#include "CoreMinimal.h"
+#include "NiagaraTypes.h"
+#include "Widgets/SCompoundWidget.h"
+
+/**
+ * 
+ */
+class NIAGARAEDITOR_API SNiagaraHierarchyAssignment : public SCompoundWidget
+{
+public:
+	SLATE_BEGIN_ARGS(SNiagaraHierarchyAssignment)
+	{
+	}
+
+	SLATE_END_ARGS()
+
+	/** Constructs this widget with InArgs */
+	void Construct(const FArguments& InArgs, class UNiagaraNodeAssignment& InAssignmentNode);
+};
diff --git a/Engine/Plugins/FX/Niagara/Source/NiagaraEditor/Private/Widgets/SNiagaraParameterPanel.cpp b/Engine/Plugins/FX/Niagara/Source/NiagaraEditor/Private/Widgets/SNiagaraParameterPanel.cpp
--- a/Engine/Plugins/FX/Niagara/Source/NiagaraEditor/Private/Widgets/SNiagaraParameterPanel.cpp
+++ b/Engine/Plugins/FX/Niagara/Source/NiagaraEditor/Private/Widgets/SNiagaraParameterPanel.cpp
@@ -102,11 +102,7 @@ void SNiagaraParameterPanel::Construct(const FArguments& InArgs, const TSharedPt
 	.CategoryBorderBackgroundPadding(FMargin(0.0f, 3.0f))
 	.OnGetKeyForItem(this, &SNiagaraParameterPanel::OnGetKeyForItem)
 	.OnGetKeyForCategory(this, &SNiagaraParameterPanel::OnGetKeyForCategory)
-	.ExpandInitially(false)
-	.SearchBoxAdjacentContent()
-	[
-		ParameterPanelViewModel->GenerateAdjacentWidget()
-	];
+	.ExpandInitially(false);
 
 	// Finalize the widget
 	ChildSlot
diff --git a/Engine/Plugins/FX/Niagara/Source/NiagaraEditor/Public/NiagaraEditorCommon.h b/Engine/Plugins/FX/Niagara/Source/NiagaraEditor/Public/NiagaraEditorCommon.h
--- a/Engine/Plugins/FX/Niagara/Source/NiagaraEditor/Public/NiagaraEditorCommon.h
+++ b/Engine/Plugins/FX/Niagara/Source/NiagaraEditor/Public/NiagaraEditorCommon.h
@@ -221,6 +221,10 @@ public:
 		check(!InputName.IsNone());		
 	}
 
+	FGuid GetFunctionGuid() const { return FunctionGuid; }
+	FGuid GetInputGuid() const { return InputGuid; }
+	FName GetInputName() const { return InputName; }
+	
 	bool operator==(const FFunctionInputSummaryViewKey& Other) const
 	{
 		return FunctionGuid == Other.FunctionGuid && InputGuid == Other.InputGuid && InputName == Other.InputName;
diff --git a/Engine/Plugins/FX/Niagara/Source/NiagaraEditor/Public/NiagaraEditorStyle.h b/Engine/Plugins/FX/Niagara/Source/NiagaraEditor/Public/NiagaraEditorStyle.h
--- a/Engine/Plugins/FX/Niagara/Source/NiagaraEditor/Public/NiagaraEditorStyle.h
+++ b/Engine/Plugins/FX/Niagara/Source/NiagaraEditor/Public/NiagaraEditorStyle.h
@@ -33,6 +33,7 @@ private:
 	void InitToolbarIcons();
 	void InitTabIcons();
 	void InitIcons();
+	void InitTextStyles();
 	void InitOverview();
 	void InitViewportStyle();
 	void InitEmitterDetails();
diff --git a/Engine/Plugins/FX/Niagara/Source/NiagaraEditor/Public/NiagaraEditorUtilities.h b/Engine/Plugins/FX/Niagara/Source/NiagaraEditor/Public/NiagaraEditorUtilities.h
--- a/Engine/Plugins/FX/Niagara/Source/NiagaraEditor/Public/NiagaraEditorUtilities.h
+++ b/Engine/Plugins/FX/Niagara/Source/NiagaraEditor/Public/NiagaraEditorUtilities.h
@@ -472,6 +472,9 @@ namespace FNiagaraParameterUtilities
 		bool bShowVisibilityConditionIcon = false;
 		bool bShowEditConditionIcon = false;
 		bool bShowAdvanced = false;
+		TOptional<TAttribute<FText>> NameOverride;
+		TOptional<TAttribute<EVisibility>> NameOverrideVisibility;
+		TOptional<TAttribute<FText>> NameOverrideTooltip;
 	};
 	
 	/** A more advanced utility function that also makes use of parameter metadata to generate a parameter widget */
diff --git a/Engine/Plugins/FX/Niagara/Source/NiagaraEditor/Public/NiagaraEmitterEditorData.h b/Engine/Plugins/FX/Niagara/Source/NiagaraEditor/Public/NiagaraEmitterEditorData.h
--- a/Engine/Plugins/FX/Niagara/Source/NiagaraEditor/Public/NiagaraEmitterEditorData.h
+++ b/Engine/Plugins/FX/Niagara/Source/NiagaraEditor/Public/NiagaraEmitterEditorData.h
@@ -1,3 +1,4 @@
+
 // Copyright Epic Games, Inc. All Rights Reserved.
 
 #pragma once
@@ -5,6 +6,7 @@
 #include "NiagaraEditorDataBase.h"
 #include "NiagaraEditorCommon.h"
 #include "NiagaraStackSection.h"
+#include "ViewModels/HierarchyEditor/NiagaraHierarchyViewModelBase.h"
 #include "NiagaraEmitterEditorData.generated.h"
 
 class UNiagaraStackEditorData;
@@ -26,6 +28,7 @@ public:
 	UNiagaraEmitterEditorData(const FObjectInitializer& ObjectInitializer);
 
 	virtual void PostLoad() override;
+	virtual void PostLoadFromOwner(UObject* InOwner) override;
 #if WITH_EDITORONLY_DATA
 	NIAGARAEDITOR_API static void DeclareConstructClasses(TArray<FTopLevelAssetPath>& OutConstructClasses, const UClass* SpecificSubclass);
 #endif
@@ -42,13 +45,8 @@ public:
 	NIAGARAEDITOR_API void SetShowSummaryView(bool bInShouldShowSummaryView);
 	NIAGARAEDITOR_API void ToggleShowSummaryView();
 	
-	NIAGARAEDITOR_API const TMap<FFunctionInputSummaryViewKey, FFunctionInputSummaryViewMetadata>& GetSummaryViewMetaDataMap() const;
-	NIAGARAEDITOR_API TOptional<FFunctionInputSummaryViewMetadata> GetSummaryViewMetaData(const FFunctionInputSummaryViewKey& Key) const;
-	NIAGARAEDITOR_API void SetSummaryViewMetaData(const FFunctionInputSummaryViewKey& Key, TOptional<FFunctionInputSummaryViewMetadata> NewMetadata);
-
-	const TArray<FNiagaraStackSection>& GetSummarySections() const;
-	void SetSummarySections(const TArray<FNiagaraStackSection>& InSummarySections);
-
+	UNiagaraHierarchyRoot* GetSummaryRoot() const { return SummaryViewRoot; }
+	const TArray<UNiagaraHierarchySection*>& GetSummarySections() const;
 private:
 	UPROPERTY(Instanced)
 	TObjectPtr<UNiagaraStackEditorData> StackEditorData;
@@ -64,14 +62,19 @@ private:
 
 	/** Stores metadata for filtering function inputs when in Filtered/Simple view. */
 	UPROPERTY()
-	TMap<FFunctionInputSummaryViewKey, FFunctionInputSummaryViewMetadata> SummaryViewFunctionInputMetadata;
+	TMap<FFunctionInputSummaryViewKey, FFunctionInputSummaryViewMetadata> SummaryViewFunctionInputMetadata_DEPRECATED;
 
-	UPROPERTY(EditAnywhere, Category=Summary)
-	TArray<FNiagaraStackSection> SummarySections;
+	UPROPERTY()
+	TArray<FNiagaraStackSection> SummarySections_DEPRECATED;
+
+	UPROPERTY()
+	TObjectPtr<UNiagaraHierarchyRoot> SummaryViewRoot = nullptr;
 
 	FSimpleMulticastDelegate OnSummaryViewStateChangedDelegate;
 	
-	void StackEditorDataChanged();	
+	void StackEditorDataChanged();
+
+	void PostLoad_TransferSummaryDataToNewFormat();
 };
 
 
diff --git a/Engine/Plugins/FX/Niagara/Source/NiagaraEditor/Public/NiagaraGraph.h b/Engine/Plugins/FX/Niagara/Source/NiagaraEditor/Public/NiagaraGraph.h
--- a/Engine/Plugins/FX/Niagara/Source/NiagaraEditor/Public/NiagaraGraph.h
+++ b/Engine/Plugins/FX/Niagara/Source/NiagaraEditor/Public/NiagaraGraph.h
@@ -301,6 +301,9 @@ class UNiagaraGraph : public UEdGraph
 
 	UNiagaraScriptVariable* GetScriptVariable(FNiagaraVariable Parameter) const;
 	NIAGARAEDITOR_API UNiagaraScriptVariable* GetScriptVariable(FName ParameterName) const;
+	NIAGARAEDITOR_API UNiagaraScriptVariable* GetScriptVariable(FGuid VariableGuid) const;
+	NIAGARAEDITOR_API TArray<UNiagaraScriptVariable*> GetChildScriptVariablesForInput(FGuid VariableGuid) const;
+	NIAGARAEDITOR_API TArray<FGuid> GetChildScriptVariableGuidsForInput(FGuid VariableGuid) const;
 
 	/** Adds parameter to the VariableToScriptVariable map.*/
 	UNiagaraScriptVariable* AddParameter(const FNiagaraVariable& Parameter, bool bIsStaticSwitch = false);
