Summary View 2.0

diff --git a/Engine/Plugins/FX/Niagara/Source/NiagaraEditor/Public/ViewModels/HierarchyEditor/NiagaraHierarchyViewModelBase.h b/Engine/Plugins/FX/Niagara/Source/NiagaraEditor/Public/ViewModels/HierarchyEditor/NiagaraHierarchyViewModelBase.h
--- a/Engine/Plugins/FX/Niagara/Source/NiagaraEditor/Public/ViewModels/HierarchyEditor/NiagaraHierarchyViewModelBase.h
+++ b/Engine/Plugins/FX/Niagara/Source/NiagaraEditor/Public/ViewModels/HierarchyEditor/NiagaraHierarchyViewModelBase.h
@@ -7,6 +7,9 @@
 #include "PropertyEditorDelegates.h"
 #include "ViewModels/Stack/NiagaraStackFunctionInput.h"
 #include "ToolMenuSection.h"
+#include "IPropertyRowGenerator.h"
+#include "Misc/TransactionObjectEvent.h"
+#include "ScopedTransaction.h"
 #include "NiagaraHierarchyViewModelBase.generated.h"
 
 USTRUCT()
@@ -25,6 +28,11 @@ struct FNiagaraHierarchyIdentity
 	UPROPERTY()
 	TArray<FName> Names;
 
+	bool IsValid() const
+	{
+		return Guids.Num() > 0 || Names.Num() > 0;
+	}
+	
 	bool operator==(const FNiagaraHierarchyIdentity& OtherIdentity) const
 	{
 		if(Guids.Num() != OtherIdentity.Guids.Num() || Names.Num() != OtherIdentity.Names.Num())
@@ -50,10 +58,42 @@ struct FNiagaraHierarchyIdentity
 
 		return true;
 	}
+
+	bool operator!=(const FNiagaraHierarchyIdentity& OtherIdentity) const
+	{
+		return !(*this == OtherIdentity);
+	}
 };
 
+FORCEINLINE uint32 GetTypeHash(const FNiagaraHierarchyIdentity& Identity)
+{
+	uint32 Hash = 0;
+	
+	for(const FGuid& Guid : Identity.Guids)
+	{
+		HashCombine(Hash, GetTypeHash(Guid));
+	}
+	
+	for(const FName& Name : Identity.Names)
+	{
+		HashCombine(Hash, GetTypeHash(Name));
+	}
+	
+	return Hash;
+}
+
+/** A base class that is used to refresh data that represents external data. Inherit from this class if you need more context data. */
 UCLASS()
-class UNiagaraHierarchyItemBase : public UObject
+class UNiagaraHierarchyDataRefreshContext : public UObject
+{
+	GENERATED_BODY()
+
+	UPROPERTY(Transient)
+	TMap<FNiagaraHierarchyIdentity, TObjectPtr<UObject>> IdentityToObjectMap;
+};
+
+UCLASS()
+class NIAGARAEDITOR_API UNiagaraHierarchyItemBase : public UObject
 {
 	GENERATED_BODY()
 
@@ -66,6 +106,14 @@ public:
 
 	template<class ChildClass>
 	ChildClass* AddChild();
+
+	UNiagaraHierarchyItemBase* FindChildWithIdentity(FNiagaraHierarchyIdentity ChildIdentity, bool bSearchRecursively = false);
+
+	UNiagaraHierarchyItemBase* CopyAndAddItemAsChild(const UNiagaraHierarchyItemBase& ItemToCopy);
+	UNiagaraHierarchyItemBase* CopyAndAddItemUnderParentIdentity(const UNiagaraHierarchyItemBase& ItemToCopy, FNiagaraHierarchyIdentity ParentIdentity);
+	
+	/** Remove a child with a given identity. Can be searched recursively. This function operates under the assumption there will be only one item with a given identity. */
+	bool RemoveChildWithIdentity(FNiagaraHierarchyIdentity ChildIdentity, bool bSearchRecursively = false);
 	
 	template<class ChildClass>
 	bool DoesOneChildExist(bool bRecursive = false) const;
@@ -73,28 +121,19 @@ public:
 	template<class ChildClass>
 	TArray<ChildClass*> GetChildrenOfType(TArray<ChildClass*>& Out, bool bRecursive = false) const;
 
-	/** Used to clean out unneeded data using the bFinalized flag. */
-	void Refresh();
-
+	template<class PREDICATE_CLASS>
+	void SortChildren(const PREDICATE_CLASS& Predicate, bool bRecursive = false);
+	
 	virtual FString ToString() const { return GetName(); }
 
-	/** Finalize an item to mark it is no longer used and should be deleted. */
-	void Finalize();
-	bool IsFinalized() const { return bFinalized; }
-
 	/** An identity can be optionally set to create a mapping from previously existing guids or names to hierarchy items that represent them. */
 	void SetIdentity(FNiagaraHierarchyIdentity InIdentity) { Identity = InIdentity; }
-	virtual FNiagaraHierarchyIdentity GetPersistentIdentity() const { return Identity; }
+	FNiagaraHierarchyIdentity GetPersistentIdentity() const { return Identity; }
+	TArray<FNiagaraHierarchyIdentity> GetParentIdentities() const;
 	
-	virtual bool CanHaveChildren() const { return false; }
-
 	/** Overridden modify method to also mark all children as modified */
 	virtual bool Modify(bool bAlwaysMarkDirty = true) override;
 
-protected:
-	/** Called by the public Refresh function. Can be overridden to further customize the refresh process, i.e. refreshing section data in the root. */
-	virtual void RefreshDataInternal() { }
-
 protected:
 	virtual void PostLoad() override;
 
@@ -107,9 +146,6 @@ protected:
 	/** An optional guid; can be used if hierarchy items represent outside items */
 	UPROPERTY()
 	FGuid Guid_DEPRECATED;
-
-	UPROPERTY(Transient)
-	bool bFinalized = false;
 };
 
 template <class ChildClass>
@@ -168,26 +204,43 @@ TArray<ChildClass*> UNiagaraHierarchyItemBase::GetChildrenOfType(TArray<ChildCla
 	return Out;
 }
 
+template <class PREDICATE_CLASS>
+void UNiagaraHierarchyItemBase::SortChildren(const PREDICATE_CLASS& Predicate, bool bRecursive)
+{
+	Children.Sort(Predicate);
+
+	if(bRecursive)
+	{
+		for(TObjectPtr<UNiagaraHierarchyItemBase> Child : Children)
+		{
+			Child->SortChildren(Predicate, bRecursive);
+		}
+	}
+}
+
 class UNiagaraHierarchySection;
 
 UCLASS()
-class UNiagaraHierarchyRoot : public UNiagaraHierarchyItemBase
+class NIAGARAEDITOR_API UNiagaraHierarchyRoot : public UNiagaraHierarchyItemBase
 {
 	GENERATED_BODY()
 public:
 	UNiagaraHierarchyRoot() {}
 	virtual ~UNiagaraHierarchyRoot() override {}
-	
-	virtual void RefreshDataInternal() override;
-	
-	virtual bool CanHaveChildren() const override { return true; }
-	
-	class UNiagaraHierarchySection* AddSection(FText InNewSectionName);
-	void RemoveSection(FText SectionName);
-	TSet<FName> GetSections() const;
+
 	const TArray<UNiagaraHierarchySection*>& GetSectionData() const { return Sections; }
 	TArray<UNiagaraHierarchySection*>& GetSectionDataMutable() { return Sections; }
+
+	TSet<FName> GetSections() const;
 	int32 GetSectionIndex(UNiagaraHierarchySection* Section) const;
+
+	UNiagaraHierarchySection* AddSection(FText InNewSectionName, int32 InsertIndex = INDEX_NONE);
+	UNiagaraHierarchySection* FindSectionByIdentity(FNiagaraHierarchyIdentity SectionIdentity);
+	/** This will copy the section element itself */
+	void DuplicateSectionFromOtherRoot(const UNiagaraHierarchySection& SectionToCopy);
+	void RemoveSection(FText SectionName);
+	void RemoveSectionByIdentity(FNiagaraHierarchyIdentity SectionIdentity);
+	
 	virtual bool Modify(bool bAlwaysMarkDirty = true) override;
 protected:
 
@@ -196,19 +249,16 @@ protected:
 };
 
 UCLASS()
-class UNiagaraHierarchyItem : public UNiagaraHierarchyItemBase
+class NIAGARAEDITOR_API UNiagaraHierarchyItem : public UNiagaraHierarchyItemBase
 {
 	GENERATED_BODY()
 public:
 	UNiagaraHierarchyItem() {}
 	virtual ~UNiagaraHierarchyItem() override {}
-
-	/** An item is assumed to not have any children. */
-	virtual bool CanHaveChildren() const override { return false; }
 };
 
 UCLASS()
-class UNiagaraHierarchyCategory : public UNiagaraHierarchyItemBase
+class NIAGARAEDITOR_API UNiagaraHierarchyCategory : public UNiagaraHierarchyItemBase
 {
 	GENERATED_BODY()
 public:
@@ -224,10 +274,13 @@ public:
 	void SetSection(UNiagaraHierarchySection* InSection) { Section = InSection; }
 	const UNiagaraHierarchySection* GetSection() const { return Section; }
 
-	virtual bool CanHaveChildren() const override { return true; }
-
 	virtual FString ToString() const override { return Category.ToString(); }
 
+	/** Since the category points to a section object, during merge or copy paste etc. it is possible the section pointer will point at a section from another root.
+	 *  We fix this up by looking through our available sections and match up via persistent identity.
+	 *  This function expects the correct section with the same identity to exist already at the root level
+	 */
+	void FixupSectionLinkage();
 private:
 	UPROPERTY(EditAnywhere, Category = "Category")
 	FName Category;
@@ -241,7 +294,7 @@ private:
 };
 
 UCLASS()
-class UNiagaraHierarchySection : public UNiagaraHierarchyItemBase
+class NIAGARAEDITOR_API UNiagaraHierarchySection : public UNiagaraHierarchyItemBase
 {
 	GENERATED_BODY()
 
@@ -266,6 +319,15 @@ private:
 	FText Tooltip;
 };
 
+UCLASS()
+class NIAGARAEDITOR_API UNiagaraHierarchyObjectProperty : public UNiagaraHierarchyItem
+{
+	GENERATED_BODY()
+public:
+	/** To know what object this ObjectProperty is referring to, a persistent guid that can be mapped back to an object is required. */
+	void Initialize(FGuid ObjectGuid, FString PropertyName);
+};
+
 struct FNiagaraHierarchyItemViewModelBase;
 struct FNiagaraHierarchySectionViewModel;
 
@@ -312,45 +374,54 @@ public:
 };
 
 UCLASS(Abstract)
-class UNiagaraHierarchyViewModelBase : public UObject, public FEditorUndoClient
+class UNiagaraHierarchyViewModelBase : public UObject, public FSelfRegisteringEditorUndoClient
 {
 public:
 	DECLARE_MULTICAST_DELEGATE(FOnHierarchyChanged)
-	DECLARE_DELEGATE_OneParam(FRefreshTreeWidget, bool bFullRefresh)
-	DECLARE_DELEGATE_RetVal_OneParam(bool, FIsItemSelected, TSharedPtr<FNiagaraHierarchyItemViewModelBase> Item)
-	DECLARE_DELEGATE_OneParam(FSelectObjectInDetailsPanel, UObject* Object)
-	DECLARE_MULTICAST_DELEGATE_OneParam(FOnSectionActivated, TSharedPtr<FNiagaraHierarchySectionViewModel> Section)
-	DECLARE_MULTICAST_DELEGATE_OneParam(FOnItemAdded, TSharedPtr<FNiagaraHierarchyItemViewModelBase> AddedItem)
+	DECLARE_MULTICAST_DELEGATE(FOnHierarchyPropertiesChanged)
+	DECLARE_DELEGATE_OneParam(FOnSectionActivated, TSharedPtr<FNiagaraHierarchySectionViewModel> Section)
+	DECLARE_DELEGATE_OneParam(FOnItemAdded, TSharedPtr<FNiagaraHierarchyItemViewModelBase> AddedItem)
+	DECLARE_DELEGATE_OneParam(FOnRefreshViewRequested, bool bForceFullRefresh)
+	DECLARE_DELEGATE_OneParam(FOnNavigateToItemInHierarchyRequested, FNiagaraHierarchyIdentity Identity)
 
 	GENERATED_BODY()
-	
+
 	UNiagaraHierarchyViewModelBase();
 	virtual ~UNiagaraHierarchyViewModelBase() override;
 
 	void Initialize();
 	void Finalize();
-
+	
 	void AddCategory() const;
 	void AddSection() const;
 
+	NIAGARAEDITOR_API void DeleteItemWithIdentity(FNiagaraHierarchyIdentity Identity);
+	NIAGARAEDITOR_API void DeleteItemsWithIdentities(TArray<FNiagaraHierarchyIdentity> Identities);
+
+	NIAGARAEDITOR_API void NavigateToItemInHierarchy(const FNiagaraHierarchyIdentity& NiagaraHierarchyIdentity);
+
 	/** Refreshes all data and widgets */
 	void ForceFullRefresh();
+	void ForceFullRefreshOnTimer();
+	void RequestFullRefreshNextFrame();
 
 	/** The hierarchy root the widget is editing. This should point to persistent data stored somewhere else as the serialized root of the hierarchy. */
-	virtual UNiagaraHierarchyRoot* GetHierarchyDataRoot() const PURE_VIRTUAL(UNiagaraHierarchyViewModelBase::GetHierarchyDataRoot, return nullptr;);
-	TSharedPtr<struct FNiagaraHierarchyRootViewModel> GetHierarchyViewModelRoot() const { return HierarchyViewModelRoot; }
-
+	virtual UNiagaraHierarchyRoot* GetHierarchyRoot() const PURE_VIRTUAL(UNiagaraHierarchyViewModelBase::GetHierarchyRoot, return nullptr;);
+	TSharedPtr<struct FNiagaraHierarchyRootViewModel> GetHierarchyRootViewModel() const { return HierarchyRootViewModel; }
+	
 	/** This function will create the view model for a given item. Override to customize view model behavior by providing custom classes. */
 	virtual TSharedPtr<FNiagaraHierarchyItemViewModelBase> CreateViewModelForData(UNiagaraHierarchyItemBase* ItemBase, TSharedPtr<FNiagaraHierarchyItemViewModelBase> Parent);
 	
-	/** Source items reflect the base, unedited status of items to edit into a hierarchy */
-	const TArray<TSharedPtr<FNiagaraHierarchyItemViewModelBase>>& GetSourceItems() const;
 	/** Hierarchy items reflect the already edited hierarchy. This should generally be constructed from persistent serialized data. */
 	const TArray<TSharedPtr<FNiagaraHierarchyItemViewModelBase>>& GetHierarchyItems() const;
 	
-	/** Prepares the items we want to create a hierarchy for. The idea is to Drag & Drop these into the hierarchy or click on them to edit details.  */
-	virtual void PrepareSourceItems() PURE_VIRTUAL(UNiagaraHierarchyViewModelBase::PrepareSourceItems,);
-
+	const UNiagaraHierarchyDataRefreshContext* GetRefreshContext() const { return RefreshContext; }
+	void SetRefreshContext(UNiagaraHierarchyDataRefreshContext* InContext) { RefreshContext = InContext; }
+	
+	/** Prepares the items we want to create a hierarchy for. Primary purpose is to add children to the source root to gather the items to display in the source panel.
+	 * The root view model is also given as a way to forcefully sync view models to access additional functionality, if needed */
+	virtual void PrepareSourceItems(UNiagaraHierarchyRoot* SourceRoot, TSharedPtr<FNiagaraHierarchyRootViewModel> SourceRootViewModel) PURE_VIRTUAL(UNiagaraHierarchyViewModelBase::PrepareSourceItems,);
+	
 	/** Additional commands can be specified overriding the Commands function. */
 	virtual void SetupCommands() {}
 	TSharedRef<FUICommandList> GetCommands() const { return Commands.ToSharedRef(); }
@@ -358,80 +429,108 @@ public:
 	/** Function to implement for custom drag drop ops. */
 	virtual TSharedRef<FNiagaraHierarchyDragDropOp> CreateDragDropOp(TSharedRef<FNiagaraHierarchyItemViewModelBase> Item) PURE_VIRTUAL(UNiagaraHierarchyViewModelBase::CreateDragDropOp, return MakeShared<FNiagaraHierarchyDragDropOp>(nullptr););
 	
-	virtual void OnSelectionChanged(TSharedPtr<FNiagaraHierarchyItemViewModelBase> HierarchyItem) {}
-	
 	void OnGetChildren(TSharedPtr<FNiagaraHierarchyItemViewModelBase> Item, TArray<TSharedPtr<FNiagaraHierarchyItemViewModelBase>>& OutChildren) const;
 	
 	virtual bool SupportsDetailsPanel() { return false; }
 	// Overriding this will give the details panel instance customizations for specific UClasses
 	virtual TArray<TTuple<UClass*, FOnGetDetailCustomizationInstance>> GetInstanceCustomizations() { return {}; }
-	
+
+	void RefreshAllViews(bool bFullRefresh = false) const;
 	void RefreshSourceView(bool bFullRefresh = false) const;
 	void RefreshHierarchyView(bool bFullRefresh = false) const;
-	void RefreshSectionsWidget() const;
+	void RefreshSectionsView() const;
 	
 	// Delegate that call functions from SNiagaraHierarchy
-	FRefreshTreeWidget& OnRefreshSourceView() { return RefreshSourceViewDelegate; }
-	FRefreshTreeWidget& OnRefreshHierarchyView() { return RefreshHierarchyWidgetDelegate; }
-	FSimpleDelegate& OnRefreshSections() { return RefreshSectionsWidgetDelegate; }
+	FSimpleDelegate& OnRefreshSourceItemsRequested() { return RefreshSourceItemsRequestedDelegate; }
+	FOnRefreshViewRequested& OnRefreshSourceView() { return RefreshSourceViewDelegate; }
+	FOnRefreshViewRequested& OnRefreshHierarchyView() { return RefreshHierarchyWidgetDelegate; }
+	FSimpleDelegate& OnRefreshSectionsView() { return RefreshSectionsViewDelegate; }
 
 	// Delegates for external systems
 	FOnHierarchyChanged& OnHierarchyChanged() { return OnHierarchyChangedDelegate; } 
-
+	FOnHierarchyChanged& OnHierarchyPropertiesChanged() { return OnHierarchyPropertiesChangedDelegate; } 
 	FOnItemAdded& OnItemAdded() { return OnItemAddedDelegate; }
+	FOnRefreshViewRequested& OnRefreshViewRequested() { return RefreshAllViewsRequestedDelegate; }
+	FOnNavigateToItemInHierarchyRequested& OnNavigateToItemInHierarchyRequested() { return OnNavigateToItemInHierarchyRequestedDelegate; }
+	FSimpleDelegate& OnInitialized() { return OnInitializedDelegate; }
 	
 	// Sections
-	void SetActiveSection(TSharedPtr<struct FNiagaraHierarchySectionViewModel>);
-	TSharedPtr<FNiagaraHierarchySectionViewModel> GetActiveSection() const;
-	UNiagaraHierarchySection* GetActiveSectionData() const;
-	bool IsSectionActive(const UNiagaraHierarchySection* Section) const;
-	FOnSectionActivated& OnSectionActivated() { return OnSectionActivatedDelegate; }
-	
-	FString OnItemToStringDebug(TSharedPtr<FNiagaraHierarchyItemViewModelBase> ItemBaseViewModel) const; 
+	NIAGARAEDITOR_API void SetActiveHierarchySection(TSharedPtr<struct FNiagaraHierarchySectionViewModel>);
+	TSharedPtr<FNiagaraHierarchySectionViewModel> GetActiveHierarchySection() const;
+	UNiagaraHierarchySection* GetActiveHierarchySectionData() const;
+	bool IsHierarchySectionActive(const UNiagaraHierarchySection* Section) const;
+	FOnSectionActivated& OnHierarchySectionActivated() { return OnHierarchySectionActivatedDelegate; }
+
+	FString OnItemToStringDebug(TSharedPtr<FNiagaraHierarchyItemViewModelBase> ItemBaseViewModel) const;
+
+	const TArray<TSharedRef<IDetailTreeNode>>& RequestDetailTreeNodesForObject(UObject* Object);
 protected:
 	virtual void InitializeInternal() {}
 	virtual void FinalizeInternal() {}
 	
 	virtual void PostUndo(bool bSuccess) override;
 	virtual void PostRedo(bool bSuccess) override;
-protected:
-	UPROPERTY()
-	TObjectPtr<UNiagaraHierarchyRoot> SourceRoot;
 
+	virtual bool MatchesContext(const FTransactionContext& InContext, const TArray<TPair<UObject*, FTransactionObjectEvent>>& TransactionObjectContexts) const override;
+private:
+	bool FilterForHierarchySection(TSharedPtr<const FNiagaraHierarchyItemViewModelBase> ItemViewModel) const;
+	bool FilterForUncategorizedRootItemsInAllSection(TSharedPtr<const FNiagaraHierarchyItemViewModelBase> ItemViewModel) const;
+protected:
 	UPROPERTY()
 	TObjectPtr<UNiagaraHierarchyRoot> HierarchyRoot;
+	
+	TSharedPtr<struct FNiagaraHierarchyRootViewModel> HierarchyRootViewModel;
 
-	TSharedPtr<struct FNiagaraHierarchyRootViewModel> SourceViewModelRoot;
-	TSharedPtr<struct FNiagaraHierarchyRootViewModel> HierarchyViewModelRoot;
+	TWeakPtr<struct FNiagaraHierarchySectionViewModel> ActiveHierarchySection;
 
-	TWeakPtr<struct FNiagaraHierarchySectionViewModel> ActiveSection;
 	TSharedPtr<FUICommandList> Commands;
 
+	TMap<UObject*, TSharedRef<IPropertyRowGenerator>> ObjectToPropertyRowGeneratorMap;
+	TMap<UObject*, TArray<TSharedRef<IDetailTreeNode>>> ObjectTreeNodeMap;
+
+	UPROPERTY(Transient)
+	TObjectPtr<UNiagaraHierarchyDataRefreshContext> RefreshContext = nullptr;
+	
 protected:
 	// delegate collection to call UI functions
-	FRefreshTreeWidget RefreshSourceViewDelegate;
-	FRefreshTreeWidget RefreshHierarchyWidgetDelegate;
-	FSimpleDelegate RefreshSectionsWidgetDelegate;
-
+	FSimpleDelegate RefreshSourceItemsRequestedDelegate;
+	FOnRefreshViewRequested RefreshAllViewsRequestedDelegate;
+	FOnRefreshViewRequested RefreshSourceViewDelegate;
+	FOnRefreshViewRequested RefreshHierarchyWidgetDelegate;
+	FSimpleDelegate RefreshSectionsViewDelegate;
+	FOnNavigateToItemInHierarchyRequested OnNavigateToItemInHierarchyRequestedDelegate;
+	
 	FOnItemAdded OnItemAddedDelegate;
-	FOnSectionActivated OnSectionActivatedDelegate;
+	FOnSectionActivated OnHierarchySectionActivatedDelegate;
+	FOnSectionActivated OnSourceSectionActivatedDelegate;
 	FOnHierarchyChanged OnHierarchyChangedDelegate;
+	FOnHierarchyPropertiesChanged OnHierarchyPropertiesChangedDelegate;
+	
+	FSimpleDelegate OnInitializedDelegate;
+
+	FTimerHandle FullRefreshNextFrameHandle;
 };
 
 struct NIAGARAEDITOR_API FNiagaraHierarchyItemViewModelBase : TSharedFromThis<FNiagaraHierarchyItemViewModelBase>, public FTickableEditorObject
 {
-	DECLARE_MULTICAST_DELEGATE(FOnSynced)
-
-	FNiagaraHierarchyItemViewModelBase(UNiagaraHierarchyItemBase* InItemBase, TSharedPtr<FNiagaraHierarchyItemViewModelBase> InParent, TWeakObjectPtr<UNiagaraHierarchyViewModelBase> InHierarchyViewModel)
+	DECLARE_DELEGATE(FOnSynced)
+	DECLARE_DELEGATE_RetVal_OneParam(bool, FOnFilterChild, const TSharedPtr<const FNiagaraHierarchyItemViewModelBase> Child);
+	DECLARE_DELEGATE_OneParam(FOnChildRequestedDeletion, TSharedPtr<FNiagaraHierarchyItemViewModelBase> Child)
+	
+	FNiagaraHierarchyItemViewModelBase(UNiagaraHierarchyItemBase* InItemBase, TSharedPtr<FNiagaraHierarchyItemViewModelBase> InParent, TWeakObjectPtr<UNiagaraHierarchyViewModelBase> InHierarchyViewModel, bool bInIsForHierarchy)
 		: ItemBase(InItemBase)
 		, Parent(InParent)
 		, HierarchyViewModel(InHierarchyViewModel)
+		, bIsForHierarchy(bInIsForHierarchy)
 	{		
 		
 	}
-	
-	virtual ~FNiagaraHierarchyItemViewModelBase() {}
 
+	/** Can be implemented for additional logic that the constructor isn't valid for. */
+	virtual void Initialize() {}
+	
+	virtual ~FNiagaraHierarchyItemViewModelBase() override;
+	
 	UNiagaraHierarchyItemBase* GetDataMutable() const { return ItemBase; }
 	const UNiagaraHierarchyItemBase* GetData() const { return ItemBase; }
 	
@@ -440,54 +539,83 @@ struct NIAGARAEDITOR_API FNiagaraHierarchyItemViewModelBase : TSharedFromThis<FN
 	
 	template<class T>
 	const T* GetData() const { return Cast<T>(ItemBase); }
-
+	
 	virtual void Tick(float DeltaTime) override;
 	virtual TStatId GetStatId() const override;
 	
 	virtual FString ToString() const { return ItemBase->ToString(); }
-	
-	void SyncToData();
+	FText ToStringAsText() const { return FText::FromString(ToString()); }
+	virtual TArray<FString> GetSearchTerms() const { return {ToString()} ;}
+
+	void RefreshChildrenData();
+	void SyncViewModelsToData();
 	const TArray<TSharedPtr<FNiagaraHierarchyItemViewModelBase>>& GetChildren() const { return Children; }
+	TArray<TSharedPtr<FNiagaraHierarchyItemViewModelBase>>& GetChildrenMutable() { return Children; }
 	const TArray<TSharedPtr<FNiagaraHierarchyItemViewModelBase>>& GetFilteredChildren() const;
 
+	void AddChildFilter(FOnFilterChild InFilterChild);
+	
 	template<class DataClass, class ViewModelChildClass>
 	void GetChildrenViewModelsForType(TArray<TSharedPtr<ViewModelChildClass>>& OutChildren, bool bRecursive = false);
-
+	
 	bool HasParent(TSharedPtr<FNiagaraHierarchyItemViewModelBase> ParentCandidate, bool bRecursive = false);
 
 	void AddChild(TSharedPtr<FNiagaraHierarchyItemViewModelBase> Item);
 	TSharedPtr<FNiagaraHierarchyItemViewModelBase> AddNewItem(TSubclassOf<UNiagaraHierarchyItemBase> NewItemClass);
-	void DuplicateToThis(TSharedPtr<FNiagaraHierarchyItemViewModelBase> ItemToDuplicate, int32 InsertIndex = INDEX_NONE);
-	void ReparentToThis(TSharedPtr<FNiagaraHierarchyItemViewModelBase> ItemToMove, int32 InsertIndex = INDEX_NONE);
+	TSharedRef<FNiagaraHierarchyItemViewModelBase> DuplicateToThis(TSharedPtr<FNiagaraHierarchyItemViewModelBase> ItemToDuplicate, int32 InsertIndex = INDEX_NONE);
+	TSharedRef<FNiagaraHierarchyItemViewModelBase> ReparentToThis(TSharedPtr<FNiagaraHierarchyItemViewModelBase> ItemToMove, int32 InsertIndex = INDEX_NONE);
 
-	TSharedPtr<FNiagaraHierarchyItemViewModelBase> FindViewModelForChild(UNiagaraHierarchyItemBase* Child) const;
-	int32 FindIndexOfChild(TSharedPtr<FNiagaraHierarchyItemViewModelBase> Child) const;
+	TSharedPtr<FNiagaraHierarchyItemViewModelBase> FindViewModelForChild(UNiagaraHierarchyItemBase* Child, bool bSearchRecursively = false) const;
+	TSharedPtr<FNiagaraHierarchyItemViewModelBase> FindViewModelForChild(FNiagaraHierarchyIdentity ChildIdentity, bool bSearchRecursively = false) const;
 	int32 FindIndexOfChild(UNiagaraHierarchyItemBase* Child) const;
 	int32 FindIndexOfDataChild(TSharedPtr<FNiagaraHierarchyItemViewModelBase> Child) const;
 	int32 FindIndexOfDataChild(UNiagaraHierarchyItemBase* Child) const;
 
-	/** Deleting will finalize as well, but is a user triggered action whereas Finalize is an internal call. */
+	UNiagaraHierarchyItemBase* AddChild(TSubclassOf<UNiagaraHierarchyItemBase> NewChildClass, FNiagaraHierarchyIdentity ChildIdentity);
+	
+	/** Deleting will ask the parent to delete its child */
 	void Delete();
-	/** Finalizing the view model will finalize the data & the children as well as handle cleanup. */
-	void Finalize();
+	void DeleteChild(TSharedPtr<FNiagaraHierarchyItemViewModelBase> Child);
 
 	TWeakObjectPtr<UNiagaraHierarchyViewModelBase> GetHierarchyViewModel() const { return HierarchyViewModel; }
-	
-	virtual FReply OnDroppedOn(const FDragDropEvent& DragDropEvent, EItemDropZone ItemDropZone, TSharedPtr<FNiagaraHierarchyItemViewModelBase> Item)
+
+	/** Returns a set result if the item can accept a drop either above/onto/below the item.  */
+	TOptional<EItemDropZone> OnCanRowAcceptDrop(const FDragDropEvent& DragDropEvent, EItemDropZone ItemDropZone, TSharedPtr<FNiagaraHierarchyItemViewModelBase> Item);
+	virtual FReply OnDroppedOnRow(const FDragDropEvent& DragDropEvent, EItemDropZone ItemDropZone, TSharedPtr<FNiagaraHierarchyItemViewModelBase> Item);
+	void OnRowDragLeave(const FDragDropEvent& DragDropEvent);
+
+	struct FCanPerformActionResults
 	{
-		return FReply::Unhandled();
-	}
+		FCanPerformActionResults(bool bInCanPerform) : bCanPerform(bInCanPerform) {}
+		
+		bool bCanPerform = false;
+		/** A message that is used when bCanPerform is false. Will either be used in tooltips in the hierarchy editor or as popup message. */
+		FText CanPerformMessage;
 
-	TOptional<EItemDropZone> OnCanAcceptDrop(const FDragDropEvent& DragDropEvent, EItemDropZone ItemDropZone, TSharedPtr<FNiagaraHierarchyItemViewModelBase> Item);
-	virtual TOptional<EItemDropZone> OnCanAcceptDropInternal(TSharedPtr<FDragDropOperation> DragDropOp, EItemDropZone ItemDropZone);
+		bool operator==(const bool& bOther) const
+		{
+			return bCanPerform == bOther;
+		}
 
-	/** Should return true if draggable */
-	virtual bool CanDrag() { return false; }
+		bool operator!=(const bool& bOther) const
+		{
+			return !(*this==bOther);
+		}
+	};
 
-	/** Should return true if renamable */
-	virtual bool CanRename() { return false; }
+	/** Should return true if properties are supposed to be editable & needs to be true if typical operations should work on it (renaming, dragging, deleting etc.) */
+	virtual FCanPerformActionResults IsEditableByUser() { return FCanPerformActionResults(false); }
+	
+	/** Needs to be true in order to allow drag & drop operations to parent items to this item */
+	virtual bool CanHaveChildren() const { return false; }
+	
+	/** Should return true if an item should be draggable. An uneditable item can not be dragged even if CanDragInternal returns true. */
+	FCanPerformActionResults CanDrag();
+	
+	/** Returns true if renamable */
+	bool CanRename() { return IsEditableByUser().bCanPerform && CanRenameInternal(); }
 
-	virtual bool CanDelete() { return true; }
+	void Rename(FName NewName) { RenameInternal(NewName); HierarchyViewModel->OnHierarchyPropertiesChanged().Broadcast(); }
 
 	void RequestRename()
 	{
@@ -505,6 +633,27 @@ struct NIAGARAEDITOR_API FNiagaraHierarchyItemViewModelBase : TSharedFromThis<FN
 			bRenamePending = true;
 		}
 	}
+	
+	/** Returns true if deletable */
+	bool CanDelete() { return IsEditableByUser().bCanPerform && CanDeleteInternal(); }
+
+	/** Returns true if the given item can be dropped on the given target area. */
+	FCanPerformActionResults CanDropOn(TSharedPtr<FNiagaraHierarchyItemViewModelBase> DraggedItem, EItemDropZone ItemDropZone) { return CanDropOnInternal(DraggedItem, ItemDropZone); }
+
+	/** Gets executed when an item was dropped on this. */
+	void OnDroppedOn(TSharedPtr<FNiagaraHierarchyItemViewModelBase> DroppedItem, EItemDropZone ItemDropZone) { OnDroppedOnInternal(DroppedItem, ItemDropZone); }
+
+	/** Determines the section this item belongs to. */
+	const UNiagaraHierarchySection* GetSection() const { return GetSectionInternal(); }
+	
+	/** For data cleanup that represents external data, this needs to return true in order for live cleanup to work. */
+	virtual bool RepresentsExternalData() const { return false; }
+	/** This function determines whether a hierarchy item that represents that external data should be maintained during data refresh
+	 * Needs to be implemented if RepresentsExternalData return true.
+	 * The context object can be used to add arbitrary data. */
+	virtual bool DoesExternalDataStillExist(const UNiagaraHierarchyDataRefreshContext* Context) const { return false; }
+
+	bool IsForHierarchy() const { return bIsForHierarchy; }
 
 	/** Override this to register dynamic context menu entries when right clicking a hierarchy item */
 	virtual void PopulateDynamicContextMenuSection(FToolMenuSection& DynamicSection) {}
@@ -519,22 +668,54 @@ struct NIAGARAEDITOR_API FNiagaraHierarchyItemViewModelBase : TSharedFromThis<FN
 
 	FSimpleDelegate& GetOnRequestRename() { return OnRequestRenameDelegate; }
 	FOnSynced& GetOnSynced() { return OnSyncedDelegate; }
+	FOnChildRequestedDeletion& OnChildRequestedDeletion() { return OnChildRequestedDeletionDelegate; }
 
 	TWeakPtr<FNiagaraHierarchyItemViewModelBase> GetParent() { return Parent; }
 
+protected:
+	/** Should return true if draggable. An optional message can be provided if false that will show as a slate notification. */
+	virtual FCanPerformActionResults CanDragInternal() { return false; }
+
+	/** Should return true if renamable */
+	virtual bool CanRenameInternal() { return false; }
+
+	virtual void RenameInternal(FName NewName) {}
+	
+	/** Should return true if deletable. By default, we can delete items in the hierarchy, not in the source. */
+	virtual bool CanDeleteInternal() { return IsForHierarchy(); }
+
+	/** Should return true if the given drag drop operation is allowed to succeed. */
+	virtual FCanPerformActionResults CanDropOnInternal(TSharedPtr<FNiagaraHierarchyItemViewModelBase>, EItemDropZone ItemDropZone);
+
+	/** Override this to handle drop-on logic. */
+	virtual void OnDroppedOnInternal(TSharedPtr<FNiagaraHierarchyItemViewModelBase>, EItemDropZone ItemDropZone) { }
+
+	/** Can be overridden to support custom sections.
+	 * In the hierarchy only categories can be parented directly to the root, but using this it is possible to add items to custom sections in the source panel.
+	 * This will only work for top-level objects, i.e. anything directly under the root. */
+	virtual const UNiagaraHierarchySection* GetSectionInternal() const { return nullptr; }
 private:
-	virtual void SyncChildrenViewModelsInternal() {}
+	/** View models can implement this to add or remove data. */
+	virtual void RefreshChildrenDataInternal() {}
+	/** View models can implement this to further customize the view model sync process.
+	 * An example for this is how the root view model handles sections, as sections exist outside the children hierarchy */
+	virtual void SyncViewModelsToDataInternal() {}
 	virtual void FinalizeInternal() {}
+
+	void PropagateOnChildSynced();
 protected:
 	UNiagaraHierarchyItemBase* const ItemBase;
 	/** Parent should be valid for all instances of this struct except for root objects */
 	TWeakPtr<FNiagaraHierarchyItemViewModelBase> Parent;
 	TArray<TSharedPtr<FNiagaraHierarchyItemViewModelBase>> Children;
+	TArray<FOnFilterChild> ChildFilters;
 	mutable TArray<TSharedPtr<FNiagaraHierarchyItemViewModelBase>> FilteredChildren;
 	TWeakObjectPtr<UNiagaraHierarchyViewModelBase> HierarchyViewModel;
 	FSimpleDelegate OnRequestRenameDelegate;
 	FOnSynced OnSyncedDelegate;
+	FOnChildRequestedDeletion OnChildRequestedDeletionDelegate;
 	bool bRenamePending = false;
+	bool bIsForHierarchy = false;
 };
 
 template <class DataClass, class ViewModelClass>
@@ -559,26 +740,46 @@ void FNiagaraHierarchyItemViewModelBase::GetChildrenViewModelsForType(TArray<TSh
 
 struct FNiagaraHierarchyRootViewModel : FNiagaraHierarchyItemViewModelBase
 {
-	FNiagaraHierarchyRootViewModel(UNiagaraHierarchyItemBase* InItem, TWeakObjectPtr<UNiagaraHierarchyViewModelBase> InHierarchyViewModel) : FNiagaraHierarchyItemViewModelBase(InItem, nullptr, InHierarchyViewModel) {}
+	DECLARE_DELEGATE(FOnSyncPropagated)
+	DECLARE_DELEGATE(FOnSectionsChanged)
+	DECLARE_DELEGATE_OneParam(FOnSingleSectionChanged, TSharedPtr<FNiagaraHierarchySectionViewModel> AddedSection)
 	
-	virtual ~FNiagaraHierarchyRootViewModel() override {}
+	FNiagaraHierarchyRootViewModel(UNiagaraHierarchyItemBase* InItem, TWeakObjectPtr<UNiagaraHierarchyViewModelBase> InHierarchyViewModel, bool bInIsForHierarchy) : FNiagaraHierarchyItemViewModelBase(InItem, nullptr, InHierarchyViewModel, bInIsForHierarchy) {}
+	virtual ~FNiagaraHierarchyRootViewModel() override;
 
-	virtual bool CanDrag() override { return false; }
+	virtual void Initialize() override;
+	
+	virtual FCanPerformActionResults CanDropOnInternal(TSharedPtr<FNiagaraHierarchyItemViewModelBase> DraggedItem, EItemDropZone ItemDropZone) override;
+	virtual void OnDroppedOnInternal(TSharedPtr<FNiagaraHierarchyItemViewModelBase> DroppedItem, EItemDropZone ItemDropZone) override;
 
-	virtual TOptional<EItemDropZone> OnCanAcceptDropInternal(TSharedPtr<FDragDropOperation> DragDropOp, EItemDropZone ItemDropZone) override;
-	virtual FReply OnDroppedOn(const FDragDropEvent& DragDropEvent, EItemDropZone ItemDropZone, TSharedPtr<FNiagaraHierarchyItemViewModelBase> Item) override;
+	virtual bool CanHaveChildren() const override { return true; }
 
-	TSharedPtr<struct FNiagaraHierarchySectionViewModel> AddNewSection();
+	TSharedPtr<struct FNiagaraHierarchySectionViewModel> AddSection();
+	void DeleteSection(TSharedPtr<FNiagaraHierarchyItemViewModelBase> SectionViewModel);
 	TArray<TSharedPtr<struct FNiagaraHierarchySectionViewModel>>& GetSectionViewModels() { return SectionViewModels; }
 
+	FOnSyncPropagated& OnSyncPropagated() { return OnSyncPropagatedDelegate; }
+
+	/** General purpose delegate for when sections change */
+	FOnSectionsChanged& OnSectionsChanged() { return OnSectionsChangedDelegate; }
+	/** Delegates for when a section is added or removed */
+	FOnSingleSectionChanged& OnSectionAdded() { return OnSectionAddedDelegate; }
+	FOnSingleSectionChanged& OnSectionDeleted() { return OnSectionDeletedDelegate; }
+
 private:
-	virtual void SyncChildrenViewModelsInternal() override;
+	void PropagateOnSynced();
+	virtual void SyncViewModelsToDataInternal() override;
 	TArray<TSharedPtr<struct FNiagaraHierarchySectionViewModel>> SectionViewModels;
+
+	FOnSyncPropagated OnSyncPropagatedDelegate;
+	FOnSingleSectionChanged OnSectionAddedDelegate;
+	FOnSingleSectionChanged OnSectionDeletedDelegate;
+	FOnSectionsChanged OnSectionsChangedDelegate;
 };
 
 struct FNiagaraHierarchySectionViewModel : FNiagaraHierarchyItemViewModelBase
 {
-	FNiagaraHierarchySectionViewModel(UNiagaraHierarchySection* InItem, TSharedPtr<FNiagaraHierarchyRootViewModel> InParent, TWeakObjectPtr<UNiagaraHierarchyViewModelBase> InHierarchyViewModel) : FNiagaraHierarchyItemViewModelBase(InItem, InParent, InHierarchyViewModel) {}
+	FNiagaraHierarchySectionViewModel(UNiagaraHierarchySection* InItem, TSharedPtr<FNiagaraHierarchyRootViewModel> InParent, TWeakObjectPtr<UNiagaraHierarchyViewModelBase> InHierarchyViewModel, bool bInIsForHierarchy) : FNiagaraHierarchyItemViewModelBase(InItem, InParent, InHierarchyViewModel, bInIsForHierarchy) {}
 	
 	virtual ~FNiagaraHierarchySectionViewModel() override {}
 
@@ -588,42 +789,62 @@ struct FNiagaraHierarchySectionViewModel : FNiagaraHierarchyItemViewModelBase
 	void SetSectionNameAsText(const FText& Text) const { Cast<UNiagaraHierarchySection>(ItemBase)->SetSectionNameAsText(Text); }
 	FText GetSectionNameAsText() const { return Cast<UNiagaraHierarchySection>(ItemBase)->GetSectionNameAsText(); }
 	FText GetSectionTooltip() const { return Cast<UNiagaraHierarchySection>(ItemBase)->GetTooltip(); }
-	virtual bool CanDrag() override { return false; }
-	virtual bool CanRename() override { return true; }
-	virtual bool CanDelete() override { return true; }
+	
+	void SetSectionImage(const FSlateBrush* InSectionImage) { SectionImage = InSectionImage; }
+	const FSlateBrush* GetSectionImage() const { return SectionImage; }
 
-	virtual TOptional<EItemDropZone> OnCanAcceptDropInternal(TSharedPtr<FDragDropOperation> DragDropOp, EItemDropZone ItemDropZone) override;
-	virtual FReply OnDroppedOn(const FDragDropEvent& DragDropEvent, EItemDropZone ItemDropZone, TSharedPtr<FNiagaraHierarchyItemViewModelBase> Item) override;
+protected:
+	/** Only hierarchy sections are editable */
+	virtual FCanPerformActionResults IsEditableByUser() override { return FCanPerformActionResults(IsForHierarchy()); }
+	virtual bool CanHaveChildren() const override { return false; }
+	virtual FCanPerformActionResults CanDragInternal() override { return bIsForHierarchy; }
+	/** We can only rename hierarchy sections */
+	virtual bool CanRenameInternal() override { return IsForHierarchy(); }
+	virtual void RenameInternal(FName NewName) override { GetDataMutable<UNiagaraHierarchySection>()->SetSectionName(NewName); }
+
+	virtual void PopulateDynamicContextMenuSection(FToolMenuSection& DynamicSection) override;
+
+	virtual FCanPerformActionResults CanDropOnInternal(TSharedPtr<FNiagaraHierarchyItemViewModelBase> DraggedItem, EItemDropZone ItemDropZone) override;
+	virtual void OnDroppedOnInternal(TSharedPtr<FNiagaraHierarchyItemViewModelBase> DroppedItem, EItemDropZone ItemDropZone) override;
 
 	virtual void FinalizeInternal() override;
+
+private:
+	const FSlateBrush* SectionImage = nullptr;
 };
 
 struct NIAGARAEDITOR_API FNiagaraHierarchyItemViewModel : FNiagaraHierarchyItemViewModelBase
 {
-	FNiagaraHierarchyItemViewModel(UNiagaraHierarchyItem* InItem, TSharedPtr<FNiagaraHierarchyItemViewModelBase> InParent, TWeakObjectPtr<UNiagaraHierarchyViewModelBase> InHierarchyViewModel) : FNiagaraHierarchyItemViewModelBase(InItem, InParent, InHierarchyViewModel) {}
+	FNiagaraHierarchyItemViewModel(UNiagaraHierarchyItem* InItem, TSharedPtr<FNiagaraHierarchyItemViewModelBase> InParent, TWeakObjectPtr<UNiagaraHierarchyViewModelBase> InHierarchyViewModel, bool bInIsForHierarchy) : FNiagaraHierarchyItemViewModelBase(InItem, InParent, InHierarchyViewModel, bInIsForHierarchy) {}
 	
 	virtual ~FNiagaraHierarchyItemViewModel() override {}
 
-	virtual bool CanDrag() override { return true; }
+	virtual FCanPerformActionResults IsEditableByUser() override { return FCanPerformActionResults(true); }
+	virtual bool CanHaveChildren() const override { return false; }
+	virtual FCanPerformActionResults CanDragInternal() override { return true; }
 
-	virtual TOptional<EItemDropZone> OnCanAcceptDropInternal(TSharedPtr<FDragDropOperation> DragDropOp, EItemDropZone ItemDropZone) override;
-	virtual FReply OnDroppedOn(const FDragDropEvent& DragDropEvent, EItemDropZone ItemDropZone, TSharedPtr<FNiagaraHierarchyItemViewModelBase> Item) override;
+	virtual FCanPerformActionResults CanDropOnInternal(TSharedPtr<FNiagaraHierarchyItemViewModelBase>, EItemDropZone ItemDropZone) override;
+	virtual void OnDroppedOnInternal(TSharedPtr<FNiagaraHierarchyItemViewModelBase> DroppedItem, EItemDropZone ItemDropZone) override;
 };
 
 struct FNiagaraHierarchyCategoryViewModel : FNiagaraHierarchyItemViewModelBase
 {
-	FNiagaraHierarchyCategoryViewModel(UNiagaraHierarchyCategory* InCategory, TSharedPtr<FNiagaraHierarchyItemViewModelBase> InParent, TWeakObjectPtr<UNiagaraHierarchyViewModelBase> InHierarchyViewModel) : FNiagaraHierarchyItemViewModelBase(InCategory, InParent, InHierarchyViewModel) {}
+	FNiagaraHierarchyCategoryViewModel(UNiagaraHierarchyCategory* InCategory, TSharedPtr<FNiagaraHierarchyItemViewModelBase> InParent, TWeakObjectPtr<UNiagaraHierarchyViewModelBase> InHierarchyViewModel, bool bInIsForHierarchy) : FNiagaraHierarchyItemViewModelBase(InCategory, InParent, InHierarchyViewModel, bInIsForHierarchy) {}
 	virtual ~FNiagaraHierarchyCategoryViewModel() override{}
 
-	virtual bool CanDrag() override { return true; }
-	virtual bool CanRename() override { return true; }
-
-	const UNiagaraHierarchySection* GetSection() const { return Cast<UNiagaraHierarchyCategory>(ItemBase)->GetSection(); }
+	FText GetCategoryName() const { return GetData<UNiagaraHierarchyCategory>()->GetCategoryAsText(); }
+	
+	virtual FCanPerformActionResults IsEditableByUser() override { return FCanPerformActionResults(true); }
+	virtual bool CanHaveChildren() const override { return true; }
+	virtual FCanPerformActionResults CanDragInternal() override { return true; }
+	virtual bool CanRenameInternal() override { return true; }
+	virtual void RenameInternal(FName NewName) override { GetDataMutable<UNiagaraHierarchyCategory>()->SetCategoryName(NewName); }
+	virtual const UNiagaraHierarchySection* GetSectionInternal() const override { return Cast<UNiagaraHierarchyCategory>(ItemBase)->GetSection(); }
 
 	bool IsTopCategoryActive() const;
 
-	virtual TOptional<EItemDropZone> OnCanAcceptDropInternal(TSharedPtr<FDragDropOperation> DragDropOp, EItemDropZone ItemDropZone) override;
-	virtual FReply OnDroppedOn(const FDragDropEvent& DragDropEvent, EItemDropZone ItemDropZone, TSharedPtr<FNiagaraHierarchyItemViewModelBase> Item) override;
+	virtual FCanPerformActionResults CanDropOnInternal(TSharedPtr<FNiagaraHierarchyItemViewModelBase>, EItemDropZone ItemDropZone) override;
+	virtual void OnDroppedOnInternal(TSharedPtr<FNiagaraHierarchyItemViewModelBase> DroppedItem, EItemDropZone ItemDropZone) override;
 };
 
 class FNiagaraSectionDragDropOp : public FNiagaraHierarchyDragDropOp
diff --git a/Engine/Plugins/FX/Niagara/Source/NiagaraEditor/Public/ViewModels/HierarchyEditor/NiagaraSummaryViewViewModel.h b/Engine/Plugins/FX/Niagara/Source/NiagaraEditor/Public/ViewModels/HierarchyEditor/NiagaraSummaryViewViewModel.h
new file mode 100644
--- /dev/null
+++ b/Engine/Plugins/FX/Niagara/Source/NiagaraEditor/Public/ViewModels/HierarchyEditor/NiagaraSummaryViewViewModel.h
@@ -0,0 +1,472 @@
+// Copyright Epic Games, Inc. All Rights Reserved.
+
+#pragma once
+
+#include "CoreMinimal.h"
+#include "NiagaraNodeFunctionCall.h"
+#include "NiagaraNodeAssignment.h"
+#include "NiagaraSimulationStageBase.h"
+#include "ViewModels/HierarchyEditor/NiagaraHierarchyViewModelBase.h"
+#include "NiagaraSummaryViewViewModel.generated.h"
+
+class UNiagaraStackEventHandlerPropertiesItem;
+class FNiagaraEmitterViewModel;
+
+UCLASS()
+class UNiagaraHierarchySummaryDataRefreshContext : public UNiagaraHierarchyDataRefreshContext
+{
+	GENERATED_BODY()
+
+public:
+	UPROPERTY(Transient)
+	TArray<TObjectPtr<UNiagaraRendererProperties>> Renderers;
+
+	TSharedPtr<FNiagaraEmitterViewModel> EmitterViewModel;
+};
+
+UCLASS()
+class NIAGARAEDITOR_API UNiagaraHierarchyModule : public UNiagaraHierarchyItem
+{
+	GENERATED_BODY()
+public:
+	void Initialize(const UNiagaraNodeFunctionCall& ModuleNode);
+};
+
+UCLASS()
+class NIAGARAEDITOR_API UNiagaraHierarchyModuleInput : public UNiagaraHierarchyItem
+{
+	GENERATED_BODY()
+
+public:
+	UNiagaraHierarchyModuleInput() {}
+	virtual ~UNiagaraHierarchyModuleInput() override {}
+
+	void Initialize(const UNiagaraNodeFunctionCall& FunctionCall, FGuid InputGuid);
+
+	void SetDisplayNameOverride(const FText& InText) { DisplayNameOverride = InText; }
+	FText GetDisplayNameOverride() const { return DisplayNameOverride; }
+
+	FText GetTooltipOverride() const { return TooltipOverride; }
+private:
+	/** If specified, will override how this input is presented in the stack. */
+	UPROPERTY(EditAnywhere, Category="Niagara")
+	FText DisplayNameOverride;
+
+	/** If specified, will override how the tooltip of this input in the stack. */
+	UPROPERTY(EditAnywhere, Category="Niagara")
+	FText TooltipOverride;
+};
+
+UCLASS()
+class NIAGARAEDITOR_API UNiagaraHierarchyAssignmentInput : public UNiagaraHierarchyItem
+{
+	GENERATED_BODY()
+
+public:
+	UNiagaraHierarchyAssignmentInput() {}
+	virtual ~UNiagaraHierarchyAssignmentInput() override {}
+
+	void Initialize(const UNiagaraNodeAssignment& AssignmentNode, FName AssignmentTarget);
+
+	FText GetTooltipOverride() const { return TooltipOverride; }
+private:
+	/** If specified, will override how the tooltip of this input in the stack. */
+	UPROPERTY(EditAnywhere, Category="Niagara")
+	FText TooltipOverride;
+};
+
+UCLASS()
+class NIAGARAEDITOR_API UNiagaraHierarchyRenderer : public UNiagaraHierarchyItem
+{
+	GENERATED_BODY()
+public:
+	void Initialize(const UNiagaraRendererProperties& Renderer);
+};
+
+UCLASS()
+class NIAGARAEDITOR_API UNiagaraHierarchyEventHandler : public UNiagaraHierarchyItem
+{
+	GENERATED_BODY()
+public:
+	void Initialize(const FNiagaraEventScriptProperties& EventHandler);
+};
+
+UCLASS()
+class NIAGARAEDITOR_API UNiagaraHierarchyEventHandlerProperties : public UNiagaraHierarchyItem
+{
+	GENERATED_BODY()
+public:
+	void Initialize(const FNiagaraEventScriptProperties& EventHandler);
+	
+	static FNiagaraHierarchyIdentity MakeIdentity(const FNiagaraEventScriptProperties& EventHandler);
+};
+
+UCLASS()
+class NIAGARAEDITOR_API UNiagaraHierarchySimStage : public UNiagaraHierarchyItem
+{
+	GENERATED_BODY()
+public:
+	void Initialize(const UNiagaraSimulationStageBase& SimStage);
+};
+
+UCLASS()
+class NIAGARAEDITOR_API UNiagaraHierarchySimStageProperties : public UNiagaraHierarchyItem
+{
+	GENERATED_BODY()
+public:
+	void Initialize(const UNiagaraSimulationStageBase& SimStage);
+	
+	static FNiagaraHierarchyIdentity MakeIdentity(const UNiagaraSimulationStageBase& SimStage);
+};
+
+UCLASS()
+class NIAGARAEDITOR_API UNiagaraSummaryViewViewModel : public UNiagaraHierarchyViewModelBase
+{
+	GENERATED_BODY()
+public:
+	UNiagaraSummaryViewViewModel() {}
+	virtual ~UNiagaraSummaryViewViewModel() override
+	{
+	}
+
+	void Initialize(TSharedRef<FNiagaraEmitterViewModel> EmitterViewModel);
+	virtual void FinalizeInternal() override;
+	
+	TSharedRef<FNiagaraEmitterViewModel> GetEmitterViewModel() const;
+	
+	virtual UNiagaraHierarchyRoot* GetHierarchyRoot() const override;
+	virtual TSharedPtr<FNiagaraHierarchyItemViewModelBase> CreateViewModelForData(UNiagaraHierarchyItemBase* ItemBase, TSharedPtr<FNiagaraHierarchyItemViewModelBase> Parent) override;
+	
+	virtual void PrepareSourceItems(UNiagaraHierarchyRoot* SourceRoot, TSharedPtr<FNiagaraHierarchyRootViewModel>) override;
+	virtual void SetupCommands() override;
+	
+	virtual TSharedRef<FNiagaraHierarchyDragDropOp> CreateDragDropOp(TSharedRef<FNiagaraHierarchyItemViewModelBase> Item) override;
+	
+	virtual bool SupportsDetailsPanel() override { return true; }
+	virtual TArray<TTuple<UClass*, FOnGetDetailCustomizationInstance>> GetInstanceCustomizations() override;
+
+	TMap<FGuid, UObject*> GetObjectsForProperties();
+
+	UNiagaraNodeFunctionCall* GetFunctionCallNode(const FGuid& NodeIdentity);
+	void ClearFunctionCallNodeCache(const FGuid& NodeIdentity);
+	TOptional<struct FInputData> GetInputData(const UNiagaraHierarchyModuleInput& Input);
+
+private:
+	void OnScriptGraphChanged(const FEdGraphEditAction& Action, const UNiagaraScript& Script);
+	void OnRenderersChanged();
+	void OnSimStagesChanged();
+	void OnEventHandlersChanged();
+protected:
+	// The cache is used to speed up access across different inputs, as the view models for both regular inputs & modules, dynamic inputs & assignment nodes need to 'find' these nodes which is expensive
+	TMap<FGuid, TWeakObjectPtr<UNiagaraNodeFunctionCall>> FunctionCallCache;
+	TWeakPtr<FNiagaraEmitterViewModel> EmitterViewModelWeak;
+};
+
+/** The view model for both module nodes & dynamic input nodes */
+struct FNiagaraFunctionViewModel : public FNiagaraHierarchyItemViewModel
+{
+	FNiagaraFunctionViewModel(UNiagaraHierarchyModule* HierarchyModule, TSharedPtr<FNiagaraHierarchyItemViewModelBase> InParent, TWeakObjectPtr<UNiagaraSummaryViewViewModel> ViewModel, bool bInIsForHierarchy)
+	: FNiagaraHierarchyItemViewModel(HierarchyModule, InParent, ViewModel, bInIsForHierarchy)
+	{ }
+
+	virtual ~FNiagaraFunctionViewModel() override;
+	
+	TWeakObjectPtr<UNiagaraNodeFunctionCall> GetFunctionCallNode() const;	
+	
+	bool IsFromBaseEmitter() const;
+
+	void SetSection(UNiagaraHierarchySection& InSection) { Section = &InSection; }
+
+	bool IsDynamicInput() const { return bIsDynamicInput; }
+private:	
+	virtual void Initialize() override;
+	virtual void RefreshChildrenDataInternal() override;
+	void RefreshChildrenInputs(bool bClearCache = false) const;
+	
+	virtual FString ToString() const override;
+
+	virtual FCanPerformActionResults IsEditableByUser() override;
+	virtual bool CanHaveChildren() const override { return bIsForHierarchy == false; }
+	virtual bool CanRenameInternal() override { return false; }
+	virtual FCanPerformActionResults CanDropOnInternal(TSharedPtr<FNiagaraHierarchyItemViewModelBase>, EItemDropZone ItemDropZone) override;
+	virtual const UNiagaraHierarchySection* GetSectionInternal() const override;
+	
+	virtual bool RepresentsExternalData() const override { return true; }
+	virtual bool DoesExternalDataStillExist(const UNiagaraHierarchyDataRefreshContext* Context) const override { ClearCache(); return GetFunctionCallNode().IsValid(); }
+	
+	void OnScriptApplied(UNiagaraScript* NiagaraScript, FGuid Guid);
+
+	void ClearCache() const;
+private:
+	FDelegateHandle OnScriptAppliedHandle;
+	mutable TOptional<bool> IsFromBaseEmitterCache;
+	TWeakObjectPtr<UNiagaraHierarchySection> Section;
+	bool bIsDynamicInput = false;
+};
+
+struct FInputData
+{
+	FName InputName;
+	FNiagaraTypeDefinition Type;
+	FNiagaraVariableMetaData MetaData;
+	bool bIsStatic = false;
+	TArray<FGuid> ChildrenInputGuids;
+	UNiagaraNodeFunctionCall* FunctionCallNode;
+};
+
+struct FNiagaraModuleInputViewModel : public FNiagaraHierarchyItemViewModel
+{	
+	FNiagaraModuleInputViewModel(UNiagaraHierarchyModuleInput* ModuleInput, TSharedPtr<FNiagaraHierarchyItemViewModelBase> InParent, TWeakObjectPtr<UNiagaraSummaryViewViewModel> ViewModel, bool bInIsForHierarchy)
+	: FNiagaraHierarchyItemViewModel(ModuleInput, InParent, ViewModel, bInIsForHierarchy) {}
+
+public:
+	TOptional<FInputData> GetInputData() const;
+	FText GetSummaryInputNameOverride() const;
+	
+protected:
+	TWeakObjectPtr<UNiagaraNodeFunctionCall> GetModuleNode() const;
+
+	virtual FString ToString() const override;
+	virtual TArray<FString> GetSearchTerms() const override;
+	
+	bool IsFromBaseEmitter() const;
+
+	void ClearCache() const;
+
+	void RefreshChildDynamicInputs(bool bClearCache = false);
+	
+	virtual bool CanHaveChildren() const override;
+	virtual FCanPerformActionResults IsEditableByUser() override;
+	virtual bool RepresentsExternalData() const override { return true; }
+	virtual bool DoesExternalDataStillExist(const UNiagaraHierarchyDataRefreshContext* Context) const override { ClearCache(); return GetInputData().IsSet(); }
+	virtual void RefreshChildrenDataInternal() override;
+
+	virtual FCanPerformActionResults CanDropOnInternal(TSharedPtr<FNiagaraHierarchyItemViewModelBase>, EItemDropZone ItemDropZone) override;
+	virtual void OnDroppedOnInternal(TSharedPtr<FNiagaraHierarchyItemViewModelBase> DroppedItem, EItemDropZone ItemDropZone) override;
+
+	virtual void PopulateDynamicContextMenuSection(FToolMenuSection& DynamicSection) override;
+private:
+	TOptional<FInputData> FindInputDataInternal() const;
+	void AddNativeChildrenInputs();
+	bool CanAddNativeChildrenInputs() const;
+	TArray<FNiagaraHierarchyIdentity> GetNativeChildInputIdentities() const;
+private:
+	mutable TOptional<FInputData> InputDataCache;
+	mutable TOptional<bool> IsFromBaseEmitterCache;
+};
+
+struct FNiagaraAssignmentInputViewModel : public FNiagaraHierarchyItemViewModel
+{	
+	FNiagaraAssignmentInputViewModel(UNiagaraHierarchyAssignmentInput* ModuleInput, TSharedPtr<FNiagaraHierarchyItemViewModelBase> InParent, TWeakObjectPtr<UNiagaraSummaryViewViewModel> ViewModel, bool bInIsForHierarchy)
+	: FNiagaraHierarchyItemViewModel(ModuleInput, InParent, ViewModel, bInIsForHierarchy) {}
+
+	virtual FCanPerformActionResults CanDropOnInternal(TSharedPtr<FNiagaraHierarchyItemViewModelBase>, EItemDropZone ItemDropZone) override;
+	
+	TWeakObjectPtr<UNiagaraNodeAssignment> GetAssignmentNode() const;
+	TOptional<FNiagaraStackGraphUtilities::FMatchingFunctionInputData> GetInputData() const;
+
+	virtual bool CanHaveChildren() const override { return false; }
+
+	virtual FString ToString() const override;
+	virtual TArray<FString> GetSearchTerms() const override;
+	
+	bool IsFromBaseEmitter() const;
+
+	void ClearCache() const;
+protected:
+	virtual FCanPerformActionResults IsEditableByUser() override;
+	virtual bool RepresentsExternalData() const override { return true; }
+	virtual bool DoesExternalDataStillExist(const UNiagaraHierarchyDataRefreshContext* Context) const override { ClearCache(); return GetInputData().IsSet(); }
+private:
+	TOptional<FNiagaraStackGraphUtilities::FMatchingFunctionInputData> FindInputDataInternal() const;
+private:
+	mutable TWeakObjectPtr<UNiagaraNodeAssignment> AssignmentNodeCache;
+	mutable TOptional<FNiagaraStackGraphUtilities::FMatchingFunctionInputData> InputDataCache;
+	mutable TOptional<bool> IsFromBaseEmitterCache;
+};
+
+struct FNiagaraHierarchySummaryCategoryViewModel : public FNiagaraHierarchyCategoryViewModel
+{
+	FNiagaraHierarchySummaryCategoryViewModel(UNiagaraHierarchyCategory* Category, TSharedPtr<FNiagaraHierarchyItemViewModelBase> InParent, TWeakObjectPtr<UNiagaraSummaryViewViewModel> ViewModel, bool bInIsForHierarchy)
+	: FNiagaraHierarchyCategoryViewModel(Category, InParent, ViewModel, bInIsForHierarchy) {}
+	
+	bool IsFromBaseEmitter() const;
+	
+protected:
+	virtual FCanPerformActionResults IsEditableByUser() override;
+	
+	mutable TOptional<bool> IsFromBaseEmitterCache;
+};
+
+struct FNiagaraHierarchyPropertyViewModel : public FNiagaraHierarchyItemViewModel
+{
+	FNiagaraHierarchyPropertyViewModel(UNiagaraHierarchyObjectProperty* ObjectProperty, TSharedPtr<FNiagaraHierarchyItemViewModelBase> InParent, TWeakObjectPtr<UNiagaraSummaryViewViewModel> ViewModel, bool bInIsForHierarchy)
+		: FNiagaraHierarchyItemViewModel(ObjectProperty, InParent, ViewModel, bInIsForHierarchy) {}
+
+	virtual FString ToString() const override;
+
+	bool IsFromBaseEmitter() const;
+protected:
+	virtual bool RepresentsExternalData() const override { return true; }
+	virtual bool DoesExternalDataStillExist(const UNiagaraHierarchyDataRefreshContext* Context) const override;
+	
+	virtual FCanPerformActionResults IsEditableByUser() override;
+	
+	mutable TOptional<bool> IsFromBaseEmitterCache;	
+};
+
+struct FNiagaraHierarchyRendererViewModel : public FNiagaraHierarchyItemViewModel
+{
+	FNiagaraHierarchyRendererViewModel(UNiagaraHierarchyRenderer* Renderer, TSharedPtr<FNiagaraHierarchyItemViewModelBase> InParent, TWeakObjectPtr<UNiagaraSummaryViewViewModel> ViewModel, bool bInIsForHierarchy)
+	: FNiagaraHierarchyItemViewModel(Renderer, InParent, ViewModel, bInIsForHierarchy) {}
+
+	virtual FString ToString() const override;
+	virtual bool CanRenameInternal() override { return false; }
+
+	UNiagaraRendererProperties* GetRendererProperties() const;
+
+	void SetSection(UNiagaraHierarchySection& InSection) { Section = &InSection; }
+
+	bool IsFromBaseEmitter() const;
+protected:
+	virtual void RefreshChildrenDataInternal() override;
+	
+	virtual FCanPerformActionResults IsEditableByUser() override;
+	virtual bool CanHaveChildren() const override { return bIsForHierarchy == false; }
+	
+	virtual bool RepresentsExternalData() const override { return true; }
+	virtual bool DoesExternalDataStillExist(const UNiagaraHierarchyDataRefreshContext* Context) const override { return GetRendererProperties() != nullptr;}
+
+	virtual const UNiagaraHierarchySection* GetSectionInternal() const override;
+	
+	TWeakObjectPtr<UNiagaraHierarchySection> Section;
+	mutable TOptional<bool> IsFromBaseEmitterCache;
+};
+
+struct FNiagaraHierarchyEventHandlerViewModel : public FNiagaraHierarchyItemViewModel
+{
+	FNiagaraHierarchyEventHandlerViewModel(UNiagaraHierarchyEventHandler* EventHandler, TSharedPtr<FNiagaraHierarchyItemViewModelBase> InParent, TWeakObjectPtr<UNiagaraSummaryViewViewModel> ViewModel, bool bInIsForHierarchy)
+	: FNiagaraHierarchyItemViewModel(EventHandler, InParent, ViewModel, bInIsForHierarchy) {}
+
+	virtual FString ToString() const override;
+	virtual bool CanRenameInternal() override { return false; }
+
+	FNiagaraEventScriptProperties* GetEventScriptProperties() const;
+
+	void SetSection(UNiagaraHierarchySection& InSection) { Section = &InSection; }
+
+	bool IsFromBaseEmitter() const;
+protected:
+	virtual void RefreshChildrenDataInternal() override;
+	
+	virtual FCanPerformActionResults IsEditableByUser() override;
+	virtual bool CanHaveChildren() const override { return bIsForHierarchy == false; }
+	
+	virtual bool RepresentsExternalData() const override { return true; }
+	virtual bool DoesExternalDataStillExist(const UNiagaraHierarchyDataRefreshContext* Context) const override { return GetEventScriptProperties() != nullptr;}
+
+	virtual const UNiagaraHierarchySection* GetSectionInternal() const override;
+	
+	TWeakObjectPtr<UNiagaraHierarchySection> Section;
+	mutable TOptional<bool> IsFromBaseEmitterCache;
+};
+
+struct FNiagaraHierarchyEventHandlerPropertiesViewModel : public FNiagaraHierarchyItemViewModel
+{
+	FNiagaraHierarchyEventHandlerPropertiesViewModel(UNiagaraHierarchyEventHandlerProperties* EventHandlerProperties, TSharedPtr<FNiagaraHierarchyItemViewModelBase> InParent, TWeakObjectPtr<UNiagaraSummaryViewViewModel> ViewModel, bool bInIsForHierarchy)
+	: FNiagaraHierarchyItemViewModel(EventHandlerProperties, InParent, ViewModel, bInIsForHierarchy) {}
+
+	virtual FString ToString() const override;
+	virtual bool CanRenameInternal() override { return false; }
+
+	FNiagaraEventScriptProperties* GetEventScriptProperties() const;
+
+	bool IsFromBaseEmitter() const;
+protected:
+	virtual void RefreshChildrenDataInternal() override;
+
+	virtual FCanPerformActionResults IsEditableByUser() override;
+	virtual bool CanHaveChildren() const override { return bIsForHierarchy == false; }
+	
+	virtual bool RepresentsExternalData() const override { return true; }
+	virtual bool DoesExternalDataStillExist(const UNiagaraHierarchyDataRefreshContext* Context) const override { return GetEventScriptProperties() != nullptr;}
+
+	mutable TOptional<bool> IsFromBaseEmitterCache;
+};
+
+struct FNiagaraHierarchySimStageViewModel : public FNiagaraHierarchyItemViewModel
+{
+	FNiagaraHierarchySimStageViewModel(UNiagaraHierarchySimStage* SimStage, TSharedPtr<FNiagaraHierarchyItemViewModelBase> InParent, TWeakObjectPtr<UNiagaraSummaryViewViewModel> ViewModel, bool bInIsForHierarchy)
+	: FNiagaraHierarchyItemViewModel(SimStage, InParent, ViewModel, bInIsForHierarchy) {}
+
+	virtual FString ToString() const override;
+	virtual bool CanRenameInternal() override { return false; }
+
+	UNiagaraSimulationStageBase* GetSimStage() const;
+
+	void SetSection(UNiagaraHierarchySection& InSection) { Section = &InSection; }
+
+	bool IsFromBaseEmitter() const;
+protected:
+	virtual void RefreshChildrenDataInternal() override;
+	
+	virtual FCanPerformActionResults IsEditableByUser() override;
+	virtual bool CanHaveChildren() const override { return bIsForHierarchy == false; }
+	
+	virtual bool RepresentsExternalData() const override { return true; }
+	virtual bool DoesExternalDataStillExist(const UNiagaraHierarchyDataRefreshContext* Context) const override { return GetSimStage() != nullptr;}
+
+	virtual const UNiagaraHierarchySection* GetSectionInternal() const override;
+	
+	TWeakObjectPtr<UNiagaraHierarchySection> Section;
+	mutable TOptional<bool> IsFromBaseEmitterCache;
+};
+
+struct FNiagaraHierarchySimStagePropertiesViewModel : public FNiagaraHierarchyItemViewModel
+{
+	FNiagaraHierarchySimStagePropertiesViewModel(UNiagaraHierarchySimStageProperties* SimStage, TSharedPtr<FNiagaraHierarchyItemViewModelBase> InParent, TWeakObjectPtr<UNiagaraSummaryViewViewModel> ViewModel, bool bInIsForHierarchy)
+	: FNiagaraHierarchyItemViewModel(SimStage, InParent, ViewModel, bInIsForHierarchy) {}
+
+	virtual FString ToString() const override;
+	virtual bool CanRenameInternal() override { return false; }
+
+	UNiagaraSimulationStageBase* GetSimStage() const;
+
+	bool IsFromBaseEmitter() const;
+protected:
+	virtual void RefreshChildrenDataInternal() override;
+
+	virtual FCanPerformActionResults IsEditableByUser() override;
+	virtual bool CanHaveChildren() const override { return bIsForHierarchy == false; }
+	
+	virtual bool RepresentsExternalData() const override { return true; }
+	virtual bool DoesExternalDataStillExist(const UNiagaraHierarchyDataRefreshContext* Context) const override { return GetSimStage() != nullptr;}
+
+	mutable TOptional<bool> IsFromBaseEmitterCache;
+};
+
+class FNiagaraHierarchyInputParameterHierarchyDragDropOp : public FNiagaraHierarchyDragDropOp
+{
+public:
+	DRAG_DROP_OPERATOR_TYPE(FNiagaraInputParameterHierarchyDragDropOp, FNiagaraHierarchyDragDropOp)
+
+	FNiagaraHierarchyInputParameterHierarchyDragDropOp(TSharedPtr<FNiagaraModuleInputViewModel> InputViewModel) : FNiagaraHierarchyDragDropOp(InputViewModel) {}
+	
+	virtual TSharedRef<SWidget> CreateCustomDecorator() const override;
+};
+
+class SNiagaraHierarchyModule : public SCompoundWidget
+{
+	SLATE_BEGIN_ARGS(SNiagaraHierarchyModule)
+		{}
+	SLATE_END_ARGS()
+
+public:
+	void Construct(const FArguments& InArgs, TSharedPtr<struct FNiagaraFunctionViewModel> InModuleViewModel);
+		
+	FText GetModuleDisplayName() const;
+	
+private:
+	TWeakPtr<struct FNiagaraFunctionViewModel> ModuleViewModel;
+	TSharedPtr<SInlineEditableTextBlock> InlineEditableTextBlock;
+};
\ No newline at end of file
diff --git a/Engine/Plugins/FX/Niagara/Source/NiagaraEditor/Public/ViewModels/HierarchyEditor/NiagaraUserParametersHierarchyViewModel.h b/Engine/Plugins/FX/Niagara/Source/NiagaraEditor/Public/ViewModels/HierarchyEditor/NiagaraUserParametersHierarchyViewModel.h
--- a/Engine/Plugins/FX/Niagara/Source/NiagaraEditor/Public/ViewModels/HierarchyEditor/NiagaraUserParametersHierarchyViewModel.h
+++ b/Engine/Plugins/FX/Niagara/Source/NiagaraEditor/Public/ViewModels/HierarchyEditor/NiagaraUserParametersHierarchyViewModel.h
@@ -6,6 +6,19 @@
 #include "NiagaraScriptVariable.h"
 #include "NiagaraUserParametersHierarchyViewModel.generated.h"
 
+UCLASS()
+class UNiagaraHierarchyUserParameterRefreshContext : public UNiagaraHierarchyDataRefreshContext
+{
+	GENERATED_BODY()
+
+public:
+	void SetSystem(UNiagaraSystem* InSystem) { System = InSystem; }
+	const UNiagaraSystem* GetSystem() const { return System; }
+private:
+	UPROPERTY(Transient)
+	TObjectPtr<UNiagaraSystem> System = nullptr;
+};
+
 UCLASS()
 class UNiagaraHierarchyUserParameter : public UNiagaraHierarchyItem
 {
@@ -16,8 +29,6 @@ public:
 	virtual ~UNiagaraHierarchyUserParameter() override {}
 	
 	void Initialize(UNiagaraScriptVariable& InUserParameterScriptVariable, UNiagaraSystem& System);
-
-	virtual void RefreshDataInternal() override;
 	
 	const FNiagaraVariable& GetUserParameter() const { return UserParameterScriptVariable->Variable; }
 
@@ -25,18 +36,18 @@ public:
 private:
 	UPROPERTY()
 	TObjectPtr<UNiagaraScriptVariable> UserParameterScriptVariable;
-
-	UPROPERTY()
-	TObjectPtr<UNiagaraSystem> System;
 };
 
 struct FNiagaraHierarchyUserParameterViewModel : public FNiagaraHierarchyItemViewModel
 {
-	FNiagaraHierarchyUserParameterViewModel(UNiagaraHierarchyUserParameter* InItem, TSharedPtr<FNiagaraHierarchyItemViewModelBase> InParent, TWeakObjectPtr<UNiagaraHierarchyViewModelBase> InHierarchyViewModel)
-	: FNiagaraHierarchyItemViewModel(InItem, InParent, InHierarchyViewModel) {}
+	FNiagaraHierarchyUserParameterViewModel(UNiagaraHierarchyUserParameter* InItem, TSharedPtr<FNiagaraHierarchyItemViewModelBase> InParent, TWeakObjectPtr<UNiagaraHierarchyViewModelBase> InHierarchyViewModel, bool bInIsForHierarchy)
+	: FNiagaraHierarchyItemViewModel(InItem, InParent, InHierarchyViewModel, bInIsForHierarchy) {}
 	
 	/** For editing in the details panel we want to handle the script variable that is represented by the hierarchy item, not the hierarchy item itself. */
 	virtual UObject* GetDataForEditing() override;
+
+	virtual bool RepresentsExternalData() const override { return true; }
+	virtual bool DoesExternalDataStillExist(const UNiagaraHierarchyDataRefreshContext* Context) const override;
 };
 
 UCLASS()
@@ -53,10 +64,10 @@ public:
 
 	void Initialize(TSharedRef<FNiagaraSystemViewModel> InSystemViewModel);
 	
-	virtual UNiagaraHierarchyRoot* GetHierarchyDataRoot() const override;
+	virtual UNiagaraHierarchyRoot* GetHierarchyRoot() const override;
 	virtual TSharedPtr<FNiagaraHierarchyItemViewModelBase> CreateViewModelForData(UNiagaraHierarchyItemBase* ItemBase, TSharedPtr<FNiagaraHierarchyItemViewModelBase> Parent) override;
 	
-	virtual void PrepareSourceItems() override;
+	virtual void PrepareSourceItems(UNiagaraHierarchyRoot* SourceRoot, TSharedPtr<FNiagaraHierarchyRootViewModel>) override;
 	virtual void SetupCommands() override;
 	
 	virtual TSharedRef<FNiagaraHierarchyDragDropOp> CreateDragDropOp(TSharedRef<FNiagaraHierarchyItemViewModelBase> Item) override;
diff --git a/Engine/Plugins/FX/Niagara/Source/NiagaraEditor/Public/ViewModels/NiagaraEmitterViewModel.h b/Engine/Plugins/FX/Niagara/Source/NiagaraEditor/Public/ViewModels/NiagaraEmitterViewModel.h
--- a/Engine/Plugins/FX/Niagara/Source/NiagaraEditor/Public/ViewModels/NiagaraEmitterViewModel.h
+++ b/Engine/Plugins/FX/Niagara/Source/NiagaraEditor/Public/ViewModels/NiagaraEmitterViewModel.h
@@ -24,6 +24,7 @@ class FNiagaraEmitterHandleViewModel;
 struct FNiagaraEventScriptProperties;
 class UNiagaraScriptVariable;
 class UNiagaraParameterDefinitions;
+class UNiagaraSummaryViewViewModel;
 
 
 /** The view model for the UNiagaraEmitter objects */
@@ -31,6 +32,7 @@ class FNiagaraEmitterViewModel
 	: public TSharedFromThis<FNiagaraEmitterViewModel>
 	, public TNiagaraViewModelManager<UNiagaraEmitter, FNiagaraEmitterViewModel>
 	, public INiagaraParameterDefinitionsSubscriberViewModel
+	, public FGCObject
 {
 public:
 	DECLARE_MULTICAST_DELEGATE(FOnEmitterChanged);
@@ -90,7 +92,9 @@ public:
 	
 	/** Geta a view model for the update/spawn Script. */
 	TSharedRef<FNiagaraScriptViewModel> GetSharedScriptViewModel();
-	
+
+	NIAGARAEDITOR_API UNiagaraSummaryViewViewModel* GetSummaryHierarchyViewModel();
+
 	/* Get the latest status of this view-model's script compilation.*/
 	ENiagaraScriptCompileStatus GetLatestCompileStatus();
 
@@ -121,9 +125,6 @@ public:
 	/** Add an event script to the owned emitter. Sets the Usage, UsageID and Source of the EventScriptProperties. */
 	NIAGARAEDITOR_API void AddEventHandler(FNiagaraEventScriptProperties& EventScriptProperties, bool bResetGraphForOutput = false);
 
-	bool GetSummaryIsInEditMode() const { return bSummaryIsInEditMode; }
-	void SetSummaryIsInEditMode(bool bInSummaryIsInEditMode) { bSummaryIsInEditMode = bInSummaryIsInEditMode; }
-
 	void Cleanup();
 
 	/** When the parent emitter version is changed, this is set to generate the changelist in the stack. */
@@ -149,6 +150,11 @@ private:
 
 	void OnEmitterPropertiesChanged();
 
+	void OnSummaryViewHierarchyChanged();
+	
+	virtual FString GetReferencerName() const override;
+	virtual void AddReferencedObjects(FReferenceCollector& Collector) override;
+	
 	/** The text format stats display .*/
 	static const FText StatsFormat;
 
@@ -167,6 +173,8 @@ private:
 	/** The view model for the update/spawn/event script. */
 	TSharedPtr<FNiagaraScriptViewModel> SharedScriptViewModel;
 
+	UNiagaraSummaryViewViewModel* SummaryViewHierarchyViewModel;
+
 	/** A flag to prevent reentrancy when updating selection sets. */
 	bool bUpdatingSelectionInternally;
 
@@ -197,6 +205,4 @@ private:
 	TMap<FObjectKey, FDelegateHandle> ScriptToOnParameterStoreChangedHandleMap;
 
 	TSharedPtr<SWindow> NewParentWindow;
-
-	bool bSummaryIsInEditMode;
 };
\ No newline at end of file
diff --git a/Engine/Plugins/FX/Niagara/Source/NiagaraEditor/Public/ViewModels/NiagaraParameterPanelViewModel.h b/Engine/Plugins/FX/Niagara/Source/NiagaraEditor/Public/ViewModels/NiagaraParameterPanelViewModel.h
--- a/Engine/Plugins/FX/Niagara/Source/NiagaraEditor/Public/ViewModels/NiagaraParameterPanelViewModel.h
+++ b/Engine/Plugins/FX/Niagara/Source/NiagaraEditor/Public/ViewModels/NiagaraParameterPanelViewModel.h
@@ -262,9 +262,6 @@ public:
 	const TArray<FNiagaraParameterPanelItem>& GetCachedViewedParameterItems() const;
 	
 	void SelectParameterItemByName(const FName ParameterName, const bool bRequestRename);
-
-	// TODO (me) Remove again when the user parameter tab goes live, since this is only used for the "Edit Hierarchy" button
-	virtual TSharedRef<SWidget> GenerateAdjacentWidget() { return SNullWidget::NullWidget; }
 	
 	FOnParameterPanelViewModelExternalSelectionChanged& GetOnExternalSelectionChangedDelegate() { return OnParameterPanelViewModelExternalSelectionChangedDelegate; };
 
diff --git a/Engine/Plugins/FX/Niagara/Source/NiagaraEditor/Public/ViewModels/Stack/NiagaraStackEmitterSettingsGroup.h b/Engine/Plugins/FX/Niagara/Source/NiagaraEditor/Public/ViewModels/Stack/NiagaraStackEmitterSettingsGroup.h
--- a/Engine/Plugins/FX/Niagara/Source/NiagaraEditor/Public/ViewModels/Stack/NiagaraStackEmitterSettingsGroup.h
+++ b/Engine/Plugins/FX/Niagara/Source/NiagaraEditor/Public/ViewModels/Stack/NiagaraStackEmitterSettingsGroup.h
@@ -3,6 +3,7 @@
 #pragma once
 
 #include "NiagaraEmitter.h"
+#include "NiagaraStackSummaryViewInputCollection.h"
 #include "ViewModels/Stack/NiagaraStackScriptItemGroup.h"
 #include "ViewModels/Stack/NiagaraStackItem.h"
 #include "NiagaraStackEmitterSettingsGroup.generated.h"
@@ -10,7 +11,6 @@
 class FNiagaraEmitterViewModel;
 class FNiagaraScriptViewModel;
 class UNiagaraStackObject;
-class UNiagaraStackSummaryViewObject;
 class IDetailTreeNode;
 
 UCLASS()
@@ -70,8 +70,10 @@ public:
 	virtual bool GetShouldShowInOverview() const override { return false; }
 
 	virtual bool SupportsEditMode() const override { return true; }
-	virtual bool GetEditModeIsActive() const override;
-	virtual void SetEditModeIsActive(bool bInEditModeIsActive) override;
+	virtual void OnEditButtonClicked() override;
+	virtual TOptional<FText> GetEditModeButtonText() const override;
+	virtual TOptional<FText> GetEditModeButtonTooltip() const override;
+	virtual EVisibility IsEditButtonVisible() const override;
 
 protected:
 
@@ -84,10 +86,7 @@ private:
 	FVersionedNiagaraEmitterWeakPtr Emitter;
 
 	UPROPERTY()
-	TObjectPtr<UNiagaraStackSummaryViewObject> FilteredObject;
-
-	UPROPERTY()
-	TObjectPtr<UNiagaraStackObject> SummaryEditorData;
+	TObjectPtr<UNiagaraStackSummaryViewCollection> FilteredObject;
 };
 
 UCLASS()
diff --git a/Engine/Plugins/FX/Niagara/Source/NiagaraEditor/Public/ViewModels/Stack/NiagaraStackEntry.h b/Engine/Plugins/FX/Niagara/Source/NiagaraEditor/Public/ViewModels/Stack/NiagaraStackEntry.h
--- a/Engine/Plugins/FX/Niagara/Source/NiagaraEditor/Public/ViewModels/Stack/NiagaraStackEntry.h
+++ b/Engine/Plugins/FX/Niagara/Source/NiagaraEditor/Public/ViewModels/Stack/NiagaraStackEntry.h
@@ -519,6 +519,13 @@ public:
 	/** Handler for when a rename is committed for this stack entry. */
 	virtual void OnRenamed(FText NewName);
 
+	virtual bool SupportsSummaryView() const { return false; }
+	virtual struct FNiagaraHierarchyIdentity DetermineSummaryIdentity() const;
+	bool IsInSummaryView() const;
+	bool IsAnyParentInSummaryView() const;
+	bool IsAnyChildInSummaryView(bool bRecursive = false) const;
+	bool ExistsInParentEmitterSummary() const;
+
 	virtual EIconMode GetSupportedIconMode() const { return EIconMode::None; }
 
 	virtual const FSlateBrush* GetIconBrush() const { return nullptr; }
diff --git a/Engine/Plugins/FX/Niagara/Source/NiagaraEditor/Public/ViewModels/Stack/NiagaraStackEventScriptItemGroup.h b/Engine/Plugins/FX/Niagara/Source/NiagaraEditor/Public/ViewModels/Stack/NiagaraStackEventScriptItemGroup.h
--- a/Engine/Plugins/FX/Niagara/Source/NiagaraEditor/Public/ViewModels/Stack/NiagaraStackEventScriptItemGroup.h
+++ b/Engine/Plugins/FX/Niagara/Source/NiagaraEditor/Public/ViewModels/Stack/NiagaraStackEventScriptItemGroup.h
@@ -52,6 +52,8 @@ protected:
 
 	virtual void RefreshChildrenInternal(const TArray<UNiagaraStackEntry*>& CurrentChildren, TArray<UNiagaraStackEntry*>& NewChildren, TArray<FStackIssue>& NewIssues) override;
 
+	virtual bool SupportsSummaryView() const override { return true; }
+	virtual FNiagaraHierarchyIdentity DetermineSummaryIdentity() const override;
 private:
 	void EventHandlerPropertiesChanged();
 
@@ -110,6 +112,8 @@ protected:
 
 	virtual void RefreshChildrenInternal(const TArray<UNiagaraStackEntry*>& CurrentChildren, TArray<UNiagaraStackEntry*>& NewChildren, TArray<FStackIssue>& NewIssues) override;
 
+	virtual bool SupportsSummaryView() const override { return true; }
+	virtual FNiagaraHierarchyIdentity DetermineSummaryIdentity() const override;
 private:
 	bool HasBaseEventHandler() const;
 	
diff --git a/Engine/Plugins/FX/Niagara/Source/NiagaraEditor/Public/ViewModels/Stack/NiagaraStackFunctionInput.h b/Engine/Plugins/FX/Niagara/Source/NiagaraEditor/Public/ViewModels/Stack/NiagaraStackFunctionInput.h
--- a/Engine/Plugins/FX/Niagara/Source/NiagaraEditor/Public/ViewModels/Stack/NiagaraStackFunctionInput.h
+++ b/Engine/Plugins/FX/Niagara/Source/NiagaraEditor/Public/ViewModels/Stack/NiagaraStackFunctionInput.h
@@ -112,6 +112,8 @@ public:
 	virtual FText GetPasteTransactionText(const UNiagaraClipboardContent* ClipboardContent) const override;
 	virtual void Paste(const UNiagaraClipboardContent* ClipboardContent, FText& OutPasteWarning) override;
 	virtual bool HasOverridenContent() const override;
+	virtual bool SupportsSummaryView() const override;
+	virtual FNiagaraHierarchyIdentity DetermineSummaryIdentity() const override;
 	
 	/** Gets the tooltip that should be shown for the value of this input. */
 	FText GetValueToolTip() const;
@@ -119,7 +121,8 @@ public:
 	/** Gets the tooltip that should be shown for the value of this input. */
 	FText GetCollapsedStateText() const;
 
-	void SetSummaryViewDisiplayName(TOptional<FText> InDisplayName);
+	void SetSummaryViewDisplayName(TAttribute<FText> InDisplayName);
+	void SetSummaryViewTooltip(TAttribute<FText> InTooltipOverride);
 
 	/** Gets the path of parameter handles from the owning module to the function call which owns this input. */
 	const TArray<FNiagaraParameterHandle>& GetInputParameterHandlePath() const;
@@ -280,10 +283,6 @@ public:
 	//~ UNiagaraStackEntry interface
 	virtual void GetSearchItems(TArray<FStackSearchItem>& SearchItems) const override;
 
-	TOptional<FNiagaraVariableMetaData> GetMetadata() const;
-	TOptional<FGuid> GetMetadataGuid() const;
-
-
 	virtual const FCollectedUsageData& GetCollectedUsageData() const override;
 
 	bool OpenSourceAsset() const;
@@ -427,7 +426,8 @@ private:
 
 	/** Optional override for the display name*/
 	TOptional<FText> DisplayNameOverride;
-	TOptional<FText> SummaryViewDisplayNameOverride;
+	TOptional<TAttribute<FText>> SummaryViewDisplayNameOverride;
+	TOptional<TAttribute<FText>> SummaryViewTooltipOverride;
 
 	/** The default value for this input defined in the defining script. */
 	FInputValues DefaultInputValues;
diff --git a/Engine/Plugins/FX/Niagara/Source/NiagaraEditor/Public/ViewModels/Stack/NiagaraStackFunctionInputCollection.h b/Engine/Plugins/FX/Niagara/Source/NiagaraEditor/Public/ViewModels/Stack/NiagaraStackFunctionInputCollection.h
--- a/Engine/Plugins/FX/Niagara/Source/NiagaraEditor/Public/ViewModels/Stack/NiagaraStackFunctionInputCollection.h
+++ b/Engine/Plugins/FX/Niagara/Source/NiagaraEditor/Public/ViewModels/Stack/NiagaraStackFunctionInputCollection.h
@@ -12,16 +12,16 @@
 class UNiagaraClipboardFunctionInput;
 class UEdGraphPin;
 
-
-
+/** A base class for value collections. Values can be all kinds of input data, such as module inputs, object properties etc.
+ * This has base functionality for sections.
+ */
 UCLASS()
-class NIAGARAEDITOR_API UNiagaraStackFunctionInputCollectionBase : public UNiagaraStackItemContent
+class NIAGARAEDITOR_API UNiagaraStackValueCollection : public UNiagaraStackItemContent
 {
 	GENERATED_BODY()
-
 public:
 	void Initialize(FRequiredEntryData InRequiredEntryData, FString InOwningStackItemEditorDataKey, FString InStackEditorDataKey);
-
+	
 	bool GetShouldDisplayLabel() const { return bShouldDisplayLabel; }
 	void SetShouldDisplayLabel(bool bInShouldShowLabel);
 
@@ -31,89 +31,32 @@
 
 	void SetActiveSection(FText InActiveSection);
 
-	virtual bool GetCanExpand() const override;
-	virtual bool GetShouldShowInStack() const override;
-
+	FText GetTooltipForSection(FString Section) const;
 public:
 	static FText UncategorizedName;
 
 	static FText AllSectionName;
-
 protected:
-
-	struct FInputData
-	{
-		FNiagaraVariable InputVariable;
-		int32 SortKey;
-		TOptional<FText> DisplayName;
-		FText Category;
-		bool bIsStatic;
-		bool bIsHidden;
-		bool bShouldShowInSummary;
-
-		UNiagaraNodeFunctionCall* ModuleNode;
-		UNiagaraNodeFunctionCall* InputFunctionCallNode;
-		
-		TArray<FInputData*> Children;
-		bool bIsChild = false;
-	};
-
-	struct FNiagaraParentData
-	{
-		FNiagaraVariable ParentVariable;
-		TArray<int32> ChildIndices;
-	};
-
-	struct FFunctionCallNodesState
-	{		
-		TArray<FInputData> InputDataCollection;
-		TMap<FName, FNiagaraParentData> ParentMapping;
-	};
-	
 	virtual void RefreshChildrenInternal(const TArray<UNiagaraStackEntry*>& CurrentChildren, TArray<UNiagaraStackEntry*>& NewChildren, TArray<FStackIssue>& NewIssues) override;
-
-	virtual int32 GetChildIndentLevel() const override;
-
-	void RefreshChildrenForFunctionCall(UNiagaraNodeFunctionCall* ModuleNode, UNiagaraNodeFunctionCall* InputFunctionCallNode, const TArray<UNiagaraStackEntry*>& CurrentChildren, 
-		TArray<UNiagaraStackEntry*>& NewChildren, TArray<FStackIssue>& NewIssues, bool bShouldApplySummaryFilter);
-	
-	void AppendInputsForFunctionCall(FFunctionCallNodesState& State, UNiagaraNodeFunctionCall* ModuleNode, UNiagaraNodeFunctionCall* InputFunctionCallNode, TArray<FStackIssue>& NewIssues, bool bShouldApplySummaryFilter);
-	
-	void ApplyAllFunctionInputsToChildren(FFunctionCallNodesState& State, const TArray<UNiagaraStackEntry*>& CurrentChildren, TArray<UNiagaraStackEntry*>& NewChildren, TArray<FStackIssue>& NewIssues, bool bShouldApplySummaryFilter);
-
-	void RefreshIssues(UNiagaraNodeFunctionCall* InputFunctionCallNode, const TArray<FName>& DuplicateInputNames, const TArray<FName>& ValidAliasedInputNames, const TArray<FNiagaraVariable>& InputsWithInvalidTypes, const TMap<FName, UEdGraphPin*>& StaticSwitchInputs, TArray<FStackIssue>& NewIssues);
-
-	void OnFunctionInputsChanged();
-
-	FStackIssueFix GetNodeRemovalFix(UEdGraphPin* PinToRemove, FText FixDescription);
-
-	FStackIssueFix GetResetPinFix(UEdGraphPin* PinToReset, FText FixDescription);
-
-	FStackIssueFix GetUpgradeVersionFix(FText FixDescription);
-
-	void AddInvalidChildStackIssue(FName PinName, TArray<FStackIssue>& OutIssues);
-
-
-	void AddInputToCategory(const FInputData& InputData, const TArray<UNiagaraStackEntry*>& CurrentChildren, TArray<UNiagaraStackEntry*>& NewChildren);
-
 	virtual void GetSectionsInternal(TArray<FNiagaraStackSection>& OutStackSections) const { }
-
-private:
-	bool FilterByActiveSection(const UNiagaraStackEntry& Child) const;
-
 	void UpdateCachedSectionData() const;
-
+	virtual bool FilterByActiveSection(const UNiagaraStackEntry& Child) const;
 private:
-	bool bShouldDisplayLabel;
-
+	virtual bool GetCanExpand() const override;
+	virtual bool GetShouldShowInStack() const override;
+	virtual int32 GetChildIndentLevel() const override;
+private:
 	mutable TOptional<TArray<FText>> SectionsCache;
 	mutable TOptional<TMap<FString, TArray<FText>>> SectionToCategoryMapCache;
+	mutable TOptional<TMap<FString, FText>> SectionToTooltipMapCache;
 	mutable TOptional<FText> ActiveSectionCache;
 	FText LastActiveSection;
+
+	bool bShouldDisplayLabel;
 };
 
 UCLASS()
-class NIAGARAEDITOR_API UNiagaraStackFunctionInputCollection : public UNiagaraStackFunctionInputCollectionBase
+class NIAGARAEDITOR_API UNiagaraStackFunctionInputCollection : public UNiagaraStackValueCollection
 {
 	GENERATED_BODY()
 
@@ -144,7 +87,36 @@
 
 	TArray<UNiagaraStackFunctionInput*> GetInlineParameterInputs() const;
 
-protected:
+private:
+
+	struct FInputData
+	{
+		FNiagaraVariable InputVariable;
+		int32 SortKey;
+		TOptional<FText> DisplayName;
+		FText Category;
+		bool bIsStatic;
+		bool bIsHidden;
+
+		UNiagaraNodeFunctionCall* ModuleNode;
+		UNiagaraNodeFunctionCall* InputFunctionCallNode;
+		
+		TArray<FInputData*> Children;
+		bool bIsChild = false;
+	};
+
+	struct FNiagaraParentData
+	{
+		FNiagaraVariable ParentVariable;
+		TArray<int32> ChildIndices;
+	};
+
+	struct FFunctionCallNodesState
+	{		
+		TArray<FInputData> InputDataCollection;
+		TMap<FName, FNiagaraParentData> ParentMapping;
+	};
+	
 	void OnScriptApplied(UNiagaraScript* NiagaraScript, FGuid Guid);
 	
 	virtual void FinalizeInternal() override;
@@ -153,6 +125,26 @@
 
 	virtual void GetSectionsInternal(TArray<FNiagaraStackSection>& OutStackSections) const override;
 	
+	void RefreshChildrenForFunctionCall(const TArray<UNiagaraStackEntry*>& CurrentChildren, TArray<UNiagaraStackEntry*>& NewChildren, TArray<FStackIssue>& NewIssues, bool bShouldApplySummaryFilter);
+	
+	void AppendInputsForFunctionCall(FFunctionCallNodesState& State, TArray<FStackIssue>& NewIssues, bool bShouldApplySummaryFilter);
+	
+	void ApplyAllFunctionInputsToChildren(FFunctionCallNodesState& State, const TArray<UNiagaraStackEntry*>& CurrentChildren, TArray<UNiagaraStackEntry*>& NewChildren, TArray<FStackIssue>& NewIssues, bool bShouldApplySummaryFilter);
+
+	void RefreshIssues(const TArray<FName>& DuplicateInputNames, const TArray<FName>& ValidAliasedInputNames, const TArray<FNiagaraVariable>& InputsWithInvalidTypes, const TMap<FName, UEdGraphPin*>& StaticSwitchInputs, TArray<FStackIssue>& NewIssues);
+
+	void OnFunctionInputsChanged();
+
+	FStackIssueFix GetNodeRemovalFix(UEdGraphPin* PinToRemove, FText FixDescription);
+
+	FStackIssueFix GetResetPinFix(UEdGraphPin* PinToReset, FText FixDescription);
+
+	FStackIssueFix GetUpgradeVersionFix(FText FixDescription);
+
+	void AddInvalidChildStackIssue(FName PinName, TArray<FStackIssue>& OutIssues);
+	
+	void AddInputToCategory(const FInputData& InputData, const TArray<UNiagaraStackEntry*>& CurrentChildren, TArray<UNiagaraStackEntry*>& NewChildren);
+	
 private:
 	UNiagaraNodeFunctionCall* ModuleNode;
 	UNiagaraNodeFunctionCall* InputFunctionCallNode;
diff --git a/Engine/Plugins/FX/Niagara/Source/NiagaraEditor/Public/ViewModels/Stack/NiagaraStackGraphUtilities.h b/Engine/Plugins/FX/Niagara/Source/NiagaraEditor/Public/ViewModels/Stack/NiagaraStackGraphUtilities.h
--- a/Engine/Plugins/FX/Niagara/Source/NiagaraEditor/Public/ViewModels/Stack/NiagaraStackGraphUtilities.h
+++ b/Engine/Plugins/FX/Niagara/Source/NiagaraEditor/Public/ViewModels/Stack/NiagaraStackGraphUtilities.h
@@ -4,10 +4,10 @@
 
 #include "NiagaraTypes.h"
 #include "NiagaraCommon.h"
+#include "NiagaraParameterMapHistory.h"
 #include "ViewModels/Stack/NiagaraParameterHandle.h"
 #include "ViewModels/Stack/NiagaraStackEntry.h"
 #include "AssetRegistry/AssetData.h"
-#include "NiagaraStackEditorData.h"
 
 struct FNiagaraHierarchyIdentity;
 class UEdGraph;
@@ -30,6 +30,7 @@ class UNiagaraStackErrorItem;
 class FCompileConstantResolver;
 class INiagaraMessage;
 class FNiagaraParameterMapHistoryBuilder;
+class UNiagaraHierarchyModuleInput;
 struct FNiagaraStackModuleData;
 struct FNiagaraModuleDependency;
 
@@ -50,6 +51,9 @@ namespace FNiagaraStackGraphUtilities
 	TArray<UNiagaraNodeFunctionCall*> GetAllModuleNodes(TSharedRef<FNiagaraEmitterViewModel> EmitterViewModel);
 	TArray<UNiagaraNodeFunctionCall*> GetAllModuleNodes(UNiagaraGraph* Graph);
 
+	TArray<UNiagaraNodeFunctionCall*> GetAllSimStagesModuleNodes(TSharedRef<FNiagaraEmitterViewModel> EmitterViewModel);
+	TArray<UNiagaraNodeFunctionCall*> GetAllEventHandlerModuleNodes(TSharedRef<FNiagaraEmitterViewModel> EmitterViewModel);
+
 	UNiagaraNodeFunctionCall* GetPreviousModuleNode(UNiagaraNodeFunctionCall& CurrentNode);
 
 	UNiagaraNodeFunctionCall* GetNextModuleNode(UNiagaraNodeFunctionCall& CurrentNode);
@@ -89,6 +93,18 @@ namespace FNiagaraStackGraphUtilities
 	TArray<FName> StackContextResolution(FVersionedNiagaraEmitter OwningEmitter, UNiagaraNodeOutput* OutputNodeInChain);
 	void BuildParameterMapHistoryWithStackContextResolution(FVersionedNiagaraEmitter OwningEmitter, UNiagaraNodeOutput* OutputNodeInChain, UNiagaraNode* NodeToVisit, FNiagaraParameterMapHistoryBuilder& Builder, bool bRecursive = true, bool bFilterForCompilation = true);
 
+	struct FInputDataCollection
+	{
+		TMap<FGuid, UNiagaraNodeFunctionCall*> NodeGuidToModuleNodeMap;
+		TMap<UNiagaraHierarchyModuleInput*, FNiagaraVariable> HierarchyInputToAssignmentMap;
+		TMap<UNiagaraHierarchyModuleInput*, TObjectPtr<UNiagaraScriptVariable>> HierarchyInputToScriptVariableMap;
+		/** These two aren't currently in use as children inputs are handled manually by the author */
+		TMap<UNiagaraHierarchyModuleInput*, TArray<FGuid>> HierarchyInputToChildrenGuidMap;
+		TMap<FGuid, TObjectPtr<UNiagaraScriptVariable>> ChildrenGuidToScriptVariablesMap;
+	};
+	
+	void GatherInputRelationsForStack(FInputDataCollection& State, TSharedRef<FNiagaraEmitterViewModel> EmitterViewModel);
+
 	enum class ENiagaraGetStackFunctionInputPinsOptions
 	{
 		AllInputs,
@@ -119,14 +135,55 @@ namespace FNiagaraStackGraphUtilities
 		FName InputName;
 		FNiagaraTypeDefinition Type;
 		FNiagaraVariableMetaData MetaData;
+		bool bIsHidden;
+		bool bIsStatic;
 		UNiagaraNodeFunctionCall* FunctionCallNode = nullptr;
 		TArray<FGuid> ChildrenInputGuids;
+		TArray<FGuid> HiddenChildrenInputGuids;
 	};
 
+	struct FInputDataCacheKey
+	{
+		FInputDataCacheKey(FGuid InNodeGuid, FGuid InVariableGuid) : NodeGuid(InNodeGuid), VariableGuid(InVariableGuid)	{}
+		
+		FGuid NodeGuid;
+		FGuid VariableGuid;
+		
+		bool operator==(const FInputDataCacheKey& OtherKey) const
+		{
+			return NodeGuid == OtherKey.NodeGuid && VariableGuid == OtherKey.VariableGuid;
+		}
+
+		bool operator!=(const FInputDataCacheKey& OtherKey) const
+		{
+			return !(*this == OtherKey);
+		}
+	};
+
+	FORCEINLINE uint32 GetTypeHash(const FInputDataCacheKey& InputDataCacheKey)
+	{
+		uint32 Hash = 0;
+		HashCombine(Hash, GetTypeHash(InputDataCacheKey.NodeGuid));
+		HashCombine(Hash, GetTypeHash(InputDataCacheKey.VariableGuid));
+		return Hash;
+	}
+
 	UNiagaraNodeFunctionCall* FindModuleNode(FGuid ModuleNodeGuid, TSharedRef<FNiagaraEmitterViewModel> EmitterViewModel);
+	UNiagaraNodeAssignment* FindAssignmentNode(FGuid AssignmentNodeGuid, TSharedRef<FNiagaraEmitterViewModel> EmitterViewModel);
+	UNiagaraNodeFunctionCall* FindDynamicInputNodeForInput(UNiagaraNodeFunctionCall& OwningFunctionNode, FName UnaliasedParameterName);
+	/** When you have the guid of the node available, we don't need to traverse pins and can just look up the correct node directly. */
+	UNiagaraNodeFunctionCall* FindFunctionCallNode(FGuid FunctionCallGuid, TSharedRef<FNiagaraEmitterViewModel> EmitterViewModel);
+	TArray<UNiagaraNodeFunctionCall*> FindModuleNodesForSimulationStage(UNiagaraSimulationStageBase& SimStage, TSharedRef<FNiagaraEmitterViewModel> EmitterViewModel);
+	TArray<UNiagaraNodeFunctionCall*> FindModuleNodesForEventHandler(FNiagaraEventScriptProperties& EventScriptProperties, TSharedRef<FNiagaraEmitterViewModel> EmitterViewModel);
 	TOptional<FMatchingFunctionInputData> FindInputData(FNiagaraHierarchyIdentity ModuleInputIdentity, TSharedRef<FNiagaraEmitterViewModel> EmitterViewModel);
-	TOptional<FMatchingFunctionInputData> FindInputData(const UNiagaraNodeFunctionCall& FunctionCallNode, FGuid InputGuid, TSharedRef<FNiagaraEmitterViewModel> EmitterViewModel, bool bIncludeChildrenInputs = true);
+	/** A more optimized version of the function above. Assumes you already have the correct node. */
+	TOptional<FMatchingFunctionInputData> FindInputData(const UNiagaraNodeFunctionCall& FunctionCallNode, FNiagaraHierarchyIdentity InputIdentity, TSharedRef<FNiagaraEmitterViewModel> EmitterViewModel);
+	TOptional<FMatchingFunctionInputData> FindAssignmentInputData(const UNiagaraNodeAssignment& AssignmentNode, FName VariableName, TSharedRef<FNiagaraEmitterViewModel> EmitterViewModel);
+	/** Finds the input data for a given function call node and a variable guid. An optional cache can be used that, if specified, will be used to get variable data if found and if not, will write into the cache
+	 * This also means maintaining the cache becomes the responsibility of the caller */
+	TOptional<FMatchingFunctionInputData> FindModuleInputData(const UNiagaraNodeFunctionCall& FunctionCallNode, FGuid VariableGuid, TSharedRef<FNiagaraEmitterViewModel> EmitterViewModel, bool bIncludeChildrenInputs = true, TMap<FInputDataCacheKey, FMatchingFunctionInputData>* OptionalCache = nullptr);
 	TArray<FGuid> GetChildrenInputGuids(const UNiagaraNodeFunctionCall& FunctionCallNode, FName ParentInputName);
+	TSet<FGuid> GetHiddenChildrenInputGuids(const UNiagaraNodeFunctionCall& FunctionCallNode, FName ParentInput, TSharedRef<FNiagaraEmitterViewModel> EmitterViewModel);
 	
 	/* Module script calls do not have direct inputs, but rely on the parameter map being initialized correctly. This utility function resolves which of the module's parameters are reachable during compilation and returns a list of pins on the parameter map node that do not have to be compiled. */
 	TArray<UEdGraphPin*> GetUnusedFunctionInputPins(UNiagaraNodeFunctionCall& FunctionCallNode, FCompileConstantResolver ConstantResolver);
diff --git a/Engine/Plugins/FX/Niagara/Source/NiagaraEditor/Public/ViewModels/Stack/NiagaraStackInputCategory.h b/Engine/Plugins/FX/Niagara/Source/NiagaraEditor/Public/ViewModels/Stack/NiagaraStackInputCategory.h
--- a/Engine/Plugins/FX/Niagara/Source/NiagaraEditor/Public/ViewModels/Stack/NiagaraStackInputCategory.h
+++ b/Engine/Plugins/FX/Niagara/Source/NiagaraEditor/Public/ViewModels/Stack/NiagaraStackInputCategory.h
@@ -4,6 +4,8 @@
 
 #include "ViewModels/Stack/NiagaraStackItem.h"
 #include "NiagaraTypes.h"
+#include "ViewModels/HierarchyEditor/NiagaraHierarchyViewModelBase.h"
+#include "ViewModels/HierarchyEditor/NiagaraSummaryViewViewModel.h"
 #include "NiagaraStackInputCategory.generated.h"
 
 class UNiagaraStackFunctionInput;
@@ -17,7 +19,37 @@ enum class EStackParameterBehavior
 };
 
 UCLASS()
-class NIAGARAEDITOR_API UNiagaraStackInputCategory : public UNiagaraStackItemContent
+class NIAGARAEDITOR_API UNiagaraStackCategory : public UNiagaraStackItemContent
+{
+	GENERATED_BODY()
+
+public:
+	void Initialize(FRequiredEntryData InRequiredEntryData, FString InOwningStackItemEditorDataKey, FString InStackEditorDataKey);
+
+	//~ UNiagaraStackEntry interface
+	virtual FText GetDisplayName() const override;
+	virtual bool GetShouldShowInStack() const override;
+	virtual EStackRowStyle GetStackRowStyle() const override;
+	virtual void GetSearchItems(TArray<FStackSearchItem>& SearchItems) const override;
+
+	virtual bool IsTopLevelCategory() const { return false; }
+	virtual int32 GetChildIndentLevel() const override;
+
+protected:
+	virtual void RefreshChildrenInternal(const TArray<UNiagaraStackEntry*>& CurrentChildren, TArray<UNiagaraStackEntry*>& NewChildren, TArray<FStackIssue>& NewIssues) override;
+private:
+private:
+	bool FilterForVisibleCondition(const UNiagaraStackEntry& Child) const;
+	bool FilterForIsInlineEditConditionToggle(const UNiagaraStackEntry& Child) const;
+protected:
+	bool bShouldShowInStack;
+
+	UPROPERTY()
+	TObjectPtr<UNiagaraStackSpacer> CategorySpacer;
+};
+
+UCLASS()
+class NIAGARAEDITOR_API UNiagaraStackInputCategory : public UNiagaraStackCategory
 {
 	GENERATED_BODY() 
 
@@ -29,19 +61,14 @@ public:
 		bool bInIsTopLevelCategory,
 		FString InOwnerStackItemEditorDataKey);
 	
-	const FText& GetCategoryName() const;
+	bool GetIsEnabled() const;
+
+	virtual FText GetDisplayName() const override;
 
 	void ResetInputs();
 
 	void AddInput(UNiagaraNodeFunctionCall* InModuleNode, UNiagaraNodeFunctionCall* InInputFunctionCallNode, FName InInputParameterHandle, FNiagaraTypeDefinition InInputType, EStackParameterBehavior InParameterBehavior, TOptional<FText> InOptionalDisplayName, bool bIsHidden, bool bIsChildInput);
 
-	//~ UNiagaraStackEntry interface
-	virtual FText GetDisplayName() const override;
-	virtual bool GetShouldShowInStack() const override;
-	virtual EStackRowStyle GetStackRowStyle() const override;
-	virtual bool GetIsEnabled() const override;
-	virtual void GetSearchItems(TArray<FStackSearchItem>& SearchItems) const override;
-
 	void SetShouldShowInStack(bool bInShouldShowInStack);
 
 	void ToClipboardFunctionInputs(UObject* InOuter, TArray<const UNiagaraClipboardFunctionInput*>& OutClipboardFunctionInputs) const;
@@ -55,12 +82,7 @@ public:
 protected:
 	//~ UNiagaraStackEntry interface
 	virtual void RefreshChildrenInternal(const TArray<UNiagaraStackEntry*>& CurrentChildren, TArray<UNiagaraStackEntry*>& NewChildren, TArray<FStackIssue>& NewIssues) override;
-	virtual int32 GetChildIndentLevel() const override;
-
-private:
-	bool FilterForVisibleCondition(const UNiagaraStackEntry& Child) const;
-	bool FilterForIsInlineEditConditionToggle(const UNiagaraStackEntry& Child) const;
-
+	virtual bool IsTopLevelCategory() const override { return bIsTopLevelCategory; }
 private:
 	struct FInputParameterHandleAndType
 	{
@@ -74,13 +96,40 @@ private:
 		bool bIsChildInput;
 	};
 
-	TOptional<FText> DisplayName;
 	FText CategoryName;
+	TOptional<FText> DisplayName;
+	bool bIsTopLevelCategory;
+
 	TArray<FInputParameterHandleAndType> Inputs;
+};
 
-	UPROPERTY()
-	TObjectPtr<UNiagaraStackSpacer> CategorySpacer;
+void AddSummaryItem(UNiagaraHierarchyItemBase* HierarchyItem, UNiagaraStackEntry* Parent);
 
-	bool bShouldShowInStack;
-	bool bIsTopLevelCategory;
-};
\ No newline at end of file
+UCLASS()
+class NIAGARAEDITOR_API UNiagaraStackSummaryCategory : public UNiagaraStackCategory
+{
+	GENERATED_BODY()
+
+public:
+	UNiagaraStackSummaryCategory() {}
+
+	void Initialize(
+		FRequiredEntryData InRequiredEntryData,
+		TSharedPtr<FNiagaraHierarchyCategoryViewModel> InCategory,
+		FString InOwnerStackItemEditorDataKey);
+
+	//~ UNiagaraStackEntry interface
+	virtual FText GetDisplayName() const override;
+	virtual bool GetIsEnabled() const override { return true; }
+	
+	TWeakPtr<FNiagaraHierarchyCategoryViewModel> GetHierarchyCategory() const { return CategoryViewModelWeakPtr; }
+
+protected:
+	virtual void RefreshChildrenInternal(const TArray<UNiagaraStackEntry*>& CurrentChildren, TArray<UNiagaraStackEntry*>& NewChildren, TArray<FStackIssue>& NewIssues) override;
+	virtual bool IsTopLevelCategory() const override;
+	virtual FText GetTooltipText() const override;
+	virtual int32 GetChildIndentLevel() const override;
+
+private:
+	TWeakPtr<FNiagaraHierarchyCategoryViewModel> CategoryViewModelWeakPtr;
+};
diff --git a/Engine/Plugins/FX/Niagara/Source/NiagaraEditor/Public/ViewModels/Stack/NiagaraStackItem.h b/Engine/Plugins/FX/Niagara/Source/NiagaraEditor/Public/ViewModels/Stack/NiagaraStackItem.h
--- a/Engine/Plugins/FX/Niagara/Source/NiagaraEditor/Public/ViewModels/Stack/NiagaraStackItem.h
+++ b/Engine/Plugins/FX/Niagara/Source/NiagaraEditor/Public/ViewModels/Stack/NiagaraStackItem.h
@@ -39,11 +39,13 @@ public:
 	virtual void ResetToBase() { }
 
 	virtual bool SupportsEditMode() const { return false; }
-	virtual bool GetEditModeIsActive() const { return false; }
-	virtual void SetEditModeIsActive(bool bInEditModeIsActive) { }
+	virtual void OnEditButtonClicked() { }
+	virtual TOptional<FText> GetEditModeButtonText() const { return TOptional<FText>(); }
+	virtual TOptional<FText> GetEditModeButtonTooltip() const { return TOptional<FText>(); }
+	virtual EVisibility IsEditButtonVisible() const { return SupportsEditMode() ? EVisibility::Visible : EVisibility::Collapsed; }
 
 	virtual bool GetIsInherited() const { return false; }
-	
+
 protected:
 	virtual void RefreshChildrenInternal(const TArray<UNiagaraStackEntry*>& CurrentChildren, TArray<UNiagaraStackEntry*>& NewChildren, TArray<FStackIssue>& NewIssues) override;
 
diff --git a/Engine/Plugins/FX/Niagara/Source/NiagaraEditor/Public/ViewModels/Stack/NiagaraStackModuleItem.h b/Engine/Plugins/FX/Niagara/Source/NiagaraEditor/Public/ViewModels/Stack/NiagaraStackModuleItem.h
--- a/Engine/Plugins/FX/Niagara/Source/NiagaraEditor/Public/ViewModels/Stack/NiagaraStackModuleItem.h
+++ b/Engine/Plugins/FX/Niagara/Source/NiagaraEditor/Public/ViewModels/Stack/NiagaraStackModuleItem.h
@@ -99,6 +99,9 @@ public:
 	virtual bool GetIsInherited() const override;
 	virtual FText GetInheritanceMessage() const override;
 
+	virtual bool SupportsSummaryView() const override { return true; }
+	virtual FNiagaraHierarchyIdentity DetermineSummaryIdentity() const override;
+	
 	bool IsScratchModule() const;
 
 	void SetOnRequestDeprecationRecommended(FOnRequestDeprecationRecommended InOnRequest)
diff --git a/Engine/Plugins/FX/Niagara/Source/NiagaraEditor/Public/ViewModels/Stack/NiagaraStackObject.h b/Engine/Plugins/FX/Niagara/Source/NiagaraEditor/Public/ViewModels/Stack/NiagaraStackObject.h
--- a/Engine/Plugins/FX/Niagara/Source/NiagaraEditor/Public/ViewModels/Stack/NiagaraStackObject.h
+++ b/Engine/Plugins/FX/Niagara/Source/NiagaraEditor/Public/ViewModels/Stack/NiagaraStackObject.h
@@ -46,6 +46,9 @@ public:
 	virtual bool GetShouldShowInStack() const override;
 	virtual UObject* GetDisplayedObject() const override;
 
+	void SetObjectGuid(FGuid InGuid) { ObjectGuid = InGuid; }
+	virtual bool SupportsSummaryView() const override { return ObjectGuid.IsSet() && ObjectGuid->IsValid(); }
+	virtual FNiagaraHierarchyIdentity DetermineSummaryIdentity() const override;
 protected:
 	virtual void FinalizeInternal() override;
 
@@ -82,6 +85,9 @@ private:
 	TSharedPtr<IPropertyRowGenerator> PropertyRowGenerator;
 	bool bIsRefresingDataInterfaceErrors;
 
+	/** An optional object guid that can be provided to identify the object this stack entry represents. This can be used for summary view purposes and is also given to the child property rows. */
+	TOptional<FGuid> ObjectGuid;
+	
 	FGuid MessageManagerRegistrationKey;
 	TArray<FStackIssue> MessageManagerIssues;
 	FGuid MessageLogGuid;
diff --git a/Engine/Plugins/FX/Niagara/Source/NiagaraEditor/Public/ViewModels/Stack/NiagaraStackPropertyRow.h b/Engine/Plugins/FX/Niagara/Source/NiagaraEditor/Public/ViewModels/Stack/NiagaraStackPropertyRow.h
--- a/Engine/Plugins/FX/Niagara/Source/NiagaraEditor/Public/ViewModels/Stack/NiagaraStackPropertyRow.h
+++ b/Engine/Plugins/FX/Niagara/Source/NiagaraEditor/Public/ViewModels/Stack/NiagaraStackPropertyRow.h
@@ -15,7 +15,7 @@ class NIAGARAEDITOR_API UNiagaraStackPropertyRow : public UNiagaraStackItemConte
 		
 public:
 	void Initialize(FRequiredEntryData InRequiredEntryData, TSharedRef<IDetailTreeNode> InDetailTreeNode, bool bInIsTopLevelProperty, FString InOwnerStackItemEditorDataKey, FString InOwnerStackEditorDataKey, UNiagaraNode* InOwningNiagaraNode);
-
+	
 	virtual EStackRowStyle GetStackRowStyle() const override;
 
 	TSharedRef<IDetailTreeNode> GetDetailTreeNode() const;
@@ -28,6 +28,7 @@ public:
 
 	virtual bool CanDrag() const override;
 
+	void SetOwnerGuid(TOptional<FGuid> InGuid) { OwnerGuid = InGuid; }
 protected:
 	virtual void FinalizeInternal() override;
 
@@ -41,6 +42,8 @@ protected:
 
 	virtual TOptional<FDropRequestResponse> DropInternal(const FDropRequest& DropRequest) override;
 
+	virtual bool SupportsSummaryView() const override;
+	virtual FNiagaraHierarchyIdentity DetermineSummaryIdentity() const override;
 private:
 	TSharedPtr<IDetailTreeNode> DetailTreeNode;
 	UNiagaraNode* OwningNiagaraNode;
@@ -49,6 +52,9 @@ private:
 	UPROPERTY()
 	TObjectPtr<UNiagaraStackSpacer> CategorySpacer;
 
+	/** An optional guid that can be used to identify the parent object. The parent object is responsible for setting this. It is used to create the summary view identity */
+	TOptional<FGuid> OwnerGuid;
+
 	bool bCannotEditInThisContext;
 	bool bIsTopLevelProperty;
 };
\ No newline at end of file
