Summary View 2.0

diff --git a/Engine/Plugins/FX/Niagara/Source/NiagaraEditor/Public/ViewModels/Stack/NiagaraStackRendererItem.h b/Engine/Plugins/FX/Niagara/Source/NiagaraEditor/Public/ViewModels/Stack/NiagaraStackRendererItem.h
--- a/Engine/Plugins/FX/Niagara/Source/NiagaraEditor/Public/ViewModels/Stack/NiagaraStackRendererItem.h
+++ b/Engine/Plugins/FX/Niagara/Source/NiagaraEditor/Public/ViewModels/Stack/NiagaraStackRendererItem.h
@@ -23,7 +23,8 @@ public:
 	void Initialize(FRequiredEntryData InRequiredEntryData, UNiagaraRendererProperties* InRendererProperties);
 
 	UNiagaraRendererProperties* GetRendererProperties();
-
+	const UNiagaraRendererProperties* GetRendererProperties() const;
+	
 	virtual FText GetDisplayName() const override;
 
 	virtual bool SupportsCut() const override { return true; }
@@ -49,6 +50,9 @@ public:
 	virtual bool SupportsInheritance() const override { return true; }
 	virtual bool GetIsInherited() const override;
 	virtual FText GetInheritanceMessage() const override;
+
+	virtual bool SupportsSummaryView() const override { return true; }
+	virtual FNiagaraHierarchyIdentity DetermineSummaryIdentity() const override;
 	
 	virtual bool SupportsRename() const override { return true; }
 
diff --git a/Engine/Plugins/FX/Niagara/Source/NiagaraEditor/Public/ViewModels/Stack/NiagaraStackSimulationStageGroup.h b/Engine/Plugins/FX/Niagara/Source/NiagaraEditor/Public/ViewModels/Stack/NiagaraStackSimulationStageGroup.h
--- a/Engine/Plugins/FX/Niagara/Source/NiagaraEditor/Public/ViewModels/Stack/NiagaraStackSimulationStageGroup.h
+++ b/Engine/Plugins/FX/Niagara/Source/NiagaraEditor/Public/ViewModels/Stack/NiagaraStackSimulationStageGroup.h
@@ -29,11 +29,14 @@ public:
 
 	void SetSimulationStageEnabled(bool bIsEnabled);
 
+	TWeakObjectPtr<UNiagaraSimulationStageBase> GetSimulationStage() const { return SimulationStage; }
 protected:
 	virtual void FinalizeInternal() override;
 
 	virtual void RefreshChildrenInternal(const TArray<UNiagaraStackEntry*>& CurrentChildren, TArray<UNiagaraStackEntry*>& NewChildren, TArray<FStackIssue>& NewIssues) override;
 
+	virtual bool SupportsSummaryView() const override { return true; }
+	virtual FNiagaraHierarchyIdentity DetermineSummaryIdentity() const override;
 private:
 	void SimulationStagePropertiesChanged();
 
@@ -92,6 +95,8 @@ protected:
 
 	virtual TOptional<FDropRequestResponse> DropInternal(const FDropRequest& DropRequest) override;
 
+	virtual bool SupportsSummaryView() const override { return true; }
+	virtual FNiagaraHierarchyIdentity DetermineSummaryIdentity() const override;
 private:
 	void SimulationStagePropertiesChanged();
 
diff --git a/Engine/Plugins/FX/Niagara/Source/NiagaraEditor/Public/ViewModels/Stack/NiagaraStackSummaryViewInputCollection.h b/Engine/Plugins/FX/Niagara/Source/NiagaraEditor/Public/ViewModels/Stack/NiagaraStackSummaryViewInputCollection.h
--- a/Engine/Plugins/FX/Niagara/Source/NiagaraEditor/Public/ViewModels/Stack/NiagaraStackSummaryViewInputCollection.h
+++ b/Engine/Plugins/FX/Niagara/Source/NiagaraEditor/Public/ViewModels/Stack/NiagaraStackSummaryViewInputCollection.h
@@ -3,19 +3,12 @@
 #pragma once
 
 #include "NiagaraEmitter.h"
-#include "ViewModels/Stack/NiagaraStackItem.h"
-#include "PropertyEditorDelegates.h"
 #include "NiagaraStackFunctionInputCollection.h"
+#include "ViewModels/HierarchyEditor/NiagaraHierarchyViewModelBase.h"
 #include "NiagaraStackSummaryViewInputCollection.generated.h"
 
-class IPropertyRowGenerator;
-class UNiagaraNode;
-class IDetailTreeNode;
-class UNiagaraStackFunctionInputCollection;
-class UNiagaraNodeFunctionCall;
-
 UCLASS()
-class NIAGARAEDITOR_API UNiagaraStackSummaryViewObject : public UNiagaraStackFunctionInputCollectionBase
+class NIAGARAEDITOR_API UNiagaraStackSummaryViewCollection : public UNiagaraStackValueCollection
 {
 	GENERATED_BODY()
 		
@@ -23,29 +16,25 @@ public:
 	DECLARE_DELEGATE_TwoParams(FOnSelectRootNodes, TArray<TSharedRef<IDetailTreeNode>>, TArray<TSharedRef<IDetailTreeNode>>*);
 
 public:
-	UNiagaraStackSummaryViewObject();
+	UNiagaraStackSummaryViewCollection() {}
 
 	void Initialize(FRequiredEntryData InRequiredEntryData, FVersionedNiagaraEmitterWeakPtr InEmitter, FString InOwningStackItemEditorDataKey);
 	virtual void FinalizeInternal() override;
 
 	virtual FText GetDisplayName() const override;
-	virtual bool GetIsEnabled() const;
-
+	virtual bool GetIsEnabled() const override;
+	
 protected:
-
 	virtual void RefreshChildrenInternal(const TArray<UNiagaraStackEntry*>& CurrentChildren, TArray<UNiagaraStackEntry*>& NewChildren, TArray<FStackIssue>& NewIssues) override;
 
-	void AppendEmitterCategory(FFunctionCallNodesState& State, TSharedPtr<FNiagaraScriptViewModel> ScriptViewModelPinned, ENiagaraScriptUsage ScriptUsage, FGuid ScriptUsageId, TArray<FStackIssue>& NewIssues);
-
-	virtual void PostRefreshChildrenInternal() override;
-
 	virtual void GetSectionsInternal(TArray<FNiagaraStackSection>& OutStackSections) const override;
+	const TArray<UNiagaraHierarchySection*>& GetHierarchySections() const;
 
+	virtual bool FilterByActiveSection(const UNiagaraStackEntry& Child) const override;
 private:
-
 	void OnViewStateChanged();
 private:
 
 	FVersionedNiagaraEmitterWeakPtr Emitter;
-	TMap<FGuid, UNiagaraStackFunctionInputCollection*> KnownInputCollections;
 };
+
diff --git a/Engine/Plugins/FX/Niagara/Source/NiagaraEditor/Public/Widgets/SNiagaraHierarchy.h b/Engine/Plugins/FX/Niagara/Source/NiagaraEditor/Public/Widgets/SNiagaraHierarchy.h
--- a/Engine/Plugins/FX/Niagara/Source/NiagaraEditor/Public/Widgets/SNiagaraHierarchy.h
+++ b/Engine/Plugins/FX/Niagara/Source/NiagaraEditor/Public/Widgets/SNiagaraHierarchy.h
@@ -7,16 +7,76 @@
 #include "Widgets/Views/STreeView.h"
 #include "Widgets/Views/STableRow.h"
 #include "Widgets/Text/SInlineEditableTextBlock.h"
-#include "Widgets/SBoxPanel.h"
 #include "Widgets/Input/SMenuAnchor.h"
 #include "Styling/SlateTypes.h"
+#include "ViewModels/HierarchyEditor/NiagaraHierarchyViewModelBase.h"
+#include "Widgets/Input/SSearchBox.h"
+#include "Widgets/Input/SCheckBox.h"
 
 class UNiagaraHierarchyViewModelBase;
 class UNiagaraHierarchySection;
 struct FNiagaraHierarchyItemViewModelBase;
 struct FNiagaraHierarchySectionViewModel;
 
-class SNiagaraHierarchy : public SCompoundWidget
+class SNiagaraHierarchySection : public SCompoundWidget
 {
+	DECLARE_DELEGATE_OneParam(FOnSectionActivated, TSharedPtr<FNiagaraHierarchySectionViewModel> SectionViewModel)
+	
+	SLATE_BEGIN_ARGS(SNiagaraHierarchySection)
+		: _bForbidDropOn(false)
+		{}
+		SLATE_ATTRIBUTE(ECheckBoxState, IsSectionActive)
+		SLATE_EVENT(FOnSectionActivated, OnSectionActivated)
+		SLATE_ARGUMENT(bool, bForbidDropOn)
+	SLATE_END_ARGS()
+
+public:
+	void Construct(const FArguments& InArgs, TSharedPtr<struct FNiagaraHierarchySectionViewModel> InSection, TWeakObjectPtr<UNiagaraHierarchyViewModelBase> HierarchyViewModel);
+	virtual ~SNiagaraHierarchySection() override;
+
+	bool OnCanAcceptDrop(TSharedPtr<FDragDropOperation> DragDropOperation, EItemDropZone ItemDropZone) const;
+	FReply OnDroppedOn(const FGeometry&, const FDragDropEvent& DragDropEvent, EItemDropZone DropZone) const;
+
+	void EnterEditingMode() const;
+
+	TSharedPtr<struct FNiagaraHierarchySectionViewModel> GetSectionViewModel();
+private:
+	virtual FReply OnMouseButtonDown(const FGeometry& MyGeometry, const FPointerEvent& MouseEvent) override;
+	virtual FReply OnMouseButtonUp(const FGeometry& MyGeometry, const FPointerEvent& MouseEvent) override;
+	
+	virtual void OnDragEnter(const FGeometry& MyGeometry, const FDragDropEvent& DragDropEvent) override;
+	virtual FReply OnDragDetected(const FGeometry& MyGeometry, const FPointerEvent& MouseEvent) override;
+	virtual void OnDragLeave(const FDragDropEvent& DragDropEvent) override;
+
+	virtual FReply OnKeyDown(const FGeometry& MyGeometry, const FKeyEvent& InKeyEvent) override;
+	
+	virtual void OnMouseLeave(const FPointerEvent& MouseEvent) override;
+
+	UNiagaraHierarchySection* TryGetSectionData() const;
+
+	FText GetText() const;
+	FText GetTooltipText() const;
+	void OnRenameSection(const FText& Text, ETextCommit::Type CommitType) const;
+	bool OnVerifySectionRename(const FText& NewName, FText& OutTooltip) const;
+
+	bool IsSectionSelected() const;
+	ECheckBoxState GetSectionCheckState() const;
+	void OnSectionCheckChanged(ECheckBoxState NewState);
+	EActiveTimerReturnType ActivateSectionIfDragging(double CurrentTime, float DeltaTime) const;
+private:
+	TSharedPtr<SMenuAnchor> MenuAnchor;
+	TSharedPtr<SCheckBox> CheckBox;
+	TSharedPtr<SInlineEditableTextBlock> InlineEditableTextBlock;
+	TWeakObjectPtr<UNiagaraHierarchyViewModelBase> HierarchyViewModel;
+	TSharedPtr<FNiagaraHierarchySectionViewModel> SectionViewModel;
+private:
+	TAttribute<ECheckBoxState> IsSectionActive;
+	FOnSectionActivated OnSectionActivatedDelegate;
+	bool bForbidDropOn = false;
+	mutable bool bDraggedOn = false;
+};
+
+class SNiagaraHierarchy : public SCompoundWidget, public FGCObject, public FNotifyHook
+{
 	DECLARE_DELEGATE_RetVal_OneParam(TSharedRef<SWidget>, FOnGenerateRowContentWidget, TSharedRef<FNiagaraHierarchyItemViewModelBase> HierarchyItem);
 	DECLARE_DELEGATE_RetVal_OneParam(TSharedRef<SWidget>, FOnGenerateCustomDetailsPanelNameWidget, TSharedPtr<FNiagaraHierarchyItemViewModelBase> HierarchyItem);
@@ -32,95 +92,135 @@
 public:
 	void Construct(const FArguments& InArgs, TObjectPtr<UNiagaraHierarchyViewModelBase> InHierarchyViewModel);
 	virtual ~SNiagaraHierarchy() override;
-	
+
+	void RefreshSourceItems();
+	void RefreshAllViews(bool bFullRefresh = false);
+	void RequestRefreshAllViewsNextFrame(bool bFullRefresh = false);
 	void RefreshSourceView(bool bFullRefresh = false) const;
+	void RequestRefreshSourceViewNextFrame(bool bFullRefresh = false);
 	void RefreshHierarchyView(bool bFullRefresh = false) const;
-	void RefreshSectionsWidget();
+	void RequestRefreshHierarchyViewNextFrame(bool bFullRefresh = false);
+	void RefreshSectionsView();
+	void RequestRefreshSectionsViewNextFrame();
+
+	void NavigateToHierarchyItem(FNiagaraHierarchyIdentity Identity) const;
+	void NavigateToHierarchyItem(TSharedPtr<FNiagaraHierarchyItemViewModelBase> Item) const;
 	bool IsItemSelected(TSharedPtr<FNiagaraHierarchyItemViewModelBase> Item) const;
-	void SelectObjectInDetailsPanel(UObject* Object) const;
 	
 private:
+	// need to do this to enable focus so we can handle shortcuts
+	virtual bool SupportsKeyboardFocus() const override { return true; }
+	
 	TSharedRef<ITableRow> GenerateSourceItemRow(TSharedPtr<FNiagaraHierarchyItemViewModelBase> HierarchyItem, const TSharedRef<STableViewBase>& TableViewBase);
 	TSharedRef<ITableRow> GenerateHierarchyItemRow(TSharedPtr<FNiagaraHierarchyItemViewModelBase> HierarchyItem, const TSharedRef<STableViewBase>& TableViewBase);
-private:	
+
+	bool FilterForSourceSection(TSharedPtr<const FNiagaraHierarchyItemViewModelBase> ItemViewModel) const;
+private:
+	void Reinitialize();
+
+	void BindToHierarchyRootViewModel();
+	void UnbindFromHierarchyRootViewModel() const;
+	
+	/** Source items reflect the base, unedited status of items to edit into a hierarchy */
+	const TArray<TSharedPtr<FNiagaraHierarchyItemViewModelBase>>& GetSourceItems() const;
+	
+	bool IsDetailsPanelEditingAllowed() const;
+	
 	void RequestRenameSelectedItem();
 	bool CanRequestRenameSelectedItem() const;
 
-	void DeleteItems(TArray<TSharedPtr<FNiagaraHierarchyItemViewModelBase>> ItemsToDelete);
-	
-	void DeleteSelectedItems();
+	void ClearSourceItems() const;
+
+	void DeleteItems(TArray<TSharedPtr<FNiagaraHierarchyItemViewModelBase>> ItemsToDelete) const;
+	void DeleteSelectedItems() const;
 	bool CanDeleteSelectedItems() const;
 
 	void OnItemAdded(TSharedPtr<FNiagaraHierarchyItemViewModelBase> AddedItem);
-	void OnSectionActivated(TSharedPtr<FNiagaraHierarchySectionViewModel> Section) const;
+	void OnHierarchySectionActivated(TSharedPtr<FNiagaraHierarchySectionViewModel> Section);
+	void OnSourceSectionActivated(TSharedPtr<FNiagaraHierarchySectionViewModel> Section);
+	void OnHierarchySectionAdded(TSharedPtr<FNiagaraHierarchySectionViewModel> AddedSection);
+	void OnHierarchySectionDeleted(TSharedPtr<FNiagaraHierarchySectionViewModel> DeletedSection);
+
+	void SetActiveSourceSection(TSharedPtr<struct FNiagaraHierarchySectionViewModel>);
+	TSharedPtr<FNiagaraHierarchySectionViewModel> GetActiveSourceSection() const;
+	UNiagaraHierarchySection* GetActiveSourceSectionData() const;
+	
 	void OnSelectionChanged(TSharedPtr<FNiagaraHierarchyItemViewModelBase> HierarchyItem, ESelectInfo::Type Type, bool bFromHierarchy) const;
 
+	void RunSourceSearch();
+	void OnSourceSearchTextChanged(const FText& Text);
+	void OnSourceSearchTextCommitted(const FText& Text, ETextCommit::Type CommitType);
+	void OnSearchButtonClicked(SSearchBox::SearchDirection SearchDirection);
+
 	virtual FReply OnKeyDown(const FGeometry& MyGeometry, const FKeyEvent& InKeyEvent) override;
 	virtual FReply OnMouseButtonDown(const FGeometry& MyGeometry, const FPointerEvent& MouseEvent) override;
 	virtual FReply OnMouseButtonUp(const FGeometry& MyGeometry, const FPointerEvent& MouseEvent) override;
 
 	TSharedPtr<SWidget> SummonContextMenuForSelectedRows(bool bFromHierarchy) const;
+
+	struct FSearchItem
+	{
+		TArray<TSharedPtr<FNiagaraHierarchyItemViewModelBase>> Path;
+
+		TSharedPtr<FNiagaraHierarchyItemViewModelBase> GetEntry() const
+		{
+			return Path.Num() > 0 ? 
+				Path[Path.Num() - 1] : 
+				nullptr;
+		}
+
+		bool operator==(const FSearchItem& Item) const
+		{
+			return Path == Item.Path;
+		}
+	};
+
+	/** This will recursively generated parent chain paths for all items within the given root. Used for expansion purposes. */
+	void GenerateSearchItems(TSharedRef<FNiagaraHierarchyItemViewModelBase> Root, TArray<TSharedPtr<FNiagaraHierarchyItemViewModelBase>> ParentChain, TArray<FSearchItem>& OutSearchItems);
+	void ExpandSourceSearchResults();
+	void SelectNextSourceSearchResult();
+	void SelectPreviousSourceSearchResult();
+	TOptional<SSearchBox::FSearchResultData> GetSearchResultData() const;
 	
-	/** This handles the drag & drop at the root level of the hierarchy, i.e. when dragging and dropping a parameter into the "empty" hierarchy */
+	FNiagaraHierarchyItemViewModelBase::FCanPerformActionResults CanDropOnRoot(TSharedPtr<FNiagaraHierarchyItemViewModelBase> DraggedItem) const;
+
+	/** Callback functions for the root widget */
 	FReply HandleHierarchyRootDrop(const FGeometry&, const FDragDropEvent& DragDropEvent) const;
-	bool CanDropOnRoot(TSharedPtr<FDragDropOperation> DragDropOp) const;
+	bool OnCanDropOnRoot(TSharedPtr<FDragDropOperation> DragDropOperation) const;
+	void OnRootDragEnter(const FDragDropEvent& DragDropEvent) const;
+	void OnRootDragLeave(const FDragDropEvent& DragDropEvent) const;
+	FSlateColor GetRootIconColor() const;
+
+	virtual FString GetReferencerName() const override; 
+	virtual void AddReferencedObjects(FReferenceCollector& Collector) override;
+
+	virtual void NotifyPostChange(const FPropertyChangedEvent& PropertyChangedEvent, FProperty* PropertyThatChanged) override;
 private:
 	TWeakObjectPtr<UNiagaraHierarchyViewModelBase> HierarchyViewModel;
+	
+	TArray<FSearchItem> SourceSearchResults;
+	TOptional<FSearchItem> FocusedSearchResult;
+
+	mutable TWeakPtr<FNiagaraHierarchyItemViewModelBase> SelectedDetailsPanelItemViewModel;
+	
+	FSlateBrush LightBackgroundBrush;
+	FSlateBrush RecessedBackgroundBrush;
+private:
+	UNiagaraHierarchyRoot* SourceRoot = nullptr;
+	TSharedPtr<FNiagaraHierarchyRootViewModel> SourceRootViewModel;
 	TSharedPtr<STreeView<TSharedPtr<FNiagaraHierarchyItemViewModelBase>>> SourceTreeView;
 	TSharedPtr<STreeView<TSharedPtr<FNiagaraHierarchyItemViewModelBase>>> HierarchyTreeView;
-	TMap<TSharedPtr<FNiagaraHierarchySectionViewModel>, TSharedPtr<class SNiagaraHierarchySection>> SectionsWidgetMap;
-	TSharedPtr<class SWrapBox> SectionBox;
+	TWeakPtr<struct FNiagaraHierarchySectionViewModel> ActiveSourceSection;
+	TSharedPtr<class SWrapBox> SourceSectionBox;
+	TSharedPtr<class SWrapBox> HierarchySectionBox;
+	TSharedPtr<SSearchBox> SourceSearchBox;
 	TSharedPtr<IDetailsView> DetailsPanel;
-	FSlateBrush LightBackgroundBrush;
-	FSlateBrush RecessedBackgroundBrush;
 private:
 	FOnGenerateRowContentWidget OnGenerateRowContentWidget;
 	FOnGenerateCustomDetailsPanelNameWidget OnGenerateCustomDetailsPanelNameWidget;
+	TSharedPtr<FActiveTimerHandle> RefreshHierarchyViewNextFrameHandle;
+	TSharedPtr<FActiveTimerHandle> RefreshSourceViewNextFrameHandle;
+	TSharedPtr<FActiveTimerHandle> RefreshSectionsViewNextFrameHandle;
 };
 
-class SNiagaraHierarchySection : public SCompoundWidget
-{
-	DECLARE_DELEGATE_OneParam(FOnSectionActivated, TSharedPtr<FNiagaraHierarchySectionViewModel> SectionViewModel)
-	
-	SLATE_BEGIN_ARGS(SNiagaraHierarchySection)
-	{}
-		SLATE_EVENT(FIsSelected, IsSelected)
-		SLATE_EVENT(FOnSectionActivated, OnSectionActivated)
-	SLATE_END_ARGS()
-
-public:
-	void Construct(const FArguments& InArgs, TSharedPtr<struct FNiagaraHierarchySectionViewModel> InSection, TWeakObjectPtr<UNiagaraHierarchyViewModelBase> HierarchyViewModel);
-	virtual ~SNiagaraHierarchySection() override;
-
-	bool OnCanAcceptDrop(TSharedPtr<FDragDropOperation> DragDropOperation, EItemDropZone ItemDropZone) const;
-	FReply OnDroppedOn(const FGeometry&, const FDragDropEvent& DragDropEvent, EItemDropZone DropZone) const;
-
-	void EnterEditingMode() const;
-
-	TSharedPtr<struct FNiagaraHierarchySectionViewModel> GetSectionViewModel();
-private:
-	virtual FReply OnMouseButtonDown(const FGeometry& MyGeometry, const FPointerEvent& MouseEvent) override;
-	virtual FReply OnMouseButtonUp(const FGeometry& MyGeometry, const FPointerEvent& MouseEvent) override;
-	
-	virtual FReply OnDragDetected(const FGeometry& MyGeometry, const FPointerEvent& MouseEvent) override;
-	virtual void OnDragLeave(const FDragDropEvent& DragDropEvent) override;
-
-	UNiagaraHierarchySection* TryGetSectionData() const;
-
-	FText GetText() const;
-	FText GetTooltipText() const;
-	void OnRenameSection(const FText& Text, ETextCommit::Type CommitType) const;
-	bool OnVerifySectionRename(const FText& NewName, FText& OutTooltip) const;
-
-	ECheckBoxState GetSectionCheckState() const;
-	void OnSectionCheckChanged(ECheckBoxState NewState);
-private:
-	TSharedPtr<SMenuAnchor> MenuAnchor;
-	TSharedPtr<SInlineEditableTextBlock> InlineEditableTextBlock;
-	TWeakObjectPtr<UNiagaraHierarchyViewModelBase> HierarchyViewModel;
-	TSharedPtr<FNiagaraHierarchySectionViewModel> SectionViewModel;
-private:
-	FOnSectionActivated OnSectionActivatedDelegate;
-};
-
 class SNiagaraHierarchyCategory : public SCompoundWidget
@@ -131,7 +231,7 @@
 	SLATE_END_ARGS()
 
 public:
-	void Construct(const FArguments& InArgs, TSharedPtr<struct FNiagaraHierarchyCategoryViewModel> InCategory);
+	void Construct(const FArguments& InArgs, TSharedPtr<struct FNiagaraHierarchyCategoryViewModel> InCategoryViewModel);
 
 	void EnterEditingMode() const;
 	bool OnVerifyCategoryRename(const FText& NewName, FText& OutTooltip) const;
@@ -140,7 +240,7 @@
 	void OnRenameCategory(const FText& NewText, ETextCommit::Type) const;
 	
 private:
-	TWeakPtr<struct FNiagaraHierarchyCategoryViewModel> Category;
+	TWeakPtr<struct FNiagaraHierarchyCategoryViewModel> CategoryViewModel;
 	TSharedPtr<SInlineEditableTextBlock> InlineEditableTextBlock;
 };
 
@@ -149,12 +249,17 @@
 public:
 	SLATE_BEGIN_ARGS(SNiagaraSectionDragDropTarget)
 		: _DropTargetArgs(SDropTarget::FArguments())
-	{}
+	{}	
 		SLATE_ARGUMENT(SDropTarget::FArguments, DropTargetArgs)
 	SLATE_END_ARGS()
 
-	void Construct(const FArguments& InArgs);
+	void Construct(const FArguments& InArgs, TSharedPtr<FNiagaraHierarchySectionViewModel> OwningSection, EItemDropZone ItemDropZone);
 	
 	virtual FReply OnDragOver(const FGeometry& MyGeometry, const FDragDropEvent& DragDropEvent) override;
 	virtual void OnDragLeave(const FDragDropEvent& DragDropEvent) override;
+
+private:
+	TSharedPtr<FNiagaraHierarchySectionViewModel> OwningSection;
+	/** The drop zone this drop target represents */
+	EItemDropZone DropZone = EItemDropZone::OntoItem;
 };
diff --git a/Engine/Plugins/FX/Niagara/Source/NiagaraEditorWidgets/Private/NiagaraEditorWidgetsUtilities.cpp b/Engine/Plugins/FX/Niagara/Source/NiagaraEditorWidgets/Private/NiagaraEditorWidgetsUtilities.cpp
--- a/Engine/Plugins/FX/Niagara/Source/NiagaraEditorWidgets/Private/NiagaraEditorWidgetsUtilities.cpp
+++ b/Engine/Plugins/FX/Niagara/Source/NiagaraEditorWidgets/Private/NiagaraEditorWidgetsUtilities.cpp
@@ -559,61 +559,4 @@ FString FNiagaraStackEditorWidgetsUtilities::StackEntryToStringForListDebug(UNia
 	return FString::Printf(TEXT("0x%08x - %s - %s"), StackEntry, *StackEntry->GetClass()->GetName(), *StackEntry->GetDisplayName().ToString());
 }
 
-TOptional<FFunctionInputSummaryViewKey> FNiagaraStackEditorWidgetsUtilities::GetSummaryViewInputKeyForFunctionInput(UNiagaraStackFunctionInput* FunctionInput)
-{
-	const FGuid NodeGuid = FunctionInput->GetInputFunctionCallNode().NodeGuid;
-
-	const UNiagaraNodeFunctionCall& InputCallNode = FunctionInput->GetInputFunctionCallNode();
-	if (InputCallNode.IsA<UNiagaraNodeAssignment>())
-	{		
-		return FFunctionInputSummaryViewKey(NodeGuid, FunctionInput->GetInputParameterHandle().GetParameterHandleString());
-	}
-	
-	const TOptional<FGuid> VariableGuid = FunctionInput->GetMetadataGuid();
-	if (VariableGuid.IsSet() && VariableGuid.GetValue().IsValid())
-	{
-		return FFunctionInputSummaryViewKey(NodeGuid, VariableGuid.GetValue());
-	}
-	
-	return TOptional<FFunctionInputSummaryViewKey>();
-}
-
-UNiagaraStackFunctionInput* FNiagaraStackEditorWidgetsUtilities::GetParentInputForSummaryView(UNiagaraStackFunctionInput* FunctionInput)
-{
-	const UNiagaraEmitterEditorData* EditorData = (FunctionInput && FunctionInput->GetEmitterViewModel())? &FunctionInput->GetEmitterViewModel()->GetEditorData() : nullptr;
-
-	// Find outermost function input first, as this will traverse the chain of dynamic inputs to find the base input which is what needs to go to summary view
-	
-	if (EditorData)
-	{
-		while (FunctionInput->GetTypedOuter<UNiagaraStackFunctionInput>())
-		{
-			FunctionInput = FunctionInput->GetTypedOuter<UNiagaraStackFunctionInput>();
-		}
-		
-		TOptional<FFunctionInputSummaryViewKey> Key = FNiagaraStackEditorWidgetsUtilities::GetSummaryViewInputKeyForFunctionInput(FunctionInput);
-
-		TOptional<FNiagaraVariableMetaData> Metadata = FunctionInput->GetMetadata();
-		if (Metadata.IsSet() && !Metadata->ParentAttribute.IsNone())
-		{
-			UNiagaraStackInputCategory* Category = CastChecked<UNiagaraStackInputCategory>(FunctionInput->GetOuter());
-
-			TArray<UNiagaraStackEntry*> Children;
-			Category->GetFilteredChildrenOfTypes(Children, { UNiagaraStackFunctionInput::StaticClass() });
-
-			for (UNiagaraStackEntry* Entry : Children)
-			{
-				UNiagaraStackFunctionInput* Input = CastChecked<UNiagaraStackFunctionInput>(Entry);
-
-				if (Input->GetInputParameterHandle().GetName() == Metadata->ParentAttribute)
-				{
-					return Input;
-				}				
-			}
-		}
-	}
-
-	return FunctionInput;	
-}
-
 #undef LOCTEXT_NAMESPACE
diff --git a/Engine/Plugins/FX/Niagara/Source/NiagaraEditorWidgets/Private/NiagaraEditorWidgetsUtilities.h b/Engine/Plugins/FX/Niagara/Source/NiagaraEditorWidgets/Private/NiagaraEditorWidgetsUtilities.h
--- a/Engine/Plugins/FX/Niagara/Source/NiagaraEditorWidgets/Private/NiagaraEditorWidgetsUtilities.h
+++ b/Engine/Plugins/FX/Niagara/Source/NiagaraEditorWidgets/Private/NiagaraEditorWidgetsUtilities.h
@@ -45,8 +45,4 @@ namespace FNiagaraStackEditorWidgetsUtilities
 	bool HandleDropForStackEntry(const FDragDropEvent& InDragDropEvent, EItemDropZone InDropZone, UNiagaraStackEntry* InTargetEntry, UNiagaraStackEntry::EDropOptions DropOptions);
 
 	FString StackEntryToStringForListDebug(UNiagaraStackEntry* StackEntry);
-	
-	TOptional<FFunctionInputSummaryViewKey> GetSummaryViewInputKeyForFunctionInput(UNiagaraStackFunctionInput* FunctionInput);
-
-	UNiagaraStackFunctionInput* GetParentInputForSummaryView(UNiagaraStackFunctionInput* FunctionInput);
 }
diff --git a/Engine/Plugins/FX/Niagara/Source/NiagaraEditorWidgets/Private/SNiagaraStack.cpp b/Engine/Plugins/FX/Niagara/Source/NiagaraEditorWidgets/Private/SNiagaraStack.cpp
--- a/Engine/Plugins/FX/Niagara/Source/NiagaraEditorWidgets/Private/SNiagaraStack.cpp
+++ b/Engine/Plugins/FX/Niagara/Source/NiagaraEditorWidgets/Private/SNiagaraStack.cpp
@@ -936,9 +936,9 @@ SNiagaraStack::FRowWidgets SNiagaraStack::ConstructNameAndValueWidgetsForItem(UN
 				.IsReadOnly(true)
 			);
 	}
-	else if (Item->IsA<UNiagaraStackFunctionInputCollectionBase>())
+	else if (Item->IsA<UNiagaraStackValueCollection>())
 	{
-		UNiagaraStackFunctionInputCollectionBase* InputCollection = CastChecked<UNiagaraStackFunctionInputCollectionBase>(Item);
+		UNiagaraStackValueCollection* InputCollection = CastChecked<UNiagaraStackValueCollection>(Item);
 		return FRowWidgets(SNew(SNiagaraStackFunctionInputCollection, InputCollection));
 	}
 	else if (Item->IsA<UNiagaraStackModuleItemOutputCollection>() ||
@@ -961,7 +961,7 @@ SNiagaraStack::FRowWidgets SNiagaraStack::ConstructNameAndValueWidgetsForItem(UN
 
 		TAttribute<bool> IsEnabled;
 		IsEnabled.BindUObject(Item, &UNiagaraStackEntry::GetIsEnabledAndOwnerIsEnabled);
-
+		
 		Container->AddFillRowContextMenuHandler(FNiagaraStackPropertyRowUtilities::CreateOnFillRowContextMenu(PropertyRow->GetDetailTreeNode()->CreatePropertyHandle(), PropertyRowWidgets.Actions));
 
 		if (PropertyRowWidgets.WholeRowWidget.IsValid())
@@ -1063,6 +1063,18 @@ SNiagaraStack::FRowWidgets SNiagaraStack::ConstructNameAndValueWidgetsForItem(UN
 			.AutoWrapText(true)
 			.Justification(ETextJustify::Center));
 	}
+	else if (Item->IsA<UNiagaraStackSummaryCategory>())
+	{
+		Container->SetOverrideNameAlignment(HAlign_Left, VAlign_Center);
+		return FRowWidgets(SNew(STextBlock)
+			.TextStyle(&FAppStyle::GetWidgetStyle<FTextBlockStyle>("DetailsView.CategoryTextStyle"))
+			.ToolTipText_UObject(Item, &UNiagaraStackEntry::GetTooltipText)
+			.Text_UObject(Item, &UNiagaraStackEntry::GetDisplayName)
+			.ColorAndOpacity(this, &SNiagaraStack::GetTextColorForItem, Item)
+			.HighlightText_UObject(StackViewModel, &UNiagaraStackViewModel::GetCurrentSearchText)
+			.IsEnabled_UObject(Item, &UNiagaraStackEntry::GetOwnerIsEnabled),
+			SNullWidget::NullWidget);
+	}
 	else if (Item->IsA<UNiagaraStackSpacer>())
 	{
 		UNiagaraStackSpacer* Spacer = CastChecked<UNiagaraStackSpacer>(Item);
diff --git a/Engine/Plugins/FX/Niagara/Source/NiagaraEditorWidgets/Private/Stack/SNiagaraStackFunctionInputCollection.cpp b/Engine/Plugins/FX/Niagara/Source/NiagaraEditorWidgets/Private/Stack/SNiagaraStackFunctionInputCollection.cpp
--- a/Engine/Plugins/FX/Niagara/Source/NiagaraEditorWidgets/Private/Stack/SNiagaraStackFunctionInputCollection.cpp
+++ b/Engine/Plugins/FX/Niagara/Source/NiagaraEditorWidgets/Private/Stack/SNiagaraStackFunctionInputCollection.cpp
@@ -7,10 +7,10 @@
 #include "Widgets/Input/SCheckBox.h"
 #include "Widgets/Layout/SWrapBox.h"
 
-void SNiagaraStackFunctionInputCollection::Construct(const FArguments& InArgs, UNiagaraStackFunctionInputCollectionBase* InInputCollection)
+void SNiagaraStackFunctionInputCollection::Construct(const FArguments& InArgs, UNiagaraStackValueCollection* PropertyCollectionBase)
 {
-	InputCollection = InInputCollection;
-	InputCollection->OnStructureChanged().AddSP(this, &SNiagaraStackFunctionInputCollection::InputCollectionStructureChanged);
+	PropertyCollection = PropertyCollectionBase;
+	PropertyCollection->OnStructureChanged().AddSP(this, &SNiagaraStackFunctionInputCollection::InputCollectionStructureChanged);
 
 	ChildSlot
 	[
@@ -21,9 +21,9 @@ void SNiagaraStackFunctionInputCollection::Construct(const FArguments& InArgs, U
 		[
 			SNew(STextBlock)
 			.TextStyle(FNiagaraEditorWidgetsStyle::Get(), "NiagaraEditor.Stack.ItemText")
-			.ToolTipText_UObject(InputCollection, &UNiagaraStackEntry::GetTooltipText)
-			.Text_UObject(InputCollection, &UNiagaraStackEntry::GetDisplayName)
-			.IsEnabled_UObject(InputCollection, &UNiagaraStackEntry::GetOwnerIsEnabled)
+			.ToolTipText_UObject(PropertyCollection, &UNiagaraStackEntry::GetTooltipText)
+			.Text_UObject(PropertyCollection, &UNiagaraStackEntry::GetDisplayName)
+			.IsEnabled_UObject(PropertyCollection, &UNiagaraStackEntry::GetOwnerIsEnabled)
 			.Visibility(this, &SNiagaraStackFunctionInputCollection::GetLabelVisibility)
 		]
 		+ SVerticalBox::Slot()
@@ -41,13 +41,13 @@ void SNiagaraStackFunctionInputCollection::Construct(const FArguments& InArgs, U
 
 EVisibility SNiagaraStackFunctionInputCollection::GetLabelVisibility() const
 {
-	return InputCollection->GetShouldDisplayLabel() ? EVisibility::Visible : EVisibility::Collapsed;
+	return PropertyCollection->GetShouldDisplayLabel() ? EVisibility::Visible : EVisibility::Collapsed;
 }
 
 void SNiagaraStackFunctionInputCollection::ConstructSectionButtons()
 {
 	SectionSelectorBox->ClearChildren();
-	for (FText Section : InputCollection->GetSections())
+	for (FText Section : PropertyCollection->GetSections())
 	{
 		SectionSelectorBox->AddSlot()
 		[
@@ -55,6 +55,7 @@ void SNiagaraStackFunctionInputCollection::ConstructSectionButtons()
 			.Style(FAppStyle::Get(), "DetailsView.SectionButton")
 			.OnCheckStateChanged(this, &SNiagaraStackFunctionInputCollection::OnSectionChecked, Section)
 			.IsChecked(this, &SNiagaraStackFunctionInputCollection::GetSectionCheckState, Section)
+			.ToolTipText(this, &SNiagaraStackFunctionInputCollection::GetTooltipText, Section)
 			[
 				SNew(STextBlock)
 				.TextStyle(FAppStyle::Get(), "SmallText")
@@ -71,13 +72,18 @@ void SNiagaraStackFunctionInputCollection::InputCollectionStructureChanged(ENiag
 
 ECheckBoxState SNiagaraStackFunctionInputCollection::GetSectionCheckState(FText Section) const
 {
-	return Section.IdenticalTo(InputCollection->GetActiveSection()) ? ECheckBoxState::Checked : ECheckBoxState::Unchecked;
+	return Section.IdenticalTo(PropertyCollection->GetActiveSection()) ? ECheckBoxState::Checked : ECheckBoxState::Unchecked;
 }
 
 void SNiagaraStackFunctionInputCollection::OnSectionChecked(ECheckBoxState CheckState, FText Section)
 {
 	if (CheckState == ECheckBoxState::Checked)
 	{
-		InputCollection->SetActiveSection(Section);
+		PropertyCollection->SetActiveSection(Section);
 	}
-}
\ No newline at end of file
+}
+
+FText SNiagaraStackFunctionInputCollection::GetTooltipText(FText Section) const
+{
+	return PropertyCollection->GetTooltipForSection(Section.ToString());
+}
diff --git a/Engine/Plugins/FX/Niagara/Source/NiagaraEditorWidgets/Private/Stack/SNiagaraStackFunctionInputCollection.h b/Engine/Plugins/FX/Niagara/Source/NiagaraEditorWidgets/Private/Stack/SNiagaraStackFunctionInputCollection.h
--- a/Engine/Plugins/FX/Niagara/Source/NiagaraEditorWidgets/Private/Stack/SNiagaraStackFunctionInputCollection.h
+++ b/Engine/Plugins/FX/Niagara/Source/NiagaraEditorWidgets/Private/Stack/SNiagaraStackFunctionInputCollection.h
@@ -3,6 +3,7 @@
 #pragma once
 
 #include "NiagaraEditorCommon.h"
+#include "ViewModels/Stack/NiagaraStackFunctionInputCollection.h"
 #include "Widgets/SCompoundWidget.h"
 #include "Widgets/DeclarativeSyntaxSupport.h"
 #include "Widgets/Views/SExpanderArrow.h"
@@ -17,7 +18,7 @@ public:
 	SLATE_BEGIN_ARGS(SNiagaraStackFunctionInputCollection) {}
 	SLATE_END_ARGS();
 
-	void Construct(const FArguments& InArgs, UNiagaraStackFunctionInputCollectionBase* InInputCollection);
+	void Construct(const FArguments& InArgs, UNiagaraStackValueCollection* PropertyCollectionBase);
 
 private:
 	EVisibility GetLabelVisibility() const;
@@ -30,8 +31,10 @@ private:
 
 	void OnSectionChecked(ECheckBoxState CheckState, FText Section);
 
+	FText GetTooltipText(FText Section) const;
+
 private:
-	UNiagaraStackFunctionInputCollectionBase* InputCollection;
+	UNiagaraStackValueCollection* PropertyCollection;
 
 	TSharedPtr<SWrapBox> SectionSelectorBox;
 };
diff --git a/Engine/Plugins/FX/Niagara/Source/NiagaraEditorWidgets/Private/Stack/SNiagaraStackFunctionInputValue.cpp b/Engine/Plugins/FX/Niagara/Source/NiagaraEditorWidgets/Private/Stack/SNiagaraStackFunctionInputValue.cpp
--- a/Engine/Plugins/FX/Niagara/Source/NiagaraEditorWidgets/Private/Stack/SNiagaraStackFunctionInputValue.cpp
+++ b/Engine/Plugins/FX/Niagara/Source/NiagaraEditorWidgets/Private/Stack/SNiagaraStackFunctionInputValue.cpp
@@ -250,29 +250,6 @@ void SNiagaraStackFunctionInputValue::Construct(const FArguments& InArgs, UNiaga
 		]
 	];
 
-	if (FunctionInput->GetTypedOuter<UNiagaraStackSummaryViewObject>() != nullptr)
-	{
-		ChildrenBox->AddSlot()
-		.AutoWidth()
-		.VAlign(VAlign_Center)
-		.Padding(3, 0, 0, 0)
-		[
-			SNew(SComboButton)
-			.ButtonStyle(FNiagaraEditorWidgetsStyle::Get(), "NiagaraEditor.Stack.SimpleButton")
-			.OnGetMenuContent(this, &SNiagaraStackFunctionInputValue::GetFilteredViewPropertiesContent)
-			.Visibility(this, &SNiagaraStackFunctionInputValue::GetFilteredViewContextButtonVisibility)
-			.ToolTipText(LOCTEXT("InputSummaryOptionsToolTip", "Input Summary Options"))
-			.IsEnabled(true)
-			.ButtonContent()
-			[
-				SNew(STextBlock)
-				.Font(FAppStyle::Get().GetFontStyle("FontAwesome.10"))
-				.Text(FText::FromString(FString(TEXT("\xf0ca")/* fa-list-ul */)))
-				.ColorAndOpacity(FStyleColors::AccentYellow)
-			]
-		];
-	}
-
 	ValueModeForGeneratedWidgets = FunctionInput->GetValueMode();
 }
 
@@ -1442,304 +1419,4 @@ void SNiagaraStackFunctionInputValue::TriggerRefresh(const TMap<EScriptSource, b
 	ActionSelector->ExpandTree();
 }
 
-const FSlateBrush* SNiagaraStackFunctionInputValue::GetFilteredViewIcon() const
-{
-	return FAppStyle::GetBrush(TEXT("TextureEditor.GreenChannel.Small"));
-}
-
-EVisibility SNiagaraStackFunctionInputValue::GetFilteredViewContextButtonVisibility() const
-{
-	UNiagaraEmitterEditorData* EditorData = FunctionInput->GetEmitterViewModel()? &FunctionInput->GetEmitterViewModel()->GetOrCreateEditorData() : nullptr;
-	
-	if (!EditorData || EditorData->ShouldShowSummaryView() || FunctionInput->GetTypedOuter<UNiagaraStackSummaryViewObject>() == nullptr || FunctionInput->GetEmitterViewModel()->GetSummaryIsInEditMode() == false)
-	{
-		return EVisibility::Hidden;
-	}
-	
-	UNiagaraStackFunctionInput* ParentInput = FNiagaraStackEditorWidgetsUtilities::GetParentInputForSummaryView(FunctionInput);
-	TOptional<FFunctionInputSummaryViewKey> Key = FNiagaraStackEditorWidgetsUtilities::GetSummaryViewInputKeyForFunctionInput(ParentInput);
-	TOptional<FFunctionInputSummaryViewMetadata> Metadata = Key.IsSet() ? EditorData->GetSummaryViewMetaData(Key.GetValue()) : TOptional<FFunctionInputSummaryViewMetadata>();
-	return Metadata.IsSet() ? EVisibility::Visible : EVisibility::Collapsed;
-}
-
-bool SNiagaraStackFunctionInputValue::GetSummaryViewChangeEnabledStateAllowed() const
-{
-	// Only show if we aren't a child input, and if we're in the summary view category
-	UNiagaraStackFunctionInput* ParentInput = FNiagaraStackEditorWidgetsUtilities::GetParentInputForSummaryView(FunctionInput);
-	return FunctionInput == ParentInput;	
-}
-
-bool SNiagaraStackFunctionInputValue::GetSummaryViewChangeDisplayNameAllowed() const
-{
-	// Only show if we aren't a child input, and if we're in the summary view category
-	// Also disallow changing display name on a direct attribute set.
-	UNiagaraStackFunctionInput* ParentInput = FNiagaraStackEditorWidgetsUtilities::GetParentInputForSummaryView(FunctionInput);	
-	const UNiagaraNodeFunctionCall& InputCallNode = FunctionInput->GetInputFunctionCallNode();
-	return FunctionInput == ParentInput && !InputCallNode.IsA<UNiagaraNodeAssignment>();
-}
-
-bool SNiagaraStackFunctionInputValue::GetSummaryViewChangeCategoryAllowed() const
-{
-	// Only show if we aren't a child input, and if we're in the summary view category
-	UNiagaraStackFunctionInput* ParentInput = FNiagaraStackEditorWidgetsUtilities::GetParentInputForSummaryView(FunctionInput);
-	return FunctionInput == ParentInput;	
-}
-
-TSharedRef<SWidget> SNiagaraStackFunctionInputValue::GetFilteredViewPropertiesContent()
-{
-	FMenuBuilder MenuBuilder(false, MakeShared<FUICommandList>());
-	MenuBuilder.BeginSection("SummaryViewDetails", LOCTEXT("SummaryViewDetails", "Emitter Summary Details"));
-
-
-	// Display Name
-	TSharedRef<SWidget> NameWidget =
-		SNew(SBox)
-		.WidthOverride(100)
-		.Padding(FMargin(5, 0, 0, 0))
-		[
-			SNew(SEditableTextBox)
-			.Text(this, &SNiagaraStackFunctionInputValue::GetFilteredViewDisplayName)
-			.IsEnabled(this, &SNiagaraStackFunctionInputValue::GetSummaryViewChangeDisplayNameAllowed)
-			.OnVerifyTextChanged(this, &SNiagaraStackFunctionInputValue::VerifyFilteredViewDisplayName)
-			.OnTextCommitted(this, &SNiagaraStackFunctionInputValue::FilteredViewDisplayNameTextCommitted)
-		];
-	MenuBuilder.AddWidget(NameWidget, LOCTEXT("FilteredViewDisplayName", "Display Name"));
-
-
-	// Category
-	TSharedRef<SWidget> CategoryWidget =
-		SNew(SBox)
-		.WidthOverride(100)
-		.Padding(FMargin(5, 0, 0, 0))
-		[
-			SNew(SEditableTextBox)
-			.IsEnabled(this, &SNiagaraStackFunctionInputValue::GetSummaryViewChangeCategoryAllowed)
-			.Text(this, &SNiagaraStackFunctionInputValue::GetFilteredViewCategory)
-			.OnVerifyTextChanged(this, &SNiagaraStackFunctionInputValue::VerifyFilteredViewCategory)
-			.OnTextCommitted(this, &SNiagaraStackFunctionInputValue::FilteredViewCategoryTextCommitted)
-		];
-	MenuBuilder.AddWidget(CategoryWidget, LOCTEXT("FilteredViewCategory", "Category"));
-
-
-	// Sort Index
-	TSharedRef<SWidget> SortIndexWidget =
-		SNew(SBox)
-		.WidthOverride(100)
-		.Padding(FMargin(5, 0, 0, 0))
-		[			
-			SNew(SEditableTextBox)
-			.Text(this, &SNiagaraStackFunctionInputValue::GetFilteredViewSortIndex)
-			.OnVerifyTextChanged(this, &SNiagaraStackFunctionInputValue::VerifyFilteredSortIndex)
-			.OnTextCommitted(this, &SNiagaraStackFunctionInputValue::FilteredSortIndexTextCommitted)
-		];
-	MenuBuilder.AddWidget(SortIndexWidget, LOCTEXT("FilteredViewSortIndex", "Sort Index"));
-
-	// Visible
-	TSharedRef<SWidget> VisibleWidget =
-		SNew(SBox)
-		.WidthOverride(100)
-		.Padding(FMargin(5, 0, 0, 0))
-		[
-			SNew(SCheckBox)
-			.IsChecked(this, &SNiagaraStackFunctionInputValue::GetFilteredViewVisibleCheckState)
-			.OnCheckStateChanged(this, &SNiagaraStackFunctionInputValue::FilteredVisibleCheckStateChanged)
-		];
-	MenuBuilder.AddWidget(VisibleWidget, LOCTEXT("FilteredViewVisible", "Is Visible"));
-
-	return MenuBuilder.MakeWidget();
-}
-
-FText SNiagaraStackFunctionInputValue::GetFilteredViewDisplayName() const
-{
-	const UNiagaraEmitterEditorData* EditorData = FunctionInput->GetEmitterViewModel()? &FunctionInput->GetEmitterViewModel()->GetEditorData() : nullptr;
-	UNiagaraStackFunctionInput* ParentInput = FNiagaraStackEditorWidgetsUtilities::GetParentInputForSummaryView(FunctionInput);
-
-	TOptional<FFunctionInputSummaryViewKey> Key = FNiagaraStackEditorWidgetsUtilities::GetSummaryViewInputKeyForFunctionInput(ParentInput);
-	TOptional<FFunctionInputSummaryViewMetadata> Metadata = EditorData != nullptr && Key.IsSet() ? EditorData->GetSummaryViewMetaData(Key.GetValue()) : TOptional<FFunctionInputSummaryViewMetadata>();
-
-	const FName SummaryViewName = Metadata.IsSet() ? Metadata->DisplayName : NAME_None;
-	return (SummaryViewName != NAME_None) ? FText::FromName(SummaryViewName) : FunctionInput->GetDisplayName();
-}
-
-bool SNiagaraStackFunctionInputValue::VerifyFilteredViewDisplayName(const FText& InText, FText& OutErrorMessage) const
-{
-	const FString DisplayName = InText.ToString();
-
-	if (DisplayName.Len() >= FNiagaraConstants::MaxCategoryNameLength)
-	{
-		OutErrorMessage = FText::FormatOrdered(LOCTEXT("FilteredViewDisplayNameTooLong", "Display name name cannot exceed {0} characters."), FNiagaraConstants::MaxCategoryNameLength);
-		return false;
-	}
-
-	return true;
-}
-
-void SNiagaraStackFunctionInputValue::FilteredViewDisplayNameTextCommitted(const FText& Text, ETextCommit::Type CommitType)
-{
-	UNiagaraEmitterEditorData* EditorData = FunctionInput->GetEmitterViewModel()? &FunctionInput->GetEmitterViewModel()->GetOrCreateEditorData() : nullptr;	
-	UNiagaraStackFunctionInput* ParentInput = FNiagaraStackEditorWidgetsUtilities::GetParentInputForSummaryView(FunctionInput);
-	TOptional<FFunctionInputSummaryViewKey> Key = FNiagaraStackEditorWidgetsUtilities::GetSummaryViewInputKeyForFunctionInput(ParentInput);
-	
-	if (EditorData && Key.IsSet())
-	{
-		const FName NewName = Text.IsEmptyOrWhitespace()? NAME_None : FName(Text.ToString());
-		TOptional<FFunctionInputSummaryViewMetadata> SummaryViewMetaData = EditorData->GetSummaryViewMetaData(Key.GetValue());	
-		if (SummaryViewMetaData.IsSet() == false)
-		{
-			SummaryViewMetaData = FFunctionInputSummaryViewMetadata();
-		}
-		if (NewName != SummaryViewMetaData->DisplayName)
-		{
-			FScopedTransaction ScopedTransaction(FText::Format(LOCTEXT("SummaryViewChangedInputDisplayName", "Changed summary view display name for {0}"), FunctionInput->GetDisplayName()));
-			EditorData->Modify();
-			
-			SummaryViewMetaData->DisplayName = NewName;
-			EditorData->SetSummaryViewMetaData(Key.GetValue(), SummaryViewMetaData);
-		}		
-	}
-}
-
-FText GetFunctionInputCategory(UNiagaraStackFunctionInput* Input)
-{
-	UNiagaraStackInputCategory* Category = Cast<UNiagaraStackInputCategory>(Input->GetOuter());
-
-	return Category && !Category->GetCategoryName().EqualTo(UNiagaraStackFunctionInputCollectionBase::UncategorizedName)? Category->GetCategoryName() : FText();
-}
-
-FText SNiagaraStackFunctionInputValue::GetFilteredViewCategory() const
-{
-	const UNiagaraEmitterEditorData* EditorData = FunctionInput->GetEmitterViewModel()? &FunctionInput->GetEmitterViewModel()->GetEditorData() : nullptr;	
-	UNiagaraStackFunctionInput* ParentInput = FNiagaraStackEditorWidgetsUtilities::GetParentInputForSummaryView(FunctionInput);
-	
-	TOptional<FFunctionInputSummaryViewKey> Key = FNiagaraStackEditorWidgetsUtilities::GetSummaryViewInputKeyForFunctionInput(ParentInput);
-	TOptional<FFunctionInputSummaryViewMetadata> Metadata = EditorData != nullptr && Key.IsSet() ? EditorData->GetSummaryViewMetaData(Key.GetValue()) : TOptional<FFunctionInputSummaryViewMetadata>();
-
-	const FName SummaryViewCategory = Metadata.IsSet() ? Metadata->Category : NAME_None;
-	return (SummaryViewCategory != NAME_None) ? FText::FromName(SummaryViewCategory) : GetFunctionInputCategory(FunctionInput);
-}
-
-bool SNiagaraStackFunctionInputValue::VerifyFilteredViewCategory(const FText& InText, FText& OutErrorMessage) const
-{
-	const FString CategoryName = InText.ToString();
-
-	if (CategoryName.Len() >= FNiagaraConstants::MaxCategoryNameLength)
-	{
-		OutErrorMessage = FText::FormatOrdered(LOCTEXT("FilteredViewCategoryNameTooLong", "Category name cannot exceed {0} characters."), FNiagaraConstants::MaxCategoryNameLength);
-		return false;
-	}
-
-	return true;
-}
-
-void SNiagaraStackFunctionInputValue::FilteredViewCategoryTextCommitted(const FText& Text, ETextCommit::Type CommitType)
-{
-	UNiagaraEmitterEditorData* EditorData = FunctionInput->GetEmitterViewModel()? &FunctionInput->GetEmitterViewModel()->GetOrCreateEditorData() : nullptr;	
-	UNiagaraStackFunctionInput* ParentInput = FNiagaraStackEditorWidgetsUtilities::GetParentInputForSummaryView(FunctionInput);
-	TOptional<FFunctionInputSummaryViewKey> Key = FNiagaraStackEditorWidgetsUtilities::GetSummaryViewInputKeyForFunctionInput(ParentInput);
-	
-	if (EditorData && Key.IsSet())
-	{
-		const FName NewCategory = Text.IsEmptyOrWhitespace()? NAME_None : FName(Text.ToString());
-		TOptional<FFunctionInputSummaryViewMetadata> SummaryViewMetaData = EditorData->GetSummaryViewMetaData(Key.GetValue());	
-		if (SummaryViewMetaData.IsSet() == false)
-		{
-			SummaryViewMetaData = FFunctionInputSummaryViewMetadata();
-		}
-		if (NewCategory != SummaryViewMetaData->Category)
-		{
-			FScopedTransaction ScopedTransaction(FText::Format(LOCTEXT("SummaryViewChangedInputCategory", "Changed summary view category for {0}"), FunctionInput->GetDisplayName()));
-			EditorData->Modify();
-			
-			SummaryViewMetaData->Category = NewCategory;
-			EditorData->SetSummaryViewMetaData(Key.GetValue(), SummaryViewMetaData);
-		}		
-	}
-}
-
-
-FText SNiagaraStackFunctionInputValue::GetFilteredViewSortIndex() const
-{
-	UNiagaraEmitterEditorData* EditorData = FunctionInput->GetEmitterViewModel()? &FunctionInput->GetEmitterViewModel()->GetOrCreateEditorData() : nullptr;	
-	UNiagaraStackFunctionInput* ParentInput = FNiagaraStackEditorWidgetsUtilities::GetParentInputForSummaryView(FunctionInput);
-
-	TOptional<FFunctionInputSummaryViewKey> Key = FNiagaraStackEditorWidgetsUtilities::GetSummaryViewInputKeyForFunctionInput(ParentInput);
-	TOptional<FFunctionInputSummaryViewMetadata> Metadata = EditorData != nullptr && Key.IsSet() ? EditorData->GetSummaryViewMetaData(Key.GetValue()) : TOptional<FFunctionInputSummaryViewMetadata>();
-
-	int32 SortIndex = Metadata.IsSet() ? Metadata->SortIndex : INDEX_NONE;
-	return SortIndex != INDEX_NONE ? FText::FromString(FString::FromInt(SortIndex)) : FText::GetEmpty();
-}
-
-bool SNiagaraStackFunctionInputValue::VerifyFilteredSortIndex(const FText& InText, FText& OutErrorMessage) const
-{
-	if (!InText.IsEmptyOrWhitespace() && !InText.IsNumeric())
-	{
-		OutErrorMessage = LOCTEXT("FilteredSortIndexNotValid", "Sort Index must be empty or a valid whole number");
-		return false;		
-	}
-	return true;
-}
-
-void SNiagaraStackFunctionInputValue::FilteredSortIndexTextCommitted(const FText& Text, ETextCommit::Type CommitType)
-{
-	UNiagaraEmitterEditorData* EditorData = FunctionInput->GetEmitterViewModel()? &FunctionInput->GetEmitterViewModel()->GetOrCreateEditorData() : nullptr;	
-	UNiagaraStackFunctionInput* ParentInput = FNiagaraStackEditorWidgetsUtilities::GetParentInputForSummaryView(FunctionInput);
-	TOptional<FFunctionInputSummaryViewKey> Key = FNiagaraStackEditorWidgetsUtilities::GetSummaryViewInputKeyForFunctionInput(ParentInput);
-
-	int32 NewValue = !Text.IsEmptyOrWhitespace() && Text.IsNumeric()? FCString::Atoi(*Text.ToString()) : INDEX_NONE;
-	
-	if (EditorData && Key.IsSet())
-	{		
-		TOptional<FFunctionInputSummaryViewMetadata> SummaryViewMetaData = EditorData->GetSummaryViewMetaData(Key.GetValue());	
-		if (SummaryViewMetaData.IsSet() == false)
-		{
-			SummaryViewMetaData = FFunctionInputSummaryViewMetadata();
-		}
-		if (NewValue != SummaryViewMetaData->SortIndex)
-		{
-			FScopedTransaction ScopedTransaction(FText::Format(LOCTEXT("SummaryViewChangedInputSortIndex", "Changed summary view sorty index for {0}"), FunctionInput->GetDisplayName()));
-			EditorData->Modify();
-			
-			SummaryViewMetaData->SortIndex = NewValue;
-			EditorData->SetSummaryViewMetaData(Key.GetValue(), SummaryViewMetaData);
-		}		
-	}
-}
-
-ECheckBoxState SNiagaraStackFunctionInputValue::GetFilteredViewVisibleCheckState() const
-{
-	UNiagaraEmitterEditorData* EditorData = FunctionInput->GetEmitterViewModel() ? &FunctionInput->GetEmitterViewModel()->GetOrCreateEditorData() : nullptr;
-	UNiagaraStackFunctionInput* ParentInput = FNiagaraStackEditorWidgetsUtilities::GetParentInputForSummaryView(FunctionInput);
-
-	TOptional<FFunctionInputSummaryViewKey> Key = FNiagaraStackEditorWidgetsUtilities::GetSummaryViewInputKeyForFunctionInput(ParentInput);
-	TOptional<FFunctionInputSummaryViewMetadata> Metadata = EditorData != nullptr && Key.IsSet() ? EditorData->GetSummaryViewMetaData(Key.GetValue()) : TOptional<FFunctionInputSummaryViewMetadata>();
-
-	return Metadata.IsSet() && Metadata->bVisible ? ECheckBoxState::Checked : ECheckBoxState::Unchecked;
-}
-
-void SNiagaraStackFunctionInputValue::FilteredVisibleCheckStateChanged(ECheckBoxState CheckBoxState)
-{
-	UNiagaraEmitterEditorData* EditorData = FunctionInput->GetEmitterViewModel() ? &FunctionInput->GetEmitterViewModel()->GetOrCreateEditorData() : nullptr;
-	UNiagaraStackFunctionInput* ParentInput = FNiagaraStackEditorWidgetsUtilities::GetParentInputForSummaryView(FunctionInput);
-	TOptional<FFunctionInputSummaryViewKey> Key = FNiagaraStackEditorWidgetsUtilities::GetSummaryViewInputKeyForFunctionInput(ParentInput);
-
-	bool bNewValue = CheckBoxState == ECheckBoxState::Checked;
-	if (EditorData && Key.IsSet())
-	{
-		TOptional<FFunctionInputSummaryViewMetadata> SummaryViewMetaData = EditorData->GetSummaryViewMetaData(Key.GetValue());
-		if (SummaryViewMetaData.IsSet() == false)
-		{
-			SummaryViewMetaData = FFunctionInputSummaryViewMetadata();
-		}
-		if (bNewValue != SummaryViewMetaData->bVisible)
-		{
-			FScopedTransaction ScopedTransaction(FText::Format(LOCTEXT("SummaryViewChangedVisibilityForInput", "Changed summary view visibility for {0}"), FunctionInput->GetDisplayName()));
-			EditorData->Modify();
-
-			SummaryViewMetaData->bVisible = bNewValue;
-			EditorData->SetSummaryViewMetaData(Key.GetValue(), SummaryViewMetaData);
-		}
-	}
-}
-
-
 #undef LOCTEXT_NAMESPACE
diff --git a/Engine/Plugins/FX/Niagara/Source/NiagaraEditorWidgets/Private/Stack/SNiagaraStackFunctionInputValue.h b/Engine/Plugins/FX/Niagara/Source/NiagaraEditorWidgets/Private/Stack/SNiagaraStackFunctionInputValue.h
--- a/Engine/Plugins/FX/Niagara/Source/NiagaraEditorWidgets/Private/Stack/SNiagaraStackFunctionInputValue.h
+++ b/Engine/Plugins/FX/Niagara/Source/NiagaraEditorWidgets/Private/Stack/SNiagaraStackFunctionInputValue.h
@@ -157,28 +157,6 @@ private:
 
 	FReply ScratchButtonPressed() const;
 
-	const FSlateBrush* GetFilteredViewIcon() const;
-	EVisibility GetFilteredViewContextButtonVisibility() const;
-	bool GetSummaryViewChangeEnabledStateAllowed() const;
-	bool GetSummaryViewChangeDisplayNameAllowed() const;
-	bool GetSummaryViewChangeCategoryAllowed() const;
-	TSharedRef<SWidget> GetFilteredViewPropertiesContent();
-	
-	FText GetFilteredViewDisplayName() const;
-	bool VerifyFilteredViewDisplayName(const FText& InText, FText& OutErrorMessage) const;
-	void FilteredViewDisplayNameTextCommitted(const FText& Text, ETextCommit::Type CommitType);
-
-	FText GetFilteredViewCategory() const;
-	bool VerifyFilteredViewCategory(const FText& InText, FText& OutErrorMessage) const;
-	void FilteredViewCategoryTextCommitted(const FText& Text, ETextCommit::Type CommitType);
-
-	FText GetFilteredViewSortIndex() const;
-	bool VerifyFilteredSortIndex(const FText& InText, FText& OutErrorMessage) const;
-	void FilteredSortIndexTextCommitted(const FText& Text, ETextCommit::Type CommitType);
-
-	ECheckBoxState GetFilteredViewVisibleCheckState() const;
-	void FilteredVisibleCheckStateChanged(ECheckBoxState CheckBoxState);
-
 private:
 	UNiagaraStackFunctionInput* FunctionInput;
 
diff --git a/Engine/Plugins/FX/Niagara/Source/NiagaraEditorWidgets/Private/Stack/SNiagaraStackItem.cpp b/Engine/Plugins/FX/Niagara/Source/NiagaraEditorWidgets/Private/Stack/SNiagaraStackItem.cpp
--- a/Engine/Plugins/FX/Niagara/Source/NiagaraEditorWidgets/Private/Stack/SNiagaraStackItem.cpp
+++ b/Engine/Plugins/FX/Niagara/Source/NiagaraEditorWidgets/Private/Stack/SNiagaraStackItem.cpp
@@ -72,23 +72,38 @@ void SNiagaraStackItem::Construct(const FArguments& InArgs, UNiagaraStackItem& I
 	// Edit Mode Button
 	if (Item->SupportsEditMode())
 	{
+		TSharedPtr<SButton> EditButton;
+		
 		RowBox->AddSlot()
 		.AutoWidth()
 		.VAlign(VAlign_Center)
 		.Padding(0, 0, 2, 0)
 		[
-			SNew(SButton)
+			SAssignNew(EditButton, SButton)
 			.IsFocusable(false)
 			.ToolTipText(this, &SNiagaraStackItem::GetEditModeButtonToolTipText)
-			.ButtonStyle(FNiagaraEditorWidgetsStyle::Get(), "NiagaraEditor.Stack.SimpleButton")
+			.ButtonStyle(FAppStyle::Get(), "RoundButton")
 			.OnClicked(this, &SNiagaraStackItem::EditModeButtonClicked)
 			.ContentPadding(1)
-			[
+			.Visibility(this, &SNiagaraStackItem::IsEditButtonVisible)
+		];
+
+		if(Item->GetEditModeButtonText().IsSet())
+		{
+			EditButton->SetContent
+			(
+				SNew(STextBlock)
+				.Text(Item->GetEditModeButtonText().GetValue())			
+			);
+		}
+		else
+		{
+			EditButton->SetContent
+			(
 				SNew(SImage)
 				.Image(FAppStyle::Get().GetBrush("Icons.Edit"))
-				.ColorAndOpacity(this, &SNiagaraStackItem::GetEditModeButtonIconColor)
-			]
-		];
+			);
+		}
 	}
 
 	// Reset to base button
@@ -191,25 +206,26 @@ TSharedRef<SWidget> SNiagaraStackItem::AddContainerForRowWidgets(TSharedRef<SWid
 
 FText SNiagaraStackItem::GetEditModeButtonToolTipText() const
 {
-	return Item->GetEditModeIsActive() 
-		? LOCTEXT("DisableEditModeToolTip", "Disable Edit Mode")
-		: LOCTEXT("EnableEditModeToolTip", "Enable Edit Mode");
+	return Item->GetEditModeButtonTooltip().IsSet() ? Item->GetEditModeButtonTooltip().GetValue() : LOCTEXT("EditButtonDefaultTooltip", "Activate Edit Mode");
 }
 
 FReply SNiagaraStackItem::EditModeButtonClicked()
 {
 	if (Item->SupportsEditMode())
 	{
-		Item->SetEditModeIsActive(!Item->GetEditModeIsActive());
+		Item->OnEditButtonClicked();
 	}
 	return FReply::Handled();
 }
 
-FSlateColor SNiagaraStackItem::GetEditModeButtonIconColor() const
+EVisibility SNiagaraStackItem::IsEditButtonVisible() const
 {
-	return Item->GetEditModeIsActive()
-		? FStyleColors::AccentYellow
-		: FSlateColor::UseSubduedForeground();
+	if(Item->SupportsEditMode())
+	{
+		return Item->IsEditButtonVisible();
+	}
+
+	return EVisibility::Collapsed;
 }
 
 EVisibility SNiagaraStackItem::GetResetToBaseButtonVisibility() const
diff --git a/Engine/Plugins/FX/Niagara/Source/NiagaraEditorWidgets/Private/Stack/SNiagaraStackItem.h b/Engine/Plugins/FX/Niagara/Source/NiagaraEditorWidgets/Private/Stack/SNiagaraStackItem.h
--- a/Engine/Plugins/FX/Niagara/Source/NiagaraEditorWidgets/Private/Stack/SNiagaraStackItem.h
+++ b/Engine/Plugins/FX/Niagara/Source/NiagaraEditorWidgets/Private/Stack/SNiagaraStackItem.h
@@ -25,11 +25,9 @@ protected:
 	virtual TSharedRef<SWidget> AddContainerForRowWidgets(TSharedRef<SWidget> RowWidgets);
 
 private:
-	FText GetEditModeButtonToolTipText() const;
-
 	FReply EditModeButtonClicked();
-
-	FSlateColor GetEditModeButtonIconColor() const;
+	EVisibility IsEditButtonVisible() const;
+	FText GetEditModeButtonToolTipText() const;
 
 	EVisibility GetResetToBaseButtonVisibility() const;
 
diff --git a/Engine/Plugins/FX/Niagara/Source/NiagaraEditorWidgets/Private/Stack/SNiagaraStackTableRow.cpp b/Engine/Plugins/FX/Niagara/Source/NiagaraEditorWidgets/Private/Stack/SNiagaraStackTableRow.cpp
--- a/Engine/Plugins/FX/Niagara/Source/NiagaraEditorWidgets/Private/Stack/SNiagaraStackTableRow.cpp
+++ b/Engine/Plugins/FX/Niagara/Source/NiagaraEditorWidgets/Private/Stack/SNiagaraStackTableRow.cpp
@@ -7,6 +7,7 @@
 #include "ViewModels/Stack/NiagaraStackItem.h"
 #include "ViewModels/Stack/NiagaraStackItemGroup.h"
 #include "ViewModels/Stack/NiagaraStackEntry.h"
+#include "ViewModels/Stack/NiagaraStackModuleItem.h"
 #include "NiagaraEditorWidgetsUtilities.h"
 #include "NiagaraEmitter.h"
 #include "NiagaraNodeFunctionCall.h"
@@ -22,6 +23,10 @@
 #include "Styling/StyleColors.h"
 #include "ScopedTransaction.h"
 #include "NiagaraEmitterEditorData.h"
+#include "ViewModels/HierarchyEditor/NiagaraSummaryViewViewModel.h"
+#include "ViewModels/Stack/NiagaraStackPropertyRow.h"
+#include "ViewModels/Stack/NiagaraStackRendererItem.h"
+#include "ViewModels/Stack/NiagaraStackSimulationStageGroup.h"
 
 #define LOCTEXT_NAMESPACE "NiagaraStackTableRow"
 
@@ -372,22 +377,44 @@ FReply SNiagaraStackTableRow::OnMouseButtonUp(const FGeometry& MyGeometry, const
 			FUIAction(FExecuteAction::CreateSP(this, &SNiagaraStackTableRow::CollapseChildren)));
 		MenuBuilder.EndSection();
 
-
-		if (IsValidForSummaryView())
-		{
-			FUIAction ShowHideSummaryViewAction(
+		if(StackEntry->GetEmitterViewModel().IsValid() && StackEntry->GetEmitterViewModel())
+		{			
+			if(StackEntry->SupportsSummaryView())
+			{
+				FNiagaraHierarchyIdentity Identity = StackEntry->DetermineSummaryIdentity();
+				if(Identity.IsValid())
+				{
+					FUIAction ShowHideSummaryViewAction(
 			FExecuteAction::CreateSP(this, &SNiagaraStackTableRow::ToggleShowInSummaryView),
-				FCanExecuteAction(),
-				FIsActionChecked::CreateSP(this, &SNiagaraStackTableRow::ShouldShowInSummaryView));
-		
-			MenuBuilder.BeginSection("SummaryViewActions", LOCTEXT("SummaryViewActions", "Emitter Summary"));
-			MenuBuilder.AddMenuEntry(
-				LOCTEXT("SummaryViewShow", "Show in Emitter Summary"),
-				LOCTEXT("SummaryViewShowTooltip", "Should this parameter be visible in the emitter summary?"),
-				FSlateIcon(),
-				ShowHideSummaryViewAction,
-				NAME_None, EUserInterfaceActionType::ToggleButton);
-			MenuBuilder.EndSection();
+						FCanExecuteAction::CreateSP(this, &SNiagaraStackTableRow::CanToggleShowInSummary),
+						FIsActionChecked::CreateSP(this, &SNiagaraStackTableRow::IsStackEntryInSummary));
+			
+					MenuBuilder.BeginSection("SummaryViewActions", LOCTEXT("SummaryViewActions", "Emitter Summary"));
+						MenuBuilder.AddMenuEntry(
+							TAttribute<FText>::CreateLambda([this]
+							{
+								bool bIsInSummaryView = StackEntry->IsInSummaryView();
+								FText BaseText = LOCTEXT("SummaryViewShow", "{0} Emitter Summary");
+								return FText::FormatOrdered(BaseText, bIsInSummaryView ? LOCTEXT("HideSummaryItem", "Remove from") : LOCTEXT("AddSummaryItem", "Add to"));
+							}),
+							TAttribute<FText>::CreateSP(this, &SNiagaraStackTableRow::GetToggleShowSummaryActionTooltip),
+							FSlateIcon(),
+							ShowHideSummaryViewAction,
+							NAME_None, EUserInterfaceActionType::ToggleButton);
+
+					FUIAction ShowInSummaryViewAction(
+			FExecuteAction::CreateSP(this, &SNiagaraStackTableRow::NavigateToSummaryView),
+						FCanExecuteAction::CreateSP(this, &SNiagaraStackTableRow::IsStackEntryInSummary));
+					
+						MenuBuilder.AddMenuEntry(
+								LOCTEXT("ShowInSummaryEditorLabel", "Show in Summary Editor"),
+								LOCTEXT("ShowInSummaryEditorTooltip", "Summon the Summary Editor and navigate to the selected item"),
+								FSlateIcon(),
+								ShowInSummaryViewAction,
+								NAME_None, EUserInterfaceActionType::Button);
+					MenuBuilder.EndSection();
+				}						
+			}
 		}
 
 		FWidgetPath WidgetPath = MouseEvent.GetEventPath() != nullptr ? *MouseEvent.GetEventPath() : FWidgetPath();
@@ -495,81 +522,122 @@ void SNiagaraStackTableRow::NavigateTo(UNiagaraStackEntry* Item)
 	OwnerTree->RequestNavigateToItem(Item, 0);
 }
 
-bool SNiagaraStackTableRow::IsValidForSummaryView() const
+void SNiagaraStackTableRow::ToggleShowInSummaryView() const
 {
-	UNiagaraStackFunctionInput* FunctionInput = Cast<UNiagaraStackFunctionInput>(StackEntry);
-	const UNiagaraEmitterEditorData* EditorData = (FunctionInput && FunctionInput->GetEmitterViewModel())? &FunctionInput->GetEmitterViewModel()->GetEditorData() : nullptr;
-
-	if (EditorData)
+	if(StackEntry->GetEmitterViewModel().IsValid())
 	{
-		UNiagaraStackFunctionInput* ParentInput = FNiagaraStackEditorWidgetsUtilities::GetParentInputForSummaryView(FunctionInput);
-		return FNiagaraStackEditorWidgetsUtilities::GetSummaryViewInputKeyForFunctionInput(ParentInput).IsSet();
+		FNiagaraHierarchyIdentity Identity = StackEntry->DetermineSummaryIdentity();
+		if(Identity.IsValid())
+		{
+			bool bItemExistsInHierarchy = StackEntry->GetEmitterViewModel()->GetSummaryHierarchyViewModel()->GetHierarchyRootViewModel()->FindViewModelForChild(Identity, true) != nullptr;
+			if(bItemExistsInHierarchy)
+			{
+				StackEntry->GetEmitterViewModel()->GetSummaryHierarchyViewModel()->DeleteItemWithIdentity(Identity);
+			}
+			else
+			{
+				UClass* HierarchyClass = DetermineHierarchyClassForSummaryView();
+				check(HierarchyClass);
+				StackEntry->GetEmitterViewModel()->GetSummaryHierarchyViewModel()->GetHierarchyRootViewModel()->AddChild(HierarchyClass, Identity);
+			}
+		}
 	}
-	return false;
 }
 
-void SNiagaraStackTableRow::ToggleShowInSummaryView()
+bool SNiagaraStackTableRow::IsStackEntryInSummary() const
 {
-	UNiagaraStackFunctionInput* FunctionInput = Cast<UNiagaraStackFunctionInput>(StackEntry);
-	UNiagaraEmitterEditorData* EditorData = (FunctionInput && FunctionInput->GetEmitterViewModel())? &FunctionInput->GetEmitterViewModel()->GetOrCreateEditorData() : nullptr;
-	
-	if (EditorData)
+	return StackEntry->IsInSummaryView();
+}
+
+bool SNiagaraStackTableRow::CanToggleShowInSummary() const
+{	
+	return StackEntry->IsAnyParentInSummaryView() == false && StackEntry->IsAnyChildInSummaryView() == false && StackEntry->ExistsInParentEmitterSummary() == false;
+}
+
+FText SNiagaraStackTableRow::GetToggleShowSummaryActionTooltip() const
+{
+	if(CanToggleShowInSummary())
 	{
-		UNiagaraStackFunctionInput* ParentInput = FNiagaraStackEditorWidgetsUtilities::GetParentInputForSummaryView(FunctionInput);
-		TOptional<FFunctionInputSummaryViewKey> Key = FNiagaraStackEditorWidgetsUtilities::GetSummaryViewInputKeyForFunctionInput(ParentInput);
+		return LOCTEXT("CanToggleSummaryStatus", "You can add or remove this item to or from summary view.\nTo view the emitter summary, please collapse the emitter node.");
+	}
+
+	FString ResultMessage;
+	FText CantToggleMessage = LOCTEXT("CantToggleSummaryViewBaseMessage", "You can not add or remove this item due to:\n");
+	FNiagaraHierarchyIdentity Identity = StackEntry->DetermineSummaryIdentity();
 	
-		if (Key.IsSet())
-		{
-			// TODO: Move this parent handling to the UNiagaraStackFunctionInput and merge manager.
-			bool bHasParentSummaryData = false;;
-			FVersionedNiagaraEmitter ParentEmitter = FunctionInput->GetEmitterViewModel()->GetEmitter().GetEmitterData()->GetParent();
-			if (ParentEmitter.GetEmitterData())
-			{
-				if (UNiagaraEmitterEditorData* ParentEmitterEditorData = Cast<UNiagaraEmitterEditorData>(ParentEmitter.GetEmitterData()->GetEditorData()))
-				{
-					bHasParentSummaryData = ParentEmitterEditorData->GetSummaryViewMetaData(Key.GetValue()).IsSet();
-				}
-			}
+	bool bExistsInParentEmitter = StackEntry->ExistsInParentEmitterSummary();
+	bool bIsAnyParentInSummaryView = StackEntry->IsAnyParentInSummaryView();
+	bool bIsAnyChildInSummaryView = StackEntry->IsAnyChildInSummaryView(true);
+	
+	FText StackEntryAddedInParentEmitter = LOCTEXT("CantToggleSummaryDueToAddedInParentEmitter", "This item was added in the parent emitter and can't me modified.\n");
+	FText ParentStackEntryIsInSummaryText = LOCTEXT("CantToggleSummaryDueToParentStackEntryInSummary", "One of the parent items of this item was already added to the summary view, which prevents you from adding this separately.\nIf you wish to add or remove it, remove the parent item first.\n");
+	FText ChildStackEntryIsInSummaryText = LOCTEXT("CantToggleSummaryDueToChildStackEntryInSummary", "One of the child items of this item was already added to the summary view, which prevents you from adding this separately.\nIf you wish to add or remove it, remove the child item first.\n");
 
-			TOptional<FFunctionInputSummaryViewMetadata> SummaryViewMetaData = EditorData->GetSummaryViewMetaData(Key.GetValue());
-			if (SummaryViewMetaData.IsSet() == false)
-			{
-				SummaryViewMetaData = FFunctionInputSummaryViewMetadata();
-			}
+	ResultMessage.Append(CantToggleMessage.ToString());
+	if(bExistsInParentEmitter)
+	{
+		ResultMessage.Append(StackEntryAddedInParentEmitter.ToString());
+	}
 
-			FScopedTransaction ScopedTransaction(FText::Format(LOCTEXT("SummaryViewChangedInputVisibility", "Changed summary view visibility for {0}"), FunctionInput->GetDisplayName()));
-			EditorData->Modify();
-			
-			SummaryViewMetaData->bVisible = !SummaryViewMetaData->bVisible;
-			if (bHasParentSummaryData == false && SummaryViewMetaData->bVisible == false)
-			{
-				// If there is no parent summary data, and the input is no longer visible, reset the optional value
-				// to remove the input from the summary.
-				SummaryViewMetaData.Reset();
-			}
-			EditorData->SetSummaryViewMetaData(Key.GetValue(), SummaryViewMetaData);
-		}
-	}	
+	if(bIsAnyParentInSummaryView)
+	{
+		ResultMessage.Append(ParentStackEntryIsInSummaryText.ToString());
+	}
+
+	if(bIsAnyChildInSummaryView)
+	{
+		ResultMessage.Append(ChildStackEntryIsInSummaryText.ToString());
+	}
+
+	ResultMessage.RemoveFromEnd("\n");
+	return FText::FromString(ResultMessage);
 }
 
-bool SNiagaraStackTableRow::ShouldShowInSummaryView() const
+bool SNiagaraStackTableRow::DoesItemExistInParentSummary() const
 {
-	UNiagaraStackFunctionInput* FunctionInput = Cast<UNiagaraStackFunctionInput>(StackEntry);
-	const UNiagaraEmitterEditorData* EditorData = (FunctionInput && FunctionInput->GetEmitterViewModel())? &FunctionInput->GetEmitterViewModel()->GetEditorData() : nullptr;
-	
-	if (EditorData)
+	return StackEntry->ExistsInParentEmitterSummary();
+}
+
+void SNiagaraStackTableRow::NavigateToSummaryView() const
+{
+	StackEntry->GetSystemViewModel()->FocusTab(FName("NiagaraSystemEditor_SummaryViewHierarchyEditor"), true);
+	StackEntry->GetEmitterViewModel()->GetSummaryHierarchyViewModel()->SetActiveHierarchySection(nullptr);
+	StackEntry->GetEmitterViewModel()->GetSummaryHierarchyViewModel()->NavigateToItemInHierarchy(StackEntry->DetermineSummaryIdentity());
+}
+
+TSubclassOf<UNiagaraHierarchyItemBase> SNiagaraStackTableRow::DetermineHierarchyClassForSummaryView() const
+{
+	if(UNiagaraStackFunctionInput* FunctionInput = Cast<UNiagaraStackFunctionInput>(StackEntry))
 	{
-		UNiagaraStackFunctionInput* ParentInput = FNiagaraStackEditorWidgetsUtilities::GetParentInputForSummaryView(FunctionInput);
-		TOptional<FFunctionInputSummaryViewKey> Key = FNiagaraStackEditorWidgetsUtilities::GetSummaryViewInputKeyForFunctionInput(ParentInput);
-	
-		if (Key.IsSet())
+		if(FunctionInput->GetInputFunctionCallNode().IsA<UNiagaraNodeAssignment>())
 		{
-			TOptional<FFunctionInputSummaryViewMetadata> Metadata = EditorData->GetSummaryViewMetaData(Key.GetValue());
-			return Metadata.IsSet() && Metadata->bVisible;
-		}			
+			return UNiagaraHierarchyAssignmentInput::StaticClass();
+		}
+		
+		return UNiagaraHierarchyModuleInput::StaticClass();
+	}
+	else if(StackEntry->IsA<UNiagaraStackModuleItem>())
+	{
+		return UNiagaraHierarchyModule::StaticClass();
+	}
+	else if(StackEntry->IsA<UNiagaraStackRendererItem>())
+	{
+		return UNiagaraHierarchyRenderer::StaticClass();
+	}
+	else if(StackEntry->IsA<UNiagaraStackPropertyRow>())
+	{
+		return UNiagaraHierarchyObjectProperty::StaticClass();
+	}
+	else if(StackEntry->IsA<UNiagaraStackSimulationStageGroup>())
+	{
+		return UNiagaraHierarchySimStage::StaticClass();
+	}
+	else if(StackEntry->IsA<UNiagaraStackSimulationStagePropertiesItem>())
+	{
+		return UNiagaraHierarchySimStageProperties::StaticClass();
 	}
 
-	return false;
+	return nullptr;
 }
 
 #undef LOCTEXT_NAMESPACE
diff --git a/Engine/Plugins/FX/Niagara/Source/NiagaraEditorWidgets/Private/Stack/SNiagaraStackTableRow.h b/Engine/Plugins/FX/Niagara/Source/NiagaraEditorWidgets/Private/Stack/SNiagaraStackTableRow.h
--- a/Engine/Plugins/FX/Niagara/Source/NiagaraEditorWidgets/Private/Stack/SNiagaraStackTableRow.h
+++ b/Engine/Plugins/FX/Niagara/Source/NiagaraEditorWidgets/Private/Stack/SNiagaraStackTableRow.h
@@ -10,6 +10,7 @@
 #include "Styling/SlateTypes.h"
 #include "Layout/Visibility.h"
 
+class UNiagaraHierarchyItemBase;
 class UNiagaraStackViewModel;
 class UNiagaraStackEntry;
 class FNiagaraStackCommandContext;
@@ -87,10 +88,15 @@ private:
 
 	void NavigateTo(UNiagaraStackEntry* Item);
 
-	bool IsValidForSummaryView() const;
-	void ToggleShowInSummaryView();
-	bool ShouldShowInSummaryView() const;
+	void ToggleShowInSummaryView() const;
+	bool IsStackEntryInSummary() const;
+	bool CanToggleShowInSummary() const;
+	FText GetToggleShowSummaryActionTooltip() const;
+	bool DoesItemExistInParentSummary() const;
 
+	void NavigateToSummaryView() const;
+	
+	TSubclassOf<UNiagaraHierarchyItemBase> DetermineHierarchyClassForSummaryView() const;
 private:
 	UNiagaraStackViewModel* StackViewModel;
 	UNiagaraStackEntry* StackEntry;
