Add SVT as an option for Volume RT DI sim cache data storage

diff --git a/Engine/Plugins/FX/Niagara/Shaders/Private/NiagaraSVTToDenseBuffer.usf b/Engine/Plugins/FX/Niagara/Shaders/Private/NiagaraSVTToDenseBuffer.usf
new file mode 100644
--- /dev/null
+++ b/Engine/Plugins/FX/Niagara/Shaders/Private/NiagaraSVTToDenseBuffer.usf
@@ -0,0 +1,24 @@
+// Copyright Epic Games, Inc. All Rights Reserved.
+
+#include "/Engine/Private/Common.ush"
+#include "/Engine/Private/SparseVolumeTexture/SparseVolumeTextureCommon.ush"
+RWTexture3D<float4>	DestinationBuffer;
+
+SamplerState	TileDataTextureSampler;
+Texture3D<uint>	SparseVolumeTexturePageTable;
+Texture3D		SparseVolumeTextureA;
+uint4			PackedSVTUniforms0;
+uint4			PackedSVTUniforms1;
+int3			TextureSize;
+int				MipLevels;
+
+[numthreads(4, 4, 4)]
+void PerformCopyCS(uint3 Index : SV_DispatchThreadID)
+{
+	
+	float4 SVTSample;
+	const FSparseVolumeTextureUniforms Uniforms = SparseVolumeTextureUnpackUniforms(PackedSVTUniforms0, PackedSVTUniforms1);
+	const int3 VoxelCoord = SparseVolumeTextureLoadPageTable(SparseVolumeTexturePageTable, Uniforms, Index, 0);	
+
+	DestinationBuffer[Index] = SparseVolumeTextureA.Load(int4(VoxelCoord, 0));
+}
diff --git a/Engine/Plugins/FX/Niagara/Source/Niagara/Private/NiagaraDataInterfaceRenderTargetVolume.cpp b/Engine/Plugins/FX/Niagara/Source/Niagara/Private/NiagaraDataInterfaceRenderTargetVolume.cpp
--- a/Engine/Plugins/FX/Niagara/Source/Niagara/Private/NiagaraDataInterfaceRenderTargetVolume.cpp
+++ b/Engine/Plugins/FX/Niagara/Source/Niagara/Private/NiagaraDataInterfaceRenderTargetVolume.cpp
@@ -23,6 +23,13 @@
 #include "Misc/Paths.h"
 #include "DataDrivenShaderPlatformInfo.h"
 
+#include "RenderGraphUtils.h"
+
+#include "SparseVolumeTexture/SparseVolumeTexture.h"
+#include "SparseVolumeTexture/SparseVolumeTextureData.h"
+#include "SparseVolumeTexture/SparseVolumeTextureSceneProxy.h"
+
+
 #include UE_INLINE_GENERATED_CPP_BY_NAME(NiagaraDataInterfaceRenderTargetVolume)
 
 #define LOCTEXT_NAMESPACE "NiagaraDataInterfaceRenderTargetVolume"
@@ -81,11 +88,28 @@ namespace NDIRenderTargetVolumeLocal
 		ECVF_Default
 	);
 
-	int32 GSimCacheUseOpenVDB = true;
-	static FAutoConsoleVariableRef CVarSimCacheUseOpenVDB(
-		TEXT("fx.Niagara.RenderTargetVolume.SimCacheUseOpenVDB"),
-		GSimCacheUseOpenVDB,
-		TEXT("Use OpenVDB as the backing data for the sim cache."),
+	enum SimCacheStorageMode
+	{
+		Dense = 0,
+		OpenVDB = 1,
+		SVT = 2
+	};
+
+	int32 GSimCacheDataStorageMode = SimCacheStorageMode::OpenVDB;
+	static FAutoConsoleVariableRef CVarSimCacheDataStorageMode(
+		TEXT("fx.Niagara.RenderTargetVolume.SimCacheDataStorageMode"),
+		GSimCacheDataStorageMode,
+		TEXT("Backing storage type for Volume RT sim cache data.  0 uses raw data, 1 uses OpenVDB, 2 uses SVT"),
+		ECVF_Default
+	);
+
+
+
+	int32 GSimCacheUseOpenVDBFloatGrids = false;
+	static FAutoConsoleVariableRef CVarSimCacheUseOpenVDBFloatGrids(
+		TEXT("fx.Niagara.RenderTargetVolume.SimCacheUseOpenVDBFloatGrids"),
+		GSimCacheUseOpenVDBFloatGrids,
+		TEXT("Use OpenVDB float grids as output."),
 		ECVF_Default
 	);
 
@@ -561,7 +585,7 @@ bool UNiagaraDataInterfaceRenderTargetVolume::GetExposedVariableValue(const FNia
 
 UObject* UNiagaraDataInterfaceRenderTargetVolume::SimCacheBeginWrite(UObject* InSimCache, FNiagaraSystemInstance* NiagaraSystemInstance, const void* OptionalPerInstanceData, FNiagaraSimCacheFeedbackContext& FeedbackContext) const
 {
-	if (NDIRenderTargetVolumeLocal::GSimCacheUseOpenVDB)
+	if (NDIRenderTargetVolumeLocal::GSimCacheDataStorageMode == NDIRenderTargetVolumeLocal::SimCacheStorageMode::OpenVDB)
 	{		
 		UVolumeCache* OpenVDBSimCacheData = nullptr;
 
@@ -607,6 +631,15 @@ UObject* UNiagaraDataInterfaceRenderTargetVolume::SimCacheBeginWrite(UObject* In
 
 		return OpenVDBSimCacheData;
 	}
+	else if (NDIRenderTargetVolumeLocal::GSimCacheDataStorageMode == NDIRenderTargetVolumeLocal::SimCacheStorageMode::SVT)
+	{		
+		UNiagaraSimCache* SimCache = CastChecked<UNiagaraSimCache>(InSimCache);
+		UAnimatedSparseVolumeTexture* CurrCache = NewObject<UAnimatedSparseVolumeTexture>(SimCache);			
+		
+		CurrCache->BeginInitialize(1);
+
+		return CurrCache;
+	}	
 	else
 	{
 		UNDIRenderTargetVolumeSimCacheData* SimCacheData = nullptr;
@@ -622,6 +655,76 @@ UObject* UNiagaraDataInterfaceRenderTargetVolume::SimCacheBeginWrite(UObject* In
 	return nullptr;
 }
 
+
+
+class FSparseVolumeTextureDataConstructionVolumeRenderTargetAdapter : public ISparseVolumeTextureDataConstructionAdapter
+{
+public:
+
+	FSparseVolumeTextureDataConstructionVolumeRenderTargetAdapter(TArray<FFloat16Color>& SourceTextureData, EPixelFormat SourceFormat, FIntVector3 SourceGridResolution)
+		: TextureData(SourceTextureData), Format(SourceFormat), GridResolution(SourceGridResolution)
+	{}
+
+	virtual void GetAttributesInfo(FAttributesInfo& OutInfoA, FAttributesInfo& OutInfoB) const
+	{
+		OutInfoA.Format = Format;
+		OutInfoA.FallbackValue = FVector4f(0, 0, 0, 0);
+		OutInfoA.NormalizeScale = FVector4f(0, 0, 0, 0);
+		OutInfoA.NormalizeBias = FVector4f(0, 0, 0, 0);
+		OutInfoA.bNormalized = false;
+
+		OutInfoB.Format = PF_Unknown;
+		OutInfoB.FallbackValue = FVector4f(0, 0, 0, 0);
+		OutInfoB.NormalizeScale = FVector4f(0, 0, 0, 0);
+		OutInfoB.NormalizeBias = FVector4f(0, 0, 0, 0);
+		OutInfoB.bNormalized = false;
+	}
+
+	virtual FIntVector3 GetAABBMin() const
+	{
+		return FIntVector3(0, 0, 0);
+	}
+
+	virtual FIntVector3 GetAABBMax() const 
+	{
+		return GridResolution;
+	}
+
+	virtual FIntVector3 GetResolution() const
+	{
+		return GridResolution;
+	}
+
+	virtual void IteratePhysicalSource(TFunctionRef<void(const FIntVector3& Coord, int32 AttributesIdx, int32 ComponentIdx, float VoxelValue)> OnVisit) const
+	{
+		for (int32 Z = 0; Z < GridResolution.Z - 1; ++Z)
+		{
+			for (int32 Y = 0; Y < GridResolution.Y - 1; ++Y)
+			{
+				for (int32 X = 0; X < GridResolution.X - 1; ++X)
+				{
+					FIntVector3 CurrCoord(X, Y, Z);
+
+					FFloat16Color CurrCellValue;
+					CurrCellValue = TextureData[Y * GridResolution.X + Z * GridResolution.X * GridResolution.Y + X];
+
+					OnVisit(CurrCoord, 0, 0, CurrCellValue.R);
+					OnVisit(CurrCoord, 0, 1, CurrCellValue.G);
+					OnVisit(CurrCoord, 0, 2, CurrCellValue.B);
+					OnVisit(CurrCoord, 0, 3, CurrCellValue.A);
+				}
+			}
+		}
+	}
+
+	virtual ~FSparseVolumeTextureDataConstructionVolumeRenderTargetAdapter() = default;
+
+private:		
+	const TArray<FFloat16Color>& TextureData;
+	EPixelFormat Format;
+	FIntVector3 GridResolution;
+};
+
 bool UNiagaraDataInterfaceRenderTargetVolume::SimCacheWriteFrame(UObject* StorageObject, int FrameIndex, FNiagaraSystemInstance* SystemInstance, const void* OptionalPerInstanceData, FNiagaraSimCacheFeedbackContext& FeedbackContext) const
 {
 	check(OptionalPerInstanceData);
@@ -655,18 +758,31 @@ bool UNiagaraDataInterfaceRenderTargetVolume::SimCacheWriteFrame(UObject* Storag
 
 		if (TextureData.Num() > 0)
 		{
-			if (NDIRenderTargetVolumeLocal::GSimCacheUseOpenVDB)
+			if (NDIRenderTargetVolumeLocal::GSimCacheDataStorageMode == NDIRenderTargetVolumeLocal::SimCacheStorageMode::OpenVDB)
 			{
 #if UE_USE_OPENVDB
 				UVolumeCache* OpenVDBSimCacheData = CastChecked<UVolumeCache>(StorageObject);
 				
 				FString FullPath = OpenVDBSimCacheData->GetAssetPath(FrameIndex);
 
-				OpenVDBTools::WriteImageDataToOpenVDBFile(FullPath, InstanceData_GT->Size, TextureData, false);
+				OpenVDBTools::WriteImageDataToOpenVDBFile(FullPath, InstanceData_GT->Size, TextureData, (bool) NDIRenderTargetVolumeLocal::GSimCacheUseOpenVDBFloatGrids);
 
 				OpenVDBSimCacheData->FrameRangeStart = FMath::Min(FrameIndex, OpenVDBSimCacheData->FrameRangeStart);
 				OpenVDBSimCacheData->FrameRangeEnd = FMath::Max(FrameIndex, OpenVDBSimCacheData->FrameRangeEnd);
 				OpenVDBSimCacheData->Resolution = InstanceData_GT->Size;
+#endif
+			}
+			else if (NDIRenderTargetVolumeLocal::GSimCacheDataStorageMode == NDIRenderTargetVolumeLocal::SimCacheStorageMode::SVT)
+			{
+#if WITH_EDITOR
+				UAnimatedSparseVolumeTexture* CurrCache = CastChecked<UAnimatedSparseVolumeTexture>(StorageObject);
+								
+				FSparseVolumeTextureData SparseTextureData{};
+				FSparseVolumeTextureDataConstructionVolumeRenderTargetAdapter Adapter(TextureData, PF_FloatRGBA, InstanceData_GT->Size);
+				
+				SparseTextureData.Construct(Adapter);
+				
+				CurrCache->AppendFrame(SparseTextureData);
 #endif
 			}
 			else
@@ -719,9 +835,44 @@ bool UNiagaraDataInterfaceRenderTargetVolume::SimCacheWriteFrame(UObject* Storag
 
 bool UNiagaraDataInterfaceRenderTargetVolume::SimCacheEndWrite(UObject* StorageObject) const
 {
+	if (NDIRenderTargetVolumeLocal::GSimCacheDataStorageMode == NDIRenderTargetVolumeLocal::SimCacheStorageMode::SVT)
+	{	
+		UAnimatedSparseVolumeTexture* CurrCache = CastChecked<UAnimatedSparseVolumeTexture>(StorageObject);
+		CurrCache->EndInitialize();
+		CurrCache->PostLoad();
+	}
+
 	return true;
 }
 
+
+
+class FNiagaraCopySVTToDenseBufferCS : public FGlobalShader
+{
+	DECLARE_GLOBAL_SHADER(FNiagaraCopySVTToDenseBufferCS);
+	SHADER_USE_PARAMETER_STRUCT(FNiagaraCopySVTToDenseBufferCS, FGlobalShader);
+	
+	static void ModifyCompilationEnvironment(const FGlobalShaderPermutationParameters& Parameters, FShaderCompilerEnvironment& OutEnvironment)
+	{
+		FGlobalShader::ModifyCompilationEnvironment(Parameters, OutEnvironment);		
+	}
+
+	BEGIN_SHADER_PARAMETER_STRUCT(FParameters, )
+
+		SHADER_PARAMETER_RDG_TEXTURE_UAV(RWTexture3D<float4>, DestinationBuffer)
+		SHADER_PARAMETER_SAMPLER(SamplerState, TileDataTextureSampler)
+		SHADER_PARAMETER_TEXTURE(Texture3D<uint>, SparseVolumeTexturePageTable)
+		SHADER_PARAMETER_TEXTURE(Texture3D, SparseVolumeTextureA)		
+		SHADER_PARAMETER(FUintVector4, PackedSVTUniforms0)
+		SHADER_PARAMETER(FUintVector4, PackedSVTUniforms1)
+		SHADER_PARAMETER(FIntVector, TextureSize)
+		SHADER_PARAMETER(int32, MipLevel)
+	END_SHADER_PARAMETER_STRUCT()
+};
+
+IMPLEMENT_GLOBAL_SHADER(FNiagaraCopySVTToDenseBufferCS, "/Plugin/FX/Niagara/Private/NiagaraSVTToDenseBuffer.usf", "PerformCopyCS", SF_Compute);
+
+
 bool UNiagaraDataInterfaceRenderTargetVolume::SimCacheReadFrame(UObject* StorageObject, int FrameA, int FrameB, float Interp, FNiagaraSystemInstance* SystemInstance, void* OptionalPerInstanceData)
 {
 	FRenderTargetVolumeRWInstanceData_GameThread* InstanceData_GT = reinterpret_cast<FRenderTargetVolumeRWInstanceData_GameThread*>(OptionalPerInstanceData);
@@ -761,6 +912,115 @@ bool UNiagaraDataInterfaceRenderTargetVolume::SimCacheReadFrame(UObject* Storage
 		}
 #endif
 	}
+	else if (Cast<UAnimatedSparseVolumeTexture>(StorageObject))
+	{			
+		UAnimatedSparseVolumeTexture* SVT = Cast<UAnimatedSparseVolumeTexture>(StorageObject);
+		
+		USparseVolumeTextureFrame *SVTFrame = USparseVolumeTextureFrame::CreateFrame(SVT, FrameA, 0);
+		
+		FUintVector4 CurrentPackedUniforms0 = FUintVector4();
+		FUintVector4 CurrentPackedUniforms1 = FUintVector4();
+		SVTFrame->GetPackedUniforms(CurrentPackedUniforms0, CurrentPackedUniforms1);
+
+		FIntVector VolumeResolution = SVT->GetVolumeResolution();						
+
+		InstanceData_GT->Size = VolumeResolution;
+		InstanceData_GT->Format = PF_FloatRGBA;
+		InstanceData_GT->Filter = TextureFilter::TF_Default;
+
+		PerInstanceTick(InstanceData_GT, SystemInstance, 0.0f);
+		PerInstanceTickPostSimulate(InstanceData_GT, SystemInstance, 0.0f);
+
+		FNiagaraDataInterfaceProxyRenderTargetVolumeProxy* RT_Proxy = GetProxyAs<FNiagaraDataInterfaceProxyRenderTargetVolumeProxy>();
+
+
+		// Execute compute shader
+		ENQUEUE_RENDER_COMMAND(NDIRenderTargetVolumeUpdate)
+			(
+				[RT_Proxy, RT_InstanceID = SystemInstance->GetId(),
+				RT_VolumeResolution = VolumeResolution, RT_SVTProxy = SVTFrame->GetSparseVolumeTextureSceneProxy(),
+				FeatureLevel = SystemInstance->GetFeatureLevel(),
+				CurrentPackedUniforms0, CurrentPackedUniforms1](FRHICommandListImmediate& RHICmdList)
+				{
+					if (RT_SVTProxy == nullptr)
+					{
+						return;
+					}
+
+					if (FRenderTargetVolumeRWInstanceData_RenderThread * InstanceData_RT = RT_Proxy->SystemInstancesToProxyData_RT.Find(RT_InstanceID))
+					{
+						FRDGBuilder GraphBuilder(RHICmdList);
+
+						FIntVector ThreadGroupSize = FNiagaraShader::GetDefaultThreadGroupSize(ENiagaraGpuDispatchType::ThreeD);
+
+						TShaderMapRef<FNiagaraCopySVTToDenseBufferCS> ComputeShader(GetGlobalShaderMap(FeatureLevel));						
+
+						const FIntVector NumThreadGroups(
+							FMath::DivideAndRoundUp(RT_VolumeResolution.X, ThreadGroupSize.X),
+							FMath::DivideAndRoundUp(RT_VolumeResolution.Y, ThreadGroupSize.Y),
+							FMath::DivideAndRoundUp(RT_VolumeResolution.Z, ThreadGroupSize.Z)
+						);
+						
+						FNiagaraCopySVTToDenseBufferCS::FParameters* PassParameters = GraphBuilder.AllocParameters<FNiagaraCopySVTToDenseBufferCS::FParameters>();
+
+						if (InstanceData_RT->RenderTarget.IsValid() && InstanceData_RT->TransientRDGTexture == nullptr)
+						{
+							InstanceData_RT->TransientRDGTexture = GraphBuilder.RegisterExternalTexture(InstanceData_RT->RenderTarget);							
+							InstanceData_RT->TransientRDGUAV = GraphBuilder.CreateUAV(InstanceData_RT->TransientRDGTexture);
+
+							// #todo(dmp): needed?							
+							//Context.GetRDGExternalAccessQueue().Add(InstanceData_RT->TransientRDGTexture);
+						}
+
+						if (InstanceData_RT->RenderTarget.IsValid() && InstanceData_RT->TransientRDGUAV != nullptr)
+						{
+							PassParameters->DestinationBuffer = InstanceData_RT->TransientRDGUAV;
+						}
+						else
+						{
+							PassParameters->DestinationBuffer = GraphBuilder.CreateUAV(GraphBuilder.CreateTexture(
+									FRDGTextureDesc::Create3D(FIntVector(1, 1, 1), PF_A16B16G16R16, FClearValueBinding::Black, ETextureCreateFlags::ShaderResource | ETextureCreateFlags::UAV),
+									TEXT("NiagaraEmptyTextureUAV::Texture3D")
+								),
+								ERDGUnorderedAccessViewFlags::SkipBarrier
+							);
+						}						
+
+						FRHITexture* PageTableTexture = RT_SVTProxy->GetPageTableTextureRHI();
+						FRHITexture* TextureA = RT_SVTProxy->GetPhysicalTileDataATextureRHI();						
+
+						PassParameters->TileDataTextureSampler = TStaticSamplerState<SF_Point, AM_Clamp, AM_Clamp, AM_Clamp>::GetRHI();
+						PassParameters->SparseVolumeTexturePageTable = PageTableTexture;
+						PassParameters->SparseVolumeTextureA = TextureA;
+							
+						PassParameters->PackedSVTUniforms0 = CurrentPackedUniforms0;
+						PassParameters->PackedSVTUniforms1 = CurrentPackedUniforms1;
+
+						PassParameters->TextureSize = RT_VolumeResolution;
+						PassParameters->MipLevel = 0;
+						
+						GraphBuilder.AddPass(
+							// Friendly name of the pass for profilers using printf semantics.
+							RDG_EVENT_NAME("Copy SVT to Volume RT"),
+							// Parameters provided to RDG.
+							PassParameters,
+							// Issues compute commands.
+							ERDGPassFlags::Compute,
+							// This is deferred until Execute. May execute in parallel with other passes.
+							[PassParameters, ComputeShader, NumThreadGroups](FRHIComputeCommandList& RHICmdList)
+							{
+								FComputeShaderUtils::Dispatch(RHICmdList, ComputeShader, *PassParameters, NumThreadGroups);
+							});
+
+						// Execute the graph.
+						GraphBuilder.Execute();						
+						
+						InstanceData_RT->TransientRDGTexture = nullptr;
+						InstanceData_RT->TransientRDGSRV = nullptr;
+						InstanceData_RT->TransientRDGUAV = nullptr;
+					}					
+				});	
+	}
 	else
 	{
 		UNDIRenderTargetVolumeSimCacheData* SimCacheData = CastChecked<UNDIRenderTargetVolumeSimCacheData>(StorageObject);
diff --git a/Engine/Plugins/FX/Niagara/Source/Niagara/Private/VolumeCache.cpp b/Engine/Plugins/FX/Niagara/Source/Niagara/Private/VolumeCache.cpp
--- a/Engine/Plugins/FX/Niagara/Source/Niagara/Private/VolumeCache.cpp
+++ b/Engine/Plugins/FX/Niagara/Source/Niagara/Private/VolumeCache.cpp
@@ -151,7 +151,7 @@ void FOpenVDBCacheData::Init(FIntVector Resolution)
 	if (OpenVDBGrids.Num() == 0)
 	{
 		TotalMemoryUsage = 0;
-}
+	}
 }
 
 bool FOpenVDBCacheData::LoadFile(FString Path, int frame)
@@ -191,19 +191,12 @@ bool FOpenVDBCacheData::LoadFile(FString Path, int frame)
 			return false;
 		}
 
-		Vec4Grid::Ptr ColorGrid;		
-		ColorGrid = openvdb::gridPtrCast<Vec4Grid>(file.readGrid(openvdb::Name("Color")));		
-
-		if (ColorGrid == nullptr)
-		{
-			UE_LOG(LogVolumeCache, Warning, TEXT("Cache has invalid grids: %s"), *FullPath);
-			return false;
-		}
-
 		openvdb::MetaMap::Ptr Metadata = file.getMetadata();
-		FIntVector Size;
+
 		try
 		{
+			FIntVector Size;
+
 			Size.X = Metadata->metaValue<int>("DenseResolutionX");
 			Size.Y = Metadata->metaValue<int>("DenseResolutionY");
 			Size.Z = Metadata->metaValue<int>("DenseResolutionZ");
@@ -217,32 +210,51 @@ bool FOpenVDBCacheData::LoadFile(FString Path, int frame)
 				return false;
 			}
 
-			OpenVDBGrids.Add(frame, ColorGrid);
-
-			// add total memory usage reported in bytes
-			TotalMemoryUsage += ColorGrid->memUsage() / 1000000;
+			bool UseFloatGrids = Metadata->metaValue<bool>("FloatGridStorage");
 
-			// enqueue 
-			CacheQueue.Enqueue(frame);
-
-			// if total memory usage is over the limit, dequeue oldest one
-			int64 MaxMemory = GetDefault<UNiagaraSettings>()->SimCacheMaxCPUMemoryVolumetrics;
-			if (TotalMemoryUsage > MaxMemory && !CacheQueue.IsEmpty())
+			if (UseFloatGrids)
 			{
-				int FrameToRemove = CacheQueue.Pop();
-				UnloadFile(FrameToRemove);
+				UE_LOG(LogVolumeCache, Warning, TEXT("Cache has float grids, which are not used for cache reading: %s"), *FullPath);
+				return false;
 			}
-
-			// if dense vdb buffer doesn't match current resolution, re initialize it				
-			FScopeLock ScopeLock(&DenseGridGuard);
-			if (DenseGridPtr == nullptr || (DenseGridPtr->bbox().dim() != openvdb::Coord(DenseResolution.X - 1, DenseResolution.Y - 1, DenseResolution.Z - 1)))
+			else
 			{
-				Init(DenseResolution);
-			}
+				Vec4Grid::Ptr ColorGrid;
+				ColorGrid = openvdb::gridPtrCast<Vec4Grid>(file.readGrid(openvdb::Name("Color")));
 
-			// load to dense buffer	
-			openvdb::tools::CopyToDense<Vec4Tree, Vec4Dense> Copier(ColorGrid->tree(), *DenseGridPtr);
-			Copier.copy();
+				if (ColorGrid == nullptr)
+				{
+					UE_LOG(LogVolumeCache, Warning, TEXT("Cache has invalid grids: %s"), *FullPath);
+					return false;
+				}
+
+				OpenVDBGrids.Add(frame, ColorGrid);
+
+				// add total memory usage reported in bytes
+				TotalMemoryUsage += ColorGrid->memUsage() / 1000000;
+
+				// enqueue 
+				CacheQueue.Enqueue(frame);
+
+				// if total memory usage is over the limit, dequeue oldest one
+				int64 MaxMemory = GetDefault<UNiagaraSettings>()->SimCacheMaxCPUMemoryVolumetrics;
+				if (TotalMemoryUsage > MaxMemory && !CacheQueue.IsEmpty())
+				{
+					int FrameToRemove = CacheQueue.Pop();
+					UnloadFile(FrameToRemove);
+				}
+
+				// if dense vdb buffer doesn't match current resolution, re initialize it				
+				FScopeLock ScopeLock(&DenseGridGuard);
+				if (DenseGridPtr == nullptr || (DenseGridPtr->bbox().dim() != openvdb::Coord(DenseResolution.X - 1, DenseResolution.Y - 1, DenseResolution.Z - 1)))
+				{
+					Init(DenseResolution);
+				}
+
+				// load to dense buffer	
+				openvdb::tools::CopyToDense<Vec4Tree, Vec4Dense> Copier(ColorGrid->tree(), *DenseGridPtr);
+				Copier.copy();
+			}
 		}
 		catch (openvdb::Exception e)
 		{
@@ -462,6 +474,9 @@ bool OpenVDBTools::WriteImageDataToOpenVDBFile(FStringView FilePath, FIntVector
 	outMeta.insertMeta("creator",
 		openvdb::StringMetadata("Unreal Engine"));
 
+	outMeta.insertMeta("FloatGridStorage", 
+		openvdb::BoolMetadata(UseFloatGrids));
+
 	outMeta.insertMeta("DenseResolutionX",
 		openvdb::Int32Metadata(ImageSize.X));
 	outMeta.insertMeta("DenseResolutionY",
