Removed cyclic dependency

diff --git a/Engine/Plugins/FX/Niagara/Source/NiagaraEditor/Private/ViewModels/Stack/NiagaraStackGraphUtilities.cpp b/Engine/Plugins/FX/Niagara/Source/NiagaraEditor/Private/ViewModels/Stack/NiagaraStackGraphUtilities.cpp
--- a/Engine/Plugins/FX/Niagara/Source/NiagaraEditor/Private/ViewModels/Stack/NiagaraStackGraphUtilities.cpp
+++ b/Engine/Plugins/FX/Niagara/Source/NiagaraEditor/Private/ViewModels/Stack/NiagaraStackGraphUtilities.cpp
@@ -36,6 +36,7 @@
 #include "ViewModels/NiagaraEmitterViewModel.h"
 #include "EdGraphUtilities.h"
 #include "ObjectTools.h"
+#include "NiagaraParameterMapHistory.h"
 #include "NiagaraMessageManager.h"
 #include "NiagaraSimulationStageBase.h"
 #include "NiagaraScriptVariable.h"
@@ -44,6 +45,8 @@
 #include "NiagaraGraphDataCache.h"
 #include "ViewModels/NiagaraScriptGraphViewModel.h"
 #include "NiagaraSettings.h"
+#include "NiagaraSimulationStageBase.h"
+#include "NiagaraEmitter.h"
 #include "ViewModels/HierarchyEditor/NiagaraHierarchyViewModelBase.h"
 #include "ViewModels/HierarchyEditor/NiagaraSummaryViewViewModel.h"
 
diff --git a/Engine/Plugins/FX/Niagara/Source/NiagaraEditor/Public/ViewModels/Stack/NiagaraStackGraphUtilities.h b/Engine/Plugins/FX/Niagara/Source/NiagaraEditor/Public/ViewModels/Stack/NiagaraStackGraphUtilities.h
--- a/Engine/Plugins/FX/Niagara/Source/NiagaraEditor/Public/ViewModels/Stack/NiagaraStackGraphUtilities.h
+++ b/Engine/Plugins/FX/Niagara/Source/NiagaraEditor/Public/ViewModels/Stack/NiagaraStackGraphUtilities.h
@@ -4,7 +4,6 @@
 
 #include "NiagaraTypes.h"
 #include "NiagaraCommon.h"
-#include "NiagaraParameterMapHistory.h"
 #include "ViewModels/Stack/NiagaraParameterHandle.h"
 #include "ViewModels/Stack/NiagaraStackEntry.h"
 #include "AssetRegistry/AssetData.h"
@@ -31,6 +30,8 @@ class FCompileConstantResolver;
 class INiagaraMessage;
 class FNiagaraParameterMapHistoryBuilder;
 class UNiagaraHierarchyModuleInput;
+class UNiagaraSimulationStageBase;
+struct FNiagaraEventScriptProperties;
 struct FNiagaraStackModuleData;
 struct FNiagaraModuleDependency;
 
