Fixed SA issues

diff --git a/Engine/Plugins/FX/Niagara/Source/NiagaraEditor/Private/ViewModels/Stack/NiagaraStackGraphUtilities.cpp b/Engine/Plugins/FX/Niagara/Source/NiagaraEditor/Private/ViewModels/Stack/NiagaraStackGraphUtilities.cpp
--- a/Engine/Plugins/FX/Niagara/Source/NiagaraEditor/Private/ViewModels/Stack/NiagaraStackGraphUtilities.cpp
+++ b/Engine/Plugins/FX/Niagara/Source/NiagaraEditor/Private/ViewModels/Stack/NiagaraStackGraphUtilities.cpp
@@ -1063,45 +1063,46 @@ void FNiagaraStackGraphUtilities::GatherInputRelationsForStack(FInputDataCollect
 	for(UNiagaraHierarchyModuleInput* Input : ModuleInputs)
 	{
 		UNiagaraGraph* GraphWithInput = HierarchyInputToFunctionCallNodeMap[Input]->GetCalledGraph();
-		ensure(GraphWithInput != nullptr);
-
-		bool bIsModuleInput = Input->GetPersistentIdentity().Guids.Num() == 2;
-
-		// the input guid is only valid for module inputs, not assignment targets
-		if(bIsModuleInput)
+		if(ensure(GraphWithInput != nullptr))
 		{
-			TMap<FNiagaraVariable, TObjectPtr<UNiagaraScriptVariable>> MetaDataMap = GraphWithInput->GetAllMetaData();
-			TArray<TObjectPtr<UNiagaraScriptVariable>> ScriptVariables;
-			MetaDataMap.GenerateValueArray(ScriptVariables);
+			bool bIsModuleInput = Input->GetPersistentIdentity().Guids.Num() == 2;
 
-			FGuid InputGuid = Input->GetPersistentIdentity().Guids[1];
-			TObjectPtr<UNiagaraScriptVariable>* MatchingScriptVariable = ScriptVariables.FindByPredicate([InputGuid](TObjectPtr<UNiagaraScriptVariable> ScriptVariableCandidate)
+			// the input guid is only valid for module inputs, not assignment targets
+			if(bIsModuleInput)
 			{
-				return ScriptVariableCandidate->Metadata.GetVariableGuid() == InputGuid;
-			});
+				TMap<FNiagaraVariable, TObjectPtr<UNiagaraScriptVariable>> MetaDataMap = GraphWithInput->GetAllMetaData();
+				TArray<TObjectPtr<UNiagaraScriptVariable>> ScriptVariables;
+				MetaDataMap.GenerateValueArray(ScriptVariables);
 
-			TOptional<bool> bIsStaticSwitch = GraphWithInput->IsStaticSwitch((*MatchingScriptVariable)->Variable);
-			if(ensure(MatchingScriptVariable != nullptr))
-			{
-				HierarchyInputToScriptVariableMap.Add(Input, *MatchingScriptVariable);
-				HierarchyInputToStaticSwitchMap.Add(Input, bIsStaticSwitch.GetValue());
-			}
-		}
-		// otherwise we have to use the name
-		else
-		{
-			TArray<FNiagaraVariable> Variables;
-			GraphWithInput->GetAllVariables(Variables);
+				FGuid InputGuid = Input->GetPersistentIdentity().Guids[1];
+				TObjectPtr<UNiagaraScriptVariable>* MatchingScriptVariable = ScriptVariables.FindByPredicate([InputGuid](TObjectPtr<UNiagaraScriptVariable> ScriptVariableCandidate)
+				{
+					return ScriptVariableCandidate->Metadata.GetVariableGuid() == InputGuid;
+				});
 
-			FName InputName = Input->GetPersistentIdentity().Names[0];
-			FNiagaraVariableBase* MatchingVariable = Variables.FindByPredicate([InputName](FNiagaraVariable& Candidate)
+				TOptional<bool> bIsStaticSwitch = GraphWithInput->IsStaticSwitch((*MatchingScriptVariable)->Variable);
+				if(ensure(MatchingScriptVariable != nullptr))
+				{
+					HierarchyInputToScriptVariableMap.Add(Input, *MatchingScriptVariable);
+					HierarchyInputToStaticSwitchMap.Add(Input, bIsStaticSwitch.GetValue());
+				}
+			}
+			// otherwise we have to use the name
+			else
 			{
-				return Candidate.GetName() == InputName;
-			});
+				TArray<FNiagaraVariable> Variables;
+				GraphWithInput->GetAllVariables(Variables);
 
-			if(ensure(MatchingVariable != nullptr))
-			{
-				HierarchyInputToAssignmentMap.Add(Input, *MatchingVariable);
+				FName InputName = Input->GetPersistentIdentity().Names[0];
+				FNiagaraVariableBase* MatchingVariable = Variables.FindByPredicate([InputName](FNiagaraVariable& Candidate)
+				{
+					return Candidate.GetName() == InputName;
+				});
+
+				if(ensure(MatchingVariable != nullptr))
+				{
+					HierarchyInputToAssignmentMap.Add(Input, *MatchingVariable);
+				}
 			}
 		}
 	}
@@ -1112,64 +1113,65 @@ void FNiagaraStackGraphUtilities::GatherInputRelationsForStack(FInputDataCollect
 	for(UNiagaraHierarchyModuleInput* Input : ModuleInputs)
 	{
 		UNiagaraGraph* AssetGraph = HierarchyInputToFunctionCallNodeMap[Input]->GetCalledGraph();
-		ensure(AssetGraph != nullptr);
-
-		TArray<FNiagaraVariable> Variables;
-		AssetGraph->GetAllVariables(Variables);
-
-		FNiagaraVariable ThisInput;
-		if(Input->GetPersistentIdentity().Guids.Num() == 2)
-		{
-			ThisInput = HierarchyInputToScriptVariableMap[Input]->Variable;
-		}
-		else
+		if(ensure(AssetGraph != nullptr))
 		{
-			FNiagaraVariable* FoundTargetVariable = Variables.FindByPredicate([InputName = Input->GetPersistentIdentity().Names[0]](FNiagaraVariable& CandidateVariable)
-			{
-				return CandidateVariable.GetName() == InputName;	
-			});
+			TArray<FNiagaraVariable> Variables;
+			AssetGraph->GetAllVariables(Variables);
 
-			if(FoundTargetVariable)
+			FNiagaraVariable ThisInput;
+			if(Input->GetPersistentIdentity().Guids.Num() == 2)
 			{
-				ThisInput = *FoundTargetVariable;
+				ThisInput = HierarchyInputToScriptVariableMap[Input]->Variable;
 			}
-		}
-
-		ensure(ThisInput.IsValid());
-		
-		TArray<FNiagaraVariable> ChildrenInputs;
-		FNiagaraParameterHandle ThisInputHandle(ThisInput.GetName());
-
-		TMap<FGuid, int32> ChildrenGuidSortOrderMap;
-		TArray<FGuid> ChildrenInputGuids;
-		for(const FNiagaraVariable& Variable : Variables)
-		{
-			bool bIsModuleInput = Variable.IsInNameSpace(FNiagaraConstants::ModuleNamespaceString);
-			TOptional<bool> bIsStaticSwitchInputOptional = AssetGraph->IsStaticSwitch(Variable);
-			if(!bIsModuleInput && !bIsStaticSwitchInputOptional.Get(false))
+			else
 			{
-				continue;
+				FNiagaraVariable* FoundTargetVariable = Variables.FindByPredicate([InputName = Input->GetPersistentIdentity().Names[0]](FNiagaraVariable& CandidateVariable)
+				{
+					return CandidateVariable.GetName() == InputName;	
+				});
+
+				if(FoundTargetVariable)
+				{
+					ThisInput = *FoundTargetVariable;
+				}
 			}
 
-			TOptional<FNiagaraVariableMetaData> VariableMetaData = AssetGraph->GetMetaData(Variable);
-			if(VariableMetaData.IsSet())
+			ensure(ThisInput.IsValid());
+			
+			TArray<FNiagaraVariable> ChildrenInputs;
+			FNiagaraParameterHandle ThisInputHandle(ThisInput.GetName());
+
+			TMap<FGuid, int32> ChildrenGuidSortOrderMap;
+			TArray<FGuid> ChildrenInputGuids;
+			for(const FNiagaraVariable& Variable : Variables)
 			{
-				if(VariableMetaData.GetValue().ParentAttribute == ThisInputHandle.GetName())
+				bool bIsModuleInput = Variable.IsInNameSpace(FNiagaraConstants::ModuleNamespaceString);
+				TOptional<bool> bIsStaticSwitchInputOptional = AssetGraph->IsStaticSwitch(Variable);
+				if(!bIsModuleInput && !bIsStaticSwitchInputOptional.Get(false))
 				{
-					HierarchyInputToChildrenGuidMap.FindOrAdd(Input).Add(VariableMetaData.GetValue().GetVariableGuid());
-					ChildrenGuidToScriptVariablesMap.Add(VariableMetaData->GetVariableGuid(), AssetGraph->GetScriptVariable(Variable));
-					ChildrenGuidSortOrderMap.Add(VariableMetaData->GetVariableGuid(), VariableMetaData->EditorSortPriority);
+					continue;
 				}
-			}
-		}
 
-		// we sort children by sort order specified in the metadata
-		if(HierarchyInputToChildrenGuidMap.Contains(Input))
-		{
-			HierarchyInputToChildrenGuidMap[Input].Sort([&ChildrenGuidSortOrderMap](const FGuid& GuidA, const FGuid& GuidB)
+				TOptional<FNiagaraVariableMetaData> VariableMetaData = AssetGraph->GetMetaData(Variable);
+				if(VariableMetaData.IsSet())
+				{
+					if(VariableMetaData.GetValue().ParentAttribute == ThisInputHandle.GetName())
+					{
+						HierarchyInputToChildrenGuidMap.FindOrAdd(Input).Add(VariableMetaData.GetValue().GetVariableGuid());
+						ChildrenGuidToScriptVariablesMap.Add(VariableMetaData->GetVariableGuid(), AssetGraph->GetScriptVariable(Variable));
+						ChildrenGuidSortOrderMap.Add(VariableMetaData->GetVariableGuid(), VariableMetaData->EditorSortPriority);
+					}
+				}
+			}
+			
+			// we sort children by sort order specified in the metadata
+			if(HierarchyInputToChildrenGuidMap.Contains(Input))
 			{
-				return ChildrenGuidSortOrderMap[GuidA] < ChildrenGuidSortOrderMap[GuidB];
-			});
+				HierarchyInputToChildrenGuidMap[Input].Sort([&ChildrenGuidSortOrderMap](const FGuid& GuidA, const FGuid& GuidB)
+				{
+					return ChildrenGuidSortOrderMap[GuidA] < ChildrenGuidSortOrderMap[GuidB];
+				});
+			}
 		}
 	}
 
