Fixed nullptr access leading to crash on cook

diff --git a/Engine/Plugins/FX/Niagara/Source/NiagaraEditor/Private/NiagaraScriptMergeManager.cpp b/Engine/Plugins/FX/Niagara/Source/NiagaraEditor/Private/NiagaraScriptMergeManager.cpp
--- a/Engine/Plugins/FX/Niagara/Source/NiagaraEditor/Private/NiagaraScriptMergeManager.cpp
+++ b/Engine/Plugins/FX/Niagara/Source/NiagaraEditor/Private/NiagaraScriptMergeManager.cpp
@@ -2444,6 +2444,11 @@ void FNiagaraScriptMergeManager::DiffRenderers(const TArray<TSharedRef<FNiagaraR
 
 void FNiagaraScriptMergeManager::DiffEmitterSummary(const UNiagaraEmitterEditorData* BaseEditorData, const UNiagaraEmitterEditorData* OtherEditorData, FNiagaraEmitterDiffResults& DiffResults) const
 {
+	if(BaseEditorData == nullptr || OtherEditorData == nullptr)
+	{
+		return;
+	}
+	
 	TArray<UNiagaraHierarchyItemBase*> AddedItemsInOther;
 	TArray<UNiagaraHierarchyItemBase*> RemovedItemsInBase;
 	TArray<UNiagaraHierarchySection*> AddedSections;
