- Fix for emitter name not showing up in customization after copy and paste

diff --git a/Engine/Plugins/FX/Niagara/Source/NiagaraEditorWidgets/Private/DetailCustomizations/NiagaraDataInterfaceParticleReadDetails.cpp b/Engine/Plugins/FX/Niagara/Source/NiagaraEditorWidgets/Private/DetailCustomizations/NiagaraDataInterfaceParticleReadDetails.cpp
--- a/Engine/Plugins/FX/Niagara/Source/NiagaraEditorWidgets/Private/DetailCustomizations/NiagaraDataInterfaceParticleReadDetails.cpp
+++ b/Engine/Plugins/FX/Niagara/Source/NiagaraEditorWidgets/Private/DetailCustomizations/NiagaraDataInterfaceParticleReadDetails.cpp
@@ -11,6 +11,14 @@
 
 #define LOCTEXT_NAMESPACE "NiagaraDataInterfaceParticleReadDetails"
 
+FNiagaraDataInterfaceParticleReadDetails::~FNiagaraDataInterfaceParticleReadDetails()
+{
+	if (ReadDataInterfaceWeak.IsValid())
+	{
+		ReadDataInterfaceWeak->OnChanged().RemoveAll(this);
+	}
+}
+
 void FNiagaraDataInterfaceParticleReadDetails::CustomizeDetails(IDetailLayoutBuilder& DetailBuilder)
 {
 	FNiagaraDataInterfaceDetailsBase::CustomizeDetails(DetailBuilder);
@@ -26,6 +34,8 @@ void FNiagaraDataInterfaceParticleReadDetails::CustomizeDetails(IDetailLayoutBui
 	ReadDataInterfaceWeak = ReadDataInterface;
 	NiagaraSystemWeak = ReadDataInterface->GetTypedOuter<UNiagaraSystem>();
 
+	ReadDataInterface->OnChanged().AddSP(this, &FNiagaraDataInterfaceParticleReadDetails::OnDataInterfaceChanged);
+
 	static const FName ReadCategoryName = TEXT("ParticleRead");
 	IDetailCategoryBuilder& CategoryBuilder = DetailBuilder.EditCategory(ReadCategoryName);
 	
@@ -46,14 +56,15 @@ void FNiagaraDataInterfaceParticleReadDetails::CustomizeDetails(IDetailLayoutBui
 				]
 				.ValueContent()
 				[
-					SNew(SSuggestionTextBox)
+					SAssignNew(TextBoxWidget, SSuggestionTextBox)
 						.ForegroundColor(FSlateColor::UseForeground())
 						.Font(IDetailLayoutBuilder::GetDetailFont())
-						.Text(FText::FromString(ReadDataInterface->EmitterName))
 						.OnTextCommitted(this, &FNiagaraDataInterfaceParticleReadDetails::EmitterName_SetText)
 						.HintText(LOCTEXT("EmitterNameHint", "Enter Emitter Name..."))
 						.OnShowingSuggestions(this, &FNiagaraDataInterfaceParticleReadDetails::EmitterName_GetSuggestions)
 				];
+
+			OnDataInterfaceChanged();
 		}
 		else
 		{
@@ -67,6 +78,14 @@ TSharedRef<IDetailCustomization> FNiagaraDataInterfaceParticleReadDetails::MakeI
 	return MakeShared<FNiagaraDataInterfaceParticleReadDetails>();
 }
 
+void FNiagaraDataInterfaceParticleReadDetails::OnDataInterfaceChanged()
+{
+	if (UNiagaraDataInterfaceParticleRead* ReadDataInterface = ReadDataInterfaceWeak.Get())
+	{
+		TextBoxWidget->SetText(FText::FromString(ReadDataInterface->EmitterName));
+	}
+}
+
 void FNiagaraDataInterfaceParticleReadDetails::EmitterName_SetText(const FText& NewText, ETextCommit::Type CommitInfo)
 {
 	if (EmitterNameHandle)
diff --git a/Engine/Plugins/FX/Niagara/Source/NiagaraEditorWidgets/Private/DetailCustomizations/NiagaraDataInterfaceParticleReadDetails.h b/Engine/Plugins/FX/Niagara/Source/NiagaraEditorWidgets/Private/DetailCustomizations/NiagaraDataInterfaceParticleReadDetails.h
--- a/Engine/Plugins/FX/Niagara/Source/NiagaraEditorWidgets/Private/DetailCustomizations/NiagaraDataInterfaceParticleReadDetails.h
+++ b/Engine/Plugins/FX/Niagara/Source/NiagaraEditorWidgets/Private/DetailCustomizations/NiagaraDataInterfaceParticleReadDetails.h
@@ -19,14 +19,18 @@ class SNiagaraNamePropertySelector;
 class FNiagaraDataInterfaceParticleReadDetails : public FNiagaraDataInterfaceDetailsBase
 {
 public:
+	virtual ~FNiagaraDataInterfaceParticleReadDetails();
+
 	virtual void CustomizeDetails(IDetailLayoutBuilder& DetailBuilder) override;
 	static TSharedRef<IDetailCustomization> MakeInstance();
 
 private:
+	void OnDataInterfaceChanged();
 	void EmitterName_SetText(const FText& NewText, ETextCommit::Type CommitInfo);
 	void EmitterName_GetSuggestions(const FString& CurrText, TArray<FString>& OutSuggestions);
 
 	TWeakObjectPtr<UNiagaraDataInterfaceParticleRead> ReadDataInterfaceWeak;
 	TWeakObjectPtr<UNiagaraSystem> NiagaraSystemWeak;
 	TSharedPtr<IPropertyHandle> EmitterNameHandle;
+	TSharedPtr<class SSuggestionTextBox> TextBoxWidget;
 };
