- Modify custom depth write check to only be enabled if the proxy will go into the custom depth write path

#rnx

diff --git a/Engine/Plugins/FX/Niagara/Source/Niagara/Private/NiagaraRendererSprites.cpp b/Engine/Plugins/FX/Niagara/Source/Niagara/Private/NiagaraRendererSprites.cpp
--- a/Engine/Plugins/FX/Niagara/Source/Niagara/Private/NiagaraRendererSprites.cpp
+++ b/Engine/Plugins/FX/Niagara/Source/Niagara/Private/NiagaraRendererSprites.cpp
@@ -247,14 +247,19 @@ void FNiagaraRendererSprites::PrepareParticleSpriteRenderData(FParticleSpriteRen
 	ParticleSpriteRenderData.bHasTranslucentMaterials = IsTranslucentBlendMode(Material);
 
 	// If these conditions change please update the DebugHUD display also to reflect it
-	const bool bLowLatencyTranslucencyEnabled =
+	bool bLowLatencyTranslucencyEnabled =
 		ParticleSpriteRenderData.bHasTranslucentMaterials &&
 		bGpuLowLatencyTranslucency &&
 		GpuReadyTickStage >= CurrentParticleData->GetGPUDataReadyStage() &&
-		Material.IsTranslucencyWritingCustomDepth() &&
 		!SceneProxy->CastsVolumetricTranslucentShadow() &&
 		ViewFamilySupportLowLatencyTranslucency(ViewFamily);
 
+	if (bLowLatencyTranslucencyEnabled && SceneProxy->ShouldRenderCustomDepth())
+	{
+		bLowLatencyTranslucencyEnabled &= !Material.IsTranslucencyWritingCustomDepth();
+	}
+
+
 	ParticleSpriteRenderData.SourceParticleData = ParticleSpriteRenderData.DynamicDataSprites->GetParticleDataToRender(bLowLatencyTranslucencyEnabled);
 	if ( !ParticleSpriteRenderData.SourceParticleData || (SourceMode == ENiagaraRendererSourceDataMode::Particles && ParticleSpriteRenderData.SourceParticleData->GetNumInstances() == 0) )
 	{
diff --git a/Engine/Source/Runtime/Renderer/Private/FXRenderingUtils.cpp b/Engine/Source/Runtime/Renderer/Private/FXRenderingUtils.cpp
--- a/Engine/Source/Runtime/Renderer/Private/FXRenderingUtils.cpp
+++ b/Engine/Source/Runtime/Renderer/Private/FXRenderingUtils.cpp
@@ -27,7 +27,7 @@ bool UE::FXRenderingUtils::CanMaterialRenderBeforeFXPostOpaque(
 	const FMaterial& Material)
 {
 	// Opaque materials & translucent that write custom depth will always need to render before FFXSystemInterface::PostOpaqueRender
-	if (!IsTranslucentBlendMode(Material) || Material.IsTranslucencyWritingCustomDepth())
+	if (!IsTranslucentBlendMode(Material) || (SceneProxy.ShouldRenderCustomDepth() && Material.IsTranslucencyWritingCustomDepth()))
 	{
 		return true;
 	}
