- Move Niagara SVT into shaders module to fix load order issue

#rnx

diff --git a/Engine/Plugins/FX/Niagara/Source/Niagara/Private/NiagaraDataInterfaceRenderTargetVolume.cpp b/Engine/Plugins/FX/Niagara/Source/Niagara/Private/NiagaraDataInterfaceRenderTargetVolume.cpp
--- a/Engine/Plugins/FX/Niagara/Source/Niagara/Private/NiagaraDataInterfaceRenderTargetVolume.cpp
+++ b/Engine/Plugins/FX/Niagara/Source/Niagara/Private/NiagaraDataInterfaceRenderTargetVolume.cpp
@@ -11,6 +11,7 @@
 #include "NiagaraDataInterfaceRenderTargetCommon.h"
 #include "NiagaraGpuComputeDispatchInterface.h"
 #include "NiagaraSettings.h"
+#include "NiagaraSVTShaders.h"
 #include "NiagaraSystemInstance.h"
 #include "NiagaraStats.h"
 #include "NiagaraShaderParametersBuilder.h"
@@ -845,34 +846,6 @@ bool UNiagaraDataInterfaceRenderTargetVolume::SimCacheEndWrite(UObject* StorageO
 	return true;
 }
 
-
-
-class FNiagaraCopySVTToDenseBufferCS : public FGlobalShader
-{
-	DECLARE_GLOBAL_SHADER(FNiagaraCopySVTToDenseBufferCS);
-	SHADER_USE_PARAMETER_STRUCT(FNiagaraCopySVTToDenseBufferCS, FGlobalShader);
-	
-	static void ModifyCompilationEnvironment(const FGlobalShaderPermutationParameters& Parameters, FShaderCompilerEnvironment& OutEnvironment)
-	{
-		FGlobalShader::ModifyCompilationEnvironment(Parameters, OutEnvironment);		
-	}
-
-	BEGIN_SHADER_PARAMETER_STRUCT(FParameters, )
-
-		SHADER_PARAMETER_RDG_TEXTURE_UAV(RWTexture3D<float4>, DestinationBuffer)
-		SHADER_PARAMETER_SAMPLER(SamplerState, TileDataTextureSampler)
-		SHADER_PARAMETER_TEXTURE(Texture3D<uint>, SparseVolumeTexturePageTable)
-		SHADER_PARAMETER_TEXTURE(Texture3D, SparseVolumeTextureA)		
-		SHADER_PARAMETER(FUintVector4, PackedSVTUniforms0)
-		SHADER_PARAMETER(FUintVector4, PackedSVTUniforms1)
-		SHADER_PARAMETER(FIntVector, TextureSize)
-		SHADER_PARAMETER(int32, MipLevel)
-	END_SHADER_PARAMETER_STRUCT()
-};
-
-IMPLEMENT_GLOBAL_SHADER(FNiagaraCopySVTToDenseBufferCS, "/Plugin/FX/Niagara/Private/NiagaraSVTToDenseBuffer.usf", "PerformCopyCS", SF_Compute);
-
-
 bool UNiagaraDataInterfaceRenderTargetVolume::SimCacheReadFrame(UObject* StorageObject, int FrameA, int FrameB, float Interp, FNiagaraSystemInstance* SystemInstance, void* OptionalPerInstanceData)
 {
 	FRenderTargetVolumeRWInstanceData_GameThread* InstanceData_GT = reinterpret_cast<FRenderTargetVolumeRWInstanceData_GameThread*>(OptionalPerInstanceData);
diff --git a/Engine/Plugins/FX/Niagara/Source/NiagaraShader/Private/NiagaraSVTShaders.cpp b/Engine/Plugins/FX/Niagara/Source/NiagaraShader/Private/NiagaraSVTShaders.cpp
new file mode 100644
--- /dev/null
+++ b/Engine/Plugins/FX/Niagara/Source/NiagaraShader/Private/NiagaraSVTShaders.cpp
@@ -0,0 +1,10 @@
+// Copyright Epic Games, Inc. All Rights Reserved.
+
+#include "NiagaraSVTShaders.h"
+
+void FNiagaraCopySVTToDenseBufferCS::ModifyCompilationEnvironment(const FGlobalShaderPermutationParameters& Parameters, FShaderCompilerEnvironment& OutEnvironment)
+{
+	FGlobalShader::ModifyCompilationEnvironment(Parameters, OutEnvironment);
+}
+
+IMPLEMENT_GLOBAL_SHADER(FNiagaraCopySVTToDenseBufferCS, "/Plugin/FX/Niagara/Private/NiagaraSVTToDenseBuffer.usf", "PerformCopyCS", SF_Compute);
diff --git a/Engine/Plugins/FX/Niagara/Source/NiagaraShader/Public/NiagaraSVTShaders.h b/Engine/Plugins/FX/Niagara/Source/NiagaraShader/Public/NiagaraSVTShaders.h
new file mode 100644
--- /dev/null
+++ b/Engine/Plugins/FX/Niagara/Source/NiagaraShader/Public/NiagaraSVTShaders.h
@@ -0,0 +1,25 @@
+// Copyright Epic Games, Inc. All Rights Reserved.
+
+#pragma once
+
+#include "GlobalShader.h"
+#include "ShaderParameterStruct.h"
+
+class FNiagaraCopySVTToDenseBufferCS : public FGlobalShader
+{
+	DECLARE_EXPORTED_GLOBAL_SHADER(FNiagaraCopySVTToDenseBufferCS, NIAGARASHADER_API);
+	SHADER_USE_PARAMETER_STRUCT(FNiagaraCopySVTToDenseBufferCS, FGlobalShader);
+	
+	static void ModifyCompilationEnvironment(const FGlobalShaderPermutationParameters& Parameters, FShaderCompilerEnvironment& OutEnvironment);
+
+	BEGIN_SHADER_PARAMETER_STRUCT(FParameters, NIAGARASHADER_API)
+		SHADER_PARAMETER_RDG_TEXTURE_UAV(RWTexture3D<float4>, DestinationBuffer)
+		SHADER_PARAMETER_SAMPLER(SamplerState, TileDataTextureSampler)
+		SHADER_PARAMETER_TEXTURE(Texture3D<uint>, SparseVolumeTexturePageTable)
+		SHADER_PARAMETER_TEXTURE(Texture3D, SparseVolumeTextureA)		
+		SHADER_PARAMETER(FUintVector4, PackedSVTUniforms0)
+		SHADER_PARAMETER(FUintVector4, PackedSVTUniforms1)
+		SHADER_PARAMETER(FIntVector, TextureSize)
+		SHADER_PARAMETER(int32, MipLevel)
+	END_SHADER_PARAMETER_STRUCT()
+};
