- Update simulation stages to allow various parameters to be a binding or a constant
- Iteration count is now a single field rather than Count or Binding
- Element Count can now be set directly rather than a binding only
- Num Threads can be set from a static int alllowing for greater customization
- FNiagaraParameterBindingWithValue was added to faciliate having either a binding or a value
- Also fixed the lack of variable / emitter renaming / deletion support on simulation stage bindings, this uses a field iterator to avoid having to remember to add support for each property

diff --git a/Engine/Plugins/FX/Niagara/Config/BaseNiagara.ini b/Engine/Plugins/FX/Niagara/Config/BaseNiagara.ini
--- a/Engine/Plugins/FX/Niagara/Config/BaseNiagara.ini
+++ b/Engine/Plugins/FX/Niagara/Config/BaseNiagara.ini
@@ -32,6 +32,8 @@
 
 +PropertyRedirects=(OldName="NiagaraRibbonRendererProperties.RibbonDistanceFromStart", NewName="NiagaraRibbonRendererProperties.RibbonUVDistance")
 
++PropertyRedirects=(OldName="NiagaraParameterBinding.Parameter", NewName="NiagaraParameterBinding.ResolvedParameter")
+
 ; Flipbook to Baker
 +ClassRedirects=(OldName="NiagaraFlipbookSettings", NewName="/Script/Niagara.NiagaraBakerSettings")
 +PropertyRedirects=(OldName="NiagaraSystem.FlipbookSettings", NewName="NiagaraSystem.BakerSettings")
diff --git a/Engine/Plugins/FX/Niagara/Source/Niagara/Private/NiagaraDecalRendererProperties.cpp b/Engine/Plugins/FX/Niagara/Source/Niagara/Private/NiagaraDecalRendererProperties.cpp
--- a/Engine/Plugins/FX/Niagara/Source/Niagara/Private/NiagaraDecalRendererProperties.cpp
+++ b/Engine/Plugins/FX/Niagara/Source/Niagara/Private/NiagaraDecalRendererProperties.cpp
@@ -377,7 +377,7 @@ UMaterialInterface* UNiagaraDecalRendererProperties::GetMaterial(const FNiagaraE
 	UMaterialInterface* MaterialInterface = nullptr;
 	if (InEmitter != nullptr)
 	{
-		MaterialInterface = Cast<UMaterialInterface>(InEmitter->FindBinding(MaterialParameterBinding.Parameter));
+		MaterialInterface = Cast<UMaterialInterface>(InEmitter->FindBinding(MaterialParameterBinding.ResolvedParameter));
 	}
 
 #if WITH_EDITORONLY_DATA
diff --git a/Engine/Plugins/FX/Niagara/Source/Niagara/Private/NiagaraEmitter.cpp b/Engine/Plugins/FX/Niagara/Source/Niagara/Private/NiagaraEmitter.cpp
--- a/Engine/Plugins/FX/Niagara/Source/Niagara/Private/NiagaraEmitter.cpp
+++ b/Engine/Plugins/FX/Niagara/Source/Niagara/Private/NiagaraEmitter.cpp
@@ -1332,14 +1332,13 @@ void UNiagaraEmitter::HandleVariableRenamed(const FNiagaraVariable& InOldVariabl
 		Prop->RenameVariable(InOldVariable, InNewVariable, FVersionedNiagaraEmitter(this, EmitterVersion));
 	}
 
-	// Rename any simulation stage iteration sources
+	// Rename any simulation stage variables
 	for (UNiagaraSimulationStageBase* SimStage : Data->SimulationStages)
 	{
-		UNiagaraSimulationStageGeneric* GenericStage = Cast<UNiagaraSimulationStageGeneric>(SimStage);
-		if (GenericStage && GenericStage->DataInterface.BoundVariable.GetName() == InOldVariable.GetName())
+		if (UNiagaraSimulationStageGeneric* GenericStage = Cast<UNiagaraSimulationStageGeneric>(SimStage))
 		{
 			GenericStage->Modify(false);
-			GenericStage->DataInterface.BoundVariable = InNewVariable;
+			GenericStage->RenameVariable(InOldVariable, InNewVariable, FVersionedNiagaraEmitter(this, EmitterVersion));
 		}
 	}
 
@@ -1361,6 +1360,16 @@ void UNiagaraEmitter::HandleVariableRemoved(const FNiagaraVariable& InOldVariabl
 		Prop->RemoveVariable(InOldVariable, FVersionedNiagaraEmitter(this, EmitterVersion));
 	}
 
+	// Remove any simulation stage variables
+	for (UNiagaraSimulationStageBase* SimStage : Data->SimulationStages)
+	{
+		if (UNiagaraSimulationStageGeneric* GenericStage = Cast<UNiagaraSimulationStageGeneric>(SimStage))
+		{
+			GenericStage->Modify(false);
+			GenericStage->RemoveVariable(InOldVariable, FVersionedNiagaraEmitter(this, EmitterVersion));
+		}
+	}
+
 	Data->RebuildRendererBindings(*this);
 
 	if (bUpdateContexts)
@@ -2905,6 +2914,15 @@ void FVersionedNiagaraEmitterData::SyncEmitterAlias(const FString& InOldName, co
 				Renderer->RenameEmitter(*InOldName, &InEmitter);
 			}
 		}
+
+		for (UNiagaraSimulationStageBase* SimStage : SimulationStages)
+		{
+			if (UNiagaraSimulationStageGeneric* GenericStage = Cast<UNiagaraSimulationStageGeneric>(SimStage))
+			{
+				GenericStage->Modify(false);
+				GenericStage->RenameEmitter(*InOldName, &InEmitter);
+			}
+		}
 	}
 }
 
diff --git a/Engine/Plugins/FX/Niagara/Source/Niagara/Private/NiagaraGPUSystemTick.cpp b/Engine/Plugins/FX/Niagara/Source/Niagara/Private/NiagaraGPUSystemTick.cpp
--- a/Engine/Plugins/FX/Niagara/Source/Niagara/Private/NiagaraGPUSystemTick.cpp
+++ b/Engine/Plugins/FX/Niagara/Source/Niagara/Private/NiagaraGPUSystemTick.cpp
@@ -243,15 +243,15 @@ void FNiagaraGPUSystemTick::Init(FNiagaraSystemInstance* InSystemInstance)
 						// This will cause a PVS warning, if we add something before this enum entry add the code back in
 						//if (SimStageMetaData.GpuDispatchType >= ENiagaraGpuDispatchType::OneD)
 						{
-							ElementCountXYZ.X = SimStageMetaData.ElementCountXBinding.IsNone() ? 0 : BoundParamStore.GetParameterValueOrDefault(FNiagaraVariable(FNiagaraTypeDefinition::GetIntDef(), SimStageMetaData.ElementCountXBinding), 0);
+							ElementCountXYZ.X = SimStageMetaData.ElementCountXBinding.IsNone() ? SimStageMetaData.ElementCount.X : BoundParamStore.GetParameterValueOrDefault(FNiagaraVariable(FNiagaraTypeDefinition::GetIntDef(), SimStageMetaData.ElementCountXBinding), 0);
 						}
 						if (SimStageMetaData.GpuDispatchType >= ENiagaraGpuDispatchType::TwoD)
 						{
-							ElementCountXYZ.Y = SimStageMetaData.ElementCountYBinding.IsNone() ? 0 : BoundParamStore.GetParameterValueOrDefault(FNiagaraVariable(FNiagaraTypeDefinition::GetIntDef(), SimStageMetaData.ElementCountYBinding), 0);
+							ElementCountXYZ.Y = SimStageMetaData.ElementCountYBinding.IsNone() ? SimStageMetaData.ElementCount.Y : BoundParamStore.GetParameterValueOrDefault(FNiagaraVariable(FNiagaraTypeDefinition::GetIntDef(), SimStageMetaData.ElementCountYBinding), 0);
 						}
 						if (SimStageMetaData.GpuDispatchType >= ENiagaraGpuDispatchType::ThreeD)
 						{
-							ElementCountXYZ.Z = SimStageMetaData.ElementCountZBinding.IsNone() ? 0 : BoundParamStore.GetParameterValueOrDefault(FNiagaraVariable(FNiagaraTypeDefinition::GetIntDef(), SimStageMetaData.ElementCountZBinding), 0);
+							ElementCountXYZ.Z = SimStageMetaData.ElementCountZBinding.IsNone() ? SimStageMetaData.ElementCount.Z : BoundParamStore.GetParameterValueOrDefault(FNiagaraVariable(FNiagaraTypeDefinition::GetIntDef(), SimStageMetaData.ElementCountZBinding), 0);
 						}
 
 						// If any of the element count values are <= 0 we can kill the stage as it won't execute
diff --git a/Engine/Plugins/FX/Niagara/Source/Niagara/Private/NiagaraMeshRendererProperties.cpp b/Engine/Plugins/FX/Niagara/Source/Niagara/Private/NiagaraMeshRendererProperties.cpp
--- a/Engine/Plugins/FX/Niagara/Source/Niagara/Private/NiagaraMeshRendererProperties.cpp
+++ b/Engine/Plugins/FX/Niagara/Source/Niagara/Private/NiagaraMeshRendererProperties.cpp
@@ -141,18 +141,19 @@ namespace NiagaraMeshRendererPropertiesInternal
 		OutStaticMesh = nullptr;
 		if (EmitterInstance)
 		{
-			if (MeshProperties.MeshParameterBinding.Parameter.IsValid())
+			const FNiagaraVariableBase& MeshParameter = MeshProperties.MeshParameterBinding.ResolvedParameter;
+			if (MeshParameter.IsValid())
 			{
-				if (MeshProperties.MeshParameterBinding.Parameter.IsDataInterface())
+				if (MeshParameter.IsDataInterface())
 				{
-					OutInterface = Cast<INiagaraRenderableMeshInterface>(EmitterInstance->GetRendererBoundVariables().GetDataInterface(MeshProperties.MeshParameterBinding.Parameter));
+					OutInterface = Cast<INiagaraRenderableMeshInterface>(EmitterInstance->GetRendererBoundVariables().GetDataInterface(MeshParameter));
 					if (OutInterface != nullptr)
 					{
 						return;
 					}
 				}
 
-				UStaticMesh* BoundMesh = Cast<UStaticMesh>(EmitterInstance->GetRendererBoundVariables().GetUObject(MeshProperties.MeshParameterBinding.Parameter));
+				UStaticMesh* BoundMesh = Cast<UStaticMesh>(EmitterInstance->GetRendererBoundVariables().GetUObject(MeshParameter));
 				if (BoundMesh && BoundMesh->GetRenderData())
 				{
 					OutStaticMesh = BoundMesh;
@@ -229,7 +230,7 @@ FNiagaraRenderableMeshPtr FNiagaraMeshRendererMeshProperties::ResolveRenderableM
 
 bool FNiagaraMeshRendererMeshProperties::HasValidRenderableMesh() const
 {
-	return Mesh || MeshParameterBinding.Parameter.IsValid();
+	return Mesh || MeshParameterBinding.ResolvedParameter.IsValid();
 }
 
 UNiagaraMeshRendererProperties::UNiagaraMeshRendererProperties()
@@ -718,9 +719,9 @@ bool UNiagaraMeshRendererProperties::PopulateRequiredBindings(FNiagaraParameterS
 
 	for (FNiagaraMeshRendererMeshProperties& Binding : Meshes)
 	{
-		if (Binding.MeshParameterBinding.Parameter.IsValid())
+		if (Binding.MeshParameterBinding.ResolvedParameter.IsValid())
 		{
-			InParameterStore.AddParameter(Binding.MeshParameterBinding.Parameter, false);
+			InParameterStore.AddParameter(Binding.MeshParameterBinding.ResolvedParameter, false);
 			bAnyAdded = true;
 		}
 	}
@@ -757,7 +758,7 @@ void UNiagaraMeshRendererProperties::PostLoad()
 #if WITH_EDITORONLY_DATA
 		if (MeshProperties.UserParamBinding_DEPRECATED.Parameter.GetName().IsNone() == false)
 		{
-			MeshProperties.MeshParameterBinding.Parameter = MeshProperties.UserParamBinding_DEPRECATED.Parameter;
+			MeshProperties.MeshParameterBinding.ResolvedParameter = MeshProperties.UserParamBinding_DEPRECATED.Parameter;
 			MeshProperties.MeshParameterBinding.AliasedParameter = MeshProperties.UserParamBinding_DEPRECATED.Parameter;
 		}
 #endif
@@ -880,10 +881,10 @@ void UNiagaraMeshRendererProperties::GetRendererTooltipWidgets(const FNiagaraEmi
 		}
 
 		// we override the previous thumbnail tooltip with a text indicating parameter binding, if it exists
-		if(MeshProperties.MeshParameterBinding.Parameter.IsValid())
+		if(MeshProperties.MeshParameterBinding.ResolvedParameter.IsValid())
 		{
 			TooltipWidget = SNew(STextBlock)
-				.Text(FText::Format(LOCTEXT("MeshBoundTooltip", "Mesh slot is bound to parameter {0}"), FText::FromName(MeshProperties.MeshParameterBinding.Parameter.GetName())));
+				.Text(FText::Format(LOCTEXT("MeshBoundTooltip", "Mesh slot is bound to parameter {0}"), FText::FromName(MeshProperties.MeshParameterBinding.ResolvedParameter.GetName())));
 		}
 		
 		OutWidgets.Add(TooltipWidget);
diff --git a/Engine/Plugins/FX/Niagara/Source/Niagara/Private/NiagaraModule.cpp b/Engine/Plugins/FX/Niagara/Source/Niagara/Private/NiagaraModule.cpp
--- a/Engine/Plugins/FX/Niagara/Source/Niagara/Private/NiagaraModule.cpp
+++ b/Engine/Plugins/FX/Niagara/Source/Niagara/Private/NiagaraModule.cpp
@@ -1041,6 +1041,12 @@ FNiagaraTypeDefinition FNiagaraTypeDefinition::ToStaticDef() const
 	return Def;
 }
 
+FNiagaraTypeDefinition FNiagaraTypeDefinition::RemoveStaticDef() const
+{
+	FNiagaraTypeDefinition Def(*this);
+	Def.Flags &= ~TF_Static;
+	return Def;
+}
 
 #if WITH_EDITORONLY_DATA
 bool FNiagaraTypeDefinition::IsInternalType() const
diff --git a/Engine/Plugins/FX/Niagara/Source/Niagara/Private/NiagaraParameterBinding.cpp b/Engine/Plugins/FX/Niagara/Source/Niagara/Private/NiagaraParameterBinding.cpp
--- a/Engine/Plugins/FX/Niagara/Source/Niagara/Private/NiagaraParameterBinding.cpp
+++ b/Engine/Plugins/FX/Niagara/Source/Niagara/Private/NiagaraParameterBinding.cpp
@@ -12,6 +12,14 @@ bool FNiagaraParameterBinding::CanBindTo(FNiagaraTypeDefinition TypeDefinition)
 		return true;
 	}
 
+	if (TypeDefinition.IsStatic() && AllowStaticVariables())
+	{
+		if (AllowedTypeDefinitions.Contains(TypeDefinition.RemoveStaticDef()))
+		{
+			return true;
+		}
+	}
+
 	if (UClass* TypeClass = TypeDefinition.GetClass())
 	{
 		if (TypeClass->IsChildOf(UNiagaraDataInterface::StaticClass()))
@@ -84,27 +92,143 @@ void FNiagaraParameterBinding::OnRenameEmitter(FStringView EmitterName)
 {
 	if (AliasedParameter.IsInNameSpace(FNiagaraConstants::EmitterNamespaceString))
 	{
-		Parameter = AliasedParameter;
-		Parameter.ReplaceRootNamespace(FNiagaraConstants::EmitterNamespaceString, EmitterName);
+		ResolvedParameter = AliasedParameter;
+		ResolvedParameter.ReplaceRootNamespace(FNiagaraConstants::EmitterNamespaceString, EmitterName);
+	}
+	if (DefaultAliasedParameter.IsInNameSpace(FNiagaraConstants::EmitterNamespaceString))
+	{
+		DefaultResolvedParameter = DefaultAliasedParameter;
+		DefaultResolvedParameter.ReplaceRootNamespace(FNiagaraConstants::EmitterNamespaceString, EmitterName);
 	}
 }
 
 void FNiagaraParameterBinding::OnRenameVariable(const FNiagaraVariableBase& OldVariable, const FNiagaraVariableBase& NewVariable, FStringView EmitterName)
 {
-	if (AliasedParameter.GetName() == OldVariable.GetName() && AliasedParameter.GetType() == OldVariable.GetType() && Parameter.IsInNameSpace(EmitterName))
+	if (AliasedParameter.GetName() == OldVariable.GetName() && AliasedParameter.GetType() == OldVariable.GetType() && ResolvedParameter.IsInNameSpace(EmitterName))
 	{
 		AliasedParameter = NewVariable;
-		Parameter = AliasedParameter;
-		Parameter.ReplaceRootNamespace(FNiagaraConstants::EmitterNamespaceString, EmitterName);
+		ResolvedParameter = AliasedParameter;
+		ResolvedParameter.ReplaceRootNamespace(FNiagaraConstants::EmitterNamespaceString, EmitterName);
 	}
 }
 
 void FNiagaraParameterBinding::OnRemoveVariable(const FNiagaraVariableBase& OldVariable, FStringView EmitterName)
 {
-	if (AliasedParameter.GetName() == OldVariable.GetName() && AliasedParameter.GetType() == OldVariable.GetType() && Parameter.IsInNameSpace(EmitterName))
+	if (AliasedParameter.GetName() == OldVariable.GetName() && AliasedParameter.GetType() == OldVariable.GetType() && ResolvedParameter.IsInNameSpace(EmitterName))
+	{
+		AliasedParameter = DefaultAliasedParameter;
+		ResolvedParameter = DefaultResolvedParameter;
+	}
+}
+
+void FNiagaraParameterBinding::SetDefaultParameter(const FNiagaraVariable& Variable)
+{
+	DefaultResolvedParameter = Variable;
+	DefaultAliasedParameter = Variable;
+
+	SetToDefault();
+}
+
+void FNiagaraParameterBinding::SetToDefault()
+{
+	AliasedParameter = DefaultAliasedParameter;
+	ResolvedParameter = DefaultResolvedParameter;
+
+	if (HasDefaultValueEditorOnly() && DefaultAliasedParameter.IsDataAllocated())
+	{
+		SetDefaultValueEditorOnly(MakeArrayView(DefaultAliasedParameter.GetData(), DefaultAliasedParameter.GetAllocatedSizeInBytes()));
+	}
+}
+
+bool FNiagaraParameterBinding::IsSetoToDefault() const
+{
+	if (HasDefaultValueEditorOnly() && DefaultAliasedParameter.IsDataAllocated())
+	{
+		TConstArrayView<uint8> DefaultValue = GetDefaultValueEditorOnly();
+		check(DefaultValue.Num() == DefaultAliasedParameter.GetAllocatedSizeInBytes());
+		if ( FMemory::Memcmp(DefaultValue.GetData(), DefaultAliasedParameter.GetData(), DefaultAliasedParameter.GetAllocatedSizeInBytes()) != 0 )
+		{
+			return false;
+		}
+	}
+	return AliasedParameter == DefaultAliasedParameter;
+}
+
+FString FNiagaraParameterBinding::ToString() const
+{
+	TStringBuilder<128> Builder;
+	Builder.Append(TEXT("Type("));
+	Builder.Append(ResolvedParameter.GetType().GetName());
+	Builder.Append(TEXT(") Name("));
+	ResolvedParameter.GetName().ToString(Builder);
+	if (HasDefaultValueEditorOnly())
 	{
-		AliasedParameter = FNiagaraVariableBase();
-		Parameter = FNiagaraVariableBase();
+		TConstArrayView<uint8> DefaultValue = GetDefaultValueEditorOnly();
+		Builder.Append(TEXT(") Value("));
+		Builder.Append(BytesToHex(DefaultValue.GetData(), DefaultValue.Num()));
 	}
+	Builder.Append(TEXT(")"));
+	return Builder.ToString();
+}
+
+void FNiagaraParameterBinding::ForEachRenameEmitter(UObject* InObject, FStringView EmitterName)
+{
+	for (TFieldIterator<FStructProperty> PropIt(InObject->GetClass()); PropIt; ++PropIt)
+	{
+		FStructProperty* StructProp = *PropIt;
+		if (StructProp->Struct == FNiagaraParameterBinding::StaticStruct())
+		{
+			FNiagaraParameterBinding* ParameterBinding = StructProp->ContainerPtrToValuePtr<FNiagaraParameterBinding>(InObject);
+			ParameterBinding->OnRenameEmitter(EmitterName);
+		}
+	}
+}
+
+void FNiagaraParameterBinding::ForEachRenameVariable(UObject* InObject, const FNiagaraVariableBase& OldVariable, const FNiagaraVariableBase& NewVariable, FStringView EmitterName)
+{
+	for (TFieldIterator<FStructProperty> PropIt(InObject->GetClass()); PropIt; ++PropIt)
+	{
+		FStructProperty* StructProp = *PropIt;
+		if (StructProp->Struct == FNiagaraParameterBinding::StaticStruct())
+		{
+			FNiagaraParameterBinding* ParameterBinding = StructProp->ContainerPtrToValuePtr<FNiagaraParameterBinding>(InObject);
+			ParameterBinding->OnRenameVariable(OldVariable, NewVariable, EmitterName);
+		}
+	}
+}
+
+void FNiagaraParameterBinding::ForEachRemoveVariable(UObject* InObject, const FNiagaraVariableBase& OldVariable, FStringView EmitterName)
+{
+	for (TFieldIterator<FStructProperty> PropIt(InObject->GetClass()); PropIt; ++PropIt)
+	{
+		FStructProperty* StructProp = *PropIt;
+		if (StructProp->Struct == FNiagaraParameterBinding::StaticStruct())
+		{
+			FNiagaraParameterBinding* ParameterBinding = StructProp->ContainerPtrToValuePtr<FNiagaraParameterBinding>(InObject);
+			ParameterBinding->OnRemoveVariable(OldVariable, EmitterName);
+		}
+	}
+}
+
+bool FNiagaraParameterBindingWithValue::HasDefaultValueEditorOnly() const
+{
+	return true;
+}
+
+TConstArrayView<uint8> FNiagaraParameterBindingWithValue::GetDefaultValueEditorOnly() const
+{
+	check(DefaultValue.Num() > 0);
+	return DefaultValue;
+}
+
+void FNiagaraParameterBindingWithValue::SetDefaultValueEditorOnly(TConstArrayView<uint8> Memory)
+{
+	DefaultValue = Memory;
+}
+
+void FNiagaraParameterBindingWithValue::SetDefaultValueEditorOnly(const uint8* Memory)
+{
+	check(DefaultValue.Num() > 0);
+	FMemory::Memcpy(DefaultValue.GetData(), Memory, DefaultValue.Num());
 }
 #endif
diff --git a/Engine/Plugins/FX/Niagara/Source/Niagara/Private/NiagaraScript.cpp b/Engine/Plugins/FX/Niagara/Source/Niagara/Private/NiagaraScript.cpp
--- a/Engine/Plugins/FX/Niagara/Source/Niagara/Private/NiagaraScript.cpp
+++ b/Engine/Plugins/FX/Niagara/Source/Niagara/Private/NiagaraScript.cpp
@@ -1131,13 +1131,20 @@ void UNiagaraScript::ComputeVMCompilationId(FNiagaraVMExecutableDataId& Id, FGui
 					if (const UNiagaraSimulationStageGeneric* SimStageGeneric = Cast<const UNiagaraSimulationStageGeneric>(SimStageBase))
 					{
 						AddAttributeToPreserve(SimStageGeneric->EnabledBinding.GetParamMapBindableVariable().GetName());
+
 						if (SimStageGeneric->IterationSource == ENiagaraIterationSource::DirectSet)
 						{
-							AddAttributeToPreserve(SimStageGeneric->ElementCountXBinding.GetParamMapBindableVariable().GetName());
-							AddAttributeToPreserve(SimStageGeneric->ElementCountYBinding.GetParamMapBindableVariable().GetName());
-							AddAttributeToPreserve(SimStageGeneric->ElementCountZBinding.GetParamMapBindableVariable().GetName());
+							AddAttributeToPreserve(SimStageGeneric->ElementCountX.ResolvedParameter.GetName());
+							AddAttributeToPreserve(SimStageGeneric->ElementCountY.ResolvedParameter.GetName());
+							AddAttributeToPreserve(SimStageGeneric->ElementCountZ.ResolvedParameter.GetName());
+						}
+						AddAttributeToPreserve(SimStageGeneric->NumIterations.ResolvedParameter.GetName());
+						if ( SimStageGeneric->bOverrideGpuDispatchNumThreads )
+						{
+							AddAttributeToPreserve(SimStageGeneric->OverrideGpuDispatchNumThreadsX.ResolvedParameter.GetName());
+							AddAttributeToPreserve(SimStageGeneric->OverrideGpuDispatchNumThreadsY.ResolvedParameter.GetName());
+							AddAttributeToPreserve(SimStageGeneric->OverrideGpuDispatchNumThreadsZ.ResolvedParameter.GetName());
 						}
-						AddAttributeToPreserve(SimStageGeneric->NumIterationsBinding.GetParamMapBindableVariable().GetName());
 					}
 				}
 
diff --git a/Engine/Plugins/FX/Niagara/Source/Niagara/Private/NiagaraSimulationStageBase.cpp b/Engine/Plugins/FX/Niagara/Source/Niagara/Private/NiagaraSimulationStageBase.cpp
--- a/Engine/Plugins/FX/Niagara/Source/Niagara/Private/NiagaraSimulationStageBase.cpp
+++ b/Engine/Plugins/FX/Niagara/Source/Niagara/Private/NiagaraSimulationStageBase.cpp
@@ -7,6 +7,8 @@
 #include "NiagaraSystem.h"
 #include "NiagaraScriptSourceBase.h"
 
+#include "UObject/UE5MainStreamObjectVersion.h"
+
 #include UE_INLINE_GENERATED_CPP_BY_NAME(NiagaraSimulationStageBase)
 
 const FName UNiagaraSimulationStageBase::ParticleSpawnUpdateName("ParticleSpawnUpdate");
@@ -20,6 +22,7 @@ namespace NiagaraSimulationStageLocal
 	static FNiagaraVariableBase GetDefaultParticleStateBinding() { return FNiagaraVariableBase(FNiagaraTypeDefinition::GetIntDef(), NAME_Particles_StateIndex); }
 }
 
+#if WITH_EDITORONLY_DATA
 bool UNiagaraSimulationStageBase::AppendCompileHash(FNiagaraCompileHashVisitor* InVisitor) const
 {
 #if WITH_EDITORONLY_DATA
@@ -31,6 +34,7 @@ bool UNiagaraSimulationStageBase::AppendCompileHash(FNiagaraCompileHashVisitor*
 	InVisitor->UpdatePOD(TEXT("Enabled"), bEnabled ? 1 : 0);
 	return true;
 }
+#endif
 
 #if WITH_EDITOR
 void UNiagaraSimulationStageBase::SetEnabled(bool bInEnabled)
@@ -79,16 +83,16 @@ void UNiagaraSimulationStageBase::PostEditChangeProperty(struct FPropertyChanged
 }
 #endif
 
-bool UNiagaraSimulationStageGeneric::AppendCompileHash(FNiagaraCompileHashVisitor* InVisitor) const 
+#if WITH_EDITORONLY_DATA
+bool UNiagaraSimulationStageGeneric::AppendCompileHash(FNiagaraCompileHashVisitor* InVisitor) const
 {
 	Super::AppendCompileHash(InVisitor);
 
 	InVisitor->UpdateString(TEXT("EnabledBinding"), EnabledBinding.GetDataSetBindableVariable().GetName().ToString());
-	InVisitor->UpdateString(TEXT("ElementCountXBinding"), ElementCountXBinding.GetDataSetBindableVariable().GetName().ToString());
-	InVisitor->UpdateString(TEXT("ElementCountYBinding"), ElementCountYBinding.GetDataSetBindableVariable().GetName().ToString());
-	InVisitor->UpdateString(TEXT("ElementCountZBinding"), ElementCountZBinding.GetDataSetBindableVariable().GetName().ToString());
-	InVisitor->UpdatePOD(TEXT("Iterations"), Iterations);
-	InVisitor->UpdateString(TEXT("NumIterationsBinding"), NumIterationsBinding.GetDataSetBindableVariable().GetName().ToString());
+	InVisitor->UpdateString(TEXT("ElementCountX"), ElementCountX.ToString());
+	InVisitor->UpdateString(TEXT("ElementCountY"), ElementCountY.ToString());
+	InVisitor->UpdateString(TEXT("ElementCountZ"), ElementCountZ.ToString());
+	InVisitor->UpdateString(TEXT("NumIterations"), NumIterations.ToString());
 	InVisitor->UpdatePOD(TEXT("IterationSource"), (int32)IterationSource);
 	InVisitor->UpdatePOD(TEXT("ExecuteBehavior"), (int32)ExecuteBehavior);
 	InVisitor->UpdatePOD(TEXT("bDisablePartialParticleUpdate"), bDisablePartialParticleUpdate ? 1 : 0);
@@ -101,23 +105,25 @@ bool UNiagaraSimulationStageGeneric::AppendCompileHash(FNiagaraCompileHashVisito
 	InVisitor->UpdatePOD(TEXT("DirectDispatchType"), (int32)DirectDispatchType);
 	InVisitor->UpdatePOD(TEXT("DirectDispatchElementType"), int32(DirectDispatchElementType));
 	InVisitor->UpdatePOD(TEXT("bOverrideGpuDispatchNumThreads"), bOverrideGpuDispatchNumThreads ? 1 : 0);
-	InVisitor->UpdateString(TEXT("OverrideGpuDispatchNumThreads"), FString::Printf(TEXT("%d,%d,%d"), OverrideGpuDispatchNumThreads.X, OverrideGpuDispatchNumThreads.Y, OverrideGpuDispatchNumThreads.Z));
+	InVisitor->UpdateString(TEXT("OverrideGpuDispatchNumThreadsX"), OverrideGpuDispatchNumThreadsX.ToString());
+	InVisitor->UpdateString(TEXT("OverrideGpuDispatchNumThreadsY"), OverrideGpuDispatchNumThreadsY.ToString());
+	InVisitor->UpdateString(TEXT("OverrideGpuDispatchNumThreadsZ"), OverrideGpuDispatchNumThreadsZ.ToString());
 
 	return true;
 }
 
-#if WITH_EDITORONLY_DATA
 bool UNiagaraSimulationStageGeneric::FillCompilationData(TArray<FNiagaraSimulationStageCompilationData>& CompilationSimStageData) const
 {
 	FNiagaraSimulationStageCompilationData& SimStageData = CompilationSimStageData.AddDefaulted_GetRef();
 	SimStageData.StageGuid = Script->GetUsageId();
 	SimStageData.StageName = SimulationStageName;
 	SimStageData.EnabledBinding = EnabledBinding.GetName();
-	SimStageData.ElementCountXBinding = ElementCountXBinding.GetName();
-	SimStageData.ElementCountYBinding = ElementCountYBinding.GetName();
-	SimStageData.ElementCountZBinding = ElementCountZBinding.GetName();
-	SimStageData.NumIterations = Iterations;
-	SimStageData.NumIterationsBinding = NumIterationsBinding.GetName();
+	SimStageData.ElementCount = FIntVector(ElementCountX.GetDefaultValue<int32>(), ElementCountY.GetDefaultValue<int32>(), ElementCountZ.GetDefaultValue<int32>());
+	SimStageData.ElementCountXBinding = ElementCountX.ResolvedParameter.GetName();
+	SimStageData.ElementCountYBinding = ElementCountY.ResolvedParameter.GetName();
+	SimStageData.ElementCountZBinding = ElementCountZ.ResolvedParameter.GetName();
+	SimStageData.NumIterations = NumIterations.GetDefaultValue<int32>();
+	SimStageData.NumIterationsBinding = NumIterations.ResolvedParameter.GetName();
 	SimStageData.IterationSourceType = IterationSource;
 	SimStageData.IterationDataInterface = IterationSource == ENiagaraIterationSource::DataInterface ? DataInterface.BoundVariable.GetName() : FName();
 	SimStageData.ExecuteBehavior = ExecuteBehavior;
@@ -128,7 +134,12 @@ bool UNiagaraSimulationStageGeneric::FillCompilationData(TArray<FNiagaraSimulati
 	SimStageData.DirectDispatchType = DirectDispatchType;
 	SimStageData.DirectDispatchElementType = DirectDispatchElementType;
 	SimStageData.bOverrideGpuDispatchNumThreads = bOverrideGpuDispatchNumThreads;
-	SimStageData.OverrideGpuDispatchNumThreads = OverrideGpuDispatchNumThreads;
+	SimStageData.OverrideGpuDispatchNumThreadsXBinding = OverrideGpuDispatchNumThreadsX.AliasedParameter;
+	SimStageData.OverrideGpuDispatchNumThreadsYBinding = OverrideGpuDispatchNumThreadsY.AliasedParameter;
+	SimStageData.OverrideGpuDispatchNumThreadsZBinding = OverrideGpuDispatchNumThreadsZ.AliasedParameter;
+	SimStageData.OverrideGpuDispatchNumThreads.X = OverrideGpuDispatchNumThreadsX.GetDefaultValue<int32>();
+	SimStageData.OverrideGpuDispatchNumThreads.Y = OverrideGpuDispatchNumThreadsY.GetDefaultValue<int32>();
+	SimStageData.OverrideGpuDispatchNumThreads.Z = OverrideGpuDispatchNumThreadsZ.GetDefaultValue<int32>();
 
 	if (SimStageData.bParticleIterationStateEnabled)
 	{
@@ -140,6 +151,21 @@ bool UNiagaraSimulationStageGeneric::FillCompilationData(TArray<FNiagaraSimulati
 	}
 	return true;
 }
+
+void UNiagaraSimulationStageGeneric::RenameEmitter(const FName& InOldName, const UNiagaraEmitter* InRenamedEmitter)
+{
+	FNiagaraParameterBinding::ForEachRenameEmitter(this, InRenamedEmitter->GetUniqueEmitterName());
+}
+
+void UNiagaraSimulationStageGeneric::RenameVariable(const FNiagaraVariableBase& OldVariable, const FNiagaraVariableBase& NewVariable, const FVersionedNiagaraEmitter& InEmitter)
+{
+	FNiagaraParameterBinding::ForEachRenameVariable(this, OldVariable, NewVariable, InEmitter.Emitter->GetUniqueEmitterName());
+}
+
+void UNiagaraSimulationStageGeneric::RemoveVariable(const FNiagaraVariableBase& OldVariable, const FVersionedNiagaraEmitter& InEmitter)
+{
+	FNiagaraParameterBinding::ForEachRemoveVariable(this, OldVariable, InEmitter.Emitter->GetUniqueEmitterName());
+}
 #endif //WITH_EDITORONLY_DATA
 
 void UNiagaraSimulationStageGeneric::PostInitProperties()
@@ -152,12 +178,42 @@ void UNiagaraSimulationStageGeneric::PostInitProperties()
 		EnabledBinding.Setup(GetDefaultEnabledBinding(), GetDefaultEnabledBinding(), ENiagaraRendererSourceDataMode::Emitter);
 #if WITH_EDITORONLY_DATA
 		ElementCountBinding_DEPRECATED.Setup(GetDefaultElementCountBinding(), GetDefaultElementCountBinding(), ENiagaraRendererSourceDataMode::Emitter);
+		NumIterationsBinding_DEPRECATED.Setup(GetDefaultNumIterationsBinding(), GetDefaultNumIterationsBinding(), ENiagaraRendererSourceDataMode::Emitter);
+		ElementCountXBinding_DEPRECATED.Setup(GetDefaultElementCountBinding(), GetDefaultElementCountBinding(), ENiagaraRendererSourceDataMode::Emitter);
+		ElementCountYBinding_DEPRECATED.Setup(GetDefaultElementCountBinding(), GetDefaultElementCountBinding(), ENiagaraRendererSourceDataMode::Emitter);
+		ElementCountZBinding_DEPRECATED.Setup(GetDefaultElementCountBinding(), GetDefaultElementCountBinding(), ENiagaraRendererSourceDataMode::Emitter);
 #endif
-		ElementCountXBinding.Setup(GetDefaultElementCountBinding(), GetDefaultElementCountBinding(), ENiagaraRendererSourceDataMode::Emitter);
-		ElementCountYBinding.Setup(GetDefaultElementCountBinding(), GetDefaultElementCountBinding(), ENiagaraRendererSourceDataMode::Emitter);
-		ElementCountZBinding.Setup(GetDefaultElementCountBinding(), GetDefaultElementCountBinding(), ENiagaraRendererSourceDataMode::Emitter);
-		NumIterationsBinding.Setup(GetDefaultNumIterationsBinding(), GetDefaultNumIterationsBinding(), ENiagaraRendererSourceDataMode::Emitter);
 		ParticleIterationStateBinding.Setup(GetDefaultParticleStateBinding(), GetDefaultParticleStateBinding(), ENiagaraRendererSourceDataMode::Particles);
+
+#if WITH_EDITORONLY_DATA
+		ElementCountX.SetUsage(ENiagaraParameterBindingUsage::NotParticle);
+		ElementCountX.SetAllowedTypeDefinitions({ FNiagaraTypeDefinition::GetIntDef() });
+		ElementCountX.SetDefaultParameter(FNiagaraTypeDefinition::GetIntDef(), 0);
+
+		ElementCountY.SetUsage(ENiagaraParameterBindingUsage::NotParticle);
+		ElementCountY.SetAllowedTypeDefinitions({ FNiagaraTypeDefinition::GetIntDef() });
+		ElementCountY.SetDefaultParameter(FNiagaraTypeDefinition::GetIntDef(), 0);
+
+		ElementCountZ.SetUsage(ENiagaraParameterBindingUsage::NotParticle);
+		ElementCountZ.SetAllowedTypeDefinitions({ FNiagaraTypeDefinition::GetIntDef() });
+		ElementCountZ.SetDefaultParameter(FNiagaraTypeDefinition::GetIntDef(), 0);
+
+		NumIterations.SetUsage(ENiagaraParameterBindingUsage::NotParticle);
+		NumIterations.SetAllowedTypeDefinitions({FNiagaraTypeDefinition::GetIntDef()});
+		NumIterations.SetDefaultParameter(FNiagaraTypeDefinition::GetIntDef(), 1);
+
+		OverrideGpuDispatchNumThreadsX.SetUsage(ENiagaraParameterBindingUsage::NotParticle);
+		OverrideGpuDispatchNumThreadsX.SetAllowedTypeDefinitions({ FNiagaraTypeDefinition::GetIntDef().ToStaticDef() });
+		OverrideGpuDispatchNumThreadsX.SetDefaultParameter(FNiagaraTypeDefinition::GetIntDef().ToStaticDef(), 64);
+
+		OverrideGpuDispatchNumThreadsY.SetUsage(ENiagaraParameterBindingUsage::NotParticle);
+		OverrideGpuDispatchNumThreadsY.SetAllowedTypeDefinitions({ FNiagaraTypeDefinition::GetIntDef().ToStaticDef() });
+		OverrideGpuDispatchNumThreadsY.SetDefaultParameter(FNiagaraTypeDefinition::GetIntDef().ToStaticDef(), 1);
+
+		OverrideGpuDispatchNumThreadsZ.SetUsage(ENiagaraParameterBindingUsage::NotParticle);
+		OverrideGpuDispatchNumThreadsZ.SetAllowedTypeDefinitions({ FNiagaraTypeDefinition::GetIntDef().ToStaticDef() });
+		OverrideGpuDispatchNumThreadsZ.SetDefaultParameter(FNiagaraTypeDefinition::GetIntDef().ToStaticDef(), 1);
+#endif
 	}
 }
 
@@ -166,27 +222,29 @@ void UNiagaraSimulationStageGeneric::PostLoad()
 	Super::PostLoad();
 
 #if WITH_EDITORONLY_DATA
+	const int32 UE5Version = GetLinkerCustomVersion(FUE5MainStreamObjectVersion::GUID);
+
 	// Ensure data wasn't somehow saved incorrectly
 	using namespace NiagaraSimulationStageLocal;
 	if (EnabledBinding.GetType() != GetDefaultEnabledBinding().GetType())
 	{
 		EnabledBinding.Setup(GetDefaultEnabledBinding(), GetDefaultEnabledBinding(), ENiagaraRendererSourceDataMode::Emitter);
 	}
-	if (ElementCountXBinding.GetType() != GetDefaultElementCountBinding().GetType())
+	if (ElementCountXBinding_DEPRECATED.GetType() != GetDefaultElementCountBinding().GetType())
 	{
-		ElementCountXBinding.Setup(GetDefaultElementCountBinding(), GetDefaultElementCountBinding(), ENiagaraRendererSourceDataMode::Emitter);
+		ElementCountXBinding_DEPRECATED.Setup(GetDefaultElementCountBinding(), GetDefaultElementCountBinding(), ENiagaraRendererSourceDataMode::Emitter);
 	}
-	if (ElementCountYBinding.GetType() != GetDefaultElementCountBinding().GetType())
+	if (ElementCountYBinding_DEPRECATED.GetType() != GetDefaultElementCountBinding().GetType())
 	{
-		ElementCountYBinding.Setup(GetDefaultElementCountBinding(), GetDefaultElementCountBinding(), ENiagaraRendererSourceDataMode::Emitter);
+		ElementCountYBinding_DEPRECATED.Setup(GetDefaultElementCountBinding(), GetDefaultElementCountBinding(), ENiagaraRendererSourceDataMode::Emitter);
 	}
-	if (ElementCountZBinding.GetType() != GetDefaultElementCountBinding().GetType())
+	if (ElementCountZBinding_DEPRECATED.GetType() != GetDefaultElementCountBinding().GetType())
 	{
-		ElementCountZBinding.Setup(GetDefaultElementCountBinding(), GetDefaultElementCountBinding(), ENiagaraRendererSourceDataMode::Emitter);
+		ElementCountZBinding_DEPRECATED.Setup(GetDefaultElementCountBinding(), GetDefaultElementCountBinding(), ENiagaraRendererSourceDataMode::Emitter);
 	}
-	if (NumIterationsBinding.GetType() != GetDefaultNumIterationsBinding().GetType())
+	if (NumIterationsBinding_DEPRECATED.GetType() != GetDefaultNumIterationsBinding().GetType())
 	{
-		NumIterationsBinding.Setup(GetDefaultNumIterationsBinding(), GetDefaultNumIterationsBinding(), ENiagaraRendererSourceDataMode::Emitter);
+		NumIterationsBinding_DEPRECATED.Setup(GetDefaultNumIterationsBinding(), GetDefaultNumIterationsBinding(), ENiagaraRendererSourceDataMode::Emitter);
 	}
 	if (ParticleIterationStateBinding.GetType() != GetDefaultParticleStateBinding().GetType())
 	{
@@ -198,7 +256,7 @@ void UNiagaraSimulationStageGeneric::PostLoad()
 	{
 		bOverrideGpuDispatchType_DEPRECATED = true;
 		OverrideGpuDispatchType_DEPRECATED = ENiagaraGpuDispatchType::OneD;
-		ElementCountXBinding = ElementCountBinding_DEPRECATED;
+		ElementCountXBinding_DEPRECATED = ElementCountBinding_DEPRECATED;
 	}
 
 	if (bOverrideGpuDispatchType_DEPRECATED)
@@ -206,6 +264,18 @@ void UNiagaraSimulationStageGeneric::PostLoad()
 		IterationSource = ENiagaraIterationSource::DirectSet;
 		DirectDispatchType = OverrideGpuDispatchType_DEPRECATED;
 	}
+
+	if (UE5Version < FUE5MainStreamObjectVersion::NiagaraSimStageNumIterationsBindings)
+	{
+		NumIterations.SetDefaultParameter(NumIterationsBinding_DEPRECATED.GetName(), NumIterationsBinding_DEPRECATED.GetType(), Iterations_DEPRECATED);
+		ElementCountX.SetDefaultParameter(ElementCountXBinding_DEPRECATED.GetName(), ElementCountXBinding_DEPRECATED.GetType(), 0);
+		ElementCountY.SetDefaultParameter(ElementCountYBinding_DEPRECATED.GetName(), ElementCountYBinding_DEPRECATED.GetType(), 0);
+		ElementCountZ.SetDefaultParameter(ElementCountZBinding_DEPRECATED.GetName(), ElementCountZBinding_DEPRECATED.GetType(), 0);
+
+		OverrideGpuDispatchNumThreadsX.SetDefaultValueEditorOnly(OverrideGpuDispatchNumThreads_DEPRECATED.X);
+		OverrideGpuDispatchNumThreadsY.SetDefaultValueEditorOnly(OverrideGpuDispatchNumThreads_DEPRECATED.Y);
+		OverrideGpuDispatchNumThreadsZ.SetDefaultValueEditorOnly(OverrideGpuDispatchNumThreads_DEPRECATED.Z);
+	}
 #endif
 }
 
@@ -228,11 +298,10 @@ void UNiagaraSimulationStageGeneric::PostEditChangeProperty(struct FPropertyChan
 	static const TSet<FName> ForceCompilePropertyNames =
 	{
 		GET_MEMBER_NAME_CHECKED(UNiagaraSimulationStageGeneric, EnabledBinding),
-		GET_MEMBER_NAME_CHECKED(UNiagaraSimulationStageGeneric, ElementCountXBinding),
-		GET_MEMBER_NAME_CHECKED(UNiagaraSimulationStageGeneric, ElementCountYBinding),
-		GET_MEMBER_NAME_CHECKED(UNiagaraSimulationStageGeneric, ElementCountZBinding),
-		GET_MEMBER_NAME_CHECKED(UNiagaraSimulationStageGeneric, Iterations),
-		GET_MEMBER_NAME_CHECKED(UNiagaraSimulationStageGeneric, NumIterationsBinding),
+		GET_MEMBER_NAME_CHECKED(UNiagaraSimulationStageGeneric, ElementCountX),
+		GET_MEMBER_NAME_CHECKED(UNiagaraSimulationStageGeneric, ElementCountY),
+		GET_MEMBER_NAME_CHECKED(UNiagaraSimulationStageGeneric, ElementCountZ),
+		GET_MEMBER_NAME_CHECKED(UNiagaraSimulationStageGeneric, NumIterations),
 		GET_MEMBER_NAME_CHECKED(UNiagaraSimulationStageGeneric, IterationSource),
 		GET_MEMBER_NAME_CHECKED(UNiagaraSimulationStageGeneric, ExecuteBehavior),
 		GET_MEMBER_NAME_CHECKED(UNiagaraSimulationStageGeneric, bDisablePartialParticleUpdate),
@@ -244,6 +313,9 @@ void UNiagaraSimulationStageGeneric::PostEditChangeProperty(struct FPropertyChan
 		GET_MEMBER_NAME_CHECKED(UNiagaraSimulationStageGeneric, DirectDispatchType),
 		GET_MEMBER_NAME_CHECKED(UNiagaraSimulationStageGeneric, DirectDispatchElementType),
 		GET_MEMBER_NAME_CHECKED(UNiagaraSimulationStageGeneric, bOverrideGpuDispatchNumThreads),
+		GET_MEMBER_NAME_CHECKED(UNiagaraSimulationStageGeneric, OverrideGpuDispatchNumThreadsX),
+		GET_MEMBER_NAME_CHECKED(UNiagaraSimulationStageGeneric, OverrideGpuDispatchNumThreadsY),
+		GET_MEMBER_NAME_CHECKED(UNiagaraSimulationStageGeneric, OverrideGpuDispatchNumThreadsZ),
 	};
 
 	static const TSet<FName> ForceCompileMemberPropertyNamess =
@@ -251,16 +323,7 @@ void UNiagaraSimulationStageGeneric::PostEditChangeProperty(struct FPropertyChan
 		GET_MEMBER_NAME_CHECKED(UNiagaraSimulationStageGeneric, ParticleIterationStateRange),
 	};
 
-	bool bNeedsRecompile = ForceCompilePropertyNames.Contains(PropertyName) || ForceCompileMemberPropertyNamess.Contains(MemberPropertyName);
-
-	if (MemberPropertyName == GET_MEMBER_NAME_CHECKED(UNiagaraSimulationStageGeneric, OverrideGpuDispatchNumThreads))
-	{
-		OverrideGpuDispatchNumThreads.X = FMath::Max(OverrideGpuDispatchNumThreads.X, 1);
-		OverrideGpuDispatchNumThreads.Y = FMath::Max(OverrideGpuDispatchNumThreads.Y, 1);
-		OverrideGpuDispatchNumThreads.Z = FMath::Max(OverrideGpuDispatchNumThreads.Z, 1);
-		bNeedsRecompile = true;
-	}
-
+	const bool bNeedsRecompile = ForceCompilePropertyNames.Contains(PropertyName) || ForceCompileMemberPropertyNamess.Contains(MemberPropertyName);
 	if (bNeedsRecompile)
 	{
 		RequestRecompile();
diff --git a/Engine/Plugins/FX/Niagara/Source/Niagara/Private/NiagaraVolumeRendererProperties.cpp b/Engine/Plugins/FX/Niagara/Source/Niagara/Private/NiagaraVolumeRendererProperties.cpp
--- a/Engine/Plugins/FX/Niagara/Source/Niagara/Private/NiagaraVolumeRendererProperties.cpp
+++ b/Engine/Plugins/FX/Niagara/Source/Niagara/Private/NiagaraVolumeRendererProperties.cpp
@@ -232,7 +232,7 @@ void UNiagaraVolumeRendererProperties::GetUsedMaterials(const FNiagaraEmitterIns
 	UMaterialInterface* MaterialInterface = nullptr;
 	if (InEmitter != nullptr)
 	{
-		MaterialInterface = Cast<UMaterialInterface>(InEmitter->FindBinding(MaterialParameterBinding.Parameter));
+		MaterialInterface = Cast<UMaterialInterface>(InEmitter->FindBinding(MaterialParameterBinding.ResolvedParameter));
 	}
 
 #if WITH_EDITORONLY_DATA
diff --git a/Engine/Plugins/FX/Niagara/Source/Niagara/Public/NiagaraParameterBinding.h b/Engine/Plugins/FX/Niagara/Source/Niagara/Public/NiagaraParameterBinding.h
--- a/Engine/Plugins/FX/Niagara/Source/Niagara/Public/NiagaraParameterBinding.h
+++ b/Engine/Plugins/FX/Niagara/Source/Niagara/Public/NiagaraParameterBinding.h
@@ -29,6 +29,8 @@ struct NIAGARA_API FNiagaraParameterBinding
 	GENERATED_USTRUCT_BODY()
 
 #if WITH_EDITORONLY_DATA
+	virtual ~FNiagaraParameterBinding() = default;
+
 	bool CanBindTo(FNiagaraTypeDefinition TypeDefinition) const;
 	bool CanBindTo(FNiagaraVariableBase InVariable, FNiagaraVariableBase& OutAliasedVariable, FStringView EmitterName) const;
 
@@ -42,6 +44,12 @@ struct NIAGARA_API FNiagaraParameterBinding
 	void SetAllowedInterfaces(TArray<UClass*> InClasses) { AllowedInterfaces = InClasses; }
 	void SetAllowedTypeDefinitions(TArray<FNiagaraTypeDefinition> InTypeDefs) { AllowedTypeDefinitions = InTypeDefs; }
 
+	void SetDefaultParameter(const FNiagaraVariable& InVariable);
+	void SetDefaultParameter(const FName& InName, const FNiagaraTypeDefinition& InTypeDef) { SetDefaultParameter(FNiagaraVariable(InTypeDef, InName)); }
+
+	const FNiagaraVariable& GetDefaultAliasedParameter() { return DefaultAliasedParameter; }
+	const FNiagaraVariable& GetDefaultResolvedParameter() { return DefaultResolvedParameter; }
+
 	bool AllowUserParameters() const { return EnumHasAllFlags(BindingUsage, ENiagaraParameterBindingUsage::User); }
 	bool AllowSystemParameters() const { return EnumHasAllFlags(BindingUsage, ENiagaraParameterBindingUsage::System); }
 	bool AllowEmitterParameters() const { return EnumHasAllFlags(BindingUsage, ENiagaraParameterBindingUsage::Emitter); }
@@ -52,18 +60,45 @@ struct NIAGARA_API FNiagaraParameterBinding
 	TConstArrayView<UClass*> GetAllowedObjects() const { return MakeArrayView(AllowedObjects); }
 	TConstArrayView<UClass*> GetAllowedInterfaces() const { return MakeArrayView(AllowedInterfaces); }
 	TConstArrayView<FNiagaraTypeDefinition> GetAllowedTypeDefinitions() const { return MakeArrayView(AllowedTypeDefinitions); }
+
+	bool IsSetoToDefault() const;
+	void SetToDefault();
+
+	FString ToString() const;
+
+	static void ForEachRenameEmitter(UObject* InObject, FStringView EmitterName);
+	static void ForEachRenameVariable(UObject* InObject, const FNiagaraVariableBase& OldVariable, const FNiagaraVariableBase& NewVariable, FStringView EmitterName);
+	static void ForEachRemoveVariable(UObject* InObject, const FNiagaraVariableBase& OldVariable, FStringView EmitterName);
+
+	virtual bool HasDefaultValueEditorOnly() const { return false; }
+	virtual TConstArrayView<uint8> GetDefaultValueEditorOnly() const { checkNoEntry(); return MakeArrayView<uint8>(nullptr, 0); }
+	virtual void SetDefaultValueEditorOnly(TConstArrayView<uint8> Memory) { checkNoEntry(); }
+	virtual void SetDefaultValueEditorOnly(const uint8* Memory) { checkNoEntry(); }
 #endif
 
-	/** Parameter binding used by the runtime fully resolved, contains a fallback value, i.e. NamedEmitter.Parameter */
+	bool operator==(const FNiagaraParameterBinding& Other) const
+	{
+		return
+		#if WITH_EDITORONLY_DATA
+			AliasedParameter == Other.AliasedParameter &&
+		#endif
+			ResolvedParameter == Other.ResolvedParameter;
+	}
+
+	/** Parameter binding used by the runtime fully resolved, i.e. NamedEmitter.Parameter */
 	UPROPERTY(EditAnywhere, Category = "Parameter Binding")
-	FNiagaraVariableBase Parameter;
+	FNiagaraVariableBase ResolvedParameter;
 
 #if WITH_EDITORONLY_DATA
 	/** Parameter binding used in the UI, i.e. Emitter.Parameter */
-	UPROPERTY(EditAnywhere, Category = "Parameter Binding")
+	UPROPERTY()
 	FNiagaraVariableBase AliasedParameter;
 
 protected:
+	/** Default parameter for the binding, can also contain the default value. */
+	FNiagaraVariable DefaultAliasedParameter;
+	FNiagaraVariable DefaultResolvedParameter;
+
 	/* Set the usage for the binding. */
 	ENiagaraParameterBindingUsage BindingUsage = ENiagaraParameterBindingUsage::NotParticle;
 
@@ -82,9 +117,42 @@ protected:
 	/** List of explicit type definitions allowed, must be an exact match. */
 	TArray<FNiagaraTypeDefinition>	AllowedTypeDefinitions;
 #endif
+};
 
-	//FORCEINLINE bool operator==(const FNiagaraParameterBinding& Other)const
-	//{
-		//return Other.Parameter == Parameter;
-	//}
-};
\ No newline at end of file
+USTRUCT()
+struct NIAGARA_API FNiagaraParameterBindingWithValue : public FNiagaraParameterBinding
+{
+	GENERATED_USTRUCT_BODY()
+
+#if WITH_EDITORONLY_DATA
+	template<typename T> void SetDefaultParameter(const FNiagaraVariableBase& InVariable, const T& InDefaultValue) { FNiagaraVariable NewVariable(InVariable); NewVariable.SetValue<T>(InDefaultValue);  FNiagaraParameterBinding::SetDefaultParameter(NewVariable); }
+	template<typename T> void SetDefaultParameter(const FNiagaraTypeDefinition& InTypeDef, const T& InDefaultValue) { SetDefaultParameter(FNiagaraVariableBase(InTypeDef, NAME_None), InDefaultValue); }
+	template<typename T> void SetDefaultParameter(const FName& InName, const FNiagaraTypeDefinition& InTypeDef, const T& InDefaultValue) { SetDefaultParameter(FNiagaraVariableBase(InTypeDef, InName), InDefaultValue); }
+
+	virtual bool HasDefaultValueEditorOnly() const override;
+	virtual TConstArrayView<uint8> GetDefaultValueEditorOnly() const override;
+	virtual void SetDefaultValueEditorOnly(TConstArrayView<uint8> Memory) override;
+	virtual void SetDefaultValueEditorOnly(const uint8* Memory) override;
+	template<typename T> void SetDefaultValueEditorOnly(const T& InDefaultValue) { check(sizeof(T) == DefaultAliasedParameter.GetType().GetSize()); SetDefaultValueEditorOnly(MakeArrayView(reinterpret_cast<const uint8*>(&InDefaultValue), sizeof(InDefaultValue))); }
+#endif
+
+	template<typename T> T GetDefaultValue() const
+	{
+		check(DefaultValue.Num() == sizeof(T));
+		T Value;
+		FMemory::Memcpy(&Value, DefaultValue.GetData(), sizeof(T));
+		return Value;
+	}
+
+	bool operator==(const FNiagaraParameterBindingWithValue& Other) const
+	{
+		return
+			FNiagaraParameterBinding::operator==(Other) &&
+			DefaultValue == Other.DefaultValue;
+	}
+
+protected:
+	/** Default value will only have contents if one is provided. */
+	UPROPERTY()
+	TArray<uint8> DefaultValue;
+};
diff --git a/Engine/Plugins/FX/Niagara/Source/Niagara/Public/NiagaraSimulationStageBase.h b/Engine/Plugins/FX/Niagara/Source/Niagara/Public/NiagaraSimulationStageBase.h
--- a/Engine/Plugins/FX/Niagara/Source/Niagara/Public/NiagaraSimulationStageBase.h
+++ b/Engine/Plugins/FX/Niagara/Source/Niagara/Public/NiagaraSimulationStageBase.h
@@ -7,6 +7,7 @@
 #include "NiagaraCommon.h"
 #include "NiagaraTypes.h"
 #include "NiagaraConstants.h"
+#include "NiagaraParameterBinding.h"
 #include "NiagaraSimulationStageCompileData.h"
 #include "NiagaraSimulationStageBase.generated.h"
 
@@ -37,8 +38,9 @@ public:
 	UPROPERTY()
 	uint32 bEnabled : 1;
 
-	virtual bool AppendCompileHash(FNiagaraCompileHashVisitor* InVisitor) const;
 #if WITH_EDITORONLY_DATA
+	virtual bool AppendCompileHash(FNiagaraCompileHashVisitor* InVisitor) const;
+
 	virtual bool FillCompilationData(TArray<FNiagaraSimulationStageCompilationData>& CompilationSimStageData) const PURE_VIRTUAL(UNiagaraSimulationStageBase::FillCompileSimStageData, return false;);
 	/** Return the FName to use in place of the default for the location in the stack context. If this would be the default, return NAME_None.*/
 	virtual FName GetStackContextReplacementName() const { return NAME_None; }
@@ -74,15 +76,10 @@ public:
 	/**
 	Number of times (or iterations) the simulation stage will execute in a row.
 	For example, setting this to 10 will mean this simulation stage runs 10 times in a row before the next stage.
-	*/
-	UPROPERTY(EditAnywhere, Category = "Simulation Stage", meta = (NoSpinbox = "true", ClampMin = 1, DisplayName = "Num Iterations"))
-	int32 Iterations = 1;
-
-	/**
-	Optional integer binding allowing scripts to control the number of iterations.
+	Can also be bound to a attribute so the simulation can dynamically decide
 	*/
 	UPROPERTY(EditAnywhere, Category = "Simulation Stage")
-	FNiagaraVariableAttributeBinding NumIterationsBinding;
+	FNiagaraParameterBindingWithValue NumIterations;
 
 	/**
 	Controls when the simulation stage should execute, only valid for data interface iteration stages
@@ -128,8 +125,17 @@ public:
 	UPROPERTY(EditAnywhere, Category = "Simulation Stage")
 	uint32 bOverrideGpuDispatchNumThreads : 1;
 
+	/** Parameter binding / constant value for Num Threads X */
 	UPROPERTY(EditAnywhere, Category = "Simulation Stage")
-	FIntVector OverrideGpuDispatchNumThreads = FIntVector(64, 1, 1);
+	FNiagaraParameterBindingWithValue OverrideGpuDispatchNumThreadsX;
+
+	/** Parameter binding / constant value for Num Threads Y */
+	UPROPERTY(EditAnywhere, Category = "Simulation Stage", meta = (EditCondition = "DirectDispatchType == ENiagaraGpuDispatchType::TwoD || DirectDispatchType == ENiagaraGpuDispatchType::ThreeD"))
+	FNiagaraParameterBindingWithValue OverrideGpuDispatchNumThreadsY;
+
+	/** Parameter binding / constant value for Num Threads Z */
+	UPROPERTY(EditAnywhere, Category = "Simulation Stage", meta = (EditCondition = "DirectDispatchType == ENiagaraGpuDispatchType::ThreeD"))
+	FNiagaraParameterBindingWithValue OverrideGpuDispatchNumThreadsZ;
 
 	/** Dimensions to use for dispatch. */
 	UPROPERTY(EditAnywhere, Category = "Simulation Stage")
@@ -145,7 +151,7 @@ public:
 	set an int to the triangle count in an emitter script and bind that as the element count.
 	*/
 	UPROPERTY(EditAnywhere, Category = "Simulation Stage", meta = (DisplayName = "ElementCount X"))
-	FNiagaraVariableAttributeBinding ElementCountXBinding;
+	FNiagaraParameterBindingWithValue ElementCountX;
 
 	/**
 	Integer binding to override the number of elements the stage will execute along Y.
@@ -153,7 +159,7 @@ public:
 	in an emitter script and bind that as the element count.
 	*/
 	UPROPERTY(EditAnywhere, Category = "Simulation Stage", meta = (DisplayName = "ElementCount Y", EditCondition = "DirectDispatchType == ENiagaraGpuDispatchType::TwoD || DirectDispatchType == ENiagaraGpuDispatchType::ThreeD"))
-	FNiagaraVariableAttributeBinding ElementCountYBinding;
+	FNiagaraParameterBindingWithValue ElementCountY;
 
 	/**
 	Integer binding to override the number of elements the stage will execute along Z.
@@ -161,13 +167,22 @@ public:
 	in an emitter script and bind that as the element count.
 	*/
 	UPROPERTY(EditAnywhere, Category = "Simulation Stage", meta = (DisplayName = "ElementCount Z", EditCondition = "DirectDispatchType == ENiagaraGpuDispatchType::ThreeD"))
-	FNiagaraVariableAttributeBinding ElementCountZBinding;
+	FNiagaraParameterBindingWithValue ElementCountZ;
 
 #if WITH_EDITORONLY_DATA
 	// Deprecated Properties
 	UPROPERTY()
 	FNiagaraVariableAttributeBinding ElementCountBinding_DEPRECATED;
 
+	UPROPERTY()
+	FNiagaraVariableAttributeBinding ElementCountXBinding_DEPRECATED;
+
+	UPROPERTY()
+	FNiagaraVariableAttributeBinding ElementCountYBinding_DEPRECATED;
+
+	UPROPERTY()
+	FNiagaraVariableAttributeBinding ElementCountZBinding_DEPRECATED;
+
 	UPROPERTY()
 	uint32 bSpawnOnly_DEPRECATED : 1;
 	
@@ -176,15 +191,27 @@ public:
 
 	UPROPERTY()
 	ENiagaraGpuDispatchType OverrideGpuDispatchType_DEPRECATED = ENiagaraGpuDispatchType::OneD;
+
+	UPROPERTY()
+	int32 Iterations_DEPRECATED = 1;
+
+	UPROPERTY()
+	FNiagaraVariableAttributeBinding NumIterationsBinding_DEPRECATED;
+
+	UPROPERTY()
+	FIntVector OverrideGpuDispatchNumThreads_DEPRECATED = FIntVector(64, 1, 1);
 #endif
 	//////////////////////////////////////////////////////////////////////////////////////////////////////////////////////////////////
 
 
 	virtual void PostInitProperties() override;
 	virtual void PostLoad() override;
-	virtual bool AppendCompileHash(FNiagaraCompileHashVisitor* InVisitor) const override;
 #if WITH_EDITORONLY_DATA
+	virtual bool AppendCompileHash(FNiagaraCompileHashVisitor* InVisitor) const override;
 	virtual bool FillCompilationData(TArray<FNiagaraSimulationStageCompilationData>& CompilationSimStageData) const override;
+	virtual void RenameEmitter(const FName& InOldName, const UNiagaraEmitter* InRenamedEmitter);
+	virtual void RenameVariable(const FNiagaraVariableBase& OldVariable, const FNiagaraVariableBase& NewVariable, const FVersionedNiagaraEmitter& InEmitter);
+	virtual void RemoveVariable(const FNiagaraVariableBase& OldVariable, const FVersionedNiagaraEmitter& InEmitter);
 #endif
 #if WITH_EDITOR
 	virtual FName GetStackContextReplacementName() const override; 
diff --git a/Engine/Plugins/FX/Niagara/Source/Niagara/Public/NiagaraSimulationStageCompileData.h b/Engine/Plugins/FX/Niagara/Source/Niagara/Public/NiagaraSimulationStageCompileData.h
--- a/Engine/Plugins/FX/Niagara/Source/Niagara/Public/NiagaraSimulationStageCompileData.h
+++ b/Engine/Plugins/FX/Niagara/Source/Niagara/Public/NiagaraSimulationStageCompileData.h
@@ -2,6 +2,7 @@
 
 #pragma once
 
+#include "NiagaraTypes.h"
 #include "NiagaraScriptBase.h"
 
 struct FNiagaraSimulationStageCompilationData
@@ -9,6 +10,7 @@ struct FNiagaraSimulationStageCompilationData
 	FGuid                           StageGuid;
 	FName                           StageName;
 	FName                           EnabledBinding;
+	FIntVector3						ElementCount = FIntVector::ZeroValue;
 	FName                           ElementCountXBinding;
 	FName                           ElementCountYBinding;
 	FName                           ElementCountZBinding;
@@ -27,4 +29,7 @@ struct FNiagaraSimulationStageCompilationData
 	ENiagaraDirectDispatchElementType DirectDispatchElementType = ENiagaraDirectDispatchElementType::NumThreads;
 	bool                            bOverrideGpuDispatchNumThreads = false;
 	FIntVector                      OverrideGpuDispatchNumThreads = FIntVector(1, 1, 1);
+	FNiagaraVariableBase			OverrideGpuDispatchNumThreadsXBinding;
+	FNiagaraVariableBase			OverrideGpuDispatchNumThreadsYBinding;
+	FNiagaraVariableBase			OverrideGpuDispatchNumThreadsZBinding;
 };
diff --git a/Engine/Plugins/FX/Niagara/Source/Niagara/Public/NiagaraTypes.h b/Engine/Plugins/FX/Niagara/Source/Niagara/Public/NiagaraTypes.h
--- a/Engine/Plugins/FX/Niagara/Source/Niagara/Public/NiagaraTypes.h
+++ b/Engine/Plugins/FX/Niagara/Source/Niagara/Public/NiagaraTypes.h
@@ -1021,7 +1021,8 @@ public:
 	bool IsIndexWildcard() const { return ClassStructOrEnum == FNiagaraTypeDefinition::GetWildcardStruct(); }
 
 	FNiagaraTypeDefinition ToStaticDef() const;
-	
+	FNiagaraTypeDefinition RemoveStaticDef() const;
+
 	int32 GetSize() const
 	{
 		if (Size == INDEX_NONE)
diff --git a/Engine/Plugins/FX/Niagara/Source/NiagaraEditor/Private/Commandlets/NiagaraSystemAuditCommandlet.cpp b/Engine/Plugins/FX/Niagara/Source/NiagaraEditor/Private/Commandlets/NiagaraSystemAuditCommandlet.cpp
--- a/Engine/Plugins/FX/Niagara/Source/NiagaraEditor/Private/Commandlets/NiagaraSystemAuditCommandlet.cpp
+++ b/Engine/Plugins/FX/Niagara/Source/NiagaraEditor/Private/Commandlets/NiagaraSystemAuditCommandlet.cpp
@@ -345,7 +345,7 @@ bool UNiagaraSystemAuditCommandlet::ProcessNiagaraSystems()
 						UNiagaraSimulationStageGeneric* SimStageGeneric = Cast<UNiagaraSimulationStageGeneric>(SimStage);
 						if (SimStage->bEnabled && SimStageGeneric)
 						{
-							TotalIterations += SimStageGeneric->Iterations;
+							TotalIterations += SimStageGeneric->NumIterations.GetDefaultValue<int32>();
 						}
 					}
 					EmittersWithSimulationStages.Appendf(TEXT("%s(%d Stages %d Iterations) "), EmitterData->GetDebugSimName(), SimulationStages.Num(), TotalIterations);
diff --git a/Engine/Plugins/FX/Niagara/Source/NiagaraEditor/Private/Customizations/NiagaraParameterBindingCustomization.cpp b/Engine/Plugins/FX/Niagara/Source/NiagaraEditor/Private/Customizations/NiagaraParameterBindingCustomization.cpp
--- a/Engine/Plugins/FX/Niagara/Source/NiagaraEditor/Private/Customizations/NiagaraParameterBindingCustomization.cpp
+++ b/Engine/Plugins/FX/Niagara/Source/NiagaraEditor/Private/Customizations/NiagaraParameterBindingCustomization.cpp
@@ -14,6 +14,8 @@
 #include "Framework/Application/SlateApplication.h"
 #include "Widgets/Input/SComboButton.h"
 
+#include "INiagaraEditorTypeUtilities.h"
+#include "NiagaraEditorModule.h"
 #include "NiagaraEmitter.h"
 #include "NiagaraNodeParameterMapBase.h"
 #include "NiagaraParameterMapHistory.h"
@@ -22,6 +24,7 @@
 #include "NiagaraSimulationStageBase.h"
 #include "NiagaraSystem.h"
 #include "NiagaraParameterBinding.h"
+#include "SNiagaraParameterEditor.h"
 #include "Customizations/NiagaraTypeCustomizations.h"
 #include "Widgets/Input/SButton.h"
 #include "Widgets/SNiagaraParameterName.h"
@@ -31,27 +34,46 @@
 #define LOCTEXT_NAMESPACE "NiagaraParameterBindingCustomization"
 
 FNiagaraParameterBindingAction::FNiagaraParameterBindingAction()
-	: FEdGraphSchemaAction(
-		FText(),
+	: FEdGraphSchemaAction(FText(), FText(), FText(), 0)
+{
+}
+
+FNiagaraParameterBindingAction::FNiagaraParameterBindingAction(FText InName, FText InTooltip)
+	: FEdGraphSchemaAction(FText(), InName, InTooltip, 0)
+{
+}
+
+TSharedRef<FNiagaraParameterBindingAction> FNiagaraParameterBindingAction::MakeNone()
+{
+	TSharedRef<FNiagaraParameterBindingAction> Action = MakeShared<FNiagaraParameterBindingAction>(
 		LOCTEXT("None", "None"),
-		LOCTEXT("NoneTooltip", "Not bound to any variable will use the default value"),
-		1,
-		FText()
-	)
+		LOCTEXT("NoneTooltip", "Not bound to any parameter.")
+	);
+	Action->Grouping = 1;
+	return Action;
+}
+
+TSharedRef<FNiagaraParameterBindingAction> FNiagaraParameterBindingAction::MakeConstant(const FNiagaraTypeDefinition& TypeDef, const uint8* ValueData)
 {
+	TSharedRef<FNiagaraParameterBindingAction> Action = MakeShared<FNiagaraParameterBindingAction>(
+		FText::Format(LOCTEXT("ConstantFormat", "Constant Value \"{0}\""), FText::FromString(TypeDef.ToString(ValueData))),
+		LOCTEXT("ConstantTooltip", "Use user defined constant value, not bound to any parameter.")
+	);
+	Action->Grouping = 1;
+	return Action;
 }
 
-FNiagaraParameterBindingAction::FNiagaraParameterBindingAction(FNiagaraVariableBase InAliasedParameter, FNiagaraVariableBase InParameter)
-	: FEdGraphSchemaAction(
-		FText(),
-		FText::FromName(InAliasedParameter.GetName()),
-		FText::Format(LOCTEXT("VariableBindActionTooltip", "Bind to variable \"{0}\""), FText::FromName(InAliasedParameter.GetName())),
-		0,
-		FText()
-	)
-	, AliasedParameter(InAliasedParameter)
-	, Parameter(InParameter)
+TSharedRef<FNiagaraParameterBindingAction> FNiagaraParameterBindingAction::MakeParameter(FNiagaraVariableBase InAliasedParameter, FNiagaraVariableBase InResolvedParameter)
 {
+	TSharedRef<FNiagaraParameterBindingAction> Action = MakeShared<FNiagaraParameterBindingAction>(
+		FText::FromName(InAliasedParameter.GetName()),
+		FText::Format(LOCTEXT("VariableBindActionTooltip", "Bind to variable \"{0}\""), FText::FromName(InAliasedParameter.GetName()))
+	);
+	Action->Grouping = 0;
+	Action->bIsValidParameter = true;
+	Action->AliasedParameter = InAliasedParameter;
+	Action->ResolvedParameter = InResolvedParameter;
+	return Action;
 }
 
 TSharedRef<IPropertyTypeCustomization> FNiagaraParameterBindingCustomization::MakeInstance()
@@ -66,10 +88,11 @@ void FNiagaraParameterBindingCustomization::CustomizeHeader(TSharedRef<IProperty
 	SystemWeakPtr.Reset();
 
 	// We only handle single object
+	UObject* OwnerObject = nullptr;
 	{
 		TArray<UObject*> OuterObjects;
 		PropertyHandle->GetOuterObjects(OuterObjects);
-		UObject* OwnerObject = OuterObjects.Num() == 1 ? OuterObjects[0] : nullptr;
+		OwnerObject = OuterObjects.Num() == 1 ? OuterObjects[0] : nullptr;
 		if (OwnerObject == nullptr)
 		{
 			HeaderRow
@@ -87,12 +110,60 @@ void FNiagaraParameterBindingCustomization::CustomizeHeader(TSharedRef<IProperty
 
 			return;
 		}
+	}
+
+	InPropertyHandle->SetOnPropertyResetToDefault(FSimpleDelegate::CreateLambda([]() { }));
+	PropertyHandle->MarkResetToDefaultCustomized(true);
+
+	OwnerWeakPtr = OwnerObject;
+	EmitterWeakPtr = OwnerObject->GetTypedOuter<UNiagaraEmitter>();
+	SystemWeakPtr = OwnerObject->GetTypedOuter<UNiagaraSystem>();
+
+	TSharedPtr<SHorizontalBox> ParameterPanel = SNew(SHorizontalBox);
 
-		OwnerWeakPtr = OwnerObject;
-		EmitterWeakPtr = OwnerObject->GetTypedOuter<UNiagaraEmitter>();
-		SystemWeakPtr = OwnerObject->GetTypedOuter<UNiagaraSystem>();
+	FNiagaraParameterBinding* ParameterBinding = GetParameterBinding();
+	if (ParameterBinding && ParameterBinding->HasDefaultValueEditorOnly())
+	{
+		const FNiagaraVariableBase& DefaultParameter = ParameterBinding->GetDefaultAliasedParameter();
+		TConstArrayView<uint8> DefaultValue = ParameterBinding->GetDefaultValueEditorOnly();
+		DefaultValueStructOnScope = MakeShared<FStructOnScope>(DefaultParameter.GetType().GetStruct());
+		FMemory::Memcpy(DefaultValueStructOnScope->GetStructMemory(), DefaultValue.GetData(), DefaultValue.Num());
+
+		FNiagaraEditorModule& NiagaraEditorModule = FModuleManager::GetModuleChecked<FNiagaraEditorModule>("NiagaraEditor");
+		TSharedPtr<INiagaraEditorTypeUtilities, ESPMode::ThreadSafe> TypeEditorUtilities = NiagaraEditorModule.GetTypeUtilities(DefaultParameter.GetType());
+		if (TypeEditorUtilities.IsValid())
+		{
+			DefaultValueParameterEditor = TypeEditorUtilities->CreateParameterEditor(ParameterBinding->ResolvedParameter.GetType());
+			DefaultValueParameterEditor->UpdateInternalValueFromStruct(DefaultValueStructOnScope.ToSharedRef());
+			DefaultValueParameterEditor->SetOnValueChanged(SNiagaraParameterEditor::FOnValueChange::CreateSP(this, &FNiagaraParameterBindingCustomization::OnValueChanged));
+
+			ParameterPanel->AddSlot()
+			.AutoWidth()
+			.Padding(0, 0, 0, 3)
+			[
+				SNew(SBox)
+				.HAlign(DefaultValueParameterEditor->GetHorizontalAlignment())
+				.VAlign(DefaultValueParameterEditor->GetVerticalAlignment())
+				.IsEnabled(this, &FNiagaraParameterBindingCustomization::IsConstantEnabled)
+				.Visibility(this, &FNiagaraParameterBindingCustomization::IsConstantVisibile)
+				[
+					DefaultValueParameterEditor.ToSharedRef()
+				]
+			];
+		}
 	}
 
+	ParameterPanel->AddSlot()
+	.AutoWidth()
+	[
+		SNew(SNiagaraParameterName)
+		.ParameterName(this, &FNiagaraParameterBindingCustomization::GetVariableName)
+		.IsReadOnly(true)
+		.IsEnabled(this, &FNiagaraParameterBindingCustomization::IsBindingEnabled)
+		.Visibility(this, &FNiagaraParameterBindingCustomization::IsBindingVisibile)
+		.ToolTipText(this, &FNiagaraParameterBindingCustomization::GetTooltipText)
+	];
+
 	HeaderRow
 	.NameContent()
 	[
@@ -102,18 +173,39 @@ void FNiagaraParameterBindingCustomization::CustomizeHeader(TSharedRef<IProperty
 	[
 		SNew(SHorizontalBox)
 		+ SHorizontalBox::Slot()
+		[
+			ParameterPanel.ToSharedRef()
+		]
+		+ SHorizontalBox::Slot()
 		.FillWidth(1.0f)
-		.Padding(0.0f, 0.0f, 4.0f, 0.0f)
+		.HAlign(HAlign_Right)
+		.Padding(3, 0, 0, 0)
 		[
 			SNew(SComboButton)					
+			.ButtonStyle(FAppStyle::Get(), "HoverHintOnly")
+			.ForegroundColor(FSlateColor::UseForeground())
 			.OnGetMenuContent(this, &FNiagaraParameterBindingCustomization::OnGetMenuContent)
-			.ContentPadding(1)
+			.ContentPadding(2)
+			.MenuPlacement(MenuPlacement_BelowRightAnchor)
 			.ToolTipText(this, &FNiagaraParameterBindingCustomization::GetTooltipText)
-			.ButtonContent()
+		]
+		+ SHorizontalBox::Slot()
+		.AutoWidth()
+		.VAlign(VAlign_Center)
+		.Padding(2, 1)
+		[
+			SNew(SButton)
+			.IsFocusable(false)
+			.ToolTipText(LOCTEXT("ResetToDefaultToolTip", "Reset to Default"))
+			.ButtonStyle(FAppStyle::Get(), "NoBorder")
+			.ContentPadding(0)
+			.Visibility(this, &FNiagaraParameterBindingCustomization::IsResetToDefaultsVisible)
+			.OnClicked(this, &FNiagaraParameterBindingCustomization::OnResetToDefaultsClicked)
+			.Content()
 			[
-				SNew(SNiagaraParameterName)
-				.ParameterName(this, &FNiagaraParameterBindingCustomization::GetVariableName)
-				.IsReadOnly(true)
+				SNew(SImage)
+				.Image(FAppStyle::GetBrush("PropertyWindow.DiffersFromDefault"))
+				.ColorAndOpacity(FSlateColor::UseForeground())
 			]
 		]
 	];
@@ -156,21 +248,18 @@ TSharedRef<SWidget> FNiagaraParameterBindingCustomization::OnGetMenuContent() co
 		];
 }
 
-void FNiagaraParameterBindingCustomization::CollectAllActions(FGraphActionListBuilderBase& OutAllActions) const
+bool FNiagaraParameterBindingCustomization::ForEachBindableVariable(TFunction<bool(FNiagaraVariableBase, FNiagaraVariableBase)> Delegate) const
 {
-	// Add the unbound action this is where we will fallback to the value only
-	OutAllActions.AddAction(MakeShared<FNiagaraParameterBindingAction>());
-
 	// Add bindings from system
 	FNiagaraParameterBinding* ParameterBinding = GetParameterBinding();
 	UNiagaraSystem* NiagaraSystem = SystemWeakPtr.Get();
 	UNiagaraEmitter* NiagaraEmitter = EmitterWeakPtr.Get();
 	if (ParameterBinding == nullptr || NiagaraSystem == nullptr)
 	{
-		return;
+		return true;
 	}
 
-	TArray<FNiagaraParameterMapHistory> Histories;
+	TArray<FNiagaraParameterMapHistory, TInlineAllocator<2>> Histories;
 	if (UNiagaraScriptSource* Source = Cast<UNiagaraScriptSource>(NiagaraSystem->GetSystemUpdateScript()->GetLatestSource()))
 	{
 		Histories.Append(UNiagaraNodeParameterMapBase::GetParameterMaps(Source->NodeGraph));
@@ -198,7 +287,7 @@ void FNiagaraParameterBindingCustomization::CollectAllActions(FGraphActionListBu
 		}
 	}
 
-	TArray<TPair<FNiagaraVariableBase, FNiagaraVariableBase>> AllowedVariables;
+	TArray<FNiagaraVariable> VisitedVariables;
 
 	const FStringView NiagaraEmitterName = NiagaraEmitter ? NiagaraEmitter->GetUniqueEmitterName() : FStringView();
 	for (const FNiagaraParameterMapHistory& History : Histories)
@@ -206,10 +295,13 @@ void FNiagaraParameterBindingCustomization::CollectAllActions(FGraphActionListBu
 		for (const FNiagaraVariable& Variable : History.Variables)
 		{
 			FNiagaraVariableBase AliasedParameter = Variable;
-			if (ParameterBinding->CanBindTo(Variable, AliasedParameter, NiagaraEmitterName))
+			if (ParameterBinding->CanBindTo(Variable, AliasedParameter, NiagaraEmitterName) && !VisitedVariables.Contains(Variable))
 			{
-				AllowedVariables.AddUnique({ AliasedParameter, Variable });
-				continue;
+				if (Delegate(AliasedParameter, Variable) == false)
+				{
+					return false;
+				}
+				VisitedVariables.Add(Variable);
 			}
 		}
 	}
@@ -218,17 +310,51 @@ void FNiagaraParameterBindingCustomization::CollectAllActions(FGraphActionListBu
 	{
 		for (const FNiagaraVariableBase& Variable : NiagaraSystem->GetExposedParameters().ReadParameterVariables())
 		{
-			if (ParameterBinding->CanBindTo(Variable.GetType()))
+			if (ParameterBinding->CanBindTo(Variable.GetType()) && !VisitedVariables.Contains(Variable))
 			{
-				AllowedVariables.AddUnique({ Variable, Variable });
+				if (Delegate(Variable, Variable) == false)
+				{
+					return false;
+				}
+				VisitedVariables.Add(Variable);
 			}
 		}
 	}
+	return true;
+}
+
+
+void FNiagaraParameterBindingCustomization::CollectAllActions(FGraphActionListBuilderBase& OutAllActions) const
+{
+	// Add default action
+	bool bHasConstantValue = false;
+	if (FNiagaraParameterBinding* ParameterBinding = GetParameterBinding())
+	{
+		UNiagaraSystem* NiagaraSystem = SystemWeakPtr.Get();
+		UNiagaraEmitter* NiagaraEmitter = EmitterWeakPtr.Get();
 
-	for (const TPair<FNiagaraVariableBase, FNiagaraVariableBase>& AllowedVariable : AllowedVariables)
+		if (ParameterBinding->HasDefaultValueEditorOnly())
+		{
+			OutAllActions.AddAction(FNiagaraParameterBindingAction::MakeConstant(ParameterBinding->GetDefaultAliasedParameter().GetType(), ParameterBinding->GetDefaultValueEditorOnly().GetData()));
+		}
+		else
+		{
+			OutAllActions.AddAction(FNiagaraParameterBindingAction::MakeNone());
+		}
+	}
+	else
 	{
-		OutAllActions.AddAction(MakeShared<FNiagaraParameterBindingAction>(AllowedVariable.Key, AllowedVariable.Value));
+		OutAllActions.AddAction(FNiagaraParameterBindingAction::MakeNone());
 	}
+
+	// Add each variable
+	ForEachBindableVariable(
+		[&OutAllActions](FNiagaraVariableBase AliasedVariable, FNiagaraVariableBase ResolvedVariable)
+		{
+			OutAllActions.AddAction(FNiagaraParameterBindingAction::MakeParameter(AliasedVariable, ResolvedVariable));
+			return true;
+		}
+	);
 }
 
 TSharedRef<SWidget> FNiagaraParameterBindingCustomization::OnCreateWidgetForAction(struct FCreateWidgetForActionData* const InCreateData) const
@@ -236,16 +362,31 @@ TSharedRef<SWidget> FNiagaraParameterBindingCustomization::OnCreateWidgetForActi
 	const FNiagaraParameterBindingAction* BindingAction = static_cast<const FNiagaraParameterBindingAction*>(InCreateData->Action.Get());
 	check(BindingAction);
 
-	return
-		SNew(SVerticalBox)
-		+ SVerticalBox::Slot()
-		.AutoHeight()
-		[
-			SNew(SNiagaraParameterName)
-			.ParameterName(BindingAction->AliasedParameter.GetName())
-			.IsReadOnly(true)
-			.ToolTipText(InCreateData->Action->GetTooltipDescription())
-		];
+	if (BindingAction->bIsValidParameter)
+	{
+		return
+			SNew(SVerticalBox)
+			+ SVerticalBox::Slot()
+			.AutoHeight()
+			[
+				SNew(SNiagaraParameterName)
+				.ParameterName(BindingAction->AliasedParameter.GetName())
+				.IsReadOnly(true)
+				.ToolTipText(InCreateData->Action->GetTooltipDescription())
+			];
+	}
+	else
+	{
+		return
+			SNew(SVerticalBox)
+			+ SVerticalBox::Slot()
+			.AutoHeight()
+			[
+				SNew(STextBlock)
+				.Text(InCreateData->Action->GetMenuDescription())
+				.ToolTipText(InCreateData->Action->GetTooltipDescription())
+			];
+	}
 }
 
 void FNiagaraParameterBindingCustomization::OnActionSelected(const TArray<TSharedPtr<FEdGraphSchemaAction>>& SelectedActions, ESelectInfo::Type InSelectionType) const
@@ -261,7 +402,7 @@ void FNiagaraParameterBindingCustomization::OnActionSelected(const TArray<TShare
 		if (EdAction.IsValid())
 		{
 			FNiagaraParameterBindingAction* BindAction = static_cast<FNiagaraParameterBindingAction*>(EdAction.Get());
-			FScopedTransaction Transaction(FText::Format(LOCTEXT("ChangeParameterBinding", " Change Parameter Binding to \"{0}\" "), FText::FromName(BindAction->Parameter.GetName())));
+			FScopedTransaction Transaction(FText::Format(LOCTEXT("ChangeParameterBinding", " Change Parameter Binding to \"{0}\" "), FText::FromName(BindAction->ResolvedParameter.GetName())));
 			FSlateApplication::Get().DismissAllMenus();
 
 			TArray<UObject*> Objects;
@@ -272,8 +413,16 @@ void FNiagaraParameterBindingCustomization::OnActionSelected(const TArray<TShare
 			}
 			PropertyHandle->NotifyPreChange();
 
-			ParameterBinding->AliasedParameter = BindAction->AliasedParameter;
-			ParameterBinding->Parameter = BindAction->Parameter;
+			if (BindAction->bIsValidParameter)
+			{
+				ParameterBinding->AliasedParameter = BindAction->AliasedParameter;
+				ParameterBinding->ResolvedParameter = BindAction->ResolvedParameter;
+			}
+			else
+			{
+				ParameterBinding->AliasedParameter = FNiagaraVariableBase();
+				ParameterBinding->ResolvedParameter = FNiagaraVariableBase();
+			}
 
 			PropertyHandle->NotifyPostChange(EPropertyChangeType::ValueSet);
 			PropertyHandle->NotifyFinishedChangingProperties();
@@ -281,22 +430,87 @@ void FNiagaraParameterBindingCustomization::OnActionSelected(const TArray<TShare
 	}
 }
 
+void FNiagaraParameterBindingCustomization::OnValueChanged() const
+{
+	FNiagaraParameterBinding* ParameterBinding = GetParameterBinding();
+	if (ParameterBinding == nullptr || !PropertyHandle.IsValid())
+	{
+		return;
+	}
+
+	FScopedTransaction Transaction(LOCTEXT("ChangeParameterBinding", " Change Parameter Value"));
+
+	TArray<UObject*> Objects;
+	PropertyHandle->GetOuterObjects(Objects);
+	for (UObject* Obj : Objects)
+	{
+		Obj->Modify();
+	}
+	PropertyHandle->NotifyPreChange();
+
+	DefaultValueParameterEditor->UpdateStructFromInternalValue(DefaultValueStructOnScope.ToSharedRef());
+	ParameterBinding->SetDefaultValueEditorOnly(DefaultValueStructOnScope->GetStructMemory());
+
+	PropertyHandle->NotifyPostChange(EPropertyChangeType::ValueSet);
+	PropertyHandle->NotifyFinishedChangingProperties();
+}
+
+EVisibility FNiagaraParameterBindingCustomization::IsBindingVisibile() const
+{
+	FNiagaraParameterBinding* ParameterBinding = GetParameterBinding();
+	if (ParameterBinding && ParameterBinding->HasDefaultValueEditorOnly())
+	{
+		return ParameterBinding->ResolvedParameter.GetName().IsNone() ? EVisibility::Collapsed : EVisibility::Visible;
+	}
+	return EVisibility::Visible;
+}
+
+EVisibility FNiagaraParameterBindingCustomization::IsConstantVisibile() const
+{
+	return IsBindingVisibile() == EVisibility::Visible ? EVisibility::Collapsed : EVisibility::Visible;
+}
+
+bool FNiagaraParameterBindingCustomization::IsBindingEnabled() const
+{
+	FNiagaraParameterBinding* ParameterBinding = GetParameterBinding();
+	if (ParameterBinding && ParameterBinding->HasDefaultValueEditorOnly())
+	{
+		return ParameterBinding->ResolvedParameter.GetName().IsNone() == false;
+	}
+	return true;
+}
+
+bool FNiagaraParameterBindingCustomization::IsConstantEnabled() const
+{
+	return !IsBindingEnabled();
+}
+
 FText FNiagaraParameterBindingCustomization::GetTooltipText() const
 {
 	if ( FNiagaraParameterBinding* ParameterBinding = GetParameterBinding() )
 	{
-		const FName BoundParameterName = ParameterBinding->AliasedParameter.GetName();
-		const FText BoundText =
-			BoundParameterName.IsNone() ?
-			LOCTEXT("BindingTooltip_UnboundParameter", "Not bound to any parameter.") :
-			FText::Format(LOCTEXT("BindingTooltip_BoundParameter", "Bound to variable \"{0}\""), FText::FromName(BoundParameterName));
-
-		const FText ValueText =
-			//ParameterBinding->Parameter.IsValid() && ParameterBinding->Parameter.GetAllocatedSizeInBytes() > 0 ?
-			//FText::Format(LOCTEXT("BindingTooltip_Value", "Default value \"{0}\""), FText::FromString(ParameterBinding->Parameter.ToString())) :
-			LOCTEXT("BindingTooltip_NoValue", "No default value");
-
-		return FText::Format(LOCTEXT("BindingTooltip", "{0}\n{1}"), BoundText, ValueText);
+		if (ParameterBinding->AliasedParameter.GetName().IsNone())
+		{
+			if (ParameterBinding->HasDefaultValueEditorOnly())
+			{
+				const FNiagaraTypeDefinition ConstantType = ParameterBinding->GetDefaultAliasedParameter().GetType();
+				return FText::Format(LOCTEXT("BindingTooltip_BoundConstant", "Bound to constant value \"{0}\" type \"{1}\""),
+					FText::FromString(ConstantType.ToString(ParameterBinding->GetDefaultValueEditorOnly().GetData())),
+					FText::FromString(ConstantType.GetName())
+				);
+			}
+			else
+			{
+				return LOCTEXT("BindingTooltip_UnboundParameter", "Not bound to any parameter.");
+			}
+		}
+		else
+		{
+			return FText::Format(LOCTEXT("BindingTooltip_BoundParameter", "Bound to parameter \"{0}\" type \"{1}\""),
+				FText::FromName(ParameterBinding->AliasedParameter.GetName()),
+				FText::FromString(ParameterBinding->AliasedParameter.GetType().GetName())
+			);
+		}
 	}
 	return FText::FromString(TEXT("Missing"));
 }
@@ -307,9 +521,27 @@ FName FNiagaraParameterBindingCustomization::GetVariableName() const
 	return ParameterBinding ? ParameterBinding->AliasedParameter.GetName() : FName();
 }
 
-void FNiagaraParameterBindingCustomization::ResetToDefault()
+EVisibility FNiagaraParameterBindingCustomization::IsResetToDefaultsVisible() const
+{
+	FNiagaraParameterBinding* ParameterBinding = GetParameterBinding();
+	return ParameterBinding && ParameterBinding->IsSetoToDefault() ? EVisibility::Hidden : EVisibility::Visible;
+}
+
+FReply FNiagaraParameterBindingCustomization::OnResetToDefaultsClicked()
 {
+	if (FNiagaraParameterBinding* ParameterBinding = GetParameterBinding())
+	{
+		ParameterBinding->SetToDefault();
+
+		if (DefaultValueStructOnScope.IsValid())
+		{
+			TConstArrayView<uint8> DefaultValue = ParameterBinding->GetDefaultValueEditorOnly();
+			FMemory::Memcpy(DefaultValueStructOnScope->GetStructMemory(), DefaultValue.GetData(), DefaultValue.Num());
+			DefaultValueParameterEditor->UpdateInternalValueFromStruct(DefaultValueStructOnScope.ToSharedRef());
+		}
+	}
 
+	return FReply::Handled();
 }
 
 #undef LOCTEXT_NAMESPACE
diff --git a/Engine/Plugins/FX/Niagara/Source/NiagaraEditor/Private/Customizations/NiagaraParameterBindingCustomization.h b/Engine/Plugins/FX/Niagara/Source/NiagaraEditor/Private/Customizations/NiagaraParameterBindingCustomization.h
--- a/Engine/Plugins/FX/Niagara/Source/NiagaraEditor/Private/Customizations/NiagaraParameterBindingCustomization.h
+++ b/Engine/Plugins/FX/Niagara/Source/NiagaraEditor/Private/Customizations/NiagaraParameterBindingCustomization.h
@@ -5,6 +5,7 @@
 #include "EdGraph/EdGraphSchema.h"
 #include "CoreMinimal.h"
 #include "IPropertyTypeCustomization.h"
+#include "Layout/Visibility.h"
 
 #include "NiagaraCommon.h"
 #include "NiagaraParameterBindingCustomization.generated.h"
@@ -20,6 +21,7 @@ class SWidget;
 class UNiagaraEmitter;
 class UNiagaraSystem;
 struct FNiagaraParameterBinding;
+class SNiagaraParameterEditor;
 
 USTRUCT()
 struct FNiagaraParameterBindingAction : public FEdGraphSchemaAction
@@ -27,10 +29,15 @@ struct FNiagaraParameterBindingAction : public FEdGraphSchemaAction
 	GENERATED_USTRUCT_BODY();
 
 	FNiagaraParameterBindingAction();
-	FNiagaraParameterBindingAction(FNiagaraVariableBase InAliasedParameter, FNiagaraVariableBase InParameter);
+	explicit FNiagaraParameterBindingAction(FText InName, FText InTooltip);
 
+	static TSharedRef<FNiagaraParameterBindingAction> MakeNone();
+	static TSharedRef<FNiagaraParameterBindingAction> MakeConstant(const FNiagaraTypeDefinition& TypeDef, const uint8* ValueData);
+	static TSharedRef<FNiagaraParameterBindingAction> MakeParameter(FNiagaraVariableBase InAliasedParameter, FNiagaraVariableBase InResolvedParameter);
+
+	bool bIsValidParameter = false;
 	FNiagaraVariableBase AliasedParameter;
-	FNiagaraVariableBase Parameter;
+	FNiagaraVariableBase ResolvedParameter;
 };
 
 class FNiagaraParameterBindingCustomization : public IPropertyTypeCustomization
@@ -48,18 +55,30 @@ public:
 	FNiagaraParameterBinding* GetDefaultParameterBinding() const;
 
 	TSharedRef<SWidget> OnGetMenuContent() const;
+	bool ForEachBindableVariable(TFunction<bool(FNiagaraVariableBase, FNiagaraVariableBase)> Delegate) const;
 	void CollectAllActions(FGraphActionListBuilderBase& OutAllActions) const;
 	TSharedRef<SWidget> OnCreateWidgetForAction(struct FCreateWidgetForActionData* const InCreateData) const;
 	void OnActionSelected(const TArray<TSharedPtr<FEdGraphSchemaAction>>& SelectedActions, ESelectInfo::Type InSelectionType) const;
+	void OnValueChanged() const;
+
+	bool IsBindingValid(FNiagaraParameterBinding* ParameterBinding) const;
+	EVisibility IsBindingVisibile() const;
+	EVisibility IsConstantVisibile() const;
+	bool IsBindingEnabled() const;
+	bool IsConstantEnabled() const;
 
 	FText GetTooltipText() const;
 	FName GetVariableName() const;
 
-	void ResetToDefault();
+	EVisibility IsResetToDefaultsVisible() const;
+	FReply OnResetToDefaultsClicked();
 
 protected:
 	TSharedPtr<IPropertyHandle>		PropertyHandle;
 	TWeakObjectPtr<UObject>			OwnerWeakPtr;
 	TWeakObjectPtr<UNiagaraEmitter>	EmitterWeakPtr;
 	TWeakObjectPtr<UNiagaraSystem>	SystemWeakPtr;
+
+	TSharedPtr<SNiagaraParameterEditor>	DefaultValueParameterEditor;
+	TSharedPtr<FStructOnScope>			DefaultValueStructOnScope;
 };
diff --git a/Engine/Plugins/FX/Niagara/Source/NiagaraEditor/Private/Customizations/NiagaraSimulationStageCustomization.cpp b/Engine/Plugins/FX/Niagara/Source/NiagaraEditor/Private/Customizations/NiagaraSimulationStageCustomization.cpp
--- a/Engine/Plugins/FX/Niagara/Source/NiagaraEditor/Private/Customizations/NiagaraSimulationStageCustomization.cpp
+++ b/Engine/Plugins/FX/Niagara/Source/NiagaraEditor/Private/Customizations/NiagaraSimulationStageCustomization.cpp
@@ -86,11 +86,8 @@ void FNiagaraSimulationStageGenericCustomization::CustomizeDetails(class IDetail
 
 	// Main Category
 	SimStageCategory.AddProperty(DetailBuilder.GetProperty(GET_MEMBER_NAME_CHECKED(UNiagaraSimulationStageGeneric, EnabledBinding)));
-
 	SimStageCategory.AddProperty(DetailBuilder.GetProperty(GET_MEMBER_NAME_CHECKED(UNiagaraSimulationStageGeneric, IterationSource)));
-	SimStageCategory.AddProperty(DetailBuilder.GetProperty(GET_MEMBER_NAME_CHECKED(UNiagaraSimulationStageGeneric, Iterations)));
-	SimStageCategory.AddProperty(DetailBuilder.GetProperty(GET_MEMBER_NAME_CHECKED(UNiagaraSimulationStageGeneric, NumIterationsBinding)));
-
+	SimStageCategory.AddProperty(DetailBuilder.GetProperty(GET_MEMBER_NAME_CHECKED(UNiagaraSimulationStageGeneric, NumIterations)));
 	SimStageCategory.AddProperty(DetailBuilder.GetProperty(GET_MEMBER_NAME_CHECKED(UNiagaraSimulationStageGeneric, ExecuteBehavior)));
 
 	// Paterial Iteration Category
@@ -118,9 +115,9 @@ void FNiagaraSimulationStageGenericCustomization::CustomizeDetails(class IDetail
 		{
 			DirectSetStageCategory.AddProperty(DetailBuilder.GetProperty(GET_MEMBER_NAME_CHECKED(UNiagaraSimulationStageGeneric, DirectDispatchType)));
 			DirectSetStageCategory.AddProperty(DetailBuilder.GetProperty(GET_MEMBER_NAME_CHECKED(UNiagaraSimulationStageGeneric, DirectDispatchElementType)));
-			DirectSetStageCategory.AddProperty(DetailBuilder.GetProperty(GET_MEMBER_NAME_CHECKED(UNiagaraSimulationStageGeneric, ElementCountXBinding)));
-			DirectSetStageCategory.AddProperty(DetailBuilder.GetProperty(GET_MEMBER_NAME_CHECKED(UNiagaraSimulationStageGeneric, ElementCountYBinding)));
-			DirectSetStageCategory.AddProperty(DetailBuilder.GetProperty(GET_MEMBER_NAME_CHECKED(UNiagaraSimulationStageGeneric, ElementCountZBinding)));
+			DirectSetStageCategory.AddProperty(DetailBuilder.GetProperty(GET_MEMBER_NAME_CHECKED(UNiagaraSimulationStageGeneric, ElementCountX)));
+			DirectSetStageCategory.AddProperty(DetailBuilder.GetProperty(GET_MEMBER_NAME_CHECKED(UNiagaraSimulationStageGeneric, ElementCountY)));
+			DirectSetStageCategory.AddProperty(DetailBuilder.GetProperty(GET_MEMBER_NAME_CHECKED(UNiagaraSimulationStageGeneric, ElementCountZ)));
 			break;
 		}
 	}
@@ -129,7 +126,9 @@ void FNiagaraSimulationStageGenericCustomization::CustomizeDetails(class IDetail
 	DispatchParametersCategory.AddProperty(DetailBuilder.GetProperty(GET_MEMBER_NAME_CHECKED(UNiagaraSimulationStageGeneric, bOverrideGpuDispatchNumThreads)));
 	if (SimStage->bOverrideGpuDispatchNumThreads)
 	{
-		DispatchParametersCategory.AddProperty(DetailBuilder.GetProperty(GET_MEMBER_NAME_CHECKED(UNiagaraSimulationStageGeneric, OverrideGpuDispatchNumThreads)));
+		DispatchParametersCategory.AddProperty(DetailBuilder.GetProperty(GET_MEMBER_NAME_CHECKED(UNiagaraSimulationStageGeneric, OverrideGpuDispatchNumThreadsX)));
+		DispatchParametersCategory.AddProperty(DetailBuilder.GetProperty(GET_MEMBER_NAME_CHECKED(UNiagaraSimulationStageGeneric, OverrideGpuDispatchNumThreadsY)));
+		DispatchParametersCategory.AddProperty(DetailBuilder.GetProperty(GET_MEMBER_NAME_CHECKED(UNiagaraSimulationStageGeneric, OverrideGpuDispatchNumThreadsZ)));
 	}
 }
 
diff --git a/Engine/Plugins/FX/Niagara/Source/NiagaraEditor/Private/NiagaraEditorModule.cpp b/Engine/Plugins/FX/Niagara/Source/NiagaraEditor/Private/NiagaraEditorModule.cpp
--- a/Engine/Plugins/FX/Niagara/Source/NiagaraEditor/Private/NiagaraEditorModule.cpp
+++ b/Engine/Plugins/FX/Niagara/Source/NiagaraEditor/Private/NiagaraEditorModule.cpp
@@ -960,6 +960,10 @@ void FNiagaraEditorModule::StartupModule()
 		FNiagaraParameterBinding::StaticStruct()->GetFName(),
 		FOnGetPropertyTypeCustomizationInstance::CreateStatic(&FNiagaraParameterBindingCustomization::MakeInstance));
 
+	PropertyModule.RegisterCustomPropertyTypeLayout(
+		FNiagaraParameterBindingWithValue::StaticStruct()->GetFName(),
+		FOnGetPropertyTypeCustomizationInstance::CreateStatic(&FNiagaraParameterBindingCustomization::MakeInstance));
+
 	PropertyModule.RegisterCustomPropertyTypeLayout(
 		FNiagaraScriptVariableBinding::StaticStruct()->GetFName(),
 		FOnGetPropertyTypeCustomizationInstance::CreateStatic(&FNiagaraScriptVariableBindingCustomization::MakeInstance));
diff --git a/Engine/Plugins/FX/Niagara/Source/NiagaraEditor/Private/NiagaraHlslTranslator.cpp b/Engine/Plugins/FX/Niagara/Source/NiagaraEditor/Private/NiagaraHlslTranslator.cpp
--- a/Engine/Plugins/FX/Niagara/Source/NiagaraEditor/Private/NiagaraHlslTranslator.cpp
+++ b/Engine/Plugins/FX/Niagara/Source/NiagaraEditor/Private/NiagaraHlslTranslator.cpp
@@ -1293,6 +1293,7 @@ const FNiagaraTranslateResults &FHlslNiagaraTranslator::Translate(const FNiagara
 					TranslationStages[Index].bCopyPreviousParams = false;
 					TranslationStages[Index].SimulationStageIndex = DestSimStageIndex;					
 					TranslationStages[Index].EnabledBinding = CompileSimStageData.EnabledBinding;
+					TranslationStages[Index].ElementCount = CompileSimStageData.ElementCount;
 					TranslationStages[Index].ElementCountXBinding = CompileSimStageData.ElementCountXBinding;
 					TranslationStages[Index].ElementCountYBinding = CompileSimStageData.ElementCountYBinding;
 					TranslationStages[Index].ElementCountZBinding = CompileSimStageData.ElementCountZBinding;
@@ -1312,6 +1313,13 @@ const FNiagaraTranslateResults &FHlslNiagaraTranslator::Translate(const FNiagara
 					TranslationStages[Index].bOverrideGpuDispatchNumThreads = CompileSimStageData.bOverrideGpuDispatchNumThreads;
 					TranslationStages[Index].OverrideGpuDispatchNumThreads = CompileSimStageData.OverrideGpuDispatchNumThreads;
 
+					if (CompileSimStageData.bOverrideGpuDispatchNumThreads)
+					{
+						TranslationStages[Index].OverrideGpuDispatchNumThreads.X = CompileData->GetStaticVariableValue(CompileSimStageData.OverrideGpuDispatchNumThreadsXBinding, TranslationStages[Index].OverrideGpuDispatchNumThreads.X);
+						TranslationStages[Index].OverrideGpuDispatchNumThreads.Y = CompileData->GetStaticVariableValue(CompileSimStageData.OverrideGpuDispatchNumThreadsYBinding, TranslationStages[Index].OverrideGpuDispatchNumThreads.Y);
+						TranslationStages[Index].OverrideGpuDispatchNumThreads.Z = CompileData->GetStaticVariableValue(CompileSimStageData.OverrideGpuDispatchNumThreadsZBinding, TranslationStages[Index].OverrideGpuDispatchNumThreads.Z);
+					}
+
 					ParamMapHistories.AddDefaulted(1);
 					ParamMapHistoriesSourceInOtherHistories.AddDefaulted(1);
 
@@ -1369,6 +1377,7 @@ const FNiagaraTranslateResults &FHlslNiagaraTranslator::Translate(const FNiagara
 					FSimulationStageMetaData& SimulationStageMetaData = CompilationOutput.ScriptData.SimulationStageMetaData.AddDefaulted_GetRef();
 					SimulationStageMetaData.SimulationStageName = InCompileData->CompileSimStageData[SourceSimStageIndex].StageName;
 					SimulationStageMetaData.EnabledBinding = InCompileData->CompileSimStageData[SourceSimStageIndex].EnabledBinding;
+					SimulationStageMetaData.ElementCount = InCompileData->CompileSimStageData[SourceSimStageIndex].ElementCount;
 					SimulationStageMetaData.ElementCountXBinding = InCompileData->CompileSimStageData[SourceSimStageIndex].ElementCountXBinding;
 					SimulationStageMetaData.ElementCountYBinding = InCompileData->CompileSimStageData[SourceSimStageIndex].ElementCountYBinding;
 					SimulationStageMetaData.ElementCountZBinding = InCompileData->CompileSimStageData[SourceSimStageIndex].ElementCountZBinding;
diff --git a/Engine/Plugins/FX/Niagara/Source/NiagaraEditor/Private/NiagaraHlslTranslator.h b/Engine/Plugins/FX/Niagara/Source/NiagaraEditor/Private/NiagaraHlslTranslator.h
--- a/Engine/Plugins/FX/Niagara/Source/NiagaraEditor/Private/NiagaraHlslTranslator.h
+++ b/Engine/Plugins/FX/Niagara/Source/NiagaraEditor/Private/NiagaraHlslTranslator.h
@@ -129,6 +129,12 @@ public:
 	TArray<FNiagaraVariable> StaticVariables;
 	TArray<FNiagaraVariable> StaticVariablesWithMultipleWrites;
 
+	template<typename T> T GetStaticVariableValue(const FNiagaraVariableBase Variable, const T DefaultValue) const
+	{
+		const int32 Index = StaticVariables.Find(Variable);
+		return Index != INDEX_NONE && StaticVariables[Index].IsDataAllocated() ? StaticVariables[Index].GetValue<T>() : DefaultValue;
+	}
+
 	void CompareAgainst(FNiagaraGraphCachedBuiltHistory* InCachedDataBase);
 
 	static void SortOutputNodesByDependencies(TArray<class UNiagaraNodeOutput*>& NodesToSort, const TArray<class UNiagaraSimulationStageBase*>* SimStages);
@@ -374,6 +380,7 @@ public:
 	FName IterationDirectBinding;
 	int32 SimulationStageIndex = -1;
 	FName EnabledBinding;
+	FIntVector3 ElementCount = FIntVector3::ZeroValue;
 	FName ElementCountXBinding;
 	FName ElementCountYBinding;
 	FName ElementCountZBinding;
diff --git a/Engine/Plugins/FX/Niagara/Source/NiagaraEditor/Private/NiagaraValidationRules.cpp b/Engine/Plugins/FX/Niagara/Source/NiagaraEditor/Private/NiagaraValidationRules.cpp
--- a/Engine/Plugins/FX/Niagara/Source/NiagaraEditor/Private/NiagaraValidationRules.cpp
+++ b/Engine/Plugins/FX/Niagara/Source/NiagaraEditor/Private/NiagaraValidationRules.cpp
@@ -843,15 +843,16 @@ void UNiagaraValidationRule_SimulationStageBudget::CheckValidity(const FNiagaraV
 				continue;
 			}
 
+			const int32 StageNumIterations = SimStage->NumIterations.GetDefaultValue<int32>();
 			++TotalEnabledStages;
-			TotalIterations += SimStage->Iterations;
-			if ( bMaxIterationsPerStageEnabled && SimStage->Iterations > MaxIterationsPerStage )
+			TotalIterations += StageNumIterations;
+			if ( bMaxIterationsPerStageEnabled && StageNumIterations > MaxIterationsPerStage)
 			{
 				UNiagaraStackEmitterPropertiesItem* EmitterProperties = NiagaraValidation::GetStackEntry<UNiagaraStackEmitterPropertiesItem>(EmitterHandleModel.Get().GetEmitterStackViewModel());
 				OutResults.Emplace_GetRef(
 					Severity,
 					FText::Format(LOCTEXT("SimStageTooManyIterationsFormat", "Simulation Stage '{0}' has too many iterations"), FText::FromName(SimStage->SimulationStageName)),
-					FText::Format(LOCTEXT("SimStageTooManyIterationsDetailedFormat", "Simulation Stage '{0}' has {1} iterations and we only allow {2}"), FText::FromName(SimStage->SimulationStageName), FText::AsNumber(SimStage->Iterations), FText::AsNumber(MaxIterationsPerStage)),
+					FText::Format(LOCTEXT("SimStageTooManyIterationsDetailedFormat", "Simulation Stage '{0}' has {1} iterations and we only allow {2}"), FText::FromName(SimStage->SimulationStageName), FText::AsNumber(StageNumIterations), FText::AsNumber(MaxIterationsPerStage)),
 					EmitterProperties
 				);
 			}
diff --git a/Engine/Plugins/FX/Niagara/Source/NiagaraShader/Public/NiagaraScriptBase.h b/Engine/Plugins/FX/Niagara/Source/NiagaraShader/Public/NiagaraScriptBase.h
--- a/Engine/Plugins/FX/Niagara/Source/NiagaraShader/Public/NiagaraScriptBase.h
+++ b/Engine/Plugins/FX/Niagara/Source/NiagaraShader/Public/NiagaraScriptBase.h
@@ -69,6 +69,10 @@ public:
 	UPROPERTY()
 	FName EnabledBinding;
 
+	/** Optional binding to manually specify the element count. */
+	UPROPERTY()
+	FIntVector ElementCount = FIntVector::ZeroValue;
+
 	/** Optional binding to manually specify the element count. */
 	UPROPERTY()
 	FName ElementCountXBinding;
diff --git a/Engine/Source/Runtime/Core/Public/UObject/UE5MainStreamObjectVersion.h b/Engine/Source/Runtime/Core/Public/UObject/UE5MainStreamObjectVersion.h
--- a/Engine/Source/Runtime/Core/Public/UObject/UE5MainStreamObjectVersion.h
+++ b/Engine/Source/Runtime/Core/Public/UObject/UE5MainStreamObjectVersion.h
@@ -349,6 +349,9 @@ struct CORE_API FUE5MainStreamObjectVersion
 		// Changed the material property connected bitmasks from 32bit to 64bit
 		IncreaseMaterialAttributesInputMask,
 
+		// Combines proprties into a new binding so users can select constant or binding
+		NiagaraSimStageNumIterationsBindings,
+
 		// -----<new versions can be added above this line>-------------------------------------------------
 		VersionPlusOne,
 		LatestVersion = VersionPlusOne - 1
