- Added 'default' UI state option for emitters that will determine if they will be added in summary view or default view when added to a system or creating a new emitter
- Changed 'Behavior Examples' label to 'Learning Examples'
- Removed the fallback logic for GetEditorData; imo Emitter Editor Data should safely exist whenever an emitter exists
- Added logic to ensure emitter editor data exists when an emitter is post loaded or initialized

diff --git a/Engine/Plugins/FX/Niagara/Source/Niagara/Classes/NiagaraEmitter.h b/Engine/Plugins/FX/Niagara/Source/Niagara/Classes/NiagaraEmitter.h
--- a/Engine/Plugins/FX/Niagara/Source/Niagara/Classes/NiagaraEmitter.h
+++ b/Engine/Plugins/FX/Niagara/Source/Niagara/Classes/NiagaraEmitter.h
@@ -440,6 +440,10 @@ struct NIAGARA_API FVersionedNiagaraEmitterData
 	UPROPERTY(EditAnywhere, AdvancedDisplay, Category = "Emitter")
 	TArray<FString> AttributesToPreserve;
 
+	/** This determines how emitters will be added to a system by default. If summary view is setup, consider setting this to 'Summary'. */
+	UPROPERTY(EditAnywhere, AdvancedDisplay, Category = "Asset Options")
+	ENiagaraEmitterDefaultSummaryState AddEmitterDefaultViewState = ENiagaraEmitterDefaultSummaryState::Default;
+	
 	UPROPERTY()
 	FNiagaraEmitterScriptProperties EmitterSpawnScriptProps;
 
@@ -640,6 +644,8 @@ public:
 	template <typename TAction>
 	void ForEachVersionData(TAction Func) const;
 
+	template <typename TAction>
+	void ForEachVersionData(TAction Func);
 #if WITH_EDITORONLY_DATA
 	NIAGARA_API virtual TSharedPtr<FNiagaraVersionDataAccessor> GetVersionDataAccessor(const FGuid& Version) override; 
 
@@ -1014,6 +1020,15 @@ void UNiagaraEmitter::ForEachVersionData(TAction Func) const
 	}
 }
 
+template <typename TAction>
+void UNiagaraEmitter::ForEachVersionData(TAction Func)
+{
+	for (FVersionedNiagaraEmitterData& Data : VersionData)
+	{
+		Func(Data);
+	}
+}
+
 template <typename TAction>
 void FVersionedNiagaraEmitterData::ForEachRenderer(TAction Func) const
 {
diff --git a/Engine/Plugins/FX/Niagara/Source/Niagara/Classes/NiagaraScript.h b/Engine/Plugins/FX/Niagara/Source/Niagara/Classes/NiagaraScript.h
--- a/Engine/Plugins/FX/Niagara/Source/Niagara/Classes/NiagaraScript.h
+++ b/Engine/Plugins/FX/Niagara/Source/Niagara/Classes/NiagaraScript.h
@@ -87,9 +87,19 @@ enum class ENiagaraScriptTemplateSpecification : uint8
 {
 	None,
 	Template,
-	Behavior UMETA(DisplayName = "Behavior Example")
+	Behavior UMETA(DisplayName = "Learning Example")
 };
 
+UENUM()
+enum class ENiagaraEmitterDefaultSummaryState : uint8
+{
+	/** The emitter will be added in default, full view by default. */
+	Default,
+	/** The emitter will be added in summary state by default. */
+	Summary
+};
+
+
 /** Defines different usages for a niagara script's module dependecies. */
 UENUM()
 enum class ENiagaraModuleDependencyUsage : uint8
diff --git a/Engine/Plugins/FX/Niagara/Source/Niagara/Private/NiagaraEmitter.cpp b/Engine/Plugins/FX/Niagara/Source/Niagara/Private/NiagaraEmitter.cpp
--- a/Engine/Plugins/FX/Niagara/Source/Niagara/Private/NiagaraEmitter.cpp
+++ b/Engine/Plugins/FX/Niagara/Source/Niagara/Private/NiagaraEmitter.cpp
@@ -259,11 +259,20 @@ void FVersionedNiagaraEmitterData::PostInitProperties(UNiagaraEmitter* Outer)
 		GPUComputeScript->SetUsage(ENiagaraScriptUsage::ParticleGPUComputeScript);
 
 #if WITH_EDITORONLY_DATA && WITH_EDITOR
-	if (EditorParameters == nullptr)
+	if(EditorParameters == nullptr || EditorData == nullptr)
 	{
 		INiagaraModule& NiagaraModule = FModuleManager::GetModuleChecked<INiagaraModule>("Niagara");
-		EditorParameters = NiagaraModule.GetEditorOnlyDataUtilities().CreateDefaultEditorParameters(Outer);
-	}
+
+		if(EditorParameters == nullptr)
+		{
+			EditorParameters = NiagaraModule.GetEditorOnlyDataUtilities().CreateDefaultEditorParameters(Outer);
+		}
+
+		if(EditorData == nullptr)
+		{
+			EditorData = NiagaraModule.GetEditorOnlyDataUtilities().CreateDefaultEditorData(Outer);
+		}
+	}	
 #endif
 }
 
@@ -779,6 +788,11 @@ void FVersionedNiagaraEmitterData::PostLoad(UNiagaraEmitter& Emitter, int32 Niag
 		EditorData->ConditionalPostLoad();
 		EditorData->PostLoadFromOwner(&Emitter);
 	}
+	else
+	{
+		INiagaraModule& NiagaraModule = FModuleManager::GetModuleChecked<INiagaraModule>("Niagara");
+		EditorData = NiagaraModule.GetEditorOnlyDataUtilities().CreateDefaultEditorData(&Emitter);
+	}
 	
 	if (!GPUComputeScript)
 	{
diff --git a/Engine/Plugins/FX/Niagara/Source/NiagaraEditor/Private/NiagaraEditorModule.cpp b/Engine/Plugins/FX/Niagara/Source/NiagaraEditor/Private/NiagaraEditorModule.cpp
--- a/Engine/Plugins/FX/Niagara/Source/NiagaraEditor/Private/NiagaraEditorModule.cpp
+++ b/Engine/Plugins/FX/Niagara/Source/NiagaraEditor/Private/NiagaraEditorModule.cpp
@@ -208,8 +208,7 @@ class FNiagaraEditorOnlyDataUtilities : public INiagaraEditorOnlyDataUtilities
 
 	virtual UNiagaraEditorDataBase* CreateDefaultEditorData(UObject* InOuter) const override
 	{
-		UNiagaraSystem* System = Cast<UNiagaraSystem>(InOuter);
-		if (System != nullptr)
+		if (UNiagaraSystem* System = Cast<UNiagaraSystem>(InOuter))
 		{
 			UNiagaraSystemEditorData* SystemEditorData = NewObject<UNiagaraSystemEditorData>(InOuter);
 			SystemEditorData->SetFlags(RF_Transactional);
@@ -217,6 +216,12 @@ class FNiagaraEditorOnlyDataUtilities : public INiagaraEditorOnlyDataUtilities
 			SystemEditorData->InitOnSyncScriptVariables(*System);
 			return SystemEditorData;
 		}
+		else if(UNiagaraEmitter* Emitter = Cast<UNiagaraEmitter>(InOuter))
+		{
+			UNiagaraEmitterEditorData* EmitterEditorData = NewObject<UNiagaraEmitterEditorData>(InOuter);
+			EmitterEditorData->SetFlags(RF_Transactional);
+			return EmitterEditorData;
+		}
 		return nullptr;
 	}
 
diff --git a/Engine/Plugins/FX/Niagara/Source/NiagaraEditor/Private/NiagaraEditorUtilities.cpp b/Engine/Plugins/FX/Niagara/Source/NiagaraEditor/Private/NiagaraEditorUtilities.cpp
--- a/Engine/Plugins/FX/Niagara/Source/NiagaraEditor/Private/NiagaraEditorUtilities.cpp
+++ b/Engine/Plugins/FX/Niagara/Source/NiagaraEditor/Private/NiagaraEditorUtilities.cpp
@@ -2100,7 +2100,8 @@ const FGuid FNiagaraEditorUtilities::AddEmitterToSystem(UNiagaraSystem& InSystem
 		EmitterHandle = FNiagaraEmitterHandle(InEmitterToAdd, EmitterVersion);
 		InSystem.AddEmitterHandleDirect(EmitterHandle);		
 	}
-	
+
+	Cast<UNiagaraEmitterEditorData>(EmitterHandle.GetEmitterData()->GetEditorData())->SetShowSummaryView(InEmitterToAdd.GetEmitterData(EmitterVersion)->AddEmitterDefaultViewState == ENiagaraEmitterDefaultSummaryState::Summary ? true : false);
 	FNiagaraStackGraphUtilities::RebuildEmitterNodes(InSystem);
 	SystemEditorData->SynchronizeOverviewGraphWithSystem(InSystem);
 
diff --git a/Engine/Plugins/FX/Niagara/Source/NiagaraEditor/Private/NiagaraEmitterFactoryNew.cpp b/Engine/Plugins/FX/Niagara/Source/NiagaraEditor/Private/NiagaraEmitterFactoryNew.cpp
--- a/Engine/Plugins/FX/Niagara/Source/NiagaraEditor/Private/NiagaraEmitterFactoryNew.cpp
+++ b/Engine/Plugins/FX/Niagara/Source/NiagaraEditor/Private/NiagaraEmitterFactoryNew.cpp
@@ -12,6 +12,7 @@
 #include "AssetRegistry/AssetData.h"
 #include "ViewModels/Stack/NiagaraStackGraphUtilities.h"
 #include "NiagaraConstants.h"
+#include "NiagaraEmitterEditorData.h"
 #include "NiagaraNodeAssignment.h"
 #include "SNewEmitterDialog.h"
 #include "Misc/MessageDialog.h"
@@ -143,6 +144,13 @@ UObject* UNiagaraEmitterFactoryNew::FactoryCreateNew(UClass* Class, UObject* InP
 		NewEmitter = NewObject<UNiagaraEmitter>(InParent, Class, Name, Flags | RF_Transactional);
 		InitializeEmitter(NewEmitter, bAddDefaultModulesAndRenderersToEmptyEmitter);
 	}
+	
+	NewEmitter->ForEachVersionData([](FVersionedNiagaraEmitterData& EmitterVersionData)
+	{
+		UNiagaraEmitterEditorData* EmitterEditorData = CastChecked<UNiagaraEmitterEditorData>(EmitterVersionData.GetEditorData());
+		EmitterEditorData->SetShowSummaryView(EmitterVersionData.AddEmitterDefaultViewState == ENiagaraEmitterDefaultSummaryState::Summary ? true : false);
+	});
+	
 	return NewEmitter;
 }
 
diff --git a/Engine/Plugins/FX/Niagara/Source/NiagaraEditor/Private/ViewModels/NiagaraEmitterViewModel.cpp b/Engine/Plugins/FX/Niagara/Source/NiagaraEditor/Private/ViewModels/NiagaraEmitterViewModel.cpp
--- a/Engine/Plugins/FX/Niagara/Source/NiagaraEditor/Private/ViewModels/NiagaraEmitterViewModel.cpp
+++ b/Engine/Plugins/FX/Niagara/Source/NiagaraEditor/Private/ViewModels/NiagaraEmitterViewModel.cpp
@@ -74,7 +74,7 @@ void FNiagaraEmitterViewModel::Cleanup()
 
 	if(EmitterWeakPtr.IsValid())
 	{
-		GetOrCreateEditorData().OnPersistentDataChanged().RemoveAll(this);
+		GetEditorData().OnPersistentDataChanged().RemoveAll(this);
 	}
 
 	RemoveScriptEventHandlers();
@@ -327,26 +327,13 @@ UNiagaraSummaryViewViewModel* FNiagaraEmitterViewModel::GetSummaryHierarchyViewM
 const UNiagaraEmitterEditorData& FNiagaraEmitterViewModel::GetEditorData() const
 {
 	check(EmitterWeakPtr.Emitter.IsValid());
-
-	const UNiagaraEmitterEditorData* EditorData = Cast<UNiagaraEmitterEditorData>(EmitterWeakPtr.GetEmitterData()->GetEditorData());
-	if (EditorData == nullptr)
-	{
-		EditorData = GetDefault<UNiagaraEmitterEditorData>();
-	}
-	return *EditorData;
+	return *Cast<UNiagaraEmitterEditorData>(EmitterWeakPtr.GetEmitterData()->GetEditorData());
 }
 
-UNiagaraEmitterEditorData& FNiagaraEmitterViewModel::GetOrCreateEditorData()
+UNiagaraEmitterEditorData& FNiagaraEmitterViewModel::GetEditorData()
 {
 	check(EmitterWeakPtr.IsValid());
-	UNiagaraEmitterEditorData* EditorData = Cast<UNiagaraEmitterEditorData>(EmitterWeakPtr.GetEmitterData()->GetEditorData());
-	if (EditorData == nullptr)
-	{
-		EditorData = NewObject<UNiagaraEmitterEditorData>(EmitterWeakPtr.Emitter.Get(), NAME_None, RF_Transactional);
-		EmitterWeakPtr.Emitter->Modify(); 
-		EmitterWeakPtr.Emitter->SetEditorData(EditorData, EmitterWeakPtr.Version);
-	}
-	return *EditorData;
+	return *Cast<UNiagaraEmitterEditorData>(EmitterWeakPtr.GetEmitterData()->GetEditorData());
 }
 
 void FNiagaraEmitterViewModel::AddEventHandler(FNiagaraEventScriptProperties& EventScriptProperties, bool bResetGraphForOutput /*= false*/)
@@ -606,7 +593,7 @@ void FNiagaraEmitterViewModel::OnEmitterPropertiesChanged()
 
 void FNiagaraEmitterViewModel::OnSummaryViewHierarchyChanged()
 {
-	GetOrCreateEditorData().OnPersistentDataChanged().Broadcast();
+	GetEditorData().OnPersistentDataChanged().Broadcast();
 }
 
 FString FNiagaraEmitterViewModel::GetReferencerName() const
diff --git a/Engine/Plugins/FX/Niagara/Source/NiagaraEditor/Private/ViewModels/NiagaraSystemViewModel.cpp b/Engine/Plugins/FX/Niagara/Source/NiagaraEditor/Private/ViewModels/NiagaraSystemViewModel.cpp
--- a/Engine/Plugins/FX/Niagara/Source/NiagaraEditor/Private/ViewModels/NiagaraSystemViewModel.cpp
+++ b/Engine/Plugins/FX/Niagara/Source/NiagaraEditor/Private/ViewModels/NiagaraSystemViewModel.cpp
@@ -538,7 +538,9 @@ TSharedPtr<FNiagaraEmitterHandleViewModel> FNiagaraSystemViewModel::AddEmitter(U
 	}
 
 	RefreshAll();
-
+	
+	GetEmitterHandleViewModelById(NewEmitterHandleId)->GetEmitterViewModel()->GetEditorData().SetShowSummaryView(Emitter.GetEmitterData(EmitterVersion)->AddEmitterDefaultViewState == ENiagaraEmitterDefaultSummaryState::Summary ? true : false);
+	
 	TRange<float> SystemPlaybackRange = GetEditorData().GetPlaybackRange();
 	TRange<float> EmitterPlaybackRange = GetEmitterHandleViewModelById(NewEmitterHandleId)->GetEmitterViewModel()->GetEditorData().GetPlaybackRange();
 	TRange<float> NewSystemPlaybackRange = TRange<float>(
@@ -2550,8 +2552,8 @@ void FNiagaraSystemViewModel::SequencerDataChanged(EMovieSceneDataChangeType Dat
 
 				VaildTrackEmitterHandleIds.Add(EmitterTrack->GetEmitterHandleViewModel()->GetId());
 				EmitterTrack->UpdateEmitterHandleFromTrackChange(NiagaraSequence->GetMovieScene()->GetTickResolution());
-				EmitterTrack->GetEmitterHandleViewModel()->GetEmitterViewModel()->GetOrCreateEditorData().Modify();
-				EmitterTrack->GetEmitterHandleViewModel()->GetEmitterViewModel()->GetOrCreateEditorData().SetPlaybackRange(GetEditorData().GetPlaybackRange());
+				EmitterTrack->GetEmitterHandleViewModel()->GetEmitterViewModel()->GetEditorData().Modify();
+				EmitterTrack->GetEmitterHandleViewModel()->GetEmitterViewModel()->GetEditorData().SetPlaybackRange(GetEditorData().GetPlaybackRange());
 				if (EmitterTrack->GetDisplayName().ToString() != EmitterTrack->GetEmitterHandleViewModel()->GetNameText().ToString())
 				{
 					EmitterHandlesToRename.Add(TTuple<TSharedPtr<FNiagaraEmitterHandleViewModel>, FName>(EmitterTrack->GetEmitterHandleViewModel(), *EmitterTrack->GetDisplayName().ToString()));
diff --git a/Engine/Plugins/FX/Niagara/Source/NiagaraEditor/Private/ViewModels/Stack/NiagaraStackRoot.cpp b/Engine/Plugins/FX/Niagara/Source/NiagaraEditor/Private/ViewModels/Stack/NiagaraStackRoot.cpp
--- a/Engine/Plugins/FX/Niagara/Source/NiagaraEditor/Private/ViewModels/Stack/NiagaraStackRoot.cpp
+++ b/Engine/Plugins/FX/Niagara/Source/NiagaraEditor/Private/ViewModels/Stack/NiagaraStackRoot.cpp
@@ -50,7 +50,7 @@ void UNiagaraStackRoot::Initialize(FRequiredEntryData InRequiredEntryData, bool
 
 	if (bInIncludeEmitterInformation && GetEmitterViewModel())
 	{
-		GetEmitterViewModel()->GetOrCreateEditorData().OnSummaryViewStateChanged().AddUObject(this, &UNiagaraStackRoot::OnSummaryViewStateChanged);
+		GetEmitterViewModel()->GetEditorData().OnSummaryViewStateChanged().AddUObject(this, &UNiagaraStackRoot::OnSummaryViewStateChanged);
 	}
 }
 
@@ -58,7 +58,7 @@ void UNiagaraStackRoot::FinalizeInternal()
 {
 	if (bIncludeEmitterInformation && GetEmitterViewModel() && GetEmitterViewModel()->GetEmitter().GetEmitterData())
 	{
-		GetEmitterViewModel()->GetOrCreateEditorData().OnSummaryViewStateChanged().RemoveAll(this);
+		GetEmitterViewModel()->GetEditorData().OnSummaryViewStateChanged().RemoveAll(this);
 	}	
 	Super::FinalizeInternal();
 }
@@ -126,7 +126,7 @@ void UNiagaraStackRoot::RefreshChildrenInternal(const TArray<UNiagaraStackEntry*
 		CommentCollection->Initialize(RequiredEntryData);
 	}
 	
-	bool bShouldShowSummaryView = GetEmitterViewModel()? GetEmitterViewModel()->GetOrCreateEditorData().ShouldShowSummaryView() : false;
+	bool bShouldShowSummaryView = GetEmitterViewModel()? GetEmitterViewModel()->GetEditorData().ShouldShowSummaryView() : false;
 
 	// we clear references to entries we don't want to display in summary view.
 	// Since they also won't get added to the NewChildren array, they will get finalized and eventually GC'd
@@ -143,7 +143,7 @@ void UNiagaraStackRoot::RefreshChildrenInternal(const TArray<UNiagaraStackEntry*
 		EmitterPropertiesGroup = NewObject<UNiagaraStackEmitterPropertiesGroup>(this);
 		FRequiredEntryData RequiredEntryData(GetSystemViewModel(), GetEmitterViewModel(),
 			FExecutionCategoryNames::Emitter, FExecutionSubcategoryNames::Settings,
-			GetEmitterViewModel()->GetOrCreateEditorData().GetStackEditorData());
+			GetEmitterViewModel()->GetEditorData().GetStackEditorData());
 		EmitterPropertiesGroup->Initialize(RequiredEntryData);
 	}
 
@@ -152,7 +152,7 @@ void UNiagaraStackRoot::RefreshChildrenInternal(const TArray<UNiagaraStackEntry*
 		EmitterSummaryGroup = NewObject<UNiagaraStackEmitterSummaryGroup>(this);
 		FRequiredEntryData RequiredEntryData(GetSystemViewModel(), GetEmitterViewModel(),
 			FExecutionCategoryNames::Emitter, FExecutionSubcategoryNames::Settings,
-			GetEmitterViewModel()->GetOrCreateEditorData().GetStackEditorData());
+			GetEmitterViewModel()->GetEditorData().GetStackEditorData());
 		FText DisplayName = LOCTEXT("EmitterSummaryGroupName", "Emitter Summary");
 		FText Tooltip = LOCTEXT("EmitterSummaryTooltip", "Summary of parameters for this Emitter.");
 		EmitterSummaryGroup->Initialize(RequiredEntryData, DisplayName, Tooltip, nullptr);
@@ -163,7 +163,7 @@ void UNiagaraStackRoot::RefreshChildrenInternal(const TArray<UNiagaraStackEntry*
 		EmitterSpawnGroup = NewObject<UNiagaraStackScriptItemGroup>(this);
 		FRequiredEntryData RequiredEntryData(GetSystemViewModel(), GetEmitterViewModel(),
 			FExecutionCategoryNames::Emitter, FExecutionSubcategoryNames::Spawn,
-			GetEmitterViewModel()->GetOrCreateEditorData().GetStackEditorData());
+			GetEmitterViewModel()->GetEditorData().GetStackEditorData());
 		FText DisplayName = LOCTEXT("EmitterSpawnGroupName", "Emitter Spawn");
 		FText ToolTip = LOCTEXT("EmitterSpawnGroupTooltip", "Occurs once at Emitter creation on the CPU. Modules in this stage should initialize defaults and/or do initial setup.\r\nModules are executed in order from top to bottom of the stack.");
 		EmitterSpawnGroup->Initialize(RequiredEntryData, DisplayName, ToolTip, GetEmitterViewModel()->GetSharedScriptViewModel(), ENiagaraScriptUsage::EmitterSpawnScript);
@@ -174,7 +174,7 @@ void UNiagaraStackRoot::RefreshChildrenInternal(const TArray<UNiagaraStackEntry*
 		EmitterUpdateGroup = NewObject<UNiagaraStackScriptItemGroup>(this);
 		FRequiredEntryData RequiredEntryData(GetSystemViewModel(), GetEmitterViewModel(),
 			FExecutionCategoryNames::Emitter, FExecutionSubcategoryNames::Update,
-			GetEmitterViewModel()->GetOrCreateEditorData().GetStackEditorData());
+			GetEmitterViewModel()->GetEditorData().GetStackEditorData());
 		FText DisplayName = LOCTEXT("EmitterUpdateGroupName", "Emitter Update");
 		FText ToolTip = LOCTEXT("EmitterUpdateGroupTooltip", "Occurs every Emitter tick on the CPU. Modules in this stage should compute values for parameters for Particle Update or Spawning this frame.\r\nModules are executed in order from top to bottom of the stack.");
 		EmitterUpdateGroup->Initialize(RequiredEntryData, DisplayName, ToolTip, GetEmitterViewModel()->GetSharedScriptViewModel(), ENiagaraScriptUsage::EmitterUpdateScript);
@@ -185,7 +185,7 @@ void UNiagaraStackRoot::RefreshChildrenInternal(const TArray<UNiagaraStackEntry*
 		ParticleSpawnGroup = NewObject<UNiagaraStackScriptItemGroup>(this);
 		FRequiredEntryData RequiredEntryData(GetSystemViewModel(), GetEmitterViewModel(),
 			FExecutionCategoryNames::Particle, FExecutionSubcategoryNames::Spawn,
-			GetEmitterViewModel()->GetOrCreateEditorData().GetStackEditorData());
+			GetEmitterViewModel()->GetEditorData().GetStackEditorData());
 		FText DisplayName = LOCTEXT("ParticleSpawnGroupName", "Particle Spawn");
 		FText ToolTip = LOCTEXT("ParticleSpawnGroupTooltip", "Called once per created particle. Modules in this stage should set up initial values for each particle.\r\nIf \"Use Interpolated Spawning\" is set, we will also run the Particle Update stage after the Particle Spawn stage.\r\nModules are executed in order from top to bottom of the stack.");
 		ParticleSpawnGroup->Initialize(RequiredEntryData, DisplayName, ToolTip, GetEmitterViewModel()->GetSharedScriptViewModel(), ENiagaraScriptUsage::ParticleSpawnScript);
@@ -196,7 +196,7 @@ void UNiagaraStackRoot::RefreshChildrenInternal(const TArray<UNiagaraStackEntry*
 		ParticleUpdateGroup = NewObject<UNiagaraStackScriptItemGroup>(this);
 		FRequiredEntryData RequiredEntryData(GetSystemViewModel(), GetEmitterViewModel(),
 			FExecutionCategoryNames::Particle, FExecutionSubcategoryNames::Update,
-			GetEmitterViewModel()->GetOrCreateEditorData().GetStackEditorData());
+			GetEmitterViewModel()->GetEditorData().GetStackEditorData());
 		FText DisplayName = LOCTEXT("ParticleUpdateGroupName", "Particle Update");
 		FText ToolTip = LOCTEXT("ParticleUpdateGroupTooltip", "Called every frame per particle. Modules in this stage should update new values for this frame.\r\nModules are executed in order from top to bottom of the stack.");
 		ParticleUpdateGroup->Initialize(RequiredEntryData, DisplayName, ToolTip, GetEmitterViewModel()->GetSharedScriptViewModel(), ENiagaraScriptUsage::ParticleUpdateScript);
@@ -207,7 +207,7 @@ void UNiagaraStackRoot::RefreshChildrenInternal(const TArray<UNiagaraStackEntry*
 		RenderGroup = NewObject<UNiagaraStackRenderItemGroup>(this);
 		FRequiredEntryData RequiredEntryData(GetSystemViewModel(), GetEmitterViewModel(),
 			FExecutionCategoryNames::Render, FExecutionSubcategoryNames::Render,
-			GetEmitterViewModel()->GetOrCreateEditorData().GetStackEditorData());
+			GetEmitterViewModel()->GetEditorData().GetStackEditorData());
 		RenderGroup->Initialize(RequiredEntryData);
 	}
 
@@ -216,7 +216,7 @@ void UNiagaraStackRoot::RefreshChildrenInternal(const TArray<UNiagaraStackEntry*
 		SummaryCollapseButton = NewObject<UNiagaraStackSummaryViewCollapseButton>(this);
 		FRequiredEntryData RequiredEntryData(GetSystemViewModel(), GetEmitterViewModel(),
 			FExecutionCategoryNames::Emitter, FExecutionSubcategoryNames::Settings,
-			GetEmitterViewModel()->GetOrCreateEditorData().GetStackEditorData());
+			GetEmitterViewModel()->GetEditorData().GetStackEditorData());
 		SummaryCollapseButton->Initialize(RequiredEntryData);
 	}
 
diff --git a/Engine/Plugins/FX/Niagara/Source/NiagaraEditor/Private/ViewModels/Stack/NiagaraStackSummaryViewInputCollection.cpp b/Engine/Plugins/FX/Niagara/Source/NiagaraEditor/Private/ViewModels/Stack/NiagaraStackSummaryViewInputCollection.cpp
--- a/Engine/Plugins/FX/Niagara/Source/NiagaraEditor/Private/ViewModels/Stack/NiagaraStackSummaryViewInputCollection.cpp
+++ b/Engine/Plugins/FX/Niagara/Source/NiagaraEditor/Private/ViewModels/Stack/NiagaraStackSummaryViewInputCollection.cpp
@@ -35,7 +35,7 @@ void UNiagaraStackSummaryViewCollection::Initialize(FRequiredEntryData InRequire
 	Super::Initialize(InRequiredEntryData, InOwningStackItemEditorDataKey, ObjectStackEditorDataKey);
 
 	Emitter = InEmitter;
-	GetEmitterViewModel()->GetOrCreateEditorData().OnSummaryViewStateChanged().AddUObject(this, &UNiagaraStackSummaryViewCollection::OnViewStateChanged);
+	GetEmitterViewModel()->GetEditorData().OnSummaryViewStateChanged().AddUObject(this, &UNiagaraStackSummaryViewCollection::OnViewStateChanged);
 	GetEmitterViewModel()->GetSummaryHierarchyViewModel()->OnHierarchyChanged().AddUObject(this, &UNiagaraStackSummaryViewCollection::OnViewStateChanged);
 	GetEmitterViewModel()->GetSummaryHierarchyViewModel()->OnHierarchyPropertiesChanged().AddUObject(this, &UNiagaraStackSummaryViewCollection::OnViewStateChanged);
 }
@@ -45,7 +45,7 @@ void UNiagaraStackSummaryViewCollection::FinalizeInternal()
 	TSharedPtr<FNiagaraEmitterViewModel> EmitterViewModelPtr = GetEmitterViewModel();
 	if (EmitterViewModelPtr.IsValid() && GetEmitterViewModel()->GetEmitter().GetEmitterData())
 	{
-		GetEmitterViewModel()->GetOrCreateEditorData().OnSummaryViewStateChanged().RemoveAll(this);
+		GetEmitterViewModel()->GetEditorData().OnSummaryViewStateChanged().RemoveAll(this);
 	}
 
 	if(EmitterViewModelPtr.IsValid() && EmitterViewModelPtr->GetSummaryHierarchyViewModel() != nullptr)
@@ -73,7 +73,7 @@ void UNiagaraStackSummaryViewCollection::RefreshChildrenInternal(const TArray<UN
 	TSharedPtr<FNiagaraScriptViewModel> ScriptViewModelPinned = ViewModel->GetSharedScriptViewModel();
 	checkf(ScriptViewModelPinned.IsValid(), TEXT("Can not refresh children when the script view model has been deleted."));
 
-	const UNiagaraHierarchyRoot* Root = ViewModel->GetOrCreateEditorData().GetSummaryRoot();
+	const UNiagaraHierarchyRoot* Root = ViewModel->GetEditorData().GetSummaryRoot();
 	TSharedPtr<FNiagaraHierarchyRootViewModel> RootViewModel = ViewModel->GetSummaryHierarchyViewModel()->GetHierarchyRootViewModel();
 	
 	FNiagaraStackGraphUtilities::FInputDataCollection State;
@@ -489,7 +489,7 @@ void UNiagaraStackSummaryViewCollection::GetSectionsInternal(TArray<FNiagaraStac
 	TMap<const UNiagaraHierarchySection*, TArray<UNiagaraHierarchyCategory*>> SectionToCategoryMap;
 
 	TSharedPtr<FNiagaraEmitterViewModel> ViewModel = GetEmitterViewModel();
-	const UNiagaraHierarchyRoot* Root = ViewModel->GetOrCreateEditorData().GetSummaryRoot();
+	const UNiagaraHierarchyRoot* Root = ViewModel->GetEditorData().GetSummaryRoot();
 	
 	TArray<UNiagaraHierarchyCategory*> HierarchyCategories;
 	Root->GetChildrenOfType<UNiagaraHierarchyCategory>(HierarchyCategories, false);
@@ -521,7 +521,7 @@ const TArray<UNiagaraHierarchySection*>& UNiagaraStackSummaryViewCollection::Get
 {
 	if (GetEmitterViewModel().IsValid())
 	{
-		return GetEmitterViewModel()->GetOrCreateEditorData().GetSummarySections();
+		return GetEmitterViewModel()->GetEditorData().GetSummarySections();
 	}
 
 	static TArray<UNiagaraHierarchySection*> Dummy;
diff --git a/Engine/Plugins/FX/Niagara/Source/NiagaraEditor/Private/ViewModels/Stack/NiagaraStackViewModel.cpp b/Engine/Plugins/FX/Niagara/Source/NiagaraEditor/Private/ViewModels/Stack/NiagaraStackViewModel.cpp
--- a/Engine/Plugins/FX/Niagara/Source/NiagaraEditor/Private/ViewModels/Stack/NiagaraStackViewModel.cpp
+++ b/Engine/Plugins/FX/Niagara/Source/NiagaraEditor/Private/ViewModels/Stack/NiagaraStackViewModel.cpp
@@ -121,7 +121,7 @@ void UNiagaraStackViewModel::InitializeWithViewModels(TSharedPtr<FNiagaraSystemV
 
 			if (UNiagaraEmitter* Emitter = EmitterViewModel->GetEmitter().Emitter)
 			{
-				EmitterViewModel->GetOrCreateEditorData().OnSummaryViewStateChanged().AddUObject(this, &UNiagaraStackViewModel::RequestRefreshDeferred);
+				EmitterViewModel->GetEditorData().OnSummaryViewStateChanged().AddUObject(this, &UNiagaraStackViewModel::RequestRefreshDeferred);
 			}
 
 		}
@@ -170,7 +170,7 @@ void UNiagaraStackViewModel::InitializeWithRootEntry(UNiagaraStackEntry* InRootE
 	if (EmitterHandleViewModelPinned.IsValid())
 	{
 		TSharedPtr<FNiagaraEmitterViewModel> EmitterViewModel = EmitterHandleViewModelPinned->GetEmitterViewModel();
-		EmitterViewModel->GetOrCreateEditorData().OnSummaryViewStateChanged().RemoveAll(this);
+		EmitterViewModel->GetEditorData().OnSummaryViewStateChanged().RemoveAll(this);
 	}
 }
 
@@ -198,7 +198,7 @@ void UNiagaraStackViewModel::Reset()
 		EmitterViewModel->OnParentRemoved().RemoveAll(this);
 		if (EmitterViewModel->GetEmitter().GetEmitterData())
 		{
-			EmitterViewModel->GetOrCreateEditorData().OnSummaryViewStateChanged().RemoveAll(this);
+			EmitterViewModel->GetEditorData().OnSummaryViewStateChanged().RemoveAll(this);
 		}
 	}
 
@@ -761,7 +761,7 @@ void UNiagaraStackViewModel::SetLastScrollPosition(double InLastScrollPosition)
 	// TODO: Fix this with the new overview paradigm.
 	if (EmitterHandleViewModel.IsValid())
 	{
-		EmitterHandleViewModel.Pin()->GetEmitterViewModel()->GetOrCreateEditorData().GetStackEditorData().SetLastScrollPosition(InLastScrollPosition);
+		EmitterHandleViewModel.Pin()->GetEmitterViewModel()->GetEditorData().GetStackEditorData().SetLastScrollPosition(InLastScrollPosition);
 	}
 }
 
diff --git a/Engine/Plugins/FX/Niagara/Source/NiagaraEditor/Private/Widgets/SNiagaraFilterBox.cpp b/Engine/Plugins/FX/Niagara/Source/NiagaraEditor/Private/Widgets/SNiagaraFilterBox.cpp
--- a/Engine/Plugins/FX/Niagara/Source/NiagaraEditor/Private/Widgets/SNiagaraFilterBox.cpp
+++ b/Engine/Plugins/FX/Niagara/Source/NiagaraEditor/Private/Widgets/SNiagaraFilterBox.cpp
@@ -403,7 +403,7 @@ FText SNiagaraTemplateTabBox::DetermineControlLabel(ENiagaraScriptTemplateSpecif
 			case ENiagaraScriptTemplateSpecification::Template:
 				return LOCTEXT("EmitterTemplateTabLabel", "Templates");
 			case ENiagaraScriptTemplateSpecification::Behavior:
-				return LOCTEXT("EmitterBehaviorExampleTabLabel", "Behavior Examples");
+				return LOCTEXT("EmitterLearningExampleTabLabel", "Learning Examples");
 			default:
 				return FText::GetEmpty();
 			}
@@ -417,7 +417,7 @@ FText SNiagaraTemplateTabBox::DetermineControlLabel(ENiagaraScriptTemplateSpecif
 			case ENiagaraScriptTemplateSpecification::Template:
 				return LOCTEXT("SystemTemplateTabLabel", "Templates");
 			case ENiagaraScriptTemplateSpecification::Behavior:
-				return LOCTEXT("SystemBehaviorExampleTabLabel", "Behavior Examples");
+				return LOCTEXT("SystemLearningExampleTabLabel", "Learning Examples");
 			default:
 				return FText::GetEmpty();
 			}
diff --git a/Engine/Plugins/FX/Niagara/Source/NiagaraEditor/Public/ViewModels/NiagaraEmitterViewModel.h b/Engine/Plugins/FX/Niagara/Source/NiagaraEditor/Public/ViewModels/NiagaraEmitterViewModel.h
--- a/Engine/Plugins/FX/Niagara/Source/NiagaraEditor/Public/ViewModels/NiagaraEmitterViewModel.h
+++ b/Engine/Plugins/FX/Niagara/Source/NiagaraEditor/Public/ViewModels/NiagaraEmitterViewModel.h
@@ -120,7 +120,7 @@ public:
 	NIAGARAEDITOR_API const UNiagaraEmitterEditorData& GetEditorData() const;
 
 	/** Gets editor specific data which is stored per emitter.  If this data hasn't been created then it will be created. */
-	NIAGARAEDITOR_API UNiagaraEmitterEditorData& GetOrCreateEditorData();
+	NIAGARAEDITOR_API UNiagaraEmitterEditorData& GetEditorData();
 
 	/** Add an event script to the owned emitter. Sets the Usage, UsageID and Source of the EventScriptProperties. */
 	NIAGARAEDITOR_API void AddEventHandler(FNiagaraEventScriptProperties& EventScriptProperties, bool bResetGraphForOutput = false);
diff --git a/Engine/Plugins/FX/Niagara/Source/NiagaraEditorWidgets/Private/SNiagaraOverviewStackNode.cpp b/Engine/Plugins/FX/Niagara/Source/NiagaraEditorWidgets/Private/SNiagaraOverviewStackNode.cpp
--- a/Engine/Plugins/FX/Niagara/Source/NiagaraEditorWidgets/Private/SNiagaraOverviewStackNode.cpp
+++ b/Engine/Plugins/FX/Niagara/Source/NiagaraEditorWidgets/Private/SNiagaraOverviewStackNode.cpp
@@ -1003,7 +1003,7 @@ FText SNiagaraOverviewStackNode::GetSummaryViewCollapseTooltipText() const
 
 FReply SNiagaraOverviewStackNode::ExpandSummaryViewClicked()
 {
-	if (UNiagaraEmitterEditorData* EditorData = EmitterHandleViewModelWeak.IsValid()? &EmitterHandleViewModelWeak.Pin()->GetEmitterViewModel()->GetOrCreateEditorData() : nullptr)
+	if (UNiagaraEmitterEditorData* EditorData = EmitterHandleViewModelWeak.IsValid()? &EmitterHandleViewModelWeak.Pin()->GetEmitterViewModel()->GetEditorData() : nullptr)
 	{
 		EditorData->ToggleShowSummaryView();
 	}
