Fix path for shader debug info for niagara shaders (use the path of the system rather than the name of the object, which is usually useless like NiagaraPrecompileContainer_0)

diff --git a/Engine/Plugins/FX/Niagara/Source/NiagaraEditor/Private/NiagaraCompiler.cpp b/Engine/Plugins/FX/Niagara/Source/NiagaraEditor/Private/NiagaraCompiler.cpp
--- a/Engine/Plugins/FX/Niagara/Source/NiagaraEditor/Private/NiagaraCompiler.cpp
+++ b/Engine/Plugins/FX/Niagara/Source/NiagaraEditor/Private/NiagaraCompiler.cpp
@@ -969,7 +969,7 @@ TSharedPtr<FNiagaraCompileRequestDataBase, ESPMode::ThreadSafe> FNiagaraEditorMo
 	TArray<TSharedPtr<FNiagaraCompileRequestData, ESPMode::ThreadSafe>> DependentRequests;
 	FCompileConstantResolver EmptyResolver;
 
-	BasePtr->SourceName = InObj->GetName();
+	BasePtr->SourceName = LogName;
 
 	if (Script)
 	{
