Fix for partial depth collision queries to never use the GBuffer normals but only the normal from depth derivative. This is required to get normal matching the particl depth buffer withhout opaque particles.

diff --git a/Engine/Plugins/FX/Niagara/Shaders/Private/NiagaraDataInterfaceCollisionQuery.ush b/Engine/Plugins/FX/Niagara/Shaders/Private/NiagaraDataInterfaceCollisionQuery.ush
--- a/Engine/Plugins/FX/Niagara/Shaders/Private/NiagaraDataInterfaceCollisionQuery.ush
+++ b/Engine/Plugins/FX/Niagara/Shaders/Private/NiagaraDataInterfaceCollisionQuery.ush
@@ -112,14 +112,9 @@ void NDICollisionQuery_QueryScenePartialDepthGPU(in float3 In_SamplePos, in floa
 
 		FLWCVector3 WorldPosition = WorldPositionFromSceneDepth(ScreenPosition.xy, SceneDepth);
 
-#if NDICOLLISIONQUERY_USE_GBUFFER_NORMAL
-#if STRATA_ENABLED
-		uint2 ScreenPos = uint2(ScreenUV * View.BufferSizeAndInvSize.xy);
-		float3 WorldNormal = StrataPublic_GetWorldNormal(ScreenPos);
-#else
-		float3 WorldNormal = Texture2DSampleLevel(SceneTexturesStruct.GBufferATexture, SceneTexturesStruct_GBufferATextureSampler, ScreenUV, 0).xyz * 2.0 - 1.0;
-#endif
-#else
+		// We cannot sample the normal from the GBuffer when using partial depth collision queries.
+		// => It must be derived form the partial depth buffer itself to get a match.
+
 		float CollisionDepthBounds = 500.0f;
 		float SceneDepth0 = NDICollisionQuery_GetScenePartialDepth(ScreenUV + float2(View.BufferSizeAndInvSize.z, 0.0));
 		float SceneDepth1 = NDICollisionQuery_GetScenePartialDepth(ScreenUV + float2(0.0, View.BufferSizeAndInvSize.w));
@@ -135,7 +130,6 @@ void NDICollisionQuery_QueryScenePartialDepthGPU(in float3 In_SamplePos, in floa
 		FLWCVector3 WorldPosition0 = WorldPositionFromSceneDepth(ScreenPosition.xy + float2(2 * View.ViewSizeAndInvSize.z, 0.0), SceneDepth0);
 		FLWCVector3 WorldPosition1 = WorldPositionFromSceneDepth(ScreenPosition.xy - float2(0.0, 2 * View.ViewSizeAndInvSize.w), SceneDepth1);
 		float3 WorldNormal = normalize(cross(LWCSubtract(WorldPosition0, WorldPosition).Offset, LWCSubtract(WorldPosition1, WorldPosition).Offset));
-#endif
 
 		// convert LWC position back to simulation space
 		LWCSetTile(WorldPosition, LWCGetTile(WorldPosition) - In_LWCTile);
