Fixed non unity compile issue

diff --git a/Engine/Plugins/FX/Niagara/Source/NiagaraEditorWidgets/Private/SNiagaraOverviewGraph.cpp b/Engine/Plugins/FX/Niagara/Source/NiagaraEditorWidgets/Private/SNiagaraOverviewGraph.cpp
--- a/Engine/Plugins/FX/Niagara/Source/NiagaraEditorWidgets/Private/SNiagaraOverviewGraph.cpp
+++ b/Engine/Plugins/FX/Niagara/Source/NiagaraEditorWidgets/Private/SNiagaraOverviewGraph.cpp
@@ -18,6 +18,7 @@
 #include "NiagaraOverviewNode.h"
 #include "NiagaraSystem.h"
 #include "NiagaraSystemEditorData.h"
+#include "NiagaraStackEditorData.h"
 #include "SNiagaraAssetPickerList.h"
 #include "SNiagaraOverviewGraphTitleBar.h"
 #include "SNiagaraStack.h"
