Fixed non unity compile issue

diff --git a/Engine/Plugins/FX/Niagara/Source/NiagaraEditorWidgets/Private/NiagaraStackCommandContext.cpp b/Engine/Plugins/FX/Niagara/Source/NiagaraEditorWidgets/Private/NiagaraStackCommandContext.cpp
--- a/Engine/Plugins/FX/Niagara/Source/NiagaraEditorWidgets/Private/NiagaraStackCommandContext.cpp
+++ b/Engine/Plugins/FX/Niagara/Source/NiagaraEditorWidgets/Private/NiagaraStackCommandContext.cpp
@@ -6,9 +6,9 @@
 #include "NiagaraEditorCommon.h"
 #include "NiagaraEditorUtilities.h"
 #include "NiagaraSystemEditorData.h"
+#include "NiagaraStackEditorData.h"
 #include "ViewModels/Stack/NiagaraStackClipboardUtilities.h"
 #include "ViewModels/Stack/NiagaraStackEntry.h"
-
 #include "Framework/Notifications/NotificationManager.h"
 #include "Logging/LogMacros.h"
 #include "Widgets/Notifications/SNotificationList.h"
