- Fix localization issue

#rnx

diff --git a/Engine/Plugins/FX/Niagara/Source/NiagaraEditor/Private/Customizations/NiagaraParameterBindingCustomization.cpp b/Engine/Plugins/FX/Niagara/Source/NiagaraEditor/Private/Customizations/NiagaraParameterBindingCustomization.cpp
--- a/Engine/Plugins/FX/Niagara/Source/NiagaraEditor/Private/Customizations/NiagaraParameterBindingCustomization.cpp
+++ b/Engine/Plugins/FX/Niagara/Source/NiagaraEditor/Private/Customizations/NiagaraParameterBindingCustomization.cpp
@@ -438,7 +438,7 @@ void FNiagaraParameterBindingCustomization::OnValueChanged() const
 		return;
 	}
 
-	FScopedTransaction Transaction(LOCTEXT("ChangeParameterBinding", " Change Parameter Value"));
+	FScopedTransaction Transaction(LOCTEXT("ChangeParameterValue", " Change Parameter Value"));
 
 	TArray<UObject*> Objects;
 	PropertyHandle->GetOuterObjects(Objects);
