Fixed file includes

diff --git a/Engine/Plugins/FX/Niagara/Source/Niagara/Private/NiagaraDataInterfaceRenderTargetVolume.cpp b/Engine/Plugins/FX/Niagara/Source/Niagara/Private/NiagaraDataInterfaceRenderTargetVolume.cpp
--- a/Engine/Plugins/FX/Niagara/Source/Niagara/Private/NiagaraDataInterfaceRenderTargetVolume.cpp
+++ b/Engine/Plugins/FX/Niagara/Source/Niagara/Private/NiagaraDataInterfaceRenderTargetVolume.cpp
@@ -15,6 +15,7 @@
 #include "NiagaraSystemInstance.h"
 #include "NiagaraStats.h"
 #include "NiagaraShaderParametersBuilder.h"
+#include "NiagaraShader.h"
 #include "NiagaraGpuComputeDebugInterface.h"
 #include "VolumeCache.h"
 #include "RenderGraphBuilder.h"
diff --git a/Engine/Plugins/FX/Niagara/Source/Niagara/Private/NiagaraRenderer.cpp b/Engine/Plugins/FX/Niagara/Source/Niagara/Private/NiagaraRenderer.cpp
--- a/Engine/Plugins/FX/Niagara/Source/Niagara/Private/NiagaraRenderer.cpp
+++ b/Engine/Plugins/FX/Niagara/Source/Niagara/Private/NiagaraRenderer.cpp
@@ -7,6 +7,7 @@
 #include "NiagaraDataSet.h"
 #include "NiagaraSceneProxy.h"
 #include "NiagaraStats.h"
+#include "NiagaraShader.h"
 #include "DynamicBufferAllocator.h"
 #include "NiagaraComputeExecutionContext.h"
 #include "NiagaraGPUSortInfo.h"
