Preparing more engine+plugins callsites for upcoming TArray/TArrayView mutable conversion restrictions

diff --git a/Engine/Plugins/Experimental/ChaosClothAssetEditor/Source/ChaosClothAssetEditor/Private/ChaosClothAsset/ClothEditorToolkit.cpp b/Engine/Plugins/Experimental/ChaosClothAssetEditor/Source/ChaosClothAssetEditor/Private/ChaosClothAsset/ClothEditorToolkit.cpp
--- a/Engine/Plugins/Experimental/ChaosClothAssetEditor/Source/ChaosClothAssetEditor/Private/ChaosClothAsset/ClothEditorToolkit.cpp
+++ b/Engine/Plugins/Experimental/ChaosClothAssetEditor/Source/ChaosClothAssetEditor/Private/ChaosClothAsset/ClothEditorToolkit.cpp
@@ -570,7 +570,7 @@ void FChaosClothAssetEditorToolkit::UnregisterTabSpawners(const TSharedRef<FTabM
 UChaosClothAsset* FChaosClothAssetEditorToolkit::GetAsset() const
 {
 	TArray<TObjectPtr<UObject>> ObjectsToEdit;
-	OwningAssetEditor->GetObjectsToEdit(ObjectsToEdit);
+	OwningAssetEditor->GetObjectsToEdit(MutableView(ObjectsToEdit));
 	
 	UObject* ObjectToEdit = nullptr;
 	if (ensure(ObjectsToEdit.Num() == 1))
@@ -681,7 +681,7 @@ TSharedRef<SDockTab> FChaosClothAssetEditorToolkit::SpawnTab_NodeDetails(const F
 void FChaosClothAssetEditorToolkit::InitDetailsViewPanel()
 {
 	TArray<TObjectPtr<UObject>> ObjectsToEdit;
-	OwningAssetEditor->GetObjectsToEdit(ObjectsToEdit);
+	OwningAssetEditor->GetObjectsToEdit(MutableView(ObjectsToEdit));
 
 	if (ObjectsToEdit.Num() > 0)
 	{
@@ -970,7 +970,7 @@ void FChaosClothAssetEditorToolkit::OnClothAssetChanged()
 	ClothPreviewViewportClient->ClearSelectedComponents();
 
 	TArray<TObjectPtr<UObject>> ObjectsToEdit;
-	OwningAssetEditor->GetObjectsToEdit(ObjectsToEdit);
+	OwningAssetEditor->GetObjectsToEdit(MutableView(ObjectsToEdit));
 
 	UChaosClothAssetEditorMode* const ClothMode = CastChecked<UChaosClothAssetEditorMode>(EditorModeManager->GetActiveScriptableMode(UChaosClothAssetEditorMode::EM_ChaosClothAssetEditorModeId));
 
diff --git a/Engine/Plugins/Experimental/UserToolBoxCore/Source/UserToolBoxCore/Private/Editor/UTBTabEditor.cpp b/Engine/Plugins/Experimental/UserToolBoxCore/Source/UserToolBoxCore/Private/Editor/UTBTabEditor.cpp
--- a/Engine/Plugins/Experimental/UserToolBoxCore/Source/UserToolBoxCore/Private/Editor/UTBTabEditor.cpp
+++ b/Engine/Plugins/Experimental/UserToolBoxCore/Source/UserToolBoxCore/Private/Editor/UTBTabEditor.cpp
@@ -676,7 +676,7 @@ inline void FUTBTabEditor::CreateOrRebuildUTBTabTabWidget()
 		
 		TSharedPtr<SWrapBox> PanelWidget= SNew(SWrapBox)
 		.UseAllottedSize(true);
-		TArray<UUTBBaseCommand*>& CurrentSectionArray=Tab->GetOrCreateSection(CurrentSection)->Commands;
+		TArray<TObjectPtr<UUTBBaseCommand>>& CurrentSectionArray=Tab->GetOrCreateSection(CurrentSection)->Commands;
 		int Index=0;
 		for (UUTBBaseCommand* Data:CurrentSectionArray)
 		{
diff --git a/Engine/Plugins/FX/Niagara/Source/NiagaraEditor/Private/ViewModels/HierarchyEditor/NiagaraHierarchyViewModelBase.cpp b/Engine/Plugins/FX/Niagara/Source/NiagaraEditor/Private/ViewModels/HierarchyEditor/NiagaraHierarchyViewModelBase.cpp
--- a/Engine/Plugins/FX/Niagara/Source/NiagaraEditor/Private/ViewModels/HierarchyEditor/NiagaraHierarchyViewModelBase.cpp
+++ b/Engine/Plugins/FX/Niagara/Source/NiagaraEditor/Private/ViewModels/HierarchyEditor/NiagaraHierarchyViewModelBase.cpp
@@ -1217,7 +1217,7 @@ void FNiagaraHierarchySectionViewModel::OnDroppedOnInternal(TSharedPtr<FNiagaraH
 		int32 IndexOfThis = HierarchyViewModel->GetHierarchyRoot()->GetSectionData().Find(GetDataMutable<UNiagaraHierarchySection>());
 		int32 DraggedSectionIndex = HierarchyViewModel->GetHierarchyRoot()->GetSectionData().Find(DraggedSectionData);
 
-		TArray<UNiagaraHierarchySection*>& SectionData = HierarchyViewModel->GetHierarchyRoot()->GetSectionDataMutable();
+		TArray<TObjectPtr<UNiagaraHierarchySection>>& SectionData = HierarchyViewModel->GetHierarchyRoot()->GetSectionDataMutable();
 		int32 Count = SectionData.Num();
 
 		bool bDropSucceeded = false;
diff --git a/Engine/Plugins/FX/Niagara/Source/NiagaraEditor/Private/ViewModels/HierarchyEditor/NiagaraSummaryViewViewModel.cpp b/Engine/Plugins/FX/Niagara/Source/NiagaraEditor/Private/ViewModels/HierarchyEditor/NiagaraSummaryViewViewModel.cpp
--- a/Engine/Plugins/FX/Niagara/Source/NiagaraEditor/Private/ViewModels/HierarchyEditor/NiagaraSummaryViewViewModel.cpp
+++ b/Engine/Plugins/FX/Niagara/Source/NiagaraEditor/Private/ViewModels/HierarchyEditor/NiagaraSummaryViewViewModel.cpp
@@ -1515,7 +1515,7 @@ void UNiagaraSummaryViewViewModel::PrepareSourceItems(UNiagaraHierarchyRoot* Sou
 	if(UsageMap.Contains(ENiagaraScriptUsage::EmitterSpawnScript))
 	{
 		FText SectionName = FText::FromString("Emitter Spawn");
-		UNiagaraHierarchySection** FoundHierarchySection = SourceRootViewModel->GetDataMutable<UNiagaraHierarchyRoot>()->GetSectionDataMutable().FindByPredicate([SectionName](UNiagaraHierarchySection* Candidate)
+		auto* FoundHierarchySection = SourceRootViewModel->GetDataMutable<UNiagaraHierarchyRoot>()->GetSectionDataMutable().FindByPredicate([SectionName](UNiagaraHierarchySection* Candidate)
 		{
 			return Candidate->GetSectionNameAsText().EqualTo(SectionName);
 		});
@@ -1542,7 +1542,7 @@ void UNiagaraSummaryViewViewModel::PrepareSourceItems(UNiagaraHierarchyRoot* Sou
 	if(UsageMap.Contains(ENiagaraScriptUsage::EmitterUpdateScript))
 	{
 		FText SectionName = FText::FromString("Emitter Update");
-		UNiagaraHierarchySection** FoundHierarchySection = SourceRootViewModel->GetDataMutable<UNiagaraHierarchyRoot>()->GetSectionDataMutable().FindByPredicate([SectionName](UNiagaraHierarchySection* Candidate)
+		auto* FoundHierarchySection = SourceRootViewModel->GetDataMutable<UNiagaraHierarchyRoot>()->GetSectionDataMutable().FindByPredicate([SectionName](UNiagaraHierarchySection* Candidate)
 		{
 			return Candidate->GetSectionNameAsText().EqualTo(SectionName);
 		});
@@ -1569,7 +1569,7 @@ void UNiagaraSummaryViewViewModel::PrepareSourceItems(UNiagaraHierarchyRoot* Sou
 	if(UsageMap.Contains(ENiagaraScriptUsage::ParticleSpawnScript))
 	{
 		FText SectionName = FText::FromString("Particle Spawn");
-		UNiagaraHierarchySection** FoundHierarchySection = SourceRootViewModel->GetDataMutable<UNiagaraHierarchyRoot>()->GetSectionDataMutable().FindByPredicate([SectionName](UNiagaraHierarchySection* Candidate)
+		auto* FoundHierarchySection = SourceRootViewModel->GetDataMutable<UNiagaraHierarchyRoot>()->GetSectionDataMutable().FindByPredicate([SectionName](UNiagaraHierarchySection* Candidate)
 		{
 			return Candidate->GetSectionNameAsText().EqualTo(SectionName);
 		});
@@ -1596,7 +1596,7 @@ void UNiagaraSummaryViewViewModel::PrepareSourceItems(UNiagaraHierarchyRoot* Sou
 	if(UsageMap.Contains(ENiagaraScriptUsage::ParticleUpdateScript))
 	{
 		FText SectionName = FText::FromString("Particle Update");
-		UNiagaraHierarchySection** FoundHierarchySection = SourceRootViewModel->GetDataMutable<UNiagaraHierarchyRoot>()->GetSectionDataMutable().FindByPredicate([SectionName](UNiagaraHierarchySection* Candidate)
+		auto* FoundHierarchySection = SourceRootViewModel->GetDataMutable<UNiagaraHierarchyRoot>()->GetSectionDataMutable().FindByPredicate([SectionName](UNiagaraHierarchySection* Candidate)
 		{
 			return Candidate->GetSectionNameAsText().EqualTo(SectionName);
 		});
@@ -1623,7 +1623,7 @@ void UNiagaraSummaryViewViewModel::PrepareSourceItems(UNiagaraHierarchyRoot* Sou
 	if(HierarchyEventHandlers.Num() > 0)
 	{
 		FText SectionName = FText::FromString("Events");
-		UNiagaraHierarchySection** FoundHierarchySection = SourceRootViewModel->GetDataMutable<UNiagaraHierarchyRoot>()->GetSectionDataMutable().FindByPredicate([SectionName](UNiagaraHierarchySection* Candidate)
+		auto* FoundHierarchySection = SourceRootViewModel->GetDataMutable<UNiagaraHierarchyRoot>()->GetSectionDataMutable().FindByPredicate([SectionName](UNiagaraHierarchySection* Candidate)
 		{
 			return Candidate->GetSectionNameAsText().EqualTo(SectionName);
 		});
@@ -1650,7 +1650,7 @@ void UNiagaraSummaryViewViewModel::PrepareSourceItems(UNiagaraHierarchyRoot* Sou
 	if(HierarchySimStages.Num() > 0)
 	{
 		FText SectionName = FText::FromString("Sim Stages");
-		UNiagaraHierarchySection** FoundHierarchySection = SourceRootViewModel->GetDataMutable<UNiagaraHierarchyRoot>()->GetSectionDataMutable().FindByPredicate([SectionName](UNiagaraHierarchySection* Candidate)
+		auto* FoundHierarchySection = SourceRootViewModel->GetDataMutable<UNiagaraHierarchyRoot>()->GetSectionDataMutable().FindByPredicate([SectionName](UNiagaraHierarchySection* Candidate)
 		{
 			return Candidate->GetSectionNameAsText().EqualTo(SectionName);
 		});
@@ -1677,7 +1677,7 @@ void UNiagaraSummaryViewViewModel::PrepareSourceItems(UNiagaraHierarchyRoot* Sou
 	if(HierarchyRenderers.Num() > 0)
 	{
 		FText SectionName = FText::FromString("Renderers");
-		UNiagaraHierarchySection** FoundHierarchySection = SourceRootViewModel->GetDataMutable<UNiagaraHierarchyRoot>()->GetSectionDataMutable().FindByPredicate([SectionName](UNiagaraHierarchySection* Candidate)
+		auto* FoundHierarchySection = SourceRootViewModel->GetDataMutable<UNiagaraHierarchyRoot>()->GetSectionDataMutable().FindByPredicate([SectionName](UNiagaraHierarchySection* Candidate)
 		{
 			return Candidate->GetSectionNameAsText().EqualTo(SectionName);
 		});
diff --git a/Engine/Plugins/FX/Niagara/Source/NiagaraEditor/Public/ViewModels/HierarchyEditor/NiagaraHierarchyViewModelBase.h b/Engine/Plugins/FX/Niagara/Source/NiagaraEditor/Public/ViewModels/HierarchyEditor/NiagaraHierarchyViewModelBase.h
--- a/Engine/Plugins/FX/Niagara/Source/NiagaraEditor/Public/ViewModels/HierarchyEditor/NiagaraHierarchyViewModelBase.h
+++ b/Engine/Plugins/FX/Niagara/Source/NiagaraEditor/Public/ViewModels/HierarchyEditor/NiagaraHierarchyViewModelBase.h
@@ -229,7 +229,7 @@ public:
 	virtual ~UNiagaraHierarchyRoot() override {}
 
 	const TArray<UNiagaraHierarchySection*>& GetSectionData() const { return Sections; }
-	TArray<UNiagaraHierarchySection*>& GetSectionDataMutable() { return Sections; }
+	TArray<TObjectPtr<UNiagaraHierarchySection>>& GetSectionDataMutable() { return Sections; }
 
 	TSet<FName> GetSections() const;
 	int32 GetSectionIndex(UNiagaraHierarchySection* Section) const;
diff --git a/Engine/Plugins/MovieScene/MovieRenderPipeline/Source/MovieRenderPipelineCore/Private/Graph/MovieGraphNode.cpp b/Engine/Plugins/MovieScene/MovieRenderPipeline/Source/MovieRenderPipelineCore/Private/Graph/MovieGraphNode.cpp
--- a/Engine/Plugins/MovieScene/MovieRenderPipeline/Source/MovieRenderPipelineCore/Private/Graph/MovieGraphNode.cpp
+++ b/Engine/Plugins/MovieScene/MovieRenderPipeline/Source/MovieRenderPipelineCore/Private/Graph/MovieGraphNode.cpp
@@ -119,8 +119,8 @@ void UMovieGraphNode::UpdatePins()
 		// return bAppliedEdgeChanges ? EPCGChangeType::Edge : None) | (bChangedPins ? EPCGChangeType::Node : None
 	};
 
-	UpdatePins(InputPins, InputPinProperties);
-	UpdatePins(OutputPins, OutputPinProperties);
+	UpdatePins(MutableView(InputPins), InputPinProperties);
+	UpdatePins(MutableView(OutputPins), OutputPinProperties);
 
 	OnNodeChangedDelegate.Broadcast(this);
 }
@@ -237,4 +237,4 @@ TArray<FMovieGraphPinProperties> UMovieGraphSettingNode::GetOutputPinProperties(
 	TArray<FMovieGraphPinProperties> Properties;
 	Properties.Add(FMovieGraphPinProperties::MakeBranchProperties());
 	return Properties;
-}
\ No newline at end of file
+}
diff --git a/Engine/Plugins/MovieScene/MovieRenderPipeline/Source/MovieRenderPipelineCore/Public/Graph/MoviePipelineRenderLayerSubsystem.h b/Engine/Plugins/MovieScene/MovieRenderPipeline/Source/MovieRenderPipelineCore/Public/Graph/MoviePipelineRenderLayerSubsystem.h
--- a/Engine/Plugins/MovieScene/MovieRenderPipeline/Source/MovieRenderPipelineCore/Public/Graph/MoviePipelineRenderLayerSubsystem.h
+++ b/Engine/Plugins/MovieScene/MovieRenderPipeline/Source/MovieRenderPipelineCore/Public/Graph/MoviePipelineRenderLayerSubsystem.h
@@ -280,7 +280,7 @@ public:
 	bool AddRenderLayer(UMoviePipelineRenderLayer* RenderLayer);
 
 	UFUNCTION(BlueprintCallable, Category = "Render Layers")
-	TArray<UMoviePipelineRenderLayer*>& GetRenderLayers() { return RenderLayers; }
+	const TArray<UMoviePipelineRenderLayer*>& GetRenderLayers() { return RenderLayers; }
 
 	UFUNCTION(BlueprintCallable, Category = "Render Layers")
 	void RemoveRenderLayer(const FString& RenderLayerName);
@@ -350,4 +350,4 @@ private:
 	/** A modifier used for visualization purposes (to show collections used with the modifier). */
 	UPROPERTY(Transient)
 	TObjectPtr<UMoviePipelineVisibilityModifier> VisualizationModifier_VisibleCollections = nullptr;
-};
\ No newline at end of file
+};
diff --git a/Engine/Plugins/Runtime/RigVM/Source/RigVMDeveloper/Private/RigVMModel/RigVMController.cpp b/Engine/Plugins/Runtime/RigVM/Source/RigVMDeveloper/Private/RigVMModel/RigVMController.cpp
--- a/Engine/Plugins/Runtime/RigVM/Source/RigVMDeveloper/Private/RigVMModel/RigVMController.cpp
+++ b/Engine/Plugins/Runtime/RigVM/Source/RigVMDeveloper/Private/RigVMModel/RigVMController.cpp
@@ -6318,7 +6318,7 @@ bool URigVMController::RemoveNodes(TArray<URigVMNode*> InNodes, bool bSetupUndoR
 								TArray<FRigVMExternalVariable> ExternalVariablesWithoutVariableNode;
 								{
 									URigVMGraph* EditedGraph = InNode->GetGraph();
-									TGuardValue<TArray<URigVMNode*>> TemporaryRemoveNodes(EditedGraph->Nodes, TArray<URigVMNode*>());
+									TGuardValue<TArray<TObjectPtr<URigVMNode>>> TemporaryRemoveNodes(EditedGraph->Nodes, TArray<TObjectPtr<URigVMNode>>());
 									ExternalVariablesWithoutVariableNode = EditedGraph->GetExternalVariables();
 								}
 
diff --git a/Engine/Source/Editor/BehaviorTreeEditor/Private/BehaviorTreeGraph.cpp b/Engine/Source/Editor/BehaviorTreeEditor/Private/BehaviorTreeGraph.cpp
--- a/Engine/Source/Editor/BehaviorTreeEditor/Private/BehaviorTreeGraph.cpp
+++ b/Engine/Source/Editor/BehaviorTreeEditor/Private/BehaviorTreeGraph.cpp
@@ -919,7 +919,7 @@ void UBehaviorTreeGraph::CreateBTFromGraph(UBehaviorTreeGraphNode* RootEdNode)
 		uint16 DummyExecutionIndex = MAX_uint16;
 		BTGraphHelpers::CollectDecorators(BTAsset
 			, RootEdNode
-			, BTAsset->RootDecorators
+			, MutableView(BTAsset->RootDecorators)
 			, BTAsset->RootDecoratorOps
 			, /*bInitializeNodes*/false
 			, /*RootNode*/nullptr
diff --git a/Engine/Source/Editor/UnrealEd/Private/Elements/Component/ComponentElementEditorCopyAndPaste.cpp b/Engine/Source/Editor/UnrealEd/Private/Elements/Component/ComponentElementEditorCopyAndPaste.cpp
--- a/Engine/Source/Editor/UnrealEd/Private/Elements/Component/ComponentElementEditorCopyAndPaste.cpp
+++ b/Engine/Source/Editor/UnrealEd/Private/Elements/Component/ComponentElementEditorCopyAndPaste.cpp
@@ -277,7 +277,7 @@ void FComponentElementEditorPasteImporter::Import(FContext& Context)
 	else if (PotentialSelection.Actor)
 	{
 		FString Text(Context.Text);
-		GUnrealEd->PasteComponents(ImportedComponents, PotentialSelection.Actor, true, &Text);
+		GUnrealEd->PasteComponents(MutableView(ImportedComponents), PotentialSelection.Actor, true, &Text);
 	}
 }
 
diff --git a/Engine/Source/Editor/UnrealEd/Private/Kismet2/BlueprintEditorUtils.cpp b/Engine/Source/Editor/UnrealEd/Private/Kismet2/BlueprintEditorUtils.cpp
--- a/Engine/Source/Editor/UnrealEd/Private/Kismet2/BlueprintEditorUtils.cpp
+++ b/Engine/Source/Editor/UnrealEd/Private/Kismet2/BlueprintEditorUtils.cpp
@@ -6903,10 +6903,10 @@ void FBlueprintEditorUtils::CleanNullGraphReferencesInArray(UBlueprint* Blueprin
 
 void FBlueprintEditorUtils::PurgeNullGraphs(UBlueprint* Blueprint)
 {
-	CleanNullGraphReferencesInArray(Blueprint, Blueprint->UbergraphPages);
-	CleanNullGraphReferencesInArray(Blueprint, Blueprint->FunctionGraphs);
-	CleanNullGraphReferencesInArray(Blueprint, Blueprint->DelegateSignatureGraphs);
-	CleanNullGraphReferencesInArray(Blueprint, Blueprint->MacroGraphs);
+	CleanNullGraphReferencesInArray(Blueprint, MutableView(Blueprint->UbergraphPages));
+	CleanNullGraphReferencesInArray(Blueprint, MutableView(Blueprint->FunctionGraphs));
+	CleanNullGraphReferencesInArray(Blueprint, MutableView(Blueprint->DelegateSignatureGraphs));
+	CleanNullGraphReferencesInArray(Blueprint, MutableView(Blueprint->MacroGraphs));
 }
 
 struct FConformCallsToParentFunctionUtils
