Added python script support for deprecated module changes

diff --git a/Engine/Plugins/FX/Niagara/Source/Niagara/Classes/NiagaraScript.h b/Engine/Plugins/FX/Niagara/Source/Niagara/Classes/NiagaraScript.h
--- a/Engine/Plugins/FX/Niagara/Source/Niagara/Classes/NiagaraScript.h
+++ b/Engine/Plugins/FX/Niagara/Source/Niagara/Classes/NiagaraScript.h
@@ -9,7 +9,6 @@
 #include "NiagaraCommon.h"
 #include "NiagaraScriptBase.h"
 #include "NiagaraShared.h"
-#include "NiagaraShader.h"
 #include "NiagaraParameters.h"
 #include "NiagaraDataSet.h"
 #include "NiagaraScriptExecutionParameterStore.h"
@@ -654,8 +653,24 @@ public:
 	UPROPERTY(EditAnywhere, Category = Script)
 	TObjectPtr<UNiagaraScript> DeprecationRecommendation;
 
-	/* Custom logic to convert the contents of an existing script assignment to this script.*/
+	/* If true then a python script will be executed when changing from this script to the selected deprectation recommendation. This allows the current script to transfer its inputs to the new script. */
 	UPROPERTY(EditAnywhere, Category = Script)
+	bool bUsePythonScriptConversion = false;
+
+	/** Reference to a python script that is executed when the user updates from a previous version to this version. */
+	UPROPERTY(EditAnywhere, Category = Script, meta=(EditCondition="bUsePythonScriptConversion", EditConditionHides))
+	ENiagaraPythonUpdateScriptReference ConversionScriptExecution = ENiagaraPythonUpdateScriptReference::ScriptAsset;
+
+	/** Python script to run when converting this script to the recommended deprecation update script. */
+	UPROPERTY(EditAnywhere, Category = Script, meta=(EditCondition="bUsePythonScriptConversion && ConversionScriptExecution == ENiagaraPythonUpdateScriptReference::DirectTextEntry", EditConditionHides))
+	FString PythonConversionScript;
+
+	/** Asset reference to a python script to run when converting this script to the recommended deprecation update script. */
+	UPROPERTY(EditAnywhere, Category = Script, meta=(EditCondition="bUsePythonScriptConversion && ConversionScriptExecution == ENiagaraPythonUpdateScriptReference::ScriptAsset", EditConditionHides))
+	FFilePath ConversionScriptAsset;
+
+	/* Custom logic to convert the contents of an existing script assignment to this script.*/
+	UPROPERTY(EditAnywhere, Category = Script, meta=(EditCondition="!bUsePythonScriptConversion", EditConditionHides))
 	TSubclassOf<UNiagaraConvertInPlaceUtilityBase> ConversionUtility;
 
 	/** Is this script experimental and less supported? */
diff --git a/Engine/Plugins/FX/Niagara/Source/Niagara/Private/NiagaraGpuComputeDispatch.h b/Engine/Plugins/FX/Niagara/Source/Niagara/Private/NiagaraGpuComputeDispatch.h
--- a/Engine/Plugins/FX/Niagara/Source/Niagara/Private/NiagaraGpuComputeDispatch.h
+++ b/Engine/Plugins/FX/Niagara/Source/Niagara/Private/NiagaraGpuComputeDispatch.h
@@ -17,6 +17,7 @@ the same VectorVM byte code / compute shader code
 #include "RenderGraphUtils.h"
 #include "RHIResources.h"
 #include "RendererInterface.h"
+#include "RendererUtils.h"
 
 enum class EGPUSortFlags : uint32;
 struct FNiagaraGPUSortInfo;
diff --git a/Engine/Plugins/FX/Niagara/Source/NiagaraEditor/Private/NiagaraEditorUtilities.cpp b/Engine/Plugins/FX/Niagara/Source/NiagaraEditor/Private/NiagaraEditorUtilities.cpp
--- a/Engine/Plugins/FX/Niagara/Source/NiagaraEditor/Private/NiagaraEditorUtilities.cpp
+++ b/Engine/Plugins/FX/Niagara/Source/NiagaraEditor/Private/NiagaraEditorUtilities.cpp
@@ -3230,11 +3230,9 @@ void FNiagaraEditorUtilities::RefreshAllScriptsFromExternalChanges(FRefreshAllSc
 	FNiagaraEditorUtilities::CompileExistingEmitters(AffectedEmitters);
 }
 
-TArray<UNiagaraPythonScriptModuleInput*> GetFunctionCallInputs(const FNiagaraScriptVersionUpgradeContext& UpgradeContext)
+TArray<UNiagaraPythonScriptModuleInput*> GetFunctionCallInputs(UNiagaraClipboardContent* ClipboardContent)
 {
 	TArray<UNiagaraPythonScriptModuleInput*> ScriptInputs;
-	UNiagaraClipboardContent* ClipboardContent = UNiagaraClipboardContent::Create();
-	UpgradeContext.CreateClipboardCallback(ClipboardContent);
 	for (const UNiagaraClipboardFunctionInput* FunctionInput : ClipboardContent->FunctionInputs)
 	{
 		UNiagaraPythonScriptModuleInput* ScriptInput = NewObject<UNiagaraPythonScriptModuleInput>();
@@ -3244,6 +3242,23 @@ TArray<UNiagaraPythonScriptModuleInput*> GetFunctionCallInputs(const FNiagaraScr
 	return ScriptInputs;
 }
 
+TArray<UNiagaraPythonScriptModuleInput*> GetFunctionCallInputs(const FNiagaraScriptVersionUpgradeContext& UpgradeContext)
+{
+	UNiagaraClipboardContent* ClipboardContent = UNiagaraClipboardContent::Create();
+	UpgradeContext.CreateClipboardCallback(ClipboardContent);
+	return GetFunctionCallInputs(ClipboardContent);
+}
+
+void AddStackWarning(const FText& ToAdd, bool& bLoggedWarning, FText& OutWarnings)
+{
+	if (!bLoggedWarning)
+	{
+		OutWarnings = LOCTEXT("NewWarningHeader", "Python conversion script:\n");
+	}
+	bLoggedWarning = true;
+	OutWarnings = FText::Format(FText::FromString("{0}  * {1}\n"), OutWarnings, ToAdd);
+}
+
 void AddStackWarning(const FNiagaraAssetVersion& FromVersion, const FNiagaraAssetVersion& ToVersion, const FString& ToAdd, bool& bLoggedWarning, FString& OutWarnings)
 {
 	if (!bLoggedWarning)
@@ -3263,6 +3278,84 @@ const FString PythonUpgradeScriptStub = TEXT(
 	"### End User Script ###\n"
 	"upgrade_context.cancelled_by_python_error = False\n");
 
+UNiagaraClipboardContent* FNiagaraEditorUtilities::RunPythonConversionScript(FVersionedNiagaraScriptData& NewScriptVersionData, UNiagaraClipboardContent* NewScriptInputs, FVersionedNiagaraScriptData& OldScriptVersionData, UNiagaraClipboardContent* OldScriptInputs, FText& OutWarnings)
+{
+	UUpgradeNiagaraScriptResults* Results = NewObject<UUpgradeNiagaraScriptResults>();
+	
+	FString PythonScript;
+	if (OldScriptVersionData.ConversionScriptExecution == ENiagaraPythonUpdateScriptReference::DirectTextEntry)
+	{
+		PythonScript = OldScriptVersionData.PythonConversionScript;
+	}
+	else if (OldScriptVersionData.ConversionScriptExecution == ENiagaraPythonUpdateScriptReference::ScriptAsset && !OldScriptVersionData.ConversionScriptAsset.FilePath.IsEmpty())
+	{
+		FFileHelper::LoadFileToString(PythonScript, *OldScriptVersionData.ConversionScriptAsset.FilePath);
+	}
+
+	if (PythonScript.IsEmpty() || NewScriptInputs == nullptr ||OldScriptInputs == nullptr)
+	{
+		return nullptr;
+	}
+
+	// set up script context
+	bool bLoggedWarning = false;
+	Results->OldInputs = GetFunctionCallInputs(OldScriptInputs);
+	Results->NewInputs = GetFunctionCallInputs(NewScriptInputs);
+	Results->Init();
+	
+	// save python script to a temp file to execute
+	FString TempScriptFile = FPaths::CreateTempFilename(*FPaths::ProjectIntermediateDir(), TEXT("ScriptConversion-"), TEXT(".py"));
+	IPlatformFile& PlatformFile = FPlatformFileManager::Get().GetPlatformFile();
+	ON_SCOPE_EXIT
+	{
+		// Delete temp script file
+		if (GNiagaraDeletePythonFilesOnError || !Results->bCancelledByPythonError)
+		{
+			PlatformFile.DeleteFile(*TempScriptFile);
+		}
+	};
+	if (!FFileHelper::SaveStringToFile(FString::Format(*PythonUpgradeScriptStub, {Results->GetPathName(), PythonScript}), *TempScriptFile))
+	{
+		UE_LOG(LogNiagaraEditor, Error, TEXT("Unable to save python script to file %s"), *TempScriptFile);
+		AddStackWarning(LOCTEXT("CreatePythonError", "Cannot create python script file!"), bLoggedWarning, OutWarnings);
+		return nullptr;
+	}
+
+	Results->bCancelledByPythonError = true;
+	FPythonCommandEx PythonCommand = FPythonCommandEx();
+	PythonCommand.ExecutionMode = EPythonCommandExecutionMode::ExecuteFile;
+	PythonCommand.Command = TempScriptFile;
+
+	// execute python script
+	IPythonScriptPlugin::Get()->ExecPythonCommandEx(PythonCommand);
+
+	UNiagaraClipboardContent* ResultClipboard = UNiagaraClipboardContent::Create();
+	if (Results->bCancelledByPythonError)
+	{
+		UE_LOG(LogNiagaraEditor, Error, TEXT("%s\n\nPython script:\n%s\nTo keep the intermediate script around, set fx.Niagara.DeletePythonFilesOnError to 0."), *PythonCommand.CommandResult, *PythonCommand.Command);
+		AddStackWarning(LOCTEXT("PythonExecutionError", "Python script ended with error, check the log!"), bLoggedWarning, OutWarnings);
+	}
+	else
+	{
+		for (UNiagaraPythonScriptModuleInput* ModuleInput : Results->NewInputs)
+		{
+			ResultClipboard->FunctionInputs.Add(ModuleInput->Input);	
+		}
+	}
+	if (PythonCommand.LogOutput.Num() > 0)
+	{
+		for (FPythonLogOutputEntry& Entry : PythonCommand.LogOutput)
+		{
+			AddStackWarning(FText::FromString(Entry.Output), bLoggedWarning, OutWarnings);
+		}
+	}
+	if (ResultClipboard->FunctionInputs.Num() > 0)
+	{
+		return ResultClipboard;
+	}
+	return nullptr;
+}
+
 void FNiagaraEditorUtilities::RunPythonUpgradeScripts(UNiagaraNodeFunctionCall* SourceNode,	const TArray<FVersionedNiagaraScriptData*>& UpgradeVersionData, const FNiagaraScriptVersionUpgradeContext& UpgradeContext, FString& OutWarnings)
 {
 	UUpgradeNiagaraScriptResults* Results = NewObject<UUpgradeNiagaraScriptResults>();
diff --git a/Engine/Plugins/FX/Niagara/Source/NiagaraEditor/Private/NiagaraHlslTranslator.cpp b/Engine/Plugins/FX/Niagara/Source/NiagaraEditor/Private/NiagaraHlslTranslator.cpp
--- a/Engine/Plugins/FX/Niagara/Source/NiagaraEditor/Private/NiagaraHlslTranslator.cpp
+++ b/Engine/Plugins/FX/Niagara/Source/NiagaraEditor/Private/NiagaraHlslTranslator.cpp
@@ -39,6 +39,7 @@
 #include "Misc/FileHelper.h"
 #include "Misc/CoreMiscDefines.h"
 #include "NiagaraDataInterfaceUtilities.h"
+#include "NiagaraShader.h"
 
 #define LOCTEXT_NAMESPACE "NiagaraCompiler"
 
diff --git a/Engine/Plugins/FX/Niagara/Source/NiagaraEditor/Private/UpgradeNiagaraScriptResults.cpp b/Engine/Plugins/FX/Niagara/Source/NiagaraEditor/Private/UpgradeNiagaraScriptResults.cpp
--- a/Engine/Plugins/FX/Niagara/Source/NiagaraEditor/Private/UpgradeNiagaraScriptResults.cpp
+++ b/Engine/Plugins/FX/Niagara/Source/NiagaraEditor/Private/UpgradeNiagaraScriptResults.cpp
@@ -300,6 +300,25 @@ void UUpgradeNiagaraScriptResults::SetLinkedInput(const FString& InputName, FStr
 	}
 }
 
+void UUpgradeNiagaraScriptResults::SetNewInput(const FString& InputName, UNiagaraPythonScriptModuleInput* Value)
+{
+	for (int i = 0; i < NewInputs.Num(); i++)
+	{
+		UNiagaraPythonScriptModuleInput* ModuleInput = NewInputs[i];
+		UNiagaraClipboardFunctionInput* FunctionInput = const_cast<UNiagaraClipboardFunctionInput*>(ModuleInput->Input.Get());
+		if (Value && FunctionInput && FunctionInput->InputName == InputName && FunctionInput->InputType == Value->Input->InputType)
+		{
+			FunctionInput->Data = Value->Input->Data;
+			FunctionInput->Dynamic = Value->Input->Dynamic;
+			FunctionInput->Expression = Value->Input->Expression;
+			FunctionInput->Linked = Value->Input->Linked;
+			FunctionInput->Local = Value->Input->Local;
+			FunctionInput->ValueMode = Value->Input->ValueMode;
+			return;
+		}
+	}
+}
+
 UNiagaraPythonScriptModuleInput* UUpgradeNiagaraScriptResults::GetNewInput(const FName& InputName) const
 {
 	for (UNiagaraPythonScriptModuleInput* ModuleInput : NewInputs)
diff --git a/Engine/Plugins/FX/Niagara/Source/NiagaraEditor/Private/ViewModels/NiagaraScriptStatsViewModel.cpp b/Engine/Plugins/FX/Niagara/Source/NiagaraEditor/Private/ViewModels/NiagaraScriptStatsViewModel.cpp
--- a/Engine/Plugins/FX/Niagara/Source/NiagaraEditor/Private/ViewModels/NiagaraScriptStatsViewModel.cpp
+++ b/Engine/Plugins/FX/Niagara/Source/NiagaraEditor/Private/ViewModels/NiagaraScriptStatsViewModel.cpp
@@ -11,6 +11,7 @@
 
 #include "Async/Async.h"
 #include "DataDrivenShaderPlatformInfo.h"
+#include "NiagaraShader.h"
 #include "Styling/AppStyle.h"
 #include "Framework/MultiBox/MultiBoxBuilder.h"
 #include "Interfaces/ITargetPlatformManagerModule.h"
diff --git a/Engine/Plugins/FX/Niagara/Source/NiagaraEditor/Private/ViewModels/Stack/NiagaraStackFunctionInput.cpp b/Engine/Plugins/FX/Niagara/Source/NiagaraEditor/Private/ViewModels/Stack/NiagaraStackFunctionInput.cpp
--- a/Engine/Plugins/FX/Niagara/Source/NiagaraEditor/Private/ViewModels/Stack/NiagaraStackFunctionInput.cpp
+++ b/Engine/Plugins/FX/Niagara/Source/NiagaraEditor/Private/ViewModels/Stack/NiagaraStackFunctionInput.cpp
@@ -8,7 +8,6 @@
 #include "Editor.h"
 #include "INiagaraEditorTypeUtilities.h"
 #include "NiagaraClipboard.h"
-#include "NiagaraComponent.h"
 #include "NiagaraConstants.h"
 #include "NiagaraConvertInPlaceUtilityBase.h"
 #include "NiagaraEditorModule.h"
@@ -29,7 +28,6 @@
 #include "ViewModels/NiagaraScriptGraphViewModel.h"
 #include "NiagaraScriptMergeManager.h"
 #include "NiagaraScriptSource.h"
-#include "NiagaraScriptVariable.h"
 #include "NiagaraSettings.h"
 #include "ViewModels/NiagaraSystemScriptViewModel.h"
 #include "ScopedTransaction.h"
@@ -37,7 +35,6 @@
 #include "Framework/Notifications/NotificationManager.h"
 #include "Materials/Material.h"
 #include "Materials/MaterialExpressionDynamicParameter.h"
-#include "Toolkits/NiagaraSystemToolkit.h"
 #include "UObject/StructOnScope.h"
 #include "ViewModels/NiagaraEmitterViewModel.h"
 #include "ViewModels/NiagaraPlaceholderDataInterfaceManager.h"
@@ -2374,8 +2371,9 @@ void UNiagaraStackFunctionInput::ReassignDynamicInputScript(UNiagaraScript* Dyna
 
 		UNiagaraClipboardContent* OldClipboardContent = nullptr;
 		UNiagaraScript* OldScript = InputValues.DynamicNode->FunctionScript;
-		FVersionedNiagaraScriptData* ScriptData = DynamicInputScript->GetLatestScriptData();
-		if (ScriptData->ConversionUtility != nullptr)
+		FVersionedNiagaraScriptData* OldScriptData = InputValues.DynamicNode->GetScriptData();
+		FVersionedNiagaraScriptData* NewScriptData = DynamicInputScript->GetLatestScriptData();
+		if (NewScriptData->ConversionUtility || OldScriptData->bUsePythonScriptConversion)
 		{
 			OldClipboardContent = UNiagaraClipboardContent::Create();
 			Copy(OldClipboardContent);
@@ -2402,18 +2400,18 @@ void UNiagaraStackFunctionInput::ReassignDynamicInputScript(UNiagaraScript* Dyna
 		InputValues.DynamicNode->MarkNodeRequiresSynchronization(TEXT("Dynamic input script reassigned."), true);
 		RefreshChildren();
 		
-		if (ScriptData->ConversionUtility != nullptr && OldClipboardContent != nullptr)
+		if (NewScriptData->ConversionUtility != nullptr && OldClipboardContent != nullptr)
 		{
-			UNiagaraConvertInPlaceUtilityBase* ConversionUtility = NewObject< UNiagaraConvertInPlaceUtilityBase>(GetTransientPackage(), ScriptData->ConversionUtility);
+			UNiagaraConvertInPlaceUtilityBase* ConversionUtility = NewObject< UNiagaraConvertInPlaceUtilityBase>(GetTransientPackage(), NewScriptData->ConversionUtility);
 
 			UNiagaraClipboardContent* NewClipboardContent = UNiagaraClipboardContent::Create();
 			Copy(NewClipboardContent);
 			TArray<UNiagaraStackFunctionInputCollection*> DynamicInputCollections;
 			GetUnfilteredChildrenOfType(DynamicInputCollections);
 
-			FText ConvertMessage;
 			if (ConversionUtility && DynamicInputCollections.Num() == 0)
 			{
+				FText ConvertMessage;
 				bool bConverted = ConversionUtility->Convert(OldScript, OldClipboardContent, DynamicInputScript, DynamicInputCollections[0], NewClipboardContent, InputValues.DynamicNode.Get(), ConvertMessage);
 				if (!ConvertMessage.IsEmptyOrWhitespace())
 				{
@@ -2426,7 +2424,31 @@ void UNiagaraStackFunctionInput::ReassignDynamicInputScript(UNiagaraScript* Dyna
 				}
 			}
 		}
-
+		else if (NewScriptData && OldScriptData && OldScriptData->bUsePythonScriptConversion && OldClipboardContent != nullptr)
+		{
+			UNiagaraClipboardContent* NewClipboardContent = UNiagaraClipboardContent::Create();
+			UNiagaraClipboardContent* NewPythonContent = nullptr;
+			Copy(NewClipboardContent);
+			if (OldClipboardContent->FunctionInputs.Num() == 1 && OldClipboardContent->FunctionInputs[0]->Dynamic)
+			{
+				OldClipboardContent->FunctionInputs = OldClipboardContent->FunctionInputs[0]->Dynamic->Inputs;
+			}
+			if (NewClipboardContent->FunctionInputs.Num() == 1 && NewClipboardContent->FunctionInputs[0]->Dynamic)
+			{
+				NewPythonContent = UNiagaraClipboardContent::Create();
+				NewPythonContent->FunctionInputs = NewClipboardContent->FunctionInputs[0]->Dynamic->Inputs;
+			}
+			InputValues.DynamicNode->PythonUpgradeScriptWarnings.Empty();
+			FText Warnings;
+			if (UNiagaraClipboardContent* NewInputs = FNiagaraEditorUtilities::RunPythonConversionScript(*NewScriptData, NewPythonContent, *OldScriptData, OldClipboardContent, Warnings))
+			{
+				NewClipboardContent->FunctionInputs[0]->Dynamic->Inputs = NewInputs->FunctionInputs;
+				Paste(NewClipboardContent, Warnings);
+				if (!Warnings.IsEmpty()) {
+					InputValues.DynamicNode->PythonUpgradeScriptWarnings = Warnings.ToString();
+				}
+			}			
+		}
 	}
 }
 
diff --git a/Engine/Plugins/FX/Niagara/Source/NiagaraEditor/Private/ViewModels/Stack/NiagaraStackModuleItem.cpp b/Engine/Plugins/FX/Niagara/Source/NiagaraEditor/Private/ViewModels/Stack/NiagaraStackModuleItem.cpp
--- a/Engine/Plugins/FX/Niagara/Source/NiagaraEditor/Private/ViewModels/Stack/NiagaraStackModuleItem.cpp
+++ b/Engine/Plugins/FX/Niagara/Source/NiagaraEditor/Private/ViewModels/Stack/NiagaraStackModuleItem.cpp
@@ -1395,11 +1395,12 @@ void UNiagaraStackModuleItem::ReassignModuleScript(UNiagaraScript* ModuleScript)
 
 		const FString OldName = FunctionCallNode->GetFunctionName();
 		UNiagaraScript* OldScript = FunctionCallNode->FunctionScript;
+		FVersionedNiagaraScriptData* OldScriptData = FunctionCallNode->GetScriptData();
 
 		FunctionCallNode->Modify();
 		UNiagaraClipboardContent* OldClipboardContent = nullptr;
-		FVersionedNiagaraScriptData* ScriptData = ModuleScript->GetLatestScriptData();
-		if (ScriptData->ConversionUtility != nullptr)
+		FVersionedNiagaraScriptData* NewScriptData = ModuleScript->GetLatestScriptData();
+		if (NewScriptData->ConversionUtility || OldScriptData->bUsePythonScriptConversion)
 		{
 			OldClipboardContent = UNiagaraClipboardContent::Create();
 			Copy(OldClipboardContent);
@@ -1422,16 +1423,16 @@ void UNiagaraStackModuleItem::ReassignModuleScript(UNiagaraScript* ModuleScript)
 			RefreshChildren();
 		}
 		
-		if (ScriptData->ConversionUtility != nullptr && OldClipboardContent != nullptr)
+		if (NewScriptData->ConversionUtility != nullptr && OldClipboardContent != nullptr)
 		{
-			UNiagaraConvertInPlaceUtilityBase* ConversionUtility = NewObject< UNiagaraConvertInPlaceUtilityBase>(GetTransientPackage(), ScriptData->ConversionUtility);
-			FText ConvertMessage;
+			UNiagaraConvertInPlaceUtilityBase* ConversionUtility = NewObject< UNiagaraConvertInPlaceUtilityBase>(GetTransientPackage(), NewScriptData->ConversionUtility);
 
 			UNiagaraClipboardContent* NewClipboardContent = UNiagaraClipboardContent::Create();
 			Copy(NewClipboardContent);
 
 			if (ConversionUtility )
 			{
+				FText ConvertMessage;
 				bool bConverted = ConversionUtility->Convert(OldScript, OldClipboardContent, ModuleScript, InputCollection, NewClipboardContent, FunctionCallNode, ConvertMessage);
 				if (!ConvertMessage.IsEmptyOrWhitespace())
 				{
@@ -1444,6 +1445,30 @@ void UNiagaraStackModuleItem::ReassignModuleScript(UNiagaraScript* ModuleScript)
 				}
 			}
 		}
+		else if (NewScriptData && OldScriptData && OldScriptData->bUsePythonScriptConversion && OldClipboardContent != nullptr)
+		{
+			UNiagaraClipboardContent* NewClipboardContent = UNiagaraClipboardContent::Create();
+			Copy(NewClipboardContent);
+			if (OldClipboardContent->Functions.Num() > 0)
+			{
+				OldClipboardContent->FunctionInputs = OldClipboardContent->Functions[0]->Inputs;
+				OldClipboardContent->Functions.Empty();
+			}
+			if (NewClipboardContent->Functions.Num() > 0)
+			{
+				NewClipboardContent->FunctionInputs = NewClipboardContent->Functions[0]->Inputs;
+				NewClipboardContent->Functions.Empty();
+			}
+			FunctionCallNode->PythonUpgradeScriptWarnings.Empty();
+			FText Warnings;
+			if (UNiagaraClipboardContent* NewInputs = FNiagaraEditorUtilities::RunPythonConversionScript(*NewScriptData, NewClipboardContent, *OldScriptData, OldClipboardContent, Warnings))
+			{
+				Paste(NewInputs, Warnings);
+				if (!Warnings.IsEmpty()) {
+					FunctionCallNode->PythonUpgradeScriptWarnings = Warnings.ToString();
+				}
+			}			
+		}
 	}
 }
 
diff --git a/Engine/Plugins/FX/Niagara/Source/NiagaraEditor/Private/Widgets/SNiagaraSystemViewport.cpp b/Engine/Plugins/FX/Niagara/Source/NiagaraEditor/Private/Widgets/SNiagaraSystemViewport.cpp
--- a/Engine/Plugins/FX/Niagara/Source/NiagaraEditor/Private/Widgets/SNiagaraSystemViewport.cpp
+++ b/Engine/Plugins/FX/Niagara/Source/NiagaraEditor/Private/Widgets/SNiagaraSystemViewport.cpp
@@ -15,6 +15,7 @@
 #include "NiagaraEffectType.h"
 #include "NiagaraPerfBaseline.h"
 #include "NiagaraSettings.h"
+#include "NiagaraShader.h"
 #include "NiagaraSystemImpl.h"
 #include "NiagaraSystemInstance.h"
 #include "NiagaraSystemInstanceController.h"
diff --git a/Engine/Plugins/FX/Niagara/Source/NiagaraEditor/Public/NiagaraEditorUtilities.h b/Engine/Plugins/FX/Niagara/Source/NiagaraEditor/Public/NiagaraEditorUtilities.h
--- a/Engine/Plugins/FX/Niagara/Source/NiagaraEditor/Public/NiagaraEditorUtilities.h
+++ b/Engine/Plugins/FX/Niagara/Source/NiagaraEditor/Public/NiagaraEditorUtilities.h
@@ -13,6 +13,7 @@
 #include "ViewModels/NiagaraSystemScalabilityViewModel.h"
 #include "ViewModels/NiagaraSystemViewModel.h"
 
+class UNiagaraClipboardContent;
 class UNiagaraNodeInput;
 class UNiagaraNodeOutput;
 class UNiagaraNodeFunctionCall;
@@ -44,6 +45,7 @@ struct FNiagaraNamespaceMetadata;
 class FNiagaraParameterHandle;
 class INiagaraParameterDefinitionsSubscriberViewModel;
 struct FNiagaraScriptVersionUpgradeContext;
+struct FNiagaraScriptConversionContext;
 class UUpgradeNiagaraEmitterContext;
 
 enum class ENiagaraFunctionDebugState : uint8;
@@ -345,6 +347,9 @@ namespace FNiagaraEditorUtilities
 	TSharedPtr<INiagaraParameterDefinitionsSubscriberViewModel> GetOwningLibrarySubscriberViewModelForGraph(const UNiagaraGraph* Graph);
 
 	TArray<UNiagaraParameterDefinitions*> DowncastParameterDefinitionsBaseArray(const TArray<UNiagaraParameterDefinitionsBase*> BaseArray);
+	
+	// Executes python conversion script on the given source node
+    UNiagaraClipboardContent* RunPythonConversionScript(FVersionedNiagaraScriptData& NewScriptVersionData, UNiagaraClipboardContent* NewScriptInputs, FVersionedNiagaraScriptData& OldScriptVersionData, UNiagaraClipboardContent* OldScriptInputs, FText& OutWarnings);
 
 	// Executes python upgrade scripts on the given source node for all the given in-between versions
 	void RunPythonUpgradeScripts(UNiagaraNodeFunctionCall* SourceNode, const TArray<FVersionedNiagaraScriptData*>& UpgradeVersionData, const FNiagaraScriptVersionUpgradeContext& UpgradeContext, FString& OutWarnings);
diff --git a/Engine/Plugins/FX/Niagara/Source/NiagaraEditor/Public/UpgradeNiagaraScriptResults.h b/Engine/Plugins/FX/Niagara/Source/NiagaraEditor/Public/UpgradeNiagaraScriptResults.h
--- a/Engine/Plugins/FX/Niagara/Source/NiagaraEditor/Public/UpgradeNiagaraScriptResults.h
+++ b/Engine/Plugins/FX/Niagara/Source/NiagaraEditor/Public/UpgradeNiagaraScriptResults.h
@@ -143,6 +143,9 @@ public:
     UFUNCTION(BlueprintCallable, Category = "Scripting")
     void SetLinkedInput(const FString& InputName, FString Value);
 
+	UFUNCTION(BlueprintCallable, Category = "Scripting")
+	void SetNewInput(const FString& InputName, UNiagaraPythonScriptModuleInput* Value);
+
 private:
 	UNiagaraPythonScriptModuleInput* GetNewInput(const FName& InputName) const;
 
