Fixed file includes

diff --git a/Engine/Plugins/FX/Niagara/Source/Niagara/Private/NiagaraGPUInstanceCountManager.cpp b/Engine/Plugins/FX/Niagara/Source/Niagara/Private/NiagaraGPUInstanceCountManager.cpp
--- a/Engine/Plugins/FX/Niagara/Source/Niagara/Private/NiagaraGPUInstanceCountManager.cpp
+++ b/Engine/Plugins/FX/Niagara/Source/Niagara/Private/NiagaraGPUInstanceCountManager.cpp
@@ -5,6 +5,7 @@
 #include "NiagaraGpuComputeDispatchInterface.h"
 #include "NiagaraRenderer.h"
 #include "NiagaraStats.h"
+#include "NiagaraShader.h"
 
 #include "Containers/DynamicRHIResourceArray.h"
 #include "GPUSortManager.h" // CopyUIntBufferToTargets
