Fixed a bug where changing versioned properties would not dirty a script/emitter

diff --git a/Engine/Plugins/FX/Niagara/Source/NiagaraEditor/Private/Toolkits/NiagaraScriptToolkit.cpp b/Engine/Plugins/FX/Niagara/Source/NiagaraEditor/Private/Toolkits/NiagaraScriptToolkit.cpp
--- a/Engine/Plugins/FX/Niagara/Source/NiagaraEditor/Private/Toolkits/NiagaraScriptToolkit.cpp
+++ b/Engine/Plugins/FX/Niagara/Source/NiagaraEditor/Private/Toolkits/NiagaraScriptToolkit.cpp
@@ -219,12 +219,16 @@ void FNiagaraScriptToolkit::Initialize( const EToolkitMode::Type Mode, const TSh
 	VersionMetadata = NewObject<UNiagaraVersionMetaData>(InputScript, "VersionMetadata", RF_Transient);
 	SAssignNew(VersionsWidget, SNiagaraScriptVersionWidget, EditedNiagaraScript.Script, VersionMetadata, InputScript->GetOutermost()->GetName())
 		.OnChangeToVersion(this, &FNiagaraScriptToolkit::SwitchToVersion)
-		.OnVersionDataChanged_Lambda([this]()
+		.OnVersionDataChanged_Lambda([this](const FPropertyChangedEvent* PropertyChangedEvent, FGuid SelectedVersion)
 		{
 	        if (UNiagaraScriptSource* ScriptSource = Cast<UNiagaraScriptSource>(EditedNiagaraScript.Script->GetSource(EditedNiagaraScript.Version)))
 	        {
 	            ScriptSource->NodeGraph->NotifyGraphChanged();
 	        }
+
+			static FProperty* VersionProperty = FindFProperty<FProperty>(UNiagaraScript::StaticClass(), FName("VersionData"));
+			FPropertyChangedEvent ChangeEvent = PropertyChangedEvent ? *PropertyChangedEvent : FPropertyChangedEvent(VersionProperty);
+			EditedNiagaraScript.Script->PostEditChangeVersionedProperty(ChangeEvent, SelectedVersion);
 		});
 
 	TSharedRef<FTabManager::FLayout> StandaloneDefaultLayout = FTabManager::NewLayout("Standalone_Niagara_Layout_v12")
diff --git a/Engine/Plugins/FX/Niagara/Source/NiagaraEditor/Private/Toolkits/NiagaraSystemToolkit.cpp b/Engine/Plugins/FX/Niagara/Source/NiagaraEditor/Private/Toolkits/NiagaraSystemToolkit.cpp
--- a/Engine/Plugins/FX/Niagara/Source/NiagaraEditor/Private/Toolkits/NiagaraSystemToolkit.cpp
+++ b/Engine/Plugins/FX/Niagara/Source/NiagaraEditor/Private/Toolkits/NiagaraSystemToolkit.cpp
@@ -312,13 +312,13 @@ void FNiagaraSystemToolkit::InitializeInternal(const EToolkitMode::Type Mode, co
 	VersionMetadata = NewObject<UNiagaraVersionMetaData>(ToolkitObject, "VersionMetadata", RF_Transient);
 	SAssignNew(VersionsWidget, SNiagaraEmitterVersionWidget, HasEmitter() ? GetEditedEmitterViewModel()->GetEmitter().Emitter : nullptr, VersionMetadata, HasEmitter() ? Emitter->GetOutermost()->GetName() : TEXT(""))
 		.OnChangeToVersion(this, &FNiagaraSystemToolkit::SwitchToVersion)
-		.OnVersionDataChanged_Lambda([this]()
+		.OnVersionDataChanged_Lambda([this](const FPropertyChangedEvent* InChangedEvent, FGuid SelectedVersion)
 		{
 			if (TSharedPtr<FNiagaraEmitterViewModel> EditableEmitterViewModel = SystemViewModel->GetEmitterHandleViewModels()[0]->GetEmitterViewModel())
 			{
 				FVersionedNiagaraEmitter EditableEmitter = EditableEmitterViewModel->GetEmitter();
 				FProperty* VersionProperty = FindFProperty<FProperty>(UNiagaraEmitter::StaticClass(), FName("VersionData"));
-				FPropertyChangedEvent ChangeEvent(VersionProperty);
+				FPropertyChangedEvent ChangeEvent = InChangedEvent ? *InChangedEvent : FPropertyChangedEvent(VersionProperty);
 				EditableEmitter.Emitter->PostEditChangeProperty(ChangeEvent);
 			}
 		});
diff --git a/Engine/Plugins/FX/Niagara/Source/NiagaraEditor/Private/Widgets/SNiagaraVersionWidget.cpp b/Engine/Plugins/FX/Niagara/Source/NiagaraEditor/Private/Widgets/SNiagaraVersionWidget.cpp
--- a/Engine/Plugins/FX/Niagara/Source/NiagaraEditor/Private/Widgets/SNiagaraVersionWidget.cpp
+++ b/Engine/Plugins/FX/Niagara/Source/NiagaraEditor/Private/Widgets/SNiagaraVersionWidget.cpp
@@ -236,7 +236,7 @@ void SNiagaraVersionWidget::NotifyPostChange(const FPropertyChangedEvent& Proper
 		VersionInListSelected(ObjectData->GetObjectVersion());
 	}
 
-	OnVersionDataChanged.ExecuteIfBound();
+	OnVersionDataChanged.ExecuteIfBound(&PropertyChangedEvent, SelectedVersion);
 }
 
 void SNiagaraVersionWidget::PostUndo(bool bSuccess)
@@ -252,7 +252,7 @@ void SNiagaraVersionWidget::PostRedo(bool bSuccess)
 	PostUndo(bSuccess);
 }
 
-void SNiagaraVersionWidget::SetOnVersionDataChanged(FSimpleDelegate InOnVersionDataChanged)
+void SNiagaraVersionWidget::SetOnVersionDataChanged(FOnVersionPropertyChangedDelegate InOnVersionDataChanged)
 {
 	OnVersionDataChanged = InOnVersionDataChanged;
 }
@@ -325,7 +325,7 @@ FReply SNiagaraVersionWidget::EnableVersioning()
 {
 	FScopedTransaction Transaction(LOCTEXT("EnableVersioning", "Enable versioning"));
 	VersionedObject->EnableVersioning();
-	OnVersionDataChanged.ExecuteIfBound();
+	OnVersionDataChanged.ExecuteIfBound(nullptr, SelectedVersion);
 	
 	return FReply::Handled();
 }
@@ -348,7 +348,7 @@ void SNiagaraVersionWidget::ExecuteDeleteAction(FNiagaraAssetVersion AssetVersio
 	
 	VersionListWidget->RefreshAllActions(true);
 	VersionInListSelected(VersionedObject->GetExposedVersion());
-	OnVersionDataChanged.ExecuteIfBound();
+	OnVersionDataChanged.ExecuteIfBound(nullptr, SelectedVersion);
 }
 
 void SNiagaraVersionWidget::ExecuteExposeAction(FNiagaraAssetVersion AssetVersion)
@@ -358,7 +358,7 @@ void SNiagaraVersionWidget::ExecuteExposeAction(FNiagaraAssetVersion AssetVersio
 
 	VersionListWidget->RefreshAllActions(true);
 	VersionInListSelected(AssetVersion);
-	OnVersionDataChanged.ExecuteIfBound();
+	OnVersionDataChanged.ExecuteIfBound(nullptr, SelectedVersion);
 }
 
 TSharedRef<ITableRow> SNiagaraVersionWidget::HandleVersionViewGenerateRow(TSharedRef<FNiagaraAssetVersion> Item, const TSharedRef<STableViewBase>& OwnerTable)
@@ -465,7 +465,7 @@ void SNiagaraVersionWidget::AddNewMajorVersion()
 	VersionListWidget->SelectItemByName(FName(ItemName.ToString()));
 
 	OnChangeToVersion.ExecuteIfBound(NewVersion);
-	OnVersionDataChanged.ExecuteIfBound();
+	OnVersionDataChanged.ExecuteIfBound(nullptr, SelectedVersion);
 }
 
 void SNiagaraVersionWidget::AddNewMinorVersion()
@@ -481,7 +481,7 @@ void SNiagaraVersionWidget::AddNewMinorVersion()
 	VersionListWidget->SelectItemByName(FName(ItemName.ToString()));
 
 	OnChangeToVersion.ExecuteIfBound(NewVersion);
-	OnVersionDataChanged.ExecuteIfBound();
+	OnVersionDataChanged.ExecuteIfBound(nullptr, SelectedVersion);
 }
 
 END_SLATE_FUNCTION_BUILD_OPTIMIZATION
diff --git a/Engine/Plugins/FX/Niagara/Source/NiagaraEditor/Private/Widgets/SNiagaraVersionWidget.h b/Engine/Plugins/FX/Niagara/Source/NiagaraEditor/Private/Widgets/SNiagaraVersionWidget.h
--- a/Engine/Plugins/FX/Niagara/Source/NiagaraEditor/Private/Widgets/SNiagaraVersionWidget.h
+++ b/Engine/Plugins/FX/Niagara/Source/NiagaraEditor/Private/Widgets/SNiagaraVersionWidget.h
@@ -23,6 +23,7 @@ class UNiagaraVersionMetaData;
 class UNiagaraScript;
 
 DECLARE_DELEGATE_OneParam(FOnSwitchToVersionDelegate, FGuid);
+DECLARE_DELEGATE_TwoParams(FOnVersionPropertyChangedDelegate, const FPropertyChangedEvent*, FGuid);
 
 struct FNiagaraVersionMenuAction : FNiagaraMenuAction
 {
@@ -39,7 +40,7 @@ public:
 	{}
 
 	/** Called when the version data of the script was edited by the user */
-	SLATE_EVENT(FSimpleDelegate, OnVersionDataChanged)
+	SLATE_EVENT(FOnVersionPropertyChangedDelegate, OnVersionDataChanged)
 
 	/** Called when the user does something that prompts the editor to change the current active version, e.g. delete a version or add a new version */
     SLATE_EVENT(FOnSwitchToVersionDelegate, OnChangeToVersion)
@@ -62,7 +63,7 @@ public:
 	// End of FEditorUndoClient
 
 	/** See OnVersionDataChanged event */
-	void SetOnVersionDataChanged(FSimpleDelegate InOnVersionDataChanged);
+	void SetOnVersionDataChanged(FOnVersionPropertyChangedDelegate InOnVersionDataChanged);
 
 protected:
 	virtual FText GetInfoHeaderText() const;
@@ -72,12 +73,12 @@ private:
 	FNiagaraVersionedObject* VersionedObject = nullptr;
 	bool bAssetVersionsChanged = false;
 	UNiagaraVersionMetaData* VersionMetadata = nullptr;
-	FGuid SelectedVersion;
 	
-	FSimpleDelegate OnVersionDataChanged;
+	FOnVersionPropertyChangedDelegate OnVersionDataChanged;
 	FOnSwitchToVersionDelegate OnChangeToVersion;
 	TSharedPtr<IDetailsView> VersionSettingsDetails;
 	TSharedPtr<SGraphActionMenu> VersionListWidget;
+	FGuid SelectedVersion;
 
 	void AddNewMajorVersion();
 	void AddNewMinorVersion();
