- Invalidate execution ready parameter stores when a data interface has changed as it can impact the function the VM should be binding to

diff --git a/Engine/Plugins/FX/Niagara/Source/Niagara/Private/NiagaraSystem.cpp b/Engine/Plugins/FX/Niagara/Source/Niagara/Private/NiagaraSystem.cpp
--- a/Engine/Plugins/FX/Niagara/Source/Niagara/Private/NiagaraSystem.cpp
+++ b/Engine/Plugins/FX/Niagara/Source/Niagara/Private/NiagaraSystem.cpp
@@ -2064,6 +2064,8 @@ void UNiagaraSystem::RegisterDataChannelUse(const UNiagaraDataChannel* DataChann
 
 void UNiagaraSystem::OnCompiledDataInterfaceChanged()
 {
+	// We need to make sure our execution ready parameter stores are invalidated so that the DI's will rebind functions
+	ForEachScript([&](UNiagaraScript* Script) { Script->InvalidateExecutionReadyParameterStores(); });
 	UpdatePostCompileDIInfo();
 	UpdateDITickFlags();
 }
