- Add Z Depth Query Type to collision query validation

#rnx

diff --git a/Engine/Plugins/FX/Niagara/Source/NiagaraEditor/Private/NiagaraValidationRules.cpp b/Engine/Plugins/FX/Niagara/Source/NiagaraEditor/Private/NiagaraValidationRules.cpp
--- a/Engine/Plugins/FX/Niagara/Source/NiagaraEditor/Private/NiagaraValidationRules.cpp
+++ b/Engine/Plugins/FX/Niagara/Source/NiagaraEditor/Private/NiagaraValidationRules.cpp
@@ -181,6 +181,37 @@ namespace NiagaraValidation
 		}
 		return nullptr;
 	}
+
+	TOptional<int32> GetModuleStaticInt32Value(const UNiagaraStackModuleItem* Module, FName ParameterName)
+	{
+		TArray<UNiagaraStackFunctionInput*> ModuleInputs;
+		Module->GetParameterInputs(ModuleInputs);
+
+		for (UNiagaraStackFunctionInput* Input : ModuleInputs)
+		{
+			if (Input->IsStaticParameter() && Input->GetInputParameterHandle().GetName() == ParameterName)
+			{
+				return TOptional<int32>(*(int32*)Input->GetLocalValueStruct()->GetStructMemory());
+			}
+		}
+		return TOptional<int32>();
+	}
+
+	void SetModuleStaticInt32Value(UNiagaraStackModuleItem* Module, FName ParameterName, int32 NewValue)
+	{
+		TArray<UNiagaraStackFunctionInput*> ModuleInputs;
+		Module->GetParameterInputs(ModuleInputs);
+
+		for (UNiagaraStackFunctionInput* Input : ModuleInputs)
+		{
+			if (Input->IsStaticParameter() && Input->GetInputParameterHandle().GetName() == ParameterName)
+			{
+				TSharedRef<FStructOnScope> ValueStruct = MakeShared<FStructOnScope>(Input->GetLocalValueStruct()->GetStruct());
+				*(int32*)ValueStruct->GetStructMemory() = NewValue;
+				Input->SetLocalValue(ValueStruct);
+			}
+		}
+	};
 }
 
 // --------------------------------------------------------------------------------------------------------------------------------------------
@@ -695,19 +726,8 @@ void UNiagaraValidationRule_NoOpaqueRenderMaterial::CheckValidity(const FNiagara
 	UNiagaraStackModuleItem* SourceModule = Cast<UNiagaraStackModuleItem>(Context.Source);
 	if (SourceModule && SourceModule->GetIsEnabled() && SourceModule->GetEmitterViewModel())
 	{
-		auto GetCollisionTypeInput = [](const UNiagaraStackModuleItem* Module)
-		{
-			TArray<UNiagaraStackFunctionInput*> ModuleInputs;
-			Module->GetParameterInputs(ModuleInputs);
-			for (UNiagaraStackFunctionInput* Input : ModuleInputs)
-			{
-				if (Input->IsStaticParameter() && Input->GetInputParameterHandle().GetName() == FName("GPU Collision Type"))
-				{
-					return Input;
-				}
-			}
-			return static_cast<UNiagaraStackFunctionInput*>(nullptr);
-		};
+		static const FName NAME_GPUCollisionType("GPU Collision Type");
+		static const FName NAME_ZDepthQueryType("Z Depth Query Type");
 		
 		// search for the right emitter view model
 		for (const TSharedRef<FNiagaraEmitterHandleViewModel>& EmitterHandleModel : Context.ViewModel->GetEmitterHandleViewModels())
@@ -715,21 +735,18 @@ void UNiagaraValidationRule_NoOpaqueRenderMaterial::CheckValidity(const FNiagara
 			FVersionedNiagaraEmitterData* EmitterData = EmitterHandleModel->GetEmitterHandle()->GetEmitterData();
 			if (EmitterHandleModel->GetIsEnabled() && EmitterData->SimTarget == ENiagaraSimTarget::GPUComputeSim && EmitterHandleModel->GetEmitterViewModel() == SourceModule->GetEmitterViewModel())
 			{
-				// check if we are using depth buffer collisions
-				if (UNiagaraStackFunctionInput* Input = GetCollisionTypeInput(SourceModule))
+				// Note: for these BP driven enums we can't compare the values
+				TOptional<int32> GPUCollisionType = NiagaraValidation::GetModuleStaticInt32Value(SourceModule, NAME_GPUCollisionType);
+				if (!GPUCollisionType.IsSet() || GPUCollisionType.GetValue() != 0)
 				{
-					// the first entry of the enum is depth buffer collision
-					// Unfortunately it's a BP enum, so we can't match the named values in C++
-					if (int32 Value = *(int32*)Input->GetLocalValueStruct()->GetStructMemory(); Value != 0)
-					{
-						continue;
-					}
+					continue;
 				}
-				else
+				TOptional<int32> ZDepthQueryType = NiagaraValidation::GetModuleStaticInt32Value(SourceModule, NAME_ZDepthQueryType);
+				if (!ZDepthQueryType.IsSet() || ZDepthQueryType.GetValue() != 0)
 				{
 					continue;
 				}
-				
+
 				// check the renderers
 				TArray<UNiagaraStackRendererItem*> RendererItems = NiagaraValidation::GetStackEntries<UNiagaraStackRendererItem>(EmitterHandleModel->GetEmitterStackViewModel());
 				for (UNiagaraStackRendererItem* Renderer : RendererItems)
@@ -759,16 +776,11 @@ void UNiagaraValidationRule_NoOpaqueRenderMaterial::CheckValidity(const FNiagara
 									TWeakObjectPtr<UNiagaraStackModuleItem> WeakSourceModule = SourceModule;
 								
 									DisableRendererFix.FixDelegate = FNiagaraValidationFixDelegate::CreateLambda(
-										[WeakSourceModule, GetCollisionTypeInput]()
+										[WeakSourceModule]()
 										{
 											if (UNiagaraStackModuleItem* CollisionModule = WeakSourceModule.Get())
 											{
-												if (UNiagaraStackFunctionInput* Input = GetCollisionTypeInput(CollisionModule))
-												{
-													TSharedRef<FStructOnScope> ValueStruct = MakeShared<FStructOnScope>(Input->GetLocalValueStruct()->GetStruct());
-													*(int32*)ValueStruct->GetStructMemory() = 1;
-													Input->SetLocalValue(ValueStruct);
-												}
+												NiagaraValidation::SetModuleStaticInt32Value(CollisionModule, NAME_GPUCollisionType, 1);
 											}
 										});
 								}
