Niagara - Refactor generic object based messages.
+ Create a struct to contains messages and dismissed message ids to unify the handling across system, emitter, and function call node, and to facilitate adding messages to other object types.
+ Allow asset based messages to be properly dismissed and restored using the default dismiss UI instead of a link.
+ Fix a bug where issues dismissed from a system stack could not be restored.
+ Don't log specific warnings when the issue has been dismissed by the user.

diff --git a/Engine/Plugins/FX/Niagara/Source/Niagara/Classes/NiagaraEmitter.h b/Engine/Plugins/FX/Niagara/Source/Niagara/Classes/NiagaraEmitter.h
--- a/Engine/Plugins/FX/Niagara/Source/Niagara/Classes/NiagaraEmitter.h
+++ b/Engine/Plugins/FX/Niagara/Source/Niagara/Classes/NiagaraEmitter.h
@@ -8,6 +8,7 @@
 #include "UObject/Object.h"
 #include "NiagaraScript.h"
 #include "NiagaraMessageDataBase.h"
+#include "NiagaraMessageStore.h"
 #include "INiagaraMergeManager.h"
 #include "NiagaraEffectType.h"
 #include "NiagaraDataSetAccessor.h"
@@ -911,10 +912,7 @@ public:
 	NIAGARA_API void ChangeParentVersion(const FGuid& NewParentVersion, const FGuid& EmitterVersion);
 
 	NIAGARA_API void NotifyScratchPadScriptsChanged();
-	NIAGARA_API const TMap<FGuid, TObjectPtr<UNiagaraMessageDataBase>>& GetMessages() const { return MessageKeyToMessageMap; };
-	NIAGARA_API void AddMessage(const FGuid& MessageKey, UNiagaraMessageDataBase* NewMessage) { MessageKeyToMessageMap.Add(MessageKey, NewMessage); };
-	NIAGARA_API void RemoveMessage(const FGuid& MessageKey) { MessageKeyToMessageMap.Remove(MessageKey); };
-	void RemoveMessageDelegateable(const FGuid MessageKey) { MessageKeyToMessageMap.Remove(MessageKey); };
+	NIAGARA_API FNiagaraMessageStore& GetMessageStore() { return MessageStore; }
 #endif
 
 protected:
@@ -1003,9 +1001,11 @@ private:
 #endif
 
 #if WITH_EDITORONLY_DATA
-	/** Messages associated with the Emitter asset. */
 	UPROPERTY()
-	TMap<FGuid, TObjectPtr<UNiagaraMessageDataBase>> MessageKeyToMessageMap;
+	TMap<FGuid, TObjectPtr<UNiagaraMessageDataBase>> MessageKeyToMessageMap_DEPRECATED;
+
+	UPROPERTY()
+	FNiagaraMessageStore MessageStore;
 #endif
 
 	void ResolveScalabilitySettings();
diff --git a/Engine/Plugins/FX/Niagara/Source/Niagara/Classes/NiagaraSystem.h b/Engine/Plugins/FX/Niagara/Source/Niagara/Classes/NiagaraSystem.h
--- a/Engine/Plugins/FX/Niagara/Source/Niagara/Classes/NiagaraSystem.h
+++ b/Engine/Plugins/FX/Niagara/Source/Niagara/Classes/NiagaraSystem.h
@@ -8,6 +8,7 @@
 #include "NiagaraDataSetAccessor.h"
 #include "NiagaraEffectType.h"
 #include "NiagaraEmitterHandle.h"
+#include "NiagaraMessageStore.h"
 #include "NiagaraParameterCollection.h"
 #include "NiagaraParameterDefinitionsSubscriber.h"
 #include "NiagaraUserRedirectionParameterStore.h"
@@ -502,6 +503,8 @@ public:
 	
 	UPROPERTY(transient)
 	FNiagaraSystemUpdateContext UpdateContext;
+
+	const static FGuid ComputeEmitterExecutionOrderMessageId;
 #endif
 
 	void UpdateSystemAfterLoad();
@@ -761,10 +764,7 @@ public:
 	bool CanObtainUserVariable(const FNiagaraVariableBase& InVar) const;
 
 #if WITH_EDITORONLY_DATA
-	const TMap<FGuid, TObjectPtr<UNiagaraMessageDataBase>>& GetMessages() const { return MessageKeyToMessageMap; };
-	void AddMessage(const FGuid& MessageKey, UNiagaraMessageDataBase* NewMessage) { MessageKeyToMessageMap.Add(MessageKey, NewMessage); };
-	void RemoveMessage(const FGuid& MessageKey) { MessageKeyToMessageMap.Remove(MessageKey); };
-	void RemoveMessageDelegateable(const FGuid MessageKey) { MessageKeyToMessageMap.Remove(MessageKey); };
+	FNiagaraMessageStore& GetMessageStore() { return MessageStore; }
 	const FGuid& GetAssetGuid() const {return AssetGuid;};
 #endif
 
@@ -1038,7 +1038,10 @@ protected:
 #if WITH_EDITORONLY_DATA
 	/** Messages associated with the System asset. */
 	UPROPERTY()
-	TMap<FGuid, TObjectPtr<UNiagaraMessageDataBase>> MessageKeyToMessageMap;
+	TMap<FGuid, TObjectPtr<UNiagaraMessageDataBase>> MessageKeyToMessageMap_DEPRECATED;
+
+	UPROPERTY()
+	FNiagaraMessageStore MessageStore;
 
 	FGuid AssetGuid;
 #endif
diff --git a/Engine/Plugins/FX/Niagara/Source/Niagara/Private/NiagaraEmitter.cpp b/Engine/Plugins/FX/Niagara/Source/Niagara/Private/NiagaraEmitter.cpp
--- a/Engine/Plugins/FX/Niagara/Source/Niagara/Private/NiagaraEmitter.cpp
+++ b/Engine/Plugins/FX/Niagara/Source/Niagara/Private/NiagaraEmitter.cpp
@@ -425,10 +425,13 @@ TArray<INiagaraMergeManager::FMergeEmitterResults> UNiagaraEmitter::MergeChanges
 				*GetPathName(), VersionedParent.Emitter != nullptr ? *VersionedParent.Emitter->GetPathName() : TEXT("(null)"), *MergeResults.GetErrorMessagesString());
 		}
 
-		RemoveMessage(EmitterMergeMessageId);
 		if (MergeResults.MergeNiagaraMessage != nullptr)
 		{
-			AddMessage(EmitterMergeMessageId, CastChecked<UNiagaraMessageDataBase>(StaticDuplicateObject(MergeResults.MergeNiagaraMessage, this)));
+			MessageStore.AddMessage(EmitterMergeMessageId, CastChecked<UNiagaraMessageDataBase>(StaticDuplicateObject(MergeResults.MergeNiagaraMessage, this)));
+		}
+		else
+		{
+			MessageStore.RemoveMessage(EmitterMergeMessageId);
 		}
 	}
 	return Results;
@@ -670,6 +673,12 @@ void UNiagaraEmitter::PostLoad()
 		LibraryVisibility = ENiagaraScriptLibraryVisibility::Library;
 	}
 
+	if (MessageKeyToMessageMap_DEPRECATED.IsEmpty() == false)
+	{
+		MessageStore.SetMessages(MessageKeyToMessageMap_DEPRECATED);
+		MessageKeyToMessageMap_DEPRECATED.Empty();
+	}
+
 	if (this->GetOuter()->IsA<UNiagaraSystem>() || this->GetOuter()->IsA<UNiagaraEmitter>())
 	{
 		// Remove thunbnails for non-asset emitters to prevent problems with cooked for editor emitters being referenced by uncooked emitters and systems.
@@ -2888,13 +2897,12 @@ void UNiagaraEmitter::UpdateFromMergedCopy(const INiagaraMergeManager& MergeMana
 	SetEditorData(MergedData->GetEditorData(), EmitterData->Version.VersionGuid);
 
 	// Update messages
-	for (const TPair<FGuid, TObjectPtr<UNiagaraMessageDataBase>>& Message : MergedEmitter->GetMessages())
+	for (const TPair<FGuid, TObjectPtr<UNiagaraMessageDataBase>>& Message : MergedEmitter->GetMessageStore().GetMessages())
 	{
-		TObjectPtr<UNiagaraMessageDataBase>& MessageData = MessageKeyToMessageMap.FindOrAdd(Message.Key);
 		if (Message.Value != nullptr)
 		{
 			ReouterMergedObject(this, Message.Value);
-			MessageData = Message.Value;
+			MessageStore.AddMessage(Message.Key, Message.Value);
 		}
 	}
 
diff --git a/Engine/Plugins/FX/Niagara/Source/Niagara/Private/NiagaraMessageStore.cpp b/Engine/Plugins/FX/Niagara/Source/Niagara/Private/NiagaraMessageStore.cpp
new file mode 100644
--- /dev/null
+++ b/Engine/Plugins/FX/Niagara/Source/Niagara/Private/NiagaraMessageStore.cpp
@@ -0,0 +1,62 @@
+// Copyright Epic Games, Inc. All Rights Reserved.
+
+#include "NiagaraMessageStore.h"
+#include "NiagaraMessageDataBase.h"
+
+const TMap<FGuid, TObjectPtr<UNiagaraMessageDataBase>>& FNiagaraMessageStore::GetMessages() const
+{ 
+	return MessageKeyToMessageMap;
+}
+
+void FNiagaraMessageStore::SetMessages(const TMap<FGuid, TObjectPtr<UNiagaraMessageDataBase>>& InMessageKeyToMessageMap)
+{
+	MessageKeyToMessageMap = InMessageKeyToMessageMap;
+	DismissedMessageKeys.Empty();
+}
+
+void FNiagaraMessageStore::AddMessage(const FGuid& MessageKey, UNiagaraMessageDataBase* NewMessage)
+{
+	MessageKeyToMessageMap.Add(MessageKey, NewMessage);
+}
+
+void FNiagaraMessageStore::RemoveMessage(const FGuid& MessageKey) 
+{
+	MessageKeyToMessageMap.Remove(MessageKey);
+	DismissedMessageKeys.Remove(MessageKey);
+}
+
+void FNiagaraMessageStore::DismissMessage(const FGuid& MessageKey)
+{
+	DismissedMessageKeys.AddUnique(MessageKey);
+}
+
+bool FNiagaraMessageStore::IsMessageDismissed(const FGuid& MessageKey)
+{
+	return DismissedMessageKeys.Contains(MessageKey);
+}
+
+bool FNiagaraMessageStore::HasDismissedMessages() const
+{
+	return DismissedMessageKeys.IsEmpty() == false;
+}
+
+void FNiagaraMessageStore::ClearDismissedMessages()
+{
+	DismissedMessageKeys.Empty();
+}
+
+FNiagaraMessageSourceAndStore::FNiagaraMessageSourceAndStore(UObject& InSource, FNiagaraMessageStore& InStore)
+{
+	SourceWeak = &InSource;
+	Store = &InStore;
+}
+
+UObject* FNiagaraMessageSourceAndStore::GetSource() const
+{
+	return SourceWeak.Get();
+}
+
+FNiagaraMessageStore* FNiagaraMessageSourceAndStore::GetStore() const
+{
+	return SourceWeak.IsValid() ? Store : nullptr;
+}
\ No newline at end of file
diff --git a/Engine/Plugins/FX/Niagara/Source/Niagara/Private/NiagaraSystem.cpp b/Engine/Plugins/FX/Niagara/Source/Niagara/Private/NiagaraSystem.cpp
--- a/Engine/Plugins/FX/Niagara/Source/Niagara/Private/NiagaraSystem.cpp
+++ b/Engine/Plugins/FX/Niagara/Source/Niagara/Private/NiagaraSystem.cpp
@@ -971,6 +971,13 @@ void UNiagaraSystem::PostLoad()
 	{
 		LibraryVisibility = ENiagaraScriptLibraryVisibility::Unexposed;
 	}
+
+	if (MessageKeyToMessageMap_DEPRECATED.IsEmpty() == false)
+	{
+		MessageStore.SetMessages(MessageKeyToMessageMap_DEPRECATED);
+		MessageKeyToMessageMap_DEPRECATED.Empty();
+	}
+
 #endif // WITH_EDITORONLY_DATA
 
 	//Apply platform set redirectors
@@ -1598,6 +1605,10 @@ void UNiagaraSystem::FindDataInterfaceDependencies(FVersionedNiagaraEmitterData*
 	}
 }
 
+#if WITH_EDITORONLY_DATA
+const FGuid UNiagaraSystem::ComputeEmitterExecutionOrderMessageId(0x736C2644, 0x3B4F4D0E, 0xADEDE523, 0x2DACC2CF);
+#endif
+
 void UNiagaraSystem::ComputeEmittersExecutionOrder()
 {
 	const int32 NumEmitters = EmitterHandles.Num();
@@ -1716,8 +1727,23 @@ void UNiagaraSystem::ComputeEmittersExecutionOrder()
 			{
 				if (!ComputeEmitterPriority(EmitterIdx, EmitterPriorities, EmitterDependencyGraph))
 				{
+#if WITH_EDITORONLY_DATA
 					FName EmitterName = EmitterHandles[EmitterIdx].GetName();
-					UE_LOG(LogNiagara, Error, TEXT("Found circular dependency involving emitter '%s' in system '%s'. The execution order will be undefined."), *EmitterName.ToString(), *GetName());
+					FString ErrorMessage = FString::Printf(TEXT("Found circular dependency involving emitter '%s' in system '%s'. The execution order will be undefined."), *EmitterName.ToString(), *GetName());
+					INiagaraModule& NiagaraModule = FModuleManager::GetModuleChecked<INiagaraModule>("Niagara");
+					bool bAllowDismissal = true;
+					UNiagaraMessageDataBase* ComputeEmitterExecutionOrderMessage = NiagaraModule.GetEditorOnlyDataUtilities().CreateWarningMessage(
+						this,
+						LOCTEXT("ComputeEmitterOrderError", "Error computing emitter execution order"),
+						FText::FromString(ErrorMessage),
+						"Resolve Data Interfaces",
+						bAllowDismissal);
+					MessageStore.AddMessage(ComputeEmitterExecutionOrderMessageId, ComputeEmitterExecutionOrderMessage);
+					if (MessageStore.IsMessageDismissed(ComputeEmitterExecutionOrderMessageId) == false)
+					{
+						UE_LOG(LogNiagara, Warning, TEXT("%s"), *ErrorMessage);
+					}
+#endif
 					break;
 				}
 			}
@@ -2662,7 +2688,7 @@ void UNiagaraSystem::EvaluateCompileResultDependencies() const
 			if (!bDependencyMet)
 			{
 				FNiagaraCompileEvent LinkerErrorEvent(
-					FNiagaraCVarUtilities::GetCompileEventSeverityForFailIfNotSet(), Dependency.LinkerErrorMessage, FString(), false, Dependency.NodeGuid, Dependency.PinGuid, Dependency.StackGuids, FNiagaraCompileEventSource::ScriptDependency);
+					FNiagaraCVarUtilities::GetCompileEventSeverityForFailIfNotSet(), Dependency.LinkerErrorMessage, FString(), Dependency.NodeGuid, Dependency.PinGuid, Dependency.StackGuids, FNiagaraCompileEventSource::ScriptDependency);
 				ValidationInfo.CompileResults->LastCompileEvents.Add(LinkerErrorEvent);
 				ValidationInfo.bCompileResultStatusDirty = true;
 			}
diff --git a/Engine/Plugins/FX/Niagara/Source/Niagara/Public/INiagaraEditorOnlyDataUtlities.h b/Engine/Plugins/FX/Niagara/Source/Niagara/Public/INiagaraEditorOnlyDataUtlities.h
--- a/Engine/Plugins/FX/Niagara/Source/Niagara/Public/INiagaraEditorOnlyDataUtlities.h
+++ b/Engine/Plugins/FX/Niagara/Source/Niagara/Public/INiagaraEditorOnlyDataUtlities.h
@@ -7,6 +7,7 @@
 class UNiagaraScriptSourceBase;
 class UNiagaraEditorParametersAdapterBase;
 class UNiagaraEditorDataBase;
+class UNiagaraMessageDataBase;
 
 /** Defines utility methods for creating editor only data which is stored on runtime objects. */
 class INiagaraEditorOnlyDataUtilities
@@ -19,4 +20,8 @@ public:
 	virtual UNiagaraEditorParametersAdapterBase* CreateDefaultEditorParameters(UObject* InOuter) const = 0;
 
 	virtual UObject::FAssetRegistryTag CreateClassUsageAssetRegistryTag(const UObject* SourceObject) const = 0;
+
+	virtual UNiagaraMessageDataBase* CreateErrorMessage(UObject* InOuter, FText InMessageShort, FText InMessageLong, FName InTopicName, bool bInAllowDismissal = false) const = 0;
+
+	virtual UNiagaraMessageDataBase* CreateWarningMessage(UObject* InOuter, FText InMessageShort, FText InMessageLong, FName InTopicName, bool bInAllowDismissal = false) const = 0;
 };
\ No newline at end of file
diff --git a/Engine/Plugins/FX/Niagara/Source/Niagara/Public/NiagaraMessageDataBase.h b/Engine/Plugins/FX/Niagara/Source/Niagara/Public/NiagaraMessageDataBase.h
--- a/Engine/Plugins/FX/Niagara/Source/Niagara/Public/NiagaraMessageDataBase.h
+++ b/Engine/Plugins/FX/Niagara/Source/Niagara/Public/NiagaraMessageDataBase.h
@@ -12,4 +12,6 @@ class NIAGARA_API UNiagaraMessageDataBase : public UObject
 
 public:
 	UNiagaraMessageDataBase() = default;
+
+	virtual bool GetAllowDismissal() const { return false; }
 };
diff --git a/Engine/Plugins/FX/Niagara/Source/Niagara/Public/NiagaraMessageStore.h b/Engine/Plugins/FX/Niagara/Source/Niagara/Public/NiagaraMessageStore.h
new file mode 100644
--- /dev/null
+++ b/Engine/Plugins/FX/Niagara/Source/Niagara/Public/NiagaraMessageStore.h
@@ -0,0 +1,43 @@
+// Copyright Epic Games, Inc. All Rights Reserved.
+
+#pragma once
+
+#include "NiagaraMessageStore.generated.h"
+
+class UNiagaraMessageDataBase;
+
+USTRUCT()
+struct NIAGARA_API FNiagaraMessageStore
+{
+	GENERATED_BODY();
+
+	const TMap<FGuid, TObjectPtr<UNiagaraMessageDataBase>>& GetMessages() const;
+	void SetMessages(const TMap<FGuid, TObjectPtr<UNiagaraMessageDataBase>>& InMessageKeyToMessageMap);
+	void AddMessage(const FGuid& MessageKey, UNiagaraMessageDataBase* NewMessage);
+	void RemoveMessage(const FGuid& MessageKey);
+	void DismissMessage(const FGuid& MessageKey);
+	bool IsMessageDismissed(const FGuid& MessageKey);
+	bool HasDismissedMessages() const;
+	void ClearDismissedMessages();
+
+private:
+	UPROPERTY()
+	TMap<FGuid, TObjectPtr<UNiagaraMessageDataBase>> MessageKeyToMessageMap;
+
+	UPROPERTY()
+	TArray<FGuid> DismissedMessageKeys;
+};
+
+struct NIAGARA_API FNiagaraMessageSourceAndStore
+{
+public:
+	FNiagaraMessageSourceAndStore(UObject& InSource, FNiagaraMessageStore& InStore);
+
+	UObject* GetSource() const;
+
+	FNiagaraMessageStore* GetStore() const;
+
+private:
+	TWeakObjectPtr<UObject> SourceWeak;
+	FNiagaraMessageStore* Store;
+};
diff --git a/Engine/Plugins/FX/Niagara/Source/NiagaraEditor/Private/NiagaraEditorModule.cpp b/Engine/Plugins/FX/Niagara/Source/NiagaraEditor/Private/NiagaraEditorModule.cpp
--- a/Engine/Plugins/FX/Niagara/Source/NiagaraEditor/Private/NiagaraEditorModule.cpp
+++ b/Engine/Plugins/FX/Niagara/Source/NiagaraEditor/Private/NiagaraEditorModule.cpp
@@ -69,6 +69,7 @@
 #include "NiagaraSystemEditorData.h"
 #include "NiagaraClipboard.h"
 #include "NiagaraMessageManager.h"
+#include "NiagaraMessages.h"
 #include "NiagaraComponentBroker.h"
 #include "NiagaraBakerSettings.h"
 #include "ContentBrowserModule.h"
@@ -238,6 +239,24 @@ class FNiagaraEditorOnlyDataUtilities : public INiagaraEditorOnlyDataUtilities
 		const UNiagaraEditorSettings* NiagaraEditorSettings = GetDefault<UNiagaraEditorSettings>();
 		return NiagaraEditorSettings->CreateClassUsageAssetRegistryTag(SourceObject);
 	}
+
+	virtual UNiagaraMessageDataBase* CreateErrorMessage(UObject* InOuter, FText InMessageShort, FText InMessageLong, FName InTopicName, bool bInAllowDismissal = false) const override
+	{
+		return CreateMessage(InOuter, ENiagaraMessageSeverity::Error, InMessageShort, InMessageLong, InTopicName, bInAllowDismissal);
+	}
+
+	virtual UNiagaraMessageDataBase* CreateWarningMessage(UObject* InOuter, FText InMessageShort, FText InMessageLong, FName InTopicName, bool bInAllowDismissal = false) const override
+	{
+		return CreateMessage(InOuter, ENiagaraMessageSeverity::Warning, InMessageShort, InMessageLong, InTopicName, bInAllowDismissal);
+	}
+
+	UNiagaraMessageDataBase* CreateMessage(UObject* InOuter, ENiagaraMessageSeverity Severity, FText InMessageShort, FText InMessageLong, FName InTopicName, bool bInAllowDismissal) const
+	{
+		UNiagaraMessageDataText* NewMessage = NewObject<UNiagaraMessageDataText>(InOuter);
+		NewMessage->Init(InMessageLong, InMessageShort, Severity, InTopicName);
+		NewMessage->SetAllowDismissal(bInAllowDismissal);
+		return NewMessage;
+	}
 };
 
 class FNiagaraScriptGraphPanelPinFactory : public FGraphPanelPinFactory
diff --git a/Engine/Plugins/FX/Niagara/Source/NiagaraEditor/Private/NiagaraEditorUtilities.cpp b/Engine/Plugins/FX/Niagara/Source/NiagaraEditor/Private/NiagaraEditorUtilities.cpp
--- a/Engine/Plugins/FX/Niagara/Source/NiagaraEditor/Private/NiagaraEditorUtilities.cpp
+++ b/Engine/Plugins/FX/Niagara/Source/NiagaraEditor/Private/NiagaraEditorUtilities.cpp
@@ -4501,6 +4501,20 @@ bool FNiagaraEditorUtilities::IsEnumIndexVisible(const UEnum* Enum, int32 Index)
 	return FNiagaraEnumIndexVisibilityCache::GetVisibility(Enum, Index);
 }
 
+void FNiagaraEditorUtilities::GetScriptMessageStores(UNiagaraScript* InScript, TArray<FNiagaraMessageSourceAndStore>& OutNiagaraMessageStores)
+{
+	UNiagaraScriptSource* Source = Cast<UNiagaraScriptSource>(InScript->GetLatestSource());
+	if (Source != nullptr && Source->NodeGraph != nullptr)
+	{
+		TArray<UNiagaraNodeFunctionCall*> FunctionCallNodes;
+		Source->NodeGraph->GetNodesOfClass<UNiagaraNodeFunctionCall>(FunctionCallNodes);
+		for (UNiagaraNodeFunctionCall* FunctionCallNode : FunctionCallNodes)
+		{
+			OutNiagaraMessageStores.Add(FNiagaraMessageSourceAndStore(*FunctionCallNode, FunctionCallNode->GetMessageStore()));
+		}
+	}
+}
+
 void FNiagaraParameterUtilities::FilterToRelevantStaticVariables(const TArray<FNiagaraVariable>& InVars, TArray<FNiagaraVariable>& OutVars, FName InOldEmitterAlias, FName InNewEmitterAlias, bool bFilterByEmitterAliasAndConvertToUnaliased)
 {
 	FNiagaraAliasContext RenameContext(ENiagaraScriptUsage::ParticleSpawnScript);
diff --git a/Engine/Plugins/FX/Niagara/Source/NiagaraEditor/Private/NiagaraHlslTranslator.cpp b/Engine/Plugins/FX/Niagara/Source/NiagaraEditor/Private/NiagaraHlslTranslator.cpp
--- a/Engine/Plugins/FX/Niagara/Source/NiagaraEditor/Private/NiagaraHlslTranslator.cpp
+++ b/Engine/Plugins/FX/Niagara/Source/NiagaraEditor/Private/NiagaraHlslTranslator.cpp
@@ -7033,7 +7033,7 @@ void FHlslNiagaraTranslator::Operation(UNiagaraNodeOp* Operation, TArray<int32>&
 	FText ValidationError;
 	if (bEnforceStrictTypesValidations && OpInfo && OpInfo->InputTypeValidationFunction.IsBound() && OpInfo->InputTypeValidationFunction.Execute(InputTypes, ValidationError) == false)
 	{
-		Warning(ValidationError, Operation, OutputPins[0], TEXT("Invalid op types"), true);
+		Warning(ValidationError, Operation, OutputPins[0], TEXT("Invalid op types"));
 	}
 
 	if (OpInfo && OpInfo->StaticVariableResolveFunction.IsBound() && bAllPinsStatic)
@@ -9526,13 +9526,13 @@ FString FHlslNiagaraTranslator::NodePinToMessage(FText MessageText, const UNiaga
 	return MessageString;
 }
 
-void FHlslNiagaraTranslator::Message(FNiagaraCompileEventSeverity Severity, FText MessageText, const UNiagaraNode* InNode, const UEdGraphPin* Pin, FString ShortDescription, bool bDismissable)
+void FHlslNiagaraTranslator::Message(FNiagaraCompileEventSeverity Severity, FText MessageText, const UNiagaraNode* InNode, const UEdGraphPin* Pin, FString ShortDescription)
 {
 	const UNiagaraNode* CurContextNode = ActiveHistoryForFunctionCalls.GetCallingContext();
 	const UNiagaraNode* TargetNode = InNode ? InNode : CurContextNode;
  
 	FString MessageString = NodePinToMessage(MessageText, TargetNode, Pin);
-	TranslateResults.CompileEvents.Add(FNiagaraCompileEvent(Severity, MessageString, ShortDescription, bDismissable, TargetNode ? TargetNode->NodeGuid : FGuid(), Pin ? Pin->PersistentGuid : FGuid(), GetCallstackGuids()));
+	TranslateResults.CompileEvents.Add(FNiagaraCompileEvent(Severity, MessageString, ShortDescription, TargetNode ? TargetNode->NodeGuid : FGuid(), Pin ? Pin->PersistentGuid : FGuid(), GetCallstackGuids()));
  
 	if (Severity == FNiagaraCompileEventSeverity::Error)
 	{
@@ -9544,14 +9544,14 @@ void FHlslNiagaraTranslator::Message(FNiagaraCompileEventSeverity Severity, FTex
 	}
 }
  
-void FHlslNiagaraTranslator::Error(FText ErrorText, const UNiagaraNode* InNode, const UEdGraphPin* Pin, FString ShortDescription, bool bDismissable)
+void FHlslNiagaraTranslator::Error(FText ErrorText, const UNiagaraNode* InNode, const UEdGraphPin* Pin, FString ShortDescription)
 {
-	Message(FNiagaraCompileEventSeverity::Error, ErrorText, InNode, Pin, ShortDescription, bDismissable);
+	Message(FNiagaraCompileEventSeverity::Error, ErrorText, InNode, Pin, ShortDescription);
 }
  
-void FHlslNiagaraTranslator::Warning(FText WarningText, const UNiagaraNode* InNode, const UEdGraphPin* Pin, FString ShortDescription, bool bDismissable)
+void FHlslNiagaraTranslator::Warning(FText WarningText, const UNiagaraNode* InNode, const UEdGraphPin* Pin, FString ShortDescription)
 {
-	Message(FNiagaraCompileEventSeverity::Warning, WarningText, InNode, Pin, ShortDescription, bDismissable);
+	Message(FNiagaraCompileEventSeverity::Warning, WarningText, InNode, Pin, ShortDescription);
 }
 
 void FHlslNiagaraTranslator::RegisterCompileDependency(const FNiagaraVariableBase& InVar, FText MessageText, const UNiagaraNode* Node, const UEdGraphPin* Pin, bool bEmitAsLinker, int32 ParamMapHistoryIdx)
@@ -9584,7 +9584,7 @@ void FHlslNiagaraTranslator::RegisterCompileDependency(const FNiagaraVariableBas
 	{
 		const bool bDismissable = false;
 		const FString ShortDescription = FString();
-		Message(FNiagaraCVarUtilities::GetCompileEventSeverityForFailIfNotSet(), MessageText, Node, Pin, ShortDescription, bDismissable);
+		Message(FNiagaraCVarUtilities::GetCompileEventSeverityForFailIfNotSet(), MessageText, Node, Pin, ShortDescription);
 	}
 }
 
diff --git a/Engine/Plugins/FX/Niagara/Source/NiagaraEditor/Private/NiagaraHlslTranslator.h b/Engine/Plugins/FX/Niagara/Source/NiagaraEditor/Private/NiagaraHlslTranslator.h
--- a/Engine/Plugins/FX/Niagara/Source/NiagaraEditor/Private/NiagaraHlslTranslator.h
+++ b/Engine/Plugins/FX/Niagara/Source/NiagaraEditor/Private/NiagaraHlslTranslator.h
@@ -633,9 +633,9 @@ public:
 	
 	void WriteCompilerTag(int32 InputCompileResult, const UEdGraphPin* Pin, bool bEmitMessageOnFailure, FNiagaraCompileEventSeverity FailureSeverity, const FString& Prefix = FString());
 
-	void Message(FNiagaraCompileEventSeverity Severity, FText MessageText, const UNiagaraNode* Node, const UEdGraphPin* Pin, FString ShortDescription = FString(), bool bDismissable = false);
-	virtual void Error(FText ErrorText, const UNiagaraNode* Node, const UEdGraphPin* Pin, FString ShortDescription = FString(), bool bDismissable = false);
-	virtual void Warning(FText WarningText, const UNiagaraNode* Node, const UEdGraphPin* Pin, FString ShortDescription = FString(), bool bDismissable = false);
+	void Message(FNiagaraCompileEventSeverity Severity, FText MessageText, const UNiagaraNode* Node, const UEdGraphPin* Pin, FString ShortDescription = FString());
+	virtual void Error(FText ErrorText, const UNiagaraNode* Node, const UEdGraphPin* Pin, FString ShortDescription = FString());
+	virtual void Warning(FText WarningText, const UNiagaraNode* Node, const UEdGraphPin* Pin, FString ShortDescription = FString());
 	void RegisterCompileDependency(const FNiagaraVariableBase& InVar, FText ErrorText, const UNiagaraNode* Node, const UEdGraphPin* Pin, bool bEmitAsLinker, int32 ParamMapHistoryIdx);
 	FString NodePinToMessage(FText MessageText, const UNiagaraNode* Node, const UEdGraphPin* Pin);
 
diff --git a/Engine/Plugins/FX/Niagara/Source/NiagaraEditor/Private/NiagaraMessageUtilities.cpp b/Engine/Plugins/FX/Niagara/Source/NiagaraEditor/Private/NiagaraMessageUtilities.cpp
--- a/Engine/Plugins/FX/Niagara/Source/NiagaraEditor/Private/NiagaraMessageUtilities.cpp
+++ b/Engine/Plugins/FX/Niagara/Source/NiagaraEditor/Private/NiagaraMessageUtilities.cpp
@@ -85,7 +85,6 @@ FText FNiagaraMessageUtilities::MakePostCompileSummaryText(const FText& Compiled
 
 UNiagaraStackEntry::FStackIssue FNiagaraMessageUtilities::MessageToStackIssue(TSharedRef<const INiagaraMessage> InMessage, FString InStackEditorDataKey)
 {
-	bool bDismissable = InMessage->AllowDismissal();
 	FText ShortDescription = InMessage->GenerateMessageTitle();
 	
 	TSharedRef<FTokenizedMessage> TokenizedMessage = InMessage->GenerateTokenizedMessage();
@@ -131,8 +130,9 @@ UNiagaraStackEntry::FStackIssue FNiagaraMessageUtilities::MessageToStackIssue(TS
 		ShortDescription,
 		InMessage->GenerateMessageText(),
 		InStackEditorDataKey,
-		bDismissable,
-		FixLinks);
+		InMessage->GetDismissHandler().IsBound(),
+		FixLinks,
+		InMessage->GetDismissHandler());
 }
 
 UNiagaraStackEntry::FStackIssue FNiagaraMessageUtilities::StackMessageToStackIssue(const FNiagaraStackMessage& InMessage, FString InStackEditorDataKey, const TArray<FLinkNameAndDelegate>& Links)
diff --git a/Engine/Plugins/FX/Niagara/Source/NiagaraEditor/Private/NiagaraMessages.cpp b/Engine/Plugins/FX/Niagara/Source/NiagaraEditor/Private/NiagaraMessages.cpp
--- a/Engine/Plugins/FX/Niagara/Source/NiagaraEditor/Private/NiagaraMessages.cpp
+++ b/Engine/Plugins/FX/Niagara/Source/NiagaraEditor/Private/NiagaraMessages.cpp
@@ -21,11 +21,6 @@ FText INiagaraMessage::GenerateMessageTitle() const
 	return FText::GetEmpty();
 }
 
-bool INiagaraMessage::AllowDismissal() const
-{
-	return false;
-}
-
 FNiagaraMessageCompileEvent::FNiagaraMessageCompileEvent(
 	const FNiagaraCompileEvent& InCompileEvent
 	, TArray<FNiagaraScriptNameAndAssetPath>& InContextScriptNamesAndAssetPaths
@@ -177,11 +172,6 @@ FText FNiagaraMessageCompileEvent::GenerateMessageTitle() const
 	}
 }
 
-bool FNiagaraMessageCompileEvent::AllowDismissal() const
-{
-	return CompileEvent.bDismissable;
-}
-
 void FNiagaraMessageCompileEvent::GenerateLinks(TArray<FText>& OutLinkDisplayNames, TArray<FSimpleDelegate>& OutLinkNavigationActions) const
 {
 	// TODO: Do this without instantiating these tokens.
@@ -420,11 +410,6 @@ FText FNiagaraMessageText::GenerateMessageTitle() const
 	return ShortDescription;
 }
 
-bool FNiagaraMessageText::AllowDismissal() const
-{
-	return bAllowDismissal;
-}
-
 const uint32 INiagaraMessage::GetMessageTopicBitflag() const
 {
 	if (MessageTopicBitflag == 0)
@@ -454,9 +439,9 @@ TSharedRef<const INiagaraMessage> UNiagaraMessageDataText::GenerateNiagaraMessag
 	const TArray<FLinkNameAndDelegate>& Links = InGenerateInfo.GetLinks();
 	if (Links.Num() > 0)
 	{
-		return MakeShared<const FNiagaraMessageTextWithLinks>(MessageText, ShortDescription, EMessageSeverity::Type(MessageSeverity), TopicName, bAllowDismissal, Links, InGenerateInfo.GetAssociatedObjectKeys());
+		return MakeShared<const FNiagaraMessageTextWithLinks>(MessageText, ShortDescription, EMessageSeverity::Type(MessageSeverity), TopicName, InGenerateInfo.GetDismissHandler(), Links, InGenerateInfo.GetAssociatedObjectKeys());
 	}
-	return MakeShared<const FNiagaraMessageText>(MessageText, ShortDescription, EMessageSeverity::Type(MessageSeverity), TopicName, bAllowDismissal, InGenerateInfo.GetAssociatedObjectKeys());
+	return MakeShared<const FNiagaraMessageText>(MessageText, ShortDescription, EMessageSeverity::Type(MessageSeverity), TopicName, InGenerateInfo.GetDismissHandler(), InGenerateInfo.GetAssociatedObjectKeys());
 }
 
 TSharedRef<FTokenizedMessage> FNiagaraMessageTextWithLinks::GenerateTokenizedMessage() const
diff --git a/Engine/Plugins/FX/Niagara/Source/NiagaraEditor/Private/NiagaraNodeFunctionCall.cpp b/Engine/Plugins/FX/Niagara/Source/NiagaraEditor/Private/NiagaraNodeFunctionCall.cpp
--- a/Engine/Plugins/FX/Niagara/Source/NiagaraEditor/Private/NiagaraNodeFunctionCall.cpp
+++ b/Engine/Plugins/FX/Niagara/Source/NiagaraEditor/Private/NiagaraNodeFunctionCall.cpp
@@ -165,6 +165,15 @@ void UNiagaraNodeFunctionCall::PostLoad()
 	{
 		ComputeNodeName();
 	}
+
+	if (MessageKeyToMessageMap_DEPRECATED.IsEmpty() == false)
+	{
+		for (auto KeyMessagePair : MessageKeyToMessageMap_DEPRECATED)
+		{
+			MessageStore.AddMessage(KeyMessagePair.Key, KeyMessagePair.Value);
+		}
+		MessageKeyToMessageMap_DEPRECATED.Empty();
+	}
 }
 
 void UNiagaraNodeFunctionCall::UpgradeDIFunctionCalls()
diff --git a/Engine/Plugins/FX/Niagara/Source/NiagaraEditor/Private/NiagaraNodeStaticSwitch.cpp b/Engine/Plugins/FX/Niagara/Source/NiagaraEditor/Private/NiagaraNodeStaticSwitch.cpp
--- a/Engine/Plugins/FX/Niagara/Source/NiagaraEditor/Private/NiagaraNodeStaticSwitch.cpp
+++ b/Engine/Plugins/FX/Niagara/Source/NiagaraEditor/Private/NiagaraNodeStaticSwitch.cpp
@@ -521,11 +521,11 @@ void UNiagaraNodeStaticSwitch::CheckForOutdatedEnum(FHlslNiagaraTranslator* Tran
 		{
 			if(NewPinNames.Num() == 0)
 			{
-				Translator->Message(FNiagaraCompileEventSeverity::Display, LOCTEXT("EmptyEnumEntries", "A static switch either has no enum assigned or the enum has no entries. A node refresh is advised."), this, nullptr, FString(), true);
+				Translator->Message(FNiagaraCompileEventSeverity::Display, LOCTEXT("EmptyEnumEntries", "A static switch either has no enum assigned or the enum has no entries. A node refresh is advised."), this, nullptr, FString());
 			}
 			else
 			{
-				Translator->Message(FNiagaraCompileEventSeverity::Log, LOCTEXT("PinCountMismatch", "Static switch pins reflect outdated enum entries. While the static switch is fully functional, a node refresh is advised."), this, nullptr, FString(), true);
+				Translator->Message(FNiagaraCompileEventSeverity::Log, LOCTEXT("PinCountMismatch", "Static switch pins reflect outdated enum entries. While the static switch is fully functional, a node refresh is advised."), this, nullptr, FString());
 			}
 		}
 		else
@@ -542,7 +542,7 @@ void UNiagaraNodeStaticSwitch::CheckForOutdatedEnum(FHlslNiagaraTranslator* Tran
 
 			if (bAnyPinNameChangeDetected)
 			{
-				Translator->Message(FNiagaraCompileEventSeverity::Log, LOCTEXT("EnumPinNameChangeDetected", "Static switch pins reflect outdated enum entries. While the static switch is fully functional, a node refresh is advised."), this, nullptr, FString(), true);
+				Translator->Message(FNiagaraCompileEventSeverity::Log, LOCTEXT("EnumPinNameChangeDetected", "Static switch pins reflect outdated enum entries. While the static switch is fully functional, a node refresh is advised."), this, nullptr, FString());
 			}
 		}
 	}
diff --git a/Engine/Plugins/FX/Niagara/Source/NiagaraEditor/Private/ViewModels/NiagaraEmitterHandleViewModel.cpp b/Engine/Plugins/FX/Niagara/Source/NiagaraEditor/Private/ViewModels/NiagaraEmitterHandleViewModel.cpp
--- a/Engine/Plugins/FX/Niagara/Source/NiagaraEditor/Private/ViewModels/NiagaraEmitterHandleViewModel.cpp
+++ b/Engine/Plugins/FX/Niagara/Source/NiagaraEditor/Private/ViewModels/NiagaraEmitterHandleViewModel.cpp
@@ -84,7 +84,7 @@ FGuid FNiagaraEmitterHandleViewModel::AddMessage(UNiagaraMessageData* NewMessage
 	{
 		const FGuid NewGuid = InNewGuid.IsValid() ? InNewGuid : FGuid::NewGuid();
 		
-		EmitterHandle->GetInstance().Emitter->AddMessage(NewGuid, NewMessage);
+		EmitterHandle->GetInstance().Emitter->GetMessageStore().AddMessage(NewGuid, NewMessage);
 		return NewGuid;
 	}
 	return FGuid();
@@ -94,7 +94,7 @@ void FNiagaraEmitterHandleViewModel::RemoveMessage(const FGuid& MessageKey) cons
 {
 	if (ensureMsgf(EmitterHandle != nullptr, TEXT("EmitterHandleViewModel had a null EmitterHandle!")))
 	{
-		EmitterHandle->GetInstance().Emitter->RemoveMessage(MessageKey);
+		EmitterHandle->GetInstance().Emitter->GetMessageStore().RemoveMessage(MessageKey);
 	}
 }
 
diff --git a/Engine/Plugins/FX/Niagara/Source/NiagaraEditor/Private/ViewModels/NiagaraEmitterViewModel.cpp b/Engine/Plugins/FX/Niagara/Source/NiagaraEditor/Private/ViewModels/NiagaraEmitterViewModel.cpp
--- a/Engine/Plugins/FX/Niagara/Source/NiagaraEditor/Private/ViewModels/NiagaraEmitterViewModel.cpp
+++ b/Engine/Plugins/FX/Niagara/Source/NiagaraEditor/Private/ViewModels/NiagaraEmitterViewModel.cpp
@@ -336,6 +336,20 @@ UNiagaraEmitterEditorData& FNiagaraEmitterViewModel::GetEditorData()
 	return *Cast<UNiagaraEmitterEditorData>(EmitterWeakPtr.GetEmitterData()->GetEditorData());
 }
 
+void FNiagaraEmitterViewModel::GetEmitterMessageStores(TArray<FNiagaraMessageSourceAndStore>& OutMessageStores)
+{
+	if (EmitterWeakPtr.IsValid())
+	{
+		OutMessageStores.Add(FNiagaraMessageSourceAndStore(*EmitterWeakPtr.Emitter, EmitterWeakPtr.Emitter->GetMessageStore()));
+		TArray<UNiagaraScript*> Scripts;
+		EmitterWeakPtr.GetEmitterData()->GetScripts(Scripts, false);
+		for (UNiagaraScript* Script : Scripts)
+		{
+			FNiagaraEditorUtilities::GetScriptMessageStores(Script, OutMessageStores);
+		}
+	}
+}
+
 void FNiagaraEmitterViewModel::AddEventHandler(FNiagaraEventScriptProperties& EventScriptProperties, bool bResetGraphForOutput /*= false*/)
 {
 	EventScriptProperties.Script = NewObject<UNiagaraScript>(GetEmitter().Emitter, MakeUniqueObjectName(GetEmitter().Emitter, UNiagaraScript::StaticClass(), "EventScript"), EObjectFlags::RF_Transactional);
diff --git a/Engine/Plugins/FX/Niagara/Source/NiagaraEditor/Private/ViewModels/NiagaraSystemViewModel.cpp b/Engine/Plugins/FX/Niagara/Source/NiagaraEditor/Private/ViewModels/NiagaraSystemViewModel.cpp
--- a/Engine/Plugins/FX/Niagara/Source/NiagaraEditor/Private/ViewModels/NiagaraSystemViewModel.cpp
+++ b/Engine/Plugins/FX/Niagara/Source/NiagaraEditor/Private/ViewModels/NiagaraSystemViewModel.cpp
@@ -786,7 +786,7 @@ FGuid FNiagaraSystemViewModel::AddMessage(UNiagaraMessageData* NewMessage) const
 	UNiagaraSystem& ViewedSystem = GetSystem();
 	if (EditMode == ENiagaraSystemViewModelEditMode::SystemAsset)
 	{
-		ViewedSystem.AddMessage(NewGuid, static_cast<UNiagaraMessageDataBase*>(NewMessage));
+		ViewedSystem.GetMessageStore().AddMessage(NewGuid, static_cast<UNiagaraMessageDataBase*>(NewMessage));
 	}
 	else if (EditMode == ENiagaraSystemViewModelEditMode::EmitterAsset)
 	{
@@ -810,7 +810,7 @@ void FNiagaraSystemViewModel::RemoveMessage(const FGuid& MessageKey) const
 	UNiagaraSystem& ViewedSystem = GetSystem();
 	if (EditMode == ENiagaraSystemViewModelEditMode::SystemAsset)
 	{
-		ViewedSystem.RemoveMessage(MessageKey);
+		ViewedSystem.GetMessageStore().RemoveMessage(MessageKey);
 	}
 	else if (EditMode == ENiagaraSystemViewModelEditMode::EmitterAsset)
 	{
@@ -830,21 +830,26 @@ void FNiagaraSystemViewModel::RemoveMessage(const FGuid& MessageKey) const
 FGuid FNiagaraSystemViewModel::AddStackMessage(UNiagaraMessageData* NewMessage, UNiagaraNodeFunctionCall* TargetFunctionCallNode) const
 {
 	FGuid NewGuid = FGuid::NewGuid();
-	TargetFunctionCallNode->AddMessage(NewGuid, NewMessage);
+	TargetFunctionCallNode->GetMessageStore().AddMessage(NewGuid, NewMessage);
 	bPendingAssetMessagesChanged = true;
 	return NewGuid;
 }
 
 NIAGARAEDITOR_API void FNiagaraSystemViewModel::RemoveStackMessage(const FGuid& MessageKey, UNiagaraNodeFunctionCall* TargetFunctionCallNode) const
 {
-	TargetFunctionCallNode->RemoveMessage(MessageKey);
+	TargetFunctionCallNode->GetMessageStore().RemoveMessage(MessageKey);
 	bPendingAssetMessagesChanged = true;
 }
 
-void FNiagaraSystemViewModel::ExecuteMessageDelegateAndRefreshMessages(FSimpleDelegate MessageDelegate)
+void FNiagaraSystemViewModel::DismissMessageAndRefresh(FNiagaraMessageSourceAndStore MessageSourceAndStore, FGuid MessageKey)
 {
-	MessageDelegate.ExecuteIfBound();
-	bPendingAssetMessagesChanged = true;
+	UObject* MessageSource = MessageSourceAndStore.GetSource();
+	if (MessageSource != nullptr)
+	{
+		MessageSource->Modify();
+		MessageSourceAndStore.GetStore()->DismissMessage(MessageKey);
+		bPendingAssetMessagesChanged = true;
+	}
 }
 
 void FNiagaraSystemViewModel::DeleteEmitters(TSet<FGuid> EmitterHandleIdsToDelete)
@@ -1535,6 +1540,18 @@ TArray<FNiagaraGraphParameterReference> FNiagaraSystemViewModel::GetGraphParamet
 	return GraphParameterReferences;
 }
 
+void FNiagaraSystemViewModel::GetSystemMessageStores(TArray<FNiagaraMessageSourceAndStore>& OutMessageStores)
+{
+	OutMessageStores.Add(FNiagaraMessageSourceAndStore(*System, System->GetMessageStore()));
+	FNiagaraEditorUtilities::GetScriptMessageStores(System->GetSystemSpawnScript(), OutMessageStores);
+	FNiagaraEditorUtilities::GetScriptMessageStores(System->GetSystemUpdateScript(), OutMessageStores);
+}
+
+void FNiagaraSystemViewModel::RefreshAssetMessagesDeferred()
+{
+	bPendingAssetMessagesChanged = true;
+}
+
 void FNiagaraSystemViewModel::SendLastCompileMessageJobs() const
 {
 	if (SystemMessageLogGuidKey.IsSet() == false)
@@ -1729,6 +1746,7 @@ void FNiagaraSystemViewModel::TickCompileStatus()
 			{
 				SendLastCompileMessageJobs();
 				RefreshStackViewModels();
+				bPendingAssetMessagesChanged = true;
 			}
 		}
 	}
@@ -3036,36 +3054,35 @@ void FNiagaraSystemViewModel::RefreshStackViewModels()
 void FNiagaraSystemViewModel::RefreshAssetMessages()
 {
 	FNiagaraMessageManager* MessageManager = FNiagaraMessageManager::Get();
-	for (const FObjectKey& FunctionCallNodeObjectKey : LastFunctionCallNodeObjectKeys)
+	for (const FObjectKey& MessageSourceObjectKey : MessageSourceObjectKeys)
 	{
-		MessageManager->ClearAssetMessagesForObject(SystemMessageLogGuidKey.GetValue(), FunctionCallNodeObjectKey);
+		MessageManager->ClearAssetMessagesForObject(SystemMessageLogGuidKey.GetValue(), MessageSourceObjectKey);
 	}
-	LastFunctionCallNodeObjectKeys.Reset();
+	MessageSourceObjectKeys.Reset();
 
-	auto PublishMessages = [this, &MessageManager](auto& MessageSource, const auto RemoveMessageFuncSig)->TArray<FObjectKey> /* MessageSourceObjectKeys */ {
-		TArray<FObjectKey> NewObjectKeys;
-		for (auto It = MessageSource.GetMessages().CreateConstIterator(); It; ++It)
+	auto PublishMessages = [this, &MessageManager](const FNiagaraMessageSourceAndStore& MessageSourceAndStore) {
+		for (auto It = MessageSourceAndStore.GetStore()->GetMessages().CreateConstIterator(); It; ++It)
 		{
+			const FGuid& MessageKey = It.Key();
+			if (MessageSourceAndStore.GetStore()->IsMessageDismissed(MessageKey))
+			{
+				continue;
+			}
+
 			const UNiagaraMessageData* MessageData = static_cast<UNiagaraMessageData*>(It.Value());
 			FGenerateNiagaraMessageInfo GenerateNiagaraMessageInfo = FGenerateNiagaraMessageInfo();
 
-			TArray<FObjectKey> AssociatedObjectKeys = { FObjectKey(&MessageSource) };
-			NewObjectKeys.Append(AssociatedObjectKeys);
+			TArray<FObjectKey> AssociatedObjectKeys = { FObjectKey(MessageSourceAndStore.GetSource()) };
 			GenerateNiagaraMessageInfo.SetAssociatedObjectKeys(AssociatedObjectKeys);
 
-			TArray<FLinkNameAndDelegate> Links;
-			const FText LinkText = LOCTEXT("AcknowledgeAndClearIssue", "Acknowledge and clear this issue.");
-			const FGuid MessageKey = It.Key();
-			FSimpleDelegate MessageDelegate = FSimpleDelegate::CreateUObject(&MessageSource, RemoveMessageFuncSig, MessageKey);
-			FSimpleDelegate WrapperDelegate = FSimpleDelegate::CreateSP(this, &FNiagaraSystemViewModel::ExecuteMessageDelegateAndRefreshMessages, MessageDelegate);
-			const FLinkNameAndDelegate Link = FLinkNameAndDelegate(LinkText, WrapperDelegate);
-			Links.Add(Link);
-			GenerateNiagaraMessageInfo.SetLinks(Links);
+			if (MessageData->GetAllowDismissal())
+			{
+				GenerateNiagaraMessageInfo.SetDismissHandler(FSimpleDelegate::CreateSP(this, &FNiagaraSystemViewModel::DismissMessageAndRefresh, MessageSourceAndStore, It.Key()));
+			}
 
 			TSharedRef<const INiagaraMessage> Message = MessageData->GenerateNiagaraMessage(GenerateNiagaraMessageInfo);
 			MessageManager->AddMessage(Message, SystemMessageLogGuidKey.GetValue());
 		}
-		return NewObjectKeys;
 	};
 
 	auto PublishScriptMessages = [&](const TArray<UNiagaraScript*>& Scripts)->TArray<FObjectKey> /* FunctionCallNodeObjectKeys */ {
@@ -3082,7 +3099,8 @@ void FNiagaraSystemViewModel::RefreshAssetMessages()
 			FNiagaraStackGraphUtilities::GetOrderedModuleNodes(*OutputNode, FunctionCallNodes);
 			for(UNiagaraNodeFunctionCall* FunctionCallNode : FunctionCallNodes)
 			{
-				NewObjectKeys.Append(PublishMessages(*FunctionCallNode, &UNiagaraNodeFunctionCall::RemoveMessageDelegateable));
+				PublishMessages(FNiagaraMessageSourceAndStore(*FunctionCallNode, FunctionCallNode->GetMessageStore()));
+				NewObjectKeys.Add(FObjectKey(FunctionCallNode));
 			}
 		}
 		return NewObjectKeys;
@@ -3091,7 +3109,8 @@ void FNiagaraSystemViewModel::RefreshAssetMessages()
 	UNiagaraSystem& ViewedSystem = GetSystem();
 	if (EditMode == ENiagaraSystemViewModelEditMode::SystemAsset)
 	{
-		LastFunctionCallNodeObjectKeys.Append(PublishMessages(ViewedSystem, &UNiagaraSystem::RemoveMessageDelegateable));
+		PublishMessages(FNiagaraMessageSourceAndStore(ViewedSystem, ViewedSystem.GetMessageStore()));
+		MessageSourceObjectKeys.Add(FObjectKey(&ViewedSystem));
 
 		TArray<UNiagaraScript*> Scripts;
 		Scripts.Add(ViewedSystem.GetSystemSpawnScript());
@@ -3099,10 +3118,12 @@ void FNiagaraSystemViewModel::RefreshAssetMessages()
 		for (const FNiagaraEmitterHandle& EmitterHandle : ViewedSystem.GetEmitterHandles())
 		{
 			FVersionedNiagaraEmitter VersionedEmitter = EmitterHandle.GetInstance();
-			LastFunctionCallNodeObjectKeys.Append(PublishMessages(*VersionedEmitter.Emitter, &UNiagaraEmitter::RemoveMessageDelegateable));
+			UNiagaraEmitter* Emitter = VersionedEmitter.Emitter.Get();
+			PublishMessages(FNiagaraMessageSourceAndStore(*Emitter, Emitter->GetMessageStore()));
+			MessageSourceObjectKeys.Add(FObjectKey(Emitter));
 			VersionedEmitter.GetEmitterData()->GetScripts(Scripts, false);
 		}
-		LastFunctionCallNodeObjectKeys.Append(PublishScriptMessages(Scripts));
+		MessageSourceObjectKeys.Append(PublishScriptMessages(Scripts));
 
 	}
 	else if (EditMode == ENiagaraSystemViewModelEditMode::EmitterAsset)
@@ -3111,11 +3132,13 @@ void FNiagaraSystemViewModel::RefreshAssetMessages()
 		if (ensureMsgf(EmitterHandles.Num() == 1, TEXT("There was not exactly 1 Emitter Handle for the SystemViewModel in Emitter edit mode!")))
 		{
 			FVersionedNiagaraEmitter VersionedEmitter = EmitterHandles[0].GetInstance();
-			LastFunctionCallNodeObjectKeys.Append(PublishMessages(*VersionedEmitter.Emitter, &UNiagaraEmitter::RemoveMessageDelegateable));
+			UNiagaraEmitter* Emitter = VersionedEmitter.Emitter.Get();
+			PublishMessages(FNiagaraMessageSourceAndStore(*Emitter, Emitter->GetMessageStore()));
+			MessageSourceObjectKeys.Add(FObjectKey(Emitter));
 
 			TArray<UNiagaraScript*> Scripts;
 			VersionedEmitter.GetEmitterData()->GetScripts(Scripts, false);
-			LastFunctionCallNodeObjectKeys.Append(PublishScriptMessages(Scripts));
+			MessageSourceObjectKeys.Append(PublishScriptMessages(Scripts));
 		}
 	}
 	else
diff --git a/Engine/Plugins/FX/Niagara/Source/NiagaraEditor/Private/ViewModels/Stack/NiagaraStackEntry.cpp b/Engine/Plugins/FX/Niagara/Source/NiagaraEditor/Private/ViewModels/Stack/NiagaraStackEntry.cpp
--- a/Engine/Plugins/FX/Niagara/Source/NiagaraEditor/Private/ViewModels/Stack/NiagaraStackEntry.cpp
+++ b/Engine/Plugins/FX/Niagara/Source/NiagaraEditor/Private/ViewModels/Stack/NiagaraStackEntry.cpp
@@ -73,17 +73,20 @@ UNiagaraStackEntry::EStackIssueFixStyle UNiagaraStackEntry::FStackIssueFix::GetS
 	return Style;
 }
 
-UNiagaraStackEntry::FStackIssue::FStackIssue() : Severity(EStackIssueSeverity::None)
+UNiagaraStackEntry::FStackIssue::FStackIssue() 
+	: Severity(EStackIssueSeverity::None)
+	, DismissHandler(nullptr)
 {
 }
 
-UNiagaraStackEntry::FStackIssue::FStackIssue(EStackIssueSeverity InSeverity, FText InShortDescription, FText InLongDescription, FString InStackEditorDataKey, bool bInCanBeDismissed, const TArray<FStackIssueFix>& InFixes)
+UNiagaraStackEntry::FStackIssue::FStackIssue(EStackIssueSeverity InSeverity, FText InShortDescription, FText InLongDescription, FString InStackEditorDataKey, bool bInCanBeDismissed, const TArray<FStackIssueFix>& InFixes, const FSimpleDelegate& InDismissHandler)
 	: Severity(InSeverity)
 	, ShortDescription(InShortDescription)
 	, LongDescription(InLongDescription)
 	, UniqueIdentifier(FMD5::HashAnsiString(*FString::Printf(TEXT("%s-%s-%s"), *InStackEditorDataKey, *InShortDescription.ToString(), *InLongDescription.ToString())))
 	, bCanBeDismissed(bInCanBeDismissed)
 	, Fixes(InFixes)
+	, DismissHandler(InDismissHandler)
 {
 	checkf(ShortDescription.IsEmptyOrWhitespace() == false, TEXT("Short description can not be empty."));
 	//checkf(LongDescription.IsEmptyOrWhitespace() == false, TEXT("Long description can not be empty."));
@@ -129,6 +132,11 @@ bool UNiagaraStackEntry::FStackIssue::GetCanBeDismissed() const
 	return bCanBeDismissed;
 }
 
+const FSimpleDelegate& UNiagaraStackEntry::FStackIssue::GetDismissHandler() const
+{
+	return DismissHandler;
+}
+
 const FString& UNiagaraStackEntry::FStackIssue::GetUniqueIdentifier() const
 {
 	return UniqueIdentifier;
diff --git a/Engine/Plugins/FX/Niagara/Source/NiagaraEditor/Private/ViewModels/Stack/NiagaraStackErrorItem.cpp b/Engine/Plugins/FX/Niagara/Source/NiagaraEditor/Private/ViewModels/Stack/NiagaraStackErrorItem.cpp
--- a/Engine/Plugins/FX/Niagara/Source/NiagaraEditor/Private/ViewModels/Stack/NiagaraStackErrorItem.cpp
+++ b/Engine/Plugins/FX/Niagara/Source/NiagaraEditor/Private/ViewModels/Stack/NiagaraStackErrorItem.cpp
@@ -214,8 +214,15 @@ void UNiagaraStackErrorItemDismiss::Initialize(FRequiredEntryData InRequiredEntr
 
 void UNiagaraStackErrorItemDismiss::DismissIssue()
 {
-	GetStackEditorData().Modify();
-	GetStackEditorData().DismissStackIssue(StackIssue.GetUniqueIdentifier());
+	if (StackIssue.GetDismissHandler().IsBound())
+	{
+		StackIssue.GetDismissHandler().Execute();
+	}
+	else
+	{
+		GetStackEditorData().Modify();
+		GetStackEditorData().DismissStackIssue(StackIssue.GetUniqueIdentifier());
+	}
 }
 
 UNiagaraStackEntry::EStackRowStyle UNiagaraStackErrorItemDismiss::GetStackRowStyle() const
diff --git a/Engine/Plugins/FX/Niagara/Source/NiagaraEditor/Private/ViewModels/Stack/NiagaraStackViewModel.cpp b/Engine/Plugins/FX/Niagara/Source/NiagaraEditor/Private/ViewModels/Stack/NiagaraStackViewModel.cpp
--- a/Engine/Plugins/FX/Niagara/Source/NiagaraEditor/Private/ViewModels/Stack/NiagaraStackViewModel.cpp
+++ b/Engine/Plugins/FX/Niagara/Source/NiagaraEditor/Private/ViewModels/Stack/NiagaraStackViewModel.cpp
@@ -56,6 +56,18 @@ UNiagaraStackEditorData* UNiagaraStackViewModel::FTopLevelViewModel::GetStackEdi
 	}
 }
 
+void UNiagaraStackViewModel::FTopLevelViewModel::GetMessageStores(TArray<FNiagaraMessageSourceAndStore>& OutMessageStores)
+{
+	if (SystemViewModel.IsValid())
+	{
+		SystemViewModel->GetSystemMessageStores(OutMessageStores);
+	}
+	else if (EmitterHandleViewModel.IsValid())
+	{
+		EmitterHandleViewModel->GetEmitterViewModel()->GetEmitterMessageStores(OutMessageStores);
+	}
+}
+
 FText UNiagaraStackViewModel::FTopLevelViewModel::GetDisplayName() const
 {
 	if (SystemViewModel.IsValid())
@@ -363,11 +375,13 @@ void UNiagaraStackViewModel::UndismissAllIssues()
 	FScopedTransaction ScopedTransaction(LOCTEXT("UnDismissIssues", "Undismiss issues"));
 
 	TArray<UNiagaraStackEditorData*> StackEditorDatas;
+	TArray<FNiagaraMessageSourceAndStore> MessageStores;
 	for (TSharedRef<UNiagaraStackViewModel::FTopLevelViewModel> TopLevelViewModel : TopLevelViewModels)
 	{
 		if (TopLevelViewModel->IsValid())
 		{
 			StackEditorDatas.AddUnique(TopLevelViewModel->GetStackEditorData());
+			TopLevelViewModel->GetMessageStores(MessageStores);
 		}
 	}
 
@@ -376,7 +390,14 @@ void UNiagaraStackViewModel::UndismissAllIssues()
 		StackEditorData->Modify();
 		StackEditorData->UndismissAllIssues();
 	}
-	
+
+	for (const FNiagaraMessageSourceAndStore& MessageStore : MessageStores)
+	{
+		MessageStore.GetSource()->Modify();
+		MessageStore.GetStore()->ClearDismissedMessages();
+	}
+
+	RootEntry->GetSystemViewModel()->RefreshAssetMessagesDeferred();
 	RootEntry->RefreshChildren();
 }
 
@@ -384,9 +405,24 @@ bool UNiagaraStackViewModel::HasDismissedStackIssues()
 {
 	for (TSharedRef<FTopLevelViewModel> TopLevelViewModel : TopLevelViewModels)
 	{
-		if (TopLevelViewModel->IsValid() && TopLevelViewModel->GetStackEditorData()->GetDismissedStackIssueIds().Num() > 0)
+		if (TopLevelViewModel->IsValid())
 		{
-			return true;
+			if (TopLevelViewModel->GetStackEditorData()->GetDismissedStackIssueIds().Num() > 0)
+			{
+				return true;
+			}
+			else
+			{
+				TArray<FNiagaraMessageSourceAndStore> MessageStores;
+				TopLevelViewModel->GetMessageStores(MessageStores);
+				for (const FNiagaraMessageSourceAndStore& MessageStore : MessageStores)
+				{
+					if (MessageStore.GetStore()->HasDismissedMessages())
+					{
+						return true;
+					}
+				}
+			}
 		}
 	}
 	return false;
diff --git a/Engine/Plugins/FX/Niagara/Source/NiagaraEditor/Public/NiagaraEditorUtilities.h b/Engine/Plugins/FX/Niagara/Source/NiagaraEditor/Public/NiagaraEditorUtilities.h
--- a/Engine/Plugins/FX/Niagara/Source/NiagaraEditor/Public/NiagaraEditorUtilities.h
+++ b/Engine/Plugins/FX/Niagara/Source/NiagaraEditor/Public/NiagaraEditorUtilities.h
@@ -47,6 +47,7 @@ class INiagaraParameterDefinitionsSubscriberViewModel;
 struct FNiagaraScriptVersionUpgradeContext;
 struct FNiagaraScriptConversionContext;
 class UUpgradeNiagaraEmitterContext;
+struct FNiagaraMessageStore;
 
 enum class ENiagaraFunctionDebugState : uint8;
 
@@ -406,7 +407,9 @@ namespace FNiagaraEditorUtilities
 	void GetAllowedPayloadTypes(TArray<FNiagaraTypeDefinition>& OutAllowedTypes);
 
 	bool IsEnumIndexVisible(const UEnum* Enum, int32 Index);
-};
+
+	void GetScriptMessageStores(UNiagaraScript* InScript, TArray<FNiagaraMessageSourceAndStore>& OutNiagaraMessageStores);
+}
 
 namespace FNiagaraParameterUtilities
 {
diff --git a/Engine/Plugins/FX/Niagara/Source/NiagaraEditor/Public/NiagaraMessages.h b/Engine/Plugins/FX/Niagara/Source/NiagaraEditor/Public/NiagaraMessages.h
--- a/Engine/Plugins/FX/Niagara/Source/NiagaraEditor/Public/NiagaraMessages.h
+++ b/Engine/Plugins/FX/Niagara/Source/NiagaraEditor/Public/NiagaraMessages.h
@@ -75,8 +75,7 @@ public:
 	/** Can optionally be overriden to give the message a title/short description. */
 	virtual FText GenerateMessageTitle() const;
 
-	/** Can optionally be overriden to allow for dismissal of a message (where applicable (the stack for example) */
-	virtual bool AllowDismissal() const;
+	virtual const FSimpleDelegate& GetDismissHandler() const { return DismissHandler; }
 	
 	virtual TSharedRef<FTokenizedMessage> GenerateTokenizedMessage() const = 0;
 
@@ -94,6 +93,7 @@ public:
 protected:
 	const TArray<FObjectKey> AssociatedObjectKeys;
 	mutable uint32 MessageTopicBitflag;
+	FSimpleDelegate DismissHandler;
 };
 
 class FNiagaraMessageCompileEvent : public INiagaraMessage
@@ -111,8 +111,6 @@ public:
 
 	virtual FText GenerateMessageTitle() const override;
 
-	virtual bool AllowDismissal() const override;
-	
 	virtual TSharedRef<FTokenizedMessage> GenerateTokenizedMessage() const override;
 
 	virtual void GenerateLinks(TArray<FText>& OutLinkDisplayNames, TArray<FSimpleDelegate>& OutLinkNavigationActions) const override;
@@ -133,22 +131,22 @@ private:
 class FNiagaraMessageText : public INiagaraMessage
 {
 public:
-	FNiagaraMessageText(const FText& InMessageText, const FText& InShortDescription, const EMessageSeverity::Type& InMessageSeverity, const FName& InTopicName, bool bInAllowDismissal = false, const TArray<FObjectKey>& InAssociatedObjectKeys = TArray<FObjectKey>())
+	FNiagaraMessageText(const FText& InMessageText, const FText& InShortDescription, const EMessageSeverity::Type& InMessageSeverity, const FName& InTopicName, const FSimpleDelegate& InDismissHandler = FSimpleDelegate(), const TArray<FObjectKey>& InAssociatedObjectKeys = TArray<FObjectKey>())
 		: INiagaraMessage(InAssociatedObjectKeys)
 		, MessageText(InMessageText)
 		, ShortDescription(InShortDescription)
 		, MessageSeverity(InMessageSeverity)
 		, TopicName(InTopicName)
-		, bAllowDismissal(bInAllowDismissal)
+		, DismissHandler(InDismissHandler)
 	{
 	};
 
-	FNiagaraMessageText(const FText& InMessageText, const EMessageSeverity::Type& InMessageSeverity, const FName& InTopicName, bool bInAllowDismissal = false, const TArray<FObjectKey>& InAssociatedObjectKeys = TArray<FObjectKey>())
+	FNiagaraMessageText(const FText& InMessageText, const EMessageSeverity::Type& InMessageSeverity, const FName& InTopicName, const FSimpleDelegate& InDismissHandler = FSimpleDelegate(), const TArray<FObjectKey>& InAssociatedObjectKeys = TArray<FObjectKey>())
 		: INiagaraMessage(InAssociatedObjectKeys)
 		, MessageText(InMessageText)
 		, MessageSeverity(InMessageSeverity)
 		, TopicName(InTopicName)
-		, bAllowDismissal(bInAllowDismissal)
+		, DismissHandler(InDismissHandler)
 	{
 	};
 
@@ -156,7 +154,7 @@ public:
 
 	virtual FText GenerateMessageTitle() const override;
 
-	virtual bool AllowDismissal() const override;
+	virtual const FSimpleDelegate& GetDismissHandler() const { return DismissHandler; }
 	
 	virtual TSharedRef<FTokenizedMessage> GenerateTokenizedMessage() const override;
 
@@ -169,7 +167,7 @@ private:
 	const FText ShortDescription;
 	const EMessageSeverity::Type MessageSeverity;
 	const FName TopicName;
-	const bool bAllowDismissal;
+	const FSimpleDelegate DismissHandler;
 };
 
 class FNiagaraMessageTextWithLinks : public FNiagaraMessageText
@@ -180,10 +178,10 @@ public:
 	, const FText& InShortDescription  
 	, const EMessageSeverity::Type& InMessageSeverity
 	, const FName& InTopicName
-	, const bool bInAllowDismissal
+	, const FSimpleDelegate& InDismissHandler
 	, const TArray<FLinkNameAndDelegate>& InLinks
 	, const TArray<FObjectKey>& InAssociatedObjectKeys = TArray<FObjectKey>())
-		: FNiagaraMessageText(InMessageText, InShortDescription, InMessageSeverity, InTopicName, bInAllowDismissal, InAssociatedObjectKeys)
+		: FNiagaraMessageText(InMessageText, InShortDescription, InMessageSeverity, InTopicName, InDismissHandler, InAssociatedObjectKeys)
 		, Links(InLinks)
 	{
 	};
@@ -273,10 +271,13 @@ public:
 	const TArray<FObjectKey>& GetAssociatedObjectKeys() const { return AssociatedObjectKeys; };
 	void SetLinks(const TArray<FLinkNameAndDelegate>& InLinks) { Links = InLinks; };
 	const TArray<FLinkNameAndDelegate>& GetLinks() const { return Links; };
+	void SetDismissHandler(const FSimpleDelegate& InDismissHandler) { DismissHandler = InDismissHandler; }
+	const FSimpleDelegate& GetDismissHandler() const { return DismissHandler; }
 
 private:
 	TArray<FObjectKey> AssociatedObjectKeys;
 	TArray<FLinkNameAndDelegate> Links;
+	FSimpleDelegate DismissHandler;
 };
 
 UCLASS(abstract)
@@ -300,6 +301,7 @@ public:
 
 	virtual TSharedRef<const INiagaraMessage> GenerateNiagaraMessage(const FGenerateNiagaraMessageInfo& InGenerateInfo = FGenerateNiagaraMessageInfo()) const override;
 
+	virtual bool GetAllowDismissal() const override { return bAllowDismissal; }
 	void SetAllowDismissal(bool bInAllowDismissal) { bAllowDismissal = bInAllowDismissal; }
 	
 private:
diff --git a/Engine/Plugins/FX/Niagara/Source/NiagaraEditor/Public/NiagaraNodeFunctionCall.h b/Engine/Plugins/FX/Niagara/Source/NiagaraEditor/Public/NiagaraNodeFunctionCall.h
--- a/Engine/Plugins/FX/Niagara/Source/NiagaraEditor/Public/NiagaraNodeFunctionCall.h
+++ b/Engine/Plugins/FX/Niagara/Source/NiagaraEditor/Public/NiagaraNodeFunctionCall.h
@@ -6,6 +6,7 @@
 #include "UObject/ObjectMacros.h"
 #include "NiagaraEditorCommon.h"
 #include "NiagaraMessages.h"
+#include "NiagaraMessageStore.h"
 #include "NiagaraNodeWithDynamicPins.h"
 #include "NiagaraNodeInput.h"
 #include "UpgradeNiagaraScriptResults.h"
@@ -172,10 +173,7 @@ public:
 	virtual TSharedPtr<SGraphNode> CreateVisualWidget() override;
 
 	// Messages API
-	NIAGARAEDITOR_API const auto& GetMessages() const { return MessageKeyToMessageMap; };
-	NIAGARAEDITOR_API void AddMessage(const FGuid& MessageKey, UNiagaraMessageData* NewMessage) { MessageKeyToMessageMap.Add(MessageKey, NewMessage); };
-	NIAGARAEDITOR_API void RemoveMessage(const FGuid& MessageKey) { MessageKeyToMessageMap.Remove(MessageKey); };
-	void RemoveMessageDelegateable(const FGuid MessageKey) { MessageKeyToMessageMap.Remove(MessageKey); };
+	FNiagaraMessageStore& GetMessageStore() { return MessageStore; }
 
 	// Custom Notes API
 	NIAGARAEDITOR_API const TArray<FNiagaraStackMessage>& GetCustomNotes() const { return StackMessages; };
@@ -239,7 +237,10 @@ protected:
 	/* Marking those properties explicitly as editoronly_data will make localization not pick these up. */
 #if WITH_EDITORONLY_DATA
 	UPROPERTY(meta = (SkipForCompileHash="true"))
-	TMap<FGuid, TObjectPtr<UNiagaraMessageData>> MessageKeyToMessageMap;
+	TMap<FGuid, TObjectPtr<UNiagaraMessageData>> MessageKeyToMessageMap_DEPRECATED;
+
+	UPROPERTY(meta = (SkipForCompileHash="true"))
+	FNiagaraMessageStore MessageStore;
 	
 	UPROPERTY(meta = (SkipForCompileHash="true"))
 	TArray<FNiagaraStackMessage> StackMessages;
diff --git a/Engine/Plugins/FX/Niagara/Source/NiagaraEditor/Public/ViewModels/NiagaraEmitterViewModel.h b/Engine/Plugins/FX/Niagara/Source/NiagaraEditor/Public/ViewModels/NiagaraEmitterViewModel.h
--- a/Engine/Plugins/FX/Niagara/Source/NiagaraEditor/Public/ViewModels/NiagaraEmitterViewModel.h
+++ b/Engine/Plugins/FX/Niagara/Source/NiagaraEditor/Public/ViewModels/NiagaraEmitterViewModel.h
@@ -25,6 +25,7 @@ struct FNiagaraEventScriptProperties;
 class UNiagaraScriptVariable;
 class UNiagaraParameterDefinitions;
 class UNiagaraSummaryViewViewModel;
+struct FNiagaraMessageSourceAndStore;
 
 
 /** The view model for the UNiagaraEmitter objects */
@@ -122,6 +123,8 @@ public:
 	/** Gets editor specific data which is stored per emitter.  If this data hasn't been created then it will be created. */
 	NIAGARAEDITOR_API UNiagaraEmitterEditorData& GetEditorData();
 
+	void GetEmitterMessageStores(TArray<FNiagaraMessageSourceAndStore>& OutMessageStores);
+
 	/** Add an event script to the owned emitter. Sets the Usage, UsageID and Source of the EventScriptProperties. */
 	NIAGARAEDITOR_API void AddEventHandler(FNiagaraEventScriptProperties& EventScriptProperties, bool bResetGraphForOutput = false);
 
diff --git a/Engine/Plugins/FX/Niagara/Source/NiagaraEditor/Public/ViewModels/NiagaraSystemViewModel.h b/Engine/Plugins/FX/Niagara/Source/NiagaraEditor/Public/ViewModels/NiagaraSystemViewModel.h
--- a/Engine/Plugins/FX/Niagara/Source/NiagaraEditor/Public/ViewModels/NiagaraSystemViewModel.h
+++ b/Engine/Plugins/FX/Niagara/Source/NiagaraEditor/Public/ViewModels/NiagaraSystemViewModel.h
@@ -410,8 +410,7 @@ public:
 	/** Remove a serialized message from the target function call node inside a script this viewmodel is managing. */
 	NIAGARAEDITOR_API void RemoveStackMessage(const FGuid& MessageKey, UNiagaraNodeFunctionCall* TargetFunctionCallNode) const;
 
-	/** Wrapper to set bPendingMessagesChanged after calling a delegate off of a message link. */
-	void ExecuteMessageDelegateAndRefreshMessages(FSimpleDelegate MessageDelegate);
+	void DismissMessageAndRefresh(FNiagaraMessageSourceAndStore MessageSourceAndStore, FGuid MessageKey);
 
 	/** Helper to get the current System or Emitter's EditorParameters. */
 	UNiagaraEditorParametersAdapter* GetEditorOnlyParametersAdapter() const;
@@ -453,6 +452,10 @@ public:
 	/** Utility method to gather graph parameter references */
 	TArray<FNiagaraGraphParameterReference> GetGraphParameterReferences(const FNiagaraVariable& Parameter, ENiagaraGetGraphParameterReferencesMode Mode);
 
+	void GetSystemMessageStores(TArray<FNiagaraMessageSourceAndStore>& OutMessageStores);
+
+	void RefreshAssetMessagesDeferred();
+
 private:
 	/** Sends message jobs to FNiagaraMessageManager for all compile events from the last compile. */
 	void SendLastCompileMessageJobs() const;
@@ -749,8 +752,8 @@ private:
 
 	FDelegateHandle SystemChangedDelegateHandle;
 
-	/** ObjectKeys for function call nodes that supply messages. Used to invalidate the messages of these nodes on refresh.*/
-	TArray<FObjectKey> LastFunctionCallNodeObjectKeys;
+	/** ObjectKeys for object viewed by this view model that supply messages. Used to invalidate the messages of these objects on refresh.*/
+	TArray<FObjectKey> MessageSourceObjectKeys;
 
 	/** Flag for when messages have been added/removed through the viewmodel to signal that the message manager needs to be refreshed. */
 	mutable bool bPendingAssetMessagesChanged;
diff --git a/Engine/Plugins/FX/Niagara/Source/NiagaraEditor/Public/ViewModels/Stack/NiagaraStackEntry.h b/Engine/Plugins/FX/Niagara/Source/NiagaraEditor/Public/ViewModels/Stack/NiagaraStackEntry.h
--- a/Engine/Plugins/FX/Niagara/Source/NiagaraEditor/Public/ViewModels/Stack/NiagaraStackEntry.h
+++ b/Engine/Plugins/FX/Niagara/Source/NiagaraEditor/Public/ViewModels/Stack/NiagaraStackEntry.h
@@ -210,7 +210,7 @@ public:
 	{
 		FStackIssue();
 
-		FStackIssue(EStackIssueSeverity InSeverity, FText InShortDescription, FText InLongDescription, FString InStackEditorDataKey, bool bInCanBeDismissed, const TArray<FStackIssueFix>& InFixes);
+		FStackIssue(EStackIssueSeverity InSeverity, FText InShortDescription, FText InLongDescription, FString InStackEditorDataKey, bool bInCanBeDismissed, const TArray<FStackIssueFix>& InFixes, const FSimpleDelegate& InDismissHandler = FSimpleDelegate());
 
 		FStackIssue(EStackIssueSeverity InSeverity, FText InShortDescription, FText InLongDescription, FString InStackEditorDataKey, bool bInCanBeDismissed, FStackIssueFix InFix);
 
@@ -228,6 +228,8 @@ public:
 
 		bool GetCanBeDismissed() const;
 
+		const FSimpleDelegate& GetDismissHandler() const;
+
 		const TArray<FStackIssueFix>& GetFixes() const;
 
 		bool GetIsExpandedByDefault() const;
@@ -244,6 +246,7 @@ public:
 		bool bCanBeDismissed = false;
 		bool bIsExpandedByDefault = true;
 		TArray<FStackIssueFix> Fixes;
+		FSimpleDelegate DismissHandler;
 	};
 
 public:
diff --git a/Engine/Plugins/FX/Niagara/Source/NiagaraEditor/Public/ViewModels/Stack/NiagaraStackViewModel.h b/Engine/Plugins/FX/Niagara/Source/NiagaraEditor/Public/ViewModels/Stack/NiagaraStackViewModel.h
--- a/Engine/Plugins/FX/Niagara/Source/NiagaraEditor/Public/ViewModels/Stack/NiagaraStackViewModel.h
+++ b/Engine/Plugins/FX/Niagara/Source/NiagaraEditor/Public/ViewModels/Stack/NiagaraStackViewModel.h
@@ -4,6 +4,7 @@
 
 #include "Templates/SharedPointer.h"
 #include "ViewModels/Stack/NiagaraStackEntry.h"
+#include "NiagaraMessageStore.h"
 #include "UObject/ObjectKey.h"
 #include "NiagaraStackViewModel.generated.h"
 
@@ -68,6 +69,8 @@ public:
 
 		UNiagaraStackEditorData* GetStackEditorData() const;
 
+		void GetMessageStores(TArray<FNiagaraMessageSourceAndStore>& OutMessageStores);
+
 		FText GetDisplayName() const;
 
 		bool operator==(const FTopLevelViewModel& Other)const;
diff --git a/Engine/Plugins/FX/Niagara/Source/NiagaraEditorWidgets/Private/SNiagaraStack.cpp b/Engine/Plugins/FX/Niagara/Source/NiagaraEditorWidgets/Private/SNiagaraStack.cpp
--- a/Engine/Plugins/FX/Niagara/Source/NiagaraEditorWidgets/Private/SNiagaraStack.cpp
+++ b/Engine/Plugins/FX/Niagara/Source/NiagaraEditorWidgets/Private/SNiagaraStack.cpp
@@ -62,6 +62,94 @@
 
 #define LOCTEXT_NAMESPACE "NiagaraStack"
 
+class SNiagaraStackHeaderButtons : public SCompoundWidget
+{
+public:
+	SLATE_BEGIN_ARGS(SNiagaraStackHeaderButtons) {}
+		SLATE_ATTRIBUTE(EVisibility, IssueIconVisibility);
+		SLATE_EVENT(FSimpleDelegate, OnCycleThroughIssues);
+		SLATE_ARGUMENT(TSharedPtr<SNiagaraStack>, ParentStack)
+	SLATE_END_ARGS();
+
+	void Construct(const FArguments& InArgs, UNiagaraStackEntry* InRootEntry, UNiagaraStackViewModel* InStackViewModel)
+	{
+		StackViewModel = InStackViewModel;
+		OnCycleThroughIssues = InArgs._OnCycleThroughIssues;
+		ParentStackPtr = InArgs._ParentStack;
+		TopLevelViewModel = StackViewModel->GetTopLevelViewModelForEntry(*InRootEntry);
+		ChildSlot
+		[
+			SNew(SHorizontalBox)
+			// Issue Icon
+			+ SHorizontalBox::Slot()
+			.AutoWidth()
+			.HAlign(HAlign_Right)
+			.VAlign(VAlign_Center)
+			.Padding(0, 0, 2, 0)
+			[
+				SNew(SNiagaraStackIssueIcon, StackViewModel, InRootEntry)
+				.Visibility(InArgs._IssueIconVisibility)
+				.OnClicked(this, &SNiagaraStackHeaderButtons::OnIssueIconClicked)
+			]
+			// Options button
+			+ SHorizontalBox::Slot()
+			.AutoWidth()
+			.HAlign(HAlign_Fill)
+			.VAlign(VAlign_Center)
+			.Padding(0, 0, 2, 0)
+			[
+				SAssignNew(SettingsAnchor, SMenuAnchor)
+				.Placement(MenuPlacement_MenuLeft)
+				.OnGetMenuContent(this, &SNiagaraStackHeaderButtons::OnGetContent)
+				[
+					SNew(SButton)
+					.ContentPadding(0)
+					.ForegroundColor(FSlateColor::UseForeground())
+					.ButtonStyle(FAppStyle::Get(), "SimpleButton")
+					.OnClicked(this, &SNiagaraStackHeaderButtons::OpenSubmenu)
+					.VAlign(VAlign_Center)
+					.HAlign(HAlign_Fill)
+					[
+						SNew(SImage)
+						.ColorAndOpacity(FSlateColor::UseForeground())
+						.Image(FAppStyle::Get().GetBrush("DetailsView.ViewOptions"))
+					]
+				]
+			]
+		];
+	}
+
+private:
+	FReply OnIssueIconClicked() const
+	{
+		StackViewModel->OnCycleThroughIssues(TopLevelViewModel);
+		OnCycleThroughIssues.ExecuteIfBound();
+		return FReply::Handled();
+	}
+
+	FReply OpenSubmenu()
+	{
+		SettingsAnchor->SetIsOpen(!SettingsAnchor->IsOpen());
+		return FReply::Handled();
+	}
+
+	TSharedRef<SWidget> OnGetContent() const
+	{
+		if (ParentStackPtr.IsValid())
+		{
+			return ParentStackPtr.Pin()->GenerateStackMenu(TopLevelViewModel).ToSharedRef();
+		}
+		return SNullWidget::NullWidget;
+	}
+
+private:
+	TSharedPtr<UNiagaraStackViewModel::FTopLevelViewModel> TopLevelViewModel;
+	FSimpleDelegate OnCycleThroughIssues;
+	UNiagaraStackViewModel* StackViewModel;
+	TWeakPtr<SNiagaraStack> ParentStackPtr;
+	TSharedPtr<SMenuAnchor> SettingsAnchor;
+};
+
 class SNiagaraStackEmitterHeader : public SCompoundWidget
 {
 public:
@@ -74,10 +162,6 @@
 	void Construct(const FArguments& InArgs, TSharedRef<FNiagaraEmitterHandleViewModel> InEmitterHandleViewModel, UNiagaraStackEntry* InRootEntry, UNiagaraStackViewModel* InStackViewModel)
 	{
 		EmitterHandleViewModel = InEmitterHandleViewModel;
-		OnCycleThroughIssues = InArgs._OnCycleThroughIssues;
-		StackViewModel = InStackViewModel;
-		TopLevelViewModel = StackViewModel->GetTopLevelViewModelForEntry(*InRootEntry);
-		ParentStackPtr = InArgs._ParentStack;
 		ChildSlot
 		[
 			SNew(SBorder)
@@ -136,40 +220,13 @@
 							.Visibility(this, &SNiagaraStackEmitterHeader::GetSourceEmitterNameVisibility) 
 						]
 					]
-					// Issue Icon
+
 					+ SHorizontalBox::Slot()
-					.AutoWidth()
-					.HAlign(HAlign_Right)
-					.VAlign(VAlign_Center)
-					.Padding(0, 0, 2, 0)
 					[
-						SNew(SNiagaraStackIssueIcon, StackViewModel, InRootEntry)
-						.Visibility(InArgs._IssueIconVisibility)
-						.OnClicked(this, &SNiagaraStackEmitterHeader::OnIssueIconClicked)
+						SNew(SNiagaraStackHeaderButtons, InRootEntry, InStackViewModel)
+						.IssueIconVisibility(InArgs._IssueIconVisibility)
+						.OnCycleThroughIssues(InArgs._OnCycleThroughIssues)
+						.ParentStack(InArgs._ParentStack)
 					]
-					+ SHorizontalBox::Slot()
-					.AutoWidth()
-					.HAlign(HAlign_Fill)
-					.VAlign(VAlign_Center)
-					.Padding(0, 0, 2, 0)
-					[
-						SAssignNew(SettingsAnchor, SMenuAnchor)
-						.Placement(MenuPlacement_MenuLeft)
-						.OnGetMenuContent(this, &SNiagaraStackEmitterHeader::OnGetContent)
-						[
-							SNew(SButton)
-							.ContentPadding(0)
-							.ForegroundColor(FSlateColor::UseForeground())
-							.ButtonStyle(FAppStyle::Get(), "SimpleButton")
-							.OnClicked(this, &SNiagaraStackEmitterHeader::OpenSubmenu)
-							.VAlign(VAlign_Center)
-							.HAlign(HAlign_Fill)
-							[
-								SNew(SImage)
-								.ColorAndOpacity(FSlateColor::UseForeground())
-								.Image(FAppStyle::Get().GetBrush("DetailsView.ViewOptions"))
-							]
-						]
-					]
 				]
 
@@ -228,37 +285,9 @@
 		return EmitterHandleViewModel->GetEmitterViewModel()->HasParentEmitter() ? EVisibility::Visible : EVisibility::Collapsed;
 	}
 
-	FReply OnIssueIconClicked() const
-	{
-		StackViewModel->OnCycleThroughIssues(TopLevelViewModel);
-		OnCycleThroughIssues.ExecuteIfBound();
-		return FReply::Handled();
-	}
-
-	FReply OpenSubmenu()
-	{
-		SettingsAnchor->SetIsOpen(!SettingsAnchor->IsOpen());
-		return FReply::Handled();
-	}
-
-	TSharedRef<SWidget> OnGetContent() const
-	{
-		if (ParentStackPtr.IsValid())
-		{
-			return ParentStackPtr.Pin()->GenerateStackMenu(TopLevelViewModel).ToSharedRef();
-		}
-		return SNullWidget::NullWidget;
-	}
-
 private:
 	TSharedPtr<FNiagaraEmitterHandleViewModel> EmitterHandleViewModel;
-	TSharedPtr<UNiagaraStackViewModel::FTopLevelViewModel> TopLevelViewModel;
-
 	TSharedPtr<SInlineEditableTextBlock> EmitterNameTextBlock;
-	FSimpleDelegate OnCycleThroughIssues;
-	UNiagaraStackViewModel* StackViewModel;
-	TWeakPtr<SNiagaraStack> ParentStackPtr;
-	TSharedPtr<SMenuAnchor> SettingsAnchor;
 };
 
 const float SpacerHeight = 6;
@@ -664,15 +693,15 @@
 					.TextStyle(FAppStyle::Get(), "DetailsView.ConstantTextBlockStyle")
 					.Text(Item->SystemViewModel.ToSharedRef(), &FNiagaraSystemViewModel::GetDisplayName)
 				]
-				// Issue Icon
 				+ SHorizontalBox::Slot()
 				.AutoWidth()
 				.VAlign(VAlign_Center)
-				.Padding(4, 0, 2, 0)
+				.Padding(4, 0, 0, 0)
 				[
-					SNew(SNiagaraStackIssueIcon, StackViewModel, Item->RootEntry.Get())
-					.Visibility(this, &SNiagaraStack::GetIssueIconVisibility)
-					.OnClicked(this, &SNiagaraStack::OnCycleThroughSystemIssues, Item->SystemViewModel)
+					SNew(SNiagaraStackHeaderButtons, Item->RootEntry.Get(), StackViewModel)
+					.IssueIconVisibility(this, &SNiagaraStack::GetIssueIconVisibility)
+					.OnCycleThroughIssues(this, &SNiagaraStack::OnCycleThroughIssues)
+					.ParentStack(SharedThis(this))
 				]
 			];
 	}
diff --git a/Engine/Plugins/FX/Niagara/Source/NiagaraShader/Public/NiagaraShared.h b/Engine/Plugins/FX/Niagara/Source/NiagaraShader/Public/NiagaraShared.h
--- a/Engine/Plugins/FX/Niagara/Source/NiagaraShader/Public/NiagaraShared.h
+++ b/Engine/Plugins/FX/Niagara/Source/NiagaraShader/Public/NiagaraShared.h
@@ -62,7 +62,6 @@ public:
 		Severity = FNiagaraCompileEventSeverity::Display;
 		Message = FString();
 		ShortDescription = FString();
-		bDismissable = true;
 		NodeGuid = FGuid();
 		PinGuid = FGuid();
 		StackGuids.Empty();
@@ -73,7 +72,6 @@ public:
 		FNiagaraCompileEventSeverity InSeverity,
 		const FString& InMessage,
 		FString InShortDescription = FString(),
-		bool bInDismissable = true,
 		FGuid InNodeGuid = FGuid(),
 		FGuid InPinGuid = FGuid(),
 		const TArray<FGuid>& InCallstackGuids = TArray<FGuid>(),
@@ -82,7 +80,6 @@ public:
 		: Severity(InSeverity)
 		, Message(InMessage)
 		, ShortDescription(InShortDescription)
-		, bDismissable(bInDismissable)
 		, NodeGuid(InNodeGuid)
 		, PinGuid(InPinGuid)
 		, StackGuids(InCallstackGuids)
@@ -98,9 +95,6 @@ public:
 	/** A short, optional description of the event. */
 	UPROPERTY()
 	FString ShortDescription;
-	/** Whether or not this event is dismissable or not */
-	UPROPERTY()
-	bool bDismissable;
 	/** The node guid that generated the compile event*/
 	UPROPERTY()
 	FGuid NodeGuid;
