Added virtual destructor to base class to ensure correct destruction.

#tested on Win64 address sanitizer
#rnx

diff --git a/Engine/Plugins/FX/Niagara/Source/Niagara/Public/NiagaraDataChannelHandler.h b/Engine/Plugins/FX/Niagara/Source/Niagara/Public/NiagaraDataChannelHandler.h
--- a/Engine/Plugins/FX/Niagara/Source/Niagara/Public/NiagaraDataChannelHandler.h
+++ b/Engine/Plugins/FX/Niagara/Source/Niagara/Public/NiagaraDataChannelHandler.h
@@ -18,6 +18,7 @@ Some more complex handlers may want to divide up the scene in various different
 /** Base class for all UNiagaraDataChannelHandler's Render Thread Proxies. */
 struct FNiagaraDataChannelHandlerRTProxyBase
 {
+	virtual ~FNiagaraDataChannelHandlerRTProxyBase() = default;
 };
 
 UCLASS(Experimental, abstract, BlueprintType)
