Fixed dropping items below an item reparenting instead of duplicating if the dragged item is from source.
This would remove items from the sources list by accident

diff --git a/Engine/Plugins/FX/Niagara/Source/NiagaraEditor/Private/ViewModels/HierarchyEditor/NiagaraHierarchyViewModelBase.cpp b/Engine/Plugins/FX/Niagara/Source/NiagaraEditor/Private/ViewModels/HierarchyEditor/NiagaraHierarchyViewModelBase.cpp
--- a/Engine/Plugins/FX/Niagara/Source/NiagaraEditor/Private/ViewModels/HierarchyEditor/NiagaraHierarchyViewModelBase.cpp
+++ b/Engine/Plugins/FX/Niagara/Source/NiagaraEditor/Private/ViewModels/HierarchyEditor/NiagaraHierarchyViewModelBase.cpp
@@ -1344,7 +1344,15 @@ void FNiagaraHierarchyItemViewModel::OnDroppedOnInternal(TSharedPtr<FNiagaraHier
 	else if(ItemDropZone == EItemDropZone::BelowItem)
 	{
 		int32 IndexOfThis = Parent.Pin()->FindIndexOfDataChild(AsShared());
-		Parent.Pin()->ReparentToThis(DroppedItem, FMath::Min(IndexOfThis+1, Parent.Pin()->GetChildren().Num()));
+
+		if(DroppedItem->IsForHierarchy() == false)
+		{
+			Parent.Pin()->DuplicateToThis(DroppedItem, FMath::Min(IndexOfThis+1, Parent.Pin()->GetChildren().Num()));
+		}
+		else
+		{
+			Parent.Pin()->ReparentToThis(DroppedItem, FMath::Min(IndexOfThis+1, Parent.Pin()->GetChildren().Num()));
+		}
 
 		bDropSucceeded = true;
 	}
