- Ignore bWidgetEnabled for -game otherwise the debug hud won't be visible

diff --git a/Engine/Plugins/FX/Niagara/Source/Niagara/Public/NiagaraDebuggerCommon.h b/Engine/Plugins/FX/Niagara/Source/Niagara/Public/NiagaraDebuggerCommon.h
--- a/Engine/Plugins/FX/Niagara/Source/Niagara/Public/NiagaraDebuggerCommon.h
+++ b/Engine/Plugins/FX/Niagara/Source/Niagara/Public/NiagaraDebuggerCommon.h
@@ -417,7 +417,7 @@ struct NIAGARA_API FNiagaraDebugHUDSettingsData
 	bool IsEnabled() const
 	{
 	#if WITH_EDITORONLY_DATA
-		return bWidgetEnabled && bHudEnabled;
+		return bHudEnabled && (!GIsEditor || bWidgetEnabled);
 	#else
 		return bHudEnabled;
 	#endif
