UX Pass on Sim Cache Attribute Spreadsheet

diff --git a/Engine/Plugins/FX/Niagara/Source/NiagaraEditor/Private/Toolkits/SystemToolkitModes/NiagaraSystemToolkitModeBase.cpp b/Engine/Plugins/FX/Niagara/Source/NiagaraEditor/Private/Toolkits/SystemToolkitModes/NiagaraSystemToolkitModeBase.cpp
--- a/Engine/Plugins/FX/Niagara/Source/NiagaraEditor/Private/Toolkits/SystemToolkitModes/NiagaraSystemToolkitModeBase.cpp
+++ b/Engine/Plugins/FX/Niagara/Source/NiagaraEditor/Private/Toolkits/SystemToolkitModes/NiagaraSystemToolkitModeBase.cpp
@@ -60,7 +60,6 @@ const FName FNiagaraSystemToolkitModeBase::SystemParameterDefinitionsTabID(TEXT(
 const FName FNiagaraSystemToolkitModeBase::DetailsTabID(TEXT("NiagaraSystemEditor_Details"));
 const FName FNiagaraSystemToolkitModeBase::SelectedEmitterGraphTabID(TEXT("NiagaraSystemEditor_SelectedEmitterGraph"));
 const FName FNiagaraSystemToolkitModeBase::DebugSpreadsheetTabID(TEXT("NiagaraSystemEditor_DebugAttributeSpreadsheet"));
-const FName FNiagaraSystemToolkitModeBase::DebugCaptureTabID(TEXT("NiagaraSystemEditor_DebugCacheCapture"));
 const FName FNiagaraSystemToolkitModeBase::DebugCacheSpreadsheetTabID(TEXT("NiagaraSystemEditor_DebugCacheSpreadsheet"));
 const FName FNiagaraSystemToolkitModeBase::PreviewSettingsTabId(TEXT("NiagaraSystemEditor_PreviewSettings"));
 const FName FNiagaraSystemToolkitModeBase::GeneratedCodeTabID(TEXT("NiagaraSystemEditor_GeneratedCode"));
@@ -284,12 +283,7 @@ void FNiagaraSystemToolkitModeBase::RegisterTabFactories(TSharedPtr<FTabManager>
 		.SetAutoGenerateMenuEntry(GbShowNiagaraDeveloperWindows != 0);
 
 	InTabManager->RegisterTabSpawner(DebugSpreadsheetTabID, FOnSpawnTab::CreateSP(this, &FNiagaraSystemToolkitModeBase::SpawnTab_DebugSpreadsheet))
-		.SetDisplayName(LOCTEXT("DebugSpreadsheet", "Attribute Spreadsheet"))
-		.SetGroup(WorkspaceMenuCategory.ToSharedRef())
-		.SetIcon(FSlateIcon(FNiagaraEditorStyle::Get().GetStyleSetName(), "Tab.Spreadsheet"));
-
-	InTabManager->RegisterTabSpawner(DebugCaptureTabID, FOnSpawnTab::CreateSP(this, &FNiagaraSystemToolkitModeBase::SpawnTab_DebugCapture))
-		.SetDisplayName(LOCTEXT("DebugCapture", "Attribute Capture"))
+		.SetDisplayName(LOCTEXT("DebugSpreadsheet", "Attribute Spreadsheet (Legacy)"))
 		.SetGroup(WorkspaceMenuCategory.ToSharedRef())
 		.SetIcon(FSlateIcon(FNiagaraEditorStyle::Get().GetStyleSetName(), "Tab.Spreadsheet"));
 
@@ -696,25 +690,6 @@ TSharedRef<SDockTab> FNiagaraSystemToolkitModeBase::SpawnTab_DebugSpreadsheet(co
 	return SpawnedTab;
 }
 
-TSharedRef<SDockTab> FNiagaraSystemToolkitModeBase::SpawnTab_DebugCapture(const FSpawnTabArgs& Args)
-{
-	check(Args.GetTabId().TabType == DebugCaptureTabID);
-	TSharedPtr<SNiagaraDebugCaptureView> DebugCaptureTab;
-
-	TSharedRef<SDockTab> SpawnedTab =
-		SNew(SDockTab)
-		[
-			SAssignNew(DebugCaptureTab, SNiagaraDebugCaptureView, SystemToolkit.Pin()->GetSystemViewModel().ToSharedRef(), SystemToolkit.Pin()->GetSimCacheViewModel().ToSharedRef())
-		];
-
-	DebugCaptureTab->OnRequestSpreadsheetTab().BindLambda([this]()
-	{
-		SystemToolkit.Pin()->GetTabManager()->TryInvokeTab(DebugCacheSpreadsheetTabID);
-	});
-	
-	return SpawnedTab;
-}
-
 TSharedRef<SDockTab> FNiagaraSystemToolkitModeBase::SpawnTab_DebugCacheSpreadsheet(const FSpawnTabArgs& Args)
 {
 	check(Args.GetTabId().TabType == DebugCacheSpreadsheetTabID);
@@ -723,6 +698,9 @@ TSharedRef<SDockTab> FNiagaraSystemToolkitModeBase::SpawnTab_DebugCacheSpreadshe
 		SNew(SNiagaraSimCacheView)
 		.SimCacheViewModel(SystemToolkit.Pin()->GetSimCacheViewModel());
 
+	TSharedRef<SNiagaraDebugCaptureView> DebugCaptureView =
+		SNew(SNiagaraDebugCaptureView, SystemToolkit.Pin()->GetSystemViewModel().ToSharedRef(), SystemToolkit.Pin()->GetSimCacheViewModel().ToSharedRef());
+
 	const TSharedRef<SVerticalBox> Contents = SNew(SVerticalBox);
 
 	Contents->AddSlot()
@@ -731,28 +709,19 @@ TSharedRef<SDockTab> FNiagaraSystemToolkitModeBase::SpawnTab_DebugCacheSpreadshe
 			SNew(SVerticalBox)
 			+SVerticalBox::Slot()
 			[
-				SimCacheSpreadsheetView
-			]
-			+SVerticalBox::Slot()
-			.AutoHeight()
-			[
-				SNew(SHorizontalBox)
-				+SHorizontalBox::Slot()
-				.AutoWidth()
-				.Padding(2.0f)
+				SNew(SSplitter)
+				.Orientation(Orient_Vertical)
+				+SSplitter::Slot()
+				.SizeRule(SSplitter::SizeToContent)
 				[
-					SNew(SNiagaraSimCacheViewTransportControls)
-					.WeakViewModel(SystemToolkit.Pin()->GetSimCacheViewModel())
+					DebugCaptureView
 				]
-				+SHorizontalBox::Slot()
-				.Padding(4.0f)
+				+SSplitter::Slot()
+				.Value(0.8f)
 				[
-					SNew(SNiagaraSimCacheViewTimeline)
-					.WeakViewModel(SystemToolkit.Pin()->GetSimCacheViewModel())
+					SimCacheSpreadsheetView
 				]
-				
 			]
-			
 		];
 
 	TSharedRef<SDockTab> SpawnedTab =
diff --git a/Engine/Plugins/FX/Niagara/Source/NiagaraEditor/Private/Toolkits/SystemToolkitModes/NiagaraSystemToolkitModeBase.h b/Engine/Plugins/FX/Niagara/Source/NiagaraEditor/Private/Toolkits/SystemToolkitModes/NiagaraSystemToolkitModeBase.h
--- a/Engine/Plugins/FX/Niagara/Source/NiagaraEditor/Private/Toolkits/SystemToolkitModes/NiagaraSystemToolkitModeBase.h
+++ b/Engine/Plugins/FX/Niagara/Source/NiagaraEditor/Private/Toolkits/SystemToolkitModes/NiagaraSystemToolkitModeBase.h
@@ -35,7 +35,6 @@ private:
 	TSharedRef<SDockTab> SpawnTab_Details(const FSpawnTabArgs& Args);
 	TSharedRef<SDockTab> SpawnTab_SelectedEmitterGraph(const FSpawnTabArgs& Args);
 	TSharedRef<SDockTab> SpawnTab_DebugSpreadsheet(const FSpawnTabArgs& Args);
-	TSharedRef<SDockTab> SpawnTab_DebugCapture(const FSpawnTabArgs& Args);
 	TSharedRef<SDockTab> SpawnTab_DebugCacheSpreadsheet(const FSpawnTabArgs& Args);
 	TSharedRef<SDockTab> SpawnTab_PreviewSettings(const FSpawnTabArgs& Args);
 	TSharedRef<SDockTab> SpawnTab_GeneratedCode(const FSpawnTabArgs& Args);
@@ -85,7 +84,6 @@ public:
 	static const FName DetailsTabID;
 	static const FName SelectedEmitterGraphTabID;
 	static const FName DebugSpreadsheetTabID;
-	static const FName DebugCaptureTabID;
 	static const FName DebugCacheSpreadsheetTabID;
 	static const FName PreviewSettingsTabId;
 	static const FName GeneratedCodeTabID;
diff --git a/Engine/Plugins/FX/Niagara/Source/NiagaraEditor/Private/ViewModels/NiagaraSimCacheViewModel.cpp b/Engine/Plugins/FX/Niagara/Source/NiagaraEditor/Private/ViewModels/NiagaraSimCacheViewModel.cpp
--- a/Engine/Plugins/FX/Niagara/Source/NiagaraEditor/Private/ViewModels/NiagaraSimCacheViewModel.cpp
+++ b/Engine/Plugins/FX/Niagara/Source/NiagaraEditor/Private/ViewModels/NiagaraSimCacheViewModel.cpp
@@ -140,14 +140,6 @@ FText FNiagaraSimCacheViewModel::GetComponentText(const FName ComponentName, con
 	return LOCTEXT("Error", "Error");
 }
 
-void FNiagaraSimCacheViewModel::SetComponentFilterActive(bool bNewActive)
-{
-	if(bNewActive != bComponentFilterActive)
-	{
-		bComponentFilterActive = bNewActive;
-		OnViewDataChangedDelegate.Broadcast(true);
-	}
-}
 
 int32 FNiagaraSimCacheViewModel::GetNumFrames() const
 {
diff --git a/Engine/Plugins/FX/Niagara/Source/NiagaraEditor/Private/ViewModels/NiagaraSimCacheViewModel.h b/Engine/Plugins/FX/Niagara/Source/NiagaraEditor/Private/ViewModels/NiagaraSimCacheViewModel.h
--- a/Engine/Plugins/FX/Niagara/Source/NiagaraEditor/Private/ViewModels/NiagaraSimCacheViewModel.h
+++ b/Engine/Plugins/FX/Niagara/Source/NiagaraEditor/Private/ViewModels/NiagaraSimCacheViewModel.h
@@ -49,14 +49,13 @@ public:
 
 	void SetComponentFilters(const TArray<FString>& NewComponentFilterArray)
 	{
+		bComponentFilterActive = true;
 		ComponentFilterArray.Empty();
 		ComponentFilterArray.Append(NewComponentFilterArray);
 		OnViewDataChangedDelegate.Broadcast(true);
 	}
 
 	bool IsComponentFilterActive() const { return bComponentFilterActive; }
-
-	void SetComponentFilterActive(bool bNewActive);
 	
 	int32 GetNumInstances() const { return NumInstances; }
 
diff --git a/Engine/Plugins/FX/Niagara/Source/NiagaraEditor/Private/Widgets/SNiagaraDebugCaptureView.cpp b/Engine/Plugins/FX/Niagara/Source/NiagaraEditor/Private/Widgets/SNiagaraDebugCaptureView.cpp
--- a/Engine/Plugins/FX/Niagara/Source/NiagaraEditor/Private/Widgets/SNiagaraDebugCaptureView.cpp
+++ b/Engine/Plugins/FX/Niagara/Source/NiagaraEditor/Private/Widgets/SNiagaraDebugCaptureView.cpp
@@ -2,11 +2,15 @@
 
 #include "SNiagaraDebugCaptureView.h"
 
+#include "NiagaraComponent.h"
+#include "NiagaraEditorUtilities.h"
+#include "NiagaraSettings.h"
 #include "PropertyEditorModule.h"
+#include "ToolBuilderUtil.h"
+#include "UObject/UObjectIterator.h"
 #include "ViewModels/NiagaraEmitterHandleViewModel.h"
 #include "ViewModels/NiagaraSystemSelectionViewModel.h"
 #include "Widgets/Input/SSpinBox.h"
-#include "Widgets/Input/SButton.h"
 #include "Widgets/Layout/SUniformGridPanel.h"
 
 #define LOCTEXT_NAMESPACE "NiagaraDebugCaptureView"
@@ -16,151 +20,341 @@ void SNiagaraDebugCaptureView::OnNumFramesChanged(int32 InNumFrames)
 	NumFrames = FMath::Max(1, InNumFrames);
 }
 
+void SNiagaraDebugCaptureView::CreateComponentSelectionMenuContent(FMenuBuilder& MenuBuilder) 
+{
+	for(TObjectIterator<UNiagaraComponent> NiagaraComponent; NiagaraComponent; ++NiagaraComponent)
+	{
+		// Ignore dying or CDO versions of data..
+		// No need to check the unreachable flag here as TObjectIterator already does that
+		if(!IsValid(*NiagaraComponent) || NiagaraComponent->HasAnyFlags(RF_ClassDefaultObject))
+		{
+			continue;
+		}
+
+		// Ignore any component not referencing this system.
+		if (NiagaraComponent->GetAsset() != &SystemViewModel->GetSystem())
+		{
+			continue;
+		}
+
+		if(!NiagaraComponent->GetWorld())
+		{
+			continue;
+		}
+
+		FText ComponentName;
+		FText ComponentTooltip;
+		GetComponentNameAndTooltip(*NiagaraComponent, ComponentName, ComponentTooltip);
+
+		MenuBuilder.AddMenuEntry(
+			ComponentName,
+			ComponentTooltip,
+			FSlateIcon(),
+			FUIAction(FExecuteAction::CreateLambda([&, NiagaraComponent]()
+			{
+				TargetComponent = *NiagaraComponent;
+			})));
+	}
+}
+
+void SNiagaraDebugCaptureView::GetComponentNameAndTooltip(const UNiagaraComponent* InComponent, FText& OutName, FText& OutTooltip) const
+{
+	const UNiagaraComponent* PreviewComponent = SystemViewModel->GetPreviewComponent();
+
+    if (InComponent == nullptr)
+    {
+    	OutName = LOCTEXT("NullComponentLabel", "Unknown");
+    	OutTooltip = LOCTEXT("NullComponentTooltip", "Unknown");
+    }
+    else if (PreviewComponent == InComponent)
+    {
+    	OutName = LOCTEXT("PreviewComponentLabel", "Editor Viewport");
+    	OutTooltip = LOCTEXT("PreviewComponentTooltip", "The instance of the Niagara Component in the Niagara editor viewport.");
+    }
+    else
+    {
+	    const UWorld* World = InComponent->GetWorld();
+    	const AActor* Actor = InComponent->GetOwner();
+    	OutName = FText::Format(LOCTEXT("SourceComponentLabel","World: \"{0}\" Actor: \"{1}\""), World ? FText::FromString(World->GetName()) : FText::GetEmpty(), Actor ? FText::FromString(Actor->GetActorNameOrLabel()) : FText::GetEmpty());
+    	OutTooltip = OutName;
+    }
+}
+
+TSharedRef<SWidget> SNiagaraDebugCaptureView::GenerateCaptureMenuContent()
+{
+	FMenuBuilder MenuBuilder(true, nullptr);
+
+	MenuBuilder.BeginSection("CaptureFrameMode", LOCTEXT("CaptureModeMenu", "Capture Mode"));
+	{
+		MenuBuilder.AddMenuEntry(
+			LOCTEXT("SingleFrameCaptureMode", "Single Frame"),
+			LOCTEXT("SingleFrameCaptureModeTooltip", "Capture a single frame and display it in the current window."),
+			FSlateIcon(),
+			FUIAction(
+				FExecuteAction::CreateLambda([this](){FrameMode = ENiagaraDebugCaptureFrameMode::SingleFrame;}),
+				FCanExecuteAction(),
+				FIsActionChecked::CreateLambda([&]()
+				{
+					return FrameMode == ENiagaraDebugCaptureFrameMode::SingleFrame;
+				})
+			),
+			NAME_None,
+			EUserInterfaceActionType::Check
+		);
+		MenuBuilder.AddMenuEntry(
+			LOCTEXT("MultiFrameCaptureMode", "Multi Frame (Standalone Editor)"),
+			LOCTEXT("MultiFrameCaptureModeTooltip", "Capture the specified number of frames and open the results in the standalone Sim Cache editor."),
+			FSlateIcon(),
+			FUIAction(
+				FExecuteAction::CreateLambda([this](){FrameMode = ENiagaraDebugCaptureFrameMode::MultiFrame;}),
+				FCanExecuteAction(),
+				FIsActionChecked::CreateLambda([&]()
+				{
+					return FrameMode == ENiagaraDebugCaptureFrameMode::MultiFrame;
+				})
+			),
+			NAME_None,
+			EUserInterfaceActionType::Check
+		);
+		
+	}
+	MenuBuilder.EndSection();
+	MenuBuilder.BeginSection("TargetSelection", LOCTEXT("TargetSelectionMenu", "Component Selection"));
+	{
+		MenuBuilder.AddSubMenu(
+			LOCTEXT("TargetSelection", "Target Component"),
+			LOCTEXT("TargetSelectionTooltip", "Select the Niagara Component of this system for capturing."),
+			FNewMenuDelegate::CreateRaw(this, &SNiagaraDebugCaptureView::CreateComponentSelectionMenuContent)
+		);
+	}
+	MenuBuilder.EndSection();
+
+	
+	return MenuBuilder.MakeWidget();
+}
+
+TSharedRef<SWidget> SNiagaraDebugCaptureView::GenerateFilterMenuContent()
+{
+	TSharedRef<SWidget> Content = SNew(SBox)
+	.MinDesiredHeight(400.0f)
+	.MinDesiredWidth(200.0f)
+	[
+		OverviewFilterWidget.ToSharedRef()
+	];
+	
+	return Content;
+}
+
+FText SNiagaraDebugCaptureView::GetCaptureLabel()
+{
+	switch(FrameMode)
+	{
+	case ENiagaraDebugCaptureFrameMode::SingleFrame:
+		return LOCTEXT("CaptureSingle", "Capture (Single)");
+	case ENiagaraDebugCaptureFrameMode::MultiFrame:
+		return LOCTEXT("CaptureMulti", "Capture (Multi)");
+	default:
+		return FText();
+	}
+}
+
+FText SNiagaraDebugCaptureView::GetCaptureTooltip()
+{
+	switch(FrameMode)
+	{
+	case ENiagaraDebugCaptureFrameMode::SingleFrame:
+		return LOCTEXT("CapturSingleToolTip", "Captures a single frame, and displays it in the current window.");
+	case ENiagaraDebugCaptureFrameMode::MultiFrame:
+		return LOCTEXT("CaptureMultiToolTip", "Captures the specified number of frames, and opens the cache in a standalone editor.");
+	default:
+		return FText();
+	}
+}
+
+FSlateIcon SNiagaraDebugCaptureView::GetCaptureIcon()
+{
+	switch(FrameMode)
+	{
+	case ENiagaraDebugCaptureFrameMode::SingleFrame:
+		return FSlateIcon();
+	case ENiagaraDebugCaptureFrameMode::MultiFrame:
+		return FSlateIcon();
+	default:
+		return FSlateIcon();
+	}
+}
+
 void SNiagaraDebugCaptureView::Construct(const FArguments& InArgs, const TSharedRef<FNiagaraSystemViewModel> InSystemViewModel, const TSharedRef<FNiagaraSimCacheViewModel> InSimCacheViewModel)
 {
+	NumFrames = FMath::Max(1, GetDefault<UNiagaraSettings>()->QuickSimCacheCaptureFrameCount);
+
 	TargetComponent = InSystemViewModel->GetPreviewComponent();
 	SimCacheViewModel = InSimCacheViewModel;
 	SystemViewModel = InSystemViewModel;
 
-	CapturedCache = NewObject<UNiagaraSimCache>(GetTransientPackage());
+	CapturedCache = NewObject<UNiagaraSimCache>(GetTransientPackage(), FNiagaraEditorUtilities::GetUniqueObjectName<UNiagaraSimCache>(GetTransientPackage(), "TempCache"));
 	SimCacheViewModel.Get()->Initialize(CapturedCache);
 	CapturedCache->SetFlags(RF_Standalone);
 
 	FDetailsViewArgs DetailsViewArgs;
 	DetailsViewArgs.bShowOptions = false;
+
+	SAssignNew(OverviewFilterWidget, SNiagaraSimCacheOverview)
+			.SimCacheViewModel(SimCacheViewModel);
 	
-	ChildSlot
-	[
-		SNew(SSplitter)
-		.Orientation(Orient_Vertical)
-		+SSplitter::Slot()
-		.SizeRule(SSplitter::SizeToContent)
-		[
-			// Capture controls
-			SNew(SSplitter)
-			.Orientation(Orient_Vertical)
-			+SSplitter::Slot()
-			.Resizable(false)
-			.SizeRule(SSplitter::SizeToContent)
-			[
-				// Header
-				SNew(SBorder)
-				.BorderImage(FAppStyle::GetNoBrush())
-				.Padding(5.0f)
+	FSlimHorizontalToolBarBuilder DebugCaptureToolbarBuilder (MakeShareable(new FUICommandList), FMultiBoxCustomization::None);
+
+	DebugCaptureToolbarBuilder.BeginSection("Filter");
+	{
+		DebugCaptureToolbarBuilder.AddComboButton(
+			FUIAction(),
+			FOnGetContent::CreateRaw(this, &SNiagaraDebugCaptureView::GenerateFilterMenuContent),
+			FText::GetEmpty(),
+			LOCTEXT("FilterCombo_Tooltip", "Change filter and selection"),
+			FSlateIcon(FAppStyle::GetAppStyleSetName(),"Icons.Filter")
+		);
+	}
+	DebugCaptureToolbarBuilder.EndSection();
+
+	DebugCaptureToolbarBuilder.BeginSection("Capture");
+	{
+		DebugCaptureToolbarBuilder.AddToolBarButton(
+			FUIAction(FExecuteAction::CreateSP(this, &SNiagaraDebugCaptureView::OnCaptureSelected)),
+			NAME_None,
+			TAttribute<FText>::Create(TAttribute<FText>::FGetter::CreateRaw(this, &SNiagaraDebugCaptureView::GetCaptureLabel)),
+			TAttribute<FText>::Create(TAttribute<FText>::FGetter::CreateRaw(this, &SNiagaraDebugCaptureView::GetCaptureTooltip)),
+			TAttribute< FSlateIcon >::Create(TAttribute< FSlateIcon >::FGetter::CreateRaw(this, &SNiagaraDebugCaptureView::GetCaptureIcon))
+		);
+		DebugCaptureToolbarBuilder.AddComboButton(
+			FUIAction(),
+			FOnGetContent::CreateRaw(this, &SNiagaraDebugCaptureView::GenerateCaptureMenuContent),
+			LOCTEXT("CaptureCombo_Label", "Capture Options"),
+			LOCTEXT("CaptureCombo_Tooltip", "Capture Options Menu"),
+			FSlateIcon(FAppStyle::GetAppStyleSetName(),"LevelEditor.Build"),
+			true
+		);
+	}
+	
+	DebugCaptureToolbarBuilder.EndSection();
+
+	DebugCaptureToolbarBuilder.BeginSection("Frames");
+	{
+		DebugCaptureToolbarBuilder.AddWidget(
+			SNew(SHorizontalBox)
+				+SHorizontalBox::Slot()
+				.HAlign(HAlign_Center)
+				.VAlign(VAlign_Center)
 				[
 					SNew(STextBlock)
-					.Justification(ETextJustify::Center)
-					.Text(LOCTEXT("CaptureControlsLabel", "Capture"))
-				]
-			]
-			+SSplitter::Slot()
-			.SizeRule(SSplitter::SizeToContent)
-			[
-				SNew(SUniformGridPanel)
-				.SlotPadding(2.0f)
-				// Capture Single Frame
-				+ SUniformGridPanel::Slot(0, 0)
-				[
-					SNew(SButton)
-					.OnClicked(this, &SNiagaraDebugCaptureView::OnSingleFrameSelected)
-					.ToolTipText(LOCTEXT("SingleFrameButtonTooltip", "Capture a single frame."))
-					.HAlign(HAlign_Center)
-					.VAlign(VAlign_Center)
-					[
-						SNew(STextBlock)
-						.Text(LOCTEXT("SingleFrameButton", "Single"))
-					]
-				]
-				// Capture Multi Frame
-				+SUniformGridPanel::Slot(1, 0)
-				[
-					SNew(SButton)
-					.OnClicked(this, &SNiagaraDebugCaptureView::OnMultiFrameSelected)
-					.ToolTipText(LOCTEXT("MultiFrameButtonTooltip", "Capture multiple frames."))
-					.HAlign(HAlign_Center)
-					.VAlign(VAlign_Center)
-					[
-						SNew(STextBlock)
-						.Text(LOCTEXT("MultiFrameButton", "Multi"))
-					]
+					.Text(LOCTEXT("NumberOfFrames", "Frames"))
+					.IsEnabled_Lambda([this]{return FrameMode == ENiagaraDebugCaptureFrameMode::MultiFrame;})
 				]
-				// Capture Until Complete
-				+SUniformGridPanel::Slot(2, 0)
+				+SHorizontalBox::Slot()
+				.AutoWidth()
 				[
-					SNew(SHorizontalBox)
-					+SHorizontalBox::Slot()
-					.HAlign(HAlign_Center)
-					.VAlign(VAlign_Center)
-					[
-						SNew(STextBlock)
-						.Text(LOCTEXT("NumberOfFrames", "Frames"))
-					]
-					+SHorizontalBox::Slot()
-					[
-						SNew(SSpinBox<int32>)
-						.Value(1)
-						.MinValue(1)
-						.OnValueChanged(this, &SNiagaraDebugCaptureView::OnNumFramesChanged)
-					]
+					SNew(SSpinBox<int32>)
+					.Value(NumFrames)
+					.MinValue(1)
+					.OnValueChanged(this, &SNiagaraDebugCaptureView::OnNumFramesChanged)
+					.IsEnabled_Lambda([this]{return FrameMode == ENiagaraDebugCaptureFrameMode::MultiFrame;})
 				]
-			]
-		]
-		+SSplitter::Slot()
-		.Value(0.99f)
-		[
-			SNew(SNiagaraSimCacheOverview)
-			.SimCacheViewModel(SimCacheViewModel)
-		]
+		);
+	}
+
+	DebugCaptureToolbarBuilder.EndSection();
+	
+	ChildSlot
+	[
+		DebugCaptureToolbarBuilder.MakeWidget()
 	];
+	
+}
+
+SNiagaraDebugCaptureView::SNiagaraDebugCaptureView()
+{
+	CapturedCache = nullptr;
 }
 
 SNiagaraDebugCaptureView::~SNiagaraDebugCaptureView()
 {
-	if(CapturedCache)
+	if(CapturedCache.IsValid())
 	{
 		CapturedCache->ClearFlags(RF_Standalone);
 		CapturedCache->MarkAsGarbage();
+		CapturedCache.Reset();
 	}
 }
 
-FReply SNiagaraDebugCaptureView::OnSingleFrameSelected()
+void SNiagaraDebugCaptureView::OnCaptureSelected()
 {
-	if(!bIsCaptureActive && TargetComponent && CapturedCache)
+	switch(FrameMode)
 	{
-		SystemViewModel.Get()->GetSequencer().Get()->OnPlay(false);
+	case ENiagaraDebugCaptureFrameMode::SingleFrame:
+		OnSingleFrameSelected();
+		break;
+	case ENiagaraDebugCaptureFrameMode::MultiFrame:
+		OnMultiFrameSelected();
+		break;
+	default:
+		break;
+	}
+}
+
+void SNiagaraDebugCaptureView::OnSingleFrameSelected()
+{
+	// The captured cache can be have its standalone flag cleared if a world in the transient package is cleaned up.
+	// @TODO fix the accidental GC issue by making the view model a GC Object, so we can capture single frames after closing another Niagara editor.
+	if(!bIsCaptureActive && TargetComponent && CapturedCache.IsValid())
+	{
+		TSharedPtr<ISequencer> Sequencer = SystemViewModel.Get()->GetSequencer();
+		Sequencer->OnPlay(false);
 
 		const FNiagaraSimCacheCreateParameters CreateParameters;
 
-		SimCacheCapture.CaptureCurrentFrameImmediate(CapturedCache, CreateParameters, TargetComponent, CapturedCache, true);
+		FQualifiedFrameTime StartTime = Sequencer->GetGlobalTime();
+		float CurrentAge = TargetComponent->GetDesiredAge();
+
+		UNiagaraSimCache* OutCache = CapturedCache.Get();
 
-		if(CapturedCache)
+		SimCacheCapture.CaptureCurrentFrameImmediate(CapturedCache.Get(), CreateParameters, TargetComponent, OutCache, true, 0.01666f);
+
+		if(OutCache)
 		{
-			OnCaptureComplete(CapturedCache);
+			
+			// 60fps
+			FFrameRate SystemFrameRate (60, 1);
+
+			FFrameTime NewTime = StartTime.ConvertTo(SystemFrameRate) + FQualifiedFrameTime(1, SystemFrameRate).Time;
+
+			
+			TargetComponent->SetDesiredAge(CurrentAge + 0.01666f);
+			
+			Sequencer->SetGlobalTime(ConvertFrameTime(NewTime, SystemFrameRate, StartTime.Rate));
+			
+			OnCaptureComplete(CapturedCache.Get());
 		}
-		
 	}
-	return FReply::Handled();
 }
 
-FReply SNiagaraDebugCaptureView::OnMultiFrameSelected()
+void SNiagaraDebugCaptureView::OnMultiFrameSelected()
 {
-	if(!bIsCaptureActive && TargetComponent && CapturedCache)
+	if(!bIsCaptureActive && TargetComponent)
 	{
+		UNiagaraSimCache* MultiFrameCache = NewObject<UNiagaraSimCache>(GetTransientPackage(), FNiagaraEditorUtilities::GetUniqueObjectName<UNiagaraSimCache>(GetTransientPackage(), "TempCache"));
+		MultiFrameCache->SetFlags(RF_Transient);
 		const FNiagaraSimCacheCreateParameters CreateParameters;
 		
 		FNiagaraSimCacheCaptureParameters CaptureParameters;
 		CaptureParameters.NumFrames = NumFrames;
-
-		SystemViewModel.Get()->GetSequencer().Get()->OnPlay(false);
+		CaptureParameters.bManuallyAdvanceSimulation = true;
+		CaptureParameters.AdvanceDeltaTime = 0.01666f;
 		
 		bIsCaptureActive = true;
 
-		SimCacheCapture.CaptureNiagaraSimCache(CapturedCache, CreateParameters, TargetComponent, CaptureParameters);
-		
 		SimCacheCapture.OnCaptureComplete().AddSP(this, &SNiagaraDebugCaptureView::OnCaptureComplete);
+		
+		SimCacheCapture.CaptureNiagaraSimCache(MultiFrameCache, CreateParameters, TargetComponent, CaptureParameters);
 	}
-
-	return FReply::Handled();
 }
 
 void SNiagaraDebugCaptureView::OnCaptureComplete(UNiagaraSimCache* CapturedSimCache)
@@ -190,7 +384,19 @@ void SNiagaraDebugCaptureView::OnCaptureComplete(UNiagaraSimCache* CapturedSimCa
 		}
 	}
 
-	RequestSpreadsheetTab.ExecuteIfBound();
+	switch(FrameMode)
+	{
+	case ENiagaraDebugCaptureFrameMode::SingleFrame:
+		RequestSpreadsheetTab.ExecuteIfBound();
+		break;
+	case ENiagaraDebugCaptureFrameMode::MultiFrame:
+		GEditor->GetEditorSubsystem<UAssetEditorSubsystem>()->OpenEditorForAsset(CapturedSimCache, EToolkitMode::Standalone);
+		break;
+	default:
+		break;
+	}
+
+	
 }
 
 #undef LOCTEXT_NAMESPACE
diff --git a/Engine/Plugins/FX/Niagara/Source/NiagaraEditor/Private/Widgets/SNiagaraDebugCaptureView.h b/Engine/Plugins/FX/Niagara/Source/NiagaraEditor/Private/Widgets/SNiagaraDebugCaptureView.h
--- a/Engine/Plugins/FX/Niagara/Source/NiagaraEditor/Private/Widgets/SNiagaraDebugCaptureView.h
+++ b/Engine/Plugins/FX/Niagara/Source/NiagaraEditor/Private/Widgets/SNiagaraDebugCaptureView.h
@@ -8,6 +8,14 @@
 class FNiagaraSimCacheCapture;
 class UNiagaraSimCache;
 
+enum class ENiagaraDebugCaptureFrameMode : uint8
+{
+	// Capture a single frame and display the results in the Niagara editor.
+	SingleFrame,
+	// Capture a user defined number of frames, and open the results in a standalone editor.
+	MultiFrame
+};
+
 class SNiagaraDebugCaptureView : public SCompoundWidget
 {
 public:
@@ -19,17 +27,29 @@ public:
 	DECLARE_DELEGATE(FOnRequestSpreadsheetTab)
 
 	void OnNumFramesChanged(int32 InNumFrames);
+	void CreateComponentSelectionMenuContent(FMenuBuilder& MenuBuilder);
+	void GetComponentNameAndTooltip(const UNiagaraComponent* InComponent, FText& OutName, FText& OutTooltip) const;
+	TSharedRef<SWidget> GenerateCaptureMenuContent();
+	TSharedRef<SWidget> GenerateFilterMenuContent();
+
+	FText GetCaptureLabel();
+	FText GetCaptureTooltip();
+	FSlateIcon GetCaptureIcon();
 	void Construct(const FArguments& InArgs, const TSharedRef<FNiagaraSystemViewModel> SystemViewModel, const TSharedRef<FNiagaraSimCacheViewModel> SimCacheViewModel);
+	SNiagaraDebugCaptureView();
 	virtual ~SNiagaraDebugCaptureView() override;
 	FOnRequestSpreadsheetTab& OnRequestSpreadsheetTab() {return RequestSpreadsheetTab; }
 
 protected:
-	FReply OnSingleFrameSelected();
-	FReply OnMultiFrameSelected();
+	void OnCaptureSelected();
+	void OnSingleFrameSelected();
+	void OnMultiFrameSelected();
 	void OnCaptureComplete(UNiagaraSimCache* CapturedSimCache);
+
+	ENiagaraDebugCaptureFrameMode FrameMode = ENiagaraDebugCaptureFrameMode::SingleFrame;
 	
 	FNiagaraSimCacheCapture SimCacheCapture;
-	UNiagaraSimCache* CapturedCache = nullptr;
+	TWeakObjectPtr<UNiagaraSimCache> CapturedCache = nullptr;
 
 	UNiagaraComponent* TargetComponent = nullptr;
 	int32 NumFrames = 1;
@@ -37,5 +57,5 @@ protected:
 	TSharedPtr<FNiagaraSimCacheViewModel> SimCacheViewModel;
 	TSharedPtr<FNiagaraSystemViewModel> SystemViewModel;
 	FOnRequestSpreadsheetTab RequestSpreadsheetTab;
-	
+	TSharedPtr<SWidget> OverviewFilterWidget;
 };
diff --git a/Engine/Plugins/FX/Niagara/Source/NiagaraEditor/Private/Widgets/SNiagaraSimCacheTreeView.cpp b/Engine/Plugins/FX/Niagara/Source/NiagaraEditor/Private/Widgets/SNiagaraSimCacheTreeView.cpp
--- a/Engine/Plugins/FX/Niagara/Source/NiagaraEditor/Private/Widgets/SNiagaraSimCacheTreeView.cpp
+++ b/Engine/Plugins/FX/Niagara/Source/NiagaraEditor/Private/Widgets/SNiagaraSimCacheTreeView.cpp
@@ -61,7 +61,6 @@ void SNiagaraSimCacheTreeViewFilterWidget::Construct(const FArguments& InArgs, T
 			SNew(SButton)
 			.Text(LOCTEXT("ClearAll", "Clear All"))
 			.OnClicked(this, &SNiagaraSimCacheTreeViewFilterWidget::OnClearAllReleased)
-			.IsEnabled(this, &SNiagaraSimCacheTreeViewFilterWidget::IsFilterActive)
 		]
 		+SHorizontalBox::Slot()
 		.AutoWidth()
@@ -71,44 +70,10 @@ void SNiagaraSimCacheTreeViewFilterWidget::Construct(const FArguments& InArgs, T
 			SNew(SButton)
 			.Text(LOCTEXT("SelectAll", "Select All"))
 			.OnClicked(this, &SNiagaraSimCacheTreeViewFilterWidget::OnSelectAllReleased)
-			.IsEnabled(this, &SNiagaraSimCacheTreeViewFilterWidget::IsFilterActive)
-		]
-		+SHorizontalBox::Slot()
-		.AutoWidth()
-		.Padding(2.0f)
-		[
-
-			// Filter Toggle
-			SNew(SCheckBox)
-			.Type(ESlateCheckBoxType::ToggleButton)
-			.HAlign(HAlign_Right)
-			//.ToolTip(LOCTEXT("ToggleFiltering", "Toggle Filtering"))
-			.Style(FNiagaraEditorStyle::Get(),"NiagaraEditor.SimCache.FilterToggleStyle")
-			.OnCheckStateChanged(this, &SNiagaraSimCacheTreeViewFilterWidget::OnCheckStateChanged)
-			.IsChecked(this, &SNiagaraSimCacheTreeViewFilterWidget::GetFilterState)
-			.Padding(FMargin(3.0f, 2.0f))
-			[
-					
-				SNew(SImage)
-				.Image(FAppStyle::GetBrush("Icons.Filter"))
-					
-			]
 		]
 	];
 }
 
-void SNiagaraSimCacheTreeViewFilterWidget::OnCheckStateChanged(ECheckBoxState InState)
-{
-	TSharedPtr<SNiagaraSimCacheTreeView> TreeView = WeakTreeView.Pin();
-
-	if(TreeView.IsValid())
-	{
-		const bool bCurrentlyActive = InState == ECheckBoxState::Checked;
-		TreeView->SetFilterActive(bCurrentlyActive);
-	}
-	
-}
-
 FReply SNiagaraSimCacheTreeViewFilterWidget::OnClearAllReleased()
 {
 	TSharedPtr<SNiagaraSimCacheTreeView> TreeView = WeakTreeView.Pin();
@@ -135,23 +100,6 @@ FReply SNiagaraSimCacheTreeViewFilterWidget::OnSelectAllReleased()
 	return FReply::Unhandled();
 }
 
-ECheckBoxState SNiagaraSimCacheTreeViewFilterWidget::GetFilterState() const
-{
-	return IsFilterActive() ? ECheckBoxState::Checked : ECheckBoxState::Unchecked;
-}
-
-bool SNiagaraSimCacheTreeViewFilterWidget::IsFilterActive() const
-{
-	TSharedPtr<SNiagaraSimCacheTreeView> TreeView = WeakTreeView.Pin();
-
-	if(TreeView.IsValid())
-	{
-		return TreeView->IsFilterActive();
-	}
-
-	return false;
-}
-
 //// Visibility Widget //////
 
 void SSimCacheTreeViewVisibilityWidget::Construct(const FArguments& InArgs, TWeakPtr<FNiagaraSimCacheTreeItem> InTreeItem, TWeakPtr<SNiagaraSimCacheTreeView> InWeakTreeView)
@@ -166,7 +114,6 @@ void SSimCacheTreeViewVisibilityWidget::Construct(const FArguments& InArgs, TWea
 		SNew(SCheckBox)
 		.OnCheckStateChanged(this, &SSimCacheTreeViewVisibilityWidget::OnCheckStateChanged)
 		.IsChecked(this, &SSimCacheTreeViewVisibilityWidget::GetCheckedState)
-		.IsEnabled(this, &SSimCacheTreeViewVisibilityWidget::IsFilterActive)
 	];
 	
 }
@@ -200,18 +147,6 @@ bool SSimCacheTreeViewVisibilityWidget::IsInFilter() const
 	return false;
 }
 
-bool SSimCacheTreeViewVisibilityWidget::IsFilterActive() const
-{
-	const TSharedPtr<SNiagaraSimCacheTreeView> TreeView = WeakTreeView.Pin();
-
-	if(TreeView.IsValid())
-	{
-		return TreeView->IsFilterActive();
-	}
-
-	return false;
-}
-
 bool SSimCacheTreeViewVisibilityWidget::IsItemSelected() const
 {
 	const TSharedPtr<SNiagaraSimCacheTreeView> TreeView = WeakTreeView.Pin();
@@ -254,6 +189,8 @@ void SNiagaraSimCacheTreeView::Construct(const FArguments& InArgs)
 	.SelectionMode(ESelectionMode::Multi)
 	.TreeItemsSource(ViewModel->GetCurrentRootEntries())
 	.OnGenerateRow(this, &SNiagaraSimCacheTreeView::OnGenerateRow)
+	.OnGeneratePinnedRow(this, &SNiagaraSimCacheTreeView::OnGenerateRow)
+	.ShouldStackHierarchyHeaders(true)
 	.OnGetChildren(this, &SNiagaraSimCacheTreeView::OnGetChildren);
 
 	SetupRootEntries();
@@ -297,6 +234,7 @@ void SNiagaraSimCacheTreeView::OnBufferChanged()
 	TreeView->RequestTreeRefresh();
 	SelectionForFilter.Empty();
 	ViewModel->SetComponentFilters(TArray<FString>());
+	SelectAll();
 	TreeView->SetItemExpansion((*ViewModel->GetCurrentRootEntries())[0], true);
 }
 
@@ -349,21 +287,11 @@ void SNiagaraSimCacheTreeView::UpdateSelectionFilter(TSharedRef<FNiagaraSimCache
 	UpdateStringFilters();
 }
 
-bool SNiagaraSimCacheTreeView::IsFilterActive() const
-{
-	return ViewModel->IsComponentFilterActive();
-}
-
 bool SNiagaraSimCacheTreeView::IsItemInFilter(TSharedPtr<FNiagaraSimCacheTreeItem> InItem) const
 {
 	return ViewModel->GetComponentFilters().Contains(InItem->GetFilterName());
 }
 
-void SNiagaraSimCacheTreeView::SetFilterActive(bool bNewActive)
-{
-	ViewModel->SetComponentFilterActive(bNewActive);
-}
-
 void SNiagaraSimCacheTreeView::VisibilityButtonClicked(TSharedRef<FNiagaraSimCacheTreeItem> InItem)
 {
 	UpdateSelectionFilter(InItem);
@@ -396,6 +324,11 @@ void SNiagaraSimCacheTreeView::UpdateStringFilters()
 
 void SNiagaraSimCacheTreeView::SelectAll()
 {
+	if(!ViewModel->IsCacheValid())
+	{
+		return;
+	}
+	
 	RecursiveAddToSelectionFilter(*ViewModel->GetCurrentRootEntries());
 
 	UpdateStringFilters();
diff --git a/Engine/Plugins/FX/Niagara/Source/NiagaraEditor/Private/Widgets/SNiagaraSimCacheTreeView.h b/Engine/Plugins/FX/Niagara/Source/NiagaraEditor/Private/Widgets/SNiagaraSimCacheTreeView.h
--- a/Engine/Plugins/FX/Niagara/Source/NiagaraEditor/Private/Widgets/SNiagaraSimCacheTreeView.h
+++ b/Engine/Plugins/FX/Niagara/Source/NiagaraEditor/Private/Widgets/SNiagaraSimCacheTreeView.h
@@ -8,7 +8,6 @@
 #include "Widgets/SNiagaraPinTypeSelector.h"
 #include "EdGraphSchema_Niagara.h"
 #include "SNiagaraSimCacheOverview.h"
-#include "Widgets/Images/SLayeredImage.h"
 #include "Widgets/Layout/SScaleBox.h"
 
 class SNiagaraSimCacheTreeView;
@@ -28,11 +27,8 @@ public:
 
 protected:
 	
-	void OnCheckStateChanged(ECheckBoxState InState);
 	FReply OnClearAllReleased();
 	FReply OnSelectAllReleased();
-	ECheckBoxState GetFilterState() const;
-	bool IsFilterActive() const;
 
 private:
 	TWeakPtr<FNiagaraSimCacheTreeItem> WeakTreeItem;
@@ -54,7 +50,6 @@ protected:
 	ECheckBoxState GetCheckedState() const;
 	
 	bool IsInFilter() const;
-	bool IsFilterActive() const;
 	bool IsItemSelected() const;
 private:
 	TWeakPtr<FNiagaraSimCacheTreeItem> WeakTreeItem;
@@ -72,9 +67,7 @@ SLATE_END_ARGS()
 	void SetupRootEntries();
 	void Construct(const FArguments& InArgs);
 	bool HasSelectionFilter() const {return !SelectionForFilter.IsEmpty();}
-	bool IsFilterActive() const;
 	bool IsItemInFilter(TSharedPtr<FNiagaraSimCacheTreeItem> InItem) const;
-	void SetFilterActive(bool bNewActive);
 	void VisibilityButtonClicked(TSharedRef<FNiagaraSimCacheTreeItem> InItem);
 	bool IsItemSelected(TSharedRef<FNiagaraSimCacheTreeItem> InItem);
 	void ClearFilterSelection();
@@ -215,7 +208,7 @@ struct FNiagaraSimCacheComponentTreeItem : public FNiagaraSimCacheTreeItem
 		Contents->AddSlot()
 		.HAlign(HAlign_Left)
 		.VAlign(VAlign_Center)
-		.AutoWidth()
+		.FillWidth(1.0f)
 		.Padding(5, 0)
 		[
 			SNew(STextBlock)
@@ -225,7 +218,7 @@ struct FNiagaraSimCacheComponentTreeItem : public FNiagaraSimCacheTreeItem
 
 		Contents->AddSlot()
 		.HAlign(HAlign_Right)
-		.FillWidth(1.0f)
+		.AutoWidth()
 		[
 			SAssignNew(VisibilityWidget, SSimCacheTreeViewVisibilityWidget, SharedThis(this), Owner)
 		];
diff --git a/Engine/Plugins/FX/Niagara/Source/NiagaraEditor/Private/Widgets/SNiagaraSimCacheView.cpp b/Engine/Plugins/FX/Niagara/Source/NiagaraEditor/Private/Widgets/SNiagaraSimCacheView.cpp
--- a/Engine/Plugins/FX/Niagara/Source/NiagaraEditor/Private/Widgets/SNiagaraSimCacheView.cpp
+++ b/Engine/Plugins/FX/Niagara/Source/NiagaraEditor/Private/Widgets/SNiagaraSimCacheView.cpp
@@ -2,12 +2,9 @@
 
 #include "SNiagaraSimCacheView.h"
 #include "ViewModels/NiagaraSimCacheViewModel.h"
-#include "NiagaraSimCache.h"
 
 #include "CoreMinimal.h"
 
-#include "Widgets/Input/SComboBox.h"
-#include "Widgets/Input/SEditableTextBox.h"
 #include "Widgets/Input/SSpinBox.h"
 #include "Widgets/Layout/SScrollBox.h"
 #include "Widgets/Views/SListView.h"
@@ -70,11 +67,6 @@ public:
 	TSharedPtr<FNiagaraSimCacheViewModel>   SimCacheViewModel;
 };
 
-bool SNiagaraSimCacheView::IsStringFilterEnabled() const
-{
-	return !SimCacheViewModel->IsComponentFilterActive();
-}
-
 void SNiagaraSimCacheView::Construct(const FArguments& InArgs)
 {
 	SimCacheViewModel = InArgs._SimCacheViewModel;
@@ -87,7 +79,6 @@ void SNiagaraSimCacheView::Construct(const FArguments& InArgs)
 
 	UpdateColumns(true);
 	UpdateRows(false);
-	UpdateBufferSelectionList();
 
 	TSharedRef<SScrollBar> HorizontalScrollBar =
 		SNew(SScrollBar)
@@ -104,57 +95,6 @@ void SNiagaraSimCacheView::Construct(const FArguments& InArgs)
 	ChildSlot
 	[
 		SNew(SVerticalBox)
-		+ SVerticalBox::Slot()
-		.AutoHeight()
-		[
-			// Select cache buffer
-			SNew(SHorizontalBox)
-			+ SHorizontalBox::Slot()
-			.AutoWidth()
-			.Padding(10.0f)
-			.VAlign(VAlign_Center)
-			.HAlign(HAlign_Right)
-			[
-				SNew(STextBlock)
-				.Text(LOCTEXT("CacheBufferSelection", "Cache Buffer Selection"))
-			]
-			+ SHorizontalBox::Slot()
-			.AutoWidth()
-			.VAlign(VAlign_Center)
-			.HAlign(HAlign_Left)
-			[
-				SNew(SComboBox<TSharedPtr<FBufferSelectionInfo>>)
-				.OptionsSource(&BufferSelectionList)
-				.OnGenerateWidget(this, &SNiagaraSimCacheView::BufferSelectionGenerateWidget)
-				.OnSelectionChanged(this, &SNiagaraSimCacheView::BufferSelectionChanged)
-				.InitiallySelectedItem(BufferSelectionList[0])
-				[
-					SNew(STextBlock)
-					.Text(this, &SNiagaraSimCacheView::GetBufferSelectionText)
-				]
-			]
-			// Component Filter
-			+ SHorizontalBox::Slot()
-			.AutoWidth()
-			.Padding(10.0f)
-			.VAlign(VAlign_Center)
-			.HAlign(HAlign_Right)
-			[
-				SNew(STextBlock)
-				.Text(LOCTEXT("ComponentFilter", "Component Filter"))
-				.IsEnabled(this, &SNiagaraSimCacheView::IsStringFilterEnabled)
-			]
-			+ SHorizontalBox::Slot()
-			.AutoWidth()
-			.VAlign(VAlign_Center)
-			.HAlign(HAlign_Left)
-			[
-				SNew(SEditableTextBox)
-				.OnTextChanged(this, &SNiagaraSimCacheView::OnComponentFilterChange)
-				.MinDesiredWidth(100)
-				.IsEnabled(this, &SNiagaraSimCacheView::IsStringFilterEnabled)
-			]
-		]
 		//// Main Spreadsheet View
 		+ SVerticalBox::Slot()
 		.FillHeight(1.0f)
@@ -205,6 +145,11 @@ void SNiagaraSimCacheView::GenerateColumns()
 	//  Give columns a width to prevent them from being shrunk when filtering. 
 	constexpr float ManualWidth = 125.0f;
 	HeaderRowWidget->ClearColumns();
+
+	if(!SimCacheViewModel->IsCacheValid())
+	{
+		return;
+	}
 	
 	// Generate instance count column
 	HeaderRowWidget->AddColumn(
@@ -244,26 +189,8 @@ void SNiagaraSimCacheView::UpdateColumns(const bool bReset)
 		GenerateColumns();
 	}
 	
-	TArray<FString> FilterArray;
-
-	// There are two filtering methods, from the tree view and from the spreadsheet view.
-	// The spreadsheet view uses a filter string input by the user. The tree view uses component item selection.
-	const bool bUseComponentFilter = SimCacheViewModel->IsComponentFilterActive();
-
-	FilterArray = bUseComponentFilter ? SimCacheViewModel->GetComponentFilters() : StringFilterArray;
-
-	// If the component filter is active, always apply the filter. If using the string filter, only apply it if it's not empty.
-	const bool bApplyFilter = bUseComponentFilter || !FilterArray.IsEmpty();
-
-	// If the columns are newly generated, and there are no filters to apply they are already up to date.
-	const bool bColumnsUpToDate = bReset && !bApplyFilter;
-
-	FString ColumnName;
+	HeaderRowWidget->RefreshColumns();
 	
-	if(!bColumnsUpToDate)
-	{
-		HeaderRowWidget->RefreshColumns();
-	}
 }
 
 void SNiagaraSimCacheView::UpdateRows(const bool bRefresh)
@@ -280,67 +207,10 @@ void SNiagaraSimCacheView::UpdateRows(const bool bRefresh)
 	}
 }
 
-void SNiagaraSimCacheView::UpdateBufferSelectionList()
-{
-	BufferSelectionList.Empty();
-
-	if (SimCacheViewModel->IsCacheValid())
-	{
-		BufferSelectionList.Emplace(MakeShared<FBufferSelectionInfo>(INDEX_NONE, LOCTEXT("SystemInstance", "System Instance")));
-		for (int32 i = 0; i < SimCacheViewModel->GetNumEmitterLayouts(); ++i)
-		{
-			BufferSelectionList.Emplace(
-				MakeShared<FBufferSelectionInfo>(i, FText::Format(LOCTEXT("EmitterFormat", "Emitter - {0}"), FText::FromName(SimCacheViewModel->GetEmitterLayoutName(i))))
-			);
-		}
-	}
-	else
-	{
-		BufferSelectionList.Emplace(MakeShared<FBufferSelectionInfo>(INDEX_NONE, LOCTEXT("InvalidCache", "Invalid Cache")));
-	}
-}
-
-TSharedRef<SWidget> SNiagaraSimCacheView::BufferSelectionGenerateWidget(TSharedPtr<FBufferSelectionInfo> InItem)
-{
-	return
-		SNew(STextBlock)
-		.Text(InItem->Value);
-}
-
-void SNiagaraSimCacheView::BufferSelectionChanged(TSharedPtr<FBufferSelectionInfo> NewSelection, ESelectInfo::Type SelectInfo)
-{
-	if(!NewSelection.IsValid() || NewSelection->Key == SimCacheViewModel->GetEmitterIndex())
-	{
-		return;
-	}
-	
-	SimCacheViewModel->SetEmitterIndex(NewSelection->Key);
-}
-
-FText SNiagaraSimCacheView::GetBufferSelectionText() const
-{
-	const int32 EmitterIndex = SimCacheViewModel->GetEmitterIndex();
-	for (const TSharedPtr<FBufferSelectionInfo>& SelectionInfo : BufferSelectionList)
-	{
-		if (SelectionInfo->Key == EmitterIndex)
-		{
-			return SelectionInfo->Value;
-		}
-	}
-	return FText::GetEmpty();
-}
-
-void SNiagaraSimCacheView::OnComponentFilterChange(const FText& InFilter)
-{
-	InFilter.ToString().ParseIntoArray(StringFilterArray, TEXT(","));
-	UpdateColumns(false);
-}
-
 void SNiagaraSimCacheView::OnSimCacheChanged()
 {
 	UpdateRows(true);
 	UpdateColumns(true);
-	UpdateBufferSelectionList();
 }
 
 void SNiagaraSimCacheView::OnViewDataChanged(const bool bFullRefresh)
@@ -361,18 +231,9 @@ void SNiagaraSimCacheView::OnBufferChanged()
 bool SNiagaraSimCacheView::GetShouldGenerateWidget(FName Name)
 {
 	TArray<FString> FilterArray;
+	FilterArray = SimCacheViewModel->GetComponentFilters();
 
-	// There are two filtering methods, from the tree view and from the spreadsheet view.
-	// The spreadsheet view uses a filter string input by the user. The tree view uses component item selection.
-	const bool bUseComponentFilter = SimCacheViewModel->IsComponentFilterActive();
-
-	FilterArray = bUseComponentFilter ? SimCacheViewModel->GetComponentFilters() : StringFilterArray;
-
-	// If the component filter is active, always apply the filter. If using the string filter, only apply it if it's not empty.
-	const bool bApplyFilter = bUseComponentFilter || !FilterArray.IsEmpty();
-
-	// If no filter to apply, all columns are visible.
-	if(!bApplyFilter)
+	if(!SimCacheViewModel->IsComponentFilterActive())
 	{
 		return true;
 	}
@@ -389,20 +250,8 @@ bool SNiagaraSimCacheView::GetShouldGenerateWidget(FName Name)
 	{
 		return ColumnName.Equals(ComponentFilter);
 	};
-
-	auto StringFilterPred = [&ColumnName](const FString& StringFilter)
-	{
-		return ColumnName.Contains(StringFilter);
-	};
-
-	bool bPassedFilter = true;
-
-	if(bApplyFilter)
-	{
-		bPassedFilter = bUseComponentFilter ? FilterArray.ContainsByPredicate(ComponentFilterPred) : FilterArray.ContainsByPredicate(StringFilterPred);
-	}
-
-	return bPassedFilter;
+	
+	return FilterArray.ContainsByPredicate(ComponentFilterPred);
 }
 
 #undef LOCTEXT_NAMESPACE
diff --git a/Engine/Plugins/FX/Niagara/Source/NiagaraEditor/Private/Widgets/SNiagaraSimCacheView.h b/Engine/Plugins/FX/Niagara/Source/NiagaraEditor/Private/Widgets/SNiagaraSimCacheView.h
--- a/Engine/Plugins/FX/Niagara/Source/NiagaraEditor/Private/Widgets/SNiagaraSimCacheView.h
+++ b/Engine/Plugins/FX/Niagara/Source/NiagaraEditor/Private/Widgets/SNiagaraSimCacheView.h
@@ -29,17 +29,7 @@ public:
 	void UpdateColumns(const bool bReset);
 
 	void UpdateRows(const bool bRefresh);
-
-	void UpdateBufferSelectionList();
 	
-	TSharedRef<SWidget> BufferSelectionGenerateWidget(TSharedPtr<FBufferSelectionInfo> InItem);
-
-	void BufferSelectionChanged(TSharedPtr<FBufferSelectionInfo> NewSelection, ESelectInfo::Type SelectInfo);
-
-	FText GetBufferSelectionText() const;
-
-	void OnComponentFilterChange(const FText& InFilter);
-
 	void OnSimCacheChanged();
 
 	void OnViewDataChanged(const bool bFullRefresh);
@@ -51,11 +41,7 @@ private:
 
 	TArray<TSharedPtr<int32>>					RowItems;
 	TSharedPtr<FNiagaraSimCacheViewModel>		SimCacheViewModel;
-
-	TArray<TSharedPtr<FBufferSelectionInfo>>	BufferSelectionList;
 	
 	TSharedPtr<SHeaderRow>						HeaderRowWidget;
 	TSharedPtr<SListView<TSharedPtr<int32>>>	ListViewWidget;
-
-	TArray<FString>								StringFilterArray;
 };
\ No newline at end of file
