Making Sim Cache View Model an FGCOBject

diff --git a/Engine/Plugins/FX/Niagara/Source/NiagaraEditor/Private/Customizations/NiagaraComponentDetails.cpp b/Engine/Plugins/FX/Niagara/Source/NiagaraEditor/Private/Customizations/NiagaraComponentDetails.cpp
--- a/Engine/Plugins/FX/Niagara/Source/NiagaraEditor/Private/Customizations/NiagaraComponentDetails.cpp
+++ b/Engine/Plugins/FX/Niagara/Source/NiagaraEditor/Private/Customizations/NiagaraComponentDetails.cpp
@@ -1278,7 +1278,7 @@ FNiagaraComponentDetails::~FNiagaraComponentDetails()
 
 	for(UNiagaraSimCache*& SimCache : CapturedCaches)
 	{
-		SimCache->ClearFlags(RF_Standalone);
+		SimCache->MarkAsGarbage();
 	}
 
 	FGameDelegates::Get().GetEndPlayMapDelegate().RemoveAll(this);
@@ -1556,7 +1556,7 @@ FReply FNiagaraComponentDetails::OnCaptureSelectedSystem()
 		
 		int32 CaptureIndex = ComponentCaptures.AddDefaulted();
 		UNiagaraSimCache* SimCache = NewObject<UNiagaraSimCache>(GetTransientPackage());
-		SimCache->SetFlags(RF_Standalone);
+		SimCache->SetFlags(RF_Transient);
 		CapturedCaches.Add(SimCache);
 		FNiagaraSimCacheCapture& SimCacheCapture = ComponentCaptures[CaptureIndex];
 
diff --git a/Engine/Plugins/FX/Niagara/Source/NiagaraEditor/Private/ViewModels/NiagaraSimCacheViewModel.cpp b/Engine/Plugins/FX/Niagara/Source/NiagaraEditor/Private/ViewModels/NiagaraSimCacheViewModel.cpp
--- a/Engine/Plugins/FX/Niagara/Source/NiagaraEditor/Private/ViewModels/NiagaraSimCacheViewModel.cpp
+++ b/Engine/Plugins/FX/Niagara/Source/NiagaraEditor/Private/ViewModels/NiagaraSimCacheViewModel.cpp
@@ -20,9 +20,11 @@ FNiagaraSimCacheViewModel::~FNiagaraSimCacheViewModel()
 {
 	UNiagaraSimCache::OnCacheEndWrite.RemoveAll(this);
 	bDelegatesAdded = false;
+	SimCache = nullptr;
+	PreviewComponent = nullptr;
 }
 
-void FNiagaraSimCacheViewModel::Initialize(TWeakObjectPtr<UNiagaraSimCache> SimCache) 
+void FNiagaraSimCacheViewModel::Initialize(TWeakObjectPtr<UNiagaraSimCache> InSimCache) 
 {
 	if (bDelegatesAdded == false)
 	{
@@ -30,7 +32,11 @@ void FNiagaraSimCacheViewModel::Initialize(TWeakObjectPtr<UNiagaraSimCache> SimC
 		UNiagaraSimCache::OnCacheEndWrite.AddSP(this, &FNiagaraSimCacheViewModel::OnCacheModified);
 	}
 
-	WeakSimCache = SimCache;
+	if(InSimCache.IsValid())
+	{
+		SimCache = InSimCache.Get();
+	}
+	
 	UpdateComponentInfos();
 	UpdateCachedFrame();
 	SetupPreviewComponentAndInstance();
@@ -41,15 +47,15 @@ void FNiagaraSimCacheViewModel::Initialize(TWeakObjectPtr<UNiagaraSimCache> SimC
 
 void FNiagaraSimCacheViewModel::UpdateSimCache(const FNiagaraSystemSimCacheCaptureReply& Reply)
 {
-	UNiagaraSimCache* SimCache = nullptr;
+	UNiagaraSimCache* TempSimCache = nullptr;
 	
 	if (Reply.SimCacheData.Num() > 0)
 	{
-		SimCache = NewObject<UNiagaraSimCache>();
+		TempSimCache = NewObject<UNiagaraSimCache>();
 
 		FMemoryReader ArReader(Reply.SimCacheData);
 		FObjectAndNameAsStringProxyArchive ProxyArReader(ArReader, false);
-		SimCache->Serialize(ProxyArReader);
+		TempSimCache->Serialize(ProxyArReader);
 		bComponentFilterActive = false;
 		UpdateComponentInfos();
 		
@@ -58,12 +64,11 @@ void FNiagaraSimCacheViewModel::UpdateSimCache(const FNiagaraSystemSimCacheCaptu
 	{
 		UE_LOG(LogNiagaraEditor, Warning, TEXT("Debug Spreadsheet received empty sim cache data."));
 	}
-	Initialize(SimCache);
+	Initialize(TempSimCache);
 }
 
 void FNiagaraSimCacheViewModel::SetupPreviewComponentAndInstance()
 {
-	UNiagaraSimCache* SimCache = WeakSimCache.Get();
 	UNiagaraSystem* System = SimCache ? SimCache->GetSystem(true) : nullptr;
 
 	if(SimCache && System)
@@ -143,7 +148,7 @@ FText FNiagaraSimCacheViewModel::GetComponentText(const FName ComponentName, con
 
 int32 FNiagaraSimCacheViewModel::GetNumFrames() const
 {
-	UNiagaraSimCache* SimCache = WeakSimCache.Get();
+	
 	return SimCache ? SimCache->GetNumFrames() : 0;
 }
 
@@ -151,7 +156,6 @@ void FNiagaraSimCacheViewModel::SetFrameIndex(const int32 InFrameIndex)
 {
 	FrameIndex = InFrameIndex;
 	UpdateCachedFrame();
-	const UNiagaraSimCache* SimCache = WeakSimCache.Get();
 	if(PreviewComponent && SimCache)
 	{
 		const float Duration = SimCache->GetDurationSeconds();
@@ -178,19 +182,16 @@ void FNiagaraSimCacheViewModel::SetEmitterIndex(const int32 InEmitterIndex)
 
 bool FNiagaraSimCacheViewModel::IsCacheValid()
 {
-	UNiagaraSimCache* SimCache = WeakSimCache.Get();
 	return SimCache ? SimCache->IsCacheValid() : false;
 }
 
 int32 FNiagaraSimCacheViewModel::GetNumEmitterLayouts()
 {
-	UNiagaraSimCache* SimCache = WeakSimCache.Get();
 	return SimCache ? SimCache->GetNumEmitters() : 0;
 }
 
 FName FNiagaraSimCacheViewModel::GetEmitterLayoutName(const int32 Index)
 {
-	UNiagaraSimCache* SimCache = WeakSimCache.Get();
 	return SimCache ? SimCache->GetEmitterName(Index) : NAME_None;
 }
 
@@ -209,11 +210,11 @@ FNiagaraSimCacheViewModel::FOnBufferChanged& FNiagaraSimCacheViewModel::OnBuffer
 	return OnBufferChangedDelegate;
 }
 
-void FNiagaraSimCacheViewModel::OnCacheModified(UNiagaraSimCache* SimCache)
+void FNiagaraSimCacheViewModel::OnCacheModified(UNiagaraSimCache* InSimCache)
 {
-	if ( UNiagaraSimCache* ThisSimCache = WeakSimCache.Get() )
+	if ( SimCache )
 	{
-		if ( ThisSimCache == SimCache )
+		if ( SimCache == InSimCache )
 		{
 			SetFrameIndex(0);
 			UpdateComponentInfos();
@@ -231,8 +232,6 @@ void FNiagaraSimCacheViewModel::UpdateCachedFrame()
 	HalfComponents.Empty();
 	Int32Components.Empty();
 	
-
-	UNiagaraSimCache* SimCache = WeakSimCache.Get();
 	if (SimCache == nullptr)
 	{
 		return;
@@ -270,7 +269,6 @@ void FNiagaraSimCacheViewModel::UpdateComponentInfos()
 	FoundHalfComponents = 0;
 	FoundInt32Components = 0;
 	
-	UNiagaraSimCache* SimCache = WeakSimCache.Get();
 	if (SimCache == nullptr)
 	{
 		return;
@@ -334,7 +332,6 @@ void FNiagaraSimCacheViewModel::UpdateComponentInfos()
 void FNiagaraSimCacheViewModel::BuildTreeItemChildren(TSharedPtr<FNiagaraSimCacheTreeItem> InTreeItem, TWeakPtr<SNiagaraSimCacheTreeView> OwningTreeView)
 {
 	FNiagaraSimCacheTreeItem* TreeItem = InTreeItem.Get();
-	UNiagaraSimCache* SimCache = WeakSimCache.Get();
 
 	if(TreeItem && SimCache)
 	{
@@ -484,6 +481,19 @@ TArray<TSharedRef<FNiagaraSimCacheOverviewItem>>* FNiagaraSimCacheViewModel::Get
 	return &BufferEntries;
 }
 
+void FNiagaraSimCacheViewModel::AddReferencedObjects(FReferenceCollector& Collector)
+{
+	if(SimCache != nullptr)
+	{
+		Collector.AddReferencedObject(SimCache);
+	}
+
+	if(PreviewComponent != nullptr)
+	{
+		Collector.AddReferencedObject(PreviewComponent);
+	}
+}
+
 void FNiagaraSimCacheViewModel::BuildComponentInfos(const FName Name, const UScriptStruct* Struct, TArray<FComponentInfo>& InComponentInfos)
 {
 	int32 NumProperties = 0;
diff --git a/Engine/Plugins/FX/Niagara/Source/NiagaraEditor/Private/ViewModels/NiagaraSimCacheViewModel.h b/Engine/Plugins/FX/Niagara/Source/NiagaraEditor/Private/ViewModels/NiagaraSimCacheViewModel.h
--- a/Engine/Plugins/FX/Niagara/Source/NiagaraEditor/Private/ViewModels/NiagaraSimCacheViewModel.h
+++ b/Engine/Plugins/FX/Niagara/Source/NiagaraEditor/Private/ViewModels/NiagaraSimCacheViewModel.h
@@ -14,7 +14,7 @@ struct FNiagaraSystemSimCacheCaptureReply;
 struct FNiagaraSimCacheDataBuffersLayout;
 class UNiagaraComponent;
 
-class NIAGARAEDITOR_API FNiagaraSimCacheViewModel : public TSharedFromThis<FNiagaraSimCacheViewModel>
+class NIAGARAEDITOR_API FNiagaraSimCacheViewModel : public TSharedFromThis<FNiagaraSimCacheViewModel>, public FGCObject
 {
 public:
 	struct FComponentInfo
@@ -107,11 +107,18 @@ public:
 
 	UNiagaraComponent* GetPreviewComponent() { return PreviewComponent; }
 
+	//~ FGCObject interface
+	virtual void AddReferencedObjects(FReferenceCollector& Collector) override;
+	virtual FString GetReferencerName() const override
+	{
+		return TEXT("FNiagaraSimCacheViewModel");
+	}
+
 private:
 	void BuildComponentInfos(const FName Name, const UScriptStruct* Struct, TArray<FComponentInfo>& ComponentInfos);
 	
 	// The sim cache being viewed
-	TWeakObjectPtr<UNiagaraSimCache>	WeakSimCache;
+	UNiagaraSimCache* SimCache = nullptr;
 
 	// Component for preview scene
 	UNiagaraComponent* PreviewComponent = nullptr;
diff --git a/Engine/Plugins/FX/Niagara/Source/NiagaraEditor/Private/Widgets/SNiagaraDebugCaptureView.cpp b/Engine/Plugins/FX/Niagara/Source/NiagaraEditor/Private/Widgets/SNiagaraDebugCaptureView.cpp
--- a/Engine/Plugins/FX/Niagara/Source/NiagaraEditor/Private/Widgets/SNiagaraDebugCaptureView.cpp
+++ b/Engine/Plugins/FX/Niagara/Source/NiagaraEditor/Private/Widgets/SNiagaraDebugCaptureView.cpp
@@ -194,7 +194,7 @@ void SNiagaraDebugCaptureView::Construct(const FArguments& InArgs, const TShared
 
 	CapturedCache = NewObject<UNiagaraSimCache>(GetTransientPackage(), FNiagaraEditorUtilities::GetUniqueObjectName<UNiagaraSimCache>(GetTransientPackage(), "TempCache"));
 	SimCacheViewModel.Get()->Initialize(CapturedCache);
-	CapturedCache->SetFlags(RF_Standalone);
+	CapturedCache->SetFlags(RF_Transient);
 
 	FDetailsViewArgs DetailsViewArgs;
 	DetailsViewArgs.bShowOptions = false;
@@ -279,7 +279,6 @@ SNiagaraDebugCaptureView::~SNiagaraDebugCaptureView()
 {
 	if(CapturedCache.IsValid())
 	{
-		CapturedCache->ClearFlags(RF_Standalone);
 		CapturedCache->MarkAsGarbage();
 		CapturedCache.Reset();
 	}
@@ -302,8 +301,6 @@ void SNiagaraDebugCaptureView::OnCaptureSelected()
 
 void SNiagaraDebugCaptureView::OnSingleFrameSelected()
 {
-	// The captured cache can be have its standalone flag cleared if a world in the transient package is cleaned up.
-	// @TODO fix the accidental GC issue by making the view model a GC Object, so we can capture single frames after closing another Niagara editor.
 	if(!bIsCaptureActive && TargetComponent && CapturedCache.IsValid())
 	{
 		TSharedPtr<ISequencer> Sequencer = SystemViewModel.Get()->GetSequencer();
