Niagara - Reimplement data interface binding as a post compile step, rather than collecting the DIs and binding them at runtime.  User parameter DIs are still bound at runtime through the override parameters and will now work with chained bindings, e.g. user->emitter->particle.

diff --git a/Engine/Plugins/FX/Niagara/Source/Niagara/Classes/NiagaraDataInterface.h b/Engine/Plugins/FX/Niagara/Source/Niagara/Classes/NiagaraDataInterface.h
--- a/Engine/Plugins/FX/Niagara/Source/Niagara/Classes/NiagaraDataInterface.h
+++ b/Engine/Plugins/FX/Niagara/Source/Niagara/Classes/NiagaraDataInterface.h
@@ -529,7 +529,7 @@ public:
 #endif
 
 	/** Allows data interfaces to cache any static data that may be shared between instances */
-	virtual void CacheStaticBuffers(struct FNiagaraSystemStaticBuffers& StaticBuffers, const struct FNiagaraScriptDataInterfaceInfo& DataInterfaceInfo, bool bUsedByCPU, bool bUsedByGPU) {}
+	virtual void CacheStaticBuffers(struct FNiagaraSystemStaticBuffers& StaticBuffers, const FNiagaraVariable& ResolvedVariable, bool bUsedByCPU, bool bUsedByGPU) {}
 
 	virtual bool NeedsGPUContextInit() const { return false; }
 	virtual bool GPUContextInit(const FNiagaraScriptDataInterfaceCompileInfo& InInfo, void* PerInstanceData, FNiagaraSystemInstance* SystemInstance) const { return false; }
diff --git a/Engine/Plugins/FX/Niagara/Source/Niagara/Classes/NiagaraDataInterfaceArrayImpl.h b/Engine/Plugins/FX/Niagara/Source/Niagara/Classes/NiagaraDataInterfaceArrayImpl.h
--- a/Engine/Plugins/FX/Niagara/Source/Niagara/Classes/NiagaraDataInterfaceArrayImpl.h
+++ b/Engine/Plugins/FX/Niagara/Source/Niagara/Classes/NiagaraDataInterfaceArrayImpl.h
@@ -1081,7 +1081,7 @@ struct FNDIArrayProxyImpl : public INDIArrayProxyBase
 			bool bRWGpuArray = false;
 			FNiagaraDataInterfaceUtilities::ForEachGpuFunction(
 				DataInterface, SystemInstance,
-				[&](const FNiagaraDataInterfaceGeneratedFunction& Function) -> bool
+				[&](const UNiagaraScript* Script, const FNiagaraDataInterfaceGeneratedFunction& Function) -> bool
 				{
 					bRWGpuArray = FNiagaraDataInterfaceArrayImplHelper::IsRWFunction(Function.DefinitionName);
 					return !bRWGpuArray;
diff --git a/Engine/Plugins/FX/Niagara/Source/Niagara/Classes/NiagaraDataInterfaceCurveBase.h b/Engine/Plugins/FX/Niagara/Source/Niagara/Classes/NiagaraDataInterfaceCurveBase.h
--- a/Engine/Plugins/FX/Niagara/Source/Niagara/Classes/NiagaraDataInterfaceCurveBase.h
+++ b/Engine/Plugins/FX/Niagara/Source/Niagara/Classes/NiagaraDataInterfaceCurveBase.h
@@ -122,7 +122,7 @@ public:
 	/** Gets information for all of the curves owned by this curve data interface. */
 	virtual void GetCurveData(TArray<FCurveData>& OutCurveData) { }
 
-	virtual void CacheStaticBuffers(struct FNiagaraSystemStaticBuffers& StaticBuffers, const struct FNiagaraScriptDataInterfaceInfo& DataInterfaceInfo, bool bUsedByCPU, bool bUsedByGPU) override;
+	virtual void CacheStaticBuffers(struct FNiagaraSystemStaticBuffers& StaticBuffers, const FNiagaraVariable& ResolvedVariable, bool bUsedByCPU, bool bUsedByGPU) override;
 #if WITH_EDITORONLY_DATA
 	virtual bool AppendCompileHash(FNiagaraCompileHashVisitor* InVisitor) const override;
 	virtual void GetParameterDefinitionHLSL(const FNiagaraDataInterfaceGPUParamInfo& ParamInfo, FString& OutHLSL) override;
diff --git a/Engine/Plugins/FX/Niagara/Source/Niagara/Classes/NiagaraEmitter.h b/Engine/Plugins/FX/Niagara/Source/Niagara/Classes/NiagaraEmitter.h
--- a/Engine/Plugins/FX/Niagara/Source/Niagara/Classes/NiagaraEmitter.h
+++ b/Engine/Plugins/FX/Niagara/Source/Niagara/Classes/NiagaraEmitter.h
@@ -1091,9 +1091,9 @@ void FVersionedNiagaraEmitterData::ForEachPlatformSet(TAction Func)
 	{
 		if (NiagaraScript)
 		{
-			for (const FNiagaraScriptDataInterfaceInfo& DataInterfaceInfo : NiagaraScript->GetCachedDefaultDataInterfaces())
+			for (const FNiagaraScriptResolvedDataInterfaceInfo& DataInterfaceInfo : NiagaraScript->GetResolvedDataInterfaces())
 			{
-				if (UNiagaraDataInterfacePlatformSet* PlatformSetDI = Cast<UNiagaraDataInterfacePlatformSet>(DataInterfaceInfo.DataInterface))
+				if (UNiagaraDataInterfacePlatformSet* PlatformSetDI = Cast<UNiagaraDataInterfacePlatformSet>(DataInterfaceInfo.ResolvedDataInterface))
 				{
 					Func(PlatformSetDI->Platforms);
 				}
diff --git a/Engine/Plugins/FX/Niagara/Source/Niagara/Classes/NiagaraEmitterInstance.h b/Engine/Plugins/FX/Niagara/Source/Niagara/Classes/NiagaraEmitterInstance.h
--- a/Engine/Plugins/FX/Niagara/Source/Niagara/Classes/NiagaraEmitterInstance.h
+++ b/Engine/Plugins/FX/Niagara/Source/Niagara/Classes/NiagaraEmitterInstance.h
@@ -199,9 +199,6 @@ private:
 
 	TUniquePtr<FEventInstanceData> EventInstanceData;
 
-	/** A parameter store which contains the data interfaces parameters which were defined by the scripts. */
-	FNiagaraParameterStore ScriptDefinedDataInterfaceParameters;
-
 	/* Are the cached emitter bounds dynamic */
 	bool bCachedBoundsDynamic = false;
 
diff --git a/Engine/Plugins/FX/Niagara/Source/Niagara/Classes/NiagaraScript.h b/Engine/Plugins/FX/Niagara/Source/Niagara/Classes/NiagaraScript.h
--- a/Engine/Plugins/FX/Niagara/Source/Niagara/Classes/NiagaraScript.h
+++ b/Engine/Plugins/FX/Niagara/Source/Niagara/Classes/NiagaraScript.h
@@ -1208,8 +1208,16 @@ public:
 	NIAGARA_API TConstArrayView<FNiagaraDataInterfaceGPUParamInfo> GetDataInterfaceGPUParamInfos() const { return CachedScriptVM.ShaderScriptParametersMetadata.DataInterfaceParamInfo; }
 
 	NIAGARA_API TArray<TObjectPtr<UNiagaraParameterCollection>>& GetCachedParameterCollectionReferences();
+
+#if WITH_EDITORONLY_DATA
 	TArray<FNiagaraScriptDataInterfaceInfo>& GetCachedDefaultDataInterfaces() { return CachedDefaultDataInterfaces; }
 	TConstArrayView<FNiagaraScriptDataInterfaceInfo> GetCachedDefaultDataInterfaces() const { return MakeArrayView(CachedDefaultDataInterfaces); }
+	void SetResolvedDataInterfaces(const TArray<FNiagaraScriptResolvedDataInterfaceInfo>& InResolvedDataInterfaces) { ResolvedDataInterfaces = InResolvedDataInterfaces; }
+	void SetResolvedUserDataInterfaceBindings(const TArray<FNiagaraResolvedUserDataInterfaceBinding>& InResolvedUserDataInterfaceBindings) { ResolvedUserDataInterfaceBindings = InResolvedUserDataInterfaceBindings; }
+#endif
+
+	TConstArrayView<FNiagaraScriptResolvedDataInterfaceInfo> GetResolvedDataInterfaces() const { return MakeArrayView(ResolvedDataInterfaces); }
+	TConstArrayView<FNiagaraResolvedUserDataInterfaceBinding> GetResolvedUserDataInterfaceBindings() const { return MakeArrayView(ResolvedUserDataInterfaceBindings); }
 
 #if STATS
 	TArrayView<const TStatId> GetStatScopeIDs() const { return MakeArrayView(StatScopesIDs); }
@@ -1334,8 +1342,16 @@ private:
 	UPROPERTY()
 	TArray<TObjectPtr<UNiagaraParameterCollection>> CachedParameterCollectionReferences;
 
+#if WITH_EDITORONLY_DATA
 	UPROPERTY()
 	TArray<FNiagaraScriptDataInterfaceInfo> CachedDefaultDataInterfaces;
+#endif
+
+	UPROPERTY()
+	TArray<FNiagaraScriptResolvedDataInterfaceInfo> ResolvedDataInterfaces;
+
+	UPROPERTY()
+	TArray<FNiagaraResolvedUserDataInterfaceBinding> ResolvedUserDataInterfaceBindings;
 
 	static UNiagaraDataInterface* CopyDataInterface(UNiagaraDataInterface* Src, UObject* Owner);
 
diff --git a/Engine/Plugins/FX/Niagara/Source/Niagara/Classes/NiagaraSystem.h b/Engine/Plugins/FX/Niagara/Source/Niagara/Classes/NiagaraSystem.h
--- a/Engine/Plugins/FX/Niagara/Source/Niagara/Classes/NiagaraSystem.h
+++ b/Engine/Plugins/FX/Niagara/Source/Niagara/Classes/NiagaraSystem.h
@@ -189,7 +189,6 @@ struct FNiagaraSystemCompileRequest
 	bool QueryCompileComplete(UNiagaraSystem* OwningSystem, bool bWait, const double& MaxDuration);
 	bool Resolve(UNiagaraSystem* OwningSystem, FNiagaraParameterStore& ExposedParameters);
 	void Reset();
-	void UpdateSpawnDataInterfaces();
 	void Launch(UNiagaraSystem* OwningSystem, TConstArrayView<UNiagaraScript*> ScriptsNeedingCompile, TConstArrayView<FAsyncTaskPtr> Tasks);
 
 	bool bIsValid = true;
@@ -504,6 +503,7 @@ public:
 	UPROPERTY(transient)
 	FNiagaraSystemUpdateContext UpdateContext;
 
+	const static FGuid ResolveDIsMessageId;
 	const static FGuid ComputeEmitterExecutionOrderMessageId;
 #endif
 
@@ -577,6 +577,9 @@ protected:
 	UPROPERTY(EditAnywhere, AdvancedDisplay, Category = "Performance", meta = (DisplayName="Disable Debug Switches"))
 	uint32 bDisableDebugSwitchesOnCook : 1;
 
+	UPROPERTY(EditAnywhere, AdvancedDisplay, Category = "System")
+	bool bUseLegacyEmitterDataInterfaceDepencendyResolution = false;
+
 	/* When set the system needs to compile before it can be activated. */
 	uint32 bNeedsRequestCompile : 1;
 
@@ -666,6 +669,9 @@ public:
 	UPROPERTY(EditAnywhere, BlueprintReadOnly, AdvancedDisplay, Category="Rendering", meta=(editcondition="bOverrideTranslucencySortDistanceOffset"))
 	float TranslucencySortDistanceOffset = 0.0f;
 
+	//UPROPERTY()
+	//TArray<TArray<uint8>> CachedDataInterfaceEmitterDependencies;
+
 	/** Computes emitter priorities based on the dependency information. */
 	bool ComputeEmitterPriority(int32 EmitterIdx, TArray<int32, TInlineAllocator<32>>& EmitterPriorities, const TBitArray<TInlineAllocator<32>>& EmitterDependencyGraph);
 
@@ -727,8 +733,6 @@ public:
 	FORCEINLINE bool HasAnyGPUEmitters()const{ return bHasAnyGPUEmitters; }
 	FORCEINLINE bool NeedsGPUContextInitForDataInterfaces() const { return bNeedsGPUContextInitForDataInterfaces; }
 
-	const TArray<FName>& GetUserDINamesReadInSystemScripts() const;
-
 	FBox GetFixedBounds() const;
 	FORCEINLINE void SetFixedBounds(const FBox& Box) { FixedBounds = Box;  }
 
@@ -824,6 +828,8 @@ private:
 
 	/** Helper for filling in attribute datasets per emitter. */
 	void InitEmitterDataSetCompiledData(FNiagaraDataSetCompiledData& DataSetToInit, const FNiagaraEmitterHandle& InAssociatedEmitterHandle);
+
+	void ResolveDIBindings();
 #endif
 
 	void ResolveScalabilitySettings();
@@ -963,16 +969,9 @@ protected:
 	TObjectPtr<UNiagaraBakerSettings> BakerGeneratedSettings;
 #endif
 
-	UPROPERTY()
-	bool bHasSystemScriptDIsWithPerInstanceData;
-
 	UPROPERTY()
 	bool bNeedsGPUContextInitForDataInterfaces;
 
-
-	UPROPERTY()
-	TArray<FName> UserDINamesReadInSystemScripts;
-
 	UPROPERTY()
 	TArray<TObjectPtr<UNiagaraDataChannelDefinitions>> ReferencedDataChannelDefinitions;
 
diff --git a/Engine/Plugins/FX/Niagara/Source/Niagara/Classes/NiagaraSystemImpl.h b/Engine/Plugins/FX/Niagara/Source/Niagara/Classes/NiagaraSystemImpl.h
--- a/Engine/Plugins/FX/Niagara/Source/Niagara/Classes/NiagaraSystemImpl.h
+++ b/Engine/Plugins/FX/Niagara/Source/Niagara/Classes/NiagaraSystemImpl.h
@@ -52,9 +52,9 @@ void UNiagaraSystem::ForEachPlatformSet(TAction Func)
 	{
 		if (NiagaraScript)
 		{
-			for (const FNiagaraScriptDataInterfaceInfo& DataInterfaceInfo : NiagaraScript->GetCachedDefaultDataInterfaces())
+			for (const FNiagaraScriptResolvedDataInterfaceInfo& DataInterfaceInfo : NiagaraScript->GetResolvedDataInterfaces())
 			{
-				if (UNiagaraDataInterfacePlatformSet* PlatformSetDI = Cast<UNiagaraDataInterfacePlatformSet>(DataInterfaceInfo.DataInterface))
+				if (UNiagaraDataInterfacePlatformSet* PlatformSetDI = Cast<UNiagaraDataInterfacePlatformSet>(DataInterfaceInfo.ResolvedDataInterface))
 				{
 					Func(PlatformSetDI->Platforms);
 				}
diff --git a/Engine/Plugins/FX/Niagara/Source/Niagara/Private/DataInterface/NiagaraDataInterfaceAsyncGpuTrace.cpp b/Engine/Plugins/FX/Niagara/Source/Niagara/Private/DataInterface/NiagaraDataInterfaceAsyncGpuTrace.cpp
--- a/Engine/Plugins/FX/Niagara/Source/Niagara/Private/DataInterface/NiagaraDataInterfaceAsyncGpuTrace.cpp
+++ b/Engine/Plugins/FX/Niagara/Source/Niagara/Private/DataInterface/NiagaraDataInterfaceAsyncGpuTrace.cpp
@@ -144,7 +144,7 @@ bool UNiagaraDataInterfaceAsyncGpuTrace::InitPerInstanceData(void* PerInstanceDa
 		return true;
 	}
 	
-	FNiagaraDataInterfaceUtilities::ForEachGpuFunction(this, InSystemInstance, [&](const FNiagaraDataInterfaceGeneratedFunction Function)
+	FNiagaraDataInterfaceUtilities::ForEachGpuFunction(this, InSystemInstance, [&](const UNiagaraScript* Script, const FNiagaraDataInterfaceGeneratedFunction& Function)
 	{
 		if (Function.DefinitionName == NDIAsyncGpuTraceLocal::IssueAsyncRayTraceName
 			|| Function.DefinitionName == NDIAsyncGpuTraceLocal::CreateAsyncRayTraceName
diff --git a/Engine/Plugins/FX/Niagara/Source/Niagara/Private/DataInterface/NiagaraDataInterfaceDataChannelCommon.cpp b/Engine/Plugins/FX/Niagara/Source/Niagara/Private/DataInterface/NiagaraDataInterfaceDataChannelCommon.cpp
--- a/Engine/Plugins/FX/Niagara/Source/Niagara/Private/DataInterface/NiagaraDataInterfaceDataChannelCommon.cpp
+++ b/Engine/Plugins/FX/Niagara/Source/Niagara/Private/DataInterface/NiagaraDataInterfaceDataChannelCommon.cpp
@@ -202,7 +202,7 @@ void FNDIDataChannelCompiledData::GatherAccessInfo(UNiagaraSystem* System, UNiag
 		bUsedByCPU = true;
 		return true;
 	};
-	FNiagaraDataInterfaceUtilities::ForEachVMFunctionEquals(Owner, System, HandleVMFunc);
+	FNiagaraDataInterfaceUtilities::ForEachVMFunction(Owner, System, HandleVMFunc);
 
 
 	//For every GPU script we iterate over the functions it calls and add each of them to the mapping.
@@ -228,7 +228,7 @@ void FNDIDataChannelCompiledData::GatherAccessInfo(UNiagaraSystem* System, UNiag
 		bUsedByGPU = true;
 		return true;
 	};
-	FNiagaraDataInterfaceUtilities::ForEachGpuFunctionEquals(Owner, System, HandleGpuFunc);
+	FNiagaraDataInterfaceUtilities::ForEachGpuFunction(Owner, System, HandleGpuFunc);
 
 	//Now we've generated the complete set of parameters accessed by each GPU script, we sort them to ensure identical access between the hlsl and the table we generate.
 	for (auto& Pair : GPUScriptParameterInfos)
diff --git a/Engine/Plugins/FX/Niagara/Source/Niagara/Private/DataInterface/NiagaraDataInterfaceSimCacheReader.cpp b/Engine/Plugins/FX/Niagara/Source/Niagara/Private/DataInterface/NiagaraDataInterfaceSimCacheReader.cpp
--- a/Engine/Plugins/FX/Niagara/Source/Niagara/Private/DataInterface/NiagaraDataInterfaceSimCacheReader.cpp
+++ b/Engine/Plugins/FX/Niagara/Source/Niagara/Private/DataInterface/NiagaraDataInterfaceSimCacheReader.cpp
@@ -279,7 +279,7 @@ bool UNiagaraDataInterfaceSimCacheReader::PerInstanceTick(void* PerInstanceData,
 		TArray<int32> ComponentOffsets;
 		FNiagaraDataInterfaceUtilities::ForEachGpuFunction(
 			this, SystemInstance,
-			[&](const FNiagaraDataInterfaceGeneratedFunction& FunctionInfo)
+			[&](const UNiagaraScript* Scirpt, const FNiagaraDataInterfaceGeneratedFunction& FunctionInfo)
 			{
 				if (FunctionInfo.Specifiers.Num() != 1)
 				{
diff --git a/Engine/Plugins/FX/Niagara/Source/Niagara/Private/DataInterface/NiagaraDataInterfaceStaticMesh.cpp b/Engine/Plugins/FX/Niagara/Source/Niagara/Private/DataInterface/NiagaraDataInterfaceStaticMesh.cpp
--- a/Engine/Plugins/FX/Niagara/Source/Niagara/Private/DataInterface/NiagaraDataInterfaceStaticMesh.cpp
+++ b/Engine/Plugins/FX/Niagara/Source/Niagara/Private/DataInterface/NiagaraDataInterfaceStaticMesh.cpp
@@ -25,6 +25,7 @@
 #include "Editor.h"
 #include "StaticMeshResources.h"
 #include "Subsystems/ImportSubsystem.h"
+#include "INiagaraEditorOnlyDataUtlities.h"
 #endif
 
 #include "Engine/StaticMeshActor.h"
@@ -36,7 +37,6 @@
 
 #include "NiagaraGpuComputeDispatchInterface.h"
 #include "FXRenderingUtils.h"
-
 #include UE_INLINE_GENERATED_CPP_BY_NAME(NiagaraDataInterfaceStaticMesh)
 
 #if WITH_EDITOR
@@ -3189,6 +3189,11 @@ void UNiagaraDataInterfaceStaticMesh::DrawDebugHud(FNDIDrawDebugHudContext& Debu
 #if WITH_EDITOR
 void UNiagaraDataInterfaceStaticMesh::GetFeedback(UNiagaraSystem* Asset, UNiagaraComponent* Component, TArray<FNiagaraDataInterfaceError>& OutErrors, TArray<FNiagaraDataInterfaceFeedback>& OutWarnings, TArray<FNiagaraDataInterfaceFeedback>& OutInfo)
 {
+	if (Asset == nullptr)
+	{
+		return;
+	}
+
 	AActor* SourceActor = SoftSourceActor.Get();
 	UStaticMesh* CurrentMesh = DefaultMesh;
 
@@ -3199,6 +3204,12 @@ void UNiagaraDataInterfaceStaticMesh::GetFeedback(UNiagaraSystem* Asset, UNiagar
 		bHasNoMeshAssignedWarning = CurrentMesh == nullptr;
 	}
 
+	INiagaraModule& NiagaraModule = FModuleManager::GetModuleChecked<INiagaraModule>("Niagara");
+	const INiagaraEditorOnlyDataUtilities& EditorOnlyDataUtilities = NiagaraModule.GetEditorOnlyDataUtilities();
+	UNiagaraDataInterface* RuntimeInstanceOfThis = EditorOnlyDataUtilities.IsEditorDataInterfaceInstance(this)
+		? EditorOnlyDataUtilities.GetResolvedRuntimeInstanceForEditorDataInterfaceInstance(*Asset, *this)
+		: this;
+
 	if (CurrentMesh != nullptr)
 	{
 		if (CurrentMesh->bAllowCPUAccess == false)
@@ -3209,8 +3220,8 @@ void UNiagaraDataInterfaceStaticMesh::GetFeedback(UNiagaraSystem* Asset, UNiagar
 
 			bool bCpuAccessWarning = false;
 
-			FNiagaraDataInterfaceUtilities::ForEachVMFunctionEquals(
-				this,
+			FNiagaraDataInterfaceUtilities::ForEachVMFunction(
+				RuntimeInstanceOfThis,
 				Asset,
 				[&](const UNiagaraScript* Script, const FVMExternalFunctionBindingInfo& VMFunction) -> bool
 				{
@@ -3219,8 +3230,8 @@ void UNiagaraDataInterfaceStaticMesh::GetFeedback(UNiagaraSystem* Asset, UNiagar
 				}
 			);
 
-			FNiagaraDataInterfaceUtilities::ForEachGpuFunctionEquals(
-				this,
+			FNiagaraDataInterfaceUtilities::ForEachGpuFunction(
+				RuntimeInstanceOfThis,
 				Asset,
 				[&](const UNiagaraScript* Script,const FNiagaraDataInterfaceGeneratedFunction& GpuFunction) -> bool
 				{
@@ -3258,8 +3269,7 @@ void UNiagaraDataInterfaceStaticMesh::GetFeedback(UNiagaraSystem* Asset, UNiagar
 
 	if (GetDefault<UNiagaraSettings>()->NDIStaticMesh_AllowDistanceFields == false)
 	{
-		FNiagaraDataInterfaceUtilities::ForEachGpuFunctionEquals(
-			this, Asset, Component,
+		auto GenerateWarnings = 
 			[&](const UNiagaraScript* Script, const FNiagaraDataInterfaceGeneratedFunction& FunctionBinding)
 			{
 				if (FunctionBinding.DefinitionName == NDIStaticMeshLocal::QueryDistanceFieldName )
@@ -3273,8 +3283,16 @@ void UNiagaraDataInterfaceStaticMesh::GetFeedback(UNiagaraSystem* Asset, UNiagar
 					return false;
 				}
 				return true;
-			}
-		);
+			};
+
+		if (Component != nullptr)
+		{
+			FNiagaraDataInterfaceUtilities::ForEachGpuFunction(RuntimeInstanceOfThis, Component, GenerateWarnings);
+		}
+		else if (Asset != nullptr)
+		{
+			FNiagaraDataInterfaceUtilities::ForEachGpuFunction(RuntimeInstanceOfThis, Asset, GenerateWarnings);
+		}
 	}
 }
 #endif //WITH_EDITOR
diff --git a/Engine/Plugins/FX/Niagara/Source/Niagara/Private/DataInterface/NiagaraDataInterfaceUtilities.cpp b/Engine/Plugins/FX/Niagara/Source/Niagara/Private/DataInterface/NiagaraDataInterfaceUtilities.cpp
--- a/Engine/Plugins/FX/Niagara/Source/Niagara/Private/DataInterface/NiagaraDataInterfaceUtilities.cpp
+++ b/Engine/Plugins/FX/Niagara/Source/Niagara/Private/DataInterface/NiagaraDataInterfaceUtilities.cpp
@@ -13,41 +13,23 @@
 namespace FNiagaraDataInterfaceUtilities
 {
 
-bool DataInterfaceEqual(const UNiagaraDataInterface* Lhs, const UNiagaraDataInterface* Rhs)
+void ForEachVMFunctionImpl(class UNiagaraDataInterface* RuntimeResolvedDataInterface, class UNiagaraSystem* NiagaraSystem, const FNiagaraParameterStore& ParameterStore, TFunction<bool(const UNiagaraScript* Script, const FVMExternalFunctionBindingInfo&)> Action)
 {
-	if ( !Lhs || !Rhs )
-	{
-		return false;
-	}
-	if ( Lhs == Rhs )
-	{
-		return true;
-	}
-	if ( Lhs->GetClass() != Rhs->GetClass() )
-	{
-		return false;
-	}
-	return Lhs->Equals(Rhs);
-}
-
-void ForEachVMFunctionEqualsImpl(class UNiagaraDataInterface* DataInterface, class UNiagaraSystem* NiagaraSystem, const FNiagaraParameterStore& ParameterStore, TFunction<bool(const UNiagaraScript* Script, const FVMExternalFunctionBindingInfo&)> Action)
-{
-	if ( DataInterface == nullptr || NiagaraSystem == nullptr )
+	if (RuntimeResolvedDataInterface == nullptr || NiagaraSystem == nullptr )
 	{
 		return;
 	}
 
-	// Find all override parameter names
-	TArray<FName, TInlineAllocator<8>> OverrideParameterNames;
+	// If the resolved data interface is an override DI capture the corresponding user parameter name here.
+	FNiagaraVariableBase DataInterfaceUserVariable;
+	for (const UNiagaraDataInterface* OverrideDI : ParameterStore.GetDataInterfaces())
 	{
-		for (const UNiagaraDataInterface* OverrideDI : ParameterStore.GetDataInterfaces())
+		if (OverrideDI != nullptr && OverrideDI == RuntimeResolvedDataInterface)
 		{
-			if (OverrideDI && DataInterfaceEqual(OverrideDI, DataInterface))
+			if (const FNiagaraVariableBase* Variable = ParameterStore.FindVariable(OverrideDI))
 			{
-				if (const FNiagaraVariableBase* Variable = ParameterStore.FindVariable(OverrideDI))
-				{
-					OverrideParameterNames.AddUnique(Variable->GetName());
-				}
+				DataInterfaceUserVariable = *Variable;
+				break;
 			}
 		}
 	}
@@ -68,54 +50,47 @@ void ForEachVMFunctionEqualsImpl(class UNiagaraDataInterface* DataInterface, cla
 				return;
 			}
 
-			const TArray<FNiagaraScriptDataInterfaceInfo>& CachedDefaultDIs = Script->GetCachedDefaultDataInterfaces();
-			const TArray<FNiagaraScriptDataInterfaceCompileInfo>& DataInterfaceInfos = ScriptExecutableData.DataInterfaceInfo;
-			for ( const FVMExternalFunctionBindingInfo& FunctionBinding : ScriptExecutableData.CalledVMExternalFunctions )
+			TArrayView<const FNiagaraScriptResolvedDataInterfaceInfo> ResolvedDIs = Script->GetResolvedDataInterfaces();
+			for (const FVMExternalFunctionBindingInfo& FunctionBinding : ScriptExecutableData.CalledVMExternalFunctions)
 			{
-				const int NumDataInterface = FMath::Min(DataInterfaceInfos.Num(), CachedDefaultDIs.Num());	// Note: Should always be equal but lets be safe
-				for ( int iDataInterface=0; iDataInterface < NumDataInterface; ++iDataInterface )
+				for (const FNiagaraScriptResolvedDataInterfaceInfo& ResolvedDI : ResolvedDIs)
 				{
-					const FNiagaraScriptDataInterfaceCompileInfo& DataInterfaceInfo = DataInterfaceInfos[iDataInterface];
-					if ( DataInterfaceInfo.Name != FunctionBinding.OwnerName )
+					if (ResolvedDI.CompileName != FunctionBinding.OwnerName)
 					{
 						continue;
 					}
 
-					// Do we have a match?
-					const FNiagaraScriptDataInterfaceInfo& CachedDefaultDI = CachedDefaultDIs[iDataInterface];
-					if (CachedDefaultDI.DataInterface && (DataInterfaceEqual(CachedDefaultDI.DataInterface, DataInterface) || OverrideParameterNames.Contains(CachedDefaultDI.Name)))
+					if (ResolvedDI.ResolvedDataInterface == RuntimeResolvedDataInterface || (DataInterfaceUserVariable.IsValid() && ResolvedDI.ResolvedVariable == DataInterfaceUserVariable))
 					{
-						if ( Action(Script, FunctionBinding) == false )
+						if (Action(Script, FunctionBinding) == false)
 						{
 							bContinueSearching = false;
 							return;
 						}
 					}
-					break;
 				}
 			}
 		}
 	);
 }
 
-void ForEachGpuFunctionEqualsImpl(class UNiagaraDataInterface* DataInterface, class UNiagaraSystem* NiagaraSystem, const FNiagaraParameterStore& ParameterStore, TFunction<bool(const UNiagaraScript*, const FNiagaraDataInterfaceGeneratedFunction&)> Action)
+void ForEachGpuFunctionImpl(class UNiagaraDataInterface* RuntimeResolvedDataInterface, class UNiagaraSystem* NiagaraSystem, const FNiagaraParameterStore& ParameterStore, TFunction<bool(const UNiagaraScript*, const FNiagaraDataInterfaceGeneratedFunction&)> Action)
 {
-	if (DataInterface == nullptr || NiagaraSystem == nullptr)
+	if (RuntimeResolvedDataInterface == nullptr || NiagaraSystem == nullptr)
 	{
 		return;
 	}
 
-	// Find all override parameter names
-	TArray<FName, TInlineAllocator<8>> OverrideParameterNames;
+	// If the resolved data interface is an override DI capture the corresponding user parameter name here.
+	FNiagaraVariableBase DataInterfaceUserVariable;
+	for (const UNiagaraDataInterface* OverrideDI : ParameterStore.GetDataInterfaces())
 	{
-		for (const UNiagaraDataInterface* OverrideDI : ParameterStore.GetDataInterfaces())
+		if (OverrideDI != nullptr && OverrideDI == RuntimeResolvedDataInterface)
 		{
-			if (DataInterfaceEqual(OverrideDI, DataInterface))
+			if (const FNiagaraVariableBase* Variable = ParameterStore.FindVariable(OverrideDI))
 			{
-				if (const FNiagaraVariableBase* Variable = ParameterStore.FindVariable(OverrideDI))
-				{
-					OverrideParameterNames.AddUnique(Variable->GetName());
-				}
+				DataInterfaceUserVariable = *Variable;
+				break;
 			}
 		}
 	}
@@ -132,14 +107,18 @@ void ForEachGpuFunctionEqualsImpl(class UNiagaraDataInterface* DataInterface, cl
 
 			if (FNiagaraShaderScript* ShaderScript = Script->GetRenderThreadScript())
 			{
-				const TArray<FNiagaraScriptDataInterfaceInfo>& CachedDefaultDIs = Script->GetCachedDefaultDataInterfaces();
+				TArrayView<const FNiagaraScriptResolvedDataInterfaceInfo> ResolvedDIs = Script->GetResolvedDataInterfaces();
 				const TArray<FNiagaraDataInterfaceGPUParamInfo>& DataInterfaceParamInfos = ShaderScript->GetScriptParametersMetadata()->DataInterfaceParamInfo;
 
-				const int NumDataInterfaces = FMath::Min(CachedDefaultDIs.Num(), DataInterfaceParamInfos.Num());	// Note: Should always be equal but lets be safe
+				const int NumDataInterfaces = FMath::Min(ResolvedDIs.Num(), DataInterfaceParamInfos.Num());	// Note: Should always be equal but lets be safe
 				for (int iDataInterface = 0; iDataInterface < NumDataInterfaces; ++iDataInterface)
 				{
-					const FNiagaraScriptDataInterfaceInfo& CachedDefaultDI = CachedDefaultDIs[iDataInterface];
-					if (CachedDefaultDI.DataInterface && (DataInterfaceEqual(CachedDefaultDI.DataInterface, DataInterface) || OverrideParameterNames.Contains(CachedDefaultDI.Name)))
+					const FNiagaraScriptResolvedDataInterfaceInfo& ResolvedDI = ResolvedDIs[iDataInterface];
+					if (ResolvedDI.ResolvedDataInterface == nullptr)
+					{
+						continue;
+					}
+					if (ResolvedDI.ResolvedDataInterface == RuntimeResolvedDataInterface || (DataInterfaceUserVariable.IsValid() && ResolvedDI.ResolvedVariable == DataInterfaceUserVariable))
 					{
 						for (const FNiagaraDataInterfaceGeneratedFunction& GeneratedFunction : DataInterfaceParamInfos[iDataInterface].GeneratedFunctions)
 						{
@@ -204,162 +183,45 @@ void ForEachDataInterface(FDataInterfaceUsageContext& UsageContext, const FNiaga
 
 } //namespace FNiagaraDataInterfaceUtilities
 
-void FNiagaraDataInterfaceUtilities::ForEachVMFunctionEquals(UNiagaraDataInterface* DataInterface, UNiagaraSystem* NiagaraSystem, TFunction<bool(const UNiagaraScript*, const FVMExternalFunctionBindingInfo&)> Action)
-{
-	ForEachVMFunctionEqualsImpl(DataInterface, NiagaraSystem, NiagaraSystem->GetExposedParameters(), Action);
-}
-
-void FNiagaraDataInterfaceUtilities::ForEachVMFunctionEquals(UNiagaraDataInterface* DataInterface, UNiagaraSystem* NiagaraSystem, UNiagaraComponent* Component, TFunction<bool(const UNiagaraScript*, const FVMExternalFunctionBindingInfo&)> Action)
-{
-	const FNiagaraParameterStore& Parameters = Component ? Component->GetOverrideParameters() : NiagaraSystem->GetExposedParameters();
-	ForEachVMFunctionEqualsImpl(DataInterface, NiagaraSystem, Parameters, Action);
-}
-
-void FNiagaraDataInterfaceUtilities::ForEachVMFunctionEquals(UNiagaraDataInterface* DataInterface, UNiagaraSystem* NiagaraSystem, FNiagaraSystemInstance* SystemInstance, TFunction<bool(const UNiagaraScript*, const FVMExternalFunctionBindingInfo&)> Action)
+void FNiagaraDataInterfaceUtilities::ForEachVMFunction(UNiagaraDataInterface* RuntimeResolvedDataInterface, UNiagaraSystem* NiagaraSystem, TFunction<bool(const UNiagaraScript*, const FVMExternalFunctionBindingInfo&)> Action)
 {
-	const FNiagaraParameterStore& Parameters = (SystemInstance && SystemInstance->GetOverrideParameters()) ? *SystemInstance->GetOverrideParameters() : NiagaraSystem->GetExposedParameters();
-	ForEachVMFunctionEqualsImpl(DataInterface, NiagaraSystem, Parameters, Action);
+	ForEachVMFunctionImpl(RuntimeResolvedDataInterface, NiagaraSystem, NiagaraSystem->GetExposedParameters(), Action);
 }
 
-void FNiagaraDataInterfaceUtilities::ForEachGpuFunctionEquals(class UNiagaraDataInterface* DataInterface, class UNiagaraSystem* NiagaraSystem, TFunction<bool(const UNiagaraScript*, const FNiagaraDataInterfaceGeneratedFunction&)> Action)
+void FNiagaraDataInterfaceUtilities::ForEachVMFunction(UNiagaraDataInterface* RuntimeResolvedDataInterface, UNiagaraComponent* Component, TFunction<bool(const UNiagaraScript*, const FVMExternalFunctionBindingInfo&)> Action)
 {
-	ForEachGpuFunctionEqualsImpl(DataInterface, NiagaraSystem, NiagaraSystem->GetExposedParameters(), Action);
+	if (Component != nullptr && Component->GetAsset() != nullptr)
+	{
+		ForEachVMFunctionImpl(RuntimeResolvedDataInterface, Component->GetAsset(), Component->GetOverrideParameters(), Action);
+	}
 }
 
-void FNiagaraDataInterfaceUtilities::ForEachGpuFunctionEquals(class UNiagaraDataInterface* DataInterface, class UNiagaraSystem* NiagaraSystem, UNiagaraComponent* Component, TFunction<bool(const UNiagaraScript*, const FNiagaraDataInterfaceGeneratedFunction&)> Action)
+void FNiagaraDataInterfaceUtilities::ForEachVMFunction(UNiagaraDataInterface* RuntimeResolvedDataInterface, FNiagaraSystemInstance* SystemInstance, TFunction<bool(const UNiagaraScript*, const FVMExternalFunctionBindingInfo&)> Action)
 {
-	const FNiagaraParameterStore& Parameters = Component ? Component->GetOverrideParameters() : NiagaraSystem->GetExposedParameters();
-	ForEachGpuFunctionEqualsImpl(DataInterface, NiagaraSystem, Parameters, Action);
+	if (SystemInstance != nullptr && SystemInstance->GetOverrideParameters() != nullptr && SystemInstance->GetSystem() != nullptr)
+	{
+		ForEachVMFunctionImpl(RuntimeResolvedDataInterface, SystemInstance->GetSystem(), *SystemInstance->GetOverrideParameters(), Action);
+	}
 }
 
-void FNiagaraDataInterfaceUtilities::ForEachGpuFunctionEquals(class UNiagaraDataInterface* DataInterface, class UNiagaraSystem* NiagaraSystem, FNiagaraSystemInstance* SystemInstance, TFunction<bool(const UNiagaraScript*, const FNiagaraDataInterfaceGeneratedFunction&)> Action)
+void FNiagaraDataInterfaceUtilities::ForEachGpuFunction(UNiagaraDataInterface* RuntimeResolvedDataInterface, class UNiagaraSystem* NiagaraSystem, TFunction<bool(const UNiagaraScript*, const FNiagaraDataInterfaceGeneratedFunction&)> Action)
 {
-	const FNiagaraParameterStore& Parameters = (SystemInstance && SystemInstance->GetOverrideParameters()) ? *SystemInstance->GetOverrideParameters() : NiagaraSystem->GetExposedParameters();
-	ForEachGpuFunctionEqualsImpl(DataInterface, NiagaraSystem, Parameters, Action);
+	ForEachGpuFunctionImpl(RuntimeResolvedDataInterface, NiagaraSystem, NiagaraSystem->GetExposedParameters(), Action);
 }
 
-void FNiagaraDataInterfaceUtilities::ForEachVMFunction(class UNiagaraDataInterface* DataInterface, class FNiagaraSystemInstance* SystemInstance, TFunction<bool(const UNiagaraScript*, const FVMExternalFunctionBindingInfo&)> Action)
+void FNiagaraDataInterfaceUtilities::ForEachGpuFunction(UNiagaraDataInterface* RuntimeResolvedDataInterface, UNiagaraComponent* Component, TFunction<bool(const UNiagaraScript*, const FNiagaraDataInterfaceGeneratedFunction&)> Action)
 {
-	if ( !DataInterface || !SystemInstance )
+	if (Component != nullptr && Component->GetAsset() != nullptr)
 	{
-		return;
-	}
-
-	auto SearchVMContext =
-		[](const FNiagaraScriptExecutionContextBase* ExecContext, UNiagaraDataInterface* DataInterface, const TFunction<bool(const UNiagaraScript* Script, const FVMExternalFunctionBindingInfo&)>& Action) -> bool
-		{
-			if (ExecContext == nullptr || ExecContext->Script == nullptr)
-			{
-				return true;
-			}
-
-			const FNiagaraVMExecutableData& ScriptExecutableData = ExecContext->Script->GetVMExecutableData();
-			if (!ScriptExecutableData.IsValid())
-			{
-				return true;
-			}
-
-			const TArray<UNiagaraDataInterface*>& DataInterfaces = ExecContext->GetDataInterfaces();
-			const int32 NumDataInterfaces = FMath::Min(ScriptExecutableData.DataInterfaceInfo.Num(), DataInterfaces.Num());		// Should be equal, but be safe
-			for (const FVMExternalFunctionBindingInfo& FunctionBinding : ScriptExecutableData.CalledVMExternalFunctions)
-			{
-				for (int32 iDataInterface = 0; iDataInterface < NumDataInterfaces; ++iDataInterface)
-				{
-					const FNiagaraScriptDataInterfaceCompileInfo& DataInterfaceInfo = ScriptExecutableData.DataInterfaceInfo[iDataInterface];
-					if (FunctionBinding.OwnerName != DataInterfaceInfo.Name)
-					{
-						continue;
-					}
-					if (DataInterfaces[iDataInterface] != DataInterface)
-					{
-						continue;
-					}
-
-					if (Action(ExecContext->Script, FunctionBinding) == false)
-					{
-						return false;
-					}
-				}
-			}
-			return true;
-		};
-
-
-	// Search system scripts (always VM)
-	TSharedPtr<FNiagaraSystemSimulation, ESPMode::ThreadSafe> SystemSimulation = SystemInstance->GetSystemSimulation();
-	if ( SystemSimulation.IsValid() )
-	{
-		if (SearchVMContext(SystemSimulation->GetSpawnExecutionContext(), DataInterface, Action) == false)
-		{
-			return;
-		}
-		if (SearchVMContext(SystemSimulation->GetUpdateExecutionContext(), DataInterface, Action) == false)
-		{
-			return;
-		}
-	}
-
-	// Search emitter scripts
-	for (const TSharedRef<FNiagaraEmitterInstance, ESPMode::ThreadSafe>& EmitterInstance : SystemInstance->GetEmitters())
-	{
-		if (EmitterInstance->IsDisabled() || EmitterInstance->GetCachedEmitter().Emitter == nullptr || EmitterInstance->GetGPUContext() != nullptr)
-		{
-			continue;
-		}
-
-		if ( SearchVMContext(&EmitterInstance->GetSpawnExecutionContext(), DataInterface, Action) == false )
-		{
-			return;
-		}
-
-		if (SearchVMContext(&EmitterInstance->GetUpdateExecutionContext(), DataInterface, Action) == false)
-		{
-			return;
-		}
-
-		for ( const FNiagaraScriptExecutionContext& EventExecContext : EmitterInstance->GetEventExecutionContexts() )
-		{
-			if (SearchVMContext(&EventExecContext, DataInterface, Action) == false)
-			{
-				return;
-			}
-		}
+		ForEachGpuFunctionImpl(RuntimeResolvedDataInterface, Component->GetAsset(), Component->GetOverrideParameters(), Action);
 	}
 }
 
-void FNiagaraDataInterfaceUtilities::ForEachGpuFunction(class UNiagaraDataInterface* DataInterface, class FNiagaraSystemInstance* SystemInstance, TFunction<bool(const FNiagaraDataInterfaceGeneratedFunction&)> Action)
+void FNiagaraDataInterfaceUtilities::ForEachGpuFunction(UNiagaraDataInterface* RuntimeResolvedDataInterface, FNiagaraSystemInstance* SystemInstance, TFunction<bool(const UNiagaraScript*, const FNiagaraDataInterfaceGeneratedFunction&)> Action)
 {
-	if ( !DataInterface || !SystemInstance )
-	{
-		return;
-	}
-
-	for (const TSharedRef<FNiagaraEmitterInstance, ESPMode::ThreadSafe>& EmitterInstance : SystemInstance->GetEmitters())
+	if (SystemInstance != nullptr && SystemInstance->GetOverrideParameters() != nullptr && SystemInstance->GetSystem() != nullptr)
 	{
-		if (EmitterInstance->IsDisabled() || EmitterInstance->GetCachedEmitter().Emitter == nullptr || EmitterInstance->GetGPUContext() == nullptr)
-		{
-			continue;
-		}
-
-		const FNiagaraScriptInstanceParameterStore& ParameterStore = EmitterInstance->GetGPUContext()->CombinedParamStore;
-		const TArray<UNiagaraDataInterface*>& DataInterfaces = ParameterStore.GetDataInterfaces();
-		const TArray<FNiagaraDataInterfaceGPUParamInfo>& DataInterfaceParamInfos = EmitterInstance->GetGPUContext()->GPUScript_RT->GetScriptParametersMetadata()->DataInterfaceParamInfo;
-		const int32 NumDataInterface = FMath::Min(DataInterfaces.Num(), DataInterfaceParamInfos.Num());		// Should be equal, but be safe
-		for ( int32 iDataInterface=0; iDataInterface < NumDataInterface; ++iDataInterface)
-		{
-			if ( DataInterfaces[iDataInterface] != DataInterface )
-			{
-				continue;
-			}
-
-			for (const FNiagaraDataInterfaceGeneratedFunction& GeneratedFunction : DataInterfaceParamInfos[iDataInterface].GeneratedFunctions)
-			{
-				if (Action(GeneratedFunction) == false)
-				{
-					return;
-				}
-			}
-		}
+		ForEachGpuFunctionImpl(RuntimeResolvedDataInterface, SystemInstance->GetSystem(), *SystemInstance->GetOverrideParameters(), Action);
 	}
 }
 
diff --git a/Engine/Plugins/FX/Niagara/Source/Niagara/Private/NiagaraAsyncCompile.cpp b/Engine/Plugins/FX/Niagara/Source/Niagara/Private/NiagaraAsyncCompile.cpp
--- a/Engine/Plugins/FX/Niagara/Source/Niagara/Private/NiagaraAsyncCompile.cpp
+++ b/Engine/Plugins/FX/Niagara/Source/Niagara/Private/NiagaraAsyncCompile.cpp
@@ -732,39 +732,6 @@ void FNiagaraSystemCompileRequest::Reset()
 	}
 }
 
-void FNiagaraSystemCompileRequest::UpdateSpawnDataInterfaces()
-{
-	// HACK: This is a temporary hack to fix an issue where data interfaces used by modules and dynamic inputs in the
-	// particle update script aren't being shared by the interpolated spawn script when accessed directly.  This works
-	// properly if the data interface is assigned to a named particle parameter and then linked to an input.
-	// TODO: Bind these data interfaces the same way parameter data interfaces are bound.
-	for (auto& AsyncTask : DDCTasks)
-	{
-		FEmitterCompiledScriptPair& EmitterCompiledScriptPair = AsyncTask->ScriptPair;
-		FVersionedNiagaraEmitter Emitter = EmitterCompiledScriptPair.VersionedEmitter;
-		UNiagaraScript* CompiledScript = EmitterCompiledScriptPair.CompiledScript;
-
-		if (UNiagaraScript::IsEquivalentUsage(CompiledScript->GetUsage(), ENiagaraScriptUsage::ParticleUpdateScript))
-		{
-			UNiagaraScript* SpawnScript = Emitter.GetEmitterData()->SpawnScriptProps.Script;
-			for (const FNiagaraScriptDataInterfaceInfo& UpdateDataInterfaceInfo : CompiledScript->GetCachedDefaultDataInterfaces())
-			{
-				// If the data interface isn't being written to a parameter map then it won't be bound properly so we
-				// assign the update scripts copy of the data interface to the spawn scripts copy by pointer so that they will share
-				// the data interface at runtime and will both be updated in the editor.
-				for (FNiagaraScriptDataInterfaceInfo& SpawnDataInterfaceInfo : SpawnScript->GetCachedDefaultDataInterfaces())
-				{
-					if (SpawnDataInterfaceInfo.RegisteredParameterMapWrite == NAME_None && UpdateDataInterfaceInfo.RegisteredParameterMapWrite == NAME_None && UpdateDataInterfaceInfo.Name == SpawnDataInterfaceInfo.Name)
-					{
-						SpawnDataInterfaceInfo.DataInterface = UpdateDataInterfaceInfo.DataInterface;
-						break;
-					}
-				}
-			}
-		}
-	}
-}
-
 void FNiagaraSystemCompileRequest::Launch(UNiagaraSystem* OwningSystem, TConstArrayView<UNiagaraScript*> ScriptsNeedingCompile, TConstArrayView<FAsyncTaskPtr> Tasks)
 {
 	bEvaluateParametersPending = !OwningSystem->ShouldUseRapidIterationParameters();
diff --git a/Engine/Plugins/FX/Niagara/Source/Niagara/Private/NiagaraCommon.cpp b/Engine/Plugins/FX/Niagara/Source/Niagara/Private/NiagaraCommon.cpp
--- a/Engine/Plugins/FX/Niagara/Source/Niagara/Private/NiagaraCommon.cpp
+++ b/Engine/Plugins/FX/Niagara/Source/Niagara/Private/NiagaraCommon.cpp
@@ -1048,29 +1048,6 @@ FNiagaraVariable FNiagaraUtilities::ConvertVariableToRapidIterationConstantName(
 	return Var;
 }
 
-void FNiagaraUtilities::CollectScriptDataInterfaceParameters(const UObject& Owner, const TArrayView<UNiagaraScript*>& Scripts, FNiagaraParameterStore& OutDataInterfaceParameters)
-{
-	for (UNiagaraScript* Script : Scripts)
-	{
-		for (FNiagaraScriptDataInterfaceInfo& DataInterfaceInfo : Script->GetCachedDefaultDataInterfaces())
-		{
-			if (DataInterfaceInfo.RegisteredParameterMapWrite != NAME_None)
-			{
-				FNiagaraVariable DataInterfaceParameter(DataInterfaceInfo.Type, DataInterfaceInfo.RegisteredParameterMapWrite);
-				if (OutDataInterfaceParameters.AddParameter(DataInterfaceParameter, false, false))
-				{
-					OutDataInterfaceParameters.SetDataInterface(DataInterfaceInfo.DataInterface, DataInterfaceParameter);
-				}
-				else
-				{
-					UE_LOG(LogNiagara, Error, TEXT("Duplicate data interface parameter writes found, simulation will be incorrect.  Owner: %s Parameter: %s"),
-						*Owner.GetPathName(), *DataInterfaceInfo.RegisteredParameterMapWrite.ToString());
-				}
-			}
-		}
-	}
-}
-
 bool FNiagaraScriptDataInterfaceInfo::IsUserDataInterface() const
 {
 	TStringBuilder<128> NameBuilder;
@@ -1078,6 +1055,11 @@ bool FNiagaraScriptDataInterfaceInfo::IsUserDataInterface() const
 	return FCString::Strnicmp(NameBuilder.ToString(), TEXT("user."), 5) == 0;
 }
 
+bool FNiagaraScriptResolvedDataInterfaceInfo::NeedsPerInstanceBinding() const
+{
+	return ResolvedVariable.GetName().ToString().StartsWith(TEXT("User."));
+}
+
 bool FNiagaraScriptDataInterfaceCompileInfo::CanExecuteOnTarget(ENiagaraSimTarget SimTarget) const
 {
 	// Note that this can be called on non-game threads. We ensure that the data interface CDO object is already in existence at application init time.
diff --git a/Engine/Plugins/FX/Niagara/Source/Niagara/Private/NiagaraComponentSettings.cpp b/Engine/Plugins/FX/Niagara/Source/Niagara/Private/NiagaraComponentSettings.cpp
--- a/Engine/Plugins/FX/Niagara/Source/Niagara/Private/NiagaraComponentSettings.cpp
+++ b/Engine/Plugins/FX/Niagara/Source/Niagara/Private/NiagaraComponentSettings.cpp
@@ -295,9 +295,9 @@ namespace FNiagaraComponentSettings
 				{
 					if (const UNiagaraScript* GPUComputeScript = EmitterData->GetGPUComputeScript())
 					{
-						for (const FNiagaraScriptDataInterfaceInfo& DefaultDIInfo : GPUComputeScript->GetCachedDefaultDataInterfaces())
+						for (const FNiagaraScriptResolvedDataInterfaceInfo& DefaultDIInfo : GPUComputeScript->GetResolvedDataInterfaces())
 						{
-							if (GpuDataInterfaceDenyList.Contains(DefaultDIInfo.Type.GetFName()))
+							if (GpuDataInterfaceDenyList.Contains(DefaultDIInfo.ResolvedVariable.GetType().GetFName()))
 							{
 								return false;
 							}
diff --git a/Engine/Plugins/FX/Niagara/Source/Niagara/Private/NiagaraDataInterfaceCurveBase.cpp b/Engine/Plugins/FX/Niagara/Source/Niagara/Private/NiagaraDataInterfaceCurveBase.cpp
--- a/Engine/Plugins/FX/Niagara/Source/Niagara/Private/NiagaraDataInterfaceCurveBase.cpp
+++ b/Engine/Plugins/FX/Niagara/Source/Niagara/Private/NiagaraDataInterfaceCurveBase.cpp
@@ -5,6 +5,7 @@
 #include "Internationalization/Internationalization.h"
 
 #include "NiagaraCompileHashVisitor.h"
+#include "NiagaraConstants.h"
 #include "NiagaraCustomVersion.h"
 #include "NiagaraEmitter.h"
 #include "NiagaraRenderer.h"
@@ -430,12 +431,12 @@ bool UNiagaraDataInterfaceCurveBase::Equals(const UNiagaraDataInterface* Other)
 	return bEqual;
 }
 
-void UNiagaraDataInterfaceCurveBase::CacheStaticBuffers(struct FNiagaraSystemStaticBuffers& StaticBuffers, const struct FNiagaraScriptDataInterfaceInfo& DataInterfaceInfo, bool bUsedByCPU, bool bUsedByGPU)
+void UNiagaraDataInterfaceCurveBase::CacheStaticBuffers(struct FNiagaraSystemStaticBuffers& StaticBuffers, const FNiagaraVariable& ResolvedVariable, bool bUsedByCPU, bool bUsedByGPU)
 {
 	const uint32 PrevLUTOffset = LUTOffset;
 
 	LUTOffset = INDEX_NONE;
-	if (bUsedByGPU && DataInterfaceInfo.IsUserDataInterface() == false)
+	if (bUsedByGPU && ResolvedVariable.IsInNameSpace(FNiagaraConstants::UserNamespaceString) == false)
 	{
 		LUTOffset = StaticBuffers.AddGpuData(ShaderLUT);
 	}
diff --git a/Engine/Plugins/FX/Niagara/Source/Niagara/Private/NiagaraDataInterfaceEmitterBinding.cpp b/Engine/Plugins/FX/Niagara/Source/Niagara/Private/NiagaraDataInterfaceEmitterBinding.cpp
--- a/Engine/Plugins/FX/Niagara/Source/Niagara/Private/NiagaraDataInterfaceEmitterBinding.cpp
+++ b/Engine/Plugins/FX/Niagara/Source/Niagara/Private/NiagaraDataInterfaceEmitterBinding.cpp
@@ -25,43 +25,31 @@ FNiagaraEmitterInstance* FNiagaraDataInterfaceEmitterBinding::Resolve(FNiagaraSy
 		// If we came from an emitter script we need to look over the default data interfaces to find ourself
 		else
 		{
-			auto FindDataInterfaceName =
-				[DataInterface](UNiagaraSystem* NiagaraSystem) -> FName
+			FString SourceEmitterName;
+			UNiagaraSystem* NiagaraSystem = SystemInstance->GetSystem();
+			for (UNiagaraScript* NiagaraScript : { NiagaraSystem->GetSystemUpdateScript(), NiagaraSystem->GetSystemSpawnScript() })
+			{
+				if (NiagaraScript != nullptr)
 				{
-					for (UNiagaraScript* NiagaraScript : { NiagaraSystem->GetSystemUpdateScript(), NiagaraSystem->GetSystemSpawnScript() })
+					const FNiagaraScriptResolvedDataInterfaceInfo* MatchingResolvedDataInterface = NiagaraScript->GetResolvedDataInterfaces().FindByPredicate(
+						[&DataInterface](const FNiagaraScriptResolvedDataInterfaceInfo& ResolvedDataInterface) { return ResolvedDataInterface.ResolvedDataInterface == DataInterface; });
+					if (MatchingResolvedDataInterface != nullptr)
 					{
-						if (NiagaraScript != nullptr)
-						{
-							for (const FNiagaraScriptDataInterfaceInfo& DataInterfaceInfo : NiagaraScript->GetCachedDefaultDataInterfaces())
-							{
-								if (DataInterfaceInfo.DataInterface == DataInterface)
-								{
-									return DataInterfaceInfo.RegisteredParameterMapWrite.IsNone() ? DataInterfaceInfo.Name : DataInterfaceInfo.RegisteredParameterMapWrite;
-								}
-							}
-						}
+						SourceEmitterName = MatchingResolvedDataInterface->EmitterName;
+						break;
 					}
-					return NAME_None;
-				};
+				}
+			}
 
-			const FName DataInterfaceName = FindDataInterfaceName(SystemInstance->GetSystem());
-			if (!DataInterfaceName.IsNone())
+			if (SourceEmitterName.IsEmpty() == false)
 			{
-				FNameBuilder SelfEmitterString;
-				DataInterfaceName.ToString(SelfEmitterString);
-				FStringView SelfEmitterView = SelfEmitterString.ToView();
-				int32 DotIndex;
-				if (SelfEmitterView.FindChar('.', DotIndex))
+				for (TSharedPtr<FNiagaraEmitterInstance, ESPMode::ThreadSafe> EmitterInstance : SystemInstance->GetEmitters())
 				{
-					SelfEmitterView = SelfEmitterView.Mid(0, DotIndex);
-					for (TSharedPtr<FNiagaraEmitterInstance, ESPMode::ThreadSafe> EmitterInstance : SystemInstance->GetEmitters())
+					if (const UNiagaraEmitter* CachedEmitter = EmitterInstance->GetCachedEmitter().Emitter)
 					{
-						if (const UNiagaraEmitter* CachedEmitter = EmitterInstance->GetCachedEmitter().Emitter)
+						if (CachedEmitter->GetUniqueEmitterName() == SourceEmitterName)
 						{
-							if (SelfEmitterView.Equals(CachedEmitter->GetUniqueEmitterName()))
-							{
-								return EmitterInstance.Get();
-							}
+							return EmitterInstance.Get();
 						}
 					}
 				}
diff --git a/Engine/Plugins/FX/Niagara/Source/Niagara/Private/NiagaraDataInterfaceGrid2DCollection.cpp b/Engine/Plugins/FX/Niagara/Source/Niagara/Private/NiagaraDataInterfaceGrid2DCollection.cpp
--- a/Engine/Plugins/FX/Niagara/Source/Niagara/Private/NiagaraDataInterfaceGrid2DCollection.cpp
+++ b/Engine/Plugins/FX/Niagara/Source/Niagara/Private/NiagaraDataInterfaceGrid2DCollection.cpp
@@ -2085,11 +2085,11 @@ void UNiagaraDataInterfaceGrid2DCollection::FindAttributes(TArray<FNiagaraVariab
 
 					for (UNiagaraScript* Script : Scripts)
 					{
-						for (FNiagaraScriptDataInterfaceInfo& DataInterfaceInfo : Script->GetCachedDefaultDataInterfaces())
+						for (const FNiagaraScriptResolvedDataInterfaceInfo& DataInterfaceInfo : Script->GetResolvedDataInterfaces())
 						{
-							if (DataInterfaceInfo.DataInterface == DataInterface)
+							if (DataInterfaceInfo.ResolvedDataInterface == DataInterface && DataInterfaceInfo.bIsInternal == false)
 							{
-								return DataInterfaceInfo.RegisteredParameterMapRead.IsNone() ? DataInterfaceInfo.RegisteredParameterMapWrite : DataInterfaceInfo.RegisteredParameterMapRead;
+								return DataInterfaceInfo.ResolvedVariable.GetName();
 							}
 						}
 					}
diff --git a/Engine/Plugins/FX/Niagara/Source/Niagara/Private/NiagaraDataInterfaceGrid3DCollection.cpp b/Engine/Plugins/FX/Niagara/Source/Niagara/Private/NiagaraDataInterfaceGrid3DCollection.cpp
--- a/Engine/Plugins/FX/Niagara/Source/Niagara/Private/NiagaraDataInterfaceGrid3DCollection.cpp
+++ b/Engine/Plugins/FX/Niagara/Source/Niagara/Private/NiagaraDataInterfaceGrid3DCollection.cpp
@@ -4274,11 +4274,11 @@ void UNiagaraDataInterfaceGrid3DCollection::FindAttributes(TArray<FNiagaraVariab
 
 					for (UNiagaraScript* Script : Scripts)
 					{
-						for (FNiagaraScriptDataInterfaceInfo& DataInterfaceInfo : Script->GetCachedDefaultDataInterfaces())
+						for (const FNiagaraScriptResolvedDataInterfaceInfo& DataInterfaceInfo : Script->GetResolvedDataInterfaces())
 						{
-							if (DataInterfaceInfo.DataInterface == DataInterface)
+							if (DataInterfaceInfo.ResolvedDataInterface == DataInterface && DataInterfaceInfo.bIsInternal == false)
 							{
-								return DataInterfaceInfo.RegisteredParameterMapRead.IsNone() ? DataInterfaceInfo.RegisteredParameterMapWrite : DataInterfaceInfo.RegisteredParameterMapRead;
+								return DataInterfaceInfo.ResolvedVariable.GetName();
 							}
 						}
 					}
diff --git a/Engine/Plugins/FX/Niagara/Source/Niagara/Private/NiagaraEmitterInstance.cpp b/Engine/Plugins/FX/Niagara/Source/Niagara/Private/NiagaraEmitterInstance.cpp
--- a/Engine/Plugins/FX/Niagara/Source/Niagara/Private/NiagaraEmitterInstance.cpp
+++ b/Engine/Plugins/FX/Niagara/Source/Niagara/Private/NiagaraEmitterInstance.cpp
@@ -450,35 +450,9 @@ void FNiagaraEmitterInstance::Init(int32 InEmitterIdx, FNiagaraSystemInstanceID
 		{
 			Scripts.Add(SimStage->Script);
 		}
-		FNiagaraUtilities::CollectScriptDataInterfaceParameters(*CachedEmitter.Emitter, MakeArrayView(Scripts), ScriptDefinedDataInterfaceParameters);
-
-		//Bind some stores and unbind immediately just to prime some data from those stores.
-		FNiagaraParameterStore& SystemScriptDefinedDataInterfaceParameters = ParentSystemInstance->GetSystemSimulation()->GetScriptDefinedDataInterfaceParameters();
-		
-		SystemScriptDefinedDataInterfaceParameters.Bind(&SpawnExecContext.Parameters);
-		ScriptDefinedDataInterfaceParameters.Bind(&SpawnExecContext.Parameters);
-		SpawnExecContext.Parameters.UnbindFromSourceStores();
-
-		SystemScriptDefinedDataInterfaceParameters.Bind(&UpdateExecContext.Parameters);
-		ScriptDefinedDataInterfaceParameters.Bind(&UpdateExecContext.Parameters);
-		UpdateExecContext.Parameters.UnbindFromSourceStores();
-
-		if (GPUExecContext)
-		{
-			SystemScriptDefinedDataInterfaceParameters.Bind(&GPUExecContext->CombinedParamStore);
-			ScriptDefinedDataInterfaceParameters.Bind(&GPUExecContext->CombinedParamStore);
-			GPUExecContext->CombinedParamStore.UnbindFromSourceStores();
-		}
 
 		if (EventInstanceData.IsValid())
 		{
-			for (FNiagaraScriptExecutionContext& EventContext : EventInstanceData->EventExecContexts)
-			{
-				SystemScriptDefinedDataInterfaceParameters.Bind(&EventContext.Parameters);
-				ScriptDefinedDataInterfaceParameters.Bind(&EventContext.Parameters);
-				EventContext.Parameters.UnbindFromSourceStores();
-			}
-
 			const int32 NumEventHandlers = EmitterData->GetEventHandlers().Num();
 			EventInstanceData->EventHandlingInfo.Reset();
 			EventInstanceData->EventHandlingInfo.SetNum(NumEventHandlers);
@@ -504,9 +478,6 @@ void FNiagaraEmitterInstance::Init(int32 InEmitterIdx, FNiagaraSystemInstanceID
 			if (ParentSystemInstance)
 				ParentSystemInstance->GetInstanceParameters().Bind(&RendererBindings);
 
-			SystemScriptDefinedDataInterfaceParameters.Bind(&RendererBindings);
-			ScriptDefinedDataInterfaceParameters.Bind(&RendererBindings);
-
 			if (GPUExecContext && EmitterData->SimTarget == ENiagaraSimTarget::GPUComputeSim)
 			{
 				GPUExecContext->CombinedParamStore.Bind(&RendererBindings);
diff --git a/Engine/Plugins/FX/Niagara/Source/Niagara/Private/NiagaraResolveDIHelpers.cpp b/Engine/Plugins/FX/Niagara/Source/Niagara/Private/NiagaraResolveDIHelpers.cpp
new file mode 100644
--- /dev/null
+++ b/Engine/Plugins/FX/Niagara/Source/Niagara/Private/NiagaraResolveDIHelpers.cpp
@@ -0,0 +1,359 @@
+// Copyright Epic Games, Inc. All Rights Reserved.
+
+//#if WITH_EDITORONLY_DATA
+
+#include "NiagaraResolveDIHelpers.h"
+
+#include "Internationalization/Internationalization.h"
+#include "NiagaraCommon.h"
+#include "NiagaraConstants.h"
+#include "NiagaraEmitter.h"
+#include "NiagaraScript.h"
+#include "NiagaraSimulationStageBase.h"
+#include "NiagaraSystem.h"
+
+#define LOCTEXT_NAMESPACE "NiagaraResolveDIHelpers"
+
+namespace FNiagaraResolveDIHelpers
+{
+	void CollectDIBindingsAndAssignmentsForScript(
+		const UNiagaraScript* TargetScript,
+		const UNiagaraSystem* System,
+		const FString& EmitterName,
+		TMap<FNiagaraVariableBase, UNiagaraDataInterface*>& OutVariableAssignmentMap,
+		TMap<FNiagaraVariableBase, FNiagaraVariableBase>& OutVariableBindingMap,
+		TArray<FText>& OutErrorMessages)
+	{
+		if (TargetScript == nullptr)
+		{
+			return;
+		}
+
+		for (const FNiagaraScriptDataInterfaceInfo& CachedDefaultDataInterface : TargetScript->GetCachedDefaultDataInterfaces())
+		{
+			for (FName RegisteredParameterMapWrite : CachedDefaultDataInterface.RegisteredParameterMapWrites)
+			{
+				FNiagaraVariable WriteVariable(CachedDefaultDataInterface.Type, RegisteredParameterMapWrite);
+				if (EmitterName.IsEmpty() == false)
+				{
+					WriteVariable = FNiagaraUtilities::ResolveAliases(WriteVariable, FNiagaraAliasContext().ChangeEmitterNameToEmitter(EmitterName));
+				}
+
+				if (CachedDefaultDataInterface.RegisteredParameterMapRead != NAME_None)
+				{
+					// Handle bindings.
+					FNiagaraVariable ReadVariable(CachedDefaultDataInterface.Type, CachedDefaultDataInterface.RegisteredParameterMapRead);
+					if (EmitterName.IsEmpty() == false)
+					{
+						ReadVariable = FNiagaraUtilities::ResolveAliases(ReadVariable, FNiagaraAliasContext().ChangeEmitterToEmitterName(EmitterName));
+					}
+
+					if (ReadVariable != WriteVariable)
+					{
+						FNiagaraVariableBase* CurrentBinding = OutVariableBindingMap.Find(WriteVariable);
+						if (CurrentBinding != nullptr)
+						{
+							if (*CurrentBinding != ReadVariable)
+							{
+								OutErrorMessages.Add(FText::Format(
+									LOCTEXT("MultipleBindingsFormat", "A data interface parameter was the target of multiple bindings in a single system.  The data interface used in the simulation may be incorrect.  Target Parameter: {0} First Read: {1} Current Read: {2}"),
+									FText::FromName(WriteVariable.GetName()),
+									FText::FromName(CurrentBinding->GetName()),
+									FText::FromName(ReadVariable.GetName())));
+							}
+						}
+						else
+						{
+							{
+								OutVariableBindingMap.Add(WriteVariable, ReadVariable);
+							}
+						}
+					}
+				}
+				else
+				{
+					// Handle assignments
+					UNiagaraDataInterface** CurrentAssignmentPtr = OutVariableAssignmentMap.Find(WriteVariable);
+					if (CurrentAssignmentPtr != nullptr)
+					{
+						if (*CurrentAssignmentPtr != CachedDefaultDataInterface.DataInterface)
+						{
+							OutErrorMessages.Add(FText::Format(
+								LOCTEXT("MultipleAssignmentsFormat", "A data interface parameter was thh target of an assignment multiple times in a single system.  The data interface used in the simulation may be incorrect.  Target Parameter: {0} First Assignment: {1} Current Assignment: {2}"),
+								FText::FromName(WriteVariable.GetName()),
+								FText::FromString((*CurrentAssignmentPtr)->GetName()),
+								FText::FromString((CachedDefaultDataInterface.DataInterface->GetName()))));
+						}
+					}
+					else
+					{
+						OutVariableAssignmentMap.Add(WriteVariable, CachedDefaultDataInterface.DataInterface);
+					}
+				}
+			}
+		}
+	}
+
+	void CollectDIBindingsAndAssignments(
+		const UNiagaraSystem* System,
+		TMap<FGuid, TMap<FNiagaraVariableBase, UNiagaraDataInterface*>>& OutEmitterIdToVariableAssignmentsMap,
+		TMap<FGuid, TMap<FNiagaraVariableBase, FNiagaraVariableBase>>& OutEmitterIdToVariableBindingsMap,
+		TArray<FText>& OutErrorMessages)
+	{
+		TMap<FNiagaraVariableBase, UNiagaraDataInterface*> VariableAssignmentMap;
+		TMap<FNiagaraVariableBase, FNiagaraVariableBase> VariableBindingMap;
+
+		// Add user parameters as assignments.
+		for (const FNiagaraVariableWithOffset& UserParameterWithOffset : System->GetExposedParameters().ReadParameterVariables())
+		{
+			if (UserParameterWithOffset.IsDataInterface())
+			{
+				UNiagaraDataInterface* DataInterface = System->GetExposedParameters().GetDataInterface(UserParameterWithOffset.Offset);
+				if (DataInterface != nullptr)
+				{
+					VariableAssignmentMap.Add(UserParameterWithOffset, DataInterface);
+				}
+			}
+		}
+
+		// Collect system and emitter scripts.
+		CollectDIBindingsAndAssignmentsForScript(System->GetSystemSpawnScript(), System, FString(), VariableAssignmentMap, VariableBindingMap, OutErrorMessages);
+		CollectDIBindingsAndAssignmentsForScript(System->GetSystemUpdateScript(), System, FString(), VariableAssignmentMap, VariableBindingMap, OutErrorMessages);
+		OutEmitterIdToVariableAssignmentsMap.Add(FGuid(), VariableAssignmentMap);
+		OutEmitterIdToVariableBindingsMap.Add(FGuid(), VariableBindingMap);
+
+		// Collect emitter scripts.
+		for (const FNiagaraEmitterHandle& EmitterHandle : System->GetEmitterHandles())
+		{
+			if (EmitterHandle.GetIsEnabled() == false)
+			{
+				continue;
+			}
+
+			FVersionedNiagaraEmitterData* VersionedNiagaraEmitterData = EmitterHandle.GetEmitterData();
+			if (VersionedNiagaraEmitterData == nullptr)
+			{
+				continue;
+			}
+
+			TMap<FNiagaraVariableBase, UNiagaraDataInterface*> EmitterVariableAssignmentMap = VariableAssignmentMap;
+			TMap<FNiagaraVariableBase, FNiagaraVariableBase> EmitterVariableBindingMap = VariableBindingMap;
+
+			CollectDIBindingsAndAssignmentsForScript(VersionedNiagaraEmitterData->SpawnScriptProps.Script, System, EmitterHandle.GetUniqueInstanceName(), EmitterVariableAssignmentMap, EmitterVariableBindingMap, OutErrorMessages);
+			CollectDIBindingsAndAssignmentsForScript(VersionedNiagaraEmitterData->UpdateScriptProps.Script, System, EmitterHandle.GetUniqueInstanceName(), EmitterVariableAssignmentMap, EmitterVariableBindingMap, OutErrorMessages);
+
+			for (const FNiagaraEventScriptProperties& EventHandler : VersionedNiagaraEmitterData->GetEventHandlers())
+			{
+				CollectDIBindingsAndAssignmentsForScript(EventHandler.Script, System, EmitterHandle.GetUniqueInstanceName(), EmitterVariableAssignmentMap, EmitterVariableBindingMap, OutErrorMessages);
+			}
+
+			for (const UNiagaraSimulationStageBase* SimulationStage : VersionedNiagaraEmitterData->GetSimulationStages())
+			{
+				if (SimulationStage == nullptr || SimulationStage->bEnabled == false)
+				{
+					continue;
+				}
+				CollectDIBindingsAndAssignmentsForScript(SimulationStage->Script, System, EmitterHandle.GetUniqueInstanceName(), EmitterVariableAssignmentMap, EmitterVariableBindingMap, OutErrorMessages);
+			}
+
+			OutEmitterIdToVariableAssignmentsMap.Add(EmitterHandle.GetId(), EmitterVariableAssignmentMap);
+			OutEmitterIdToVariableBindingsMap.Add(EmitterHandle.GetId(), EmitterVariableBindingMap);
+		}
+	}
+
+	void ResolveDIsForScript(
+		const UNiagaraSystem* System,
+		UNiagaraScript* TargetScript,
+		const FString& EmitterName,
+		TMap<FNiagaraVariableBase, UNiagaraDataInterface*>& VariableAssignmentMap,
+		TMap<FNiagaraVariableBase, FNiagaraVariableBase>& VariableBindingMap,
+		TArray<FText>& OutErrorMessages)
+	{
+		TArray<FNiagaraScriptResolvedDataInterfaceInfo> ResolvedDataInterfaces;
+		ResolvedDataInterfaces.Reserve(TargetScript->GetCachedDefaultDataInterfaces().Num());
+		TArray<FNiagaraResolvedUserDataInterfaceBinding> UserDataInterfaceBindings;
+		int32 ResolvedDataInterfaceIndex = 0;
+		for (const FNiagaraScriptDataInterfaceInfo& CachedDefaultDataInterface : TargetScript->GetCachedDefaultDataInterfaces())
+		{
+			FNiagaraScriptResolvedDataInterfaceInfo& ResolvedDataInterface = ResolvedDataInterfaces.AddDefaulted_GetRef();
+			ResolvedDataInterface.Name = CachedDefaultDataInterface.Name;
+			ResolvedDataInterface.CompileName = CachedDefaultDataInterface.CompileName;
+			ResolvedDataInterface.EmitterName = CachedDefaultDataInterface.EmitterName;
+			ResolvedDataInterface.UserPtrIdx = CachedDefaultDataInterface.UserPtrIdx;
+
+			if (CachedDefaultDataInterface.RegisteredParameterMapRead != NAME_None)
+			{
+				// If this DI is read from a parameter try to resolve it through the binding and assignment maps.
+				FNiagaraVariable ReadVariable(CachedDefaultDataInterface.Type, CachedDefaultDataInterface.RegisteredParameterMapRead);
+				if (EmitterName.IsEmpty() == false)
+				{
+					ReadVariable = FNiagaraUtilities::ResolveAliases(ReadVariable, FNiagaraAliasContext().ChangeEmitterToEmitterName(EmitterName));
+				}
+
+				FNiagaraVariableBase BoundVariable;
+				FNiagaraVariableBase* CurrentBoundVariable = &ReadVariable;
+				TSet<FNiagaraVariableBase> SeenBoundVariables;
+				bool bCircularReferenceDetected = false;
+				while (CurrentBoundVariable != nullptr && bCircularReferenceDetected == false)
+				{
+					SeenBoundVariables.Add(*CurrentBoundVariable);
+					BoundVariable = *CurrentBoundVariable;
+					CurrentBoundVariable = VariableBindingMap.Find(BoundVariable);
+
+					if (CurrentBoundVariable != nullptr && SeenBoundVariables.Contains(*CurrentBoundVariable))
+					{
+						bCircularReferenceDetected = true;
+						OutErrorMessages.Add(FText::Format(
+							LOCTEXT("CircularDependencyWarningFormat", "A data interface parameter circular dependency found.  The data interface used in the simulation may be incorrect.  Target Parameter: {0} Resolved Parameter: {1}"),
+							FText::FromName(ReadVariable.GetName()),
+							FText::FromName(CurrentBoundVariable->GetName())));
+					}
+				}
+
+				UNiagaraDataInterface** BoundDataInterfacePtr = VariableAssignmentMap.Find(BoundVariable);
+				if (BoundDataInterfacePtr != nullptr)
+				{
+					ResolvedDataInterface.ResolvedVariable = BoundVariable;
+					ResolvedDataInterface.SourceVariable = ReadVariable;
+					ResolvedDataInterface.bIsInternal = false;
+					ResolvedDataInterface.ResolvedDataInterface = *BoundDataInterfacePtr;
+				}
+
+				if (FNiagaraUserRedirectionParameterStore::IsUserParameter(BoundVariable))
+				{
+					const int32* UserParameterIndex = System->GetExposedParameters().FindParameterOffset(BoundVariable);
+					if (UserParameterIndex != nullptr && *UserParameterIndex != INDEX_NONE)
+					{
+						UserDataInterfaceBindings.Add(FNiagaraResolvedUserDataInterfaceBinding(*UserParameterIndex, ResolvedDataInterfaceIndex));
+					}
+				}
+			}
+
+			if (ResolvedDataInterface.ResolvedDataInterface == nullptr)
+			{
+				// If the DI was not read from a parameter or couldn't be found, use the one cached during compilation, and give it an internal
+				// name to prevent it from being bound incorrectly.
+				FNameBuilder NameBuilder;
+				NameBuilder.Append(FNiagaraConstants::InternalNamespaceString);
+				NameBuilder.AppendChar(TEXT('.'));
+				ResolvedDataInterface.Name.AppendString(NameBuilder);
+
+				FNiagaraVariable InternalVariable(CachedDefaultDataInterface.Type, FName(NameBuilder.ToString()));
+				ResolvedDataInterface.ResolvedVariable = InternalVariable;
+				ResolvedDataInterface.SourceVariable = InternalVariable;
+				ResolvedDataInterface.bIsInternal = true;
+				ResolvedDataInterface.ResolvedDataInterface = CachedDefaultDataInterface.DataInterface;
+			}
+
+			ResolvedDataInterfaceIndex++;
+		}
+		TargetScript->SetResolvedDataInterfaces(ResolvedDataInterfaces);
+		TargetScript->SetResolvedUserDataInterfaceBindings(UserDataInterfaceBindings);
+	}
+
+	void SynchronizeMatchingInternalResolvedDataInterfaces(TArrayView<const FNiagaraScriptResolvedDataInterfaceInfo> SourceResolvedDataInterfaces, TArray<FNiagaraScriptResolvedDataInterfaceInfo>& TargetResolvedDataInterfaces)
+	{
+		for (const FNiagaraScriptResolvedDataInterfaceInfo& SourceResolvedDataInterface : SourceResolvedDataInterfaces)
+		{
+			if (SourceResolvedDataInterface.bIsInternal)
+			{
+				for (FNiagaraScriptResolvedDataInterfaceInfo& TargetResolvedDataInterface : TargetResolvedDataInterfaces)
+				{
+					if (TargetResolvedDataInterface.bIsInternal && TargetResolvedDataInterface.ResolvedVariable == SourceResolvedDataInterface.ResolvedVariable)
+					{
+						TargetResolvedDataInterface.ResolvedDataInterface = SourceResolvedDataInterface.ResolvedDataInterface;
+						break;
+					}
+				}
+			}
+		}
+	}
+
+	/** Handles the special case where internal data interfaces defined in particle update need to be copied to particle spawn so that they're using the same instance. */
+	void ResolveInternalDIsForInterpolatedSpawning(UNiagaraScript* ParticleSpawnScript, UNiagaraScript* ParticleUpdateScript)
+	{
+		TArray<FNiagaraScriptResolvedDataInterfaceInfo> ResolvedSpawnDataInterfaces;
+		ResolvedSpawnDataInterfaces.Append(ParticleSpawnScript->GetResolvedDataInterfaces());
+
+		TArrayView<const FNiagaraScriptResolvedDataInterfaceInfo> ResolvedUpdateDataInterfaces = ParticleUpdateScript->GetResolvedDataInterfaces();
+		SynchronizeMatchingInternalResolvedDataInterfaces(ResolvedUpdateDataInterfaces, ResolvedSpawnDataInterfaces);
+
+		ParticleSpawnScript->SetResolvedDataInterfaces(ResolvedSpawnDataInterfaces);
+	}
+
+	/** Handles the special case where internal data interfaces defined in particle spawn or particle update need to be copied to the gpu script so that they're using the same instance. */
+	void ResolveInternalDIsForGpuScripts(UNiagaraScript* ParticleSpawnScript, UNiagaraScript* ParticleUpdateScript, UNiagaraScript* ParticleGpuScript)
+	{
+		TArray<FNiagaraScriptResolvedDataInterfaceInfo> ResolvedGpuDataInterfaces;
+		ResolvedGpuDataInterfaces.Append(ParticleGpuScript->GetResolvedDataInterfaces());
+
+		TArrayView<const FNiagaraScriptResolvedDataInterfaceInfo> ResolvedSpawnDataInterfaces = ParticleSpawnScript->GetResolvedDataInterfaces();
+		SynchronizeMatchingInternalResolvedDataInterfaces(ResolvedSpawnDataInterfaces, ResolvedGpuDataInterfaces);
+
+		TArrayView<const FNiagaraScriptResolvedDataInterfaceInfo> ResolvedUpdateDataInterfaces = ParticleUpdateScript->GetResolvedDataInterfaces();
+		SynchronizeMatchingInternalResolvedDataInterfaces(ResolvedUpdateDataInterfaces, ResolvedGpuDataInterfaces);
+
+		ParticleGpuScript->SetResolvedDataInterfaces(ResolvedGpuDataInterfaces);
+	}
+
+	void ResolveDIs(
+		UNiagaraSystem* System,
+		const TMap<FGuid, TMap<FNiagaraVariableBase, UNiagaraDataInterface*>>& EmitterIdToVariableAssignmentsMap,
+		const TMap<FGuid, TMap<FNiagaraVariableBase, FNiagaraVariableBase>>& EmitterIdToVariableBindingsMap,
+		TArray<FText>& OutErrorMessages)
+	{
+		TMap<FNiagaraVariableBase, UNiagaraDataInterface*> VariableAssignmentMap = EmitterIdToVariableAssignmentsMap[FGuid()];
+		TMap<FNiagaraVariableBase, FNiagaraVariableBase> VariableBindingMap = EmitterIdToVariableBindingsMap[FGuid()];
+		ResolveDIsForScript(System, System->GetSystemSpawnScript(), FString(), VariableAssignmentMap, VariableBindingMap, OutErrorMessages);
+		ResolveDIsForScript(System, System->GetSystemUpdateScript(), FString(), VariableAssignmentMap, VariableBindingMap, OutErrorMessages);
+
+		for (const FNiagaraEmitterHandle& EmitterHandle : System->GetEmitterHandles())
+		{
+			if (EmitterHandle.GetIsEnabled() == false)
+			{
+				continue;
+			}
+
+			FVersionedNiagaraEmitterData* VersionedNiagaraEmitterData = EmitterHandle.GetEmitterData();
+			if (VersionedNiagaraEmitterData == nullptr)
+			{
+				continue;
+			}
+
+			TMap<FNiagaraVariableBase, UNiagaraDataInterface*> EmitterVariableAssignmentMap = EmitterIdToVariableAssignmentsMap[EmitterHandle.GetId()];
+			TMap<FNiagaraVariableBase, FNiagaraVariableBase> EmitterVariableBindingMap = EmitterIdToVariableBindingsMap[EmitterHandle.GetId()];
+
+			ResolveDIsForScript(System, VersionedNiagaraEmitterData->SpawnScriptProps.Script, EmitterHandle.GetUniqueInstanceName(), EmitterVariableAssignmentMap, EmitterVariableBindingMap, OutErrorMessages);
+			ResolveDIsForScript(System, VersionedNiagaraEmitterData->UpdateScriptProps.Script, EmitterHandle.GetUniqueInstanceName(), EmitterVariableAssignmentMap, EmitterVariableBindingMap, OutErrorMessages);
+			ResolveDIsForScript(System, VersionedNiagaraEmitterData->GetGPUComputeScript(), EmitterHandle.GetUniqueInstanceName(), EmitterVariableAssignmentMap, EmitterVariableBindingMap, OutErrorMessages);
+
+			if (VersionedNiagaraEmitterData->bInterpolatedSpawning)
+			{
+				ResolveInternalDIsForInterpolatedSpawning(VersionedNiagaraEmitterData->SpawnScriptProps.Script, VersionedNiagaraEmitterData->UpdateScriptProps.Script);
+			}
+
+			if (VersionedNiagaraEmitterData->SimTarget == ENiagaraSimTarget::GPUComputeSim)
+			{
+				ResolveInternalDIsForGpuScripts(VersionedNiagaraEmitterData->SpawnScriptProps.Script, VersionedNiagaraEmitterData->EmitterUpdateScriptProps.Script, VersionedNiagaraEmitterData->GetGPUComputeScript());
+			}
+
+			for (const FNiagaraEventScriptProperties& EventHandler : VersionedNiagaraEmitterData->GetEventHandlers())
+			{
+				ResolveDIsForScript(System, EventHandler.Script, EmitterHandle.GetUniqueInstanceName(), EmitterVariableAssignmentMap, EmitterVariableBindingMap, OutErrorMessages);
+			}
+
+			for (const UNiagaraSimulationStageBase* SimulationStage : VersionedNiagaraEmitterData->GetSimulationStages())
+			{
+				if (SimulationStage == nullptr || SimulationStage->bEnabled == false)
+				{
+					continue;
+				}
+				ResolveDIsForScript(System, SimulationStage->Script, EmitterHandle.GetUniqueInstanceName(), EmitterVariableAssignmentMap, EmitterVariableBindingMap, OutErrorMessages);
+			}
+		}
+	}
+}
+
+#undef LOCTEXT_NAMESPACE
+
+//#endif // WITH_EDITORONLY_DATA
\ No newline at end of file
diff --git a/Engine/Plugins/FX/Niagara/Source/Niagara/Private/NiagaraResolveDIHelpers.h b/Engine/Plugins/FX/Niagara/Source/Niagara/Private/NiagaraResolveDIHelpers.h
new file mode 100644
--- /dev/null
+++ b/Engine/Plugins/FX/Niagara/Source/Niagara/Private/NiagaraResolveDIHelpers.h
@@ -0,0 +1,27 @@
+// Copyright Epic Games, Inc. All Rights Reserved.
+
+#pragma once
+
+#if WITH_EDITORONLY_DATA
+
+#include "NiagaraTypes.h"
+
+class UNiagaraDataInterface;
+class UNiagaraSystem;
+
+namespace FNiagaraResolveDIHelpers
+{
+	void CollectDIBindingsAndAssignments(
+		const UNiagaraSystem* System,
+		TMap<FGuid, TMap<FNiagaraVariableBase, UNiagaraDataInterface*>>& OutEmitterIdToVariableAssignmentsMap,
+		TMap<FGuid, TMap<FNiagaraVariableBase, FNiagaraVariableBase>>& OutEmitterIdToVariableBindingsMap,
+		TArray<FText>& OutErrorMessages);
+
+	void ResolveDIs(
+		UNiagaraSystem* System,
+		const TMap<FGuid, TMap<FNiagaraVariableBase, UNiagaraDataInterface*>>& EmitterIdToVariableAssignmentsMap,
+		const TMap<FGuid, TMap<FNiagaraVariableBase, FNiagaraVariableBase>>& EmitterIdToVariableBindingsMap,
+		TArray<FText>& OutErrorMessages);
+}
+
+#endif
\ No newline at end of file
diff --git a/Engine/Plugins/FX/Niagara/Source/Niagara/Private/NiagaraScript.cpp b/Engine/Plugins/FX/Niagara/Source/Niagara/Private/NiagaraScript.cpp
--- a/Engine/Plugins/FX/Niagara/Source/Niagara/Private/NiagaraScript.cpp
+++ b/Engine/Plugins/FX/Niagara/Source/Niagara/Private/NiagaraScript.cpp
@@ -2139,6 +2139,12 @@ void UNiagaraScript::PostLoad()
 		GenerateDefaultFunctionBindings();
 	}
 
+	for (FNiagaraScriptResolvedDataInterfaceInfo& ResolvedDataInterface : ResolvedDataInterfaces)
+	{
+		ResolvedDataInterface.ResolvedDataInterface->ConditionalPostLoad();
+	}
+
+#if WITH_EDITORONLY_DATA
 	// Because we might be using these cached data interfaces, we need to make sure that they are properly postloaded.
 	for (FNiagaraScriptDataInterfaceInfo& Info : CachedDefaultDataInterfaces)
 	{
@@ -2148,7 +2154,6 @@ void UNiagaraScript::PostLoad()
 		}
 	}
 
-#if WITH_EDITORONLY_DATA
 	FVersionedNiagaraScriptData* ScriptData = GetLatestScriptData();
 	ensure(ScriptData);
 	if (NiagaraVer < FNiagaraCustomVersion::AddSimulationStageUsageEnum)
@@ -2817,9 +2822,14 @@ void UNiagaraScript::SetVMCompilationResults(const FNiagaraVMExecutableDataId& I
 		int32 Idx = CachedDefaultDataInterfaces.AddDefaulted();
 		CachedDefaultDataInterfaces[Idx].UserPtrIdx = Info.UserPtrIdx;
 		CachedDefaultDataInterfaces[Idx].Name = ResolveEmitterAlias(Info.Name, EmitterUniqueName);
+		CachedDefaultDataInterfaces[Idx].CompileName = Info.Name;
 		CachedDefaultDataInterfaces[Idx].Type = Info.Type;
+		CachedDefaultDataInterfaces[Idx].EmitterName = EmitterUniqueName;
 		CachedDefaultDataInterfaces[Idx].RegisteredParameterMapRead = ResolveEmitterAlias(Info.RegisteredParameterMapRead, EmitterUniqueName);
-		CachedDefaultDataInterfaces[Idx].RegisteredParameterMapWrite = ResolveEmitterAlias(Info.RegisteredParameterMapWrite, EmitterUniqueName);
+		for (const FName& RegisteredParameterMapWrite : Info.RegisteredParameterMapWrites)
+		{
+			CachedDefaultDataInterfaces[Idx].RegisteredParameterMapWrites.Add(ResolveEmitterAlias(RegisteredParameterMapWrite, EmitterUniqueName));
+		}
 
 		// We compiled it just a bit ago, so we should be able to resolve it from the table that we passed in.
 		UNiagaraDataInterface* FindDIById = ResolveDataInterface(ObjectNameMap, CachedDefaultDataInterfaces[Idx].Name);
@@ -3624,7 +3634,10 @@ void UNiagaraScript::SyncAliases(const FNiagaraAliasContext& ResolveAliasesConte
 
 		// also update the MapRead/MapWrite member variables, they seem to typically match the Name parameter, but may be None
 		UpdateName(DataInterfaceInfo.RegisteredParameterMapRead);
-		UpdateName(DataInterfaceInfo.RegisteredParameterMapWrite);
+		for (FName& RegisteredParameterMapWrite : DataInterfaceInfo.RegisteredParameterMapWrites)
+		{
+			UpdateName(RegisteredParameterMapWrite);
+		}
 	}
 }
 
diff --git a/Engine/Plugins/FX/Niagara/Source/Niagara/Private/NiagaraScriptExecutionContext.cpp b/Engine/Plugins/FX/Niagara/Source/Niagara/Private/NiagaraScriptExecutionContext.cpp
--- a/Engine/Plugins/FX/Niagara/Source/Niagara/Private/NiagaraScriptExecutionContext.cpp
+++ b/Engine/Plugins/FX/Niagara/Source/Niagara/Private/NiagaraScriptExecutionContext.cpp
@@ -810,10 +810,11 @@ bool FNiagaraSystemScriptExecutionContext::GeneratePerInstanceDIFunctionTable(FN
 		for (int32 i = 0; i < ScriptExecutableData.DataInterfaceInfo.Num(); i++)
 		{
 			const FNiagaraScriptDataInterfaceCompileInfo& ScriptDIInfo = ScriptExecutableData.DataInterfaceInfo[i];
+			const FNiagaraScriptResolvedDataInterfaceInfo& ResolvedDIInfo = Script->GetResolvedDataInterfaces()[i];
 			//UNiagaraDataInterface* ScriptInterface = ScriptDataInterfaces[i];
 			UNiagaraDataInterface* ExternalInterface = GetDataInterfaces()[i];
 
-			if (ScriptDIInfo.Name == BindingInfo.OwnerName && ScriptDIInfo.NeedsPerInstanceBinding())
+			if (ScriptDIInfo.Name == BindingInfo.OwnerName && (ScriptDIInfo.NeedsPerInstanceBinding() || ResolvedDIInfo.NeedsPerInstanceBinding()))
 			{
 				UNiagaraDataInterface* DIToBind = nullptr;
 				FNiagaraPerInstanceDIFuncInfo& NewFuncInfo = OutFunctions.AddDefaulted_GetRef();
@@ -828,6 +829,15 @@ bool FNiagaraSystemScriptExecutionContext::GeneratePerInstanceDIFunctionTable(FN
 						InstData = Inst->FindDataInterfaceInstanceData(UserInterface);
 					}
 				}
+				else if (const int32* ResolvedDIIndex = Inst->GetInstanceParameters().FindParameterOffset(ResolvedDIInfo.ResolvedVariable))
+				{
+					//If this is a User DI we bind to the user DI and find instance data with it.
+					if (UNiagaraDataInterface* UserInterface = Inst->GetInstanceParameters().GetDataInterface(*ResolvedDIIndex))
+					{
+						DIToBind = UserInterface;
+						InstData = Inst->FindDataInterfaceInstanceData(UserInterface);
+					}
+				}
 				else
 				{
 					//Otherwise we use the script DI and search for instance data with that.
diff --git a/Engine/Plugins/FX/Niagara/Source/Niagara/Private/NiagaraScriptExecutionParameterStore.cpp b/Engine/Plugins/FX/Niagara/Source/Niagara/Private/NiagaraScriptExecutionParameterStore.cpp
--- a/Engine/Plugins/FX/Niagara/Source/Niagara/Private/NiagaraScriptExecutionParameterStore.cpp
+++ b/Engine/Plugins/FX/Niagara/Source/Niagara/Private/NiagaraScriptExecutionParameterStore.cpp
@@ -311,28 +311,11 @@ void FNiagaraScriptExecutionParameterStore::AddScriptParams(UNiagaraScript* Scri
 	}
 
 	check(Script->GetVMExecutableData().DataInterfaceInfo.Num() == Script->GetCachedDefaultDataInterfaces().Num());
-	for (FNiagaraScriptDataInterfaceInfo& Info : Script->GetCachedDefaultDataInterfaces())
+	for (const FNiagaraScriptResolvedDataInterfaceInfo& ResolvedDataInterface : Script->GetResolvedDataInterfaces())
 	{
-		FName ParameterName;
-		if (Info.RegisteredParameterMapRead != NAME_None)
-		{
-			ParameterName = Info.RegisteredParameterMapRead;
-		}
-		else
-		{
-			// If the data interface wasn't used in a parameter map, mangle the name so that it doesn't accidentally bind to
-			// a valid parameter.
-			FNameBuilder NameBuilder;
-			NameBuilder.Append(FNiagaraConstants::InternalNamespaceString);
-			NameBuilder.AppendChar(TEXT('.'));
-			Info.Name.AppendString(NameBuilder);
-			ParameterName = FName(NameBuilder);
-		}
-
-		FNiagaraVariable Var(Info.Type, ParameterName);
 		int32 VarOffset = INDEX_NONE;
-		bAdded |= AddParameter(Var, false, false, &VarOffset);
-		SetDataInterface(Info.DataInterface, VarOffset);
+		bAdded |= AddParameter(ResolvedDataInterface.SourceVariable, false, false, &VarOffset);
+		SetDataInterface(ResolvedDataInterface.ResolvedDataInterface, VarOffset);
 	}
 
 	if (bAdded && bTriggerRebind)
diff --git a/Engine/Plugins/FX/Niagara/Source/Niagara/Private/NiagaraSystem.cpp b/Engine/Plugins/FX/Niagara/Source/Niagara/Private/NiagaraSystem.cpp
--- a/Engine/Plugins/FX/Niagara/Source/Niagara/Private/NiagaraSystem.cpp
+++ b/Engine/Plugins/FX/Niagara/Source/Niagara/Private/NiagaraSystem.cpp
@@ -18,6 +18,7 @@
 #include "NiagaraEmitterHandle.h"
 #include "NiagaraModule.h"
 #include "NiagaraRendererProperties.h"
+#include "NiagaraResolveDIHelpers.h"
 #include "NiagaraScratchPadContainer.h"
 #include "NiagaraScriptSourceBase.h"
 #include "NiagaraSettings.h"
@@ -29,6 +30,7 @@
 #include "NiagaraTypes.h"
 #include "NiagaraWorldManager.h"
 #include "Algo/RemoveIf.h"
+#include "Algo/StableSort.h"
 #include "HAL/LowLevelMemTracker.h"
 #include "Misc/ScopedSlowTask.h"
 #include "Misc/ScopeExit.h"
@@ -170,7 +172,6 @@ UNiagaraSystem::UNiagaraSystem(const FObjectInitializer& ObjectInitializer)
 , WarmupTime(0.0f)
 , WarmupTickCount(0)
 , WarmupTickDelta(1.0f / 15.0f)
-, bHasSystemScriptDIsWithPerInstanceData(false)
 , bNeedsGPUContextInitForDataInterfaces(false)
 , bNeedsAsyncOptimize(true)
 , CurrentScalabilitySettings(*new FNiagaraSystemScalabilitySettings())
@@ -1568,11 +1569,11 @@ void UNiagaraSystem::FindDataInterfaceDependencies(FVersionedNiagaraEmitterData*
 				{
 					if (Script)
 					{
-						for (FNiagaraScriptDataInterfaceInfo& DataInterfaceInfo : Script->GetCachedDefaultDataInterfaces())
+						for (const FNiagaraScriptResolvedDataInterfaceInfo& DataInterfaceInfo : Script->GetResolvedDataInterfaces())
 						{
-							if ((Variable.GetType() == DataInterfaceInfo.Type) && (Variable.GetName() == DataInterfaceInfo.RegisteredParameterMapWrite))
+							if (DataInterfaceInfo.ResolvedVariable == Variable)
 							{
-								return DataInterfaceInfo.DataInterface;
+								return DataInterfaceInfo.ResolvedDataInterface;
 							}
 						}
 					}
@@ -1750,7 +1751,7 @@ void UNiagaraSystem::ComputeEmittersExecutionOrder()
 		}
 
 		// Sort the emitter indices in the execution order array so that dependencies are satisfied.
-		Algo::Sort(EmitterExecutionOrder, [&EmitterPriorities](FNiagaraEmitterExecutionIndex IdxA, FNiagaraEmitterExecutionIndex IdxB) { return EmitterPriorities[IdxA.EmitterIndex] < EmitterPriorities[IdxB.EmitterIndex]; });
+		Algo::StableSort(EmitterExecutionOrder, [&EmitterPriorities](FNiagaraEmitterExecutionIndex IdxA, FNiagaraEmitterExecutionIndex IdxB) { return EmitterPriorities[IdxA.EmitterIndex] < EmitterPriorities[IdxB.EmitterIndex]; });
 
 		// Emitters with the same priority value can execute in parallel. Look for the emitters where the priority increases and mark them as needing to start a new
 		// overlap group. This informs the execution code about where to insert synchronization points to satisfy data dependencies.
@@ -1890,17 +1891,17 @@ void UNiagaraSystem::CacheFromCompiledData()
 
 				const bool bUsedByCPU = EmitterData->SimTarget == ENiagaraSimTarget::CPUSim;
 				const bool bUsedByGPU = EmitterData->SimTarget == ENiagaraSimTarget::GPUComputeSim;
-				for (const FNiagaraScriptDataInterfaceInfo& DataInterfaceInfo : NiagaraScript->GetCachedDefaultDataInterfaces())
+				for (const FNiagaraScriptResolvedDataInterfaceInfo& DataInterfaceInfo : NiagaraScript->GetResolvedDataInterfaces())
 				{
-					if (UNiagaraDataInterface* DataInterface = DataInterfaceInfo.DataInterface)
+					if (UNiagaraDataInterface* DataInterface = DataInterfaceInfo.ResolvedDataInterface)
 					{
-						DataInterface->CacheStaticBuffers(*StaticBuffers.Get(), DataInterfaceInfo, bUsedByCPU, bUsedByGPU);
+						DataInterface->CacheStaticBuffers(*StaticBuffers.Get(), DataInterfaceInfo.ResolvedVariable, bUsedByCPU, bUsedByGPU);
 
 						if ( bUsedByGPU )
 						{
-							if (!DataInterfaceInfo.RegisteredParameterMapRead.IsNone())
+							if (DataInterfaceInfo.bIsInternal == false)
 							{
-								DataInterfaceGpuUsage.Add(DataInterfaceInfo.RegisteredParameterMapRead);
+								DataInterfaceGpuUsage.Add(DataInterfaceInfo.ResolvedVariable.GetName());
 							}
 						}
 					}
@@ -1932,12 +1933,12 @@ void UNiagaraSystem::CacheFromCompiledData()
 			continue;
 		}
 
-		for (const FNiagaraScriptDataInterfaceInfo& DataInterfaceInfo : NiagaraScript->GetCachedDefaultDataInterfaces())
+ 		for (const FNiagaraScriptResolvedDataInterfaceInfo& DataInterfaceInfo : NiagaraScript->GetResolvedDataInterfaces())
 		{
-			if (UNiagaraDataInterface* DataInterface = DataInterfaceInfo.DataInterface)
+			if (UNiagaraDataInterface* DataInterface = DataInterfaceInfo.ResolvedDataInterface)
 			{
-				const bool bUsedByGPU = DataInterfaceGpuUsage.Contains(DataInterfaceInfo.RegisteredParameterMapWrite);
-				DataInterface->CacheStaticBuffers(*StaticBuffers.Get(), DataInterfaceInfo, true, bUsedByGPU);
+				const bool bUsedByGPU = DataInterfaceGpuUsage.Contains(DataInterfaceInfo.ResolvedVariable.GetName());
+				DataInterface->CacheStaticBuffers(*StaticBuffers.Get(), DataInterfaceInfo.ResolvedVariable, true, bUsedByGPU);
 			}
 		}
 	}
@@ -1946,47 +1947,54 @@ void UNiagaraSystem::CacheFromCompiledData()
 	StaticBuffers->Finalize();
 }
 
-bool UNiagaraSystem::HasSystemScriptDIsWithPerInstanceData() const
-{
-	return bHasSystemScriptDIsWithPerInstanceData;
-}
-
-const TArray<FName>& UNiagaraSystem::GetUserDINamesReadInSystemScripts() const
-{
-	return UserDINamesReadInSystemScripts;
-}
-
 FBox UNiagaraSystem::GetFixedBounds() const
 {
 	return FixedBounds;
 }
 
-void CheckDICompileInfo(const TArray<FNiagaraScriptDataInterfaceCompileInfo>& ScriptDICompileInfos, bool& bOutbHasSystemDIsWithPerInstanceData, TArray<FName>& OutUserDINamesReadInSystemScripts)
+#if WITH_EDITORONLY_DATA
+
+const FGuid UNiagaraSystem::ResolveDIsMessageId(0xB6ACDD97, 0xA2C04B02, 0xAB9FA4E4, 0xDBC73E51);
+
+void UNiagaraSystem::ResolveDIBindings()
 {
-	for (const FNiagaraScriptDataInterfaceCompileInfo& ScriptDICompileInfo : ScriptDICompileInfos)
-	{
-		UNiagaraDataInterface* DefaultDataInterface = ScriptDICompileInfo.GetDefaultDataInterface();
-		if (DefaultDataInterface != nullptr && DefaultDataInterface->PerInstanceDataSize() > 0)
-		{
-			bOutbHasSystemDIsWithPerInstanceData = true;
-		}
+	TMap<FGuid, TMap<FNiagaraVariableBase, UNiagaraDataInterface*>> EmitterIdToVariableAssignmentsMap;
+	TMap<FGuid, TMap<FNiagaraVariableBase, FNiagaraVariableBase>> EmitterIdToVariableBindingsMap;
+	TArray<FText> ErrorMessages;
+	FNiagaraResolveDIHelpers::CollectDIBindingsAndAssignments(this, EmitterIdToVariableAssignmentsMap, EmitterIdToVariableBindingsMap, ErrorMessages);
+	FNiagaraResolveDIHelpers::ResolveDIs(this, EmitterIdToVariableAssignmentsMap, EmitterIdToVariableBindingsMap, ErrorMessages);
 
-		if (ScriptDICompileInfo.RegisteredParameterMapRead.ToString().StartsWith(TEXT("User.")))
+	if (ErrorMessages.Num() > 0)
+	{
+		INiagaraModule& NiagaraModule = FModuleManager::GetModuleChecked<INiagaraModule>("Niagara");
+		bool bAllowDismissal = true;
+		UNiagaraMessageDataBase* ResolveDIsMessage = NiagaraModule.GetEditorOnlyDataUtilities().CreateWarningMessage(
+			this,
+			LOCTEXT("ResolveDIsError", "Resolving data interfaces generated errors"),
+			FText::Join(FText::FromString("\n"), ErrorMessages),
+			"Resolve Data Interfaces",
+			bAllowDismissal);
+		MessageStore.AddMessage(ResolveDIsMessageId, ResolveDIsMessage);
+		if (MessageStore.IsMessageDismissed(ResolveDIsMessageId) == false)
 		{
-			OutUserDINamesReadInSystemScripts.AddUnique(ScriptDICompileInfo.RegisteredParameterMapRead);
+			for (const FText& ErrorMessage : ErrorMessages)
+			{
+				UE_LOG(LogNiagara, Warning, TEXT("%s"), *ErrorMessage.ToString());
+			}
 		}
 	}
+	else
+	{
+		MessageStore.RemoveMessage(ResolveDIsMessageId);
+	}
 }
 
+#endif // WITH_EDITORONLY_DATA
+
 void UNiagaraSystem::UpdatePostCompileDIInfo()
 {
-	bHasSystemScriptDIsWithPerInstanceData = false;
-	UserDINamesReadInSystemScripts.Empty();
 	bNeedsGPUContextInitForDataInterfaces = false;
 
-	CheckDICompileInfo(SystemSpawnScript->GetVMExecutableData().DataInterfaceInfo, bHasSystemScriptDIsWithPerInstanceData, UserDINamesReadInSystemScripts);
-	CheckDICompileInfo(SystemUpdateScript->GetVMExecutableData().DataInterfaceInfo, bHasSystemScriptDIsWithPerInstanceData, UserDINamesReadInSystemScripts);
-
 	for (const FNiagaraEmitterHandle& EmitterHandle : GetEmitterHandles())
 	{
 		FVersionedNiagaraEmitterData* EmitterData = EmitterHandle.GetEmitterData();
@@ -2770,10 +2778,9 @@ bool UNiagaraSystem::QueryCompileComplete(bool bWait, bool bDoPost, bool bDoNotA
 		InitEmitterCompiledData();
 		InitSystemCompiledData();
 
-		CompileRequest.UpdateSpawnDataInterfaces();
-
 		CompileRequest.RootObjects.Empty();
 
+		ResolveDIBindings();
 		UpdatePostCompileDIInfo();
 		ComputeEmittersExecutionOrder();
 		ComputeRenderersDrawOrder();
@@ -3298,7 +3305,6 @@ void UNiagaraSystem::ResetToEmptySystem()
 	SystemUpdateScript = nullptr;
 	EmitterCompiledData.Empty();
 	SystemCompiledData = FNiagaraSystemCompiledData();
-	UserDINamesReadInSystemScripts.Empty();
 	EmitterExecutionOrder.Empty();
 	RendererPostTickOrder.Empty();
 	RendererCompletionOrder.Empty();
diff --git a/Engine/Plugins/FX/Niagara/Source/Niagara/Private/NiagaraSystemInstance.cpp b/Engine/Plugins/FX/Niagara/Source/Niagara/Private/NiagaraSystemInstance.cpp
--- a/Engine/Plugins/FX/Niagara/Source/Niagara/Private/NiagaraSystemInstance.cpp
+++ b/Engine/Plugins/FX/Niagara/Source/Niagara/Private/NiagaraSystemInstance.cpp
@@ -1406,6 +1406,8 @@ void FNiagaraSystemInstance::InitDataInterfaces()
 	PerInstanceDIFunctions[(int32)ENiagaraSystemSimulationScript::Spawn].Reset();
 	PerInstanceDIFunctions[(int32)ENiagaraSystemSimulationScript::Update].Reset();
 
+	ResolveUserDataInterfaceBindings();
+
 	//Now the interfaces in the simulations are all correct, we can build the per instance data table.
 	int32 InstanceDataSize = 0;
 	DataInterfaceInstanceDataOffsets.Empty();
@@ -1596,6 +1598,53 @@ void FNiagaraSystemInstance::InitDataInterfaces()
 	}
 }
 
+void FNiagaraSystemInstance::ResolveUserDataInterfaceBindings()
+{
+	const TArray<UNiagaraDataInterface*>& InstanceParameterDataInterfaces = InstanceParameters.GetDataInterfaces();
+	auto ResolveUserDIs = [&InstanceParameterDataInterfaces](FNiagaraParameterStore& TargetParameterStore, const UNiagaraScript* TargetScript)
+	{
+		TArrayView<const FNiagaraResolvedUserDataInterfaceBinding> ResolvedUserDataInterfaceBindings = TargetScript->GetResolvedUserDataInterfaceBindings();
+		const TArray<UNiagaraDataInterface*>& TargetDataInterfaces = TargetParameterStore.GetDataInterfaces();
+		for (const FNiagaraResolvedUserDataInterfaceBinding& ResolvedUserDIBinding : ResolvedUserDataInterfaceBindings)
+		{
+			if (ResolvedUserDIBinding.UserParameterStoreDataInterfaceIndex != INDEX_NONE &&
+				ResolvedUserDIBinding.UserParameterStoreDataInterfaceIndex < InstanceParameterDataInterfaces.Num() &&
+				ResolvedUserDIBinding.ScriptParameterStoreDataInterfaceIndex != INDEX_NONE &&
+				ResolvedUserDIBinding.ScriptParameterStoreDataInterfaceIndex < TargetDataInterfaces.Num())
+			{
+				TargetParameterStore.SetDataInterface(
+					InstanceParameterDataInterfaces[ResolvedUserDIBinding.UserParameterStoreDataInterfaceIndex],
+					ResolvedUserDIBinding.ScriptParameterStoreDataInterfaceIndex);
+			}
+		}
+	};
+	
+	for (const TSharedRef<FNiagaraEmitterInstance, ESPMode::ThreadSafe>& EmitterSimulation : Emitters)
+	{
+		FNiagaraEmitterInstance& EmitterSim = EmitterSimulation.Get();
+
+		if (EmitterSim.GetSpawnExecutionContext().Script != nullptr)
+		{
+			ResolveUserDIs(EmitterSim.GetSpawnExecutionContext().Parameters, EmitterSim.GetSpawnExecutionContext().Script);
+		}
+		if (EmitterSim.GetUpdateExecutionContext().Script != nullptr)
+		{
+			ResolveUserDIs(EmitterSim.GetUpdateExecutionContext().Parameters, EmitterSim.GetUpdateExecutionContext().Script);
+		}
+		if (EmitterSim.GetGPUContext() != nullptr && EmitterSim.GetGPUContext()->GPUScript != nullptr)
+		{
+			ResolveUserDIs(EmitterSim.GetGPUContext()->CombinedParamStore, EmitterSim.GetGPUContext()->GPUScript);
+		}
+		for (FNiagaraScriptExecutionContext& EventContext : EmitterSim.GetEventExecutionContexts())
+		{
+			if (EventContext.Script != nullptr)
+			{
+				ResolveUserDIs(EventContext.Parameters, EventContext.Script);
+			}
+		}
+	}
+}
+
 void FNiagaraSystemInstance::TickDataInterfaces(float DeltaSeconds, bool bPostSimulate)
 {
 	if (!GetSystem() || IsDisabled())
diff --git a/Engine/Plugins/FX/Niagara/Source/Niagara/Private/NiagaraSystemSimulation.cpp b/Engine/Plugins/FX/Niagara/Source/Niagara/Private/NiagaraSystemSimulation.cpp
--- a/Engine/Plugins/FX/Niagara/Source/Niagara/Private/NiagaraSystemSimulation.cpp
+++ b/Engine/Plugins/FX/Niagara/Source/Niagara/Private/NiagaraSystemSimulation.cpp
@@ -717,14 +717,6 @@ bool FNiagaraSystemSimulation::Init(UNiagaraSystem* InSystem, UWorld* InWorld, b
 				BindParameterColleciton(Collection, UpdateExecContext->Parameters);
 			}
 
-			TArray<UNiagaraScript*, TInlineAllocator<2>> Scripts;
-			Scripts.Add(SpawnScript);
-			Scripts.Add(UpdateScript);
-			FNiagaraUtilities::CollectScriptDataInterfaceParameters(*System, Scripts, ScriptDefinedDataInterfaceParameters);
-
-			ScriptDefinedDataInterfaceParameters.Bind(&SpawnExecContext->Parameters);
-			ScriptDefinedDataInterfaceParameters.Bind(&UpdateExecContext->Parameters);
-
 			SpawnScript->RapidIterationParameters.Bind(&SpawnExecContext->Parameters);
 			UpdateScript->RapidIterationParameters.Bind(&UpdateExecContext->Parameters);
 		}
@@ -801,11 +793,6 @@ UNiagaraParameterCollectionInstance* FNiagaraSystemSimulation::GetParameterColle
 	return Ret;
 }
 
-FNiagaraParameterStore& FNiagaraSystemSimulation::GetScriptDefinedDataInterfaceParameters()
-{
-	return ScriptDefinedDataInterfaceParameters;
-}
-
 void FNiagaraSystemSimulation::TransferInstance(FNiagaraSystemInstance* Instance)
 {
 	check(Instance);
diff --git a/Engine/Plugins/FX/Niagara/Source/Niagara/Public/INiagaraEditorOnlyDataUtlities.h b/Engine/Plugins/FX/Niagara/Source/Niagara/Public/INiagaraEditorOnlyDataUtlities.h
--- a/Engine/Plugins/FX/Niagara/Source/Niagara/Public/INiagaraEditorOnlyDataUtlities.h
+++ b/Engine/Plugins/FX/Niagara/Source/Niagara/Public/INiagaraEditorOnlyDataUtlities.h
@@ -4,10 +4,12 @@
 
 #include "UObject/Object.h"
 
-class UNiagaraScriptSourceBase;
-class UNiagaraEditorParametersAdapterBase;
+class UNiagaraDataInterface;
 class UNiagaraEditorDataBase;
+class UNiagaraEditorParametersAdapterBase;
 class UNiagaraMessageDataBase;
+class UNiagaraScriptSourceBase;
+class UNiagaraSystem;
 
 /** Defines utility methods for creating editor only data which is stored on runtime objects. */
 class INiagaraEditorOnlyDataUtilities
@@ -24,4 +26,8 @@ public:
 	virtual UNiagaraMessageDataBase* CreateErrorMessage(UObject* InOuter, FText InMessageShort, FText InMessageLong, FName InTopicName, bool bInAllowDismissal = false) const = 0;
 
 	virtual UNiagaraMessageDataBase* CreateWarningMessage(UObject* InOuter, FText InMessageShort, FText InMessageLong, FName InTopicName, bool bInAllowDismissal = false) const = 0;
+
+	virtual bool IsEditorDataInterfaceInstance(const UNiagaraDataInterface* DataInterface) const = 0;
+
+	virtual UNiagaraDataInterface* GetResolvedRuntimeInstanceForEditorDataInterfaceInstance(const UNiagaraSystem& OwningSystem, UNiagaraDataInterface& EditorDataInterfaceInstance) const = 0;
 };
\ No newline at end of file
diff --git a/Engine/Plugins/FX/Niagara/Source/Niagara/Public/NiagaraCommon.h b/Engine/Plugins/FX/Niagara/Source/Niagara/Public/NiagaraCommon.h
--- a/Engine/Plugins/FX/Niagara/Source/Niagara/Public/NiagaraCommon.h
+++ b/Engine/Plugins/FX/Niagara/Source/Niagara/Public/NiagaraCommon.h
@@ -665,6 +665,9 @@ public:
 	UPROPERTY()
 	FName Name;
 
+	UPROPERTY()
+	FName CompileName;
+
 	/** Index of the user pointer for this data interface. */
 	UPROPERTY()
 	int32 UserPtrIdx;
@@ -676,11 +679,74 @@ public:
 	FName RegisteredParameterMapRead;
 
 	UPROPERTY()
-	FName RegisteredParameterMapWrite;
+	TArray<FName> RegisteredParameterMapWrites;
+
+	UPROPERTY()
+	FString EmitterName;
 
 	bool IsUserDataInterface() const;
 };
 
+USTRUCT()
+struct NIAGARA_API FNiagaraResolvedUserDataInterfaceBinding
+{
+	GENERATED_BODY()
+
+	int32 UserParameterStoreDataInterfaceIndex;
+	int32 ScriptParameterStoreDataInterfaceIndex;
+
+	FNiagaraResolvedUserDataInterfaceBinding()
+		: UserParameterStoreDataInterfaceIndex(INDEX_NONE)
+		, ScriptParameterStoreDataInterfaceIndex(INDEX_NONE)
+	{
+	}
+
+	FNiagaraResolvedUserDataInterfaceBinding(int32 InUserParameterStoreDataInterfaceIndex, int32 InScriptParameterStoreDataInterfaceIndex)
+		: UserParameterStoreDataInterfaceIndex(InUserParameterStoreDataInterfaceIndex)
+		, ScriptParameterStoreDataInterfaceIndex(InScriptParameterStoreDataInterfaceIndex)
+	{
+	}
+};
+
+USTRUCT()
+struct NIAGARA_API FNiagaraScriptResolvedDataInterfaceInfo
+{
+	GENERATED_USTRUCT_BODY();
+
+	FNiagaraScriptResolvedDataInterfaceInfo()
+		: bIsInternal(false)
+		, ResolvedDataInterface(nullptr)
+		, UserPtrIdx(-1)
+	{
+	}
+
+	UPROPERTY()
+	FName Name;
+
+	UPROPERTY()
+	FName CompileName;
+
+	UPROPERTY()
+	FString EmitterName;
+
+	UPROPERTY()
+	FNiagaraVariableBase ResolvedVariable;
+
+	UPROPERTY()
+	FNiagaraVariableBase SourceVariable;
+
+	UPROPERTY()
+	bool bIsInternal;
+
+	UPROPERTY()
+	TObjectPtr<UNiagaraDataInterface> ResolvedDataInterface;
+
+	UPROPERTY()
+	int32 UserPtrIdx;
+
+	bool NeedsPerInstanceBinding() const;
+};
+
 USTRUCT()
 struct NIAGARA_API FNiagaraScriptDataInterfaceCompileInfo
 {
@@ -714,7 +780,7 @@ public:
 	FName RegisteredParameterMapRead;
 
 	UPROPERTY()
-	FName RegisteredParameterMapWrite;
+	TArray<FName> RegisteredParameterMapWrites;
 
 	UPROPERTY()
 	bool bIsPlaceholder;
@@ -1298,8 +1364,6 @@ namespace FNiagaraUtilities
 	NIAGARA_API FString CreateRapidIterationConstantName(FName InVariableName, const TCHAR* InEmitterName, ENiagaraScriptUsage InUsage);
 	FNiagaraVariable NIAGARA_API ConvertVariableToRapidIterationConstantName(FNiagaraVariable InVar, const TCHAR* InEmitterName, ENiagaraScriptUsage InUsage);
 
-	void CollectScriptDataInterfaceParameters(const UObject& Owner, const TArrayView<UNiagaraScript*>& Scripts, FNiagaraParameterStore& OutDataInterfaceParameters);
-
 	inline bool SupportsNiagaraRendering(ERHIFeatureLevel::Type FeatureLevel)
 	{
 		return true;
diff --git a/Engine/Plugins/FX/Niagara/Source/Niagara/Public/NiagaraDataInterfaceUtilities.h b/Engine/Plugins/FX/Niagara/Source/Niagara/Public/NiagaraDataInterfaceUtilities.h
--- a/Engine/Plugins/FX/Niagara/Source/Niagara/Public/NiagaraDataInterfaceUtilities.h
+++ b/Engine/Plugins/FX/Niagara/Source/Niagara/Public/NiagaraDataInterfaceUtilities.h
@@ -23,28 +23,19 @@ namespace FNiagaraDataInterfaceUtilities
 		bool bIncludeInternal = false;
 	};
 
-	// Finds all VM function calls made using the data interface that equals this one (i.e. A->Equals(B))
+	// Finds all VM function calls made using the runtime resolved data interface.
 	// The same function call made be made multiple times across different scripts so you may see the same function multiple times
 	// The action should return True to continue iteration or False to stop
-	NIAGARA_API void ForEachVMFunctionEquals(UNiagaraDataInterface* DataInterface, UNiagaraSystem* NiagaraSystem, TFunction<bool(const UNiagaraScript*, const FVMExternalFunctionBindingInfo&)> Action);
-	NIAGARA_API void ForEachVMFunctionEquals(UNiagaraDataInterface* DataInterface, UNiagaraSystem* NiagaraSystem, UNiagaraComponent* Component, TFunction<bool(const UNiagaraScript*, const FVMExternalFunctionBindingInfo&)> Action);
-	NIAGARA_API void ForEachVMFunctionEquals(UNiagaraDataInterface* DataInterface, UNiagaraSystem* NiagaraSystem, FNiagaraSystemInstance* SystemInstance, TFunction<bool(const UNiagaraScript*, const FVMExternalFunctionBindingInfo&)> Action);
+	NIAGARA_API void ForEachVMFunction(UNiagaraDataInterface* ResolvedRuntimeDataInterface, UNiagaraSystem* NiagaraSystem, TFunction<bool(const UNiagaraScript*, const FVMExternalFunctionBindingInfo&)> Action);
+	NIAGARA_API void ForEachVMFunction(UNiagaraDataInterface* ResolvedRuntimeDataInterface, UNiagaraComponent* Component, TFunction<bool(const UNiagaraScript*, const FVMExternalFunctionBindingInfo&)> Action);
+	NIAGARA_API void ForEachVMFunction(UNiagaraDataInterface* ResolvedRuntimeDataInterface, FNiagaraSystemInstance* SystemInstance, TFunction<bool(const UNiagaraScript*, const FVMExternalFunctionBindingInfo&)> Action);
 
-	// Finds all Gpu function calls made using the data interface that equals this one (i.e. A->Equals(B))
+	// Finds all Gpu function calls made using the runtime resolved data interface.
 	// The same function call made be made multiple times across different scripts so you may see the same function multiple times
 	// The action should return True to continue iteration or False to stop
-	NIAGARA_API void ForEachGpuFunctionEquals(UNiagaraDataInterface* DataInterface, UNiagaraSystem* NiagaraSystem, TFunction<bool(const UNiagaraScript*, const FNiagaraDataInterfaceGeneratedFunction&)> Action);
-	NIAGARA_API void ForEachGpuFunctionEquals(UNiagaraDataInterface* DataInterface, UNiagaraSystem* NiagaraSystem, UNiagaraComponent* Component, TFunction<bool(const UNiagaraScript*, const FNiagaraDataInterfaceGeneratedFunction&)> Action);
-	NIAGARA_API void ForEachGpuFunctionEquals(UNiagaraDataInterface* DataInterface, UNiagaraSystem* NiagaraSystem, FNiagaraSystemInstance* SystemInstance, TFunction<bool(const UNiagaraScript*, const FNiagaraDataInterfaceGeneratedFunction&)> Action);
-
-	// Finds all VM function calls made using the data interface (i.e. pointer comparison A == B)
-	// The same function call made be made multiple times across different scripts so you may see the same function multiple times
-	// The action should return True to continue iteration or False to stop
-	NIAGARA_API void ForEachVMFunction(UNiagaraDataInterface* DataInterface, FNiagaraSystemInstance* SystemInstance, TFunction<bool(const UNiagaraScript*, const FVMExternalFunctionBindingInfo&)> Action);
-	// Finds all Gpu function calls made using the data interface (i.e. pointer comparison A == B)
-	// The same function call made be made multiple times across different scripts so you may see the same function multiple times
-	// The action should return True to continue iteration or False to stop
-	NIAGARA_API void ForEachGpuFunction(UNiagaraDataInterface* DataInterface, FNiagaraSystemInstance* SystemInstance, TFunction<bool(const FNiagaraDataInterfaceGeneratedFunction&)> Action);
+	NIAGARA_API void ForEachGpuFunction(UNiagaraDataInterface* ResolvedRuntimeDataInterface, UNiagaraSystem* NiagaraSystem, TFunction<bool(const UNiagaraScript*, const FNiagaraDataInterfaceGeneratedFunction&)> Action);
+	NIAGARA_API void ForEachGpuFunction(UNiagaraDataInterface* ResolvedRuntimeDataInterface, UNiagaraComponent* Component, TFunction<bool(const UNiagaraScript*, const FNiagaraDataInterfaceGeneratedFunction&)> Action);
+	NIAGARA_API void ForEachGpuFunction(UNiagaraDataInterface* ResolvedRuntimeDataInterface, FNiagaraSystemInstance* SystemInstance, TFunction<bool(const UNiagaraScript*, const FNiagaraDataInterfaceGeneratedFunction&)> Action);
 
 	// Loops over all data interfaces inside the SystemInstance
 	// The action should return True to continue iteration or False to stop
diff --git a/Engine/Plugins/FX/Niagara/Source/Niagara/Public/NiagaraSystemInstance.h b/Engine/Plugins/FX/Niagara/Source/Niagara/Public/NiagaraSystemInstance.h
--- a/Engine/Plugins/FX/Niagara/Source/Niagara/Public/NiagaraSystemInstance.h
+++ b/Engine/Plugins/FX/Niagara/Source/Niagara/Public/NiagaraSystemInstance.h
@@ -439,6 +439,8 @@ private:
 	/** Call PrepareForSImulation on each data source from the simulations and determine which need per-tick updates.*/
 	void InitDataInterfaces();
 
+	void ResolveUserDataInterfaceBindings();
+
 	/** The LWC tile of this system instance, used to offset all local simulation relative to the origin */
 	FVector3f LWCTile = FVector3f::ZeroVector;
 
diff --git a/Engine/Plugins/FX/Niagara/Source/Niagara/Public/NiagaraSystemSimulation.h b/Engine/Plugins/FX/Niagara/Source/Niagara/Public/NiagaraSystemSimulation.h
--- a/Engine/Plugins/FX/Niagara/Source/Niagara/Public/NiagaraSystemSimulation.h
+++ b/Engine/Plugins/FX/Niagara/Source/Niagara/Public/NiagaraSystemSimulation.h
@@ -296,8 +296,6 @@ public:
 
 	UNiagaraParameterCollectionInstance* GetParameterCollectionInstance(UNiagaraParameterCollection* Collection);
 
-	FNiagaraParameterStore& GetScriptDefinedDataInterfaceParameters();
-
 	/** Transfers a system instance from the current simulation into this one. */
 	void TransferInstance(FNiagaraSystemInstance* SystemInst);
 
@@ -426,9 +424,6 @@ protected:
 
 	void InitParameterDataSetBindings(FNiagaraSystemInstance* SystemInst);
 
-	/** A parameter store which contains the data interfaces parameters which were defined by the scripts. */
-	FNiagaraParameterStore ScriptDefinedDataInterfaceParameters;
-
 	TOptional<float> MaxDeltaTime;
 
 	/** Current tick batch we're filling ready for processing, potentially in an async task. */
diff --git a/Engine/Plugins/FX/Niagara/Source/NiagaraEditor/Private/Customizations/NiagaraTypeCustomizations.cpp b/Engine/Plugins/FX/Niagara/Source/NiagaraEditor/Private/Customizations/NiagaraTypeCustomizations.cpp
--- a/Engine/Plugins/FX/Niagara/Source/NiagaraEditor/Private/Customizations/NiagaraTypeCustomizations.cpp
+++ b/Engine/Plugins/FX/Niagara/Source/NiagaraEditor/Private/Customizations/NiagaraTypeCustomizations.cpp
@@ -966,7 +966,7 @@ TArray<TPair<FNiagaraVariableBase, FNiagaraVariableBase> > FNiagaraMaterialAttri
 					TArray<FNiagaraScriptDataInterfaceInfo>& CachedDIs = Script->GetCachedDefaultDataInterfaces();
 					for (const FNiagaraScriptDataInterfaceInfo& Info : CachedDIs)
 					{
-						if (Info.RegisteredParameterMapWrite == VariableName)
+						if (Info.RegisteredParameterMapWrites.Contains(VariableName))
 						{
 							return Info.DataInterface;
 						}
diff --git a/Engine/Plugins/FX/Niagara/Source/NiagaraEditor/Private/NiagaraEditorModule.cpp b/Engine/Plugins/FX/Niagara/Source/NiagaraEditor/Private/NiagaraEditorModule.cpp
--- a/Engine/Plugins/FX/Niagara/Source/NiagaraEditor/Private/NiagaraEditorModule.cpp
+++ b/Engine/Plugins/FX/Niagara/Source/NiagaraEditor/Private/NiagaraEditorModule.cpp
@@ -257,6 +257,16 @@ class FNiagaraEditorOnlyDataUtilities : public INiagaraEditorOnlyDataUtilities
 		NewMessage->SetAllowDismissal(bInAllowDismissal);
 		return NewMessage;
 	}
+
+	bool IsEditorDataInterfaceInstance(const UNiagaraDataInterface* DataInterface) const override
+	{
+		return FNiagaraEditorUtilities::IsEditorDataInterfaceInstance(DataInterface);
+	}
+
+	UNiagaraDataInterface* GetResolvedRuntimeInstanceForEditorDataInterfaceInstance(const UNiagaraSystem& OwningSystem, UNiagaraDataInterface& EditorDataInterfaceInstance) const override
+	{
+		return FNiagaraEditorUtilities::GetResolvedRuntimeInstanceForEditorDataInterfaceInstance(OwningSystem, EditorDataInterfaceInstance);
+	}
 };
 
 class FNiagaraScriptGraphPanelPinFactory : public FGraphPanelPinFactory
@@ -732,6 +742,68 @@ void LoadAllSystemsInFolder(const TArray<FString>& Arguments)
 	}
 }
 
+void DumpEmitterDependenciesInFolder(const TArray<FString>& Arguments)
+{
+	if (Arguments.Num() == 1)
+	{
+		TArray<FAssetData> SystemAssetsInFolder;
+		const FAssetRegistryModule& AssetRegistryModule = FModuleManager::LoadModuleChecked<FAssetRegistryModule>(TEXT("AssetRegistry"));
+		FARFilter Filter;
+		Filter.ClassPaths.Add(UNiagaraSystem::StaticClass()->GetClassPathName());
+		Filter.PackagePaths.Add(*Arguments[0]);
+		Filter.bRecursivePaths = true;
+		AssetRegistryModule.Get().GetAssets(Filter, SystemAssetsInFolder);
+		SystemAssetsInFolder.RemoveAll([](const FAssetData& AssetData) { return AssetData.GetObjectPathString().StartsWith("/Game/Developers"); });
+		if (SystemAssetsInFolder.Num() > 0)
+		{
+			const FText SlowTaskText = FText::Format(NSLOCTEXT("NiagaraEditor", "LoadAllSystemsInFolderFormat", "Loading {0} systems in folder {1}"), FText::AsNumber(SystemAssetsInFolder.Num()), FText::FromString(Arguments[0]));
+			FScopedSlowTask SlowTask(SystemAssetsInFolder.Num(), SlowTaskText);
+			SlowTask.MakeDialog(true);
+			int32 ItemNumber = 0;
+
+			bool bWasCancelled = false;
+			TArray<TPair<FString, TArray<int32>>> SystemPathEmitterIndicesPairs;
+			for (FAssetData& SystemAsset : SystemAssetsInFolder)
+			{
+				if (SlowTask.ShouldCancel())
+				{
+					bWasCancelled = true;
+					break;
+				}
+				ItemNumber++;
+				FText SlowTaskUpdateText = FText::Format(NSLOCTEXT("NiagaraEditor", "DumpEmitterDependenciesInFolder", "Dumping system {0} of {1}\n{2}"),
+					FText::AsNumber(ItemNumber), FText::AsNumber(SystemAssetsInFolder.Num()), FText::FromString(SystemAsset.GetFullName()));
+				SlowTask.EnterProgressFrame(1, SlowTaskUpdateText);
+				UNiagaraSystem* System = Cast<UNiagaraSystem>(SystemAsset.GetAsset());
+				if (System != nullptr)
+				{
+					System->WaitForCompilationComplete(true, false);
+					TPair<FString, TArray<int32>>& SystemPathEmitterIndicesPair = SystemPathEmitterIndicesPairs.AddDefaulted_GetRef();
+					SystemPathEmitterIndicesPair.Key = System->GetPathName();
+					for (const FNiagaraEmitterExecutionIndex& EmitterExecutionIndex : System->GetEmitterExecutionOrder())
+					{
+						SystemPathEmitterIndicesPair.Value.Add(EmitterExecutionIndex.EmitterIndex);
+					}
+				}
+			}
+
+			SystemPathEmitterIndicesPairs.Sort(
+				[](const TPair<FString, TArray<int32>>& A, const TPair<FString, TArray<int32>>& B)
+				{
+					return A.Key.Compare(B.Key) < 0;
+				});
+
+			TStringBuilder<1024> StringBuilder;
+			auto IntToString = [](const int32& Int) { return FString::Printf(TEXT("%i"), Int); };
+			for (const TPair<FString, TArray<int32>>& SystemEmitterIndiciesPair : SystemPathEmitterIndicesPairs)
+			{
+				StringBuilder.Appendf(TEXT("%s, %s\n"), *SystemEmitterIndiciesPair.Key, *FString::JoinBy(SystemEmitterIndiciesPair.Value, TEXT(","), IntToString));
+			}
+			FNiagaraEditorUtilities::WriteTextFileToDisk(FPaths::ProjectLogDir(), FString(TEXT("EmitterDependencies")) + TEXT(".csv"), StringBuilder.ToString(), true);
+		}
+	}
+}
+
 void ExecuteInvalidateNiagaraCachedScripts(const TArray< FString >& Args)
 {
 	if (Args.Num() == 0)
@@ -1288,6 +1360,11 @@ void FNiagaraEditorModule::StartupModule()
 		TEXT("Loads all niagara systems in the supplied directory and sub-directories."),
 		FConsoleCommandWithArgsDelegate::CreateStatic(&LoadAllSystemsInFolder));
 
+	DumpEmitterDependenciesCommand = IConsoleManager::Get().RegisterConsoleCommand(
+		TEXT("fx.DumpEmitterDepencenciesInFolder"),
+		TEXT("Dumps emitter dependencies for all systems in the supplied folder and sub-folders."),
+		FConsoleCommandWithArgsDelegate::CreateStatic(&DumpEmitterDependenciesInFolder));
+
 	FNiagaraMessageManager* MessageManager = FNiagaraMessageManager::Get();
 	MessageManager->RegisterMessageTopic(FNiagaraMessageTopics::CompilerTopicName);
 	MessageManager->RegisterMessageTopic(FNiagaraMessageTopics::ObjectTopicName);
diff --git a/Engine/Plugins/FX/Niagara/Source/NiagaraEditor/Private/NiagaraEditorUtilities.cpp b/Engine/Plugins/FX/Niagara/Source/NiagaraEditor/Private/NiagaraEditorUtilities.cpp
--- a/Engine/Plugins/FX/Niagara/Source/NiagaraEditor/Private/NiagaraEditorUtilities.cpp
+++ b/Engine/Plugins/FX/Niagara/Source/NiagaraEditor/Private/NiagaraEditorUtilities.cpp
@@ -4515,6 +4515,48 @@ void FNiagaraEditorUtilities::GetScriptMessageStores(UNiagaraScript* InScript, T
 	}
 }
 
+bool FNiagaraEditorUtilities::IsEditorDataInterfaceInstance(const UNiagaraDataInterface* DataInterface)
+{
+	return DataInterface->GetOuter()->IsA<UNiagaraNodeInput>();
+}
+
+UNiagaraDataInterface* FNiagaraEditorUtilities::GetResolvedRuntimeInstanceForEditorDataInterfaceInstance(const UNiagaraSystem& OwningSystem, UNiagaraDataInterface& EditorDataInterfaceInstance)
+{
+	UNiagaraNodeInput* OuterInputNode = EditorDataInterfaceInstance.GetTypedOuter<UNiagaraNodeInput>();
+	if (OuterInputNode != nullptr)
+	{
+		// If the data interface's owning node has been removed from it's graph then it's not valid so early out here.
+		bool bIsValidInputNode = OuterInputNode->GetGraph()->Nodes.Contains(OuterInputNode);
+		if (bIsValidInputNode == false)
+		{
+			return nullptr;
+		}
+
+		// If the data interface was owned by an input node, then we need to try to update the compiled version.
+		const FNiagaraEmitterHandle* EmitterHandle;
+		TArray<const UNiagaraScript*> CompiledScripts;
+		
+		FNiagaraStackGraphUtilities::GetEmitterHandleAndCompiledScriptsForStackNode(OwningSystem, *OuterInputNode, EmitterHandle, CompiledScripts);
+		if (ensureMsgf(CompiledScripts.Num() > 0, TEXT("Could not find compiled scripts for data interface input node.")))
+		{
+			for (const UNiagaraScript* CompiledScript : CompiledScripts)
+			{
+				FString EmitterName = EmitterHandle != nullptr ? EmitterHandle->GetUniqueInstanceName() : FString();
+				bool bIsParameterMapDataInterface = false;
+				FName DataInterfaceName = FHlslNiagaraTranslator::GetDataInterfaceName(OuterInputNode->Input.GetName(), EmitterName, bIsParameterMapDataInterface);
+				for (const FNiagaraScriptResolvedDataInterfaceInfo& ResolvedDI : CompiledScript->GetResolvedDataInterfaces())
+				{
+					if (ResolvedDI.Name == DataInterfaceName)
+					{
+						return ResolvedDI.ResolvedDataInterface;
+					}
+				}
+			}
+		}
+	}
+	return nullptr;
+}
+
 void FNiagaraParameterUtilities::FilterToRelevantStaticVariables(const TArray<FNiagaraVariable>& InVars, TArray<FNiagaraVariable>& OutVars, FName InOldEmitterAlias, FName InNewEmitterAlias, bool bFilterByEmitterAliasAndConvertToUnaliased)
 {
 	FNiagaraAliasContext RenameContext(ENiagaraScriptUsage::ParticleSpawnScript);
diff --git a/Engine/Plugins/FX/Niagara/Source/NiagaraEditor/Private/NiagaraEmitterFactoryNew.cpp b/Engine/Plugins/FX/Niagara/Source/NiagaraEditor/Private/NiagaraEmitterFactoryNew.cpp
--- a/Engine/Plugins/FX/Niagara/Source/NiagaraEditor/Private/NiagaraEmitterFactoryNew.cpp
+++ b/Engine/Plugins/FX/Niagara/Source/NiagaraEditor/Private/NiagaraEmitterFactoryNew.cpp
@@ -186,7 +186,7 @@ void UNiagaraEmitterFactoryNew::InitializeEmitter(UNiagaraEmitter* NewEmitter, b
 		{
 			NewEmitter->AddRenderer(NewObject<UNiagaraSpriteRendererProperties>(NewEmitter, "Renderer"), EmitterData->Version.VersionGuid);
 
-			AddModuleFromAssetPath(TEXT("/Niagara/Modules/Emitter/EmitterLifeCycle.EmitterLifeCycle"), *EmitterUpdateOutputNode);
+			AddModuleFromAssetPath(TEXT("/Niagara/Modules/Emitter/EmitterState.EmitterState"), *EmitterUpdateOutputNode);
 			
 			UNiagaraNodeFunctionCall* SpawnRateNode = AddModuleFromAssetPath(TEXT("/Niagara/Modules/Emitter/SpawnRate.SpawnRate"), *EmitterUpdateOutputNode);
 			if (SpawnRateNode != nullptr)
diff --git a/Engine/Plugins/FX/Niagara/Source/NiagaraEditor/Private/NiagaraHlslTranslator.cpp b/Engine/Plugins/FX/Niagara/Source/NiagaraEditor/Private/NiagaraHlslTranslator.cpp
--- a/Engine/Plugins/FX/Niagara/Source/NiagaraEditor/Private/NiagaraHlslTranslator.cpp
+++ b/Engine/Plugins/FX/Niagara/Source/NiagaraEditor/Private/NiagaraHlslTranslator.cpp
@@ -5493,16 +5493,7 @@ void FHlslNiagaraTranslator::ParameterMapSet(UNiagaraNodeParameterMapSet* SetNod
 					}
 
 					FNiagaraScriptDataInterfaceCompileInfo& Info = CompilationOutput.ScriptData.DataInterfaceInfo[Input];
-					if (Info.RegisteredParameterMapWrite == NAME_None)
-					{
-						Info.RegisteredParameterMapWrite = UsageName;
-					}
-					else
-					{
-						Error(FText::Format(LOCTEXT("ExternalDataInterfaceAssignedToMultipleParameters", "The data interface named {0} was added to a parameter map multiple times which isn't supported.  First usage: {1} Invalid usage:{2}"),
-							FText::FromName(Info.Name), FText::FromName(Info.RegisteredParameterMapWrite), FText::FromName(UsageName)), SetNode, Inputs[i].Pin);
-						continue;
-					}
+					Info.RegisteredParameterMapWrites.Add(UsageName);
 				}
 			}
 			else
diff --git a/Engine/Plugins/FX/Niagara/Source/NiagaraEditor/Private/NiagaraNodeFunctionCall.cpp b/Engine/Plugins/FX/Niagara/Source/NiagaraEditor/Private/NiagaraNodeFunctionCall.cpp
--- a/Engine/Plugins/FX/Niagara/Source/NiagaraEditor/Private/NiagaraNodeFunctionCall.cpp
+++ b/Engine/Plugins/FX/Niagara/Source/NiagaraEditor/Private/NiagaraNodeFunctionCall.cpp
@@ -434,7 +434,10 @@ PRAGMA_ENABLE_DEPRECATION_WARNINGS;
 		TArray<FNiagaraVariable> SwitchNodeInputs = Graph->FindStaticSwitchInputs();
 		for (FNiagaraVariable& Input : SwitchNodeInputs)
 		{
-			AddStaticSwitchInputPin(Input);
+			if (ensureMsgf(Input.IsValid(), TEXT("Static switch input %s on function call node %s is invalid"), *Input.GetName().ToString(), *GetPathName()))
+			{
+				AddStaticSwitchInputPin(Input);
+			}
 		}
 
 		for (FNiagaraVariable& Output : Outputs)
diff --git a/Engine/Plugins/FX/Niagara/Source/NiagaraEditor/Private/NiagaraScriptMergeManager.cpp b/Engine/Plugins/FX/Niagara/Source/NiagaraEditor/Private/NiagaraScriptMergeManager.cpp
--- a/Engine/Plugins/FX/Niagara/Source/NiagaraEditor/Private/NiagaraScriptMergeManager.cpp
+++ b/Engine/Plugins/FX/Niagara/Source/NiagaraEditor/Private/NiagaraScriptMergeManager.cpp
@@ -2964,7 +2964,11 @@ FNiagaraScriptMergeManager::FApplyDiffResults FNiagaraScriptMergeManager::AddInp
 		const UEdGraphSchema_Niagara* NiagaraSchema = GetDefault<UEdGraphSchema_Niagara>();
 		FNiagaraTypeDefinition InputType = NiagaraSchema->PinToTypeDefinition(OverrideToAdd->GetOverridePin());
 
-		UEdGraphPin& InputOverridePin = FNiagaraStackGraphUtilities::GetOrCreateStackFunctionInputOverridePin(TargetFunctionCall, AliasedFunctionInputHandle, InputType, OverrideToAdd->GetOverrideNode()->NodeGuid);
+		FNiagaraVariable InputVariable(InputType, FunctionInputHandle.GetParameterHandleString());
+		TOptional<FNiagaraVariableMetaData> InputMetaData = TargetFunctionCall.GetNiagaraGraph()->GetMetaData(InputVariable);
+		FGuid InputVariableGuid = InputMetaData.IsSet() ? InputMetaData->GetVariableGuid() : FGuid();
+
+		UEdGraphPin& InputOverridePin = FNiagaraStackGraphUtilities::GetOrCreateStackFunctionInputOverridePin(TargetFunctionCall, AliasedFunctionInputHandle, InputType, InputVariableGuid, OverrideToAdd->GetOverrideNode()->NodeGuid);
 		if (InputOverridePin.LinkedTo.Num() != 0)
 		{
 			Results.bSucceeded = false;
diff --git a/Engine/Plugins/FX/Niagara/Source/NiagaraEditor/Private/NiagaraStackFunctionInputBinder.cpp b/Engine/Plugins/FX/Niagara/Source/NiagaraEditor/Private/NiagaraStackFunctionInputBinder.cpp
--- a/Engine/Plugins/FX/Niagara/Source/NiagaraEditor/Private/NiagaraStackFunctionInputBinder.cpp
+++ b/Engine/Plugins/FX/Niagara/Source/NiagaraEditor/Private/NiagaraStackFunctionInputBinder.cpp
@@ -288,7 +288,7 @@ void FNiagaraStackFunctionInputBinder::SetData(const uint8* InValue, int32 InSiz
 
 			if (OverridePin == nullptr)
 			{
-				OverridePin = &FNiagaraStackGraphUtilities::GetOrCreateStackFunctionInputOverridePin(*FunctionCallNode, AliasedParameterHandle, InputType, InputVariableGuid);
+				OverridePin = &FNiagaraStackGraphUtilities::GetOrCreateStackFunctionInputOverridePin(*FunctionCallNode, AliasedParameterHandle, InputType, InputVariableGuid, FGuid());
 			}
 
 			OverridePin->Modify();
diff --git a/Engine/Plugins/FX/Niagara/Source/NiagaraEditor/Private/ViewModels/NiagaraPlaceholderDataInterfaceManager.cpp b/Engine/Plugins/FX/Niagara/Source/NiagaraEditor/Private/ViewModels/NiagaraPlaceholderDataInterfaceManager.cpp
--- a/Engine/Plugins/FX/Niagara/Source/NiagaraEditor/Private/ViewModels/NiagaraPlaceholderDataInterfaceManager.cpp
+++ b/Engine/Plugins/FX/Niagara/Source/NiagaraEditor/Private/ViewModels/NiagaraPlaceholderDataInterfaceManager.cpp
@@ -127,7 +127,7 @@ void FNiagaraPlaceholderDataInterfaceManager::PlaceholderDataInterfaceChanged(TW
 			FGuid ParameterVariableGuid = ParameterMetaData.IsSet() ? ParameterMetaData->GetVariableGuid() : FGuid();
 			UEdGraphPin& OverridePin = FNiagaraStackGraphUtilities::GetOrCreateStackFunctionInputOverridePin(
 				*PlaceholderDataInterfaceInfo->OwningFunctionCall.Get(), AliasedInputParameterHandle,
-				FNiagaraTypeDefinition(PlaceholderDataInterface->GetClass()), ParameterVariableGuid);
+				FNiagaraTypeDefinition(PlaceholderDataInterface->GetClass()), ParameterVariableGuid, FGuid());
 
 			if (OverridePin.LinkedTo.Num() == 1)
 			{
diff --git a/Engine/Plugins/FX/Niagara/Source/NiagaraEditor/Private/ViewModels/NiagaraSystemViewModel.cpp b/Engine/Plugins/FX/Niagara/Source/NiagaraEditor/Private/ViewModels/NiagaraSystemViewModel.cpp
--- a/Engine/Plugins/FX/Niagara/Source/NiagaraEditor/Private/ViewModels/NiagaraSystemViewModel.cpp
+++ b/Engine/Plugins/FX/Niagara/Source/NiagaraEditor/Private/ViewModels/NiagaraSystemViewModel.cpp
@@ -2312,43 +2312,12 @@ void GetCompiledScriptsAndEmitterNameFromInputNode(UNiagaraNode& StackNode, UNia
 	}
 }
 
-void UpdateCompiledDataInterfacesForScript(UNiagaraScript& TargetScript, FName TargetDataInterfaceName, UNiagaraDataInterface& SourceDataInterface)
-{
-	for (FNiagaraScriptDataInterfaceInfo& DataInterfaceInfo : TargetScript.GetCachedDefaultDataInterfaces())
-	{
-		if (DataInterfaceInfo.Name == TargetDataInterfaceName)
-		{
-			SourceDataInterface.CopyTo(DataInterfaceInfo.DataInterface);
-			break;
-		}
-	}
-}
-
 void FNiagaraSystemViewModel::UpdateCompiledDataInterfaces(UNiagaraDataInterface* ChangedDataInterface)
 {
-	UNiagaraNodeInput* OuterInputNode = ChangedDataInterface->GetTypedOuter<UNiagaraNodeInput>();
-	if (OuterInputNode != nullptr)
+	UNiagaraDataInterface* ResolvedCompileDataInterface = FNiagaraEditorUtilities::GetResolvedRuntimeInstanceForEditorDataInterfaceInstance(GetSystem(), *ChangedDataInterface);
+	if (ResolvedCompileDataInterface != nullptr)
 	{
-		// If the data interface's owning node has been removed from it's graph then it's not valid so early out here.
-		bool bIsValidInputNode = OuterInputNode->GetGraph()->Nodes.Contains(OuterInputNode);
-		if (bIsValidInputNode == false)
-		{
-			return;
-		}
-
-		// If the data interface was owned by an input node, then we need to try to update the compiled version.
-		TArray<UNiagaraScript*> CompiledScripts;
-		FString EmitterName;
-		GetCompiledScriptsAndEmitterNameFromInputNode(*OuterInputNode, GetSystem(), CompiledScripts, EmitterName);
-		if (ensureMsgf(CompiledScripts.Num() > 0, TEXT("Could not find compiled scripts for data interface input node.")))
-		{
-			for(UNiagaraScript* CompiledScript : CompiledScripts)
-			{
-				bool bIsParameterMapDataInterface = false;
-				FName DataInterfaceName = FHlslNiagaraTranslator::GetDataInterfaceName(OuterInputNode->Input.GetName(), EmitterName, bIsParameterMapDataInterface);
-				UpdateCompiledDataInterfacesForScript(*CompiledScript, DataInterfaceName, *ChangedDataInterface);
-			}
-		}
+		ChangedDataInterface->CopyTo(ResolvedCompileDataInterface);
 	}
 	else
 	{
diff --git a/Engine/Plugins/FX/Niagara/Source/NiagaraEditor/Private/ViewModels/Stack/NiagaraStackFunctionInput.cpp b/Engine/Plugins/FX/Niagara/Source/NiagaraEditor/Private/ViewModels/Stack/NiagaraStackFunctionInput.cpp
--- a/Engine/Plugins/FX/Niagara/Source/NiagaraEditor/Private/ViewModels/Stack/NiagaraStackFunctionInput.cpp
+++ b/Engine/Plugins/FX/Niagara/Source/NiagaraEditor/Private/ViewModels/Stack/NiagaraStackFunctionInput.cpp
@@ -2915,7 +2915,7 @@ UEdGraphPin& UNiagaraStackFunctionInput::GetOrCreateOverridePin()
 	{
 		TGuardValue<bool> Guard(bUpdatingGraphDirectly, true);
 		FGuid InputScriptVariableId = InputMetaData.IsSet() ? InputMetaData->GetVariableGuid() : FGuid();
-		OverridePin = &FNiagaraStackGraphUtilities::GetOrCreateStackFunctionInputOverridePin(*OwningFunctionCallNode, AliasedInputParameterHandle, InputType, InputScriptVariableId);
+		OverridePin = &FNiagaraStackGraphUtilities::GetOrCreateStackFunctionInputOverridePin(*OwningFunctionCallNode, AliasedInputParameterHandle, InputType, InputScriptVariableId, FGuid());
 		OverridePinCache = OverridePin;
 	}
 	return *OverridePin;
diff --git a/Engine/Plugins/FX/Niagara/Source/NiagaraEditor/Private/ViewModels/Stack/NiagaraStackGraphUtilities.cpp b/Engine/Plugins/FX/Niagara/Source/NiagaraEditor/Private/ViewModels/Stack/NiagaraStackGraphUtilities.cpp
--- a/Engine/Plugins/FX/Niagara/Source/NiagaraEditor/Private/ViewModels/Stack/NiagaraStackGraphUtilities.cpp
+++ b/Engine/Plugins/FX/Niagara/Source/NiagaraEditor/Private/ViewModels/Stack/NiagaraStackGraphUtilities.cpp
@@ -4462,6 +4462,54 @@ void FNiagaraStackGraphUtilities::FixDynamicInputNodeOutputPinsFromExternalChang
 	}
 }
 
+void FNiagaraStackGraphUtilities::GetEmitterHandleAndCompiledScriptsForStackNode(const UNiagaraSystem& OwningSystem, UNiagaraNode& StackNode, const FNiagaraEmitterHandle*& OutEmitterHandle, TArray<const UNiagaraScript*>& OutCompiledScripts)
+{
+	UNiagaraNodeOutput* OutputNode = GetEmitterOutputNodeForStackNode(StackNode);
+	if (OutputNode != nullptr)
+	{
+		if (OutputNode->GetUsage() == ENiagaraScriptUsage::SystemSpawnScript)
+		{
+			OutCompiledScripts.Add(OwningSystem.GetSystemSpawnScript());
+			OutEmitterHandle = nullptr;
+		}
+		else if (OutputNode->GetUsage() == ENiagaraScriptUsage::SystemUpdateScript)
+		{
+			OutCompiledScripts.Add(OwningSystem.GetSystemUpdateScript());
+			OutEmitterHandle = nullptr;
+		}
+		else
+		{
+			OutEmitterHandle = OwningSystem.GetEmitterHandles().FindByPredicate([&StackNode](const FNiagaraEmitterHandle& EmitterHandle)
+				{
+					return CastChecked<UNiagaraScriptSource>(EmitterHandle.GetEmitterData()->GraphSource)->NodeGraph == StackNode.GetNiagaraGraph();
+				});
+
+			if (OutEmitterHandle != nullptr)
+			{
+				switch (OutputNode->GetUsage())
+				{
+				case ENiagaraScriptUsage::EmitterSpawnScript:
+					OutCompiledScripts.Add(OwningSystem.GetSystemSpawnScript());
+					break;
+				case ENiagaraScriptUsage::EmitterUpdateScript:
+					OutCompiledScripts.Add(OwningSystem.GetSystemUpdateScript());
+					break;
+				case ENiagaraScriptUsage::ParticleSpawnScript:
+				case ENiagaraScriptUsage::ParticleUpdateScript:
+				case ENiagaraScriptUsage::ParticleEventScript:
+				case ENiagaraScriptUsage::ParticleSimulationStageScript:
+					OutCompiledScripts.Add(OutEmitterHandle->GetEmitterData()->GetScript(OutputNode->GetUsage(), OutputNode->GetUsageId()));
+					if (OutEmitterHandle->GetEmitterData()->SimTarget == ENiagaraSimTarget::GPUComputeSim)
+					{
+						OutCompiledScripts.Add(OutEmitterHandle->GetEmitterData()->GetScript(ENiagaraScriptUsage::ParticleGPUComputeScript, FGuid()));
+					}
+					break;
+				}
+			}
+		}
+	}
+}
+
 int32 FNiagaraStackGraphUtilities::DependencyUtilities::FindBestIndexForModuleInStack(UNiagaraNodeFunctionCall& ModuleNode, UNiagaraGraph& EmitterScriptGraph)
 {
 	// Check if the new module node has any dependencies to begin with. If not, early exit.
diff --git a/Engine/Plugins/FX/Niagara/Source/NiagaraEditor/Public/NiagaraEditorModule.h b/Engine/Plugins/FX/Niagara/Source/NiagaraEditor/Public/NiagaraEditorModule.h
--- a/Engine/Plugins/FX/Niagara/Source/NiagaraEditor/Public/NiagaraEditorModule.h
+++ b/Engine/Plugins/FX/Niagara/Source/NiagaraEditor/Public/NiagaraEditorModule.h
@@ -393,6 +393,7 @@ private:
 	IConsoleCommand* UpgradeAllNiagaraAssetsCommand;
 	IConsoleCommand* DumpCompileIdDataForAssetCommand;
 	IConsoleCommand* LoadAllSystemsInFolderCommand;
+	IConsoleCommand* DumpEmitterDependenciesCommand;
 
 	FOnCheckScriptToolkitsShouldFocusGraphElement OnCheckScriptToolkitsShouldFocusGraphElement;
 
diff --git a/Engine/Plugins/FX/Niagara/Source/NiagaraEditor/Public/NiagaraEditorUtilities.h b/Engine/Plugins/FX/Niagara/Source/NiagaraEditor/Public/NiagaraEditorUtilities.h
--- a/Engine/Plugins/FX/Niagara/Source/NiagaraEditor/Public/NiagaraEditorUtilities.h
+++ b/Engine/Plugins/FX/Niagara/Source/NiagaraEditor/Public/NiagaraEditorUtilities.h
@@ -409,6 +409,10 @@ namespace FNiagaraEditorUtilities
 	bool IsEnumIndexVisible(const UEnum* Enum, int32 Index);
 
 	void GetScriptMessageStores(UNiagaraScript* InScript, TArray<FNiagaraMessageSourceAndStore>& OutNiagaraMessageStores);
+
+	NIAGARAEDITOR_API bool IsEditorDataInterfaceInstance(const UNiagaraDataInterface* DataInterface);
+
+	NIAGARAEDITOR_API UNiagaraDataInterface* GetResolvedRuntimeInstanceForEditorDataInterfaceInstance(const UNiagaraSystem& OwningSystem, UNiagaraDataInterface& EditorDataInterfaceInstance);
 }
 
 namespace FNiagaraParameterUtilities
diff --git a/Engine/Plugins/FX/Niagara/Source/NiagaraEditor/Public/NiagaraEmitterFactoryNew.h b/Engine/Plugins/FX/Niagara/Source/NiagaraEditor/Public/NiagaraEmitterFactoryNew.h
--- a/Engine/Plugins/FX/Niagara/Source/NiagaraEditor/Public/NiagaraEmitterFactoryNew.h
+++ b/Engine/Plugins/FX/Niagara/Source/NiagaraEditor/Public/NiagaraEmitterFactoryNew.h
@@ -23,7 +23,7 @@ class UNiagaraEmitterFactoryNew : public UFactory
 	//~ Begin UFactory Interface	
 
 public:
-	static void InitializeEmitter(UNiagaraEmitter* NewEmitter, bool bAddDefaultModulesAndRenderers);
+	NIAGARAEDITOR_API static void InitializeEmitter(UNiagaraEmitter* NewEmitter, bool bAddDefaultModulesAndRenderers);
 };
 
 
diff --git a/Engine/Plugins/FX/Niagara/Source/NiagaraEditor/Public/NiagaraGraph.h b/Engine/Plugins/FX/Niagara/Source/NiagaraEditor/Public/NiagaraGraph.h
--- a/Engine/Plugins/FX/Niagara/Source/NiagaraEditor/Public/NiagaraGraph.h
+++ b/Engine/Plugins/FX/Niagara/Source/NiagaraEditor/Public/NiagaraGraph.h
@@ -288,7 +288,7 @@ class UNiagaraGraph : public UEdGraph
 	const TMap<FNiagaraVariable, FNiagaraGraphParameterReferenceCollection>& GetParameterReferenceMap() const; // NOTE: The const is a lie! (This indirectly calls RefreshParameterReferences, which can recreate the entire map)
 
 	/** Gets the meta-data associated with this variable, if it exists.*/
-	TOptional<FNiagaraVariableMetaData> GetMetaData(const FNiagaraVariable& InVar) const;
+	NIAGARAEDITOR_API TOptional<FNiagaraVariableMetaData> GetMetaData(const FNiagaraVariable& InVar) const;
 
 	/** Sets the meta-data associated with this variable. Creates a new UNiagaraScriptVariable if the target variable cannot be found. Illegal to call on FNiagaraVariables that are Niagara Constants. */
 	void SetMetaData(const FNiagaraVariable& InVar, const FNiagaraVariableMetaData& MetaData);
diff --git a/Engine/Plugins/FX/Niagara/Source/NiagaraEditor/Public/ViewModels/Stack/NiagaraStackGraphUtilities.h b/Engine/Plugins/FX/Niagara/Source/NiagaraEditor/Public/ViewModels/Stack/NiagaraStackGraphUtilities.h
--- a/Engine/Plugins/FX/Niagara/Source/NiagaraEditor/Public/ViewModels/Stack/NiagaraStackGraphUtilities.h
+++ b/Engine/Plugins/FX/Niagara/Source/NiagaraEditor/Public/ViewModels/Stack/NiagaraStackGraphUtilities.h
@@ -34,6 +34,7 @@ class UNiagaraSimulationStageBase;
 struct FNiagaraEventScriptProperties;
 struct FNiagaraStackModuleData;
 struct FNiagaraModuleDependency;
+struct FNiagaraEmitterHandle;
 
 namespace FNiagaraStackGraphUtilities
 {
@@ -202,7 +203,7 @@ namespace FNiagaraStackGraphUtilities
 
 	UEdGraphPin* GetStackFunctionInputOverridePin(UNiagaraNodeFunctionCall& StackFunctionCall, FNiagaraParameterHandle AliasedInputParameterHandle);
 
-	UEdGraphPin& GetOrCreateStackFunctionInputOverridePin(UNiagaraNodeFunctionCall& StackFunctionCall, FNiagaraParameterHandle AliasedInputParameterHandle, FNiagaraTypeDefinition InputType, const FGuid& InputScriptVariableId, const FGuid& PreferredOverrideNodeGuid = FGuid());
+	NIAGARAEDITOR_API UEdGraphPin& GetOrCreateStackFunctionInputOverridePin(UNiagaraNodeFunctionCall& StackFunctionCall, FNiagaraParameterHandle AliasedInputParameterHandle, FNiagaraTypeDefinition InputType, const FGuid& InputScriptVariableId, const FGuid& PreferredOverrideNodeGuid);
 
 	bool IsOverridePinForFunction(UEdGraphPin& OverridePin, UNiagaraNodeFunctionCall& FunctionCallNode);
 
@@ -215,11 +216,11 @@ namespace FNiagaraStackGraphUtilities
 	UEdGraphPin* GetLinkedValueHandleForFunctionInput(const UEdGraphPin& OverridePin);
 
 	TSet<FNiagaraVariable> GetParametersForContext(UEdGraph* Graph, UNiagaraSystem& System);
-	void SetLinkedValueHandleForFunctionInput(UEdGraphPin& OverridePin, FNiagaraParameterHandle LinkedParameterHandle, const TSet<FNiagaraVariable>& KnownParameters, ENiagaraDefaultMode DesiredDefaultMode = ENiagaraDefaultMode::FailIfPreviouslyNotSet, const FGuid& NewNodePersistentId = FGuid());
+	NIAGARAEDITOR_API void SetLinkedValueHandleForFunctionInput(UEdGraphPin& OverridePin, FNiagaraParameterHandle LinkedParameterHandle, const TSet<FNiagaraVariable>& KnownParameters, ENiagaraDefaultMode DesiredDefaultMode = ENiagaraDefaultMode::FailIfPreviouslyNotSet, const FGuid& NewNodePersistentId = FGuid());
 
-	void SetDataValueObjectForFunctionInput(UEdGraphPin& OverridePin, UClass* DataObjectType, FString InputNodeInputName, UNiagaraDataInterface*& OutDataObject, const FGuid& NewNodePersistentId = FGuid());
+	NIAGARAEDITOR_API void SetDataValueObjectForFunctionInput(UEdGraphPin& OverridePin, UClass* DataObjectType, FString InputNodeInputName, UNiagaraDataInterface*& OutDataObject, const FGuid& NewNodePersistentId = FGuid());
 
-	void SetDynamicInputForFunctionInput(UEdGraphPin& OverridePin, UNiagaraScript* DynamicInput, UNiagaraNodeFunctionCall*& OutDynamicInputFunctionCall, const FGuid& NewNodePersistentId = FGuid(), FString SuggestedName = FString(), const FGuid& InScriptVersion = FGuid());
+	NIAGARAEDITOR_API void SetDynamicInputForFunctionInput(UEdGraphPin& OverridePin, UNiagaraScript* DynamicInput, UNiagaraNodeFunctionCall*& OutDynamicInputFunctionCall, const FGuid& NewNodePersistentId = FGuid(), FString SuggestedName = FString(), const FGuid& InScriptVersion = FGuid());
 
 	void SetCustomExpressionForFunctionInput(UEdGraphPin& OverridePin, const FString& CustomExpression, UNiagaraNodeCustomHlsl*& OutDynamicInputFunctionCall, const FGuid& NewNodePersistentId = FGuid());
 
@@ -360,6 +361,8 @@ namespace FNiagaraStackGraphUtilities
 	/** Fixes orphaned and disconnected output pins on dynamic input nodes resulting from the underlying script changing and the pins being reallocated. */
 	void FixDynamicInputNodeOutputPinsFromExternalChanges(UNiagaraNodeFunctionCall& InFunctionCallNode);
 
+	void GetEmitterHandleAndCompiledScriptsForStackNode(const UNiagaraSystem& OwningSystem, UNiagaraNode& StackNode, const FNiagaraEmitterHandle*& OutEmitterHandle, TArray<const UNiagaraScript*>& OutCompiledScripts);
+
 	namespace DependencyUtilities
 	{
 		bool DoesStackModuleProvideDependency(const FNiagaraStackModuleData& StackModuleData, const FNiagaraModuleDependency& SourceModuleRequiredDependency, const UNiagaraNodeOutput& SourceOutputNode);
diff --git a/Engine/Source/Runtime/Core/Private/UObject/FortniteMainBranchObjectVersion.cpp b/Engine/Source/Runtime/Core/Private/UObject/FortniteMainBranchObjectVersion.cpp
--- a/Engine/Source/Runtime/Core/Private/UObject/FortniteMainBranchObjectVersion.cpp
+++ b/Engine/Source/Runtime/Core/Private/UObject/FortniteMainBranchObjectVersion.cpp
@@ -10,8 +10,8 @@ TMap<FGuid, FGuid> FFortniteMainBranchObjectVersion::GetSystemGuids()
 	SystemGuids.Add(DevGuids.LANDSCAPE_MOBILE_COOK_VERSION, FGuid("32D02EF867C74B71A0D4E0FA41392732"));
 	SystemGuids.Add(DevGuids.MATERIALSHADERMAP_DERIVEDDATA_VER, FGuid("CE7776E63E9E43FBB8C103EE1AB40B5D"));
 	SystemGuids.Add(DevGuids.NANITE_DERIVEDDATA_VER, FGuid("D28C003F825F4C178F698938FB0E47F7"));
-	SystemGuids.Add(DevGuids.NIAGARASHADERMAP_DERIVEDDATA_VER, FGuid("07FCB513221642B795CAC18A0D71FA01"));
-	SystemGuids.Add(DevGuids.Niagara_LatestScriptCompileVersion, FGuid("C9B73C19EBDD498996662D2E727538F6"));
+	SystemGuids.Add(DevGuids.NIAGARASHADERMAP_DERIVEDDATA_VER, FGuid("35D7960B72164B49B9B7B3FB5857E674"));
+	SystemGuids.Add(DevGuids.Niagara_LatestScriptCompileVersion, FGuid("3A499F83F8274468B50A487A3349C2BB"));
 	SystemGuids.Add(DevGuids.SkeletalMeshDerivedDataVersion, FGuid("ACF593EAAE354FCBB34CF44F5AA1BFD2"));
 	SystemGuids.Add(DevGuids.STATICMESH_DERIVEDDATA_VER, FGuid("543843B103794E0E9BA4BE80FB602F79"));
 	return SystemGuids;
