Niagara - Fix defines and includes.

diff --git a/Engine/Plugins/FX/Niagara/Source/Niagara/Private/NiagaraResolveDIHelpers.cpp b/Engine/Plugins/FX/Niagara/Source/Niagara/Private/NiagaraResolveDIHelpers.cpp
--- a/Engine/Plugins/FX/Niagara/Source/Niagara/Private/NiagaraResolveDIHelpers.cpp
+++ b/Engine/Plugins/FX/Niagara/Source/Niagara/Private/NiagaraResolveDIHelpers.cpp
@@ -1,7 +1,5 @@
 // Copyright Epic Games, Inc. All Rights Reserved.
 
-//#if WITH_EDITORONLY_DATA
-
 #include "NiagaraResolveDIHelpers.h"
 
 #include "Internationalization/Internationalization.h"
@@ -14,6 +12,8 @@
 
 #define LOCTEXT_NAMESPACE "NiagaraResolveDIHelpers"
 
+#if WITH_EDITORONLY_DATA
+
 namespace FNiagaraResolveDIHelpers
 {
 	void CollectDIBindingsAndAssignmentsForScript(
@@ -356,4 +356,4 @@ namespace FNiagaraResolveDIHelpers
 
 #undef LOCTEXT_NAMESPACE
 
-//#endif // WITH_EDITORONLY_DATA
\ No newline at end of file
+#endif // WITH_EDITORONLY_DATA
\ No newline at end of file
diff --git a/Engine/Plugins/FX/Niagara/Source/Niagara/Private/NiagaraResolveDIHelpers.h b/Engine/Plugins/FX/Niagara/Source/Niagara/Private/NiagaraResolveDIHelpers.h
--- a/Engine/Plugins/FX/Niagara/Source/Niagara/Private/NiagaraResolveDIHelpers.h
+++ b/Engine/Plugins/FX/Niagara/Source/Niagara/Private/NiagaraResolveDIHelpers.h
@@ -2,13 +2,15 @@
 
 #pragma once
 
-#if WITH_EDITORONLY_DATA
+#include "Misc/CoreMiscDefines.h"
 
 #include "NiagaraTypes.h"
 
 class UNiagaraDataInterface;
 class UNiagaraSystem;
 
+#if WITH_EDITORONLY_DATA
+
 namespace FNiagaraResolveDIHelpers
 {
 	void CollectDIBindingsAndAssignments(
