Niagara - Fix non-unity

diff --git a/Engine/Plugins/FX/Niagara/Source/Niagara/Private/DataInterface/NiagaraDataInterfaceStaticMesh.cpp b/Engine/Plugins/FX/Niagara/Source/Niagara/Private/DataInterface/NiagaraDataInterfaceStaticMesh.cpp
--- a/Engine/Plugins/FX/Niagara/Source/Niagara/Private/DataInterface/NiagaraDataInterfaceStaticMesh.cpp
+++ b/Engine/Plugins/FX/Niagara/Source/Niagara/Private/DataInterface/NiagaraDataInterfaceStaticMesh.cpp
@@ -26,6 +26,8 @@
 #include "StaticMeshResources.h"
 #include "Subsystems/ImportSubsystem.h"
 #include "INiagaraEditorOnlyDataUtlities.h"
+#include "Modules/ModuleManager.h"
+#include "NiagaraModule.h"
 #endif
 
 #include "Engine/StaticMeshActor.h"
diff --git a/Engine/Plugins/FX/Niagara/Source/Niagara/Private/NiagaraResolveDIHelpers.h b/Engine/Plugins/FX/Niagara/Source/Niagara/Private/NiagaraResolveDIHelpers.h
--- a/Engine/Plugins/FX/Niagara/Source/Niagara/Private/NiagaraResolveDIHelpers.h
+++ b/Engine/Plugins/FX/Niagara/Source/Niagara/Private/NiagaraResolveDIHelpers.h
@@ -2,6 +2,7 @@
 
 #pragma once
 
+#include "HAL/Platform.h"
 #include "Misc/CoreMiscDefines.h"
 
 #include "NiagaraTypes.h"
