Summary View 2.0
Base Hierarchy Editor Changes (also affects User Parameters):
- API Changes (mostly streamlining)
- Editability as a native feature. View models can override 'IsEditableByUser' which will make a lock icon appear on the row with a tooltip specified by the view model of the item.
- Reworked user messaging system
---> CanDropInternal now produces an optional message that is used to surface information to the user
---> Editability uses the same approach
- Moved 'CanHaveChildren' logic from data to view model
- Item view models now know if they are part of the hierarchy or part of the source. This allows us to modify behavior (modules show inputs in source view, but not in hierarchy)
- Search box & search API for view models
- Section support for sources panel (used in Summary View 2.0)
- 'Find in Hierarchy' context action for items already embedded in the hierarchy

Summary View Hierarchy Editor:
- Shows modules & renderers as source items/categories
- Parent emitter attributes can not be edited
- Override Display Name & tooltip property for inputs. Stack uses these if specified.
- Sections for Emitter/Particle Spawn/Update, Sim Stages, Events, Renderers
- Live update in the stack
- Drag & Dropping entire objects is possible (entire module, renderer, sim stage etc.)
- Bunch of QoL improvements

Merge process for summary view:
- Support for added items in children only

Post Load for existing emitters:
- data is transferred from summary view 1.0 data to 2.0 data
- Child Inputs are transferred explicitly as new setup uses explicit children

Stack changes:
- Refactored the stack category into a base class and some child classes
- Refactored the input collection into a 'property collection' base class and more child classes.
- Added new summary view support API that lets stack entries determine individually if they can be added to stack entries or not
- Added 'Show in Summary View' context action
- Added optional  'Owner Guid' functionality to stack objects & property rows to determine if a single property can be added to summary view

Some new utility functions and changes
Some new small-ish widgets

diff --git a/Engine/Plugins/FX/Niagara/Source/NiagaraEditor/Private/NiagaraParameterMapHistory.cpp b/Engine/Plugins/FX/Niagara/Source/NiagaraEditor/Private/NiagaraParameterMapHistory.cpp
--- a/Engine/Plugins/FX/Niagara/Source/NiagaraEditor/Private/NiagaraParameterMapHistory.cpp
+++ b/Engine/Plugins/FX/Niagara/Source/NiagaraEditor/Private/NiagaraParameterMapHistory.cpp
@@ -9,6 +9,7 @@
 #include "NiagaraGraph.h"
 #include "EdGraphSchema_Niagara.h"
 #include "NiagaraNode.h"
+#include "NiagaraGraph.h"
 #include "NiagaraCompiler.h"
 #include "NiagaraEditorUtilities.h"
 #include "NiagaraNodeOutput.h"
diff --git a/Engine/Plugins/FX/Niagara/Source/NiagaraEditor/Private/ViewModels/Stack/NiagaraStackGraphUtilities.cpp b/Engine/Plugins/FX/Niagara/Source/NiagaraEditor/Private/ViewModels/Stack/NiagaraStackGraphUtilities.cpp
--- a/Engine/Plugins/FX/Niagara/Source/NiagaraEditor/Private/ViewModels/Stack/NiagaraStackGraphUtilities.cpp
+++ b/Engine/Plugins/FX/Niagara/Source/NiagaraEditor/Private/ViewModels/Stack/NiagaraStackGraphUtilities.cpp
@@ -39,6 +39,7 @@
 #include "NiagaraParameterMapHistory.h"
 #include "NiagaraMessageManager.h"
 #include "NiagaraSimulationStageBase.h"
+#include "NiagaraEmitter.h"
 #include "NiagaraScriptVariable.h"
 #include "INiagaraEditorTypeUtilities.h"
 #include "NiagaraEditorData.h"
diff --git a/Engine/Plugins/FX/Niagara/Source/NiagaraEditor/Public/NiagaraParameterMapHistory.h b/Engine/Plugins/FX/Niagara/Source/NiagaraEditor/Public/NiagaraParameterMapHistory.h
--- a/Engine/Plugins/FX/Niagara/Source/NiagaraEditor/Public/NiagaraParameterMapHistory.h
+++ b/Engine/Plugins/FX/Niagara/Source/NiagaraEditor/Public/NiagaraParameterMapHistory.h
@@ -12,6 +12,7 @@
 class UNiagaraNodeOutput;
 class UEdGraphPin;
 class UEdGraphNode;
+class UNiagaraGraph;
 class UNiagaraParameterCollection;
 class FHlslNiagaraTranslator;
 class UNiagaraGraph;
diff --git a/Engine/Plugins/FX/Niagara/Source/NiagaraEditor/Public/ViewModels/NiagaraEmitterViewModel.h b/Engine/Plugins/FX/Niagara/Source/NiagaraEditor/Public/ViewModels/NiagaraEmitterViewModel.h
--- a/Engine/Plugins/FX/Niagara/Source/NiagaraEditor/Public/ViewModels/NiagaraEmitterViewModel.h
+++ b/Engine/Plugins/FX/Niagara/Source/NiagaraEditor/Public/ViewModels/NiagaraEmitterViewModel.h
@@ -26,6 +26,7 @@ class UNiagaraScriptVariable;
 class UNiagaraParameterDefinitions;
 class UNiagaraSummaryViewViewModel;
 struct FNiagaraMessageSourceAndStore;
+class UNiagaraSummaryViewViewModel;
 
 
 /** The view model for the UNiagaraEmitter objects */
