Fixed 2 issues:
- Cleanup of old rapid iteration parameters used name only, not type. This resulted in cases where old parameters that used the same name as a new parameter of different type would not get cleaned up.
- Fixed multiple lookups in the merge manager that looked up by name only instead of name + type. This in conjunction with the fix above results in more stable diff process.

All of this fixes an issue where some fluid emitters were showing or not showing the 'reset to parent emitter' arrow wrongly

diff --git a/Engine/Plugins/FX/Niagara/Source/NiagaraEditor/Private/NiagaraScriptMergeManager.cpp b/Engine/Plugins/FX/Niagara/Source/NiagaraEditor/Private/NiagaraScriptMergeManager.cpp
--- a/Engine/Plugins/FX/Niagara/Source/NiagaraEditor/Private/NiagaraScriptMergeManager.cpp
+++ b/Engine/Plugins/FX/Niagara/Source/NiagaraEditor/Private/NiagaraScriptMergeManager.cpp
@@ -1906,7 +1906,7 @@ bool FNiagaraScriptMergeManager::FindBaseModule(const FVersionedNiagaraEmitter&
 	return false;
 }
 
-bool FNiagaraScriptMergeManager::IsModuleInputDifferentFromBase(const FVersionedNiagaraEmitter& Emitter, const FVersionedNiagaraEmitter& BaseEmitter, ENiagaraScriptUsage ScriptUsage, FGuid ScriptUsageId, FGuid ModuleId, FString InputName)
+bool FNiagaraScriptMergeManager::IsModuleInputDifferentFromBase(const FVersionedNiagaraEmitter& Emitter, const FVersionedNiagaraEmitter& BaseEmitter, ENiagaraScriptUsage ScriptUsage, FGuid ScriptUsageId, FGuid ModuleId, FNiagaraVariableBase Variable)
 {
 	SCOPE_CYCLE_COUNTER(STAT_NiagaraEditor_ScriptMergeManager_IsModuleInputDifferentFromBase);
 
@@ -1936,7 +1936,8 @@ bool FNiagaraScriptMergeManager::IsModuleInputDifferentFromBase(const FVersioned
 
 	auto FindInputOverrideByInputName = [=](TSharedRef<FNiagaraStackFunctionInputOverrideMergeAdapter> InputOverride)
 	{
-		return InputOverride->GetOwningFunctionCall()->NodeGuid == ModuleId && InputOverride->GetInputName() == InputName;
+		FNiagaraVariableBase CandidateVariable(InputOverride->GetType(), FName(InputOverride->GetInputName()));
+		return InputOverride->GetOwningFunctionCall()->NodeGuid == ModuleId && CandidateVariable == Variable;
 	};
 
 	return
@@ -2592,10 +2593,13 @@ void FNiagaraScriptMergeManager::DiffScriptStacks(TSharedRef<FNiagaraScriptStack
 
 void FNiagaraScriptMergeManager::DiffFunctionInputs(TSharedRef<FNiagaraStackFunctionMergeAdapter> BaseFunctionAdapter, TSharedRef<FNiagaraStackFunctionMergeAdapter> OtherFunctionAdapter, FNiagaraScriptStackDiffResults& DiffResults) const
 {
-	FListDiffResults<TSharedRef<FNiagaraStackFunctionInputOverrideMergeAdapter>> ListDiffResults = DiffLists<TSharedRef<FNiagaraStackFunctionInputOverrideMergeAdapter>, FString>(
+	FListDiffResults<TSharedRef<FNiagaraStackFunctionInputOverrideMergeAdapter>> ListDiffResults = DiffLists<TSharedRef<FNiagaraStackFunctionInputOverrideMergeAdapter>, FNiagaraVariableBase>(
 		BaseFunctionAdapter->GetInputOverrides(),
 		OtherFunctionAdapter->GetInputOverrides(),
-		[](TSharedRef<FNiagaraStackFunctionInputOverrideMergeAdapter> InputOverrideAdapter) { return InputOverrideAdapter->GetInputName(); });
+		[](TSharedRef<FNiagaraStackFunctionInputOverrideMergeAdapter> InputOverrideAdapter)
+		{
+			return FNiagaraVariableBase(InputOverrideAdapter->GetType(), FName(InputOverrideAdapter->GetInputName()));
+		});
 
 	DiffResults.RemovedBaseInputOverrides.Append(ListDiffResults.RemovedBaseValues);
 	DiffResults.AddedOtherInputOverrides.Append(ListDiffResults.AddedOtherValues);
diff --git a/Engine/Plugins/FX/Niagara/Source/NiagaraEditor/Private/NiagaraScriptMergeManager.h b/Engine/Plugins/FX/Niagara/Source/NiagaraEditor/Private/NiagaraScriptMergeManager.h
--- a/Engine/Plugins/FX/Niagara/Source/NiagaraEditor/Private/NiagaraScriptMergeManager.h
+++ b/Engine/Plugins/FX/Niagara/Source/NiagaraEditor/Private/NiagaraScriptMergeManager.h
@@ -452,7 +452,7 @@ public:
 	// Find the module originating this module id in the scratch lists
 	bool FindBaseModule(const FVersionedNiagaraEmitter& BaseEmitter, ENiagaraScriptUsage ScriptUsage, FGuid ScriptUsageId, FGuid ModuleId, class UNiagaraScript*& OutActualScript, FGuid& OutScriptVersionGuid, FVersionedNiagaraEmitter& OutBaseEmitter);
 
-	bool IsModuleInputDifferentFromBase(const FVersionedNiagaraEmitter& Emitter, const FVersionedNiagaraEmitter& BaseEmitter, ENiagaraScriptUsage ScriptUsage, FGuid ScriptUsageId, FGuid ModuleId, FString InputName);
+	bool IsModuleInputDifferentFromBase(const FVersionedNiagaraEmitter& Emitter, const FVersionedNiagaraEmitter& BaseEmitter, ENiagaraScriptUsage ScriptUsage, FGuid ScriptUsageId, FGuid ModuleId, FNiagaraVariableBase Variable);
 
 	/** Check whether a given summary item identity exists in the parent of a given emitter. */
 	bool DoesSummaryItemExistInBase(const FVersionedNiagaraEmitter& Emitter, FNiagaraHierarchyIdentity Identity);
diff --git a/Engine/Plugins/FX/Niagara/Source/NiagaraEditor/Private/NiagaraScriptSource.cpp b/Engine/Plugins/FX/Niagara/Source/NiagaraEditor/Private/NiagaraScriptSource.cpp
--- a/Engine/Plugins/FX/Niagara/Source/NiagaraEditor/Private/NiagaraScriptSource.cpp
+++ b/Engine/Plugins/FX/Niagara/Source/NiagaraEditor/Private/NiagaraScriptSource.cpp
@@ -433,7 +433,7 @@ void UNiagaraScriptSource::FixupRenamedParameters(UNiagaraNodeFunctionCall* Func
 	}
 }
 
-void UNiagaraScriptSource::InitializeNewParameters(UNiagaraNodeFunctionCall* FunctionCallNode, TConstArrayView<FNiagaraVariable> ModuleInputVariables, FNiagaraParameterStore& RapidIterationParameters, const FVersionedNiagaraEmitter& VersionedEmitter, ENiagaraScriptUsage ScriptUsage, TSet<FName>& ValidRapidIterationParameterNames) const
+void UNiagaraScriptSource::InitializeNewParameters(UNiagaraNodeFunctionCall* FunctionCallNode, TConstArrayView<FNiagaraVariable> ModuleInputVariables, FNiagaraParameterStore& RapidIterationParameters, const FVersionedNiagaraEmitter& VersionedEmitter, ENiagaraScriptUsage ScriptUsage, TSet<FNiagaraVariableBase>& ValidRapidIterationParameters) const
 {
 	const FString UniqueEmitterName = VersionedEmitter.Emitter ? VersionedEmitter.Emitter->GetUniqueEmitterName() : FString();
 	UNiagaraGraph* Graph = FunctionCallNode->GetCalledGraph();
@@ -466,7 +466,7 @@ void UNiagaraScriptSource::InitializeNewParameters(UNiagaraNodeFunctionCall* Fun
 		{
 			FNiagaraParameterHandle AliasedFunctionInputHandle = FNiagaraParameterHandle::CreateAliasedModuleParameterHandle(ModuleInputVariable.GetName(), FunctionNodeName);
 			FNiagaraVariable RapidIterationParameter = FNiagaraStackGraphUtilities::CreateRapidIterationParameter(UniqueEmitterName, ScriptUsage, AliasedFunctionInputHandle.GetParameterHandleString(), ModuleInputVariable.GetType());
-			ValidRapidIterationParameterNames.Add(RapidIterationParameter.GetName());
+			ValidRapidIterationParameters.Add(RapidIterationParameter);
 			int32 ParameterIndex = RapidIterationParameters.IndexOf(RapidIterationParameter);
 
 			// Only set a value for the parameter if it's not already set.
@@ -531,7 +531,7 @@ void UNiagaraScriptSource::CleanUpOldAndInitializeNewRapidIterationParameters(co
 	TArray<FNiagaraVariable> OldRapidIterationVariables;
 	RapidIterationParameters.GetParameters(OldRapidIterationVariables);
 
-	TSet<FName> ValidRapidIterationParameterNames;
+	TSet<FNiagaraVariableBase> ValidRapidIterationParameters;
 	for (UNiagaraNodeOutput* OutputNode : OutputNodes)
 	{
 		if (OutputNode != nullptr)
@@ -560,7 +560,7 @@ void UNiagaraScriptSource::CleanUpOldAndInitializeNewRapidIterationParameters(co
 				if (!ModuleInputVariables.IsEmpty())
 				{
 					FixupRenamedParameters(FunctionCallNode, ModuleInputVariables, RapidIterationParameters, OldRapidIterationVariables, Emitter, ScriptUsage);
-					InitializeNewParameters(FunctionCallNode, ModuleInputVariables, RapidIterationParameters, Emitter, OutputNode->GetUsage(), ValidRapidIterationParameterNames);
+					InitializeNewParameters(FunctionCallNode, ModuleInputVariables, RapidIterationParameters, Emitter, OutputNode->GetUsage(), ValidRapidIterationParameters);
 				}
 			}
 		}
@@ -571,7 +571,7 @@ void UNiagaraScriptSource::CleanUpOldAndInitializeNewRapidIterationParameters(co
 	RapidIterationParameters.GetParameters(CurrentRapidIterationVariables);
 	for (const FNiagaraVariable& CurrentRapidIterationVariable : CurrentRapidIterationVariables)
 	{
-		if (ValidRapidIterationParameterNames.Contains(CurrentRapidIterationVariable.GetName()) == false)
+		if (ValidRapidIterationParameters.Contains(CurrentRapidIterationVariable) == false)
 		{
 			RapidIterationParameters.RemoveParameter(CurrentRapidIterationVariable);
 		}
diff --git a/Engine/Plugins/FX/Niagara/Source/NiagaraEditor/Private/ViewModels/Stack/NiagaraStackFunctionInput.cpp b/Engine/Plugins/FX/Niagara/Source/NiagaraEditor/Private/ViewModels/Stack/NiagaraStackFunctionInput.cpp
--- a/Engine/Plugins/FX/Niagara/Source/NiagaraEditor/Private/ViewModels/Stack/NiagaraStackFunctionInput.cpp
+++ b/Engine/Plugins/FX/Niagara/Source/NiagaraEditor/Private/ViewModels/Stack/NiagaraStackFunctionInput.cpp
@@ -2047,6 +2047,7 @@ bool UNiagaraStackFunctionInput::CanResetToBase() const
 				UNiagaraNodeOutput* OutputNode = FNiagaraStackGraphUtilities::GetEmitterOutputNodeForStackNode(*OwningFunctionCallNode.Get());
 				if(MergeManager->IsMergeableScriptUsage(OutputNode->GetUsage()))
 				{
+					FNiagaraVariableBase Variable(GetInputType(), GetInputParameterHandle().GetName());
 					TSharedPtr<FNiagaraEmitterViewModel> EmitterViewModelPtr = GetEmitterViewModel();
 					ensureMsgf(EmitterViewModelPtr.IsValid(), TEXT("ViewMode is nullptr and it never should be"));
 					FVersionedNiagaraEmitter BaseEmitter = EmitterViewModelPtr.IsValid() ? EmitterViewModelPtr->GetParentEmitter() : FVersionedNiagaraEmitter();
@@ -2056,7 +2057,7 @@ bool UNiagaraStackFunctionInput::CanResetToBase() const
 						OutputNode->GetUsage(),
 						OutputNode->GetUsageId(),
 						OwningModuleNode->NodeGuid,
-						InputParameterHandle.GetName().ToString());
+						Variable);
 				}
 				else
 				{
diff --git a/Engine/Plugins/FX/Niagara/Source/NiagaraEditor/Public/NiagaraScriptSource.h b/Engine/Plugins/FX/Niagara/Source/NiagaraEditor/Public/NiagaraScriptSource.h
--- a/Engine/Plugins/FX/Niagara/Source/NiagaraEditor/Public/NiagaraScriptSource.h
+++ b/Engine/Plugins/FX/Niagara/Source/NiagaraEditor/Public/NiagaraScriptSource.h
@@ -55,7 +55,7 @@ class UNiagaraScriptSource : public UNiagaraScriptSourceBase
 	NIAGARAEDITOR_API virtual bool AddModuleIfMissing(FString ModulePath, ENiagaraScriptUsage Usage, bool& bOutFoundModule)override;
 
 	virtual void FixupRenamedParameters(UNiagaraNodeFunctionCall* FunctionCallNode, TConstArrayView<FNiagaraVariable> ModuleInputVariables, FNiagaraParameterStore& RapidIterationParameters, const TArray<FNiagaraVariable>& OldRapidIterationVariables, const FVersionedNiagaraEmitter& Emitter, ENiagaraScriptUsage ScriptUsage) const;
-	virtual void InitializeNewParameters(UNiagaraNodeFunctionCall* FunctionCallNode, TConstArrayView<FNiagaraVariable> ModuleInputVariables, FNiagaraParameterStore& RapidIterationParameters, const FVersionedNiagaraEmitter& VersionedEmitter, ENiagaraScriptUsage ScriptUsage, TSet<FName>& ValidRapidIterationParameterNames) const;
+	virtual void InitializeNewParameters(UNiagaraNodeFunctionCall* FunctionCallNode, TConstArrayView<FNiagaraVariable> ModuleInputVariables, FNiagaraParameterStore& RapidIterationParameters, const FVersionedNiagaraEmitter& VersionedEmitter, ENiagaraScriptUsage ScriptUsage, TSet<FNiagaraVariableBase>& ValidRapidIterationParameters) const;
 
 	virtual void CleanUpOldAndInitializeNewRapidIterationParameters(const FVersionedNiagaraEmitter& Emitter, ENiagaraScriptUsage ScriptUsage, FGuid ScriptUsageId, FNiagaraParameterStore& RapidIterationParameters) const override;
 	virtual void ForceGraphToRecompileOnNextCheck() override;
