Added 'CollapseCategories' support for the stack. If a class has that flag, root level categories are skipped and instead their children are used.

diff --git a/Engine/Plugins/FX/Niagara/Source/Niagara/Classes/NiagaraDataInterface2DArrayTexture.h b/Engine/Plugins/FX/Niagara/Source/Niagara/Classes/NiagaraDataInterface2DArrayTexture.h
--- a/Engine/Plugins/FX/Niagara/Source/Niagara/Classes/NiagaraDataInterface2DArrayTexture.h
+++ b/Engine/Plugins/FX/Niagara/Source/Niagara/Classes/NiagaraDataInterface2DArrayTexture.h
@@ -10,7 +10,7 @@
 class UTexture2DArray;
 
 /** Data Interface allowing sampling of a texture */
-UCLASS(EditInlineNew, Category = "Texture", meta = (DisplayName = "2D Array Texture Sample"))
+UCLASS(EditInlineNew, Category = "Texture", CollapseCategories, meta = (DisplayName = "2D Array Texture Sample"))
 class NIAGARA_API UNiagaraDataInterface2DArrayTexture : public UNiagaraDataInterface
 {
 	GENERATED_UCLASS_BODY()
diff --git a/Engine/Plugins/FX/Niagara/Source/Niagara/Classes/NiagaraDataInterfaceArrayFloat.h b/Engine/Plugins/FX/Niagara/Source/Niagara/Classes/NiagaraDataInterfaceArrayFloat.h
--- a/Engine/Plugins/FX/Niagara/Source/Niagara/Classes/NiagaraDataInterfaceArrayFloat.h
+++ b/Engine/Plugins/FX/Niagara/Source/Niagara/Classes/NiagaraDataInterfaceArrayFloat.h
@@ -5,7 +5,7 @@
 #include "NiagaraDataInterfaceArray.h"
 #include "NiagaraDataInterfaceArrayFloat.generated.h"
 
-UCLASS(EditInlineNew, Category = "Array", meta = (DisplayName = "Float Array"), Blueprintable, BlueprintType)
+UCLASS(EditInlineNew, Category = "Array", CollapseCategories, meta = (DisplayName = "Float Array"), Blueprintable, BlueprintType)
 class NIAGARA_API UNiagaraDataInterfaceArrayFloat : public UNiagaraDataInterfaceArray
 {
 public:
@@ -17,7 +17,7 @@ public:
 	NDIARRAY_GENERATE_BODY(UNiagaraDataInterfaceArrayFloat, float, FloatData)
 };
 
-UCLASS(EditInlineNew, Category = "Array", meta = (DisplayName = "Vector 2D Array"), Blueprintable, BlueprintType)
+UCLASS(EditInlineNew, Category = "Array", CollapseCategories, meta = (DisplayName = "Vector 2D Array"), Blueprintable, BlueprintType)
 class NIAGARA_API UNiagaraDataInterfaceArrayFloat2 : public UNiagaraDataInterfaceArray
 {
 public:
@@ -34,7 +34,7 @@ public:
 	NDIARRAY_GENERATE_BODY_LWC(UNiagaraDataInterfaceArrayFloat2, FVector2f, FloatData)
 };
 
-UCLASS(EditInlineNew, Category = "Array", meta = (DisplayName = "Vector Array"), Blueprintable, BlueprintType)
+UCLASS(EditInlineNew, Category = "Array", CollapseCategories, meta = (DisplayName = "Vector Array"), Blueprintable, BlueprintType)
 class NIAGARA_API UNiagaraDataInterfaceArrayFloat3 : public UNiagaraDataInterfaceArray
 {
 public:
@@ -51,7 +51,7 @@ public:
 	NDIARRAY_GENERATE_BODY_LWC(UNiagaraDataInterfaceArrayFloat3, FVector3f, FloatData)
 };
 
-UCLASS(EditInlineNew, Category = "Array", meta = (DisplayName = "Position Array"), Blueprintable, BlueprintType)
+UCLASS(EditInlineNew, Category = "Array", CollapseCategories, meta = (DisplayName = "Position Array"), Blueprintable, BlueprintType)
 class NIAGARA_API UNiagaraDataInterfaceArrayPosition : public UNiagaraDataInterfaceArray
 {
 public:
@@ -63,7 +63,7 @@ public:
 	NDIARRAY_GENERATE_BODY(UNiagaraDataInterfaceArrayPosition, FNiagaraPosition, PositionData)
 };
 
-UCLASS(EditInlineNew, Category = "Array", meta = (DisplayName = "Vector 4 Array"), Blueprintable, BlueprintType)
+UCLASS(EditInlineNew, Category = "Array", CollapseCategories, meta = (DisplayName = "Vector 4 Array"), Blueprintable, BlueprintType)
 class NIAGARA_API UNiagaraDataInterfaceArrayFloat4 : public UNiagaraDataInterfaceArray
 {
 public:
@@ -80,7 +80,7 @@ public:
 	NDIARRAY_GENERATE_BODY_LWC(UNiagaraDataInterfaceArrayFloat4, FVector4f, FloatData)
 };
 
-UCLASS(EditInlineNew, Category = "Array", meta = (DisplayName = "Color Array"), Blueprintable, BlueprintType)
+UCLASS(EditInlineNew, Category = "Array", CollapseCategories, meta = (DisplayName = "Color Array"), Blueprintable, BlueprintType)
 class NIAGARA_API UNiagaraDataInterfaceArrayColor : public UNiagaraDataInterfaceArray
 {
 public:
@@ -92,7 +92,7 @@ public:
 	NDIARRAY_GENERATE_BODY(UNiagaraDataInterfaceArrayColor, FLinearColor, ColorData)
 };
 
-UCLASS(EditInlineNew, Category = "Array", meta = (DisplayName = "Quaternion Array"), Blueprintable, BlueprintType)
+UCLASS(EditInlineNew, Category = "Array", CollapseCategories, meta = (DisplayName = "Quaternion Array"), Blueprintable, BlueprintType)
 class NIAGARA_API UNiagaraDataInterfaceArrayQuat : public UNiagaraDataInterfaceArray
 {
 public:
diff --git a/Engine/Plugins/FX/Niagara/Source/Niagara/Classes/NiagaraDataInterfaceArrayInt.h b/Engine/Plugins/FX/Niagara/Source/Niagara/Classes/NiagaraDataInterfaceArrayInt.h
--- a/Engine/Plugins/FX/Niagara/Source/Niagara/Classes/NiagaraDataInterfaceArrayInt.h
+++ b/Engine/Plugins/FX/Niagara/Source/Niagara/Classes/NiagaraDataInterfaceArrayInt.h
@@ -5,7 +5,7 @@
 #include "NiagaraDataInterfaceArray.h"
 #include "NiagaraDataInterfaceArrayInt.generated.h"
 
-UCLASS(EditInlineNew, Category = "Array", meta = (DisplayName = "Int32 Array"), Blueprintable, BlueprintType)
+UCLASS(EditInlineNew, Category = "Array", CollapseCategories, meta = (DisplayName = "Int32 Array"), Blueprintable, BlueprintType)
 class NIAGARA_API UNiagaraDataInterfaceArrayInt32 : public UNiagaraDataInterfaceArray
 {
 public:
@@ -17,7 +17,7 @@ public:
 	NDIARRAY_GENERATE_BODY(UNiagaraDataInterfaceArrayInt32, int32, IntData)
 };
 
-UCLASS(EditInlineNew, Category = "Array", meta = (DisplayName = "UInt8 Array"), Blueprintable, BlueprintType)
+UCLASS(EditInlineNew, Category = "Array", CollapseCategories, meta = (DisplayName = "UInt8 Array"), Blueprintable, BlueprintType)
 class NIAGARA_API UNiagaraDataInterfaceArrayUInt8 : public UNiagaraDataInterfaceArray
 {
 public:
@@ -34,7 +34,7 @@ public:
 	NDIARRAY_GENERATE_BODY_LWC(UNiagaraDataInterfaceArrayUInt8, uint8, IntData)
 };
 
-UCLASS(EditInlineNew, Category = "Array", meta = (DisplayName = "Bool Array"), Blueprintable, BlueprintType)
+UCLASS(EditInlineNew, Category = "Array", CollapseCategories, meta = (DisplayName = "Bool Array"), Blueprintable, BlueprintType)
 class NIAGARA_API UNiagaraDataInterfaceArrayBool : public UNiagaraDataInterfaceArray
 {
 public:
diff --git a/Engine/Plugins/FX/Niagara/Source/Niagara/Classes/NiagaraDataInterfaceAsyncGpuTrace.h b/Engine/Plugins/FX/Niagara/Source/Niagara/Classes/NiagaraDataInterfaceAsyncGpuTrace.h
--- a/Engine/Plugins/FX/Niagara/Source/Niagara/Classes/NiagaraDataInterfaceAsyncGpuTrace.h
+++ b/Engine/Plugins/FX/Niagara/Source/Niagara/Classes/NiagaraDataInterfaceAsyncGpuTrace.h
@@ -11,7 +11,7 @@
 class FNiagaraSystemInstance;
 
 /** Data interface for handling latent (delayed 1 frame) traces against the scene for GPU simulations. */
-UCLASS(EditInlineNew, Category = "Collision", meta = (DisplayName = "Async Gpu Trace"))
+UCLASS(EditInlineNew, Category = "Collision", CollapseCategories, meta = (DisplayName = "Async Gpu Trace"))
 class NIAGARA_API UNiagaraDataInterfaceAsyncGpuTrace : public UNiagaraDataInterface
 {
 	GENERATED_UCLASS_BODY()
diff --git a/Engine/Plugins/FX/Niagara/Source/Niagara/Classes/NiagaraDataInterfaceAudio.h b/Engine/Plugins/FX/Niagara/Source/Niagara/Classes/NiagaraDataInterfaceAudio.h
--- a/Engine/Plugins/FX/Niagara/Source/Niagara/Classes/NiagaraDataInterfaceAudio.h
+++ b/Engine/Plugins/FX/Niagara/Source/Niagara/Classes/NiagaraDataInterfaceAudio.h
@@ -132,7 +132,7 @@ private:
 };
 
 /** Data Interface allowing sampling of recent audio data. */
-UCLASS(EditInlineNew, Category = "Audio", meta = (DisplayName = "Audio Submix"))
+UCLASS(EditInlineNew, Category = "Audio", CollapseCategories, meta = (DisplayName = "Audio Submix"))
 class NIAGARA_API UNiagaraDataInterfaceAudioSubmix : public UNiagaraDataInterface
 {
 	GENERATED_UCLASS_BODY()
diff --git a/Engine/Plugins/FX/Niagara/Source/Niagara/Classes/NiagaraDataInterfaceAudioOscilloscope.h b/Engine/Plugins/FX/Niagara/Source/Niagara/Classes/NiagaraDataInterfaceAudioOscilloscope.h
--- a/Engine/Plugins/FX/Niagara/Source/Niagara/Classes/NiagaraDataInterfaceAudioOscilloscope.h
+++ b/Engine/Plugins/FX/Niagara/Source/Niagara/Classes/NiagaraDataInterfaceAudioOscilloscope.h
@@ -87,7 +87,7 @@ private:
 };
 
 /** Data Interface allowing sampling of recent audio data. */
-UCLASS(EditInlineNew, Category = "Audio", meta = (DisplayName = "Audio Oscilloscope"))
+UCLASS(EditInlineNew, Category = "Audio", CollapseCategories, meta = (DisplayName = "Audio Oscilloscope"))
 class NIAGARA_API UNiagaraDataInterfaceAudioOscilloscope : public UNiagaraDataInterface
 {
 	GENERATED_UCLASS_BODY()
diff --git a/Engine/Plugins/FX/Niagara/Source/Niagara/Classes/NiagaraDataInterfaceCamera.h b/Engine/Plugins/FX/Niagara/Source/Niagara/Classes/NiagaraDataInterfaceCamera.h
--- a/Engine/Plugins/FX/Niagara/Source/Niagara/Classes/NiagaraDataInterfaceCamera.h
+++ b/Engine/Plugins/FX/Niagara/Source/Niagara/Classes/NiagaraDataInterfaceCamera.h
@@ -23,7 +23,7 @@ struct FCameraDataInterface_InstanceData
 	TArray<FDistanceData> ParticlesSortedByDistance;	
 };
 
-UCLASS(EditInlineNew, Category = "Camera", meta = (DisplayName = "Camera Query"))
+UCLASS(EditInlineNew, Category = "Camera", CollapseCategories, meta = (DisplayName = "Camera Query"))
 class NIAGARA_API UNiagaraDataInterfaceCamera : public UNiagaraDataInterface
 {
 	GENERATED_UCLASS_BODY()
diff --git a/Engine/Plugins/FX/Niagara/Source/Niagara/Classes/NiagaraDataInterfaceColorCurve.h b/Engine/Plugins/FX/Niagara/Source/Niagara/Classes/NiagaraDataInterfaceColorCurve.h
--- a/Engine/Plugins/FX/Niagara/Source/Niagara/Classes/NiagaraDataInterfaceColorCurve.h
+++ b/Engine/Plugins/FX/Niagara/Source/Niagara/Classes/NiagaraDataInterfaceColorCurve.h
@@ -17,7 +17,7 @@ class FNiagaraSystemInstance;
 
 
 /** Data Interface allowing sampling of color curves. */
-UCLASS(EditInlineNew, Category="Curves", meta = (DisplayName = "Curve for Colors"))
+UCLASS(EditInlineNew, Category="Curves", CollapseCategories, meta = (DisplayName = "Curve for Colors"))
 class NIAGARA_API UNiagaraDataInterfaceColorCurve : public UNiagaraDataInterfaceCurveBase
 {
 	GENERATED_UCLASS_BODY()
diff --git a/Engine/Plugins/FX/Niagara/Source/Niagara/Classes/NiagaraDataInterfaceCubeTexture.h b/Engine/Plugins/FX/Niagara/Source/Niagara/Classes/NiagaraDataInterfaceCubeTexture.h
--- a/Engine/Plugins/FX/Niagara/Source/Niagara/Classes/NiagaraDataInterfaceCubeTexture.h
+++ b/Engine/Plugins/FX/Niagara/Source/Niagara/Classes/NiagaraDataInterfaceCubeTexture.h
@@ -10,7 +10,7 @@
 class UTextureCube;
 
 /** Data Interface allowing sampling of a texture */
-UCLASS(EditInlineNew, Category = "Texture", meta = (DisplayName = "Cube Texture Sample"))
+UCLASS(EditInlineNew, Category = "Texture", CollapseCategories, meta = (DisplayName = "Cube Texture Sample"))
 class NIAGARA_API UNiagaraDataInterfaceCubeTexture : public UNiagaraDataInterface
 {
 	GENERATED_UCLASS_BODY()
diff --git a/Engine/Plugins/FX/Niagara/Source/Niagara/Classes/NiagaraDataInterfaceCurlNoise.h b/Engine/Plugins/FX/Niagara/Source/Niagara/Classes/NiagaraDataInterfaceCurlNoise.h
--- a/Engine/Plugins/FX/Niagara/Source/Niagara/Classes/NiagaraDataInterfaceCurlNoise.h
+++ b/Engine/Plugins/FX/Niagara/Source/Niagara/Classes/NiagaraDataInterfaceCurlNoise.h
@@ -8,7 +8,7 @@
 #include "NiagaraDataInterfaceCurlNoise.generated.h"
 
 /** Data Interface allowing sampling of curl noise LUT. */
-UCLASS(EditInlineNew, Category = "Curl Noise LUT", meta = (DisplayName = "Curl Noise"))
+UCLASS(EditInlineNew, Category = "Curl Noise LUT", CollapseCategories, meta = (DisplayName = "Curl Noise"))
 class NIAGARA_API UNiagaraDataInterfaceCurlNoise : public UNiagaraDataInterface
 {
 	GENERATED_UCLASS_BODY()
diff --git a/Engine/Plugins/FX/Niagara/Source/Niagara/Classes/NiagaraDataInterfaceCurve.h b/Engine/Plugins/FX/Niagara/Source/Niagara/Classes/NiagaraDataInterfaceCurve.h
--- a/Engine/Plugins/FX/Niagara/Source/Niagara/Classes/NiagaraDataInterfaceCurve.h
+++ b/Engine/Plugins/FX/Niagara/Source/Niagara/Classes/NiagaraDataInterfaceCurve.h
@@ -11,7 +11,7 @@
 
 
 /** Data Interface allowing sampling of float curves. */
-UCLASS(EditInlineNew, Category = "Curves", meta = (DisplayName = "Curve for Floats"))
+UCLASS(EditInlineNew, Category = "Curves", CollapseCategories, meta = (DisplayName = "Curve for Floats"))
 class NIAGARA_API UNiagaraDataInterfaceCurve : public UNiagaraDataInterfaceCurveBase
 {
 	GENERATED_UCLASS_BODY()
diff --git a/Engine/Plugins/FX/Niagara/Source/Niagara/Classes/NiagaraDataInterfaceGrid2DCollection.h b/Engine/Plugins/FX/Niagara/Source/Niagara/Classes/NiagaraDataInterfaceGrid2DCollection.h
--- a/Engine/Plugins/FX/Niagara/Source/Niagara/Classes/NiagaraDataInterfaceGrid2DCollection.h
+++ b/Engine/Plugins/FX/Niagara/Source/Niagara/Classes/NiagaraDataInterfaceGrid2DCollection.h
@@ -110,7 +110,7 @@ struct FNiagaraDataInterfaceParametersCS_Grid2DCollection : public FNiagaraDataI
 	LAYOUT_FIELD(TMemoryImageArray<uint32>, AttributeChannelCount);
 };
 
-UCLASS(EditInlineNew, Category = "Grid", meta = (DisplayName = "Grid2D Collection"), Blueprintable, BlueprintType)
+UCLASS(EditInlineNew, Category = "Grid", CollapseCategories, meta = (DisplayName = "Grid2D Collection"), Blueprintable, BlueprintType)
 class NIAGARA_API UNiagaraDataInterfaceGrid2DCollection : public UNiagaraDataInterfaceGrid2D
 {
 	GENERATED_UCLASS_BODY()
diff --git a/Engine/Plugins/FX/Niagara/Source/Niagara/Classes/NiagaraDataInterfaceGrid2DCollectionReader.h b/Engine/Plugins/FX/Niagara/Source/Niagara/Classes/NiagaraDataInterfaceGrid2DCollectionReader.h
--- a/Engine/Plugins/FX/Niagara/Source/Niagara/Classes/NiagaraDataInterfaceGrid2DCollectionReader.h
+++ b/Engine/Plugins/FX/Niagara/Source/Niagara/Classes/NiagaraDataInterfaceGrid2DCollectionReader.h
@@ -13,7 +13,7 @@ class FNiagaraSystemInstance;
 class UTextureRenderTarget;
 class UTextureRenderTargetVolume;
 
-UCLASS(EditInlineNew, Category = "Grid", hidecategories = (Grid2DCollection, Grid, Deprecated), meta = (DisplayName = "Grid2D Collection Reader", Experimental), Blueprintable, BlueprintType)
+UCLASS(EditInlineNew, Category = "Grid", CollapseCategories, hidecategories = (Grid2DCollection, Grid, Deprecated), meta = (DisplayName = "Grid2D Collection Reader", Experimental), Blueprintable, BlueprintType)
 class NIAGARA_API UNiagaraDataInterfaceGrid2DCollectionReader : public UNiagaraDataInterfaceGrid2DCollection
 {
 	GENERATED_UCLASS_BODY()
diff --git a/Engine/Plugins/FX/Niagara/Source/Niagara/Classes/NiagaraDataInterfaceGrid3DCollection.h b/Engine/Plugins/FX/Niagara/Source/Niagara/Classes/NiagaraDataInterfaceGrid3DCollection.h
--- a/Engine/Plugins/FX/Niagara/Source/Niagara/Classes/NiagaraDataInterfaceGrid3DCollection.h
+++ b/Engine/Plugins/FX/Niagara/Source/Niagara/Classes/NiagaraDataInterfaceGrid3DCollection.h
@@ -120,7 +120,7 @@ struct FNiagaraDataInterfaceParametersCS_Grid3DCollection : public FNiagaraDataI
 	LAYOUT_FIELD(TMemoryImageArray<uint32>, AttributeChannelCount);
 };
 
-UCLASS(EditInlineNew, Category = "Grid", meta = (DisplayName = "Grid3D Collection"), Blueprintable, BlueprintType)
+UCLASS(EditInlineNew, Category = "Grid", CollapseCategories, meta = (DisplayName = "Grid3D Collection"), Blueprintable, BlueprintType)
 class NIAGARA_API UNiagaraDataInterfaceGrid3DCollection : public UNiagaraDataInterfaceGrid3D
 {
 	GENERATED_UCLASS_BODY()
diff --git a/Engine/Plugins/FX/Niagara/Source/Niagara/Classes/NiagaraDataInterfaceGrid3DCollectionReader.h b/Engine/Plugins/FX/Niagara/Source/Niagara/Classes/NiagaraDataInterfaceGrid3DCollectionReader.h
--- a/Engine/Plugins/FX/Niagara/Source/Niagara/Classes/NiagaraDataInterfaceGrid3DCollectionReader.h
+++ b/Engine/Plugins/FX/Niagara/Source/Niagara/Classes/NiagaraDataInterfaceGrid3DCollectionReader.h
@@ -13,7 +13,7 @@ class FNiagaraSystemInstance;
 class UTextureRenderTarget;
 class UTextureRenderTargetVolume;
 
-UCLASS(EditInlineNew, Category = "Grid", hidecategories=(Grid3DCollection,Grid, Deprecated), meta = (DisplayName = "Grid3D Collection Reader"), Blueprintable, BlueprintType)
+UCLASS(EditInlineNew, Category = "Grid", CollapseCategories, hidecategories=(Grid3DCollection,Grid, Deprecated), meta = (DisplayName = "Grid3D Collection Reader"), Blueprintable, BlueprintType)
 class NIAGARA_API UNiagaraDataInterfaceGrid3DCollectionReader : public UNiagaraDataInterfaceGrid3DCollection
 {
 	GENERATED_UCLASS_BODY()
diff --git a/Engine/Plugins/FX/Niagara/Source/Niagara/Classes/NiagaraDataInterfaceIntRenderTarget2D.h b/Engine/Plugins/FX/Niagara/Source/Niagara/Classes/NiagaraDataInterfaceIntRenderTarget2D.h
--- a/Engine/Plugins/FX/Niagara/Source/Niagara/Classes/NiagaraDataInterfaceIntRenderTarget2D.h
+++ b/Engine/Plugins/FX/Niagara/Source/Niagara/Classes/NiagaraDataInterfaceIntRenderTarget2D.h
@@ -9,7 +9,7 @@
 
 class UTextureRenderTarget2D;
 
-UCLASS(EditInlineNew, Category = "Render Target", meta = (DisplayName = "Integer Render Target 2D", Experimental), Blueprintable, BlueprintType)
+UCLASS(EditInlineNew, Category = "Render Target", CollapseCategories, meta = (DisplayName = "Integer Render Target 2D", Experimental), Blueprintable, BlueprintType)
 class NIAGARA_API UNiagaraDataInterfaceIntRenderTarget2D : public UNiagaraDataInterfaceRWBase
 {
 	GENERATED_UCLASS_BODY()
diff --git a/Engine/Plugins/FX/Niagara/Source/Niagara/Classes/NiagaraDataInterfaceLandscape.h b/Engine/Plugins/FX/Niagara/Source/Niagara/Classes/NiagaraDataInterfaceLandscape.h
--- a/Engine/Plugins/FX/Niagara/Source/Niagara/Classes/NiagaraDataInterfaceLandscape.h
+++ b/Engine/Plugins/FX/Niagara/Source/Niagara/Classes/NiagaraDataInterfaceLandscape.h
@@ -34,7 +34,7 @@ enum class ENDILandscape_SourceMode : uint8
 };
 
 /** Data Interface allowing sampling of a Landscape */
-UCLASS(EditInlineNew, Category = "Landscape", meta = (DisplayName = "Landscape Sample"))
+UCLASS(EditInlineNew, Category = "Landscape", CollapseCategories, meta = (DisplayName = "Landscape Sample"))
 class NIAGARA_API UNiagaraDataInterfaceLandscape : public UNiagaraDataInterface
 {
 	GENERATED_UCLASS_BODY()
diff --git a/Engine/Plugins/FX/Niagara/Source/Niagara/Classes/NiagaraDataInterfaceMeshRendererInfo.h b/Engine/Plugins/FX/Niagara/Source/Niagara/Classes/NiagaraDataInterfaceMeshRendererInfo.h
--- a/Engine/Plugins/FX/Niagara/Source/Niagara/Classes/NiagaraDataInterfaceMeshRendererInfo.h
+++ b/Engine/Plugins/FX/Niagara/Source/Niagara/Classes/NiagaraDataInterfaceMeshRendererInfo.h
@@ -11,7 +11,7 @@ class UNiagaraMeshRendererProperties;
 class FNDIMeshRendererInfo;
 
 /** This Data Interface can be used to query information about the mesh renderers of an emitter */
-UCLASS(EditInlineNew, Category = "Mesh Particles", meta = (DisplayName = "Mesh Renderer Info"))
+UCLASS(EditInlineNew, Category = "Mesh Particles", CollapseCategories, meta = (DisplayName = "Mesh Renderer Info"))
 class NIAGARA_API UNiagaraDataInterfaceMeshRendererInfo : public UNiagaraDataInterface
 {
 	GENERATED_UCLASS_BODY()
diff --git a/Engine/Plugins/FX/Niagara/Source/Niagara/Classes/NiagaraDataInterfaceNeighborGrid3D.h b/Engine/Plugins/FX/Niagara/Source/Niagara/Classes/NiagaraDataInterfaceNeighborGrid3D.h
--- a/Engine/Plugins/FX/Niagara/Source/Niagara/Classes/NiagaraDataInterfaceNeighborGrid3D.h
+++ b/Engine/Plugins/FX/Niagara/Source/Niagara/Classes/NiagaraDataInterfaceNeighborGrid3D.h
@@ -56,7 +56,7 @@ struct FNiagaraDataInterfaceProxyNeighborGrid3D : public FNiagaraDataInterfacePr
 	TMap<FNiagaraSystemInstanceID, FNDINeighborGrid3DInstanceData_RT> SystemInstancesToProxyData_RT;
 };
 
-UCLASS(EditInlineNew, Category = "Grid", meta = (DisplayName = "Neighbor Grid3D"))
+UCLASS(EditInlineNew, Category = "Grid", CollapseCategories, meta = (DisplayName = "Neighbor Grid3D"))
 class NIAGARA_API UNiagaraDataInterfaceNeighborGrid3D : public UNiagaraDataInterfaceGrid3D
 {
 	GENERATED_UCLASS_BODY()
diff --git a/Engine/Plugins/FX/Niagara/Source/Niagara/Classes/NiagaraDataInterfaceParticleRead.h b/Engine/Plugins/FX/Niagara/Source/Niagara/Classes/NiagaraDataInterfaceParticleRead.h
--- a/Engine/Plugins/FX/Niagara/Source/Niagara/Classes/NiagaraDataInterfaceParticleRead.h
+++ b/Engine/Plugins/FX/Niagara/Source/Niagara/Classes/NiagaraDataInterfaceParticleRead.h
@@ -7,7 +7,7 @@
 #include "NiagaraEmitterInstance.h"
 #include "NiagaraDataInterfaceParticleRead.generated.h"
 
-UCLASS(EditInlineNew, Category = "ParticleRead", meta = (DisplayName = "Particle Attribute Reader"))
+UCLASS(EditInlineNew, Category = "ParticleRead", CollapseCategories, meta = (DisplayName = "Particle Attribute Reader"))
 class NIAGARA_API UNiagaraDataInterfaceParticleRead : public UNiagaraDataInterfaceRWBase
 {
 	GENERATED_UCLASS_BODY()
diff --git a/Engine/Plugins/FX/Niagara/Source/Niagara/Classes/NiagaraDataInterfacePlatformSet.h b/Engine/Plugins/FX/Niagara/Source/Niagara/Classes/NiagaraDataInterfacePlatformSet.h
--- a/Engine/Plugins/FX/Niagara/Source/Niagara/Classes/NiagaraDataInterfacePlatformSet.h
+++ b/Engine/Plugins/FX/Niagara/Source/Niagara/Classes/NiagaraDataInterfacePlatformSet.h
@@ -9,7 +9,7 @@
 
 
 /** Data Interface allowing querying of the current platform set. */
-UCLASS(EditInlineNew, Category = "Scalability", meta = (DisplayName = "Platform Set"))
+UCLASS(EditInlineNew, Category = "Scalability", CollapseCategories, meta = (DisplayName = "Platform Set"))
 class NIAGARA_API UNiagaraDataInterfacePlatformSet : public UNiagaraDataInterface
 {
 	GENERATED_UCLASS_BODY()
diff --git a/Engine/Plugins/FX/Niagara/Source/Niagara/Classes/NiagaraDataInterfaceRasterizationGrid3D.h b/Engine/Plugins/FX/Niagara/Source/Niagara/Classes/NiagaraDataInterfaceRasterizationGrid3D.h
--- a/Engine/Plugins/FX/Niagara/Source/Niagara/Classes/NiagaraDataInterfaceRasterizationGrid3D.h
+++ b/Engine/Plugins/FX/Niagara/Source/Niagara/Classes/NiagaraDataInterfaceRasterizationGrid3D.h
@@ -52,7 +52,7 @@ struct FNiagaraDataInterfaceProxyRasterizationGrid3D : public FNiagaraDataInterf
 	TMap<FNiagaraSystemInstanceID, RasterizationGrid3DRWInstanceData> SystemInstancesToProxyData;
 };
 
-UCLASS(EditInlineNew, Category = "Grid", meta = (DisplayName = "Rasterization Grid3D"))
+UCLASS(EditInlineNew, Category = "Grid", CollapseCategories, meta = (DisplayName = "Rasterization Grid3D"))
 class NIAGARA_API UNiagaraDataInterfaceRasterizationGrid3D : public UNiagaraDataInterfaceGrid3D
 {
 	GENERATED_UCLASS_BODY()
diff --git a/Engine/Plugins/FX/Niagara/Source/Niagara/Classes/NiagaraDataInterfaceRenderTarget2D.h b/Engine/Plugins/FX/Niagara/Source/Niagara/Classes/NiagaraDataInterfaceRenderTarget2D.h
--- a/Engine/Plugins/FX/Niagara/Source/Niagara/Classes/NiagaraDataInterfaceRenderTarget2D.h
+++ b/Engine/Plugins/FX/Niagara/Source/Niagara/Classes/NiagaraDataInterfaceRenderTarget2D.h
@@ -85,7 +85,7 @@ struct FNiagaraDataInterfaceProxyRenderTarget2DProxy : public FNiagaraDataInterf
 	TMap<FNiagaraSystemInstanceID, FRenderTarget2DRWInstanceData_RenderThread> SystemInstancesToProxyData_RT;
 };
 
-UCLASS(EditInlineNew, Category = "Rendering", meta = (DisplayName = "Render Target 2D"), Blueprintable, BlueprintType)
+UCLASS(EditInlineNew, Category = "Rendering", CollapseCategories, meta = (DisplayName = "Render Target 2D"), Blueprintable, BlueprintType)
 class NIAGARA_API UNiagaraDataInterfaceRenderTarget2D : public UNiagaraDataInterfaceRWBase
 {
 	GENERATED_UCLASS_BODY()
diff --git a/Engine/Plugins/FX/Niagara/Source/Niagara/Classes/NiagaraDataInterfaceRenderTarget2DArray.h b/Engine/Plugins/FX/Niagara/Source/Niagara/Classes/NiagaraDataInterfaceRenderTarget2DArray.h
--- a/Engine/Plugins/FX/Niagara/Source/Niagara/Classes/NiagaraDataInterfaceRenderTarget2DArray.h
+++ b/Engine/Plugins/FX/Niagara/Source/Niagara/Classes/NiagaraDataInterfaceRenderTarget2DArray.h
@@ -81,7 +81,7 @@ struct FNiagaraDataInterfaceProxyRenderTarget2DArrayProxy : public FNiagaraDataI
 	TMap<FNiagaraSystemInstanceID, FRenderTarget2DArrayRWInstanceData_RenderThread> SystemInstancesToProxyData_RT;
 };
 
-UCLASS(EditInlineNew, Category = "Rendering", meta = (DisplayName = "Render Target 2D Array"), Blueprintable, BlueprintType)
+UCLASS(EditInlineNew, Category = "Rendering", CollapseCategories, meta = (DisplayName = "Render Target 2D Array"), Blueprintable, BlueprintType)
 class NIAGARA_API UNiagaraDataInterfaceRenderTarget2DArray : public UNiagaraDataInterfaceRWBase
 {
 	GENERATED_UCLASS_BODY()
diff --git a/Engine/Plugins/FX/Niagara/Source/Niagara/Classes/NiagaraDataInterfaceRenderTargetCube.h b/Engine/Plugins/FX/Niagara/Source/Niagara/Classes/NiagaraDataInterfaceRenderTargetCube.h
--- a/Engine/Plugins/FX/Niagara/Source/Niagara/Classes/NiagaraDataInterfaceRenderTargetCube.h
+++ b/Engine/Plugins/FX/Niagara/Source/Niagara/Classes/NiagaraDataInterfaceRenderTargetCube.h
@@ -74,7 +74,7 @@ struct FNiagaraDataInterfaceProxyRenderTargetCubeProxy : public FNiagaraDataInte
 	TMap<FNiagaraSystemInstanceID, FRenderTargetCubeRWInstanceData_RenderThread> SystemInstancesToProxyData_RT;
 };
 
-UCLASS(EditInlineNew, Category = "Rendering", meta = (DisplayName = "Render Target Cube"), Blueprintable, BlueprintType)
+UCLASS(EditInlineNew, Category = "Rendering", CollapseCategories, meta = (DisplayName = "Render Target Cube"), Blueprintable, BlueprintType)
 class NIAGARA_API UNiagaraDataInterfaceRenderTargetCube : public UNiagaraDataInterfaceRWBase
 {
 	GENERATED_UCLASS_BODY()
diff --git a/Engine/Plugins/FX/Niagara/Source/Niagara/Classes/NiagaraDataInterfaceRenderTargetVolume.h b/Engine/Plugins/FX/Niagara/Source/Niagara/Classes/NiagaraDataInterfaceRenderTargetVolume.h
--- a/Engine/Plugins/FX/Niagara/Source/Niagara/Classes/NiagaraDataInterfaceRenderTargetVolume.h
+++ b/Engine/Plugins/FX/Niagara/Source/Niagara/Classes/NiagaraDataInterfaceRenderTargetVolume.h
@@ -74,7 +74,7 @@ struct FNiagaraDataInterfaceProxyRenderTargetVolumeProxy : public FNiagaraDataIn
 	TMap<FNiagaraSystemInstanceID, FRenderTargetVolumeRWInstanceData_RenderThread> SystemInstancesToProxyData_RT;
 };
 
-UCLASS(EditInlineNew, Category = "Rendering", meta = (DisplayName = "Render Target Volume"), Blueprintable, BlueprintType)
+UCLASS(EditInlineNew, Category = "Rendering", CollapseCategories, meta = (DisplayName = "Render Target Volume"), Blueprintable, BlueprintType)
 class NIAGARA_API UNiagaraDataInterfaceRenderTargetVolume : public UNiagaraDataInterfaceRWBase, public INiagaraSimCacheCustomStorageInterface
 {
 	GENERATED_UCLASS_BODY()
diff --git a/Engine/Plugins/FX/Niagara/Source/Niagara/Classes/NiagaraDataInterfaceSparseVolumeTexture.h b/Engine/Plugins/FX/Niagara/Source/Niagara/Classes/NiagaraDataInterfaceSparseVolumeTexture.h
--- a/Engine/Plugins/FX/Niagara/Source/Niagara/Classes/NiagaraDataInterfaceSparseVolumeTexture.h
+++ b/Engine/Plugins/FX/Niagara/Source/Niagara/Classes/NiagaraDataInterfaceSparseVolumeTexture.h
@@ -10,7 +10,7 @@
 class USparseVolumeTexture;
 
 /** Data Interface allowing sampling of a sparse volume texture */
-UCLASS(EditInlineNew, Category = "Texture", meta = (DisplayName = "Sparse Volume Texture Sample"))
+UCLASS(EditInlineNew, Category = "Texture", CollapseCategories, meta = (DisplayName = "Sparse Volume Texture Sample"))
 class NIAGARA_API UNiagaraDataInterfaceSparseVolumeTexture : public UNiagaraDataInterface
 {
 	GENERATED_UCLASS_BODY()
diff --git a/Engine/Plugins/FX/Niagara/Source/Niagara/Classes/NiagaraDataInterfaceSpline.h b/Engine/Plugins/FX/Niagara/Source/Niagara/Classes/NiagaraDataInterfaceSpline.h
--- a/Engine/Plugins/FX/Niagara/Source/Niagara/Classes/NiagaraDataInterfaceSpline.h
+++ b/Engine/Plugins/FX/Niagara/Source/Niagara/Classes/NiagaraDataInterfaceSpline.h
@@ -162,7 +162,7 @@ private:
 
 
 /** Data Interface allowing sampling of in-world spline components. Note that this data interface is very experimental. */
-UCLASS(EditInlineNew, Category = "Splines", meta = (DisplayName = "Spline"))
+UCLASS(EditInlineNew, Category = "Splines", CollapseCategories, meta = (DisplayName = "Spline"))
 class NIAGARA_API UNiagaraDataInterfaceSpline : public UNiagaraDataInterface
 {
 	GENERATED_UCLASS_BODY()
diff --git a/Engine/Plugins/FX/Niagara/Source/Niagara/Classes/NiagaraDataInterfaceSpriteRendererInfo.h b/Engine/Plugins/FX/Niagara/Source/Niagara/Classes/NiagaraDataInterfaceSpriteRendererInfo.h
--- a/Engine/Plugins/FX/Niagara/Source/Niagara/Classes/NiagaraDataInterfaceSpriteRendererInfo.h
+++ b/Engine/Plugins/FX/Niagara/Source/Niagara/Classes/NiagaraDataInterfaceSpriteRendererInfo.h
@@ -10,7 +10,7 @@
 class UNiagaraSpriteRendererProperties;
 
 /** This Data Interface can be used to query information about the sprite renderers of an emitter */
-UCLASS(EditInlineNew, Category = "Sprite Particles", meta = (DisplayName = "Sprite Renderer Info"))
+UCLASS(EditInlineNew, Category = "Sprite Particles", CollapseCategories, meta = (DisplayName = "Sprite Renderer Info"))
 class NIAGARA_API UNiagaraDataInterfaceSpriteRendererInfo : public UNiagaraDataInterface
 {
 	GENERATED_UCLASS_BODY()
diff --git a/Engine/Plugins/FX/Niagara/Source/Niagara/Classes/NiagaraDataInterfaceTexture.h b/Engine/Plugins/FX/Niagara/Source/Niagara/Classes/NiagaraDataInterfaceTexture.h
--- a/Engine/Plugins/FX/Niagara/Source/Niagara/Classes/NiagaraDataInterfaceTexture.h
+++ b/Engine/Plugins/FX/Niagara/Source/Niagara/Classes/NiagaraDataInterfaceTexture.h
@@ -8,7 +8,7 @@
 #include "NiagaraDataInterfaceTexture.generated.h"
 
 /** Data Interface allowing sampling of a texture */
-UCLASS(EditInlineNew, Category = "Texture", meta = (DisplayName = "Texture Sample"))
+UCLASS(EditInlineNew, Category = "Texture", CollapseCategories, meta = (DisplayName = "Texture Sample"))
 class NIAGARA_API UNiagaraDataInterfaceTexture : public UNiagaraDataInterface
 {
 	GENERATED_UCLASS_BODY()
diff --git a/Engine/Plugins/FX/Niagara/Source/Niagara/Classes/NiagaraDataInterfaceVector2DCurve.h b/Engine/Plugins/FX/Niagara/Source/Niagara/Classes/NiagaraDataInterfaceVector2DCurve.h
--- a/Engine/Plugins/FX/Niagara/Source/Niagara/Classes/NiagaraDataInterfaceVector2DCurve.h
+++ b/Engine/Plugins/FX/Niagara/Source/Niagara/Classes/NiagaraDataInterfaceVector2DCurve.h
@@ -11,7 +11,7 @@
 
 
 /** Data Interface allowing sampling of vector curves. */
-UCLASS(EditInlineNew, Category = "Curves", meta = (DisplayName = "Curve for Vector2D's"))
+UCLASS(EditInlineNew, Category = "Curves", CollapseCategories, meta = (DisplayName = "Curve for Vector2D's"))
 class NIAGARA_API UNiagaraDataInterfaceVector2DCurve : public UNiagaraDataInterfaceCurveBase
 {
 	GENERATED_UCLASS_BODY()
diff --git a/Engine/Plugins/FX/Niagara/Source/Niagara/Classes/NiagaraDataInterfaceVector4Curve.h b/Engine/Plugins/FX/Niagara/Source/Niagara/Classes/NiagaraDataInterfaceVector4Curve.h
--- a/Engine/Plugins/FX/Niagara/Source/Niagara/Classes/NiagaraDataInterfaceVector4Curve.h
+++ b/Engine/Plugins/FX/Niagara/Source/Niagara/Classes/NiagaraDataInterfaceVector4Curve.h
@@ -16,7 +16,7 @@ class UCurveFloat;
 class FNiagaraSystemInstance;
 
 /** Data Interface allowing sampling of vector curves. */
-UCLASS(EditInlineNew, Category = "Curves", meta = (DisplayName = "Curve for Vector4's"))
+UCLASS(EditInlineNew, Category = "Curves", CollapseCategories, meta = (DisplayName = "Curve for Vector4's"))
 class NIAGARA_API UNiagaraDataInterfaceVector4Curve : public UNiagaraDataInterfaceCurveBase
 {
 	GENERATED_UCLASS_BODY()
diff --git a/Engine/Plugins/FX/Niagara/Source/Niagara/Classes/NiagaraDataInterfaceVectorCurve.h b/Engine/Plugins/FX/Niagara/Source/Niagara/Classes/NiagaraDataInterfaceVectorCurve.h
--- a/Engine/Plugins/FX/Niagara/Source/Niagara/Classes/NiagaraDataInterfaceVectorCurve.h
+++ b/Engine/Plugins/FX/Niagara/Source/Niagara/Classes/NiagaraDataInterfaceVectorCurve.h
@@ -17,7 +17,7 @@ class FNiagaraSystemInstance;
 
 
 /** Data Interface allowing sampling of vector curves. */
-UCLASS(EditInlineNew, Category = "Curves", meta = (DisplayName = "Curve for Vector3's"))
+UCLASS(EditInlineNew, Category = "Curves", CollapseCategories, meta = (DisplayName = "Curve for Vector3's"))
 class NIAGARA_API UNiagaraDataInterfaceVectorCurve : public UNiagaraDataInterfaceCurveBase
 {
 	GENERATED_UCLASS_BODY()
diff --git a/Engine/Plugins/FX/Niagara/Source/Niagara/Classes/NiagaraDataInterfaceVectorField.h b/Engine/Plugins/FX/Niagara/Source/Niagara/Classes/NiagaraDataInterfaceVectorField.h
--- a/Engine/Plugins/FX/Niagara/Source/Niagara/Classes/NiagaraDataInterfaceVectorField.h
+++ b/Engine/Plugins/FX/Niagara/Source/Niagara/Classes/NiagaraDataInterfaceVectorField.h
@@ -7,7 +7,7 @@
 
 class FNiagaraSystemInstance;
 
-UCLASS(EditInlineNew, Category = "Vector Field", meta = (DisplayName = "Vector Field"))
+UCLASS(EditInlineNew, Category = "Vector Field", CollapseCategories, meta = (DisplayName = "Vector Field"))
 class NIAGARA_API UNiagaraDataInterfaceVectorField : public UNiagaraDataInterface
 {
 	GENERATED_UCLASS_BODY()
diff --git a/Engine/Plugins/FX/Niagara/Source/Niagara/Classes/NiagaraDataInterfaceVolumeCache.h b/Engine/Plugins/FX/Niagara/Source/Niagara/Classes/NiagaraDataInterfaceVolumeCache.h
--- a/Engine/Plugins/FX/Niagara/Source/Niagara/Classes/NiagaraDataInterfaceVolumeCache.h
+++ b/Engine/Plugins/FX/Niagara/Source/Niagara/Classes/NiagaraDataInterfaceVolumeCache.h
@@ -12,7 +12,7 @@
 
 struct FVolumeCacheInstanceData_GameThread;
 
-UCLASS(EditInlineNew, Category = "Grid",  meta = (DisplayName = "Volume Cache"))
+UCLASS(EditInlineNew, Category = "Grid", CollapseCategories, meta = (DisplayName = "Volume Cache"))
 class NIAGARA_API UNiagaraDataInterfaceVolumeCache : public UNiagaraDataInterface
 {
 	GENERATED_UCLASS_BODY()
diff --git a/Engine/Plugins/FX/Niagara/Source/Niagara/Classes/NiagaraDataInterfaceVolumeTexture.h b/Engine/Plugins/FX/Niagara/Source/Niagara/Classes/NiagaraDataInterfaceVolumeTexture.h
--- a/Engine/Plugins/FX/Niagara/Source/Niagara/Classes/NiagaraDataInterfaceVolumeTexture.h
+++ b/Engine/Plugins/FX/Niagara/Source/Niagara/Classes/NiagaraDataInterfaceVolumeTexture.h
@@ -10,7 +10,7 @@
 class UVolumeTexture;
 
 /** Data Interface allowing sampling of a texture */
-UCLASS(EditInlineNew, Category = "Texture", meta = (DisplayName = "Volume Texture Sample"))
+UCLASS(EditInlineNew, Category = "Texture", CollapseCategories, meta = (DisplayName = "Volume Texture Sample"))
 class NIAGARA_API UNiagaraDataInterfaceVolumeTexture : public UNiagaraDataInterface
 {
 	GENERATED_UCLASS_BODY()
diff --git a/Engine/Plugins/FX/Niagara/Source/Niagara/Internal/DataInterface/NiagaraDataInterfaceActorComponent.h b/Engine/Plugins/FX/Niagara/Source/Niagara/Internal/DataInterface/NiagaraDataInterfaceActorComponent.h
--- a/Engine/Plugins/FX/Niagara/Source/Niagara/Internal/DataInterface/NiagaraDataInterfaceActorComponent.h
+++ b/Engine/Plugins/FX/Niagara/Source/Niagara/Internal/DataInterface/NiagaraDataInterfaceActorComponent.h
@@ -31,7 +31,7 @@ enum class ENDIActorComponentSourceMode : uint8
 /**
 Data interface that gives you access to actor & component information.
 */
-UCLASS(EditInlineNew, Category = "Actor", meta = (DisplayName = "Actor Component Interface"))
+UCLASS(EditInlineNew, Category = "Actor", CollapseCategories, meta = (DisplayName = "Actor Component Interface"))
 class NIAGARA_API UNiagaraDataInterfaceActorComponent : public UNiagaraDataInterface
 {
 	GENERATED_UCLASS_BODY()
@@ -45,11 +45,6 @@ class NIAGARA_API UNiagaraDataInterfaceActorComponent : public UNiagaraDataInter
 	END_SHADER_PARAMETER_STRUCT();
 
 public:
-	/** When this option is disabled, we use the previous frame's data for the skeletal mesh and can often issue the simulation early. This greatly
-	reduces overhead and allows the game thread to run faster, but comes at a tradeoff if the dependencies might leave gaps or other visual artifacts.*/
-	UPROPERTY(EditAnywhere, Category = "Performance")
-	bool bRequireCurrentFrameData = true;
-
 	/** Controls how we find the actor / component we want to bind to. */
 	UPROPERTY(EditAnywhere, Category = "ActorComponent")
 	ENDIActorComponentSourceMode SourceMode = ENDIActorComponentSourceMode::Default;
@@ -65,6 +60,11 @@ public:
 	UPROPERTY(EditAnywhere, Category = "ActorComponent")
 	FNiagaraUserParameterBinding ActorOrComponentParameter;
 
+	/** When this option is disabled, we use the previous frame's data for the skeletal mesh and can often issue the simulation early. This greatly
+	reduces overhead and allows the game thread to run faster, but comes at a tradeoff if the dependencies might leave gaps or other visual artifacts.*/
+	UPROPERTY(EditAnywhere, Category = "Performance")
+	bool bRequireCurrentFrameData = true;
+	
 	//UObject Interface
 	virtual void PostInitProperties() override;
 	//UObject Interface End
diff --git a/Engine/Plugins/FX/Niagara/Source/Niagara/Internal/DataInterface/NiagaraDataInterfaceDataChannelRead.h b/Engine/Plugins/FX/Niagara/Source/Niagara/Internal/DataInterface/NiagaraDataInterfaceDataChannelRead.h
--- a/Engine/Plugins/FX/Niagara/Source/Niagara/Internal/DataInterface/NiagaraDataInterfaceDataChannelRead.h
+++ b/Engine/Plugins/FX/Niagara/Source/Niagara/Internal/DataInterface/NiagaraDataInterfaceDataChannelRead.h
@@ -31,7 +31,7 @@ struct FNDIDataChannelWriteInstanceData;
 class UNiagaraDataChannelHandler;
 class FNiagaraDataBuffer;
 
-UCLASS(Experimental, EditInlineNew, Category = "Data Channels", meta = (DisplayName = "Data Channel Reader"))
+UCLASS(Experimental, EditInlineNew, Category = "Data Channels", CollapseCategories, meta = (DisplayName = "Data Channel Reader"))
 class NIAGARA_API UNiagaraDataInterfaceDataChannelRead : public UNiagaraDataInterface
 {
 	GENERATED_UCLASS_BODY()
diff --git a/Engine/Plugins/FX/Niagara/Source/Niagara/Internal/DataInterface/NiagaraDataInterfaceDataChannelWrite.h b/Engine/Plugins/FX/Niagara/Source/Niagara/Internal/DataInterface/NiagaraDataInterfaceDataChannelWrite.h
--- a/Engine/Plugins/FX/Niagara/Source/Niagara/Internal/DataInterface/NiagaraDataInterfaceDataChannelWrite.h
+++ b/Engine/Plugins/FX/Niagara/Source/Niagara/Internal/DataInterface/NiagaraDataInterfaceDataChannelWrite.h
@@ -34,7 +34,7 @@ struct FNDIDataChannelWriteCompiledData : public FNDIDataChannelCompiledData
 	bool Init(UNiagaraSystem* System, UNiagaraDataInterfaceDataChannelWrite* OwnerDI);
 };
 
-UCLASS(Experimental, EditInlineNew, Category = "Data Channels", meta = (DisplayName = "Data Channel Writer"))
+UCLASS(Experimental, EditInlineNew, Category = "Data Channels", CollapseCategories, meta = (DisplayName = "Data Channel Writer"))
 class NIAGARA_API UNiagaraDataInterfaceDataChannelWrite : public UNiagaraDataInterface
 {
 	GENERATED_UCLASS_BODY()
diff --git a/Engine/Plugins/FX/Niagara/Source/Niagara/Private/DataInterface/NiagaraDataInterfaceDebugDraw.h b/Engine/Plugins/FX/Niagara/Source/Niagara/Private/DataInterface/NiagaraDataInterfaceDebugDraw.h
--- a/Engine/Plugins/FX/Niagara/Source/Niagara/Private/DataInterface/NiagaraDataInterfaceDebugDraw.h
+++ b/Engine/Plugins/FX/Niagara/Source/Niagara/Private/DataInterface/NiagaraDataInterfaceDebugDraw.h
@@ -12,7 +12,7 @@ struct FNiagaraFunctionSignature;
 struct FNiagaraScriptDataInterfaceCompileInfo;
 struct FVMExternalFunctionBindingInfo;
 
-UCLASS(EditInlineNew, Category = "Debug", meta = (DisplayName = "DebugDraw"))
+UCLASS(EditInlineNew, Category = "Debug", CollapseCategories, meta = (DisplayName = "DebugDraw"))
 class NIAGARA_API UNiagaraDataInterfaceDebugDraw : public UNiagaraDataInterface
 {
 	GENERATED_UCLASS_BODY()
diff --git a/Engine/Plugins/FX/Niagara/Source/Niagara/Private/DataInterface/NiagaraDataInterfaceDynamicMesh.h b/Engine/Plugins/FX/Niagara/Source/Niagara/Private/DataInterface/NiagaraDataInterfaceDynamicMesh.h
--- a/Engine/Plugins/FX/Niagara/Source/Niagara/Private/DataInterface/NiagaraDataInterfaceDynamicMesh.h
+++ b/Engine/Plugins/FX/Niagara/Source/Niagara/Private/DataInterface/NiagaraDataInterfaceDynamicMesh.h
@@ -49,7 +49,7 @@ struct FNiagaraDynamicMeshMaterial
 };
 
 /** Data Interface allowing sampling of a texture */
-UCLASS(Experimental, EditInlineNew, Category = "Mesh", meta = (DisplayName = "Dynamic Mesh"))
+UCLASS(Experimental, EditInlineNew, Category = "Mesh", CollapseCategories, meta = (DisplayName = "Dynamic Mesh"))
 class NIAGARA_API UNiagaraDataInterfaceDynamicMesh : public UNiagaraDataInterface, public INiagaraRenderableMeshInterface
 {
 	GENERATED_UCLASS_BODY()
diff --git a/Engine/Plugins/FX/Niagara/Source/Niagara/Private/DataInterface/NiagaraDataInterfaceEmitterProperties.h b/Engine/Plugins/FX/Niagara/Source/Niagara/Private/DataInterface/NiagaraDataInterfaceEmitterProperties.h
--- a/Engine/Plugins/FX/Niagara/Source/Niagara/Private/DataInterface/NiagaraDataInterfaceEmitterProperties.h
+++ b/Engine/Plugins/FX/Niagara/Source/Niagara/Private/DataInterface/NiagaraDataInterfaceEmitterProperties.h
@@ -15,7 +15,7 @@ struct FVMExternalFunctionBindingInfo;
 /**
 Allows access to various emitter properties that are not part of the simulation data.
 */
-UCLASS(EditInlineNew, Category = "DataInterface", meta=(DisplayName="Emitter Properties"))
+UCLASS(EditInlineNew, Category = "DataInterface", CollapseCategories, meta=(DisplayName="Emitter Properties"))
 class UNiagaraDataInterfaceEmitterProperties : public UNiagaraDataInterface
 {
 	GENERATED_UCLASS_BODY()
diff --git a/Engine/Plugins/FX/Niagara/Source/Niagara/Private/DataInterface/NiagaraDataInterfaceSimCacheReader.h b/Engine/Plugins/FX/Niagara/Source/Niagara/Private/DataInterface/NiagaraDataInterfaceSimCacheReader.h
--- a/Engine/Plugins/FX/Niagara/Source/Niagara/Private/DataInterface/NiagaraDataInterfaceSimCacheReader.h
+++ b/Engine/Plugins/FX/Niagara/Source/Niagara/Private/DataInterface/NiagaraDataInterfaceSimCacheReader.h
@@ -13,7 +13,7 @@ class UNiagaraSimCache;
 /**
 Data interface to read properties from a Niagara Simulation Cache
 */
-UCLASS(EditInlineNew, Category = "DataInterface", meta=(DisplayName="SimCache Reader", Experimental))
+UCLASS(EditInlineNew, Category = "DataInterface", CollapseCategories, meta=(DisplayName="SimCache Reader", Experimental))
 class UNiagaraDataInterfaceSimCacheReader : public UNiagaraDataInterface
 {
 	GENERATED_UCLASS_BODY()
diff --git a/Engine/Plugins/FX/Niagara/Source/Niagara/Private/DataInterface/NiagaraDataInterfaceSimpleCounter.h b/Engine/Plugins/FX/Niagara/Source/Niagara/Private/DataInterface/NiagaraDataInterfaceSimpleCounter.h
--- a/Engine/Plugins/FX/Niagara/Source/Niagara/Private/DataInterface/NiagaraDataInterfaceSimpleCounter.h
+++ b/Engine/Plugins/FX/Niagara/Source/Niagara/Private/DataInterface/NiagaraDataInterfaceSimpleCounter.h
@@ -14,7 +14,7 @@ struct FNiagaraDataInterfaceGeneratedFunction;
 Thread safe counter starts at the initial value on start / reset.
 When operating between CPU & GPU ensure you set the appropriate sync mode.
 */
-UCLASS(EditInlineNew, Category = "Counting", meta = (DisplayName = "Simple Counter"))
+UCLASS(EditInlineNew, Category = "Counting", CollapseCategories, meta = (DisplayName = "Simple Counter"))
 class NIAGARA_API UNiagaraDataInterfaceSimpleCounter : public UNiagaraDataInterfaceRWBase
 {
 	GENERATED_UCLASS_BODY()
diff --git a/Engine/Plugins/FX/Niagara/Source/Niagara/Private/DataInterface/NiagaraDataInterfaceUObjectPropertyReader.h b/Engine/Plugins/FX/Niagara/Source/Niagara/Private/DataInterface/NiagaraDataInterfaceUObjectPropertyReader.h
--- a/Engine/Plugins/FX/Niagara/Source/Niagara/Private/DataInterface/NiagaraDataInterfaceUObjectPropertyReader.h
+++ b/Engine/Plugins/FX/Niagara/Source/Niagara/Private/DataInterface/NiagaraDataInterfaceUObjectPropertyReader.h
@@ -29,7 +29,7 @@ struct FNiagaraUObjectPropertyReaderRemap
 Data interface to read properties from UObjects.
 Rather than having BP tick functions that push data into Niagara this data interface will instead pull them.
 */
-UCLASS(EditInlineNew, Category = "DataInterface", meta=(DisplayName="Object Reader"))
+UCLASS(EditInlineNew, Category = "DataInterface", CollapseCategories, meta=(DisplayName="Object Reader"))
 class UNiagaraDataInterfaceUObjectPropertyReader : public UNiagaraDataInterface
 {
 	GENERATED_UCLASS_BODY()
diff --git a/Engine/Plugins/FX/Niagara/Source/Niagara/Private/DataInterface/NiagaraDataInterfaceVirtualTexture.h b/Engine/Plugins/FX/Niagara/Source/Niagara/Private/DataInterface/NiagaraDataInterfaceVirtualTexture.h
--- a/Engine/Plugins/FX/Niagara/Source/Niagara/Private/DataInterface/NiagaraDataInterfaceVirtualTexture.h
+++ b/Engine/Plugins/FX/Niagara/Source/Niagara/Private/DataInterface/NiagaraDataInterfaceVirtualTexture.h
@@ -12,7 +12,7 @@ struct FNiagaraDataInterfaceGeneratedFunction;
 class URuntimeVirtualTexture;
 
 /** Data Interface allowing sampling of a texture */
-UCLASS(Experimental, EditInlineNew, Category = "Texture", meta = (DisplayName = "Runtime Virtual Texture Sample"))
+UCLASS(Experimental, EditInlineNew, Category = "Texture", CollapseCategories, meta = (DisplayName = "Runtime Virtual Texture Sample"))
 class NIAGARA_API UNiagaraDataInterfaceVirtualTexture : public UNiagaraDataInterface
 {
 	GENERATED_UCLASS_BODY()
diff --git a/Engine/Plugins/FX/Niagara/Source/Niagara/Public/NiagaraDataInterfacePhysicsAsset.h b/Engine/Plugins/FX/Niagara/Source/Niagara/Public/NiagaraDataInterfacePhysicsAsset.h
--- a/Engine/Plugins/FX/Niagara/Source/Niagara/Public/NiagaraDataInterfacePhysicsAsset.h
+++ b/Engine/Plugins/FX/Niagara/Source/Niagara/Public/NiagaraDataInterfacePhysicsAsset.h
@@ -100,7 +100,7 @@ struct FNDIPhysicsAssetData
 };
 
 /** Data Interface for interacting with PhysicsAssets */
-UCLASS(EditInlineNew, Category = "Physics", meta = (DisplayName = "Physics Asset"))
+UCLASS(EditInlineNew, Category = "Physics", CollapseCategories, meta = (DisplayName = "Physics Asset"))
 class NIAGARA_API UNiagaraDataInterfacePhysicsAsset : public UNiagaraDataInterface
 {
 	GENERATED_UCLASS_BODY()
diff --git a/Engine/Plugins/FX/Niagara/Source/NiagaraEditor/Private/ViewModels/Stack/NiagaraStackObject.cpp b/Engine/Plugins/FX/Niagara/Source/NiagaraEditor/Private/ViewModels/Stack/NiagaraStackObject.cpp
--- a/Engine/Plugins/FX/Niagara/Source/NiagaraEditor/Private/ViewModels/Stack/NiagaraStackObject.cpp
+++ b/Engine/Plugins/FX/Niagara/Source/NiagaraEditor/Private/ViewModels/Stack/NiagaraStackObject.cpp
@@ -179,6 +179,7 @@ void UNiagaraStackObject::RefreshChildrenInternal(const TArray<UNiagaraStackEntr
 
 	GatherIssueFromProperties((uint8*)Object, Object->GetClass(), true);
 
+
 	if (GetSystemViewModel()->GetIsForDataProcessingOnly() == false && PropertyRowGenerator.IsValid() == false)
 	{
 		FPropertyEditorModule& PropertyEditorModule = FModuleManager::LoadModuleChecked<FPropertyEditorModule>("PropertyEditor");
@@ -205,6 +206,7 @@ void UNiagaraStackObject::RefreshChildrenInternal(const TArray<UNiagaraStackEntr
 		PropertyRowGenerator->OnRowsRefreshed().AddUObject(this, &UNiagaraStackObject::PropertyRowsRefreshed);
 	}
 
+
 	if (!Object->HasAllFlags(EObjectFlags::RF_Transactional))
 	{
 		NewIssues.Add(FStackIssue(
@@ -228,6 +230,8 @@ void UNiagaraStackObject::RefreshChildrenInternal(const TArray<UNiagaraStackEntr
 			}));
 	}
 
+	TArray<FString> DontCollapseCategoriesOverride;
+
 	// TODO: Handle this in a more generic way.  Maybe add error apis to UNiagaraMergable, or use a UObject interface, or create a
 	// TODO: Possibly move to use INiagaraStackObjectIssueGenerator interface.
 	// data interface specific implementation of UNiagaraStackObject.
@@ -264,6 +268,8 @@ void UNiagaraStackObject::RefreshChildrenInternal(const TArray<UNiagaraStackEntr
 				GetStackEditorDataKey(),
 				false));
 		}
+
+		DontCollapseCategoriesOverride.Add("Errors");
 	}
 
 	const UNiagaraEditorSettings* NiagaraEditorSettings = GetDefault<UNiagaraEditorSettings>();
@@ -290,6 +296,41 @@ void UNiagaraStackObject::RefreshChildrenInternal(const TArray<UNiagaraStackEntr
 			RootTreeNodes = DefaultRootTreeNodes;
 		}
 
+		TArray<TSharedRef<IDetailTreeNode>> CollapsedRootTreeNodes;
+		if(Object->GetClass()->HasAnyClassFlags(CLASS_CollapseCategories))
+		{
+			if(Object->GetClass()->HasMetaData("DontCollapseCategoriesOverride"))
+			{
+				FString MetaDataValue = Object->GetClass()->GetMetaData(TEXT("DontCollapseCategoriesOverride"));
+				TArray<FString> MetaDataDontCollapseOverride;
+				MetaDataValue.ParseIntoArray(MetaDataDontCollapseOverride, TEXT(" "), true);
+				DontCollapseCategoriesOverride.Append(MetaDataDontCollapseOverride);
+			}
+			for (TSharedRef<IDetailTreeNode> RootTreeNode : RootTreeNodes)
+			{
+				if (RootTreeNode->GetNodeType() == EDetailNodeType::Advanced)
+				{
+					continue;
+				}
+				
+				if(RootTreeNode->GetNodeType() == EDetailNodeType::Category && !DontCollapseCategoriesOverride.Contains(RootTreeNode->GetNodeName().ToString()))
+				{
+					TArray<TSharedRef<IDetailTreeNode>> RootChildren;
+					RootTreeNode->GetChildren(RootChildren);
+					CollapsedRootTreeNodes.Append(RootChildren);
+				}
+				else
+				{
+					CollapsedRootTreeNodes.Add(RootTreeNode);
+				}
+			}
+		}
+
+		if(!CollapsedRootTreeNodes.IsEmpty())
+		{
+			RootTreeNodes = CollapsedRootTreeNodes;
+		}
+		
 		for (TSharedRef<IDetailTreeNode> RootTreeNode : RootTreeNodes)
 		{
 			if (RootTreeNode->GetNodeType() == EDetailNodeType::Advanced)
