Niagara - When logging warnings about data interface binding, log the path of the system so that the messages are actually useful.

diff --git a/Engine/Plugins/FX/Niagara/Source/Niagara/Private/NiagaraSystem.cpp b/Engine/Plugins/FX/Niagara/Source/Niagara/Private/NiagaraSystem.cpp
--- a/Engine/Plugins/FX/Niagara/Source/Niagara/Private/NiagaraSystem.cpp
+++ b/Engine/Plugins/FX/Niagara/Source/Niagara/Private/NiagaraSystem.cpp
@@ -1977,6 +1977,7 @@ void UNiagaraSystem::ResolveDIBindings()
 		MessageStore.AddMessage(ResolveDIsMessageId, ResolveDIsMessage);
 		if (MessageStore.IsMessageDismissed(ResolveDIsMessageId) == false)
 		{
+			UE_LOG(LogNiagara, Warning, TEXT("Issues reported while resolving data interfaces for system %s"), *this->GetPathName());
 			for (const FText& ErrorMessage : ErrorMessages)
 			{
 				UE_LOG(LogNiagara, Warning, TEXT("%s"), *ErrorMessage.ToString());
