Fix crash when changing static switch enum.

diff --git a/Engine/Plugins/FX/Niagara/Source/NiagaraEditor/Private/NiagaraNodeStaticSwitch.cpp b/Engine/Plugins/FX/Niagara/Source/NiagaraEditor/Private/NiagaraNodeStaticSwitch.cpp
--- a/Engine/Plugins/FX/Niagara/Source/NiagaraEditor/Private/NiagaraNodeStaticSwitch.cpp
+++ b/Engine/Plugins/FX/Niagara/Source/NiagaraEditor/Private/NiagaraNodeStaticSwitch.cpp
@@ -141,7 +141,11 @@ void UNiagaraNodeStaticSwitch::PreChange(const UUserDefinedEnum* Changed, FEnumE
 
 void UNiagaraNodeStaticSwitch::PostChange(const UUserDefinedEnum* Changed, FEnumEditorUtils::EEnumEditorChangeInfo ChangedType)
 {
-	if(SwitchTypeData.SwitchType == ENiagaraStaticSwitchType::Enum && Changed == SwitchTypeData.Enum && (SwitchTypeData.bAutoRefreshEnabled == true || NiagaraStaticSwitchCVars::bEnabledAutoRefreshOldStaticSwitches))
+	const bool bHasUnderlyingEnumChanged = SwitchTypeData.SwitchType == ENiagaraStaticSwitchType::Enum && Changed == SwitchTypeData.Enum;
+	const bool bShouldAutoRefresh = SwitchTypeData.bAutoRefreshEnabled == true || NiagaraStaticSwitchCVars::bEnabledAutoRefreshOldStaticSwitches;
+
+	// Can't call GetNiagaraGraph on the CDO, this should early out before then on the CDO
+	if(bHasUnderlyingEnumChanged && bShouldAutoRefresh && !GetNiagaraGraph()->IsCompilationCopy())
 	{
 		RefreshFromExternalChanges();
 	}
diff --git a/Engine/Plugins/FX/Niagara/Source/NiagaraEditor/Public/NiagaraGraph.h b/Engine/Plugins/FX/Niagara/Source/NiagaraEditor/Public/NiagaraGraph.h
--- a/Engine/Plugins/FX/Niagara/Source/NiagaraEditor/Public/NiagaraGraph.h
+++ b/Engine/Plugins/FX/Niagara/Source/NiagaraEditor/Public/NiagaraGraph.h
@@ -167,6 +167,9 @@ class UNiagaraGraph : public UEdGraph
 	/** Creates a transient copy of this graph for compilation purposes. */
 	UNiagaraGraph* CreateCompilationCopy(const TArray<ENiagaraScriptUsage>& CompileUsages);
 	void ReleaseCompilationCopy();
+
+	/** Returns if this graph is for compilation purposes only. **/
+	bool IsCompilationCopy() const { return bIsForCompilationOnly; };
 			
 	/** Find the first output node bound to the target usage type.*/
 	class UNiagaraNodeOutput* FindOutputNode(ENiagaraScriptUsage TargetUsageType, FGuid TargetUsageId = FGuid()) const;
