- Add Particle Sprite Rotation material function

diff --git a/Engine/Plugins/Enterprise/GLTFExporter/Source/GLTFExporter/Private/Converters/GLTFProxyMaterialCompiler.h b/Engine/Plugins/Enterprise/GLTFExporter/Source/GLTFExporter/Private/Converters/GLTFProxyMaterialCompiler.h
--- a/Engine/Plugins/Enterprise/GLTFExporter/Source/GLTFExporter/Private/Converters/GLTFProxyMaterialCompiler.h
+++ b/Engine/Plugins/Enterprise/GLTFExporter/Source/GLTFExporter/Private/Converters/GLTFProxyMaterialCompiler.h
@@ -169,6 +169,11 @@ public:
 		return Compiler->ParticleSize();
 	}
 
+	virtual int32 ParticleSpriteRotation() override
+	{
+		return Compiler->ParticleSpriteRotation();
+	}
+
 	virtual int32 VertexInterpolator(uint32 InterpolatorIndex) override
 	{
 		return Compiler->VertexInterpolator(InterpolatorIndex);
diff --git a/Engine/Plugins/FX/Niagara/Shaders/Private/NiagaraSpriteVertexFactory.ush b/Engine/Plugins/FX/Niagara/Shaders/Private/NiagaraSpriteVertexFactory.ush
--- a/Engine/Plugins/FX/Niagara/Shaders/Private/NiagaraSpriteVertexFactory.ush
+++ b/Engine/Plugins/FX/Niagara/Shaders/Private/NiagaraSpriteVertexFactory.ush
@@ -15,6 +15,7 @@
 #define USE_PARTICLE_TIME (NEEDS_PARTICLE_TIME)
 #define USE_PARTICLE_SIZE (NEEDS_PARTICLE_SIZE)
 #define USE_PARTICLE_RANDOM (NEEDS_PARTICLE_RANDOM)
+#define USE_PARTICLE_SPRITE_ROTATION (NEEDS_PARTICLE_SPRITE_ROTATION)
 
 // Must keep these in sync with ENiagaraSpriteAlignment in NiagaraSpriteRendererProperties.h
 /** Only Particles.SpriteRotation and FacingMode impact the alignment of the particle.*/
@@ -146,6 +147,10 @@ struct FVertexFactoryInterpolantsVSToPS
 	nointerpolation float2 ParticleSize : PARTICLE_SIZE;
 #endif
 
+#if USE_PARTICLE_SPRITE_ROTATION
+	nointerpolation float ParticleSpriteRotation : PARTICLE_SPRITE_ROTATION;
+#endif
+
 #if USE_PARTICLE_RANDOM
 	nointerpolation float ParticleRandom : PARTICLE_RANDOM;
 #endif
@@ -187,7 +192,7 @@ struct FVertexFactoryIntermediates
 	/** The particle worldspace position */
 	FLWCVector3 Position;
 	
-	float4 SizeRotSubImage;
+	float2 Size;
 	
 	float3 CustomFacingVector;
 	float3 CustomAlignmentVector;
@@ -218,10 +223,16 @@ struct FVertexFactoryIntermediates
 	/** Size of the particle */
 	float2 ParticleSize;
 
+
+#if USE_PARTICLE_SPRITE_ROTATION
+	float ParticleSpriteRotation;
+#endif
+
 #if USE_PARTICLE_RANDOM
 	/** Particle Random value */
 	float ParticleRandom;
 #endif
+
 	/** Cached primitive and instance data */
 	FSceneDataIntermediates SceneData;
 };
@@ -344,6 +355,12 @@ FMaterialPixelParameters GetMaterialPixelParameters(FVertexFactoryInterpolantsVS
 	Result.Particle.Size = Interpolants.ParticleSize;
 #endif
 
+#if USE_PARTICLE_SPRITE_ROTATION
+	Result.Particle.SpriteRotation = Interpolants.ParticleSpriteRotation;
+#else
+	Result.Particle.SpriteRotation = 0.0f;
+#endif
+
 #if USE_PARTICLE_RANDOM
 	Result.Particle.Random = Interpolants.ParticleRandom;
 #else
@@ -400,6 +417,12 @@ FMaterialVertexParameters GetMaterialVertexParameters(FVertexFactoryInput Input,
 	Result.Particle.RelativeTime = Intermediates.RelativeTime;
 #endif
 
+#if USE_PARTICLE_SPRITE_ROTATION
+	Result.Particle.SpriteRotation = Intermediates.ParticleSpriteRotation;
+#else
+	Result.Particle.SpriteRotation = 0.0f;
+#endif
+
 #if USE_PARTICLE_RANDOM
 	Result.Particle.Random = Intermediates.ParticleRandom;
 #else
@@ -531,8 +554,7 @@ void GetTangents(FVertexFactoryInput Input, FLWCVector3 ParticlePosition, float
 	// Determine the angle of rotation.
 	float SinRotation;
 	float CosRotation;
-	const float SpriteRotation = Rotation * NiagaraSpriteVF.RotationScale + NiagaraSpriteVF.RotationBias;
-	sincos(SpriteRotation, SinRotation, CosRotation);
+	sincos(Rotation, SinRotation, CosRotation);
 
 	float3 SR = Right * SinRotation;
 	float3 SU = Up * SinRotation;
@@ -691,13 +713,15 @@ float2 GetNiagaraPreviousParticleSize(uint InstanceID)
 float GetNiagaraParticleRotation(uint InstanceID)
 {
 	// See FNiagaraConstants for a synchronized value...
-	return (SafeGetFloat(NiagaraSpriteVF.RotationDataOffset, InstanceID, NiagaraSpriteVF.DefaultRotation) / 180.0f) * PI;
+	const float Rotation = (SafeGetFloat(NiagaraSpriteVF.RotationDataOffset, InstanceID, NiagaraSpriteVF.DefaultRotation) / 180.0f) * PI;
+	return Rotation * NiagaraSpriteVF.RotationScale + NiagaraSpriteVF.RotationBias;
 }
 
 float GetNiagaraPreviousParticleRotation(uint InstanceID)
 {
 	// See FNiagaraConstants for a synchronized value...
-	return (SafeGetFloat(NiagaraSpriteVF.PrevRotationDataOffset, InstanceID, NiagaraSpriteVF.DefaultPrevRotation) / 180.0f) * PI;
+	const float PrevRotation = (SafeGetFloat(NiagaraSpriteVF.PrevRotationDataOffset, InstanceID, NiagaraSpriteVF.DefaultPrevRotation) / 180.0f) * PI;
+	return PrevRotation * NiagaraSpriteVF.RotationScale + NiagaraSpriteVF.RotationBias;
 }
 
 float4 GetNiagaraParticleColor(uint InstanceID)
@@ -818,7 +842,7 @@ FVertexFactoryIntermediates GetVertexFactoryIntermediates(FVertexFactoryInput In
 	
 	Intermediates.SceneData = VF_GPUSCENE_GET_INTERMEDIATES(Input);
 	Intermediates.Position = ParticlePosition;
-	Intermediates.SizeRotSubImage = float4(ParticleSize.x, ParticleSize.y, ParticleRotation, SubImageIndex);
+	Intermediates.Size = ParticleSize;
 	Intermediates.Color	= GetNiagaraParticleColor(ParticleID);
 	Intermediates.CustomFacingVector = CustomFacing;
 	Intermediates.CustomAlignmentVector = SafeNormalize(GetNiagaraParticleAlignmentVector(ParticleID));
@@ -837,6 +861,10 @@ FVertexFactoryIntermediates GetVertexFactoryIntermediates(FVertexFactoryInput In
 	Intermediates.DynamicParameter3 = GetNiagaraParticleDynamicParameters3(ParticleID);
 #endif
 
+#if USE_PARTICLE_SPRITE_ROTATION
+	Intermediates.ParticleSpriteRotation = ParticleRotation;
+#endif
+
 	// Tangents.
 	{
 		const float3 ParticleDirection = SafeNormalizeWithFallback(ParticleVelocity, float3(0,0,1));
@@ -853,7 +881,7 @@ FVertexFactoryIntermediates GetVertexFactoryIntermediates(FVertexFactoryInput In
 	ComputeBillboardUVs(Input, ParticleSize, SubImageIndex, UVForPosition, UVForTexturing, UVForTexturingUnflipped);
 
 	UVForTexturing.xy = float2(0.5f, 0.5f) + (UVForTexturing.xy - float2(0.5f, 0.5f) ) *  UVScale.xy; 
-	const float2 Size = abs(Intermediates.SizeRotSubImage.xy);
+	const float2 Size = abs(Intermediates.Size.xy);
 
 	// When enabled try to main a pixel sized sprite and optionally fade particles based on coverage
 	// In reality this is more like 2x2 coverage to ensure we always hit 1 pixel, the alternative would be to snap
@@ -987,7 +1015,7 @@ float4 VertexFactoryGetRasterizedWorldPosition(FVertexFactoryInput Input, FVerte
 	// This avoids opaque objects intersecting the particle from causing clipping artifacts due to depth testing
 	// The downside is that the particle will clip the near plane sooner
 
-	float Radius = .5f * min(Intermediates.SizeRotSubImage.x, Intermediates.SizeRotSubImage.y);
+	float Radius = .5f * min(Intermediates.Size.x, Intermediates.Size.y);
 	return ReprojectPosition(TranslatedWorldPosition, Radius);
 #else
 	return TranslatedWorldPosition;
@@ -1077,11 +1105,14 @@ FVertexFactoryInterpolantsVSToPS VertexFactoryGetInterpolantsVSToPS(FVertexFacto
 	Interpolants.ParticleSize = Intermediates.ParticleSize;
 #endif
 
+#if USE_PARTICLE_SPRITE_ROTATION
+	Interpolants.ParticleSpriteRotation = Intermediates.ParticleSpriteRotation;
+#endif
+
 #if USE_PARTICLE_RANDOM
 	Interpolants.ParticleRandom = Intermediates.ParticleRandom;
 #endif
 
-
 #if INSTANCED_STEREO
 	Interpolants.EyeIndex = 0;
 #endif
@@ -1173,6 +1204,10 @@ FVertexFactoryRayTracingInterpolants VertexFactoryInterpolate(FVertexFactoryRayT
 	INTERPOLATE_MEMBER(InterpolantsVSToPS.ParticleSize);
 #endif
 
+#if USE_PARTICLE_SPRITE_ROTATION
+	INTERPOLATE_MEMBER(InterpolantsVSToPS.ParticleSpriteRotation);
+#endif
+
 #if USE_PARTICLE_RANDOM
 	INTERPOLATE_MEMBER(InterpolantsVSToPS.ParticleRandom);
 #endif
diff --git a/Engine/Shaders/Private/MaterialTemplate.ush b/Engine/Shaders/Private/MaterialTemplate.ush
--- a/Engine/Shaders/Private/MaterialTemplate.ush
+++ b/Engine/Shaders/Private/MaterialTemplate.ush
@@ -233,6 +233,9 @@ struct FMaterialParticleParameters
 
 	/** The size of the particle. */
 	float2 Size;
+
+	/** The sprite rotation */
+	float SpriteRotation;
 };
 
 float4 GetDynamicParameter(FMaterialParticleParameters Parameters, float4 Default, int ParameterIndex=0)
diff --git a/Engine/Shaders/Private/ParticleGPUSpriteVertexFactory.ush b/Engine/Shaders/Private/ParticleGPUSpriteVertexFactory.ush
--- a/Engine/Shaders/Private/ParticleGPUSpriteVertexFactory.ush
+++ b/Engine/Shaders/Private/ParticleGPUSpriteVertexFactory.ush
@@ -15,6 +15,7 @@
 #define USE_PARTICLE_TIME (NEEDS_PARTICLE_TIME)
 #define USE_SPHERICAL_NORMALS (GENERATE_SPHERICAL_PARTICLE_NORMALS && USE_PARTICLE_POSITION)
 #define USE_PARTICLE_SIZE (NEEDS_PARTICLE_SIZE)
+#define USE_PARTICLE_SPRITE_ROTATION (NEEDS_PARTICLE_SPRITE_ROTATION)
 
 /** Buffer containing particle indices. */
 Buffer<float2> ParticleIndices;
@@ -95,6 +96,10 @@ struct FVertexFactoryInterpolantsVSToPS
 	float2 ParticleSize			: PARTICLE_SIZE;
 #endif
 
+#if USE_PARTICLE_SPRITE_ROTATION
+	float ParticleSpriteRotation : PARTICLE_SPRITE_ROTATION;
+#endif
+
 #if INSTANCED_STEREO
 	nointerpolation uint EyeIndex : PACKED_EYE_INDEX;
 #endif
@@ -139,6 +144,9 @@ struct FVertexFactoryIntermediates
 	float ParticleRadius;
 	/** Size of the particle. */
 	float2 ParticleSize;
+#if USE_PARTICLE_SPRITE_ROTATION
+	float ParticleSpriteRotation;
+#endif
 	/** Cached primitive and instance data */
 	FSceneDataIntermediates SceneData;
 };
@@ -374,6 +382,11 @@ FVertexFactoryIntermediates GetVertexFactoryIntermediates( FVertexFactoryInput I
 
 	Intermediates.ParticleRadius = .5f * min(Size.x, Size.y);
 	Intermediates.ParticleSize = Size;
+
+#if USE_PARTICLE_SPRITE_ROTATION
+	Intermediates.ParticleSpriteRotation = Rotation;
+#endif
+
 	return Intermediates;
 }
 
@@ -476,6 +489,9 @@ FMaterialPixelParameters GetMaterialPixelParameters( FVertexFactoryInterpolantsV
 	Result.Particle.Size = Interpolants.ParticleSize;
 #endif
 
+#if USE_PARTICLE_SPRITE_ROTATION
+	Result.Particle.SpriteRotation = Interpolants.ParticleSpriteRotation;
+#endif
 	return Result;
 }
 
@@ -551,6 +567,11 @@ FMaterialVertexParameters GetMaterialVertexParameters(FVertexFactoryInput Input,
 #if USE_PARTICLE_SIZE
 	Result.Particle.Size = Intermediates.ParticleSize;
 #endif
+
+#if USE_PARTICLE_SPRITE_ROTATION
+	Result.Particle.SpriteRotation = Intermediates.ParticleSpriteRotation;
+#endif
+
 	return Result;
 }
 
@@ -675,6 +696,10 @@ FVertexFactoryInterpolantsVSToPS VertexFactoryGetInterpolantsVSToPS( FVertexFact
 	Interpolants.ParticleSize = Intermediates.ParticleSize;
 #endif
 
+#if USE_PARTICLE_SPRITE_ROTATION
+	Interpolants.ParticleSpriteRotation = Intermediates.ParticleSpriteRotation;
+#endif
+
 #if INSTANCED_STEREO
 	Interpolants.EyeIndex = 0;
 #endif
diff --git a/Engine/Shaders/Private/ParticleSpriteVertexFactory.ush b/Engine/Shaders/Private/ParticleSpriteVertexFactory.ush
--- a/Engine/Shaders/Private/ParticleSpriteVertexFactory.ush
+++ b/Engine/Shaders/Private/ParticleSpriteVertexFactory.ush
@@ -13,6 +13,7 @@
 #define USE_PARTICLE_VELOCITY (NEEDS_PARTICLE_VELOCITY)
 #define USE_PARTICLE_TIME (NEEDS_PARTICLE_TIME)
 #define USE_PARTICLE_SIZE (NEEDS_PARTICLE_SIZE)
+#define USE_PARTICLE_SPRITE_ROTATION (NEEDS_PARTICLE_SPRITE_ROTATION)
 
 // We don't compute a tangent basis in ES3.1 to save shader instructions.
 #define NEEDS_TANGENT_BASIS (FEATURE_LEVEL >= FEATURE_LEVEL_SM4)
@@ -85,6 +86,10 @@ struct FVertexFactoryInterpolantsVSToPS
 	float2 ParticleSize : PARTICLE_SIZE;
 #endif
 
+#if USE_PARTICLE_SPRITE_ROTATION
+	float ParticleSpriteRotation : PARTICLE_SPRITE_ROTATION;
+#endif
+
 #if INSTANCED_STEREO
 	nointerpolation uint EyeIndex : PACKED_EYE_INDEX;
 #endif
@@ -129,6 +134,9 @@ struct FVertexFactoryIntermediates
 	float3x3 TangentToLocal;
 	/** Size of the particle */
 	float2 ParticleSize;
+#if USE_PARTICLE_SPRITE_ROTATION
+	float ParticleSpriteRotation : PARTICLE_SPRITE_ROTATION;
+#endif
 	/** Cached primitive and instance data */
 	FSceneDataIntermediates SceneData;
 };
@@ -243,6 +251,10 @@ FMaterialPixelParameters GetMaterialPixelParameters(FVertexFactoryInterpolantsVS
 	Result.Particle.Size = Interpolants.ParticleSize;
 #endif
 
+#if USE_PARTICLE_SPRITE_ROTATION
+	Result.Particle.SpriteRotation = Interpolants.ParticleSpriteRotation;
+#endif
+
 	Result.TwoSidedSign = 1;
 	return Result;
 }
@@ -304,6 +316,11 @@ FMaterialVertexParameters GetMaterialVertexParameters(FVertexFactoryInput Input,
 #if USE_PARTICLE_SIZE
 	Result.Particle.Size = Intermediates.ParticleSize;
 #endif
+
+#if USE_PARTICLE_SPRITE_ROTATION
+	Result.Particle.SpriteRotation = Intermediates.ParticleSpriteRotation;
+#endif
+
 	return Result;
 }
 
@@ -312,7 +329,7 @@ float3 SafeNormalize(float3 V)
 	return V * rsqrt(max(dot(V,V),0.00000001));
 }
 
-void GetTangents(FVertexFactoryInput Input, FLWCVector3 WorldPosition, FLWCVector3 OldWorldPosition, out float3 OutRight, out float3 OutUp)
+void GetTangents(FLWCVector3 WorldPosition, FLWCVector3 OldWorldPosition, float SpriteRotation, out float3 OutRight, out float3 OutUp)
 {
 	// Select camera up/right vectors.
 	float3 ResolvedViewRight = lerp(ResolvedView.ViewRight, ResolvedView.HMDViewNoRollRight, SpriteVF.RemoveHMDRoll);
@@ -369,9 +386,6 @@ void GetTangents(FVertexFactoryInput Input, FLWCVector3 WorldPosition, FLWCVecto
 	// Determine the angle of rotation.
 	float SinRotation; // = 0
 	float CosRotation; // = 1
-	const float RotationScale = SpriteVF.RotationScale;
-	const float RotationBias = SpriteVF.RotationBias;
-	const float SpriteRotation = Input.SizeRotSubImage.z * RotationScale + RotationBias;
 	sincos(SpriteRotation, SinRotation, CosRotation);
 
 	// Rotate the sprite to determine final tangents.
@@ -480,9 +494,11 @@ FVertexFactoryIntermediates GetVertexFactoryIntermediates(FVertexFactoryInput In
 	FLWCVector3 ParticleOldWorldPosition = LWCAdd(LWCMultiply(Input.OldPosition.xyz, GetPrimitiveDataFromUniformBuffer().LocalToWorld), LWCTileOffset);
 	Intermediates.ParticleWorldPosition = ParticleWorldPosition;
 
+	const float SpriteRotation = Input.SizeRotSubImage.z * SpriteVF.RotationScale + SpriteVF.RotationBias;
+
 	// Tangents.
 	float3 Right,Up;
-	GetTangents(Input, ParticleWorldPosition, ParticleOldWorldPosition, Right, Up);
+	GetTangents(ParticleWorldPosition, ParticleOldWorldPosition, SpriteRotation, Right, Up);
 	Intermediates.TangentUp = Up;
 	Intermediates.TangentRight = Right;
 
@@ -563,6 +579,10 @@ FVertexFactoryIntermediates GetVertexFactoryIntermediates(FVertexFactoryInput In
 	Intermediates.ParticleSize = Size;
 #endif
 
+#if USE_PARTICLE_SPRITE_ROTATION
+	Intermediates.ParticleSpriteRotation = SpriteRotation;
+#endif
+
 	return Intermediates;
 }
 
@@ -668,6 +688,10 @@ FVertexFactoryInterpolantsVSToPS VertexFactoryGetInterpolantsVSToPS(FVertexFacto
 	Interpolants.ParticleSize = Intermediates.ParticleSize;
 #endif
 
+#if USE_PARTICLE_SPRITE_ROTATION
+	Interpolants.ParticleSpriteRotation = Intermediates.ParticleSpriteRotation;
+#endif
+
 #if INSTANCED_STEREO
 	Interpolants.EyeIndex = 0;
 #endif
diff --git a/Engine/Source/Developer/MaterialBaking/Private/ExportMaterialProxy.h b/Engine/Source/Developer/MaterialBaking/Private/ExportMaterialProxy.h
--- a/Engine/Source/Developer/MaterialBaking/Private/ExportMaterialProxy.h
+++ b/Engine/Source/Developer/MaterialBaking/Private/ExportMaterialProxy.h
@@ -96,6 +96,11 @@ struct FExportMaterialCompiler : public FProxyMaterialCompiler
 		return Compiler->Constant2(0.0f, 0.0f);
 	}
 
+	virtual int32 ParticleSpriteRotation() override
+	{
+		return Compiler->Constant2(0.0f, 0.0f);
+	}
+
 	virtual int32 CameraVector() override
 	{
 		// By returning vertex normal instead of a constant vector (like up), we ensure materials (with fresnel for example) are more correctly baked using custom mesh data.
diff --git a/Engine/Source/Developer/MaterialUtilities/Private/MaterialUtilities.cpp b/Engine/Source/Developer/MaterialUtilities/Private/MaterialUtilities.cpp
--- a/Engine/Source/Developer/MaterialUtilities/Private/MaterialUtilities.cpp
+++ b/Engine/Source/Developer/MaterialUtilities/Private/MaterialUtilities.cpp
@@ -310,6 +310,11 @@ struct FExportMaterialCompiler : public FProxyMaterialCompiler
 		return Compiler->Constant2(0.0f,0.0f);
 	}
 
+	virtual int32 ParticleSpriteRotation() override
+	{
+		return Compiler->Constant2(0.0f, 0.0f);
+	}
+
 	virtual int32 ObjectRadius() override
 	{
 		return Compiler->Constant(500);
diff --git a/Engine/Source/Editor/UnrealEd/Private/Lightmass/LightmassRender.cpp b/Engine/Source/Editor/UnrealEd/Private/Lightmass/LightmassRender.cpp
--- a/Engine/Source/Editor/UnrealEd/Private/Lightmass/LightmassRender.cpp
+++ b/Engine/Source/Editor/UnrealEd/Private/Lightmass/LightmassRender.cpp
@@ -109,6 +109,11 @@ struct FLightmassMaterialCompiler : public FProxyMaterialCompiler
 		return Compiler->Constant2(0.0f,0.0f);
 	}
 
+	virtual int32 ParticleSpriteRotation() override
+	{
+		return Compiler->Constant2(0.0f, 0.0f);
+	}
+
 	virtual int32 WorldPosition(EWorldPositionIncludedOffsets WorldPositionIncludedOffsets) override
 	{
 		//UE_LOG(LogLightmassRender, Log, TEXT("Lightmass material compiler has encountered WorldPosition... Forcing constant (0.0f,0.0f,0.0f)."));
diff --git a/Engine/Source/Runtime/Engine/Classes/Materials/MaterialExpressionParticleSpriteRotation.h b/Engine/Source/Runtime/Engine/Classes/Materials/MaterialExpressionParticleSpriteRotation.h
new file mode 100644
--- /dev/null
+++ b/Engine/Source/Runtime/Engine/Classes/Materials/MaterialExpressionParticleSpriteRotation.h
@@ -0,0 +1,30 @@
+// Copyright Epic Games, Inc. All Rights Reserved.
+
+/*==============================================================================
+	MaterialExpressionParticleSize: Exposes the size of a particle to
+		the material editor.
+==============================================================================*/
+
+#pragma once
+
+#include "CoreMinimal.h"
+#include "UObject/ObjectMacros.h"
+#include "Materials/MaterialExpression.h"
+#include "MaterialExpressionParticleSpriteRotation.generated.h"
+
+UCLASS(collapsecategories, hidecategories=Object)
+class UMaterialExpressionParticleSpriteRotation : public UMaterialExpression
+{
+	GENERATED_UCLASS_BODY()
+
+
+	//~ Begin UMaterialExpression Interface
+#if WITH_EDITOR
+	virtual int32 Compile(class FMaterialCompiler* Compiler, int32 OutputIndex) override;
+	virtual void GetCaption(TArray<FString>& OutCaptions) const override;
+#endif
+	//~ End UMaterialExpression Interface
+};
+
+
+
diff --git a/Engine/Source/Runtime/Engine/Private/Materials/HLSLMaterialTranslator.cpp b/Engine/Source/Runtime/Engine/Private/Materials/HLSLMaterialTranslator.cpp
--- a/Engine/Source/Runtime/Engine/Private/Materials/HLSLMaterialTranslator.cpp
+++ b/Engine/Source/Runtime/Engine/Private/Materials/HLSLMaterialTranslator.cpp
@@ -354,6 +354,7 @@ FHLSLMaterialTranslator::FHLSLMaterialTranslator(FMaterial* InMaterial,
 ,	bUsesSpeedTree(false)
 ,	bNeedsWorldPositionExcludingShaderOffsets(false)
 ,	bNeedsParticleSize(false)
+,	bNeedsParticleSpriteRotation(false)
 ,	bNeedsSceneTexturePostProcessInputs(false)
 ,	bUsesSkyAtmosphere(false)
 ,	bUsesVertexColor(false)
@@ -1996,6 +1997,11 @@ void FHLSLMaterialTranslator::GetMaterialEnvironment(EShaderPlatform InPlatform,
 		OutEnvironment.SetDefine(TEXT("NEEDS_PARTICLE_SIZE"), TEXT("1"));
 	}
 
+	if (bNeedsParticleSpriteRotation)
+	{
+		OutEnvironment.SetDefine(TEXT("NEEDS_PARTICLE_SPRITE_ROTATION"), TEXT("1"));
+	}
+
 	if (MaterialCompilationOutput.bNeedsSceneTextures)
 	{
 		OutEnvironment.SetDefine(TEXT("NEEDS_SCENE_TEXTURES"), TEXT("1"));
@@ -5706,6 +5712,16 @@ int32 FHLSLMaterialTranslator::ParticleSize()
 	return AddInlinedCodeChunkZeroDeriv(MCT_Float2,TEXT("Parameters.Particle.Size"));
 }
 
+int32 FHLSLMaterialTranslator::ParticleSpriteRotation()
+{
+	if (ShaderFrequency != SF_Vertex && ShaderFrequency != SF_Pixel && ShaderFrequency != SF_Compute)
+	{
+		return NonVertexOrPixelShaderExpressionError();
+	}
+	bNeedsParticleSpriteRotation = true;
+	return AddInlinedCodeChunkZeroDeriv(MCT_Float2, TEXT("float2(Parameters.Particle.SpriteRotation, Parameters.Particle.SpriteRotation * 57.2957795131f)"));
+}
+
 int32 FHLSLMaterialTranslator::WorldPosition(EWorldPositionIncludedOffsets WorldPositionIncludedOffsets)
 {
 	FString FunctionNamePattern;
diff --git a/Engine/Source/Runtime/Engine/Private/Materials/HLSLMaterialTranslator.h b/Engine/Source/Runtime/Engine/Private/Materials/HLSLMaterialTranslator.h
--- a/Engine/Source/Runtime/Engine/Private/Materials/HLSLMaterialTranslator.h
+++ b/Engine/Source/Runtime/Engine/Private/Materials/HLSLMaterialTranslator.h
@@ -364,6 +364,8 @@ protected:
 	uint32 bNeedsWorldPositionExcludingShaderOffsets : 1;
 	/** true if the material needs particle size. */
 	uint32 bNeedsParticleSize : 1;
+	/** true if the material needs particle sprite rotation. */
+	uint32 bNeedsParticleSpriteRotation : 1;
 	/** true if any scene texture expressions are reading from post process inputs */
 	uint32 bNeedsSceneTexturePostProcessInputs : 1;
 	/** true if any atmospheric fog expressions are used */
@@ -896,6 +898,7 @@ protected:
 	virtual int32 ParticleDirection() override;
 	virtual int32 ParticleSpeed() override;
 	virtual int32 ParticleSize() override;
+	virtual int32 ParticleSpriteRotation() override;
 
 	virtual int32 WorldPosition(EWorldPositionIncludedOffsets WorldPositionIncludedOffsets) override;
 
diff --git a/Engine/Source/Runtime/Engine/Private/Materials/MaterialExpressions.cpp b/Engine/Source/Runtime/Engine/Private/Materials/MaterialExpressions.cpp
--- a/Engine/Source/Runtime/Engine/Private/Materials/MaterialExpressions.cpp
+++ b/Engine/Source/Runtime/Engine/Private/Materials/MaterialExpressions.cpp
@@ -166,6 +166,7 @@
 #include "Materials/MaterialExpressionParticleRadius.h"
 #include "Materials/MaterialExpressionParticleRelativeTime.h"
 #include "Materials/MaterialExpressionParticleSize.h"
+#include "Materials/MaterialExpressionParticleSpriteRotation.h"
 #include "Materials/MaterialExpressionParticleSpeed.h"
 #include "Materials/MaterialExpressionPerInstanceFadeAmount.h"
 #include "Materials/MaterialExpressionPerInstanceRandom.h"
@@ -20671,6 +20672,50 @@ void UMaterialExpressionParticleSize::GetCaption(TArray<FString>& OutCaptions) c
 }
 #endif // WITH_EDITOR
 
+/*------------------------------------------------------------------------------
+	Particle sprite rotation material expression.
+------------------------------------------------------------------------------*/
+UMaterialExpressionParticleSpriteRotation::UMaterialExpressionParticleSpriteRotation(const FObjectInitializer& ObjectInitializer)
+	: Super(ObjectInitializer)
+{
+#if WITH_EDITORONLY_DATA
+	// Structure to hold one-time initialization
+	struct FConstructorStatics
+	{
+		FText NAME_Particles;
+		FText NAME_Constants;
+		FConstructorStatics()
+			: NAME_Particles(LOCTEXT("Particles", "Particles"))
+			, NAME_Constants(LOCTEXT("Constants", "Constants"))
+		{
+		}
+	};
+	static FConstructorStatics ConstructorStatics;
+
+	MenuCategories.Add(ConstructorStatics.NAME_Particles);
+	MenuCategories.Add(ConstructorStatics.NAME_Constants);
+
+	Outputs.Reset();
+	Outputs.Emplace(TEXT("Rad"), 1, 1, 0, 0, 0);
+	Outputs.Emplace(TEXT("Deg"), 1, 0, 1, 0, 0);
+
+	bShaderInputData = true;
+	bShowOutputNameOnPin = true;
+#endif
+}
+
+#if WITH_EDITOR
+int32 UMaterialExpressionParticleSpriteRotation::Compile(class FMaterialCompiler* Compiler, int32 OutputIndex)
+{
+	return Compiler->ParticleSpriteRotation();
+}
+
+void UMaterialExpressionParticleSpriteRotation::GetCaption(TArray<FString>& OutCaptions) const
+{
+	OutCaptions.Add(TEXT("Particle Sprite Rotation"));
+}
+#endif // WITH_EDITOR
+
 /*------------------------------------------------------------------------------
 	Atmospheric fog material expression.
 ------------------------------------------------------------------------------*/
diff --git a/Engine/Source/Runtime/Engine/Public/MaterialCompiler.h b/Engine/Source/Runtime/Engine/Public/MaterialCompiler.h
--- a/Engine/Source/Runtime/Engine/Public/MaterialCompiler.h
+++ b/Engine/Source/Runtime/Engine/Public/MaterialCompiler.h
@@ -353,6 +353,7 @@ public:
 	virtual int32 ParticleDirection() = 0;
 	virtual int32 ParticleSpeed() = 0;
 	virtual int32 ParticleSize() = 0;
+	virtual int32 ParticleSpriteRotation() = 0;
 
 	virtual int32 DynamicBranch(int32 Condition, int32 A, int32 B) = 0;
 	virtual int32 If(int32 A,int32 B,int32 AGreaterThanB,int32 AEqualsB,int32 ALessThanB,int32 Threshold) = 0;
