Niagara - Fix an issue with the new DI binding resolution where the emitter name on the resolved data interface would not be correct which was breaking self emitter references in some cases.

diff --git a/Engine/Plugins/FX/Niagara/Source/Niagara/Private/NiagaraDataInterfaceEmitterBinding.cpp b/Engine/Plugins/FX/Niagara/Source/Niagara/Private/NiagaraDataInterfaceEmitterBinding.cpp
--- a/Engine/Plugins/FX/Niagara/Source/Niagara/Private/NiagaraDataInterfaceEmitterBinding.cpp
+++ b/Engine/Plugins/FX/Niagara/Source/Niagara/Private/NiagaraDataInterfaceEmitterBinding.cpp
@@ -35,7 +35,7 @@ FNiagaraEmitterInstance* FNiagaraDataInterfaceEmitterBinding::Resolve(FNiagaraSy
 						[&DataInterface](const FNiagaraScriptResolvedDataInterfaceInfo& ResolvedDataInterface) { return ResolvedDataInterface.ResolvedDataInterface == DataInterface; });
 					if (MatchingResolvedDataInterface != nullptr)
 					{
-						SourceEmitterName = MatchingResolvedDataInterface->EmitterName;
+						SourceEmitterName = MatchingResolvedDataInterface->ResolvedSourceEmitterName;
 						break;
 					}
 				}
diff --git a/Engine/Plugins/FX/Niagara/Source/Niagara/Private/NiagaraResolveDIHelpers.cpp b/Engine/Plugins/FX/Niagara/Source/Niagara/Private/NiagaraResolveDIHelpers.cpp
--- a/Engine/Plugins/FX/Niagara/Source/Niagara/Private/NiagaraResolveDIHelpers.cpp
+++ b/Engine/Plugins/FX/Niagara/Source/Niagara/Private/NiagaraResolveDIHelpers.cpp
@@ -16,11 +16,17 @@
 
 namespace FNiagaraResolveDIHelpers
 {
+	struct FDataInterfaceSourceEmitterNamePair
+	{
+		UNiagaraDataInterface* DataInterface;
+		FString SourceEmitterName;
+	};
+
 	void CollectDIBindingsAndAssignmentsForScript(
 		const UNiagaraScript* TargetScript,
 		const UNiagaraSystem* System,
 		const FString& EmitterName,
-		TMap<FNiagaraVariableBase, UNiagaraDataInterface*>& OutVariableAssignmentMap,
+		TMap<FNiagaraVariableBase, FDataInterfaceSourceEmitterNamePair>& OutVariableAssignmentMap,
 		TMap<FNiagaraVariableBase, FNiagaraVariableBase>& OutVariableBindingMap,
 		TArray<FText>& OutErrorMessages)
 	{
@@ -73,21 +79,23 @@ namespace FNiagaraResolveDIHelpers
 				else
 				{
 					// Handle assignments
-					UNiagaraDataInterface** CurrentAssignmentPtr = OutVariableAssignmentMap.Find(WriteVariable);
-					if (CurrentAssignmentPtr != nullptr)
+					FDataInterfaceSourceEmitterNamePair* CurrentAssignment = OutVariableAssignmentMap.Find(WriteVariable);
+					if (CurrentAssignment != nullptr)
 					{
-						if (*CurrentAssignmentPtr != CachedDefaultDataInterface.DataInterface)
+						if (CurrentAssignment->DataInterface != CachedDefaultDataInterface.DataInterface)
 						{
 							OutErrorMessages.Add(FText::Format(
 								LOCTEXT("MultipleAssignmentsFormat", "A data interface parameter was thh target of an assignment multiple times in a single system.  The data interface used in the simulation may be incorrect.  Target Parameter: {0} First Assignment: {1} Current Assignment: {2}"),
 								FText::FromName(WriteVariable.GetName()),
-								FText::FromString((*CurrentAssignmentPtr)->GetName()),
-								FText::FromString((CachedDefaultDataInterface.DataInterface->GetName()))));
+								FText::FromString(CurrentAssignment->DataInterface->GetName()),
+								FText::FromString(CachedDefaultDataInterface.DataInterface->GetName())));
 						}
 					}
 					else
 					{
-						OutVariableAssignmentMap.Add(WriteVariable, CachedDefaultDataInterface.DataInterface);
+						FDataInterfaceSourceEmitterNamePair& NewAssignment = OutVariableAssignmentMap.Add(WriteVariable);
+						NewAssignment.DataInterface = CachedDefaultDataInterface.DataInterface;
+						NewAssignment.SourceEmitterName = CachedDefaultDataInterface.SourceEmitterName;
 					}
 				}
 			}
@@ -96,11 +104,11 @@ namespace FNiagaraResolveDIHelpers
 
 	void CollectDIBindingsAndAssignments(
 		const UNiagaraSystem* System,
-		TMap<FGuid, TMap<FNiagaraVariableBase, UNiagaraDataInterface*>>& OutEmitterIdToVariableAssignmentsMap,
+		TMap<FGuid, TMap<FNiagaraVariableBase, FDataInterfaceSourceEmitterNamePair>>& OutEmitterIdToVariableAssignmentsMap,
 		TMap<FGuid, TMap<FNiagaraVariableBase, FNiagaraVariableBase>>& OutEmitterIdToVariableBindingsMap,
 		TArray<FText>& OutErrorMessages)
 	{
-		TMap<FNiagaraVariableBase, UNiagaraDataInterface*> VariableAssignmentMap;
+		TMap<FNiagaraVariableBase, FDataInterfaceSourceEmitterNamePair> VariableAssignmentMap;
 		TMap<FNiagaraVariableBase, FNiagaraVariableBase> VariableBindingMap;
 
 		// Add user parameters as assignments.
@@ -111,7 +119,9 @@ namespace FNiagaraResolveDIHelpers
 				UNiagaraDataInterface* DataInterface = System->GetExposedParameters().GetDataInterface(UserParameterWithOffset.Offset);
 				if (DataInterface != nullptr)
 				{
-					VariableAssignmentMap.Add(UserParameterWithOffset, DataInterface);
+					FDataInterfaceSourceEmitterNamePair& NewAssignment = VariableAssignmentMap.Add(UserParameterWithOffset);
+					NewAssignment.DataInterface = DataInterface;
+					NewAssignment.SourceEmitterName = FString();
 				}
 			}
 		}
@@ -136,7 +146,7 @@ namespace FNiagaraResolveDIHelpers
 				continue;
 			}
 
-			TMap<FNiagaraVariableBase, UNiagaraDataInterface*> EmitterVariableAssignmentMap = VariableAssignmentMap;
+			TMap<FNiagaraVariableBase, FDataInterfaceSourceEmitterNamePair> EmitterVariableAssignmentMap = VariableAssignmentMap;
 			TMap<FNiagaraVariableBase, FNiagaraVariableBase> EmitterVariableBindingMap = VariableBindingMap;
 
 			CollectDIBindingsAndAssignmentsForScript(VersionedNiagaraEmitterData->SpawnScriptProps.Script, System, EmitterHandle.GetUniqueInstanceName(), EmitterVariableAssignmentMap, EmitterVariableBindingMap, OutErrorMessages);
@@ -165,8 +175,8 @@ namespace FNiagaraResolveDIHelpers
 		const UNiagaraSystem* System,
 		UNiagaraScript* TargetScript,
 		const FString& EmitterName,
-		TMap<FNiagaraVariableBase, UNiagaraDataInterface*>& VariableAssignmentMap,
-		TMap<FNiagaraVariableBase, FNiagaraVariableBase>& VariableBindingMap,
+		const TMap<FNiagaraVariableBase, FDataInterfaceSourceEmitterNamePair>& VariableAssignmentMap,
+		const TMap<FNiagaraVariableBase, FNiagaraVariableBase>& VariableBindingMap,
 		TArray<FText>& OutErrorMessages)
 	{
 		TArray<FNiagaraScriptResolvedDataInterfaceInfo> ResolvedDataInterfaces;
@@ -178,7 +188,6 @@ namespace FNiagaraResolveDIHelpers
 			FNiagaraScriptResolvedDataInterfaceInfo& ResolvedDataInterface = ResolvedDataInterfaces.AddDefaulted_GetRef();
 			ResolvedDataInterface.Name = CachedDefaultDataInterface.Name;
 			ResolvedDataInterface.CompileName = CachedDefaultDataInterface.CompileName;
-			ResolvedDataInterface.EmitterName = CachedDefaultDataInterface.EmitterName;
 			ResolvedDataInterface.UserPtrIdx = CachedDefaultDataInterface.UserPtrIdx;
 
 			if (CachedDefaultDataInterface.RegisteredParameterMapRead != NAME_None)
@@ -191,7 +200,7 @@ namespace FNiagaraResolveDIHelpers
 				}
 
 				FNiagaraVariableBase BoundVariable;
-				FNiagaraVariableBase* CurrentBoundVariable = &ReadVariable;
+				const FNiagaraVariableBase* CurrentBoundVariable = &ReadVariable;
 				TSet<FNiagaraVariableBase> SeenBoundVariables;
 				bool bCircularReferenceDetected = false;
 				while (CurrentBoundVariable != nullptr && bCircularReferenceDetected == false)
@@ -210,13 +219,14 @@ namespace FNiagaraResolveDIHelpers
 					}
 				}
 
-				UNiagaraDataInterface** BoundDataInterfacePtr = VariableAssignmentMap.Find(BoundVariable);
-				if (BoundDataInterfacePtr != nullptr)
+				const FDataInterfaceSourceEmitterNamePair* DataInterfaceAssignment = VariableAssignmentMap.Find(BoundVariable);
+				if (DataInterfaceAssignment != nullptr)
 				{
 					ResolvedDataInterface.ResolvedVariable = BoundVariable;
 					ResolvedDataInterface.SourceVariable = ReadVariable;
 					ResolvedDataInterface.bIsInternal = false;
-					ResolvedDataInterface.ResolvedDataInterface = *BoundDataInterfacePtr;
+					ResolvedDataInterface.ResolvedDataInterface = DataInterfaceAssignment->DataInterface;
+					ResolvedDataInterface.ResolvedSourceEmitterName = DataInterfaceAssignment->SourceEmitterName;
 				}
 
 				if (FNiagaraUserRedirectionParameterStore::IsUserParameter(BoundVariable))
@@ -243,6 +253,7 @@ namespace FNiagaraResolveDIHelpers
 				ResolvedDataInterface.SourceVariable = InternalVariable;
 				ResolvedDataInterface.bIsInternal = true;
 				ResolvedDataInterface.ResolvedDataInterface = CachedDefaultDataInterface.DataInterface;
+				ResolvedDataInterface.ResolvedSourceEmitterName = CachedDefaultDataInterface.SourceEmitterName;
 			}
 
 			ResolvedDataInterfaceIndex++;
@@ -262,6 +273,7 @@ namespace FNiagaraResolveDIHelpers
 					if (TargetResolvedDataInterface.bIsInternal && TargetResolvedDataInterface.ResolvedVariable == SourceResolvedDataInterface.ResolvedVariable)
 					{
 						TargetResolvedDataInterface.ResolvedDataInterface = SourceResolvedDataInterface.ResolvedDataInterface;
+						TargetResolvedDataInterface.ResolvedSourceEmitterName = SourceResolvedDataInterface.ResolvedSourceEmitterName;
 						break;
 					}
 				}
@@ -296,13 +308,13 @@ namespace FNiagaraResolveDIHelpers
 		ParticleGpuScript->SetResolvedDataInterfaces(ResolvedGpuDataInterfaces);
 	}
 
-	void ResolveDIs(
+	void ResolveDIsInternal(
 		UNiagaraSystem* System,
-		const TMap<FGuid, TMap<FNiagaraVariableBase, UNiagaraDataInterface*>>& EmitterIdToVariableAssignmentsMap,
+		const TMap<FGuid, TMap<FNiagaraVariableBase, FDataInterfaceSourceEmitterNamePair>>& EmitterIdToVariableAssignmentsMap,
 		const TMap<FGuid, TMap<FNiagaraVariableBase, FNiagaraVariableBase>>& EmitterIdToVariableBindingsMap,
 		TArray<FText>& OutErrorMessages)
 	{
-		TMap<FNiagaraVariableBase, UNiagaraDataInterface*> VariableAssignmentMap = EmitterIdToVariableAssignmentsMap[FGuid()];
+		TMap<FNiagaraVariableBase, FDataInterfaceSourceEmitterNamePair> VariableAssignmentMap = EmitterIdToVariableAssignmentsMap[FGuid()];
 		TMap<FNiagaraVariableBase, FNiagaraVariableBase> VariableBindingMap = EmitterIdToVariableBindingsMap[FGuid()];
 		ResolveDIsForScript(System, System->GetSystemSpawnScript(), FString(), VariableAssignmentMap, VariableBindingMap, OutErrorMessages);
 		ResolveDIsForScript(System, System->GetSystemUpdateScript(), FString(), VariableAssignmentMap, VariableBindingMap, OutErrorMessages);
@@ -320,7 +332,7 @@ namespace FNiagaraResolveDIHelpers
 				continue;
 			}
 
-			TMap<FNiagaraVariableBase, UNiagaraDataInterface*> EmitterVariableAssignmentMap = EmitterIdToVariableAssignmentsMap[EmitterHandle.GetId()];
+			TMap<FNiagaraVariableBase, FDataInterfaceSourceEmitterNamePair> EmitterVariableAssignmentMap = EmitterIdToVariableAssignmentsMap[EmitterHandle.GetId()];
 			TMap<FNiagaraVariableBase, FNiagaraVariableBase> EmitterVariableBindingMap = EmitterIdToVariableBindingsMap[EmitterHandle.GetId()];
 
 			ResolveDIsForScript(System, VersionedNiagaraEmitterData->SpawnScriptProps.Script, EmitterHandle.GetUniqueInstanceName(), EmitterVariableAssignmentMap, EmitterVariableBindingMap, OutErrorMessages);
@@ -352,6 +364,15 @@ namespace FNiagaraResolveDIHelpers
 			}
 		}
 	}
+
+	void ResolveDIs(UNiagaraSystem* System, TArray<FText>& OutErrorMessages)
+	{
+		TMap<FGuid, TMap<FNiagaraVariableBase, FDataInterfaceSourceEmitterNamePair>> EmitterIdToVariableAssignmentsMap;
+		TMap<FGuid, TMap<FNiagaraVariableBase, FNiagaraVariableBase>> EmitterIdToVariableBindingsMap;
+		TArray<FText> ErrorMessages;
+		FNiagaraResolveDIHelpers::CollectDIBindingsAndAssignments(System, EmitterIdToVariableAssignmentsMap, EmitterIdToVariableBindingsMap, ErrorMessages);
+		FNiagaraResolveDIHelpers::ResolveDIsInternal(System, EmitterIdToVariableAssignmentsMap, EmitterIdToVariableBindingsMap, ErrorMessages);
+	}
 }
 
 #undef LOCTEXT_NAMESPACE
diff --git a/Engine/Plugins/FX/Niagara/Source/Niagara/Private/NiagaraResolveDIHelpers.h b/Engine/Plugins/FX/Niagara/Source/Niagara/Private/NiagaraResolveDIHelpers.h
--- a/Engine/Plugins/FX/Niagara/Source/Niagara/Private/NiagaraResolveDIHelpers.h
+++ b/Engine/Plugins/FX/Niagara/Source/Niagara/Private/NiagaraResolveDIHelpers.h
@@ -14,17 +14,7 @@ class UNiagaraSystem;
 
 namespace FNiagaraResolveDIHelpers
 {
-	void CollectDIBindingsAndAssignments(
-		const UNiagaraSystem* System,
-		TMap<FGuid, TMap<FNiagaraVariableBase, UNiagaraDataInterface*>>& OutEmitterIdToVariableAssignmentsMap,
-		TMap<FGuid, TMap<FNiagaraVariableBase, FNiagaraVariableBase>>& OutEmitterIdToVariableBindingsMap,
-		TArray<FText>& OutErrorMessages);
-
-	void ResolveDIs(
-		UNiagaraSystem* System,
-		const TMap<FGuid, TMap<FNiagaraVariableBase, UNiagaraDataInterface*>>& EmitterIdToVariableAssignmentsMap,
-		const TMap<FGuid, TMap<FNiagaraVariableBase, FNiagaraVariableBase>>& EmitterIdToVariableBindingsMap,
-		TArray<FText>& OutErrorMessages);
+	void ResolveDIs(UNiagaraSystem* System,	TArray<FText>& OutErrorMessages);
 }
 
 #endif
\ No newline at end of file
diff --git a/Engine/Plugins/FX/Niagara/Source/Niagara/Private/NiagaraScript.cpp b/Engine/Plugins/FX/Niagara/Source/Niagara/Private/NiagaraScript.cpp
--- a/Engine/Plugins/FX/Niagara/Source/Niagara/Private/NiagaraScript.cpp
+++ b/Engine/Plugins/FX/Niagara/Source/Niagara/Private/NiagaraScript.cpp
@@ -2824,7 +2824,7 @@ void UNiagaraScript::SetVMCompilationResults(const FNiagaraVMExecutableDataId& I
 		CachedDefaultDataInterfaces[Idx].Name = ResolveEmitterAlias(Info.Name, EmitterUniqueName);
 		CachedDefaultDataInterfaces[Idx].CompileName = Info.Name;
 		CachedDefaultDataInterfaces[Idx].Type = Info.Type;
-		CachedDefaultDataInterfaces[Idx].EmitterName = EmitterUniqueName;
+		CachedDefaultDataInterfaces[Idx].SourceEmitterName = Info.SourceEmitterName;
 		CachedDefaultDataInterfaces[Idx].RegisteredParameterMapRead = ResolveEmitterAlias(Info.RegisteredParameterMapRead, EmitterUniqueName);
 		for (const FName& RegisteredParameterMapWrite : Info.RegisteredParameterMapWrites)
 		{
diff --git a/Engine/Plugins/FX/Niagara/Source/Niagara/Private/NiagaraSystem.cpp b/Engine/Plugins/FX/Niagara/Source/Niagara/Private/NiagaraSystem.cpp
--- a/Engine/Plugins/FX/Niagara/Source/Niagara/Private/NiagaraSystem.cpp
+++ b/Engine/Plugins/FX/Niagara/Source/Niagara/Private/NiagaraSystem.cpp
@@ -1958,12 +1958,8 @@ const FGuid UNiagaraSystem::ResolveDIsMessageId(0xB6ACDD97, 0xA2C04B02, 0xAB9FA4
 
 void UNiagaraSystem::ResolveDIBindings()
 {
-	TMap<FGuid, TMap<FNiagaraVariableBase, UNiagaraDataInterface*>> EmitterIdToVariableAssignmentsMap;
-	TMap<FGuid, TMap<FNiagaraVariableBase, FNiagaraVariableBase>> EmitterIdToVariableBindingsMap;
 	TArray<FText> ErrorMessages;
-	FNiagaraResolveDIHelpers::CollectDIBindingsAndAssignments(this, EmitterIdToVariableAssignmentsMap, EmitterIdToVariableBindingsMap, ErrorMessages);
-	FNiagaraResolveDIHelpers::ResolveDIs(this, EmitterIdToVariableAssignmentsMap, EmitterIdToVariableBindingsMap, ErrorMessages);
-
+	FNiagaraResolveDIHelpers::ResolveDIs(this, ErrorMessages);
 	if (ErrorMessages.Num() > 0)
 	{
 		INiagaraModule& NiagaraModule = FModuleManager::GetModuleChecked<INiagaraModule>("Niagara");
diff --git a/Engine/Plugins/FX/Niagara/Source/Niagara/Public/NiagaraCommon.h b/Engine/Plugins/FX/Niagara/Source/Niagara/Public/NiagaraCommon.h
--- a/Engine/Plugins/FX/Niagara/Source/Niagara/Public/NiagaraCommon.h
+++ b/Engine/Plugins/FX/Niagara/Source/Niagara/Public/NiagaraCommon.h
@@ -682,7 +682,7 @@ public:
 	TArray<FName> RegisteredParameterMapWrites;
 
 	UPROPERTY()
-	FString EmitterName;
+	FString SourceEmitterName;
 
 	bool IsUserDataInterface() const;
 };
@@ -727,7 +727,7 @@ struct NIAGARA_API FNiagaraScriptResolvedDataInterfaceInfo
 	FName CompileName;
 
 	UPROPERTY()
-	FString EmitterName;
+	FString ResolvedSourceEmitterName;
 
 	UPROPERTY()
 	FNiagaraVariableBase ResolvedVariable;
@@ -785,6 +785,9 @@ public:
 	UPROPERTY()
 	bool bIsPlaceholder;
 
+	UPROPERTY()
+	FString SourceEmitterName;
+
 	/** Would this data interface work on the target execution type? Only call this on the game thread.*/
 	bool CanExecuteOnTarget(ENiagaraSimTarget SimTarget) const;
 
diff --git a/Engine/Plugins/FX/Niagara/Source/NiagaraEditor/Private/NiagaraHlslTranslator.cpp b/Engine/Plugins/FX/Niagara/Source/NiagaraEditor/Private/NiagaraHlslTranslator.cpp
--- a/Engine/Plugins/FX/Niagara/Source/NiagaraEditor/Private/NiagaraHlslTranslator.cpp
+++ b/Engine/Plugins/FX/Niagara/Source/NiagaraEditor/Private/NiagaraHlslTranslator.cpp
@@ -6953,6 +6953,17 @@ int32 FHlslNiagaraTranslator::RegisterDataInterface(FNiagaraVariable& Var, UNiag
 		CompilationOutput.ScriptData.DataInterfaceInfo[Idx].Type = Var.GetType();
 		CompilationOutput.ScriptData.DataInterfaceInfo[Idx].bIsPlaceholder = bPlaceholder;
 
+		FString SourceEmitterName = FString();
+		if (CompileData->GetUniqueEmitterName().IsEmpty() == false)
+		{
+			SourceEmitterName = CompileData->GetUniqueEmitterName();
+		}
+		else if (EmitterAlias != nullptr)
+		{
+			SourceEmitterName = *EmitterAlias;
+		}
+		CompilationOutput.ScriptData.DataInterfaceInfo[Idx].SourceEmitterName = SourceEmitterName;
+
 		//Interface requires per instance data so add a user pointer table entry.
 		if (DataInterface != nullptr && DataInterface->PerInstanceDataSize() > 0)
 		{
