Niagara functions with variadic params now add their variadic params to their hlsl function signature so that they can be properly identified when coming back from the VM compiler backend.
Fixed equality operator for signatures with variadics.

diff --git a/Engine/Plugins/FX/Niagara/Source/Niagara/Public/NiagaraCommon.h b/Engine/Plugins/FX/Niagara/Source/Niagara/Public/NiagaraCommon.h
--- a/Engine/Plugins/FX/Niagara/Source/Niagara/Public/NiagaraCommon.h
+++ b/Engine/Plugins/FX/Niagara/Source/Niagara/Public/NiagaraCommon.h
@@ -516,7 +516,28 @@ struct NIAGARA_API FNiagaraFunctionSignature
 			}
 			return true;
 		}();
-		return EqualsIgnoringSpecifiers(Other) && bFunctionSpecifiersEqual;
+
+		bool bMatchingVariadics = true;
+		if (VariadicInput() || VariadicOutput())
+		{
+			bMatchingVariadics &= Other.Inputs.Num() == Inputs.Num() && RequiredInputs == Other.RequiredInputs;
+			bMatchingVariadics &= Other.Outputs.Num() == Outputs.Num() && RequiredOutputs == Other.RequiredOutputs;
+			if (bMatchingVariadics && RequiredInputs != INDEX_NONE)
+			{
+				for (int32 i = RequiredInputs; i < Inputs.Num(); ++i)
+				{
+					bMatchingVariadics &= Inputs[i] == Other.Inputs[i];
+				}
+			}
+			if (bMatchingVariadics && RequiredOutputs != INDEX_NONE)
+			{
+				for (int32 i = RequiredOutputs; i < Outputs.Num(); ++i)
+				{
+					bMatchingVariadics &= Outputs[i] == Other.Outputs[i];
+				}
+			}
+		}
+		return EqualsIgnoringSpecifiers(Other) && bFunctionSpecifiersEqual && bMatchingVariadics;
 	}
 
 	bool EqualsIgnoringSpecifiers(const FNiagaraFunctionSignature& Other) const
diff --git a/Engine/Plugins/FX/Niagara/Source/NiagaraEditor/Private/NiagaraHlslTranslator.cpp b/Engine/Plugins/FX/Niagara/Source/NiagaraEditor/Private/NiagaraHlslTranslator.cpp
--- a/Engine/Plugins/FX/Niagara/Source/NiagaraEditor/Private/NiagaraHlslTranslator.cpp
+++ b/Engine/Plugins/FX/Niagara/Source/NiagaraEditor/Private/NiagaraHlslTranslator.cpp
@@ -8603,6 +8603,22 @@ FString FHlslNiagaraTranslator::GetFunctionSignatureSymbol(const FNiagaraFunctio
 	{
 		SigStr += TEXT("_") + Specifier.Key.ToString() + Specifier.Value.ToString().Replace(TEXT("."), TEXT("_"));
 	}
+
+	auto AddVarsToSig = [&SigStr](TConstArrayView<FNiagaraVariableBase> Vars)
+	{
+		for(const FNiagaraVariableBase& Var : Vars)
+		{
+			SigStr += TEXT("_") + Var.GetName().ToString() + Var.GetType().GetName();
+		}
+	};
+
+	TArray<FNiagaraVariableBase> VariadicParams;
+	VariadicParams.Reserve(Sig.NumOptionalInputs() + Sig.NumOptionalOutputs());
+	Sig.GetVariadicInputs(VariadicParams);
+	AddVarsToSig(VariadicParams);
+	Sig.GetVariadicOutputs(VariadicParams);
+	AddVarsToSig(VariadicParams);
+
 	return GetSanitizedSymbolName(SigStr);
 }
 
