Missing file.

diff --git a/Engine/Plugins/FX/Niagara/Source/Niagara/Public/NiagaraTypes.h b/Engine/Plugins/FX/Niagara/Source/Niagara/Public/NiagaraTypes.h
--- a/Engine/Plugins/FX/Niagara/Source/Niagara/Public/NiagaraTypes.h
+++ b/Engine/Plugins/FX/Niagara/Source/Niagara/Public/NiagaraTypes.h
@@ -315,7 +315,13 @@ USTRUCT(BlueprintType, meta = (DisplayName = "Spawn Info", NiagaraClearEachFrame
 struct FNiagaraSpawnInfo
 {
 	GENERATED_USTRUCT_BODY();
-	
+
+	FNiagaraSpawnInfo() = default;
+	FNiagaraSpawnInfo(int32 InCount, float InStartDt, float InIntervalDt, int32 InSpawnGroup)
+		: Count(InCount), InterpStartDt(InStartDt), IntervalDt(InIntervalDt), SpawnGroup(InSpawnGroup)
+	{
+	}
+
 	/** How many particles to spawn. */
 	UPROPERTY(BlueprintReadWrite, Category = SpawnInfo)
 	int32 Count = 0;
