Fixed mismatched copy constructor/assignment operator in FNiagaraVariant.

diff --git a/Engine/Plugins/FX/Niagara/Source/Niagara/Private/NiagaraVariant.cpp b/Engine/Plugins/FX/Niagara/Source/Niagara/Private/NiagaraVariant.cpp
--- a/Engine/Plugins/FX/Niagara/Source/Niagara/Private/NiagaraVariant.cpp
+++ b/Engine/Plugins/FX/Niagara/Source/Niagara/Private/NiagaraVariant.cpp
@@ -13,14 +13,6 @@ FNiagaraVariant::FNiagaraVariant()
 	DataInterface = nullptr;
 }
 
-FNiagaraVariant::FNiagaraVariant(const FNiagaraVariant& Other)
-{
-	CurrentMode = Other.CurrentMode;
-	Object = Other.Object;
-	DataInterface = Other.DataInterface;
-	Bytes = Other.Bytes;
-}
-
 FNiagaraVariant::FNiagaraVariant(UObject* InObject)
 {
 	CurrentMode = ENiagaraVariantMode::Object;
diff --git a/Engine/Plugins/FX/Niagara/Source/Niagara/Public/NiagaraVariant.h b/Engine/Plugins/FX/Niagara/Source/Niagara/Public/NiagaraVariant.h
--- a/Engine/Plugins/FX/Niagara/Source/Niagara/Public/NiagaraVariant.h
+++ b/Engine/Plugins/FX/Niagara/Source/Niagara/Public/NiagaraVariant.h
@@ -24,7 +24,9 @@ struct NIAGARA_API FNiagaraVariant
 	GENERATED_BODY()
 
 	FNiagaraVariant();
-	FNiagaraVariant(const FNiagaraVariant& Other);
+	// TODO: Add move support
+	FNiagaraVariant(const FNiagaraVariant&) = default;
+	FNiagaraVariant& operator=(const FNiagaraVariant&) = default;
 	explicit FNiagaraVariant(UNiagaraDataInterface* InDataInterface);
 	explicit FNiagaraVariant(UObject* InObject);
 	explicit FNiagaraVariant(const TArray<uint8>& InBytes);
