Fixed mismatched copy constructor/assignment operator in FNiagaraGpuProfileEvent.

diff --git a/Engine/Plugins/FX/Niagara/Source/Niagara/Classes/NiagaraGPUProfilerInterface.h b/Engine/Plugins/FX/Niagara/Source/Niagara/Classes/NiagaraGPUProfilerInterface.h
--- a/Engine/Plugins/FX/Niagara/Source/Niagara/Classes/NiagaraGPUProfilerInterface.h
+++ b/Engine/Plugins/FX/Niagara/Source/Niagara/Classes/NiagaraGPUProfilerInterface.h
@@ -18,7 +18,6 @@ struct NIAGARA_API FNiagaraGpuProfileEvent
 
 	explicit FNiagaraGpuProfileEvent(const FNiagaraComputeInstanceData& InstanceData, const FNiagaraSimStageData& SimStageData, const bool bFirstInstanceData);
 	explicit FNiagaraGpuProfileEvent(const FNiagaraComputeInstanceData& InstanceData, FName CustomStageName);
-	explicit FNiagaraGpuProfileEvent(const FNiagaraGpuProfileEvent& Other) = default;
 
 private:
 	uint32									bUniqueInstance : 1;
@@ -28,7 +27,6 @@ private:
 #else
 	explicit FNiagaraGpuProfileEvent(const FNiagaraComputeInstanceData& InstanceData, const FNiagaraSimStageData& SimStageData, const bool bFirstInstanceData) {}
 	explicit FNiagaraGpuProfileEvent(const FNiagaraComputeInstanceData& InstanceData, FName CustomStageName) {}
-	explicit FNiagaraGpuProfileEvent(const FNiagaraGpuProfileEvent& Other) = default;
 #endif
 };
 
