Fixed mismatched copy constructor/assignment operator in FNiagaraTypeDefinition and FNiagaraVariable.

diff --git a/Engine/Plugins/FX/Niagara/Source/Niagara/Public/NiagaraTypes.h b/Engine/Plugins/FX/Niagara/Source/Niagara/Public/NiagaraTypes.h
--- a/Engine/Plugins/FX/Niagara/Source/Niagara/Public/NiagaraTypes.h
+++ b/Engine/Plugins/FX/Niagara/Source/Niagara/Public/NiagaraTypes.h
@@ -866,14 +866,6 @@ public:
 	
 	FORCEINLINE FNiagaraTypeDefinition(UScriptStruct* StructDef) : FNiagaraTypeDefinition(StructDef, EAllowUnfriendlyStruct::Deny) {}
 
-	FORCEINLINE FNiagaraTypeDefinition(const FNiagaraTypeDefinition &Other)
-		: ClassStructOrEnum(Other.ClassStructOrEnum), UnderlyingType(Other.UnderlyingType), Flags(Other.Flags), Size(INDEX_NONE), Alignment(INDEX_NONE)
-#if WITH_EDITORONLY_DATA
-		, Struct_DEPRECATED(nullptr), Enum_DEPRECATED(nullptr)
-#endif
-	{
-	}
-
 	FORCEINLINE FNiagaraTypeDefinition(FProperty* Property, EAllowUnfriendlyStruct AllowUnfriendlyStruct)
 	{
 		if (Property->IsA(FFloatProperty::StaticClass())) { *this = FNiagaraTypeDefinition::GetFloatDef(); }
@@ -1838,15 +1830,6 @@ struct FNiagaraVariable : public FNiagaraVariableBase
 	{
 	}
 
-	FNiagaraVariable(const FNiagaraVariable &Other)
-		: FNiagaraVariableBase(Other)
-	{
-		if (Other.GetType().IsValid() && Other.IsDataAllocated())
-		{
-			SetData(Other.GetData());
-		}
-	}
-
 	FNiagaraVariable(const FNiagaraVariableBase& Other)
 		: FNiagaraVariableBase(Other)
 	{
