Added support for curve assets to the curve data interfaces

diff --git a/Engine/Plugins/FX/Niagara/Source/Niagara/Classes/NiagaraDataInterfaceColorCurve.h b/Engine/Plugins/FX/Niagara/Source/Niagara/Classes/NiagaraDataInterfaceColorCurve.h
--- a/Engine/Plugins/FX/Niagara/Source/Niagara/Classes/NiagaraDataInterfaceColorCurve.h
+++ b/Engine/Plugins/FX/Niagara/Source/Niagara/Classes/NiagaraDataInterfaceColorCurve.h
@@ -73,6 +73,7 @@ public:
 	virtual void GetCurveData(TArray<FCurveData>& OutCurveData) override;
 
 #if WITH_EDITORONLY_DATA
+	virtual void SyncCurvesToAsset() override;
 	virtual bool GetFunctionHLSL(const FNiagaraDataInterfaceGPUParamInfo& ParamInfo, const FNiagaraDataInterfaceGeneratedFunction& FunctionInfo, int FunctionInstanceIndex, FString& OutHLSL) override;
 #endif
 
diff --git a/Engine/Plugins/FX/Niagara/Source/Niagara/Classes/NiagaraDataInterfaceCurve.h b/Engine/Plugins/FX/Niagara/Source/Niagara/Classes/NiagaraDataInterfaceCurve.h
--- a/Engine/Plugins/FX/Niagara/Source/Niagara/Classes/NiagaraDataInterfaceCurve.h
+++ b/Engine/Plugins/FX/Niagara/Source/Niagara/Classes/NiagaraDataInterfaceCurve.h
@@ -52,6 +52,7 @@ public:
 	virtual void GetCurveData(TArray<FCurveData>& OutCurveData) override;
 
 #if WITH_EDITORONLY_DATA
+	virtual void SyncCurvesToAsset() override;
 	virtual bool GetFunctionHLSL(const FNiagaraDataInterfaceGPUParamInfo& ParamInfo, const FNiagaraDataInterfaceGeneratedFunction& FunctionInfo, int FunctionInstanceIndex, FString& OutHLSL) override;
 #endif
 
diff --git a/Engine/Plugins/FX/Niagara/Source/Niagara/Classes/NiagaraDataInterfaceCurveBase.h b/Engine/Plugins/FX/Niagara/Source/Niagara/Classes/NiagaraDataInterfaceCurveBase.h
--- a/Engine/Plugins/FX/Niagara/Source/Niagara/Classes/NiagaraDataInterfaceCurveBase.h
+++ b/Engine/Plugins/FX/Niagara/Source/Niagara/Classes/NiagaraDataInterfaceCurveBase.h
@@ -2,10 +2,10 @@
 #pragma once
 
 #include "NiagaraCommon.h"
-#include "NiagaraShared.h"
 #include "VectorVM.h"
 #include "StaticMeshResources.h"
 #include "Curves/RichCurve.h"
+#include "Curves/CurveBase.h"
 #include "NiagaraDataInterface.h"
 #include "NiagaraDataInterfaceCurveBase.generated.h"
 
@@ -49,6 +49,11 @@ protected:
 	}
 
 public:
+#if WITH_EDITORONLY_DATA
+	UPROPERTY(EditAnywhere, Category = "Curve")
+	TSoftObjectPtr<UCurveBase> CurveAsset;
+#endif
+	
 	UPROPERTY(EditAnywhere, AdvancedDisplay, Category = "Curve")
 	uint32 bUseLUT : 1;
 
@@ -129,12 +134,14 @@ public:
 #endif
 	virtual void BuildShaderParameters(FNiagaraShaderParametersBuilder& ShaderParametersBuilder) const override;
 	virtual void SetShaderParameters(const FNiagaraDataInterfaceSetShaderParametersContext& Context) const override;
+	virtual void PostCompile() override;
 
 	void SetDefaultLUT();
 #if WITH_EDITORONLY_DATA
 	void UpdateLUT(bool bFromSerialize = false);
 	void OptimizeLUT();
 	void UpdateExposedTexture();
+	virtual void SyncCurvesToAsset() {};
 #endif
 
 	//UNiagaraDataInterface interface
@@ -151,7 +158,7 @@ public:
 #endif
 
 	virtual void UpdateTimeRanges() PURE_VIRTUAL(UNiagaraDataInterfaceCurveBase::UpdateTimeRanges(), )
-	virtual TArray<float> BuildLUT(int32 NumEntries) const PURE_VIRTUAL(UNiagaraDataInterfaceCurveBase::UpdateTimeRanges(), return TArray<float>(); )
+	virtual TArray<float> BuildLUT(int32 NumEntries) const PURE_VIRTUAL(UNiagaraDataInterfaceCurveBase::BuildLUT(), return TArray<float>(); )
 
 	FORCEINLINE float GetMinTime()const { return LUTMinTime; }
 	FORCEINLINE float GetMaxTime()const { return LUTMaxTime; }
diff --git a/Engine/Plugins/FX/Niagara/Source/Niagara/Classes/NiagaraDataInterfaceVectorCurve.h b/Engine/Plugins/FX/Niagara/Source/Niagara/Classes/NiagaraDataInterfaceVectorCurve.h
--- a/Engine/Plugins/FX/Niagara/Source/Niagara/Classes/NiagaraDataInterfaceVectorCurve.h
+++ b/Engine/Plugins/FX/Niagara/Source/Niagara/Classes/NiagaraDataInterfaceVectorCurve.h
@@ -72,6 +72,7 @@ public:
 #endif
 
 #if WITH_EDITORONLY_DATA
+	virtual void SyncCurvesToAsset() override;
 	virtual bool GetFunctionHLSL(const FNiagaraDataInterfaceGPUParamInfo& ParamInfo, const FNiagaraDataInterfaceGeneratedFunction& FunctionInfo, int FunctionInstanceIndex, FString& OutHLSL) override;
 #endif
 
diff --git a/Engine/Plugins/FX/Niagara/Source/Niagara/Private/NiagaraDataInterfaceColorCurve.cpp b/Engine/Plugins/FX/Niagara/Source/Niagara/Private/NiagaraDataInterfaceColorCurve.cpp
--- a/Engine/Plugins/FX/Niagara/Source/Niagara/Private/NiagaraDataInterfaceColorCurve.cpp
+++ b/Engine/Plugins/FX/Niagara/Source/Niagara/Private/NiagaraDataInterfaceColorCurve.cpp
@@ -7,6 +7,9 @@
 #include UE_INLINE_GENERATED_CPP_BY_NAME(NiagaraDataInterfaceColorCurve)
 
 #if WITH_EDITORONLY_DATA
+#include "Curves/CurveLinearColor.h"
+#include "Curves/CurveVector.h"
+#include "Curves/CurveFloat.h"
 #include "Interfaces/ITargetPlatform.h"
 #endif
 
@@ -179,6 +182,32 @@ void UNiagaraDataInterfaceColorCurve::GetFunctions(TArray<FNiagaraFunctionSignat
 // TODO: need a way to identify each specific function here
 // 
 #if WITH_EDITORONLY_DATA
+void UNiagaraDataInterfaceColorCurve::SyncCurvesToAsset()
+{
+	if (UCurveLinearColor* ColorCurve = Cast<UCurveLinearColor>(CurveAsset.LoadSynchronous()))
+	{
+		RedCurve = ColorCurve->FloatCurves[0];
+		GreenCurve = ColorCurve->FloatCurves[1];
+		BlueCurve = ColorCurve->FloatCurves[2];
+		AlphaCurve = ColorCurve->FloatCurves[3];
+	}
+	else if (UCurveVector* VecCurve = Cast<UCurveVector>(CurveAsset.LoadSynchronous()))
+	{
+		RedCurve = VecCurve->FloatCurves[0];
+		GreenCurve = VecCurve->FloatCurves[1];
+		BlueCurve = VecCurve->FloatCurves[2];
+		AlphaCurve.Reset();
+		AlphaCurve.AddKey(0, 1);
+	}
+	else if (UCurveFloat* FloatCurve = Cast<UCurveFloat>(CurveAsset.LoadSynchronous()))
+	{
+		RedCurve = FloatCurve->FloatCurve;
+		GreenCurve = FloatCurve->FloatCurve;
+		BlueCurve = FloatCurve->FloatCurve;
+		AlphaCurve = FloatCurve->FloatCurve;
+	}
+}
+
 bool UNiagaraDataInterfaceColorCurve::GetFunctionHLSL(const FNiagaraDataInterfaceGPUParamInfo& ParamInfo, const FNiagaraDataInterfaceGeneratedFunction& FunctionInfo, int FunctionInstanceIndex, FString& OutHLSL)
 {
 	if (FunctionInfo.DefinitionName == SampleCurveName)
diff --git a/Engine/Plugins/FX/Niagara/Source/Niagara/Private/NiagaraDataInterfaceCurve.cpp b/Engine/Plugins/FX/Niagara/Source/Niagara/Private/NiagaraDataInterfaceCurve.cpp
--- a/Engine/Plugins/FX/Niagara/Source/Niagara/Private/NiagaraDataInterfaceCurve.cpp
+++ b/Engine/Plugins/FX/Niagara/Source/Niagara/Private/NiagaraDataInterfaceCurve.cpp
@@ -7,6 +7,7 @@
 #include UE_INLINE_GENERATED_CPP_BY_NAME(NiagaraDataInterfaceCurve)
 
 #if WITH_EDITORONLY_DATA
+#include "Curves/CurveFloat.h"
 #include "Interfaces/ITargetPlatform.h"
 #endif
 
@@ -143,6 +144,14 @@ TArray<float> UNiagaraDataInterfaceCurve::BuildLUT(int32 NumEntries) const
 // TODO: need a way to identify each specific function here
 // 
 #if WITH_EDITORONLY_DATA
+void UNiagaraDataInterfaceCurve::SyncCurvesToAsset()
+{
+	if (UCurveFloat* LoadedCurve = Cast<UCurveFloat>(CurveAsset.LoadSynchronous()))
+	{
+		Curve = LoadedCurve->FloatCurve;
+	}
+}
+
 bool UNiagaraDataInterfaceCurve::GetFunctionHLSL(const FNiagaraDataInterfaceGPUParamInfo& ParamInfo, const FNiagaraDataInterfaceGeneratedFunction& FunctionInfo, int FunctionInstanceIndex, FString& OutHLSL)
 {
 	if (FunctionInfo.DefinitionName == SampleCurveName)
diff --git a/Engine/Plugins/FX/Niagara/Source/Niagara/Private/NiagaraDataInterfaceCurveBase.cpp b/Engine/Plugins/FX/Niagara/Source/Niagara/Private/NiagaraDataInterfaceCurveBase.cpp
--- a/Engine/Plugins/FX/Niagara/Source/Niagara/Private/NiagaraDataInterfaceCurveBase.cpp
+++ b/Engine/Plugins/FX/Niagara/Source/Niagara/Private/NiagaraDataInterfaceCurveBase.cpp
@@ -244,6 +244,7 @@ bool UNiagaraDataInterfaceCurveBase::CopyToInternal(UNiagaraDataInterface* Desti
 	DestinationTyped->ShaderLUT = ShaderLUT;
 	DestinationTyped->LUTNumSamplesMinusOne = LUTNumSamplesMinusOne;
 #if WITH_EDITORONLY_DATA
+	DestinationTyped->CurveAsset = CurveAsset;
 	DestinationTyped->bOptimizeLUT = bOptimizeLUT;
 	DestinationTyped->bOverrideOptimizeThreshold = bOverrideOptimizeThreshold;
 	DestinationTyped->OptimizeThreshold = OptimizeThreshold;
@@ -289,6 +290,7 @@ void UNiagaraDataInterfaceCurveBase::SetDefaultLUT()
 #if WITH_EDITORONLY_DATA
 void UNiagaraDataInterfaceCurveBase::UpdateLUT(bool bFromSerialize)
 {
+	SyncCurvesToAsset();
 	UpdateTimeRanges();
 	if (bUseLUT)
 	{
@@ -414,6 +416,7 @@ bool UNiagaraDataInterfaceCurveBase::Equals(const UNiagaraDataInterface* Other)
 	const UNiagaraDataInterfaceCurveBase* OtherTyped = CastChecked<UNiagaraDataInterfaceCurveBase>(Other);
 	bool bEqual = OtherTyped->bUseLUT == bUseLUT;
 #if WITH_EDITORONLY_DATA
+	bEqual &= OtherTyped->CurveAsset == CurveAsset;
 	bEqual &= OtherTyped->bOptimizeLUT == bOptimizeLUT;
 	bEqual &= OtherTyped->bOverrideOptimizeThreshold == bOverrideOptimizeThreshold;
 	if (bOverrideOptimizeThreshold)
@@ -464,6 +467,13 @@ void UNiagaraDataInterfaceCurveBase::SetShaderParameters(const FNiagaraDataInter
 	Parameters->CurveLUT = DIProxy.CurveLUT.SRV.IsValid() ? DIProxy.CurveLUT.SRV.GetReference() : FNiagaraRenderer::GetDummyFloatBuffer();
 }
 
+void UNiagaraDataInterfaceCurveBase::PostCompile()
+{
+#if WITH_EDITORONLY_DATA
+	SyncCurvesToAsset();
+#endif
+}
+
 #if WITH_EDITORONLY_DATA
 bool UNiagaraDataInterfaceCurveBase::AppendCompileHash(FNiagaraCompileHashVisitor* InVisitor) const
 {
diff --git a/Engine/Plugins/FX/Niagara/Source/Niagara/Private/NiagaraDataInterfaceVectorCurve.cpp b/Engine/Plugins/FX/Niagara/Source/Niagara/Private/NiagaraDataInterfaceVectorCurve.cpp
--- a/Engine/Plugins/FX/Niagara/Source/Niagara/Private/NiagaraDataInterfaceVectorCurve.cpp
+++ b/Engine/Plugins/FX/Niagara/Source/Niagara/Private/NiagaraDataInterfaceVectorCurve.cpp
@@ -7,6 +7,9 @@
 #include UE_INLINE_GENERATED_CPP_BY_NAME(NiagaraDataInterfaceVectorCurve)
 
 #if WITH_EDITORONLY_DATA
+#include "Curves/CurveLinearColor.h"
+#include "Curves/CurveVector.h"
+#include "Curves/CurveFloat.h"
 #include "Interfaces/ITargetPlatform.h"
 #endif
 
@@ -170,6 +173,28 @@ void UNiagaraDataInterfaceVectorCurve::GetFunctions(TArray<FNiagaraFunctionSigna
 // TODO: need a way to identify each specific function here
 // 
 #if WITH_EDITORONLY_DATA
+void UNiagaraDataInterfaceVectorCurve::SyncCurvesToAsset()
+{
+	if (UCurveLinearColor* ColorCurve = Cast<UCurveLinearColor>(CurveAsset.LoadSynchronous()))
+	{
+		XCurve = ColorCurve->FloatCurves[0];
+		YCurve = ColorCurve->FloatCurves[1];
+		ZCurve = ColorCurve->FloatCurves[2];
+	}
+	else if (UCurveVector* VecCurve = Cast<UCurveVector>(CurveAsset.LoadSynchronous()))
+	{
+		XCurve = VecCurve->FloatCurves[0];
+		YCurve = VecCurve->FloatCurves[1];
+		ZCurve = VecCurve->FloatCurves[2];
+	}
+	else if (UCurveFloat* FloatCurve = Cast<UCurveFloat>(CurveAsset.LoadSynchronous()))
+	{
+		XCurve = FloatCurve->FloatCurve;
+		YCurve = FloatCurve->FloatCurve;
+		ZCurve = FloatCurve->FloatCurve;
+	}
+}
+
 bool UNiagaraDataInterfaceVectorCurve::GetFunctionHLSL(const FNiagaraDataInterfaceGPUParamInfo& ParamInfo, const FNiagaraDataInterfaceGeneratedFunction& FunctionInfo, int FunctionInstanceIndex, FString& OutHLSL)
 {
 	if (FunctionInfo.DefinitionName == SampleCurveName)
diff --git a/Engine/Plugins/FX/Niagara/Source/NiagaraEditorWidgets/Private/DetailCustomizations/NiagaraDataInterfaceCurveDetails.cpp b/Engine/Plugins/FX/Niagara/Source/NiagaraEditorWidgets/Private/DetailCustomizations/NiagaraDataInterfaceCurveDetails.cpp
--- a/Engine/Plugins/FX/Niagara/Source/NiagaraEditorWidgets/Private/DetailCustomizations/NiagaraDataInterfaceCurveDetails.cpp
+++ b/Engine/Plugins/FX/Niagara/Source/NiagaraEditorWidgets/Private/DetailCustomizations/NiagaraDataInterfaceCurveDetails.cpp
@@ -36,13 +36,13 @@
 #include "Widgets/Input/SButton.h"
 #include "Widgets/Images/SImage.h"
 #include "Widgets/Layout/SWidgetSwitcher.h"
-#include "NiagaraEditorStyle.h"
 #include "NiagaraEditorWidgetsStyle.h"
 #include "IContentBrowserSingleton.h"
 #include "ContentBrowserModule.h"
 #include "Framework/Application/SlateApplication.h"
 #include "Framework/MultiBox/MultiBoxBuilder.h"
 #include "Framework/MultiBox/MultiBoxDefs.h"
+#include "Curves/CurveBase.h"
 #include "Curves/CurveVector.h"
 #include "Curves/CurveFloat.h"
 #include "Curves/CurveLinearColor.h"
@@ -784,125 +784,6 @@ class SNiagaraCurveThumbnail : public SLeafWidget
 	float Height;
 };
 
-class SNiagaraCurveTemplateBar : public SCompoundWidget
-{
-	SLATE_BEGIN_ARGS(SNiagaraCurveTemplateBar) { }
-	SLATE_END_ARGS()
-
-	void Construct(const FArguments& InArgs, TSharedRef<FCurveEditor> InCurveEditor)
-	{
-		CurveEditor = InCurveEditor;
-
-		const UNiagaraEditorSettings* Settings = GetDefault<UNiagaraEditorSettings>();
-
-		FToolBarBuilder ToolBarBuilder(CurveEditor->GetCommands(), FMultiBoxCustomization::None, nullptr, true);
-		ToolBarBuilder.SetLabelVisibility(EVisibility::Collapsed);
-
-		ToolBarBuilder.AddWidget(
-			SNew(SBox)
-			.VAlign(VAlign_Center)
-			.Padding(FMargin(0, 0, 5, 0))
-			[
-				SNew(STextBlock)
-				.Text(LOCTEXT("CurveTemplateLabel", "Templates"))
-			]);
-
-		for (const FNiagaraCurveTemplate& CurveTemplate : Settings->GetCurveTemplates())
-		{
-			UCurveFloat* FloatCurveAsset = Cast<UCurveFloat>(CurveTemplate.CurveAsset.TryLoad());
-			if(FloatCurveAsset != nullptr)
-			{
-				FText CurveDisplayName = CurveTemplate.DisplayNameOverride.IsEmpty()
-					? FText::FromString(FName::NameToDisplayString(FloatCurveAsset->GetName(), false))
-					: FText::FromString(CurveTemplate.DisplayNameOverride);
-
-				ToolBarBuilder.AddWidget(
-					SNew(SButton)
-					.ButtonStyle(FAppStyle::Get(), "SimpleButton")
-					.OnClicked(this, &SNiagaraCurveTemplateBar::CurveTemplateClicked, TWeakObjectPtr<UCurveFloat>(FloatCurveAsset))
-					.ToolTipText(FText::Format(LOCTEXT("ApplyCurveTemplateFormat", "{0}\nClick to apply this template to the selected curves."), CurveDisplayName))
-					.ContentPadding(FMargin(3))
-					.Content()
-					[
-						SNew(SNiagaraCurveThumbnail, &FloatCurveAsset->FloatCurve)
-					]);
-			}
-		}
-
-		ChildSlot
-		[
-			ToolBarBuilder.MakeWidget()
-		];
-	}
-
-private:
-	FReply CurveTemplateClicked(TWeakObjectPtr<UCurveFloat> FloatCurveAssetWeak)
-	{
-		UCurveFloat* FloatCurveAsset = FloatCurveAssetWeak.Get();
-		if (FloatCurveAsset != nullptr)
-		{
-			TArray<FCurveModelID> CurveModelIdsToSet;
-			if(CurveEditor->GetRootTreeItems().Num() == 1)
-			{
-				const FCurveEditorTreeItem& TreeItem = CurveEditor->GetTreeItem(CurveEditor->GetRootTreeItems()[0]);
-				for(const FCurveModelID& CurveModelId : TreeItem.GetCurves())
-				{
-					CurveModelIdsToSet.Add(CurveModelId);
-				}
-			}
-			else
-			{
-				for (const TPair<FCurveEditorTreeItemID, ECurveEditorTreeSelectionState>& TreeItemSelectionState : CurveEditor->GetTreeSelection())
-				{
-					if (TreeItemSelectionState.Value != ECurveEditorTreeSelectionState::None)
-					{
-						const FCurveEditorTreeItem& TreeItem = CurveEditor->GetTreeItem(TreeItemSelectionState.Key);
-						for (const FCurveModelID& CurveModelId : TreeItem.GetCurves())
-						{
-							CurveModelIdsToSet.Add(CurveModelId);
-						}
-					}
-				}
-			}
-
-			if (CurveModelIdsToSet.Num() > 0)
-			{
-				FScopedTransaction ApplyTemplateTransaction(LOCTEXT("ApplyCurveTemplateTransaction", "Apply curve template"));
-				for (const FCurveModelID& CurveModelId : CurveModelIdsToSet)
-				{
-					FCurveModel* CurveModel = CurveEditor->GetCurves()[CurveModelId].Get();
-					if (CurveModel != nullptr)
-					{
-						TArray<FKeyHandle> KeyHandles;
-						CurveModel->GetKeys(*CurveEditor.Get(), TNumericLimits<double>::Lowest(), TNumericLimits<double>::Max(), TNumericLimits<double>::Lowest(), TNumericLimits<double>::Max(), KeyHandles);
-						CurveModel->RemoveKeys(KeyHandles);
-
-						const FRichCurve& FloatCurve = FloatCurveAsset->FloatCurve;
-						for (auto KeyIterator = FloatCurve.GetKeyHandleIterator(); KeyIterator; ++KeyIterator)
-						{
-							const FRichCurveKey& Key = FloatCurve.GetKey(*KeyIterator);
-							FKeyPosition KeyPosition;
-							KeyPosition.InputValue = Key.Time;
-							KeyPosition.OutputValue = Key.Value;
-							FKeyAttributes KeyAttributes;
-							KeyAttributes.SetInterpMode(Key.InterpMode);
-							KeyAttributes.SetTangentMode(Key.TangentMode);
-							KeyAttributes.SetArriveTangent(Key.ArriveTangent);
-							KeyAttributes.SetLeaveTangent(Key.LeaveTangent);
-							CurveModel->AddKey(KeyPosition, KeyAttributes);
-						}
-					}
-				}
-				CurveEditor->ZoomToFit();
-			}
-		}
-		return FReply::Handled();
-	}
-
-private:
-	TSharedPtr<FCurveEditor> CurveEditor;
-};
-
 class SNiagaraDataInterfaceCurveEditor : public SCompoundWidget
 {
 private:
@@ -994,12 +875,6 @@ public:
 		[
 			SNew(SVerticalBox)
 			+ SVerticalBox::Slot()
-			.AutoHeight()
-			.Padding(KeySelectorLeftPadding, 3, 0, 3)
-			[
-				SNew(SNiagaraCurveTemplateBar, CurveEditor.ToSharedRef())
-			]
-			+ SVerticalBox::Slot()
 			.Padding(0, 0, 0, 5)
 			[
 				CurveEditorPanel.ToSharedRef()
@@ -1038,6 +913,11 @@ public:
 		return CurveEditor->GetCommands();
 	}
 
+	TSharedPtr<FCurveEditor> GetEditor()
+	{
+		return CurveEditor;
+	}
+
 private:
 	void InitializeView()
 	{
@@ -1366,6 +1246,9 @@ void FNiagaraDataInterfaceCurveDetailsBase::CustomizeDetails(IDetailLayoutBuilde
 	{
 		CurveProperty->MarkHiddenByCustomization();
 	}
+	TSharedRef<IPropertyHandle> AssetProperty = DetailBuilder.GetProperty(GET_MEMBER_NAME_CHECKED(UNiagaraDataInterfaceCurveBase, CurveAsset), UNiagaraDataInterfaceCurveBase::StaticClass());
+	IDetailPropertyRow* DefaultRow = DetailBuilder.EditDefaultProperty(AssetProperty);
+	DefaultRow->Visibility(TAttribute<EVisibility>::CreateLambda([this]() -> EVisibility { return GetUsedCurveAsset(CurveDataInterfaceWeak.Get()) ? EVisibility::Visible : EVisibility::Collapsed; }));
 
 	TSharedPtr<SNiagaraDataInterfaceCurveEditor> CurveEditor;
 	TSharedRef<SWidget> CurveDataInterfaceCurveEditor = 
@@ -1398,7 +1281,7 @@ void FNiagaraDataInterfaceCurveDetailsBase::CustomizeDetails(IDetailLayoutBuilde
 	}
 
 	FToolBarBuilder ToolBarBuilder(CurveEditor->GetCommands(), FMultiBoxCustomization::None, nullptr, true);
-	ToolBarBuilder.SetStyle(&FAppStyle::Get(), "SlimToolbar");
+	ToolBarBuilder.SetStyle(&FNiagaraEditorWidgetsStyle::Get(), "NiagaraEditor.CurveDetails.SlimToolBar");
 
 	ToolBarBuilder.AddComboButton(
 		FUIAction(), 
@@ -1439,9 +1322,50 @@ void FNiagaraDataInterfaceCurveDetailsBase::CustomizeDetails(IDetailLayoutBuilde
 			]);
 	}
 
+	// add templates to toolbar
+	ToolBarBuilder.AddSeparator();
+	const UNiagaraEditorSettings* Settings = GetDefault<UNiagaraEditorSettings>();
+	ToolBarBuilder.AddWidget(
+		SNew(SBox)
+		.VAlign(VAlign_Center)
+		.Visibility_Lambda([this]() { return GetIsColorCurve() && IsGradientVisible() ? EVisibility::Collapsed : EVisibility::Visible; })
+		.Padding(FMargin(3, 0, 5, 0))
+		[
+			SNew(STextBlock)
+			.Text(LOCTEXT("CurveTemplateLabel", "Templates"))
+		]);
+	for (const FNiagaraCurveTemplate& CurveTemplate : Settings->GetCurveTemplates())
+	{
+		UCurveFloat* FloatCurveAsset = Cast<UCurveFloat>(CurveTemplate.CurveAsset.TryLoad());
+		if(FloatCurveAsset != nullptr)
+		{
+			FText CurveDisplayName = CurveTemplate.DisplayNameOverride.IsEmpty()
+				? FText::FromString(FName::NameToDisplayString(FloatCurveAsset->GetName(), false))
+				: FText::FromString(CurveTemplate.DisplayNameOverride);
+
+			TWeakObjectPtr<UCurveFloat> WeakCurveAsset(FloatCurveAsset);
+			TSharedPtr<FCurveEditor> EditorPtr = CurveEditor->GetEditor();
+			ToolBarBuilder.AddWidget(
+				SNew(SButton)
+				.ButtonStyle(FAppStyle::Get(), "SimpleButton")
+				.OnClicked(this, &FNiagaraDataInterfaceCurveDetailsBase::CurveTemplateSelected, WeakCurveAsset, EditorPtr)
+				.ToolTipText(FText::Format(LOCTEXT("ApplyCurveTemplateFormat", "{0}\nClick to apply this template to the selected curves."), CurveDisplayName))
+				.ContentPadding(FMargin(3, 10, 3, 0))
+				.Visibility_Lambda([this]() { return GetIsColorCurve() && IsGradientVisible() ? EVisibility::Collapsed : EVisibility::Visible; })
+				.Content()
+				[
+					SNew(SNiagaraCurveThumbnail, &FloatCurveAsset->FloatCurve)
+				]);
+		}
+	}
+
 	IDetailCategoryBuilder& CurveCategory = DetailBuilder.EditCategory("Curve");
-	CurveCategory.AddCustomRow(NSLOCTEXT("NiagaraDataInterfaceCurveDetails", "CurveOptions", "Options")).WholeRowContent() [ ToolBarBuilder.MakeWidget() ];
-	CurveCategory.AddCustomRow(NSLOCTEXT("NiagaraDataInterfaceCurveDetails", "CurveFilterText", "Curve")).WholeRowContent() [ CurveDataInterfaceEditor.ToSharedRef() ];
+	CurveCategory.AddCustomRow(NSLOCTEXT("NiagaraDataInterfaceCurveDetails", "CurveOptions", "Options"))
+		.Visibility(TAttribute<EVisibility>::CreateLambda([this]() -> EVisibility { return GetUsedCurveAsset(CurveDataInterfaceWeak.Get()) ? EVisibility::Collapsed : EVisibility::Visible; }))
+		.WholeRowContent() [ ToolBarBuilder.MakeWidget() ];
+	CurveCategory.AddCustomRow(NSLOCTEXT("NiagaraDataInterfaceCurveDetails", "CurveFilterText", "Curve"))
+		.Visibility(TAttribute<EVisibility>::CreateLambda([this]() -> EVisibility { return GetUsedCurveAsset(CurveDataInterfaceWeak.Get()) ? EVisibility::Collapsed : EVisibility::Visible; }))
+		.WholeRowContent() [ CurveDataInterfaceEditor.ToSharedRef() ];
 }
 
 void FNiagaraDataInterfaceCurveDetailsBase::SetGradientVisibility(bool bInShowGradient)
@@ -1482,6 +1406,11 @@ void FNiagaraDataInterfaceCurveDetailsBase::OnShowInCurveEditor() const
 	}
 }
 
+UCurveBase* FNiagaraDataInterfaceCurveDetailsBase::GetUsedCurveAsset(UNiagaraDataInterfaceCurveBase* CurveDI) const
+{
+	return CurveDI ? CurveDI->CurveAsset.LoadSynchronous() : nullptr;
+}
+
 void FNiagaraDataInterfaceCurveDetailsBase::ImportSelectedAsset(UObject* SelectedAsset)
 {
 	if (CurveDataInterfaceWeak.IsValid() == false)
@@ -1497,6 +1426,12 @@ void FNiagaraDataInterfaceCurveDetailsBase::ImportSelectedAsset(UObject* Selecte
 	{
 		FScopedTransaction ImportTransaction(LOCTEXT("ImportCurveTransaction", "Import curve"));
 		CurveDataInterfaceWeak->Modify();
+
+		if (!bCopyAssetData)
+		{
+			CurveDataInterfaceWeak->CurveAsset = SelectedAsset;
+		}
+		
 		for (int i = 0; i < CurveProperties.Num(); i++)
 		{
 			if (CurveProperties[i]->IsValidHandle())
@@ -1525,13 +1460,47 @@ TSharedRef<SWidget> FNiagaraDataInterfaceCurveDetailsBase::GetCurveToCopyMenu()
 
 	FContentBrowserModule& ContentBrowserModule = FModuleManager::Get().LoadModuleChecked<FContentBrowserModule>(TEXT("ContentBrowser"));
 
-	return SNew(SBox)
-		.WidthOverride(300.0f)
+	return SNew(SVerticalBox)
+		+ SVerticalBox::Slot()
+		.Padding(5)
+		.AutoHeight()
 		[
-			ContentBrowserModule.Get().CreateAssetPicker(AssetPickerConfig)
+			SNew(SCheckBox)
+			.Style(FAppStyle::Get(), "RadioButton")
+			.IsChecked_Lambda([this]() { return bCopyAssetData ? ECheckBoxState::Checked : ECheckBoxState::Unchecked; })
+			.ToolTipText(LOCTEXT("CopyDataCheckBoxTooltip", "Copies the curve data from the selected asset into the stack - \nfurther updates to the asset will not affect the imported curves"))
+			.OnCheckStateChanged_Lambda([this](ECheckBoxState NewState) { bCopyAssetData = NewState == ECheckBoxState::Checked ? true : false; })
+			[
+				SNew(STextBlock)
+				.Text(LOCTEXT("CopyDataCheckBoxText", "Copy curve data"))
+			]
+		]
+		+ SVerticalBox::Slot()
+		.Padding(5)
+		.AutoHeight()
+		[
+			SNew(SCheckBox)
+			.Style(FAppStyle::Get(), "RadioButton")
+			.IsChecked_Lambda([this]() { return bCopyAssetData ? ECheckBoxState::Unchecked : ECheckBoxState::Checked; })
+			.ToolTipText(LOCTEXT("RefAssetheckBoxTooltip", "References the selected data instead of just copying the data - \nwhen the curve asset is changed and the effect is recompiled or cooked it will reflect the changes from the asset."))
+			.OnCheckStateChanged_Lambda([this](ECheckBoxState NewState) { bCopyAssetData = NewState == ECheckBoxState::Unchecked ? true : false; })
+			[
+				SNew(STextBlock)
+				.Text(LOCTEXT("RefAssetCheckBoxText", "Create asset reference"))
+			]
+		]
+		+ SVerticalBox::Slot()
+		.AutoHeight()
+		.MaxHeight(600)
+		.HAlign(HAlign_Fill)
+		.VAlign(VAlign_Fill)
+		[
+			SNew(SBox)
+			.WidthOverride(400.0f)
+			[
+				ContentBrowserModule.Get().CreateAssetPicker(AssetPickerConfig)
+			]
 		];
-	
-	return SNullWidget::NullWidget;
 }
 
 void FNiagaraDataInterfaceCurveDetailsBase::CurveToCopySelected(const FAssetData& AssetData)
@@ -1540,6 +1509,69 @@ void FNiagaraDataInterfaceCurveDetailsBase::CurveToCopySelected(const FAssetData
 	FSlateApplication::Get().DismissAllMenus();
 }
 
+FReply FNiagaraDataInterfaceCurveDetailsBase::CurveTemplateSelected(TWeakObjectPtr<UCurveFloat> WeakCurveAsset, TSharedPtr<FCurveEditor> CurveEditor)
+{
+	UCurveFloat* FloatCurveAsset = WeakCurveAsset.Get();
+	if (FloatCurveAsset != nullptr)
+	{
+		TArray<FCurveModelID> CurveModelIdsToSet;
+		if(CurveEditor->GetRootTreeItems().Num() == 1)
+		{
+			const FCurveEditorTreeItem& TreeItem = CurveEditor->GetTreeItem(CurveEditor->GetRootTreeItems()[0]);
+			for(const FCurveModelID& CurveModelId : TreeItem.GetCurves())
+			{
+				CurveModelIdsToSet.Add(CurveModelId);
+			}
+		}
+		else
+		{
+			for (const TPair<FCurveEditorTreeItemID, ECurveEditorTreeSelectionState>& TreeItemSelectionState : CurveEditor->GetTreeSelection())
+			{
+				if (TreeItemSelectionState.Value != ECurveEditorTreeSelectionState::None)
+				{
+					const FCurveEditorTreeItem& TreeItem = CurveEditor->GetTreeItem(TreeItemSelectionState.Key);
+					for (const FCurveModelID& CurveModelId : TreeItem.GetCurves())
+					{
+						CurveModelIdsToSet.Add(CurveModelId);
+					}
+				}
+			}
+		}
+
+		if (CurveModelIdsToSet.Num() > 0)
+		{
+			FScopedTransaction ApplyTemplateTransaction(LOCTEXT("ApplyCurveTemplateTransaction", "Apply curve template"));
+			for (const FCurveModelID& CurveModelId : CurveModelIdsToSet)
+			{
+				FCurveModel* CurveModel = CurveEditor->GetCurves()[CurveModelId].Get();
+				if (CurveModel != nullptr)
+				{
+					TArray<FKeyHandle> KeyHandles;
+					CurveModel->GetKeys(*CurveEditor.Get(), TNumericLimits<double>::Lowest(), TNumericLimits<double>::Max(), TNumericLimits<double>::Lowest(), TNumericLimits<double>::Max(), KeyHandles);
+					CurveModel->RemoveKeys(KeyHandles);
+
+					const FRichCurve& FloatCurve = FloatCurveAsset->FloatCurve;
+					for (auto KeyIterator = FloatCurve.GetKeyHandleIterator(); KeyIterator; ++KeyIterator)
+					{
+						const FRichCurveKey& Key = FloatCurve.GetKey(*KeyIterator);
+						FKeyPosition KeyPosition;
+						KeyPosition.InputValue = Key.Time;
+						KeyPosition.OutputValue = Key.Value;
+						FKeyAttributes KeyAttributes;
+						KeyAttributes.SetInterpMode(Key.InterpMode);
+						KeyAttributes.SetTangentMode(Key.TangentMode);
+						KeyAttributes.SetArriveTangent(Key.ArriveTangent);
+						KeyAttributes.SetLeaveTangent(Key.LeaveTangent);
+						CurveModel->AddKey(KeyPosition, KeyAttributes);
+					}
+				}
+			}
+			CurveEditor->ZoomToFit();
+		}
+	}
+	return FReply::Handled();
+}
+
 // Curve
 TSharedRef<IDetailCustomization> FNiagaraDataInterfaceCurveDetails::MakeInstance()
 {
diff --git a/Engine/Plugins/FX/Niagara/Source/NiagaraEditorWidgets/Private/DetailCustomizations/NiagaraDataInterfaceCurveDetails.h b/Engine/Plugins/FX/Niagara/Source/NiagaraEditorWidgets/Private/DetailCustomizations/NiagaraDataInterfaceCurveDetails.h
--- a/Engine/Plugins/FX/Niagara/Source/NiagaraEditorWidgets/Private/DetailCustomizations/NiagaraDataInterfaceCurveDetails.h
+++ b/Engine/Plugins/FX/Niagara/Source/NiagaraEditorWidgets/Private/DetailCustomizations/NiagaraDataInterfaceCurveDetails.h
@@ -11,6 +11,8 @@ class IPropertyHandle;
 class UNiagaraDataInterfaceCurveBase;
 class FNiagaraStackCurveEditorOptions;
 class SWidget;
+class UCurveBase;
+class UCurveFloat;
 struct FAssetData;
 struct FRichCurve;
 
@@ -24,13 +26,14 @@ protected:
 	virtual void GetCurveProperties(IDetailLayoutBuilder& DetailBuilder, TArray<TSharedRef<IPropertyHandle>>& OutCurveProperties) const = 0;
 	virtual bool GetIsColorCurve() const { return false; }
 	virtual float GetDefaultHeight() const { return 150; }
+	virtual UCurveBase* GetUsedCurveAsset(UNiagaraDataInterfaceCurveBase* CurveDI) const;
 	
 	virtual FTopLevelAssetPath GetSupportedAssetClassName() const = 0;
 	virtual void ImportSelectedAsset(UObject* SelectedAsset);
 	virtual void GetFloatCurvesFromAsset(UObject* SelectedAsset, TArray<FRichCurve>& FloatCurves) const = 0;
 	
 protected:
-	IDetailLayoutBuilder* CustomDetailBuilder;
+	IDetailLayoutBuilder* CustomDetailBuilder = nullptr;
 
 private:
 	virtual TSharedRef<SWidget> GetCurveToCopyMenu();
@@ -39,8 +42,10 @@ private:
 	bool IsGradientVisible() const;
 	void OnShowInCurveEditor() const;
 	void CurveToCopySelected(const FAssetData& AssetData);
+	FReply CurveTemplateSelected(TWeakObjectPtr<UCurveFloat> WeakCurveAsset, TSharedPtr<class FCurveEditor> CurveEditor);
 	TWeakObjectPtr<UNiagaraDataInterfaceCurveBase> CurveDataInterfaceWeak;
 	TWeakPtr<FNiagaraStackCurveEditorOptions> StackCurveEditorOptionsWeak;
+	bool bCopyAssetData = true;
 };
 
 /** Details customization for curve data interfaces. */
diff --git a/Engine/Plugins/FX/Niagara/Source/NiagaraEditorWidgets/Private/NiagaraEditorWidgetsStyle.cpp b/Engine/Plugins/FX/Niagara/Source/NiagaraEditorWidgets/Private/NiagaraEditorWidgetsStyle.cpp
--- a/Engine/Plugins/FX/Niagara/Source/NiagaraEditorWidgets/Private/NiagaraEditorWidgetsStyle.cpp
+++ b/Engine/Plugins/FX/Niagara/Source/NiagaraEditorWidgets/Private/NiagaraEditorWidgetsStyle.cpp
@@ -11,6 +11,7 @@
 #include "Settings/EditorStyleSettings.h"
 #include "Styling/AppStyle.h"
 #include "Styling/StyleColors.h"
+#include "Styling/ToolBarStyle.h"
 
 const FVector2D Icon6x6(6.0f, 6.0f);
 const FVector2D Icon8x8(8.0f, 8.0f);
@@ -299,10 +300,17 @@ FNiagaraEditorWidgetsStyle::FNiagaraEditorWidgetsStyle() : FSlateStyleSet("Niaga
 	Set("NiagaraEditor.CurveOverview.DefaultText", CurveOverviewDefaultText);
 
 	Set("NiagaraEditor.CurveDetails.TextButtonForeground", FLinearColor::White);
-
 	Set("NiagaraEditor.CurveDetails.Import.Small", new CORE_IMAGE_BRUSH("Icons/GeneralTools/Import_40x", Icon20x20));
-
 	Set("NiagaraEditor.CurveDetails.ShowInOverview.Small", new CORE_IMAGE_BRUSH("Common/GoToSource", Icon12x12, FLinearColor(.9f, .9f, .9f, 1.0f)));
+	
+	FToolBarStyle SlimToolbarStyle = FAppStyle::Get().GetWidgetStyle<FToolBarStyle>("SlimToolBar");
+	SlimToolbarStyle
+		.SetBackgroundPadding(FMargin(4.0f, 0.0f, 0.0f, 0.0f))
+		.SetIconSize(Icon16x16);
+	SlimToolbarStyle.ToggleButton
+		.SetPadding(FMargin(8.0f, 2.0f));
+	SlimToolbarStyle.SetBlockPadding(FMargin(2, 0));
+	Set("NiagaraEditor.CurveDetails.SlimToolBar", SlimToolbarStyle);
 
 	FSlateBoxBrush StackRowSelectionBrush = BOX_BRUSH("Icons/StackSelectionBorder", FMargin(2.0f / 8.0f), FStyleColors::Select);
 	FSlateBoxBrush StackRowSubduedSelectionBrush = BOX_BRUSH("Icons/StackSelectionBorder", FMargin(2.0f / 8.0f), FStyleColors::SelectInactive);
