Refactor DI interface for querying emitter dependencies - get a list of emitter names that are referenced rather than evaluating a supplied emitter

diff --git a/Engine/Plugins/FX/Niagara/Source/Niagara/Classes/NiagaraDataInterface.h b/Engine/Plugins/FX/Niagara/Source/Niagara/Classes/NiagaraDataInterface.h
--- a/Engine/Plugins/FX/Niagara/Source/Niagara/Classes/NiagaraDataInterface.h
+++ b/Engine/Plugins/FX/Niagara/Source/Niagara/Classes/NiagaraDataInterface.h
@@ -770,7 +770,16 @@ public:
 	*/
 	virtual void GetEmitterDependencies(UNiagaraSystem* Asset, TArray<FVersionedNiagaraEmitter>& Dependencies) const {}
 
-	virtual bool ReadsEmitterParticleData(const FString& EmitterName) const { return false; }
+	UE_DEPRECATED(5.4, "Implement GetEmitterReferencesByName() instead as this will be removed in a future version.")
+	virtual bool ReadsEmitterParticleData(const FString& EmitterName) const
+	{
+		TArray<FString> EmitterReferences;
+		GetEmitterReferencesByName(EmitterReferences);
+
+		return EmitterReferences.Contains(EmitterName);
+	}
+
+	virtual void GetEmitterReferencesByName(TArray<FString>& EmitterReferences) const {};
 
 	/**
 	Set the shader parameters will only be called if the data interface provided shader parameters.
diff --git a/Engine/Plugins/FX/Niagara/Source/Niagara/Classes/NiagaraDataInterfaceParticleRead.h b/Engine/Plugins/FX/Niagara/Source/Niagara/Classes/NiagaraDataInterfaceParticleRead.h
--- a/Engine/Plugins/FX/Niagara/Source/Niagara/Classes/NiagaraDataInterfaceParticleRead.h
+++ b/Engine/Plugins/FX/Niagara/Source/Niagara/Classes/NiagaraDataInterfaceParticleRead.h
@@ -64,7 +64,7 @@ public:
 	virtual bool UpgradeFunctionCall(FNiagaraFunctionSignature& FunctionSignature) override;
 #endif
 	virtual void GetEmitterDependencies(UNiagaraSystem* Asset, TArray<FVersionedNiagaraEmitter>& Dependencies) const override;
-	virtual bool ReadsEmitterParticleData(const FString& EmitterName) const override;
+	virtual void GetEmitterReferencesByName(TArray<FString>& EmitterReferences) const override;
 
 	virtual bool HasInternalAttributeReads(const UNiagaraEmitter* OwnerEmitter, const UNiagaraEmitter* Provider) const override;
 	//UNiagaraDataInterface Interface End
diff --git a/Engine/Plugins/FX/Niagara/Source/Niagara/Private/NiagaraDataInterfaceParticleRead.cpp b/Engine/Plugins/FX/Niagara/Source/Niagara/Private/NiagaraDataInterfaceParticleRead.cpp
--- a/Engine/Plugins/FX/Niagara/Source/Niagara/Private/NiagaraDataInterfaceParticleRead.cpp
+++ b/Engine/Plugins/FX/Niagara/Source/Niagara/Private/NiagaraDataInterfaceParticleRead.cpp
@@ -2550,9 +2550,12 @@ void UNiagaraDataInterfaceParticleRead::GetEmitterDependencies(UNiagaraSystem* A
 	}
 }
 
-bool UNiagaraDataInterfaceParticleRead::ReadsEmitterParticleData(const FString& InEmitterName) const 
+void UNiagaraDataInterfaceParticleRead::GetEmitterReferencesByName(TArray<FString>& EmitterReferences) const
 {
-	return EmitterName == InEmitterName;
+	if (!EmitterName.IsEmpty())
+	{
+		EmitterReferences.AddUnique(EmitterName);
+	}
 }
 
 #undef LOCTEXT_NAMESPACE
diff --git a/Engine/Plugins/FX/Niagara/Source/NiagaraEditor/Private/NiagaraCompiler.cpp b/Engine/Plugins/FX/Niagara/Source/NiagaraEditor/Private/NiagaraCompiler.cpp
--- a/Engine/Plugins/FX/Niagara/Source/NiagaraEditor/Private/NiagaraCompiler.cpp
+++ b/Engine/Plugins/FX/Niagara/Source/NiagaraEditor/Private/NiagaraCompiler.cpp
@@ -775,14 +775,17 @@ void FNiagaraCompileRequestData::FinishPrecompile(const TArray<FNiagaraVariable>
 									FCompileDataInterfaceData* DataInterfaceData = nullptr;
 									if (DataInterface != nullptr)
 									{
+										TArray<FString> EmitterReferences;
+										DataInterface->GetEmitterReferencesByName(EmitterReferences);
+
 										for (const FString& EmitterName : EmitterNames)
 										{
-											if (DataInterface->ReadsEmitterParticleData(EmitterName))
+											if (EmitterReferences.Contains(EmitterName))
 											{
 												if (DataInterfaceData == nullptr)
 												{
 													DataInterfaceData = &SharedCompileDataInterfaceData->AddDefaulted_GetRef();
-													DataInterfaceData->EmitterName = ConstantResolver.GetEmitter() != nullptr ? ConstantResolver.GetEmitter()->GetUniqueEmitterName() : FString();
+													DataInterfaceData->EmitterName = EmitterUniqueName;
 													DataInterfaceData->Usage = FoundOutputNode->GetUsage();
 													DataInterfaceData->UsageId = FoundOutputNode->GetUsageId();
 													DataInterfaceData->Variable = Variable;
@@ -1580,7 +1583,9 @@ TSharedPtr<FNiagaraGraphCachedDataBase, ESPMode::ThreadSafe> FNiagaraEditorModul
 			Builder.EndUsage();
 
 			for (const FNiagaraVariable& Var : Builder.StaticVariables)
+			{
 				StaticVariables.AddUnique(Var);
+			}
 		}
 	}
 
