Small fix to Niagara attribute trimming to limit how deep we dig for subexpressions when chaining dependencies.

diff --git a/Engine/Plugins/FX/Niagara/Source/NiagaraEditor/Private/NiagaraAttributeTrimmer.cpp b/Engine/Plugins/FX/Niagara/Source/NiagaraEditor/Private/NiagaraAttributeTrimmer.cpp
--- a/Engine/Plugins/FX/Niagara/Source/NiagaraEditor/Private/NiagaraAttributeTrimmer.cpp
+++ b/Engine/Plugins/FX/Niagara/Source/NiagaraEditor/Private/NiagaraAttributeTrimmer.cpp
@@ -15,11 +15,6 @@
 #include "NiagaraNodeParameterMapSet.h"
 #include "NiagaraScriptVariable.h"
 
-uint32 GetTypeHash(const FModuleScopedPin& ScopedPin)
-{
-	return GetTypeHash(TTuple<const UEdGraphPin*, FName>(ScopedPin.Pin, ScopedPin.ModuleName));
-}
-
 namespace NiagaraAttributeTrimming
 {
 	using FDependencySet = TSet<FModuleScopedPin>;
@@ -317,9 +312,10 @@ namespace NiagaraAttributeTrimming
 									continue;
 								}
 
-								// we want to exclude execution pins
+								// we want to exclude execution & parameter map pins because they terminate the subexpression
+								const bool IsParamMapPin = UEdGraphSchema_Niagara::PinToTypeDefinition(NodePin) == FNiagaraTypeDefinition::GetParameterMapDef();
 								const bool IsExecPin = NodePin->PinName == UNiagaraNodeParameterMapBase::SourcePinName;
-								if (IsExecPin)
+								if (IsParamMapPin || IsExecPin)
 								{
 									continue;
 								}
diff --git a/Engine/Plugins/FX/Niagara/Source/NiagaraEditor/Public/NiagaraParameterMapHistory.h b/Engine/Plugins/FX/Niagara/Source/NiagaraEditor/Public/NiagaraParameterMapHistory.h
--- a/Engine/Plugins/FX/Niagara/Source/NiagaraEditor/Public/NiagaraParameterMapHistory.h
+++ b/Engine/Plugins/FX/Niagara/Source/NiagaraEditor/Public/NiagaraParameterMapHistory.h
@@ -65,6 +65,11 @@ struct FModuleScopedPin
 
 	const UEdGraphPin* Pin = nullptr;
 	FName ModuleName = NAME_None;
+
+	friend uint32 GetTypeHash(const FModuleScopedPin& ScopedPin)
+	{
+		return GetTypeHash(TTuple<const UEdGraphPin*, FName>(ScopedPin.Pin, ScopedPin.ModuleName));
+	}
 };
 
 struct FGraphTraversalHandle
@@ -654,4 +659,4 @@ struct FNiagaraGraphCachedBuiltHistory : public FNiagaraGraphCachedDataBase
 
 	TArray<FNiagaraVariable> StaticVariables;
 	TArray<FGuid> SourceAssetChangeIds;
-};
\ No newline at end of file
+};
diff --git a/Engine/Source/Runtime/Core/Private/UObject/UE5MainStreamObjectVersion.cpp b/Engine/Source/Runtime/Core/Private/UObject/UE5MainStreamObjectVersion.cpp
--- a/Engine/Source/Runtime/Core/Private/UObject/UE5MainStreamObjectVersion.cpp
+++ b/Engine/Source/Runtime/Core/Private/UObject/UE5MainStreamObjectVersion.cpp
@@ -14,7 +14,7 @@ TMap<FGuid, FGuid> FUE5MainStreamObjectVersion::GetSystemGuids()
 	SystemGuids.Add(DevGuids.MATERIALSHADERMAP_DERIVEDDATA_VER, FGuid("65C01D817C9A4EEFAE9E988D41A1F3DD"));
 	SystemGuids.Add(DevGuids.NANITE_DERIVEDDATA_VER, FGuid("C283995F5F1346C3B6F7FB1DDF01ED7F"));
 	SystemGuids.Add(DevGuids.NIAGARASHADERMAP_DERIVEDDATA_VER, FGuid("8A37C45D24F2423CBE5F8F371DE33575"));
-	SystemGuids.Add(DevGuids.Niagara_LatestScriptCompileVersion, FGuid("4E3E29C28655E162A05986724B41711B"));
+	SystemGuids.Add(DevGuids.Niagara_LatestScriptCompileVersion, FGuid("D30EFEEB3A9D4059AC7CD2B9748949DE"));
 	SystemGuids.Add(DevGuids.SkeletalMeshDerivedDataVersion, FGuid("4F04FE8848F36B4F2AEB5D156092E08F"));
 	SystemGuids.Add(DevGuids.STATICMESH_DERIVEDDATA_VER, FGuid("98F7D79A5811013825E75BBCC41ED3E9"));
 	SystemGuids.Add(DevGuids.POSESEARCHDB_DERIVEDDATA_VER, FGuid("4E595C2AC5E947D6BA9ABC874353E5BC"));
