Few small changes to use array views and name buidlers to reduce allocaitons during Niagara compilation

diff --git a/Engine/Plugins/FX/Niagara/Source/Niagara/Private/NiagaraModule.cpp b/Engine/Plugins/FX/Niagara/Source/Niagara/Private/NiagaraModule.cpp
--- a/Engine/Plugins/FX/Niagara/Source/Niagara/Private/NiagaraModule.cpp
+++ b/Engine/Plugins/FX/Niagara/Source/Niagara/Private/NiagaraModule.cpp
@@ -1310,7 +1310,7 @@ bool FNiagaraTypeDefinition::IsLossyConversion(const FNiagaraTypeDefinition& Fro
 		|| (FromType == Vec3Def && ToType == PositionDef);
 }
 
-FNiagaraTypeDefinition FNiagaraTypeDefinition::GetNumericOutputType(const TArray<FNiagaraTypeDefinition> TypeDefinintions, ENiagaraNumericOutputTypeSelectionMode SelectionMode)
+FNiagaraTypeDefinition FNiagaraTypeDefinition::GetNumericOutputType(TConstArrayView<FNiagaraTypeDefinition> TypeDefinintions, ENiagaraNumericOutputTypeSelectionMode SelectionMode)
 {
 	checkf(SelectionMode != ENiagaraNumericOutputTypeSelectionMode::None, TEXT("Can not get numeric output type with selection mode none."));
 	checkf(SelectionMode != ENiagaraNumericOutputTypeSelectionMode::Custom, TEXT("Can not get numeric output type with selection mode custom."));
@@ -1341,7 +1341,7 @@ FNiagaraTypeDefinition FNiagaraTypeDefinition::GetNumericOutputType(const TArray
 		return NumericDef;
 	}
 
-	TArray<FNiagaraTypeDefinition> SortedTypeDefinitions = TypeDefinintions;
+	TArray<FNiagaraTypeDefinition> SortedTypeDefinitions(TypeDefinintions);
 	SortedTypeDefinitions.Sort([&](const FNiagaraTypeDefinition& TypeA, const FNiagaraTypeDefinition& TypeB)
 		{
 			int32 AIndex = OrderedNumericTypes.IndexOfByKey(TypeA);
diff --git a/Engine/Plugins/FX/Niagara/Source/Niagara/Private/NiagaraUserRedirectionParameterStore.cpp b/Engine/Plugins/FX/Niagara/Source/Niagara/Private/NiagaraUserRedirectionParameterStore.cpp
--- a/Engine/Plugins/FX/Niagara/Source/Niagara/Private/NiagaraUserRedirectionParameterStore.cpp
+++ b/Engine/Plugins/FX/Niagara/Source/Niagara/Private/NiagaraUserRedirectionParameterStore.cpp
@@ -2,6 +2,8 @@
 
 #include "NiagaraUserRedirectionParameterStore.h"
 
+#include "NiagaraConstants.h"
+
 #include UE_INLINE_GENERATED_CPP_BY_NAME(NiagaraUserRedirectionParameterStore)
 
 FNiagaraUserRedirectionParameterStore::FNiagaraUserRedirectionParameterStore() : FNiagaraParameterStore()
@@ -92,7 +94,11 @@ bool FNiagaraUserRedirectionParameterStore::AddParameter(const FNiagaraVariable&
 	}
 	else
 	{
-		AddParam = FNiagaraVariable(Param.GetType(), *(TEXT("User.") + Param.GetName().ToString()));
+		FNameBuilder NameBuilder;
+		FNiagaraConstants::UserNamespace.AppendString(NameBuilder);
+		NameBuilder << TEXT(".");
+		Param.GetName().AppendString(NameBuilder);
+		AddParam = FNiagaraVariable(Param.GetType(), FName(NameBuilder.ToView()));
 	}
 
 	UserParameterRedirects.Add(GetUserRedirection(AddParam), AddParam);
diff --git a/Engine/Plugins/FX/Niagara/Source/Niagara/Public/NiagaraTypes.h b/Engine/Plugins/FX/Niagara/Source/Niagara/Public/NiagaraTypes.h
--- a/Engine/Plugins/FX/Niagara/Source/Niagara/Public/NiagaraTypes.h
+++ b/Engine/Plugins/FX/Niagara/Source/Niagara/Public/NiagaraTypes.h
@@ -1194,7 +1194,7 @@ public:
 	
 	static bool TypesAreAssignable(const FNiagaraTypeDefinition& TypeA, const FNiagaraTypeDefinition& TypeB, bool bAllowLossyLWCConversions = false);
 	static bool IsLossyConversion(const FNiagaraTypeDefinition& FromType, const FNiagaraTypeDefinition& ToType);
-	static FNiagaraTypeDefinition GetNumericOutputType(const TArray<FNiagaraTypeDefinition> TypeDefinintions, ENiagaraNumericOutputTypeSelectionMode SelectionMode);
+	static FNiagaraTypeDefinition GetNumericOutputType(TConstArrayView<FNiagaraTypeDefinition> TypeDefinintions, ENiagaraNumericOutputTypeSelectionMode SelectionMode);
 
 	static const TArray<FNiagaraTypeDefinition>& GetNumericTypes() { return OrderedNumericTypes; }
 	static bool IsValidNumericInput(const FNiagaraTypeDefinition& TypeDef);
diff --git a/Engine/Plugins/FX/Niagara/Source/NiagaraEditor/Private/EdGraphSchema_Niagara.cpp b/Engine/Plugins/FX/Niagara/Source/NiagaraEditor/Private/EdGraphSchema_Niagara.cpp
--- a/Engine/Plugins/FX/Niagara/Source/NiagaraEditor/Private/EdGraphSchema_Niagara.cpp
+++ b/Engine/Plugins/FX/Niagara/Source/NiagaraEditor/Private/EdGraphSchema_Niagara.cpp
@@ -1917,11 +1917,11 @@ static UNiagaraParameterCollection* EnsureCollectionLoaded(FAssetData& Collectio
 
 UNiagaraParameterCollection* UEdGraphSchema_Niagara::VariableIsFromParameterCollection(const FNiagaraVariable& Var)const
 {
-	FString VarName = Var.GetName().ToString();
-	if (VarName.StartsWith(TEXT("NPC.")))
+	FNameBuilder VarName(Var.GetName());
+	if (VarName.ToView().StartsWith(TEXT("NPC.")))
 	{
 		FNiagaraEditorModule& NiagaraEditorModule = FModuleManager::GetModuleChecked<FNiagaraEditorModule>("NiagaraEditor");
-		if (UNiagaraParameterCollection* Collection = NiagaraEditorModule.FindCollectionForVariable(VarName))
+		if (UNiagaraParameterCollection* Collection = NiagaraEditorModule.FindCollectionForVariable(FString(VarName.ToView())))
 		{
 			return Collection;
 		}
diff --git a/Engine/Plugins/FX/Niagara/Source/NiagaraEditor/Private/NiagaraEditorCommon.cpp b/Engine/Plugins/FX/Niagara/Source/NiagaraEditor/Private/NiagaraEditorCommon.cpp
--- a/Engine/Plugins/FX/Niagara/Source/NiagaraEditor/Private/NiagaraEditorCommon.cpp
+++ b/Engine/Plugins/FX/Niagara/Source/NiagaraEditor/Private/NiagaraEditorCommon.cpp
@@ -53,7 +53,7 @@ void FNiagaraOpInfo::BuildName(FString InName, FString InCategory)
 	Name = FName(*(InCategory + TEXT("::") + InName));
 }
 
-int32 CountType(const TArray<FNiagaraTypeDefinition>& Types, const FNiagaraTypeDefinition& SearchType)
+int32 CountType(TConstArrayView<FNiagaraTypeDefinition> Types, const FNiagaraTypeDefinition& SearchType)
 {
 	int32 Count = 0;
 	for (const FNiagaraTypeDefinition& Type : Types)
@@ -178,7 +178,7 @@ void FNiagaraOpInfo::Init()
 				return true;
 			});
 		Op->NumericOuputTypeSelectionMode = ENiagaraNumericOutputTypeSelectionMode::Custom;
-		Op->CustomNumericResolveFunction.BindLambda([=](const TArray<FNiagaraTypeDefinition>& NonNumericInputs)
+		Op->CustomNumericResolveFunction.BindLambda([=](TConstArrayView<FNiagaraTypeDefinition> NonNumericInputs)
 			{
 				int32 PositionTypes = CountType(NonNumericInputs, FNiagaraTypeDefinition::GetPositionDef());
 				if (PositionTypes > 0)
@@ -229,7 +229,7 @@ void FNiagaraOpInfo::Init()
 				return true;
 			});
 		Op->NumericOuputTypeSelectionMode = ENiagaraNumericOutputTypeSelectionMode::Custom;
-		Op->CustomNumericResolveFunction.BindLambda([=](const TArray<FNiagaraTypeDefinition>& NonNumericInputs)
+		Op->CustomNumericResolveFunction.BindLambda([=](TConstArrayView<FNiagaraTypeDefinition> NonNumericInputs)
 			{
 				int32 PositionTypes = CountType(NonNumericInputs, FNiagaraTypeDefinition::GetPositionDef());
 				if (PositionTypes == 1)
@@ -333,7 +333,7 @@ void FNiagaraOpInfo::Init()
 				return true;
 			});
 		Op->NumericOuputTypeSelectionMode = ENiagaraNumericOutputTypeSelectionMode::Custom;
-		Op->CustomNumericResolveFunction.BindLambda([=](const TArray<FNiagaraTypeDefinition>& NonNumericInputs)
+		Op->CustomNumericResolveFunction.BindLambda([=](TConstArrayView<FNiagaraTypeDefinition> NonNumericInputs)
 			{
 				int32 PositionTypes = CountType(NonNumericInputs, FNiagaraTypeDefinition::GetPositionDef());
 				if (PositionTypes == 1)
@@ -370,7 +370,7 @@ void FNiagaraOpInfo::Init()
 				return true;
 			});
 		Op->NumericOuputTypeSelectionMode = ENiagaraNumericOutputTypeSelectionMode::Custom;
-		Op->CustomNumericResolveFunction.BindLambda([=](const TArray<FNiagaraTypeDefinition>& NonNumericInputs)
+		Op->CustomNumericResolveFunction.BindLambda([=](TConstArrayView<FNiagaraTypeDefinition> NonNumericInputs)
 			{
 				 int32 PositionTypes = CountType(NonNumericInputs, FNiagaraTypeDefinition::GetPositionDef());
 				if (PositionTypes == 2)
@@ -924,7 +924,7 @@ void FNiagaraOpInfo::Init()
 				return true;
 			});
 		Op->NumericOuputTypeSelectionMode = ENiagaraNumericOutputTypeSelectionMode::Custom;
-		Op->CustomNumericResolveFunction.BindLambda([=](const TArray<FNiagaraTypeDefinition>& NonNumericInputs)
+		Op->CustomNumericResolveFunction.BindLambda([=](TConstArrayView<FNiagaraTypeDefinition> NonNumericInputs)
 			{
 				 int32 PositionTypes = CountType(NonNumericInputs, FNiagaraTypeDefinition::GetPositionDef());
 				if (PositionTypes == 3)
@@ -958,7 +958,7 @@ void FNiagaraOpInfo::Init()
 				return true;
 			});
 		Op->NumericOuputTypeSelectionMode = ENiagaraNumericOutputTypeSelectionMode::Custom;
-		Op->CustomNumericResolveFunction.BindLambda([=](const TArray<FNiagaraTypeDefinition>& NonNumericInputs)
+		Op->CustomNumericResolveFunction.BindLambda([=](TConstArrayView<FNiagaraTypeDefinition> NonNumericInputs)
 			{
 				 int32 PositionTypes = CountType(NonNumericInputs, FNiagaraTypeDefinition::GetPositionDef());
 				if (PositionTypes == 2)
@@ -992,7 +992,7 @@ void FNiagaraOpInfo::Init()
 				return true;
 			});
 		Op->NumericOuputTypeSelectionMode = ENiagaraNumericOutputTypeSelectionMode::Custom;
-		Op->CustomNumericResolveFunction.BindLambda([=](const TArray<FNiagaraTypeDefinition>& NonNumericInputs)
+		Op->CustomNumericResolveFunction.BindLambda([=](TConstArrayView<FNiagaraTypeDefinition> NonNumericInputs)
 			{
 				 int32 PositionTypes = CountType(NonNumericInputs, FNiagaraTypeDefinition::GetPositionDef());
 				if (PositionTypes == 2)
diff --git a/Engine/Plugins/FX/Niagara/Source/NiagaraEditor/Private/NiagaraEditorUtilities.cpp b/Engine/Plugins/FX/Niagara/Source/NiagaraEditor/Private/NiagaraEditorUtilities.cpp
--- a/Engine/Plugins/FX/Niagara/Source/NiagaraEditor/Private/NiagaraEditorUtilities.cpp
+++ b/Engine/Plugins/FX/Niagara/Source/NiagaraEditor/Private/NiagaraEditorUtilities.cpp
@@ -4557,7 +4557,7 @@ UNiagaraDataInterface* FNiagaraEditorUtilities::GetResolvedRuntimeInstanceForEdi
 	return nullptr;
 }
 
-void FNiagaraParameterUtilities::FilterToRelevantStaticVariables(const TArray<FNiagaraVariable>& InVars, TArray<FNiagaraVariable>& OutVars, FName InOldEmitterAlias, FName InNewEmitterAlias, bool bFilterByEmitterAliasAndConvertToUnaliased)
+void FNiagaraParameterUtilities::FilterToRelevantStaticVariables(TConstArrayView<FNiagaraVariable> InVars, TArray<FNiagaraVariable>& OutVars, FName InOldEmitterAlias, FName InNewEmitterAlias, bool bFilterByEmitterAliasAndConvertToUnaliased)
 {
 	FNiagaraAliasContext RenameContext(ENiagaraScriptUsage::ParticleSpawnScript);
 	RenameContext.ChangeEmitterName(InOldEmitterAlias.ToString(), InNewEmitterAlias.ToString());
diff --git a/Engine/Plugins/FX/Niagara/Source/NiagaraEditor/Public/NiagaraEditorCommon.h b/Engine/Plugins/FX/Niagara/Source/NiagaraEditor/Public/NiagaraEditorCommon.h
--- a/Engine/Plugins/FX/Niagara/Source/NiagaraEditor/Public/NiagaraEditorCommon.h
+++ b/Engine/Plugins/FX/Niagara/Source/NiagaraEditor/Public/NiagaraEditorCommon.h
@@ -74,7 +74,7 @@ public:
 
 	DECLARE_DELEGATE_RetVal_OneParam(int32, FStaticVariableResolve, const TArray<int32>& );
 	DECLARE_DELEGATE_RetVal_TwoParams(bool, FInputTypeValidation, const TArray<FNiagaraTypeDefinition>&, FText& );;
-	DECLARE_DELEGATE_RetVal_OneParam(FNiagaraTypeDefinition, FCustomNumericResolve, const TArray<FNiagaraTypeDefinition>& )
+	DECLARE_DELEGATE_RetVal_OneParam(FNiagaraTypeDefinition, FCustomNumericResolve, TConstArrayView<FNiagaraTypeDefinition> )
 
 	FStaticVariableResolve StaticVariableResolveFunction;
 	FInputTypeValidation InputTypeValidationFunction;
diff --git a/Engine/Plugins/FX/Niagara/Source/NiagaraEditor/Public/NiagaraEditorUtilities.h b/Engine/Plugins/FX/Niagara/Source/NiagaraEditor/Public/NiagaraEditorUtilities.h
--- a/Engine/Plugins/FX/Niagara/Source/NiagaraEditor/Public/NiagaraEditorUtilities.h
+++ b/Engine/Plugins/FX/Niagara/Source/NiagaraEditor/Public/NiagaraEditorUtilities.h
@@ -495,7 +495,7 @@ namespace FNiagaraParameterUtilities
 	/** Creates a tooltip based on a parameter. Also shows the value, if allocated and enabled. */
 	NIAGARAEDITOR_API TSharedRef<SToolTip> GetTooltipWidget(FNiagaraVariable Variable, bool bShowValue = true, TSharedPtr<SWidget> AdditionalVerticalWidget = nullptr,  TSharedPtr<SWidget> AdditionalHorizontalWidget = nullptr);
 
-	NIAGARAEDITOR_API void FilterToRelevantStaticVariables(const TArray<FNiagaraVariable>& InVars, TArray<FNiagaraVariable>& OutVars, FName InOldEmitterAlias, FName InNewEmitterAlias, bool bFilterByEmitterAliasAndConvertToUnaliased);
+	NIAGARAEDITOR_API void FilterToRelevantStaticVariables(TConstArrayView<FNiagaraVariable> InVars, TArray<FNiagaraVariable>& OutVars, FName InOldEmitterAlias, FName InNewEmitterAlias, bool bFilterByEmitterAliasAndConvertToUnaliased);
 };
 
 namespace FNiagaraParameterDefinitionsUtilities
