Removed 'EditAnywhere' flag from Category & Section names so inline widget editing is what you have to do. This means the user can't easily create sibling categories with the same name.
Also made category dropped on logic more robust, giving more info back.
Also changes 'AddCategory' to add a category to the currently selected category, if available, instead of the root. Makes it easy to create nested categories.

diff --git a/Engine/Plugins/FX/Niagara/Source/NiagaraEditor/Private/ViewModels/HierarchyEditor/NiagaraHierarchyViewModelBase.cpp b/Engine/Plugins/FX/Niagara/Source/NiagaraEditor/Private/ViewModels/HierarchyEditor/NiagaraHierarchyViewModelBase.cpp
--- a/Engine/Plugins/FX/Niagara/Source/NiagaraEditor/Private/ViewModels/HierarchyEditor/NiagaraHierarchyViewModelBase.cpp
+++ b/Engine/Plugins/FX/Niagara/Source/NiagaraEditor/Private/ViewModels/HierarchyEditor/NiagaraHierarchyViewModelBase.cpp
@@ -274,24 +274,60 @@ FNiagaraHierarchyItemViewModelBase::FCanPerformActionResults FNiagaraHierarchyCa
 	
 	TArray<TSharedPtr<FNiagaraHierarchyCategoryViewModel>> TargetChildrenCategories;
 	GetChildrenViewModelsForType<UNiagaraHierarchyCategory, FNiagaraHierarchyCategoryViewModel>(TargetChildrenCategories);
-
+	
+	TArray<TSharedPtr<FNiagaraHierarchyCategoryViewModel>> SiblingCategories;
+	Parent.Pin()->GetChildrenViewModelsForType<UNiagaraHierarchyCategory, FNiagaraHierarchyCategoryViewModel>(SiblingCategories);
+	
 	// we only allow drops if some general conditions are fulfilled
 	if(DraggedItem->GetData() != GetData() &&
 		(!DraggedItem->HasParent(AsShared(), false) || ItemDropZone != EItemDropZone::OntoItem)  &&
 		!HasParent(DraggedItem, true))
 	{
-		// categories can be dropped on categories, but only if the resulting sibling categories have different names
+		// categories can be dropped on categories, but only if the resulting sibling categories or children categories have different names
 		if(DraggedItem->GetData()->IsA<UNiagaraHierarchyCategory>())
 		{
-			Results.bCanPerform =
-			(
-				(ItemDropZone == EItemDropZone::OntoItem && !TargetChildrenCategories.ContainsByPredicate([DraggedItem](TSharedPtr<FNiagaraHierarchyCategoryViewModel> HierarchyCategoryViewModel)
+			if(ItemDropZone != EItemDropZone::OntoItem)
+			{
+				bool bContainsSiblingWithSameName = SiblingCategories.ContainsByPredicate([DraggedItem](TSharedPtr<FNiagaraHierarchyCategoryViewModel> HierarchyCategoryViewModel)
+					{
+						return DraggedItem->ToString() == HierarchyCategoryViewModel->ToString();
+					});
+
+				if(bContainsSiblingWithSameName)
+				{
+					Results.bCanPerform = false;
+					Results.CanPerformMessage = LOCTEXT("CantDropCategoryOnCategorySameSiblingNames", "A category of the same name already exists here, potentially in a different section. Please rename your category first.");
+					return Results;
+				}
+
+				// if we are making a category a sibling of another at the root level, the section will be set to the currently active section. Let that be known.
+				if(Parent.Pin()->GetData()->IsA<UNiagaraHierarchyRoot>())
+				{
+					UNiagaraHierarchyCategory* DraggedCategory = Cast<UNiagaraHierarchyCategory>(DraggedItem->GetDataMutable());
+					if(DraggedCategory->GetSection() != HierarchyViewModel->GetActiveHierarchySectionData())
+					{
+						FText BaseMessage = LOCTEXT("CategorySectionWillUpdateDueToDrop", "The section of the category will change to {0} after the drop");
+						Results.CanPerformMessage = FText::FormatOrdered(BaseMessage, HierarchyViewModel->GetActiveHierarchySectionData() == nullptr ? FText::FromString("All") : HierarchyViewModel->GetActiveHierarchySectionData()->GetSectionNameAsText());
+					}
+				}
+			}
+			else
+			{
+				bool bContainsChildrenCategoriesWithSameName = TargetChildrenCategories.ContainsByPredicate([DraggedItem](TSharedPtr<FNiagaraHierarchyCategoryViewModel> HierarchyCategoryViewModel)
 					{
 						return DraggedItem->ToString() == HierarchyCategoryViewModel->ToString();
-					}))
-					||
-				(ItemDropZone != EItemDropZone::OntoItem && DraggedItem->ToString() != ToString())
-			);
+					});
+
+				if(bContainsChildrenCategoriesWithSameName)
+				{
+					Results.bCanPerform = false;
+					Results.CanPerformMessage = LOCTEXT("CantDropCategoryOnCategorySameSiblingNames", "A sub-category of the same name already exists! Please rename your category first.");
+					return Results;
+				}
+			}
+
+			Results.bCanPerform = true;
+			return Results;
 		}
 		else if(DraggedItem->GetData()->IsA<UNiagaraHierarchyItem>())
 		{
@@ -1444,12 +1480,14 @@ void FNiagaraHierarchyCategoryViewModel::OnDroppedOnInternal(TSharedPtr<FNiagara
 	}
 }
 
-void UNiagaraHierarchyViewModelBase::AddCategory() const
+void UNiagaraHierarchyViewModelBase::AddCategory(TSharedPtr<FNiagaraHierarchyItemViewModelBase> CategoryParent) const
 {
-	TSharedPtr<FNiagaraHierarchyItemViewModelBase> ViewModel = HierarchyRootViewModel->AddNewItem(UNiagaraHierarchyCategory::StaticClass());
+	TSharedPtr<FNiagaraHierarchyItemViewModelBase> ViewModel = CategoryParent->AddNewItem(UNiagaraHierarchyCategory::StaticClass());
 
-	if(UNiagaraHierarchyCategory* Category = ViewModel->GetDataMutable<UNiagaraHierarchyCategory>())
+	if(ensure(ViewModel.IsValid()))
 	{
+		UNiagaraHierarchyCategory* Category = CastChecked<UNiagaraHierarchyCategory>(ViewModel->GetDataMutable());
+		
 		TArray<UNiagaraHierarchyCategory*> SiblingCategories;
 		Category->GetTypedOuter<UNiagaraHierarchyItemBase>()->GetChildrenOfType<UNiagaraHierarchyCategory>(SiblingCategories);
 		
diff --git a/Engine/Plugins/FX/Niagara/Source/NiagaraEditor/Private/Widgets/SNiagaraHierarchy.cpp b/Engine/Plugins/FX/Niagara/Source/NiagaraEditor/Private/Widgets/SNiagaraHierarchy.cpp
--- a/Engine/Plugins/FX/Niagara/Source/NiagaraEditor/Private/Widgets/SNiagaraHierarchy.cpp
+++ b/Engine/Plugins/FX/Niagara/Source/NiagaraEditor/Private/Widgets/SNiagaraHierarchy.cpp
@@ -582,7 +582,17 @@ void SNiagaraHierarchy::Construct(const FArguments& InArgs, TObjectPtr<UNiagaraH
 			SNew(SButton)
 			.OnClicked_Lambda([this]()
 			{
-				HierarchyViewModel.Get()->AddCategory();
+				TArray<TSharedPtr<FNiagaraHierarchyItemViewModelBase>> SelectedItems;
+				HierarchyTreeView->GetSelectedItems(SelectedItems);
+
+				if(SelectedItems.Num() == 1 && SelectedItems[0]->GetData()->IsA<UNiagaraHierarchyCategory>())
+				{
+					HierarchyViewModel.Get()->AddCategory(SelectedItems[0]);
+				}
+				else
+				{
+					HierarchyViewModel.Get()->AddCategory(HierarchyViewModel->GetHierarchyRootViewModel());
+				}
 				return FReply::Handled();
 			})
 			.ButtonStyle(FNiagaraEditorStyle::Get(), "NiagaraEditor.HierarchyEditor.ButtonStyle")
@@ -1268,7 +1278,7 @@ void SNiagaraHierarchy::OnItemAdded(TSharedPtr<FNiagaraHierarchyItemViewModelBas
 	if(AddedItem->GetData()->IsA<UNiagaraHierarchyItem>() || AddedItem->GetData()->IsA<UNiagaraHierarchyCategory>())
 	{
 		HierarchyTreeView->RequestTreeRefresh();
-		HierarchyTreeView->SetSelection(AddedItem);
+		NavigateToHierarchyItem(AddedItem);
 	}
 	else if(AddedItem->GetData()->IsA<UNiagaraHierarchySection>())
 	{
diff --git a/Engine/Plugins/FX/Niagara/Source/NiagaraEditor/Public/ViewModels/HierarchyEditor/NiagaraHierarchyViewModelBase.h b/Engine/Plugins/FX/Niagara/Source/NiagaraEditor/Public/ViewModels/HierarchyEditor/NiagaraHierarchyViewModelBase.h
--- a/Engine/Plugins/FX/Niagara/Source/NiagaraEditor/Public/ViewModels/HierarchyEditor/NiagaraHierarchyViewModelBase.h
+++ b/Engine/Plugins/FX/Niagara/Source/NiagaraEditor/Public/ViewModels/HierarchyEditor/NiagaraHierarchyViewModelBase.h
@@ -283,7 +283,7 @@ public:
 	 */
 	void FixupSectionLinkage();
 private:
-	UPROPERTY(EditAnywhere, Category = "Category")
+	UPROPERTY()
 	FName Category;
 
 	/** The tooltip used when the user is hovering this category */
@@ -312,7 +312,7 @@ public:
 
 	virtual FString ToString() const override { return Section.ToString(); }
 private:
-	UPROPERTY(EditAnywhere, Category = "Section")
+	UPROPERTY()
 	FName Section;
 
 	/** The tooltip used when the user is hovering this section */
@@ -393,7 +393,7 @@ public:
 	void Initialize();
 	void Finalize();
 	
-	void AddCategory() const;
+	void AddCategory(TSharedPtr<FNiagaraHierarchyItemViewModelBase> CategoryParent) const;
 	void AddSection() const;
 
 	NIAGARAEDITOR_API void DeleteItemWithIdentity(FNiagaraHierarchyIdentity Identity);
@@ -687,8 +687,8 @@ protected:
 
 	/** Should return true if the given drag drop operation is allowed to succeed. */
 	virtual FCanPerformActionResults CanDropOnInternal(TSharedPtr<FNiagaraHierarchyItemViewModelBase>, EItemDropZone ItemDropZone);
-
-	/** Override this to handle drop-on logic. */
+	
+	/** Override this to handle drop-on logic. This is called when an item has been dropped on the item that has implemented this function. */
 	virtual void OnDroppedOnInternal(TSharedPtr<FNiagaraHierarchyItemViewModelBase>, EItemDropZone ItemDropZone) { }
 
 	/** Can be overridden to support custom sections.
