Static Analysis Fix.

#rnx

diff --git a/Engine/Plugins/FX/Niagara/Source/Niagara/Private/NiagaraScriptExecutionContext.cpp b/Engine/Plugins/FX/Niagara/Source/Niagara/Private/NiagaraScriptExecutionContext.cpp
--- a/Engine/Plugins/FX/Niagara/Source/Niagara/Private/NiagaraScriptExecutionContext.cpp
+++ b/Engine/Plugins/FX/Niagara/Source/Niagara/Private/NiagaraScriptExecutionContext.cpp
@@ -87,25 +87,28 @@ bool FNiagaraScriptExecutionContextBase::Init(FNiagaraSystemInstance* Instance,
 	VectorVMState = AllocVectorVMState(&OptimizeContext);
 #endif
 
-	ENiagaraScriptUsage Usage = Script->GetUsage();
-	for(auto& ResolvedDIInfo : Script->GetResolvedDataInterfaces())
+	if(Script)
 	{
-		if(ResolvedDIInfo.ResolvedDataInterface)
+		ENiagaraScriptUsage Usage = Script->GetUsage();
+		for(auto& ResolvedDIInfo : Script->GetResolvedDataInterfaces())
 		{
-			bHasDIsWithPreStageTick |= (ResolvedDIInfo.ResolvedDataInterface->HasPreStageTick(Usage));
-			bHasDIsWithPostStageTick |= (ResolvedDIInfo.ResolvedDataInterface->HasPostStageTick(Usage));
+			if(ResolvedDIInfo.ResolvedDataInterface)
+			{
+				bHasDIsWithPreStageTick |= (ResolvedDIInfo.ResolvedDataInterface->HasPreStageTick(Usage));
+				bHasDIsWithPostStageTick |= (ResolvedDIInfo.ResolvedDataInterface->HasPostStageTick(Usage));
+			}
+			if(bHasDIsWithPreStageTick && bHasDIsWithPostStageTick)
+			{
+				break;
+			}
 		}
-		if(bHasDIsWithPreStageTick && bHasDIsWithPostStageTick)
+
+		if(Instance && (bHasDIsWithPreStageTick || bHasDIsWithPostStageTick))
 		{
-			break;
+			DIStageTickHandler.Init(Script, Instance);
 		}
 	}
 
-	if(Instance && (bHasDIsWithPreStageTick || bHasDIsWithPostStageTick))
-	{
-		DIStageTickHandler.Init(Script, Instance);
-	}
-
 	return true;
 }
 
