Fixed missing User Parameters in renderer bindings

diff --git a/Engine/Plugins/FX/Niagara/Source/NiagaraEditor/Private/Customizations/NiagaraTypeCustomizations.cpp b/Engine/Plugins/FX/Niagara/Source/NiagaraEditor/Private/Customizations/NiagaraTypeCustomizations.cpp
--- a/Engine/Plugins/FX/Niagara/Source/NiagaraEditor/Private/Customizations/NiagaraTypeCustomizations.cpp
+++ b/Engine/Plugins/FX/Niagara/Source/NiagaraEditor/Private/Customizations/NiagaraTypeCustomizations.cpp
@@ -278,6 +278,15 @@ TArray<FName> FNiagaraVariableAttributeBindingCustomization::GetNames(const FVer
 	GraphsToTraverse.Add(Cast<UNiagaraScriptSource>(System->GetSystemUpdateScript()->GetLatestSource())->NodeGraph);
 
 	TSet<FNiagaraVariableBase> Vars;
+
+	if (System)
+	{
+		for (const FNiagaraVariable UserParameter : System->GetExposedParameters().ReadParameterVariables())
+		{
+			Vars.Add(UserParameter);
+		}
+	}
+	
 	TArray<TArray<FNiagaraVariable>> MapHistoryVars;
 	for(UNiagaraGraph* Graph : GraphsToTraverse)
 	{
