Fixed PVS Warning

diff --git a/Engine/Plugins/FX/Niagara/Source/NiagaraEditor/Private/Customizations/NiagaraTypeCustomizations.cpp b/Engine/Plugins/FX/Niagara/Source/NiagaraEditor/Private/Customizations/NiagaraTypeCustomizations.cpp
--- a/Engine/Plugins/FX/Niagara/Source/NiagaraEditor/Private/Customizations/NiagaraTypeCustomizations.cpp
+++ b/Engine/Plugins/FX/Niagara/Source/NiagaraEditor/Private/Customizations/NiagaraTypeCustomizations.cpp
@@ -274,13 +274,13 @@ TArray<FName> FNiagaraVariableAttributeBindingCustomization::GetNames(const FVer
 	TArray<UNiagaraGraph*> GraphsToTraverse;
 	
 	GraphsToTraverse.Add(Cast<UNiagaraScriptSource>(BaseEmitter.GetEmitterData()->GraphSource)->NodeGraph);
-	UNiagaraSystem* System = BaseEmitter.Emitter->GetTypedOuter<UNiagaraSystem>();
-	GraphsToTraverse.Add(Cast<UNiagaraScriptSource>(System->GetSystemUpdateScript()->GetLatestSource())->NodeGraph);
 
 	TSet<FNiagaraVariableBase> Vars;
 
-	if (System)
+	if(UNiagaraSystem* System = BaseEmitter.Emitter->GetTypedOuter<UNiagaraSystem>())
 	{
+    	GraphsToTraverse.Add(Cast<UNiagaraScriptSource>(System->GetSystemUpdateScript()->GetLatestSource())->NodeGraph);
+	
 		for (const FNiagaraVariable UserParameter : System->GetExposedParameters().ReadParameterVariables())
 		{
 			Vars.Add(UserParameter);
