Fixed mismatched copy constructor/assignment operator in FNiagaraTypeDefinitionHandle and FNiagaraVariableBase.

diff --git a/Engine/Plugins/FX/Niagara/Source/Niagara/Public/NiagaraTypes.h b/Engine/Plugins/FX/Niagara/Source/Niagara/Public/NiagaraTypes.h
--- a/Engine/Plugins/FX/Niagara/Source/Niagara/Public/NiagaraTypes.h
+++ b/Engine/Plugins/FX/Niagara/Source/Niagara/Public/NiagaraTypes.h
@@ -1619,7 +1619,6 @@ struct FNiagaraTypeDefinitionHandle
 
 	FNiagaraTypeDefinitionHandle() : RegisteredTypeIndex(INDEX_NONE) {}
 	explicit FNiagaraTypeDefinitionHandle(const FNiagaraTypeDefinition& Type) : RegisteredTypeIndex(Register(Type)) {}
-	explicit FNiagaraTypeDefinitionHandle(const FNiagaraTypeDefinitionHandle& Handle) : RegisteredTypeIndex(Handle.RegisteredTypeIndex) {}
 
 	const FNiagaraTypeDefinition& operator*() const { return Resolve(); }
 	const FNiagaraTypeDefinition* operator->() const { return &Resolve(); }
@@ -1667,14 +1666,6 @@ struct FNiagaraVariableBase
 #endif
 		{}
 
-	FORCEINLINE FNiagaraVariableBase(const FNiagaraVariableBase &Other)
-		: Name(Other.Name)
-		, TypeDefHandle(Other.TypeDefHandle)
-#if WITH_EDITORONLY_DATA
-		, TypeDef_DEPRECATED(Other.TypeDef_DEPRECATED)
-#endif
-		{}
-
 	FORCEINLINE FNiagaraVariableBase(const FNiagaraTypeDefinition& InType, const FName& InName)
 		: Name(InName)
 		, TypeDefHandle(InType)
