Fixed duplicate loctext key

diff --git a/Engine/Plugins/FX/Niagara/Source/NiagaraEditor/Private/ViewModels/HierarchyEditor/NiagaraHierarchyViewModelBase.cpp b/Engine/Plugins/FX/Niagara/Source/NiagaraEditor/Private/ViewModels/HierarchyEditor/NiagaraHierarchyViewModelBase.cpp
--- a/Engine/Plugins/FX/Niagara/Source/NiagaraEditor/Private/ViewModels/HierarchyEditor/NiagaraHierarchyViewModelBase.cpp
+++ b/Engine/Plugins/FX/Niagara/Source/NiagaraEditor/Private/ViewModels/HierarchyEditor/NiagaraHierarchyViewModelBase.cpp
@@ -296,7 +296,7 @@ FNiagaraHierarchyItemViewModelBase::FCanPerformActionResults FNiagaraHierarchyCa
 				if(bContainsSiblingWithSameName)
 				{
 					Results.bCanPerform = false;
-					Results.CanPerformMessage = LOCTEXT("CantDropCategoryOnCategorySameSiblingNames", "A category of the same name already exists here, potentially in a different section. Please rename your category first.");
+					Results.CanPerformMessage = LOCTEXT("CantDropCategorNextToCategorySameSiblingNames", "A category of the same name already exists here, potentially in a different section. Please rename your category first.");
 					return Results;
 				}
 
@@ -321,7 +321,7 @@ FNiagaraHierarchyItemViewModelBase::FCanPerformActionResults FNiagaraHierarchyCa
 				if(bContainsChildrenCategoriesWithSameName)
 				{
 					Results.bCanPerform = false;
-					Results.CanPerformMessage = LOCTEXT("CantDropCategoryOnCategorySameSiblingNames", "A sub-category of the same name already exists! Please rename your category first.");
+					Results.CanPerformMessage = LOCTEXT("CantDropCategoryOnCategorySameChildCategoryName", "A sub-category of the same name already exists! Please rename your category first.");
 					return Results;
 				}
 			}
