Adds (or makes public) some accessor functions on UNiagaraNode classes to fascilitate creation of compilation copies of the nodes

diff --git a/Engine/Plugins/FX/Niagara/Source/NiagaraEditor/Private/NiagaraNodeConvert.cpp b/Engine/Plugins/FX/Niagara/Source/NiagaraEditor/Private/NiagaraNodeConvert.cpp
--- a/Engine/Plugins/FX/Niagara/Source/NiagaraEditor/Private/NiagaraNodeConvert.cpp
+++ b/Engine/Plugins/FX/Niagara/Source/NiagaraEditor/Private/NiagaraNodeConvert.cpp
@@ -741,6 +741,11 @@ TArray<FNiagaraConvertConnection>& UNiagaraNodeConvert::GetConnections()
 	return Connections;
 }
 
+const TArray<FNiagaraConvertConnection>& UNiagaraNodeConvert::GetConnections() const
+{
+	return Connections;
+}
+
 void UNiagaraNodeConvert::OnPinRemoved(UEdGraphPin* PinToRemove)
 {
 	TSet<FGuid> TypePinIds;
diff --git a/Engine/Plugins/FX/Niagara/Source/NiagaraEditor/Private/NiagaraNodeConvert.h b/Engine/Plugins/FX/Niagara/Source/NiagaraEditor/Private/NiagaraNodeConvert.h
--- a/Engine/Plugins/FX/Niagara/Source/NiagaraEditor/Private/NiagaraNodeConvert.h
+++ b/Engine/Plugins/FX/Niagara/Source/NiagaraEditor/Private/NiagaraNodeConvert.h
@@ -114,6 +114,7 @@ public:
 
 	/** Gets the nodes inner connection. */
 	TArray<FNiagaraConvertConnection>& GetConnections();
+	const TArray<FNiagaraConvertConnection>& GetConnections() const;
 
 	/** Initializes this node as a swizzle by component string. */
 	void InitAsSwizzle(FString Swiz);
