Add debug command line option (-noniagaraddc) to force DDC puts/gets to fail to allow easier debugging of compilation

diff --git a/Engine/Plugins/FX/Niagara/Source/Niagara/Classes/NiagaraScript.h b/Engine/Plugins/FX/Niagara/Source/Niagara/Classes/NiagaraScript.h
--- a/Engine/Plugins/FX/Niagara/Source/Niagara/Classes/NiagaraScript.h
+++ b/Engine/Plugins/FX/Niagara/Source/Niagara/Classes/NiagaraScript.h
@@ -1148,6 +1148,8 @@ public:
 	NIAGARA_API void InvalidateCompileResults(const FString& Reason);
 	NIAGARA_API FText GetDescription(const FGuid& VersionGuid);
 
+	NIAGARA_API static bool ShouldUseDDC();
+
 	/** Helper to convert the struct from its binary data out of the DDC to it's actual in-memory version.
 		Do not call this on anything other than the game thread as it depends on the FObjectAndNameAsStringProxyArchive,
 		which calls FindStaticObject which can fail when used in any other thread!*/
diff --git a/Engine/Plugins/FX/Niagara/Source/Niagara/Private/NiagaraScript.cpp b/Engine/Plugins/FX/Niagara/Source/Niagara/Private/NiagaraScript.cpp
--- a/Engine/Plugins/FX/Niagara/Source/Niagara/Private/NiagaraScript.cpp
+++ b/Engine/Plugins/FX/Niagara/Source/Niagara/Private/NiagaraScript.cpp
@@ -2625,6 +2625,13 @@ static bool ValidateExecData(const UNiagaraScript* Script, const FNiagaraVMExecu
 	return IsValid;
 }
 
+bool UNiagaraScript::ShouldUseDDC()
+{
+	static const bool bNoNiagaraDDC = FParse::Param(FCommandLine::Get(), TEXT("noniagaraddc"));
+
+	return !bNoNiagaraDDC;
+}
+
 bool UNiagaraScript::BinaryToExecData(const UNiagaraScript* Script, const TArray<uint8>& InBinaryData, FNiagaraVMExecutableData& OutExecData)
 {
 	check(IsInGameThread());
@@ -2633,6 +2640,11 @@ bool UNiagaraScript::BinaryToExecData(const UNiagaraScript* Script, const TArray
 		return false;
 	}
 
+	if (!ShouldUseDDC())
+	{
+		return false;
+	}
+
 	FMemoryReader Ar(InBinaryData, true);
 	FObjectAndNameAsStringProxyArchive SafeAr(Ar, false);
 	OutExecData.SerializeData(SafeAr, true);
@@ -2651,6 +2663,11 @@ bool UNiagaraScript::ExecToBinaryData(const UNiagaraScript* Script, TArray<uint8
 {
 	check(IsInGameThread());
 
+	if (!ShouldUseDDC())
+	{
+		return false;
+	}
+
 	FString ValidationErrors;
 	if (!ValidateExecData(Script, InExecData, ValidationErrors))
 	{
