Update NiagaraCompiler interface where the debug name is supplied directly rather than generated internally with the precompile data

diff --git a/Engine/Plugins/FX/Niagara/Source/NiagaraEditor/Private/NiagaraCompiler.cpp b/Engine/Plugins/FX/Niagara/Source/NiagaraEditor/Private/NiagaraCompiler.cpp
--- a/Engine/Plugins/FX/Niagara/Source/NiagaraEditor/Private/NiagaraCompiler.cpp
+++ b/Engine/Plugins/FX/Niagara/Source/NiagaraEditor/Private/NiagaraCompiler.cpp
@@ -5,6 +5,7 @@
 #include "EdGraphUtilities.h"
 #include "INiagaraEditorTypeUtilities.h"
 #include "Misc/FileHelper.h"
+#include "Misc/PathViews.h"
 #include "Modules/ModuleManager.h"
 #include "NiagaraComponent.h"
 #include "NiagaraDataInterface.h"
@@ -83,6 +84,23 @@ static FAutoConsoleVariableRef CVarNiagaraEnablePrecompilerNamespaceFixup(
 	ECVF_Default
 );
 
+namespace NiagaraCompileRequestHelper
+{
+	using FDebugGroupNameBuilder = TStringBuilder<512>;
+	static void BuildScriptDebugGroupName(const FNiagaraCompileRequestData* InCompileRequest, const FNiagaraCompileOptions& InCompileOptions, FDebugGroupNameBuilder& NameBuilder)
+	{
+		FPathViews::Append(NameBuilder, InCompileRequest->SourceName);
+		FPathViews::Append(NameBuilder, InCompileRequest->EmitterUniqueName);
+		FPathViews::Append(NameBuilder, InCompileRequest->ENiagaraScriptUsageEnum->GetNameStringByValue((int64)InCompileOptions.TargetUsage));
+
+		if (InCompileOptions.TargetUsageId.IsValid())
+		{
+			NameBuilder << TEXT("_");
+			InCompileOptions.TargetUsageId.AppendString(NameBuilder, EGuidFormats::Digits);
+		}
+	}
+
+};
 
 static FCriticalSection TranslationCritSec;
 
@@ -1450,7 +1468,10 @@ int32 FNiagaraEditorModule::CompileScript(const FNiagaraCompileRequestDataBase*
 		}
 	}
 
-	int32 JobID = Compiler->CompileScript(CompileRequest, InCompileOptions, TranslateResults, &Translator.GetTranslateOutput(), Translator.GetTranslatedHLSL());
+	NiagaraCompileRequestHelper::FDebugGroupNameBuilder DebugGroupName;
+	NiagaraCompileRequestHelper::BuildScriptDebugGroupName(CompileRequest, InCompileOptions, DebugGroupName);
+
+	int32 JobID = Compiler->CompileScript(DebugGroupName, InCompileOptions, TranslateResults, &Translator.GetTranslateOutput(), Translator.GetTranslatedHLSL());
 	ActiveCompilations.Add(JobID, Compiler);
 	return JobID;
 	
@@ -1738,7 +1759,15 @@ ENiagaraScriptCompileStatus FNiagaraCompileResults::CompileResultsToSummary(cons
 	return SummaryStatus;
 }
 
-int32 FHlslNiagaraCompiler::CompileScript(const FNiagaraCompileRequestData* InCompileRequest, const FNiagaraCompileOptions& InOptions, const FNiagaraTranslateResults& InTranslateResults, FNiagaraTranslatorOutput *TranslatorOutput, FString &TranslatedHLSL)
+int32 FHlslNiagaraCompiler::CompileScript(const FNiagaraCompileRequestData* InCompileRequest, const FNiagaraCompileOptions& InOptions, const FNiagaraTranslateResults& InTranslateResults, FNiagaraTranslatorOutput* TranslatorOutput, FString& TranslatedHLSL)
+{
+	NiagaraCompileRequestHelper::FDebugGroupNameBuilder DebugGroupName;
+	NiagaraCompileRequestHelper::BuildScriptDebugGroupName(InCompileRequest, InOptions, DebugGroupName);
+
+	return CompileScript(DebugGroupName, InOptions, InTranslateResults, TranslatorOutput, TranslatedHLSL);
+}
+
+int32 FHlslNiagaraCompiler::CompileScript(const FStringView GroupName, const FNiagaraCompileOptions& InOptions, const FNiagaraTranslateResults& InTranslateResults, FNiagaraTranslatorOutput *TranslatorOutput, FString &TranslatedHLSL)
 {
 	SCOPE_CYCLE_COUNTER(STAT_NiagaraEditor_HlslCompiler_CompileScript);
 
@@ -1768,8 +1797,7 @@ int32 FHlslNiagaraCompiler::CompileScript(const FNiagaraCompileRequestData* InCo
 	Input.Environment.IncludeVirtualPathToContentsMap.Add(TEXT("/Engine/Generated/NiagaraEmitterInstance.ush"), TranslatedHLSL);
 	Input.bGenerateDirectCompileFile = CVarDumpCommandLine ? CVarDumpCommandLine->GetBool() : false;
 	Input.DumpDebugInfoRootPath = GShaderCompilingManager->GetAbsoluteShaderDebugInfoDirectory() / TEXT("VM");
-	FString UsageIdStr = !InOptions.TargetUsageId.IsValid() ? TEXT("") : (TEXT("_") + InOptions.TargetUsageId.ToString());
-	Input.DebugGroupName = InCompileRequest->SourceName / InCompileRequest->EmitterUniqueName / InCompileRequest->ENiagaraScriptUsageEnum->GetNameStringByValue((int64)InOptions.TargetUsage) + UsageIdStr;
+	Input.DebugGroupName = GroupName;
 	Input.DebugExtension.Empty();
 	Input.DumpDebugInfoPath.Empty();
 
diff --git a/Engine/Plugins/FX/Niagara/Source/NiagaraEditor/Private/NiagaraCompiler.h b/Engine/Plugins/FX/Niagara/Source/NiagaraEditor/Private/NiagaraCompiler.h
--- a/Engine/Plugins/FX/Niagara/Source/NiagaraEditor/Private/NiagaraCompiler.h
+++ b/Engine/Plugins/FX/Niagara/Source/NiagaraEditor/Private/NiagaraCompiler.h
@@ -6,7 +6,6 @@
 #include "NiagaraCommon.h"
 #include "NiagaraHlslTranslator.h"
 #include "INiagaraCompiler.h"
-#include "Kismet2/CompilerResultsLog.h"
 #include "ShaderCompiler.h"
 
 class Error;
@@ -42,7 +41,10 @@ public:
 	virtual ~FHlslNiagaraCompiler() {}
 
 	//Begin INiagaraCompiler Interface
-	virtual int32 CompileScript(const FNiagaraCompileRequestData* InCompileRequest, const FNiagaraCompileOptions& InOptions, const FNiagaraTranslateResults& InTranslateResults, FNiagaraTranslatorOutput *TranslatorOutput, FString& TranslatedHLSL) override;
+	UE_DEPRECATED(5.4, "Please update to supply the GroupName directly as this will be removed in a future version.")
+	virtual int32 CompileScript(const class FNiagaraCompileRequestData* InCompileRequest, const FNiagaraCompileOptions& InOptions, const FNiagaraTranslateResults& InTranslateResults, FNiagaraTranslatorOutput* TranslatorOutput, FString& TranslatedHLSL) override;
+
+	virtual int32 CompileScript(const FStringView GroupName, const FNiagaraCompileOptions& InOptions, const FNiagaraTranslateResults& InTranslateResults, FNiagaraTranslatorOutput* TranslatorOutput, FString& TranslatedHLSL) override;
 	virtual TOptional<FNiagaraCompileResults> GetCompileResult(int32 JobID, bool bWait = false) override;
 
 	virtual void Error(FText ErrorText) override;
diff --git a/Engine/Plugins/FX/Niagara/Source/NiagaraEditor/Public/INiagaraCompiler.h b/Engine/Plugins/FX/Niagara/Source/NiagaraEditor/Public/INiagaraCompiler.h
--- a/Engine/Plugins/FX/Niagara/Source/NiagaraEditor/Public/INiagaraCompiler.h
+++ b/Engine/Plugins/FX/Niagara/Source/NiagaraEditor/Public/INiagaraCompiler.h
@@ -3,6 +3,7 @@
 #pragma once
 
 #include "CoreMinimal.h"
+#include "Containers/StringView.h"
 #include "NiagaraEditorCommon.h"
 #include "NiagaraParameters.h"
 #include "NiagaraScript.h"
@@ -17,9 +18,7 @@ class UNiagaraDataInterface;
 struct FNiagaraTranslateResults;
 struct FNiagaraTranslatorOutput;
 struct FNiagaraVMExecutableData;
-class FNiagaraCompileRequestData;
 class FNiagaraCompileOptions;
-//struct FNiagaraCompileEvent;
 
 /** Defines information about the results of a Niagara script compile. */
 struct FNiagaraCompileResults
@@ -61,8 +60,11 @@ struct FNiagaraCompileResults
 class INiagaraCompiler
 {
 public:
+	UE_DEPRECATED(5.4, "Please update to supply the GroupName directly as this will be removed in a future version.")
+	virtual int32 CompileScript(const class FNiagaraCompileRequestData* InCompileRequest, const FNiagaraCompileOptions& InOptions, const FNiagaraTranslateResults& InTranslateResults, FNiagaraTranslatorOutput* TranslatorOutput, FString& TranslatedHLSL) = 0;
+
 	/** Starts the async compilation of a script and returns the job handle to retrieve the results */
-	virtual int32 CompileScript(const FNiagaraCompileRequestData* InCompileRequest, const FNiagaraCompileOptions& InOptions, const FNiagaraTranslateResults& InTranslateResults, FNiagaraTranslatorOutput *TranslatorOutput, FString& TranslatedHLSL) = 0;
+	virtual int32 CompileScript(const FStringView GroupName, const FNiagaraCompileOptions& InOptions, const FNiagaraTranslateResults& InTranslateResults, FNiagaraTranslatorOutput* TranslatorOutput, FString& TranslatedHLSL) = 0;
 
 	/** Returns the compile result for a given job id once the job has finished compiling. */
 	virtual TOptional<FNiagaraCompileResults> GetCompileResult(int32 JobID, bool bWait = false) = 0;
