- Fix for using instance parameters vs exposed parameters
- These are not equal and can result in binding the wrong data interface to the parameter store

#rnx

diff --git a/Engine/Plugins/FX/Niagara/Source/Niagara/Private/NiagaraResolveDIHelpers.cpp b/Engine/Plugins/FX/Niagara/Source/Niagara/Private/NiagaraResolveDIHelpers.cpp
--- a/Engine/Plugins/FX/Niagara/Source/Niagara/Private/NiagaraResolveDIHelpers.cpp
+++ b/Engine/Plugins/FX/Niagara/Source/Niagara/Private/NiagaraResolveDIHelpers.cpp
@@ -112,11 +112,12 @@ namespace FNiagaraResolveDIHelpers
 		TMap<FNiagaraVariableBase, FNiagaraVariableBase> VariableBindingMap;
 
 		// Add user parameters as assignments.
-		for (const FNiagaraVariableWithOffset& UserParameterWithOffset : System->GetExposedParameters().ReadParameterVariables())
+		const FNiagaraParameterStore& InstanceParameters = System->GetSystemCompiledData().InstanceParamStore;
+		for (const FNiagaraVariableWithOffset& UserParameterWithOffset : InstanceParameters.ReadParameterVariables())
 		{
 			if (UserParameterWithOffset.IsDataInterface())
 			{
-				UNiagaraDataInterface* DataInterface = System->GetExposedParameters().GetDataInterface(UserParameterWithOffset.Offset);
+				UNiagaraDataInterface* DataInterface = InstanceParameters.GetDataInterface(UserParameterWithOffset.Offset);
 				if (DataInterface != nullptr)
 				{
 					FDataInterfaceSourceEmitterNamePair& NewAssignment = VariableAssignmentMap.Add(UserParameterWithOffset);
@@ -231,7 +232,8 @@ namespace FNiagaraResolveDIHelpers
 
 				if (FNiagaraUserRedirectionParameterStore::IsUserParameter(BoundVariable))
 				{
-					const int32* UserParameterIndex = System->GetExposedParameters().FindParameterOffset(BoundVariable);
+					const FNiagaraParameterStore& InstanceParameters = System->GetSystemCompiledData().InstanceParamStore;
+					const int32* UserParameterIndex = InstanceParameters.FindParameterOffset(BoundVariable);
 					if (UserParameterIndex != nullptr && *UserParameterIndex != INDEX_NONE)
 					{
 						UserDataInterfaceBindings.Add(FNiagaraResolvedUserDataInterfaceBinding(*UserParameterIndex, ResolvedDataInterfaceIndex));
