- Improve error output if data interfaces do not match on GPU & avoid crashing

#rnx

diff --git a/Engine/Plugins/FX/Niagara/Source/Niagara/Private/NiagaraComputeExecutionContext.cpp b/Engine/Plugins/FX/Niagara/Source/Niagara/Private/NiagaraComputeExecutionContext.cpp
--- a/Engine/Plugins/FX/Niagara/Source/Niagara/Private/NiagaraComputeExecutionContext.cpp
+++ b/Engine/Plugins/FX/Niagara/Source/Niagara/Private/NiagaraComputeExecutionContext.cpp
@@ -137,25 +137,39 @@ bool FNiagaraComputeExecutionContext::Tick(FNiagaraSystemInstance* ParentSystemI
 	if (CombinedParamStore.GetInterfacesDirty())
 	{
 #if DO_CHECK
-		const TArray<UNiagaraDataInterface*> &DataInterfaces = CombinedParamStore.GetDataInterfaces();
 		// We must make sure that the data interfaces match up between the original script values and our overrides...
-		if (DIClassNames.Num() != DataInterfaces.Num())
+		const TArray<UNiagaraDataInterface*>& DataInterfaces = CombinedParamStore.GetDataInterfaces();
+		bool bHasMismatchedDIs = DIClassNames.Num() != DataInterfaces.Num();
+		if (!bHasMismatchedDIs)
 		{
-			UE_LOG(LogNiagara, Warning, TEXT("Mismatch between Niagara GPU Execution Context data interfaces and those in its script!"));
-			return false;
+			for (int32 i = 0; i < DIClassNames.Num(); ++i)
+			{
+				const UNiagaraDataInterface* UsedDataInterface = DataInterfaces[i];
+				const FName UsedClassName = UsedDataInterface ? UsedDataInterface->GetClass()->GetFName() : NAME_None;
+				bHasMismatchedDIs |= DIClassNames[i] != UsedClassName;
+			}
 		}
 
-		for (int32 i=0; i < DIClassNames.Num(); ++i)
+		if (bHasMismatchedDIs)
 		{
-			UNiagaraDataInterface* UsedDataInterface = DataInterfaces[i];
-			checkf(UsedDataInterface, TEXT("DataInterface for GPU execution context is nullptr. System(%s) Index(%d) ExpectedType(%s)"), *GetNameSafe(ParentSystemInstance->GetSystem()), i, *DIClassNames[i].ToString());
+			UE_LOG(LogNiagara, Error, TEXT("Niagara GPU Execution Context Mismatch with DataInterfaces. System (%s) will not run"), *GetNameSafe(ParentSystemInstance->GetSystem()));
 
-			const FName UsedClassName = UsedDataInterface->GetClass()->GetFName();
-			if (DIClassNames[i] != UsedClassName)
+			const int32 MaxDataInterfaces = FMath::Max(DIClassNames.Num(), DataInterfaces.Num());
+			for (int32 i = 0; i < MaxDataInterfaces; ++i)
 			{
-				UE_LOG(LogNiagara, Warning, TEXT("Mismatched class between Niagara GPU Execution Context data interfaces and those in its script! System(%s) Index(%d) ExpectedType(%s) CombinedParamStoreType(%s)"), *GetNameSafe(ParentSystemInstance->GetSystem()), i, *DIClassNames[i].ToString(), *UsedClassName.ToString());
+				const FName SourceType = DataInterfaces.IsValidIndex(i) ? DataInterfaces[i]->GetClass()->GetFName() : NAME_None;
+				const FName ExpectedType = DIClassNames.IsValidIndex(i) ? DIClassNames[i] : NAME_None;
+				const TCHAR* ErrorType = SourceType.IsNone() || ExpectedType.IsNone() ? TEXT("Invalid Array") : TEXT("");
+				if (SourceType != ExpectedType)
+				{
+					ErrorType = TEXT("Mismatched");
+				}
+				UE_LOG(LogNiagara, Error, TEXT(" - DI(%d) SourceType(%s) ExpectedType(%s) %s"), i, *SourceType.ToString(), *ExpectedType.ToString(), ErrorType);
 			}
+
+			return false;
 		}
+
 #endif
 		if (CombinedParamStore.GetPositionDataDirty())
 		{
