- Allow UObject user parameters to be cleared

#rnx

diff --git a/Engine/Plugins/FX/Niagara/Source/NiagaraEditor/Private/Customizations/NiagaraComponentDetails.cpp b/Engine/Plugins/FX/Niagara/Source/NiagaraEditor/Private/Customizations/NiagaraComponentDetails.cpp
--- a/Engine/Plugins/FX/Niagara/Source/NiagaraEditor/Private/Customizations/NiagaraComponentDetails.cpp
+++ b/Engine/Plugins/FX/Niagara/Source/NiagaraEditor/Private/Customizations/NiagaraComponentDetails.cpp
@@ -425,7 +425,6 @@ void FNiagaraUserParameterNodeBuilder::GenerateRowForUserParameter(IDetailChildr
 			.ObjectPath(ParameterProxy.ToSharedRef(), &FNiagaraParameterProxy::GetCurrentAssetPath)
 			.AllowedClass(UserParameter.GetType().GetClass())
 			.OnObjectChanged(ParameterProxy.ToSharedRef(), &FNiagaraParameterProxy::OnAssetSelectedFromPicker)
-			.AllowClear(false)
 			.DisplayUseSelected(true)
 			.DisplayBrowse(true)
 			.DisplayThumbnail(true)
