- Fix crash when versioning an emitter with renderer bindings
- Added missing context to destroy and wait on existing simulations
- Needed to unbind / bind emitter parameter store as the emitter data will be reallocated when the new version is added

diff --git a/Engine/Plugins/FX/Niagara/Source/Niagara/Private/NiagaraEmitter.cpp b/Engine/Plugins/FX/Niagara/Source/Niagara/Private/NiagaraEmitter.cpp
--- a/Engine/Plugins/FX/Niagara/Source/Niagara/Private/NiagaraEmitter.cpp
+++ b/Engine/Plugins/FX/Niagara/Source/Niagara/Private/NiagaraEmitter.cpp
@@ -2508,6 +2508,14 @@ FGuid UNiagaraEmitter::AddNewVersion(int32 MajorVersion, int32 MinorVersion)
 	check(MajorVersion != 1 || MinorVersion != 0);
 	Modify();
 
+	FNiagaraSystemUpdateContext UpdateContext;
+	UpdateContext.SetDestroyOnAdd(true);
+	UpdateContext.SetDestroySystemSim(true);
+	if (UNiagaraSystem* Owner = GetTypedOuter<UNiagaraSystem>())
+	{
+		UpdateContext.Add(Owner, true);
+	}
+
 	FVersionedNiagaraEmitterData NewVersionData;
 	FNiagaraAssetVersion NewVersion = { MajorVersion, MinorVersion, FGuid::NewGuid() };
 	
diff --git a/Engine/Plugins/FX/Niagara/Source/Niagara/Private/NiagaraEmitterInstance.cpp b/Engine/Plugins/FX/Niagara/Source/Niagara/Private/NiagaraEmitterInstance.cpp
--- a/Engine/Plugins/FX/Niagara/Source/Niagara/Private/NiagaraEmitterInstance.cpp
+++ b/Engine/Plugins/FX/Niagara/Source/Niagara/Private/NiagaraEmitterInstance.cpp
@@ -476,15 +476,14 @@ void FNiagaraEmitterInstance::Init(int32 InEmitterIdx, FNiagaraSystemInstanceID
 		if (bAnyRendererBindingsAdded)
 		{
 			if (ParentSystemInstance)
+			{
 				ParentSystemInstance->GetInstanceParameters().Bind(&RendererBindings);
+			}
 
 			if (GPUExecContext && EmitterData->SimTarget == ENiagaraSimTarget::GPUComputeSim)
 			{
 				GPUExecContext->CombinedParamStore.Bind(&RendererBindings);
 			}
-
-			// Handle populating static variable values into the list
-			CachedEmitter.GetEmitterData()->RendererBindings.Bind(&RendererBindings);
 		}
 	}	
 
@@ -703,6 +702,8 @@ void FNiagaraEmitterInstance::UnbindParameters(bool bExternalOnly)
 				}
 			}
 		}
+
+		CachedEmitter.GetEmitterData()->RendererBindings.Unbind(&RendererBindings);
 	}
 	else
 	{
@@ -844,6 +845,8 @@ void FNiagaraEmitterInstance::BindParameters(bool bExternalOnly)
 			ParentSystemInstance->GetInstanceParameters().Bind(&RendererBindings);
 		}
 
+		CachedEmitter.GetEmitterData()->RendererBindings.Bind(&RendererBindings);
+
 		//SystemScriptDefinedDataInterfaceParameters.Bind(&RendererBindings);
 		//ScriptDefinedDataInterfaceParameters.Bind(&RendererBindings);
 	}
