Implementing Icons for Sim Cache Attribute Spreadsheet Capture

diff --git a/Engine/Plugins/FX/Niagara/Source/NiagaraEditor/Private/NiagaraEditorStyle.cpp b/Engine/Plugins/FX/Niagara/Source/NiagaraEditor/Private/NiagaraEditorStyle.cpp
--- a/Engine/Plugins/FX/Niagara/Source/NiagaraEditor/Private/NiagaraEditorStyle.cpp
+++ b/Engine/Plugins/FX/Niagara/Source/NiagaraEditor/Private/NiagaraEditorStyle.cpp
@@ -813,6 +813,9 @@ void FNiagaraEditorStyle::InitSimCacheEditor()
 		.SetCheckedHoveredImage(FSlateRoundedBoxBrush(FStyleColors::Select, 2.0f))
 		.SetCheckedPressedImage(FSlateRoundedBoxBrush(FStyleColors::Select, 2.0f))
 		);
+
+	Set("NiagaraEditor.SimCache.CaptureSingleIcon", new IMAGE_BRUSH_SVG("Icons/AttributeSpreadsheet/NiagaraAttrSpreadsheetCaptureSingle", Icon20x20));
+	Set("NiagaraEditor.SimCache.CaptureMultiIcon", new IMAGE_BRUSH_SVG("Icons/AttributeSpreadsheet/NiagaraAttrSpreadsheetCaptureMulti", Icon20x20));
 }
 
 FNiagaraEditorStyle::FNiagaraEditorStyle() : FSlateStyleSet("NiagaraEditorStyle")
diff --git a/Engine/Plugins/FX/Niagara/Source/NiagaraEditor/Private/Widgets/SNiagaraDebugCaptureView.cpp b/Engine/Plugins/FX/Niagara/Source/NiagaraEditor/Private/Widgets/SNiagaraDebugCaptureView.cpp
--- a/Engine/Plugins/FX/Niagara/Source/NiagaraEditor/Private/Widgets/SNiagaraDebugCaptureView.cpp
+++ b/Engine/Plugins/FX/Niagara/Source/NiagaraEditor/Private/Widgets/SNiagaraDebugCaptureView.cpp
@@ -3,6 +3,7 @@
 #include "SNiagaraDebugCaptureView.h"
 
 #include "NiagaraComponent.h"
+#include "NiagaraEditorStyle.h"
 #include "NiagaraEditorUtilities.h"
 #include "NiagaraSettings.h"
 #include "PropertyEditorModule.h"
@@ -176,9 +177,9 @@ FSlateIcon SNiagaraDebugCaptureView::GetCaptureIcon()
 	switch(FrameMode)
 	{
 	case ENiagaraDebugCaptureFrameMode::SingleFrame:
-		return FSlateIcon();
+		return FSlateIcon (FNiagaraEditorStyle::Get().GetStyleSetName(), "NiagaraEditor.SimCache.CaptureSingleIcon");
 	case ENiagaraDebugCaptureFrameMode::MultiFrame:
-		return FSlateIcon();
+		return FSlateIcon (FNiagaraEditorStyle::Get().GetStyleSetName(), "NiagaraEditor.SimCache.CaptureMultiIcon");
 	default:
 		return FSlateIcon();
 	}
