Added a wrapper for sim cache backed tests

diff --git a/Engine/Plugins/FX/Niagara/Source/Niagara/Classes/NiagaraSimCache.h b/Engine/Plugins/FX/Niagara/Source/Niagara/Classes/NiagaraSimCache.h
--- a/Engine/Plugins/FX/Niagara/Source/Niagara/Classes/NiagaraSimCache.h
+++ b/Engine/Plugins/FX/Niagara/Source/Niagara/Classes/NiagaraSimCache.h
@@ -143,6 +143,9 @@ struct FNiagaraSimCacheDataBuffers
 
 	UPROPERTY()
 	TArray<uint32> InterpMapping;
+	
+	bool operator==(const FNiagaraSimCacheDataBuffers& Other) const;
+	bool operator!=(const FNiagaraSimCacheDataBuffers& Other) const;
 };
 
 USTRUCT()
@@ -497,6 +500,13 @@ public:
 	UFUNCTION(BlueprintCallable, Category = NiagaraSimCache)
 	void ReadColorAttribute(TArray<FLinearColor>& OutValues, FName AttributeName = FName("Color"), FName EmitterName = NAME_None, int FrameIndex = 0) const;
 
+	/**
+	Reads Niagara ID attributes by name from the cache frame and appends them into the OutValues array.
+	EmitterName - If left blank will return the system simulation attributes.
+	*/
+	UFUNCTION(BlueprintCallable, Category = NiagaraSimCache)
+	void ReadIDAttribute(TArray<FNiagaraID>& OutValues, FName AttributeName, FName EmitterName, int FrameIndex = 0) const;
+
 	/**
 	Reads Niagara Position attributes by name from the cache frame and appends them into the OutValues array.
 	Local space emitters provide data at local locations unless bLocalSpaceToWorld is true.
@@ -505,7 +515,6 @@ public:
 	*/
 	UFUNCTION(BlueprintCallable, Category = NiagaraSimCache, meta=(AdvancedDisplay="bLocalSpaceToWorld"))
 	void ReadPositionAttribute(TArray<FVector>& OutValues, FName AttributeName = FName("Position"), FName EmitterName = NAME_None, bool bLocalSpaceToWorld = true, int FrameIndex = 0) const;
-	
 
 	/**
 	Reads Niagara Position attributes by name from the cache frame and appends them into the OutValues array.
@@ -534,6 +543,11 @@ public:
 	UFUNCTION(BlueprintCallable, Category = NiagaraSimCache)
 	void ReadQuatAttributeWithRebase(TArray<FQuat>& OutValues, FQuat Quat, FName AttributeName = FName("MeshOrientation"), FName EmitterName = NAME_None, int FrameIndex = 0) const;
 
+#if WITH_EDITORONLY_DATA
+	UE_DEPRECATED(5.3, "For test purposes only, don't use in production code.")
+	bool IsDataEqual(const UNiagaraSimCache& OtherCache, FString& OutDifference) const;
+#endif
+	
 private:
 	UPROPERTY(VisibleAnywhere, Category=SimCache)
 	FGuid CacheGuid;
diff --git a/Engine/Plugins/FX/Niagara/Source/Niagara/Private/NiagaraSimCache.cpp b/Engine/Plugins/FX/Niagara/Source/Niagara/Private/NiagaraSimCache.cpp
--- a/Engine/Plugins/FX/Niagara/Source/Niagara/Private/NiagaraSimCache.cpp
+++ b/Engine/Plugins/FX/Niagara/Source/Niagara/Private/NiagaraSimCache.cpp
@@ -898,6 +898,19 @@ void UNiagaraSimCache::ReadColorAttribute(TArray<FLinearColor>& OutValues, FName
 	}
 }
 
+void UNiagaraSimCache::ReadIDAttribute(TArray<FNiagaraID>& OutValues, FName AttributeName, FName EmitterName, int FrameIndex) const
+{
+	FNiagaraSimCacheAttributeReaderHelper AttributeReader(this, EmitterName, AttributeName, FrameIndex);
+	if (AttributeReader.IsValid() && AttributeReader.Variable->Variable.GetType() == FNiagaraTypeDefinition::GetIDDef())
+	{
+		const int32 OutValueOffset = OutValues.AddUninitialized(AttributeReader.GetNumInstances());
+		for (int32 i = 0; i < AttributeReader.GetNumInstances(); ++i)
+		{
+			OutValues[OutValueOffset + i] = AttributeReader.ReadID(i);
+		}
+	}
+}
+
 void UNiagaraSimCache::ReadPositionAttribute(TArray<FVector>& OutValues, FName AttributeName, FName EmitterName, bool bLocalSpaceToWorld, int FrameIndex) const
 {
 	FNiagaraSimCacheAttributeReaderHelper AttributeReader(this, EmitterName, AttributeName, FrameIndex);
@@ -980,3 +993,152 @@ void UNiagaraSimCache::ReadQuatAttributeWithRebase(TArray<FQuat>& OutValues, FQu
 	}
 }
 
+bool FNiagaraSimCacheDataBuffers::operator==(const FNiagaraSimCacheDataBuffers& Other) const
+{
+	bool InstanceCountSame = NumInstances == Other.NumInstances;
+	bool SameAquireTag = IDAcquireTag == Other.IDAcquireTag;
+	bool FloatDataSame = FloatData == Other.FloatData;
+	bool HalfDataSame = HalfData == Other.HalfData;
+	bool IntDataSame = Int32Data == Other.Int32Data;
+	bool IndexTableSame = IDToIndexTable == Other.IDToIndexTable;
+	bool MappingSame = InterpMapping == Other.InterpMapping;
+	return InstanceCountSame && SameAquireTag && FloatDataSame && HalfDataSame
+		&& IntDataSame && IndexTableSame && MappingSame;
+}
+bool FNiagaraSimCacheDataBuffers::operator!=(const FNiagaraSimCacheDataBuffers& Other) const
+{
+	return !(*this == Other);
+}
+
+#if WITH_EDITORONLY_DATA
+
+template<typename T>
+bool CompareAttributeData(TArray<int32> OtherIDValues, TMap<int32, int32> ExpectedIDToRow, TArray<T> ExpectedData, TArray<T> OtherData, int32 DataCount)
+{
+	if (DataCount == 0)
+	{
+		return true;
+	}
+	for (int Index = 0; Index < OtherIDValues.Num(); Index++) {
+		int32& OtherID = OtherIDValues[Index];
+		int32 ExpectedDataRow = ExpectedIDToRow.FindOrAdd(OtherID, INDEX_NONE);
+		for (int32 Offset = 0; Offset < DataCount; Offset++)
+		{
+			int32 OtherDataIndex = Index + Offset * OtherIDValues.Num();
+			int32 ExpectedDataIndex = ExpectedDataRow + Offset * OtherIDValues.Num();
+			if (ExpectedDataRow == INDEX_NONE || !OtherData.IsValidIndex(OtherDataIndex) || !ExpectedData.IsValidIndex(ExpectedDataIndex) || OtherData[OtherDataIndex] != ExpectedData[ExpectedDataIndex])
+			{
+				return false;
+			}
+		}
+	}
+	return true;
+}
+
+void AddError(TStringBuilder<1024>& Builder, const FString& Str)
+{
+	if (Builder.Len() < 1024)
+	{
+		Builder.Append(Str);
+	}
+}
+
+bool UNiagaraSimCache::IsDataEqual(const UNiagaraSimCache& OtherCache, FString& OutDifference) const
+{
+	const TArray<FNiagaraSimCacheFrame>& OtherFrames = OtherCache.CacheFrames;
+	if (GetNumFrames() != OtherCache.GetNumFrames())
+	{
+		OutDifference = TEXT("Cache Frame count different");
+		return false;
+	}
+
+	bool bEqual = true;
+	TStringBuilder<1024> Errors;
+	for (int FrameIndex = 0; FrameIndex < CacheFrames.Num(); FrameIndex++)
+	{
+		FString Frame = FString::FromInt(FrameIndex);
+		const FNiagaraSimCacheFrame& ExpectedFrame = CacheFrames[FrameIndex];
+		const FNiagaraSimCacheFrame& OtherFrame = OtherFrames[FrameIndex];
+
+		if (OtherFrame.SimulationAge != ExpectedFrame.SimulationAge)
+		{
+			AddError(Errors, TEXT("Simulation age different - frame ") + Frame + '\n');
+			bEqual = false;
+		}
+		if (OtherFrame.EmitterData.Num() != ExpectedFrame.EmitterData.Num())
+		{
+			AddError(Errors, TEXT("Emitter count different - frame ") + Frame + '\n');
+			bEqual = false;
+			continue;
+		}
+		
+		if (OtherFrame.SystemData.SystemDataBuffers != ExpectedFrame.SystemData.SystemDataBuffers)
+		{
+			AddError(Errors, TEXT("Cached system data different - frame ") + Frame + '\n');
+			bEqual = false;
+		}
+		for (int EmitterIndex = 0; EmitterIndex < ExpectedFrame.EmitterData.Num(); EmitterIndex++)
+		{
+			FName EmitterName = CacheLayout.EmitterLayouts[EmitterIndex].LayoutName;
+			const FNiagaraSimCacheEmitterFrame& ExpectedEmitterFrame = ExpectedFrame.EmitterData[EmitterIndex];
+			const FNiagaraSimCacheEmitterFrame& OtherEmitterFrame = OtherFrame.EmitterData[EmitterIndex];
+			if (OtherEmitterFrame.ParticleDataBuffers.NumInstances != ExpectedEmitterFrame.ParticleDataBuffers.NumInstances)
+			{
+				AddError(Errors, TEXT("Particle count different - frame ") + Frame + TEXT(" emitter ") + EmitterName.ToString() + '\n');
+				bEqual = false;
+				continue;
+			}
+
+			FName IdName("UniqueID");
+			TArray<int32> ExpectedIDValues;
+			TArray<int32> OtherIDValues;
+			ReadIntAttribute(ExpectedIDValues, IdName, EmitterName, FrameIndex);
+			OtherCache.ReadIntAttribute(OtherIDValues, IdName, EmitterName, FrameIndex);
+
+			if (ExpectedIDValues.Num() != OtherIDValues.Num() || ExpectedIDValues.Num() != ExpectedEmitterFrame.ParticleDataBuffers.NumInstances)
+			{
+				AddError(Errors, TEXT("Invalid particle IDs - frame ") + Frame + TEXT(" emitter ") + EmitterName.ToString() + '\n');
+				bEqual = false;
+				continue;
+			}
+			if (ExpectedIDValues.Num() == 0)
+			{
+				continue;
+			}
+
+			TMap<int32, int32> ExpectedIDToRow;
+			for (int Index = 0; Index < ExpectedIDValues.Num(); Index++) {
+				ExpectedIDToRow.Add(ExpectedIDValues[Index], Index);
+			}
+
+			for (const FNiagaraSimCacheVariable& Var : CacheLayout.EmitterLayouts[EmitterIndex].Variables)
+			{
+				if (Var.Variable.GetName() == FName("ID"))
+				{
+					// persistent ID is not deterministic, skip
+					continue;
+				}
+				TArray<float> ExpectedFloats;
+				TArray<float> OtherFloats;
+				TArray<FFloat16> ExpectedHalfs;
+				TArray<FFloat16> OtherHalfs;
+				TArray<int32> ExpectedInts;
+				TArray<int32> OtherInts;
+				ReadAttribute(ExpectedFloats, ExpectedHalfs, ExpectedInts, Var.Variable.GetName(), EmitterName, FrameIndex);
+				OtherCache.ReadAttribute(OtherFloats, OtherHalfs, OtherInts, Var.Variable.GetName(), EmitterName, FrameIndex);
+				
+				if (!CompareAttributeData( OtherIDValues, ExpectedIDToRow, ExpectedFloats,OtherFloats, Var.FloatCount) ||
+					!CompareAttributeData( OtherIDValues, ExpectedIDToRow, ExpectedHalfs,OtherHalfs, Var.HalfCount) ||
+					!CompareAttributeData( OtherIDValues, ExpectedIDToRow, ExpectedInts,OtherInts, Var.Int32Count))
+				{
+					AddError(Errors, TEXT("Different particle data - frame ") + Frame + TEXT(" emitter ") + EmitterName.ToString() + TEXT(" attribute ") + Var.Variable.GetName().ToString() + '\n');
+					bEqual = false;
+				}
+			}
+		}
+	}
+	OutDifference = Errors.ToString();
+	return bEqual;
+}
+#endif
+
diff --git a/Engine/Plugins/FX/Niagara/Source/Niagara/Private/NiagaraSimCacheAttributeReaderHelper.h b/Engine/Plugins/FX/Niagara/Source/Niagara/Private/NiagaraSimCacheAttributeReaderHelper.h
--- a/Engine/Plugins/FX/Niagara/Source/Niagara/Private/NiagaraSimCacheAttributeReaderHelper.h
+++ b/Engine/Plugins/FX/Niagara/Source/Niagara/Private/NiagaraSimCacheAttributeReaderHelper.h
@@ -156,12 +156,11 @@ struct FNiagaraSimCacheAttributeReaderHelper
 		FMemory::Memcpy(OutBuffer, &DataBuffers->Int32Data[(Variable->Int32Offset + Component) * DataBuffers->NumInstances * sizeof(int32)], sizeof(int32) * DataBuffers->NumInstances);
 	}
 
-	int32 ReadInt(int32 Instance) const
+	int32 ReadInt(int32 Instance, int32 Offset = 0) const
 	{
 		check(IsValid());
-		check(Variable->Int32Offset != INDEX_NONE && Variable->Int32Count == 1);
-
-		const int32 Int32Offset = Instance + (Variable->Int32Offset * DataBuffers->NumInstances);
+		
+		const int32 Int32Offset = Instance + ((Variable->Int32Offset + Offset) * DataBuffers->NumInstances);
 
 		int32 Value;
 		FMemory::Memcpy(&Value, &DataBuffers->Int32Data[Int32Offset * sizeof(int32)], sizeof(int32));
@@ -185,6 +184,12 @@ struct FNiagaraSimCacheAttributeReaderHelper
 		return InternalReadFloat(0, Instance);
 	}
 
+	FNiagaraID ReadID(int32 Instance) const
+	{
+		check(Variable->Int32Offset != INDEX_NONE && Variable->Int32Count == 2);
+		return FNiagaraID(ReadInt(Instance, 0), ReadInt(Instance, 1));
+	}
+
 	FVector2f ReadFloat2f(int32 Instance) const
 	{
 		check(Variable->FloatOffset != INDEX_NONE && Variable->FloatCount == 2);
