Fixed a bug where curve soft ptr references could not be dereferenced correctly when the system itself was in Serialize()

diff --git a/Engine/Plugins/FX/Niagara/Source/Niagara/Classes/NiagaraDataInterfaceCurveBase.h b/Engine/Plugins/FX/Niagara/Source/Niagara/Classes/NiagaraDataInterfaceCurveBase.h
--- a/Engine/Plugins/FX/Niagara/Source/Niagara/Classes/NiagaraDataInterfaceCurveBase.h
+++ b/Engine/Plugins/FX/Niagara/Source/Niagara/Classes/NiagaraDataInterfaceCurveBase.h
@@ -51,7 +51,7 @@ protected:
 public:
 #if WITH_EDITORONLY_DATA
 	UPROPERTY(EditAnywhere, Category = "Curve")
-	TSoftObjectPtr<UCurveBase> CurveAsset;
+	TObjectPtr<UCurveBase> CurveAsset;
 #endif
 	
 	UPROPERTY(EditAnywhere, AdvancedDisplay, Category = "Curve")
diff --git a/Engine/Plugins/FX/Niagara/Source/Niagara/Private/NiagaraDataInterfaceColorCurve.cpp b/Engine/Plugins/FX/Niagara/Source/Niagara/Private/NiagaraDataInterfaceColorCurve.cpp
--- a/Engine/Plugins/FX/Niagara/Source/Niagara/Private/NiagaraDataInterfaceColorCurve.cpp
+++ b/Engine/Plugins/FX/Niagara/Source/Niagara/Private/NiagaraDataInterfaceColorCurve.cpp
@@ -184,14 +184,14 @@ void UNiagaraDataInterfaceColorCurve::GetFunctions(TArray<FNiagaraFunctionSignat
 #if WITH_EDITORONLY_DATA
 void UNiagaraDataInterfaceColorCurve::SyncCurvesToAsset()
 {
-	if (UCurveLinearColor* ColorCurve = Cast<UCurveLinearColor>(CurveAsset.LoadSynchronous()))
+	if (UCurveLinearColor* ColorCurve = Cast<UCurveLinearColor>(CurveAsset))
 	{
 		RedCurve = ColorCurve->FloatCurves[0];
 		GreenCurve = ColorCurve->FloatCurves[1];
 		BlueCurve = ColorCurve->FloatCurves[2];
 		AlphaCurve = ColorCurve->FloatCurves[3];
 	}
-	else if (UCurveVector* VecCurve = Cast<UCurveVector>(CurveAsset.LoadSynchronous()))
+	else if (UCurveVector* VecCurve = Cast<UCurveVector>(CurveAsset))
 	{
 		RedCurve = VecCurve->FloatCurves[0];
 		GreenCurve = VecCurve->FloatCurves[1];
@@ -199,7 +199,7 @@ void UNiagaraDataInterfaceColorCurve::SyncCurvesToAsset()
 		AlphaCurve.Reset();
 		AlphaCurve.AddKey(0, 1);
 	}
-	else if (UCurveFloat* FloatCurve = Cast<UCurveFloat>(CurveAsset.LoadSynchronous()))
+	else if (UCurveFloat* FloatCurve = Cast<UCurveFloat>(CurveAsset))
 	{
 		RedCurve = FloatCurve->FloatCurve;
 		GreenCurve = FloatCurve->FloatCurve;
diff --git a/Engine/Plugins/FX/Niagara/Source/Niagara/Private/NiagaraDataInterfaceCurve.cpp b/Engine/Plugins/FX/Niagara/Source/Niagara/Private/NiagaraDataInterfaceCurve.cpp
--- a/Engine/Plugins/FX/Niagara/Source/Niagara/Private/NiagaraDataInterfaceCurve.cpp
+++ b/Engine/Plugins/FX/Niagara/Source/Niagara/Private/NiagaraDataInterfaceCurve.cpp
@@ -146,7 +146,7 @@ TArray<float> UNiagaraDataInterfaceCurve::BuildLUT(int32 NumEntries) const
 #if WITH_EDITORONLY_DATA
 void UNiagaraDataInterfaceCurve::SyncCurvesToAsset()
 {
-	if (UCurveFloat* LoadedCurve = Cast<UCurveFloat>(CurveAsset.LoadSynchronous()))
+	if (UCurveFloat* LoadedCurve = Cast<UCurveFloat>(CurveAsset))
 	{
 		Curve = LoadedCurve->FloatCurve;
 	}
diff --git a/Engine/Plugins/FX/Niagara/Source/Niagara/Private/NiagaraDataInterfaceVectorCurve.cpp b/Engine/Plugins/FX/Niagara/Source/Niagara/Private/NiagaraDataInterfaceVectorCurve.cpp
--- a/Engine/Plugins/FX/Niagara/Source/Niagara/Private/NiagaraDataInterfaceVectorCurve.cpp
+++ b/Engine/Plugins/FX/Niagara/Source/Niagara/Private/NiagaraDataInterfaceVectorCurve.cpp
@@ -175,19 +175,19 @@ void UNiagaraDataInterfaceVectorCurve::GetFunctions(TArray<FNiagaraFunctionSigna
 #if WITH_EDITORONLY_DATA
 void UNiagaraDataInterfaceVectorCurve::SyncCurvesToAsset()
 {
-	if (UCurveLinearColor* ColorCurve = Cast<UCurveLinearColor>(CurveAsset.LoadSynchronous()))
+	if (UCurveLinearColor* ColorCurve = Cast<UCurveLinearColor>(CurveAsset))
 	{
 		XCurve = ColorCurve->FloatCurves[0];
 		YCurve = ColorCurve->FloatCurves[1];
 		ZCurve = ColorCurve->FloatCurves[2];
 	}
-	else if (UCurveVector* VecCurve = Cast<UCurveVector>(CurveAsset.LoadSynchronous()))
+	else if (UCurveVector* VecCurve = Cast<UCurveVector>(CurveAsset))
 	{
 		XCurve = VecCurve->FloatCurves[0];
 		YCurve = VecCurve->FloatCurves[1];
 		ZCurve = VecCurve->FloatCurves[2];
 	}
-	else if (UCurveFloat* FloatCurve = Cast<UCurveFloat>(CurveAsset.LoadSynchronous()))
+	else if (UCurveFloat* FloatCurve = Cast<UCurveFloat>(CurveAsset))
 	{
 		XCurve = FloatCurve->FloatCurve;
 		YCurve = FloatCurve->FloatCurve;
diff --git a/Engine/Plugins/FX/Niagara/Source/NiagaraEditorWidgets/Private/DetailCustomizations/NiagaraDataInterfaceCurveDetails.cpp b/Engine/Plugins/FX/Niagara/Source/NiagaraEditorWidgets/Private/DetailCustomizations/NiagaraDataInterfaceCurveDetails.cpp
--- a/Engine/Plugins/FX/Niagara/Source/NiagaraEditorWidgets/Private/DetailCustomizations/NiagaraDataInterfaceCurveDetails.cpp
+++ b/Engine/Plugins/FX/Niagara/Source/NiagaraEditorWidgets/Private/DetailCustomizations/NiagaraDataInterfaceCurveDetails.cpp
@@ -1408,7 +1408,7 @@ void FNiagaraDataInterfaceCurveDetailsBase::OnShowInCurveEditor() const
 
 UCurveBase* FNiagaraDataInterfaceCurveDetailsBase::GetUsedCurveAsset(UNiagaraDataInterfaceCurveBase* CurveDI) const
 {
-	return CurveDI ? CurveDI->CurveAsset.LoadSynchronous() : nullptr;
+	return CurveDI ? CurveDI->CurveAsset : nullptr;
 }
 
 void FNiagaraDataInterfaceCurveDetailsBase::ImportSelectedAsset(UObject* SelectedAsset)
@@ -1429,7 +1429,7 @@ void FNiagaraDataInterfaceCurveDetailsBase::ImportSelectedAsset(UObject* Selecte
 
 		if (!bCopyAssetData)
 		{
-			CurveDataInterfaceWeak->CurveAsset = SelectedAsset;
+			CurveDataInterfaceWeak->CurveAsset = Cast<UCurveBase>(SelectedAsset);
 		}
 		
 		for (int i = 0; i < CurveProperties.Num(); i++)
