Fixed mismatched copy constructor/assignment operator in FHlslNiagaraTranslatorOptions.

diff --git a/Engine/Plugins/FX/Niagara/Source/NiagaraEditor/Private/NiagaraHlslTranslator.h b/Engine/Plugins/FX/Niagara/Source/NiagaraEditor/Private/NiagaraHlslTranslator.h
--- a/Engine/Plugins/FX/Niagara/Source/NiagaraEditor/Private/NiagaraHlslTranslator.h
+++ b/Engine/Plugins/FX/Niagara/Source/NiagaraEditor/Private/NiagaraHlslTranslator.h
@@ -333,15 +333,6 @@ public:
 
 	}
 
-	FHlslNiagaraTranslatorOptions(const FHlslNiagaraTranslatorOptions& InOptions)
-		: SimTarget(InOptions.SimTarget)
-		, InstanceParameterNamespaces(InOptions.InstanceParameterNamespaces)
-		, bParameterRapidIteration(InOptions.bParameterRapidIteration)
-		, bDisableDebugSwitches(InOptions.bDisableDebugSwitches)
-		, OverrideModuleConstants(InOptions.OverrideModuleConstants)
-	{
-	}
-
 	ENiagaraSimTarget SimTarget;
 
 	/** Any parameters in these namespaces will be pulled from an "InstanceParameters" dataset rather than from the uniform table. */
