fix override format default

diff --git a/Engine/Plugins/FX/Niagara/Source/Niagara/Private/NiagaraDataInterfaceGrid2DCollection.cpp b/Engine/Plugins/FX/Niagara/Source/Niagara/Private/NiagaraDataInterfaceGrid2DCollection.cpp
--- a/Engine/Plugins/FX/Niagara/Source/Niagara/Private/NiagaraDataInterfaceGrid2DCollection.cpp
+++ b/Engine/Plugins/FX/Niagara/Source/Niagara/Private/NiagaraDataInterfaceGrid2DCollection.cpp
@@ -249,7 +249,7 @@ struct FNiagaraGrid2DLegacyTiled2DInfo
 };
 
 UNiagaraDataInterfaceGrid2DCollection::UNiagaraDataInterfaceGrid2DCollection(FObjectInitializer const& ObjectInitializer)
-	: Super(ObjectInitializer), OverrideBufferFormat(ENiagaraGpuBufferFormat::Float), bOverrideFormat(false)
+	: Super(ObjectInitializer), OverrideBufferFormat(GetDefault<UNiagaraSettings>()->DefaultGridFormat), bOverrideFormat(false)
 #if WITH_EDITORONLY_DATA
 	, bPreviewGrid(false)
 #endif
diff --git a/Engine/Plugins/FX/Niagara/Source/Niagara/Private/NiagaraDataInterfaceGrid3DCollection.cpp b/Engine/Plugins/FX/Niagara/Source/Niagara/Private/NiagaraDataInterfaceGrid3DCollection.cpp
--- a/Engine/Plugins/FX/Niagara/Source/Niagara/Private/NiagaraDataInterfaceGrid3DCollection.cpp
+++ b/Engine/Plugins/FX/Niagara/Source/Niagara/Private/NiagaraDataInterfaceGrid3DCollection.cpp
@@ -312,7 +312,7 @@ int32 UNiagaraDataInterfaceGrid3DCollection::GetComponentCountFromFuncName(const
 UNiagaraDataInterfaceGrid3DCollection::UNiagaraDataInterfaceGrid3DCollection(FObjectInitializer const& ObjectInitializer)
 	: Super(ObjectInitializer)
 	, NumAttributes(0)
-	, OverrideBufferFormat(ENiagaraGpuBufferFormat::Float)
+	, OverrideBufferFormat(GetDefault<UNiagaraSettings>()->DefaultGridFormat)
 	, bOverrideFormat(false)
 #if WITH_EDITORONLY_DATA
 	, bPreviewGrid(false)
