Fixed a shutdown crash in hierarchy editor

diff --git a/Engine/Plugins/FX/Niagara/Source/NiagaraEditor/Private/Widgets/SNiagaraHierarchy.cpp b/Engine/Plugins/FX/Niagara/Source/NiagaraEditor/Private/Widgets/SNiagaraHierarchy.cpp
--- a/Engine/Plugins/FX/Niagara/Source/NiagaraEditor/Private/Widgets/SNiagaraHierarchy.cpp
+++ b/Engine/Plugins/FX/Niagara/Source/NiagaraEditor/Private/Widgets/SNiagaraHierarchy.cpp
@@ -1171,10 +1171,13 @@ void SNiagaraHierarchy::BindToHierarchyRootViewModel()
 
 void SNiagaraHierarchy::UnbindFromHierarchyRootViewModel() const
 {
-	HierarchyViewModel->GetHierarchyRootViewModel()->OnSyncPropagated().Unbind();
-	HierarchyViewModel->GetHierarchyRootViewModel()->OnSectionsChanged().Unbind();
-	HierarchyViewModel->GetHierarchyRootViewModel()->OnSectionAdded().Unbind();
-	HierarchyViewModel->GetHierarchyRootViewModel()->OnSectionDeleted().Unbind();
+	if(HierarchyViewModel->GetHierarchyRootViewModel().IsValid())
+	{
+		HierarchyViewModel->GetHierarchyRootViewModel()->OnSyncPropagated().Unbind();
+		HierarchyViewModel->GetHierarchyRootViewModel()->OnSectionsChanged().Unbind();
+		HierarchyViewModel->GetHierarchyRootViewModel()->OnSectionAdded().Unbind();
+		HierarchyViewModel->GetHierarchyRootViewModel()->OnSectionDeleted().Unbind();
+	}
 }
 
 const TArray<TSharedPtr<FNiagaraHierarchyItemViewModelBase>>& SNiagaraHierarchy::GetSourceItems() const
