Fixed a bug where updating user parameter curve overrides could cause an ensure

diff --git a/Engine/Plugins/FX/Niagara/Source/Niagara/Private/NiagaraDataInterfaceCurveBase.cpp b/Engine/Plugins/FX/Niagara/Source/Niagara/Private/NiagaraDataInterfaceCurveBase.cpp
--- a/Engine/Plugins/FX/Niagara/Source/Niagara/Private/NiagaraDataInterfaceCurveBase.cpp
+++ b/Engine/Plugins/FX/Niagara/Source/Niagara/Private/NiagaraDataInterfaceCurveBase.cpp
@@ -229,6 +229,10 @@ void UNiagaraDataInterfaceCurveBase::PostEditChangeProperty(struct FPropertyChan
 	{
 		UpdateExposedTexture();
 	}
+	if (PropertyName == GET_MEMBER_NAME_CHECKED(UNiagaraDataInterfaceCurveBase, CurveAsset))
+	{
+		UpdateLUT();
+	}
 }
 #endif
 
