Made sure to sync the root view model of summary view before refreshing the summary stack to ensure we aren't accessing invalid data by accident (i.e. deleting a sim stage, which will cause a refresh in the summary view, but the item has not necessarily been deleted yet in summary data

diff --git a/Engine/Plugins/FX/Niagara/Source/NiagaraEditor/Private/ViewModels/Stack/NiagaraStackGraphUtilities.cpp b/Engine/Plugins/FX/Niagara/Source/NiagaraEditor/Private/ViewModels/Stack/NiagaraStackGraphUtilities.cpp
--- a/Engine/Plugins/FX/Niagara/Source/NiagaraEditor/Private/ViewModels/Stack/NiagaraStackGraphUtilities.cpp
+++ b/Engine/Plugins/FX/Niagara/Source/NiagaraEditor/Private/ViewModels/Stack/NiagaraStackGraphUtilities.cpp
@@ -1046,7 +1046,7 @@ void FNiagaraStackGraphUtilities::GatherInputRelationsForStack(FInputDataCollect
 		NodeGuidToModuleNodeMap.Add(FunctionCallNode->NodeGuid, FunctionCallNode);
 	}
 
-	// the inputs found in the root are the ones we want to display
+	// the inputs found in the root are the ones we want to display; these might contain outdated entries so it is important to sync the root view model to data before calling this
 	TArray<UNiagaraHierarchyModuleInput*> ModuleInputs;
 	SummaryRoot->GetChildrenOfType(ModuleInputs, true);
 	
diff --git a/Engine/Plugins/FX/Niagara/Source/NiagaraEditor/Private/ViewModels/Stack/NiagaraStackSummaryViewInputCollection.cpp b/Engine/Plugins/FX/Niagara/Source/NiagaraEditor/Private/ViewModels/Stack/NiagaraStackSummaryViewInputCollection.cpp
--- a/Engine/Plugins/FX/Niagara/Source/NiagaraEditor/Private/ViewModels/Stack/NiagaraStackSummaryViewInputCollection.cpp
+++ b/Engine/Plugins/FX/Niagara/Source/NiagaraEditor/Private/ViewModels/Stack/NiagaraStackSummaryViewInputCollection.cpp
@@ -75,7 +75,9 @@ void UNiagaraStackSummaryViewCollection::RefreshChildrenInternal(const TArray<UN
 
 	const UNiagaraHierarchyRoot* Root = ViewModel->GetEditorData().GetSummaryRoot();
 	TSharedPtr<FNiagaraHierarchyRootViewModel> RootViewModel = ViewModel->GetSummaryHierarchyViewModel()->GetHierarchyRootViewModel();
-	
+	// we make sure to sync view models to data before refreshing in order to get rid of possibly removed entries
+	RootViewModel->SyncViewModelsToData();
+
 	FNiagaraStackGraphUtilities::FInputDataCollection State;
 	FNiagaraStackGraphUtilities::GatherInputRelationsForStack(State, GetEmitterViewModel().ToSharedRef());
 	
