Fixed source sections being renameable via click + click

diff --git a/Engine/Plugins/FX/Niagara/Source/NiagaraEditor/Private/Widgets/SNiagaraHierarchy.cpp b/Engine/Plugins/FX/Niagara/Source/NiagaraEditor/Private/Widgets/SNiagaraHierarchy.cpp
--- a/Engine/Plugins/FX/Niagara/Source/NiagaraEditor/Private/Widgets/SNiagaraHierarchy.cpp
+++ b/Engine/Plugins/FX/Niagara/Source/NiagaraEditor/Private/Widgets/SNiagaraHierarchy.cpp
@@ -225,6 +225,7 @@ void SNiagaraHierarchySection::Construct(const FArguments& InArgs, TSharedPtr<FN
 							.OnTextCommitted(this, &SNiagaraHierarchySection::OnRenameSection)
 							.OnVerifyTextChanged(this, &SNiagaraHierarchySection::OnVerifySectionRename)
 							.IsSelected(this, &SNiagaraHierarchySection::IsSectionSelected)
+							.IsReadOnly(this, &SNiagaraHierarchySection::IsSectionReadOnly)
 						]
 					]
 				])
@@ -491,6 +492,11 @@ bool SNiagaraHierarchySection::IsSectionSelected() const
 	return GetSectionCheckState() == ECheckBoxState::Checked ? true : false;
 }
 
+bool SNiagaraHierarchySection::IsSectionReadOnly() const
+{
+	return SectionViewModel.IsValid() ? !SectionViewModel->CanRename() : true;
+}
+
 ECheckBoxState SNiagaraHierarchySection::GetSectionCheckState() const
 {
 	return IsSectionActive.Get();
diff --git a/Engine/Plugins/FX/Niagara/Source/NiagaraEditor/Public/Widgets/SNiagaraHierarchy.h b/Engine/Plugins/FX/Niagara/Source/NiagaraEditor/Public/Widgets/SNiagaraHierarchy.h
--- a/Engine/Plugins/FX/Niagara/Source/NiagaraEditor/Public/Widgets/SNiagaraHierarchy.h
+++ b/Engine/Plugins/FX/Niagara/Source/NiagaraEditor/Public/Widgets/SNiagaraHierarchy.h
@@ -60,6 +60,7 @@ private:
 	bool OnVerifySectionRename(const FText& NewName, FText& OutTooltip) const;
 
 	bool IsSectionSelected() const;
+	bool IsSectionReadOnly() const;
 	ECheckBoxState GetSectionCheckState() const;
 	void OnSectionCheckChanged(ECheckBoxState NewState);
 	EActiveTimerReturnType ActivateSectionIfDragging(double CurrentTime, float DeltaTime) const;
