Fixed flickering issue of summary view as assignment inputs would constantly get recreated

diff --git a/Engine/Plugins/FX/Niagara/Source/NiagaraEditor/Private/ViewModels/Stack/NiagaraStackInputCategory.cpp b/Engine/Plugins/FX/Niagara/Source/NiagaraEditor/Private/ViewModels/Stack/NiagaraStackInputCategory.cpp
--- a/Engine/Plugins/FX/Niagara/Source/NiagaraEditor/Private/ViewModels/Stack/NiagaraStackInputCategory.cpp
+++ b/Engine/Plugins/FX/Niagara/Source/NiagaraEditor/Private/ViewModels/Stack/NiagaraStackInputCategory.cpp
@@ -393,10 +393,11 @@ void UNiagaraStackSummaryCategory::RefreshChildrenInternal(const TArray<UNiagara
 			{
 				FString VariableNameWithModulePrefix = FNiagaraConstants::ModuleNamespaceString + ".";
 				VariableNameWithModulePrefix.Append(InputData->InputName.ToString());
+				FName VariableNameToTestAgainst(VariableNameWithModulePrefix);
 				
 				UNiagaraStackFunctionInput* TopLevelInput = FindCurrentChildOfTypeByPredicate<UNiagaraStackFunctionInput>(CurrentChildren, [&](UNiagaraStackFunctionInput* CurrentInput) 
 				{
-					return CurrentInput->GetInputParameterHandle().GetName().ToString() == VariableNameWithModulePrefix && CurrentInput->GetInputType() == InputData->Type && &CurrentInput->GetInputFunctionCallNode() == InputData->FunctionCallNode;
+					return CurrentInput->GetInputParameterHandle() == VariableNameToTestAgainst && CurrentInput->GetInputType() == InputData->Type && &CurrentInput->GetInputFunctionCallNode() == InputData->FunctionCallNode;
 				});
 	
 				if(TopLevelInput == nullptr)
diff --git a/Engine/Plugins/FX/Niagara/Source/NiagaraEditor/Private/ViewModels/Stack/NiagaraStackSummaryViewInputCollection.cpp b/Engine/Plugins/FX/Niagara/Source/NiagaraEditor/Private/ViewModels/Stack/NiagaraStackSummaryViewInputCollection.cpp
--- a/Engine/Plugins/FX/Niagara/Source/NiagaraEditor/Private/ViewModels/Stack/NiagaraStackSummaryViewInputCollection.cpp
+++ b/Engine/Plugins/FX/Niagara/Source/NiagaraEditor/Private/ViewModels/Stack/NiagaraStackSummaryViewInputCollection.cpp
@@ -247,10 +247,11 @@ void UNiagaraStackSummaryViewCollection::RefreshChildrenInternal(const TArray<UN
 			{
 				FString VariableNameWithModulePrefix = FNiagaraConstants::ModuleNamespaceString + ".";
 				VariableNameWithModulePrefix.Append(InputData->InputName.ToString());
-				
+				FName VariableNameToTestAgainst(VariableNameWithModulePrefix);
+
 				UNiagaraStackFunctionInput* TopLevelInput = FindCurrentChildOfTypeByPredicate<UNiagaraStackFunctionInput>(CurrentChildren, [&](UNiagaraStackFunctionInput* CurrentInput) 
 				{
-					return CurrentInput->GetInputParameterHandle().GetName().ToString() == VariableNameWithModulePrefix && CurrentInput->GetInputType() == InputData->Type && &CurrentInput->GetInputFunctionCallNode() == InputData->FunctionCallNode;
+					return CurrentInput->GetInputParameterHandle() == VariableNameToTestAgainst && CurrentInput->GetInputType() == InputData->Type && &CurrentInput->GetInputFunctionCallNode() == InputData->FunctionCallNode;
 				});
 	
 				if(TopLevelInput == nullptr)
