Fixed missing stack refresh when changing a user parameter name. It would appear like the stack didn't take the update to the parameter, but it's just the UI didn't update

diff --git a/Engine/Plugins/FX/Niagara/Source/NiagaraEditor/Private/ViewModels/NiagaraSystemViewModel.cpp b/Engine/Plugins/FX/Niagara/Source/NiagaraEditor/Private/ViewModels/NiagaraSystemViewModel.cpp
--- a/Engine/Plugins/FX/Niagara/Source/NiagaraEditor/Private/ViewModels/NiagaraSystemViewModel.cpp
+++ b/Engine/Plugins/FX/Niagara/Source/NiagaraEditor/Private/ViewModels/NiagaraSystemViewModel.cpp
@@ -1337,6 +1337,7 @@ bool FNiagaraSystemViewModel::RenameParameter(const FNiagaraVariable TargetParam
 			{
 				FNiagaraEmitterHandle* EmitterHandle = EmitterHandleViewModel->GetEmitterHandle();
 				EmitterHandle->GetInstance().Emitter->HandleVariableRenamed(TargetParameter, FNiagaraVariableBase(TargetParameter.GetType(), NewName), true, EmitterHandle->GetInstance().Version);
+				EmitterHandleViewModel->GetEmitterStackViewModel()->RequestRefreshDeferred();
 			}
 		}
 		else
