Changed the ResetSimulation command to also resume playback if currently paused.

diff --git a/Engine/Plugins/FX/Niagara/Source/NiagaraEditor/Private/ViewModels/NiagaraSystemViewModel.cpp b/Engine/Plugins/FX/Niagara/Source/NiagaraEditor/Private/ViewModels/NiagaraSystemViewModel.cpp
--- a/Engine/Plugins/FX/Niagara/Source/NiagaraEditor/Private/ViewModels/NiagaraSystemViewModel.cpp
+++ b/Engine/Plugins/FX/Niagara/Source/NiagaraEditor/Private/ViewModels/NiagaraSystemViewModel.cpp
@@ -2217,6 +2217,7 @@ void FNiagaraSystemViewModel::ResetSystem(ETimeResetMode TimeResetMode, EMultiRe
 			else
 			{
 				Sequencer->SetGlobalTime(0);
+				Sequencer->SetPlaybackStatus(EMovieScenePlayerStatus::Playing);
 			}
 		}
 	}
