Added new scratchpad action + set parameters to suggested in menu

diff --git a/Engine/Plugins/FX/Niagara/Source/NiagaraEditor/Private/ViewModels/Stack/NiagaraStackScriptItemGroup.cpp b/Engine/Plugins/FX/Niagara/Source/NiagaraEditor/Private/ViewModels/Stack/NiagaraStackScriptItemGroup.cpp
--- a/Engine/Plugins/FX/Niagara/Source/NiagaraEditor/Private/ViewModels/Stack/NiagaraStackScriptItemGroup.cpp
+++ b/Engine/Plugins/FX/Niagara/Source/NiagaraEditor/Private/ViewModels/Stack/NiagaraStackScriptItemGroup.cpp
@@ -113,14 +113,14 @@ public:
 	{
 		FText DisplayName = LOCTEXT("NewScratchModuleName", "New Scratch Pad Module");
 		FText Description = LOCTEXT("NewScratchModuleDescription", "Description: Create a new scratch pad module.");
-		return MakeShareable(new FScriptGroupAddAction( DisplayName, {}, Description, FText(), false, true, FNiagaraVariable(), false, FAssetData(), nullptr, true, false));
+		return MakeShareable(new FScriptGroupAddAction( DisplayName, {}, Description, FText(), true, true, FNiagaraVariable(), false, FAssetData(), nullptr, true, false));
 	}
 
 	static TSharedRef<FScriptGroupAddAction> CreateNewSetSpecificModuleAction()
 	{
 		FText DisplayName = LOCTEXT("NewSetSpecificModuleName", "Set new or existing parameter directly");
 		FText Description = LOCTEXT("NewSetSpecificModuleDescription", "Description: Create a new module that can set new or existing parameters directly.");
-		return MakeShareable(new FScriptGroupAddAction( DisplayName, {}, Description, FText(), false, true, FNiagaraVariable(), false, FAssetData(), nullptr, false, true));
+		return MakeShareable(new FScriptGroupAddAction( DisplayName, {}, Description, FText(), true, true, FNiagaraVariable(), false, FAssetData(), nullptr, false, true));
 	}
 
 	virtual TArray<FString> GetCategories() const override
