- Fix for lens effects being killed when tweaking parameters

diff --git a/Engine/Plugins/FX/Niagara/Source/Niagara/Private/NiagaraComponent.cpp b/Engine/Plugins/FX/Niagara/Source/Niagara/Private/NiagaraComponent.cpp
--- a/Engine/Plugins/FX/Niagara/Source/Niagara/Private/NiagaraComponent.cpp
+++ b/Engine/Plugins/FX/Niagara/Source/Niagara/Private/NiagaraComponent.cpp
@@ -29,6 +29,7 @@
 #include "SceneInterface.h"
 #include "UObject/UObjectIterator.h"
 #include "PrimitiveUniformShaderParametersBuilder.h"
+#include "NiagaraActor.h"
 
 #include UE_INLINE_GENERATED_CPP_BY_NAME(NiagaraComponent)
 
@@ -3431,7 +3432,24 @@ void UNiagaraComponent::AssetExposedParametersChanged()
 		// In that case we skip the sync with the system, as that happens in PostLoad anyways.
 		return;
 	}
-	bool bCurrentlyActive = IsActive();
+
+#if WITH_EDITOR
+	// Some owner actors might have destroy on finish enabled, if we were to tweak parameters this would trigger completion and destroy the actor
+	// Lens effects for example set this making it impossible to tweak them without storing and restoring this information
+	TWeakObjectPtr<ANiagaraActor> WeakNiagaraActor = Cast<ANiagaraActor>(GetOwner());
+	bool bDestroyOnSystemFinish = false;
+	if (ANiagaraActor* NiagaraActor = WeakNiagaraActor.Get())
+	{
+		if (NiagaraActor->GetDestroyOnSystemFinish())
+		{
+			bDestroyOnSystemFinish = true;
+			NiagaraActor->SetDestroyOnSystemFinish(false);
+		}
+	}
+#endif
+
+	const bool bCurrentlyActive = IsActive();
+
 	SynchronizeWithSourceSystem();
 #if WITH_EDITOR
 	if ( !GIsTransacting )
@@ -3442,6 +3460,16 @@ void UNiagaraComponent::AssetExposedParametersChanged()
 			// only reactivate systems that were running before
 			ReinitializeSystem();
 		}
+
+#if WITH_EDITOR
+		if (bDestroyOnSystemFinish)
+		{
+			if (ANiagaraActor* NiagaraActor = WeakNiagaraActor.Get())
+			{
+				NiagaraActor->SetDestroyOnSystemFinish(true);
+			}
+		}
+#endif
 	}
 }
 
diff --git a/Engine/Plugins/FX/Niagara/Source/Niagara/Public/NiagaraActor.h b/Engine/Plugins/FX/Niagara/Source/Niagara/Public/NiagaraActor.h
--- a/Engine/Plugins/FX/Niagara/Source/Niagara/Public/NiagaraActor.h
+++ b/Engine/Plugins/FX/Niagara/Source/Niagara/Public/NiagaraActor.h
@@ -23,6 +23,10 @@ public:
 	UFUNCTION(BlueprintCallable, Category=NiagaraActor)
 	void SetDestroyOnSystemFinish(bool bShouldDestroyOnSystemFinish);
 
+	/** Returns true if the system will destroy on finish */
+	UFUNCTION(BlueprintCallable, Category=NiagaraActor)
+	bool GetDestroyOnSystemFinish() const { return bDestroyOnSystemFinish; }
+
 private:
 	/** Pointer to System component */
 	UPROPERTY(VisibleAnywhere, BlueprintReadOnly, Category=NiagaraActor, meta = (AllowPrivateAccess = "true"))
