- Fix GPU readback not passing IDtoIndex data to the current data buffer
- Fix transition issues for GPU context readback

diff --git a/Engine/Plugins/FX/Niagara/Source/Niagara/Private/NiagaraDataSetReadback.cpp b/Engine/Plugins/FX/Niagara/Source/Niagara/Private/NiagaraDataSetReadback.cpp
--- a/Engine/Plugins/FX/Niagara/Source/Niagara/Private/NiagaraDataSetReadback.cpp
+++ b/Engine/Plugins/FX/Niagara/Source/Niagara/Private/NiagaraDataSetReadback.cpp
@@ -6,6 +6,7 @@
 #include "NiagaraEmitterInstance.h"
 #include "NiagaraSystemInstance.h"
 #include "NiagaraGpuComputeDispatchInterface.h"
+#include "NiagaraGPUInstanceCountManager.h"
 #include "NiagaraGpuReadbackManager.h"
 
 #include "Async/Async.h"
@@ -130,6 +131,16 @@ void FNiagaraDataSetReadback::GPUReadbackInternal(FRHICommandListImmediate& RHIC
 	const int32 HalfBufferStride = CurrentDataBuffer->GetHalfStride();
 	const int32 IntBufferStride = CurrentDataBuffer->GetInt32Stride();
 
+	// Transition buffers to copy
+	TArray<FRHITransitionInfo, TInlineAllocator<NumReadbackBuffers>> Transitions;
+	Transitions.Emplace(ReadbackBuffers[0].Buffer, FNiagaraGPUInstanceCountManager::kCountBufferDefaultState, ERHIAccess::CopySrc);
+	for (int32 i=1; i < ReadbackBuffers.Num(); ++i)
+	{
+		Transitions.Emplace(ReadbackBuffers[i].Buffer, ERHIAccess::SRVMask, ERHIAccess::CopySrc);
+	}
+	RHICmdList.Transition(Transitions);
+
+	// Enqueue readback
 	ReadbackManager->EnqueueReadbacks(
 		RHICmdList,
 		MakeArrayView(ReadbackBuffers),
@@ -144,21 +155,32 @@ void FNiagaraDataSetReadback::GPUReadbackInternal(FRHICommandListImmediate& RHIC
 			Readback->DataSet.CopyFromGPUReadback(FloatDataBuffer, IntDataBuffer, HalfDataBuffer, 0, InstanceCount, FloatBufferStride, IntBufferStride, HalfBufferStride);
 
 			// Copy ID to Index table
-			const int32* IDtoIndexBuffer = IDtoIndexBufferIndex == INDEX_NONE ? nullptr : reinterpret_cast<int32*>(BufferData[IDtoIndexBufferIndex].Key);
-			if (IDtoIndexBufferIndex != INDEX_NONE)
-			{
-				check(BufferData[IDtoIndexBufferIndex].Value >= InstanceCount * sizeof(int32));
-				Readback->IDToIndexTable.Empty();
-				Readback->IDToIndexTable.Reset(InstanceCount);
-				FMemory::Memcpy(Readback->IDToIndexTable.GetData(), IDtoIndexBuffer, InstanceCount * sizeof(int32));
-			}
-			else
+			if (FNiagaraDataBuffer* CurrentData = Readback->DataSet.GetCurrentData())
 			{
-				Readback->IDToIndexTable.Empty();
+				TArray<int32>& IDTable = CurrentData->GetIDTable();
+
+				const int32* IDtoIndexBuffer = IDtoIndexBufferIndex == INDEX_NONE ? nullptr : reinterpret_cast<int32*>(BufferData[IDtoIndexBufferIndex].Key);
+				if (IDtoIndexBuffer != nullptr)
+				{
+					const int32 NumIDs = BufferData[IDtoIndexBufferIndex].Value / sizeof(int32);
+					check(NumIDs >= InstanceCount);
+					IDTable.SetNumUninitialized(NumIDs);
+					FMemory::Memcpy(IDTable.GetData(), IDtoIndexBuffer, NumIDs * sizeof(int32));
+				}
+				else
+				{
+					IDTable.Empty();
+				}
 			}
 
-
 			Readback->ReadbackCompleteInternal();
 		}
 	);
+
+	// Transition buffers from copy
+	for (FRHITransitionInfo& Transition : Transitions)
+	{
+		Swap(Transition.AccessBefore, Transition.AccessAfter);
+	}
+	RHICmdList.Transition(Transitions);
 }
diff --git a/Engine/Plugins/FX/Niagara/Source/Niagara/Public/NiagaraDataSetReadback.h b/Engine/Plugins/FX/Niagara/Source/Niagara/Public/NiagaraDataSetReadback.h
--- a/Engine/Plugins/FX/Niagara/Source/Niagara/Public/NiagaraDataSetReadback.h
+++ b/Engine/Plugins/FX/Niagara/Source/Niagara/Public/NiagaraDataSetReadback.h
@@ -39,7 +39,6 @@ private:
 	//FGuid					SourceUsageId;
 	FNiagaraDataSet			DataSet;
 	FNiagaraParameterStore	ParameterStore;
-	TArray<int32>			IDToIndexTable;
 
 	FOnReadbackReady		OnReadbackReady;
 };
