Small cleanup to FCompileConstantResolver and also include handling Emitter.SimulationTarget for System resolvers

diff --git a/Engine/Plugins/FX/Niagara/Source/NiagaraEditor/Private/NiagaraParameterMapHistory.cpp b/Engine/Plugins/FX/Niagara/Source/NiagaraEditor/Private/NiagaraParameterMapHistory.cpp
--- a/Engine/Plugins/FX/Niagara/Source/NiagaraEditor/Private/NiagaraParameterMapHistory.cpp
+++ b/Engine/Plugins/FX/Niagara/Source/NiagaraEditor/Private/NiagaraParameterMapHistory.cpp
@@ -2314,73 +2314,85 @@ bool FCompileConstantResolver::ResolveConstant(FNiagaraVariable& OutConstant) co
 	}
 
 	// handle emitter case
-	FVersionedNiagaraEmitterData* EmitterData = Emitter.GetEmitterData();
-	if (EmitterData && OutConstant == FNiagaraVariable(FNiagaraTypeDefinition::GetBoolDef(), TEXT("Emitter.Localspace")))
+	if (FVersionedNiagaraEmitterData* EmitterData = Emitter.GetEmitterData())
 	{
-		OutConstant.SetValue(EmitterData->bLocalSpace ? FNiagaraBool(true) : FNiagaraBool(false));
-		return true;
-	}
-	if (EmitterData && OutConstant == FNiagaraVariable(FNiagaraTypeDefinition::GetBoolDef(), TEXT("Emitter.Determinism")))
-	{
-		OutConstant.SetValue(EmitterData->bDeterminism ? FNiagaraBool(true) : FNiagaraBool(false));
-		return true;
-	}
-	if (EmitterData && OutConstant == FNiagaraVariable(FNiagaraTypeDefinition::GetBoolDef(), TEXT("Emitter.InterpolatedSpawn")))
-	{
-		OutConstant.SetValue(EmitterData->bInterpolatedSpawning ? FNiagaraBool(true) : FNiagaraBool(false));
-		return true;
-	}
-	if (EmitterData && OutConstant == FNiagaraVariable(FNiagaraTypeDefinition::GetSimulationTargetEnum(), TEXT("Emitter.SimulationTarget")))
-	{
-		FNiagaraInt32 EnumValue;
-		EnumValue.Value = (uint8)EmitterData->SimTarget;
-		OutConstant.SetValue(EnumValue);
-		return true;
-	}
-	if (EmitterData && OutConstant == FNiagaraVariable(FNiagaraTypeDefinition::GetScriptUsageEnum(), TEXT("Script.Usage")))
-	{
-		FNiagaraInt32 EnumValue;
-		EnumValue.Value = (uint8)FNiagaraUtilities::ConvertScriptUsageToStaticSwitchUsage(Usage);
-		OutConstant.SetValue(EnumValue);
-		return true;
-	}
-	if (EmitterData && OutConstant == FNiagaraVariable(FNiagaraTypeDefinition::GetScriptContextEnum(), TEXT("Script.Context")))
-	{
-		FNiagaraInt32 EnumValue;
-		EnumValue.Value = (uint8)FNiagaraUtilities::ConvertScriptUsageToStaticSwitchContext(Usage);
-		OutConstant.SetValue(EnumValue);
-		return true;
+		if (OutConstant == FNiagaraVariable(FNiagaraTypeDefinition::GetBoolDef(), TEXT("Emitter.Localspace")))
+		{
+			OutConstant.SetValue(EmitterData->bLocalSpace ? FNiagaraBool(true) : FNiagaraBool(false));
+			return true;
+		}
+		if (OutConstant == FNiagaraVariable(FNiagaraTypeDefinition::GetBoolDef(), TEXT("Emitter.Determinism")))
+		{
+			OutConstant.SetValue(EmitterData->bDeterminism ? FNiagaraBool(true) : FNiagaraBool(false));
+			return true;
+		}
+		if (OutConstant == FNiagaraVariable(FNiagaraTypeDefinition::GetBoolDef(), TEXT("Emitter.InterpolatedSpawn")))
+		{
+			OutConstant.SetValue(EmitterData->bInterpolatedSpawning ? FNiagaraBool(true) : FNiagaraBool(false));
+			return true;
+		}
+		if (OutConstant == FNiagaraVariable(FNiagaraTypeDefinition::GetSimulationTargetEnum(), TEXT("Emitter.SimulationTarget")))
+		{
+			FNiagaraInt32 EnumValue;
+			EnumValue.Value = (uint8)EmitterData->SimTarget;
+			OutConstant.SetValue(EnumValue);
+			return true;
+		}
+		if (OutConstant == FNiagaraVariable(FNiagaraTypeDefinition::GetScriptUsageEnum(), TEXT("Script.Usage")))
+		{
+			FNiagaraInt32 EnumValue;
+			EnumValue.Value = (uint8)FNiagaraUtilities::ConvertScriptUsageToStaticSwitchUsage(Usage);
+			OutConstant.SetValue(EnumValue);
+			return true;
+		}
+		if (OutConstant == FNiagaraVariable(FNiagaraTypeDefinition::GetScriptContextEnum(), TEXT("Script.Context")))
+		{
+			FNiagaraInt32 EnumValue;
+			EnumValue.Value = (uint8)FNiagaraUtilities::ConvertScriptUsageToStaticSwitchContext(Usage);
+			OutConstant.SetValue(EnumValue);
+			return true;
+		}
 	}
 
 	// handle system case
-	if (System && OutConstant == FNiagaraVariable(FNiagaraTypeDefinition::GetBoolDef(), TEXT("Emitter.Localspace")))
-	{
-		OutConstant.SetValue(FNiagaraBool(false));
-		return true;
-	}
-	if (System && OutConstant == FNiagaraVariable(FNiagaraTypeDefinition::GetBoolDef(), TEXT("Emitter.Determinism")))
-	{
-		OutConstant.SetValue(FNiagaraBool(true));
-		return true;
-	}
-	if (System && OutConstant == FNiagaraVariable(FNiagaraTypeDefinition::GetBoolDef(), TEXT("Emitter.InterpolatedSpawn")))
-	{
-		OutConstant.SetValue(FNiagaraBool(false));
-		return true;
-	}
-	if (System && OutConstant == FNiagaraVariable(FNiagaraTypeDefinition::GetScriptUsageEnum(), TEXT("Script.Usage")))
-	{
-		FNiagaraInt32 EnumValue;
-		EnumValue.Value = (uint8)FNiagaraUtilities::ConvertScriptUsageToStaticSwitchUsage(Usage);
-		OutConstant.SetValue(EnumValue);
-		return true;
-	}
-	if (System && OutConstant == FNiagaraVariable(FNiagaraTypeDefinition::GetScriptContextEnum(), TEXT("Script.Context")))
+	if (System)
 	{
-		FNiagaraInt32 EnumValue;
-		EnumValue.Value = (uint8)FNiagaraUtilities::ConvertScriptUsageToStaticSwitchContext(Usage);
-		OutConstant.SetValue(EnumValue);
-		return true;
+		if (OutConstant == FNiagaraVariable(FNiagaraTypeDefinition::GetBoolDef(), TEXT("Emitter.Localspace")))
+		{
+			OutConstant.SetValue(FNiagaraBool(false));
+			return true;
+		}
+		if (OutConstant == FNiagaraVariable(FNiagaraTypeDefinition::GetBoolDef(), TEXT("Emitter.Determinism")))
+		{
+			OutConstant.SetValue(FNiagaraBool(true));
+			return true;
+		}
+		if (OutConstant == FNiagaraVariable(FNiagaraTypeDefinition::GetBoolDef(), TEXT("Emitter.InterpolatedSpawn")))
+		{
+			OutConstant.SetValue(FNiagaraBool(false));
+			return true;
+		}
+		if (OutConstant == FNiagaraVariable(FNiagaraTypeDefinition::GetSimulationTargetEnum(), TEXT("Emitter.SimulationTarget")))
+		{
+			FNiagaraInt32 EnumValue;
+			EnumValue.Value = (uint8)ENiagaraSimTarget::CPUSim;
+			OutConstant.SetValue(EnumValue);
+			return true;
+		}
+		if (OutConstant == FNiagaraVariable(FNiagaraTypeDefinition::GetScriptUsageEnum(), TEXT("Script.Usage")))
+		{
+			FNiagaraInt32 EnumValue;
+			EnumValue.Value = (uint8)FNiagaraUtilities::ConvertScriptUsageToStaticSwitchUsage(Usage);
+			OutConstant.SetValue(EnumValue);
+			return true;
+		}
+		if (OutConstant == FNiagaraVariable(FNiagaraTypeDefinition::GetScriptContextEnum(), TEXT("Script.Context")))
+		{
+			FNiagaraInt32 EnumValue;
+			EnumValue.Value = (uint8)FNiagaraUtilities::ConvertScriptUsageToStaticSwitchContext(Usage);
+			OutConstant.SetValue(EnumValue);
+			return true;
+		}
 	}
 
 	return false;
