Fixed missing user parameter rename verification

diff --git a/Engine/Plugins/FX/Niagara/Source/NiagaraEditor/Private/Customizations/NiagaraComponentDetails.cpp b/Engine/Plugins/FX/Niagara/Source/NiagaraEditor/Private/Customizations/NiagaraComponentDetails.cpp
--- a/Engine/Plugins/FX/Niagara/Source/NiagaraEditor/Private/Customizations/NiagaraComponentDetails.cpp
+++ b/Engine/Plugins/FX/Niagara/Source/NiagaraEditor/Private/Customizations/NiagaraComponentDetails.cpp
@@ -1005,6 +1005,45 @@ TSharedPtr<SWidget> FNiagaraSystemUserParameterBuilder::GenerateCustomNameWidget
 		})
 		.ParameterText(FText::FromName(UserParameter.GetName()))
 		.OnDragDetected(this, &FNiagaraSystemUserParameterBuilder::GenerateParameterDragDropOp, UserParameter)
+		.OnVerifyTextChanged_Lambda([this, UserParameter](const FText& InNewName, FText& OutErrorMessage)
+		{
+			if(InNewName.IsEmptyOrWhitespace())
+			{
+				OutErrorMessage = LOCTEXT("UserParameterNameCantBeEmpty", "Can not accept empty name");
+				return false;
+			}
+			
+			FNiagaraVariable NewUserParameter = UserParameter;
+			NewUserParameter.SetName(FName(InNewName.ToString()));
+			if(UserParameter == NewUserParameter)
+			{
+				// we return true, but renaming won't do anything if the name is the same as before
+				return true;
+			}
+			else
+			{
+				TArray<FNiagaraVariable> ExistingParameters;
+				UserParamToWidgetMap.GenerateKeyArray(ExistingParameters);
+
+				bool bNameExistsAlready = false;
+				for(FNiagaraVariable& ExistingUserParameter : ExistingParameters)
+				{
+					if(ExistingUserParameter.GetName().IsEqual(FName(InNewName.ToString())))
+					{
+						bNameExistsAlready = true;
+						break;
+					}
+				}
+
+				if(bNameExistsAlready)
+				{
+					OutErrorMessage = LOCTEXT("UserParameterAlreadyExists", "A User Parameter with the same name already exists");
+					return false;					
+				}
+
+				return true;
+			}
+		})
 		.OnTextCommitted_Lambda([this, UserParameter](const FText& InText, ETextCommit::Type Type)
 		{
 			RenameParameter(UserParameter, FName(InText.ToString()));
