Improved sim cache comparison for automation test output

diff --git a/Engine/Plugins/FX/Niagara/Source/Niagara/Classes/NiagaraSimCache.h b/Engine/Plugins/FX/Niagara/Source/Niagara/Classes/NiagaraSimCache.h
--- a/Engine/Plugins/FX/Niagara/Source/Niagara/Classes/NiagaraSimCache.h
+++ b/Engine/Plugins/FX/Niagara/Source/Niagara/Classes/NiagaraSimCache.h
@@ -545,7 +545,7 @@ public:
 
 #if WITH_EDITORONLY_DATA
 	UE_DEPRECATED(5.3, "For test purposes only, don't use in production code.")
-	bool IsDataEqual(const UNiagaraSimCache& OtherCache, FString& OutDifference) const;
+	bool IsDataEqual(const UNiagaraSimCache& OtherCache, float ErrorTolerance, FString& OutDifference) const;
 #endif
 	
 private:
diff --git a/Engine/Plugins/FX/Niagara/Source/Niagara/Private/NiagaraSimCache.cpp b/Engine/Plugins/FX/Niagara/Source/Niagara/Private/NiagaraSimCache.cpp
--- a/Engine/Plugins/FX/Niagara/Source/Niagara/Private/NiagaraSimCache.cpp
+++ b/Engine/Plugins/FX/Niagara/Source/Niagara/Private/NiagaraSimCache.cpp
@@ -1012,38 +1012,80 @@ bool FNiagaraSimCacheDataBuffers::operator!=(const FNiagaraSimCacheDataBuffers&
 
 #if WITH_EDITORONLY_DATA
 
-template<typename T>
-bool CompareAttributeData(TArray<int32> OtherIDValues, TMap<int32, int32> ExpectedIDToRow, TArray<T> ExpectedData, TArray<T> OtherData, int32 DataCount)
+namespace CacheCompare
 {
-	if (DataCount == 0)
+	constexpr int32 ErrorStrLen = 8000;
+
+	bool IsEqual(const float& Num1, const float& Num2, float Tolerance)
 	{
-		return true;
+		return FMath::IsNearlyEqual(Num1, Num2, Tolerance);
+	}
+	
+	bool IsEqual(const int32& Num1, const int32& Num2, float)
+	{
+		return Num1 == Num2;
+	}
+	
+	FString NumberToString(const float& Num)
+	{
+		return FString::SanitizeFloat(Num);
+	}
+
+	FString NumberToString(const int32& Num)
+	{
+		return FString::FromInt(Num);
 	}
-	for (int Index = 0; Index < OtherIDValues.Num(); Index++) {
-		int32& OtherID = OtherIDValues[Index];
-		int32 ExpectedDataRow = ExpectedIDToRow.FindOrAdd(OtherID, INDEX_NONE);
-		for (int32 Offset = 0; Offset < DataCount; Offset++)
+
+	template<typename T>
+	bool CompareAttributeData(TArray<int32> OtherIDValues, TMap<int32, int32> ExpectedIDToRow, TArray<T> ExpectedData, TArray<T> OtherData, int32 DataCount, TStringBuilder<1024>& ValueDiffString, float ErrorTolerance)
+	{
+		if (DataCount == 0)
 		{
-			int32 OtherDataIndex = Index + Offset * OtherIDValues.Num();
-			int32 ExpectedDataIndex = ExpectedDataRow + Offset * OtherIDValues.Num();
-			if (ExpectedDataRow == INDEX_NONE || !OtherData.IsValidIndex(OtherDataIndex) || !ExpectedData.IsValidIndex(ExpectedDataIndex) || OtherData[OtherDataIndex] != ExpectedData[ExpectedDataIndex])
+			return true;
+		}
+		for (int Index = 0; Index < OtherIDValues.Num(); Index++) {
+			int32& OtherID = OtherIDValues[Index];
+			int32 ExpectedDataRow = ExpectedIDToRow.FindOrAdd(OtherID, INDEX_NONE);
+			for (int32 Offset = 0; Offset < DataCount; Offset++)
 			{
-				return false;
+				int32 OtherDataIndex = Index + Offset * OtherIDValues.Num();
+				int32 ExpectedDataIndex = ExpectedDataRow + Offset * OtherIDValues.Num();
+				if (ExpectedDataRow == INDEX_NONE || !OtherData.IsValidIndex(OtherDataIndex) || !ExpectedData.IsValidIndex(ExpectedDataIndex))
+				{
+					ValueDiffString.Append(TEXT("Unable to find source data for particle ID "));
+					ValueDiffString.Append(FString::FromInt(OtherID));
+					return false;
+				}
+				
+				if (!IsEqual(OtherData[OtherDataIndex], ExpectedData[ExpectedDataIndex], ErrorTolerance))
+				{
+					ValueDiffString.Append(TEXT("Particle ID "));
+					ValueDiffString.Append(FString::FromInt(OtherID));
+					ValueDiffString.Append(TEXT(", Expected: "));
+					ValueDiffString.Append(NumberToString(ExpectedData[ExpectedDataIndex]));
+					ValueDiffString.Append(TEXT(", Actual: "));
+					ValueDiffString.Append(NumberToString(OtherData[OtherDataIndex]));
+					return false;
+				}
 			}
 		}
+		return true;
 	}
-	return true;
-}
 
-void AddError(TStringBuilder<1024>& Builder, const FString& Str)
-{
-	if (Builder.Len() < 1024)
+	void AddError(TStringBuilder<ErrorStrLen>& Builder, const FString& Str)
 	{
-		Builder.Append(Str);
+		if (Builder.Len() < ErrorStrLen)
+		{
+			Builder.Append(Str);
+			if (Builder.Len() >= ErrorStrLen)
+			{
+				Builder.Append("...\n");
+			}
+		}
 	}
 }
 
-bool UNiagaraSimCache::IsDataEqual(const UNiagaraSimCache& OtherCache, FString& OutDifference) const
+bool UNiagaraSimCache::IsDataEqual(const UNiagaraSimCache& OtherCache, float ErrorTolerance, FString& OutDifference) const
 {
 	const TArray<FNiagaraSimCacheFrame>& OtherFrames = OtherCache.CacheFrames;
 	if (GetNumFrames() != OtherCache.GetNumFrames())
@@ -1053,7 +1095,7 @@ bool UNiagaraSimCache::IsDataEqual(const UNiagaraSimCache& OtherCache, FString&
 	}
 
 	bool bEqual = true;
-	TStringBuilder<1024> Errors;
+	TStringBuilder<CacheCompare::ErrorStrLen> Errors;
 	for (int FrameIndex = 0; FrameIndex < CacheFrames.Num(); FrameIndex++)
 	{
 		FString Frame = FString::FromInt(FrameIndex);
@@ -1062,19 +1104,19 @@ bool UNiagaraSimCache::IsDataEqual(const UNiagaraSimCache& OtherCache, FString&
 
 		if (OtherFrame.SimulationAge != ExpectedFrame.SimulationAge)
 		{
-			AddError(Errors, TEXT("Simulation age different - frame ") + Frame + '\n');
+			CacheCompare::AddError(Errors, TEXT("Simulation age different - frame ") + Frame + '\n');
 			bEqual = false;
 		}
 		if (OtherFrame.EmitterData.Num() != ExpectedFrame.EmitterData.Num())
 		{
-			AddError(Errors, TEXT("Emitter count different - frame ") + Frame + '\n');
+			CacheCompare::AddError(Errors, TEXT("Emitter count different - frame ") + Frame + '\n');
 			bEqual = false;
 			continue;
 		}
 		
 		if (OtherFrame.SystemData.SystemDataBuffers != ExpectedFrame.SystemData.SystemDataBuffers)
 		{
-			AddError(Errors, TEXT("Cached system data different - frame ") + Frame + '\n');
+			CacheCompare::AddError(Errors, TEXT("Cached system data different - frame ") + Frame + '\n');
 			bEqual = false;
 		}
 		for (int EmitterIndex = 0; EmitterIndex < ExpectedFrame.EmitterData.Num(); EmitterIndex++)
@@ -1084,7 +1126,7 @@ bool UNiagaraSimCache::IsDataEqual(const UNiagaraSimCache& OtherCache, FString&
 			const FNiagaraSimCacheEmitterFrame& OtherEmitterFrame = OtherFrame.EmitterData[EmitterIndex];
 			if (OtherEmitterFrame.ParticleDataBuffers.NumInstances != ExpectedEmitterFrame.ParticleDataBuffers.NumInstances)
 			{
-				AddError(Errors, TEXT("Particle count different - frame ") + Frame + TEXT(" emitter ") + EmitterName.ToString() + '\n');
+				CacheCompare::AddError(Errors, TEXT("Particle count different - frame ") + Frame + TEXT(" emitter ") + EmitterName.ToString() + '\n');
 				bEqual = false;
 				continue;
 			}
@@ -1097,7 +1139,7 @@ bool UNiagaraSimCache::IsDataEqual(const UNiagaraSimCache& OtherCache, FString&
 
 			if (ExpectedIDValues.Num() != OtherIDValues.Num() || ExpectedIDValues.Num() != ExpectedEmitterFrame.ParticleDataBuffers.NumInstances)
 			{
-				AddError(Errors, TEXT("Invalid particle IDs - frame ") + Frame + TEXT(" emitter ") + EmitterName.ToString() + '\n');
+				CacheCompare::AddError(Errors, TEXT("Invalid particle IDs - frame ") + Frame + TEXT(" emitter ") + EmitterName.ToString() + '\n');
 				bEqual = false;
 				continue;
 			}
@@ -1126,12 +1168,15 @@ bool UNiagaraSimCache::IsDataEqual(const UNiagaraSimCache& OtherCache, FString&
 				TArray<int32> OtherInts;
 				ReadAttribute(ExpectedFloats, ExpectedHalfs, ExpectedInts, Var.Variable.GetName(), EmitterName, FrameIndex);
 				OtherCache.ReadAttribute(OtherFloats, OtherHalfs, OtherInts, Var.Variable.GetName(), EmitterName, FrameIndex);
-				
-				if (!CompareAttributeData( OtherIDValues, ExpectedIDToRow, ExpectedFloats,OtherFloats, Var.FloatCount) ||
-					!CompareAttributeData( OtherIDValues, ExpectedIDToRow, ExpectedHalfs,OtherHalfs, Var.HalfCount) ||
-					!CompareAttributeData( OtherIDValues, ExpectedIDToRow, ExpectedInts,OtherInts, Var.Int32Count))
+
+				TStringBuilder<1024> ValueDiff;
+				if (!CacheCompare::CompareAttributeData( OtherIDValues, ExpectedIDToRow, ExpectedFloats,OtherFloats, Var.FloatCount, ValueDiff, ErrorTolerance) ||
+					!CacheCompare::CompareAttributeData( OtherIDValues, ExpectedIDToRow, ExpectedHalfs,OtherHalfs, Var.HalfCount, ValueDiff, ErrorTolerance) ||
+					!CacheCompare::CompareAttributeData( OtherIDValues, ExpectedIDToRow, ExpectedInts,OtherInts, Var.Int32Count, ValueDiff, ErrorTolerance))
 				{
-					AddError(Errors, TEXT("Different particle data - frame ") + Frame + TEXT(" emitter ") + EmitterName.ToString() + TEXT(" attribute ") + Var.Variable.GetName().ToString() + '\n');
+					CacheCompare::AddError(Errors, TEXT("Different particle data - frame ") + Frame + TEXT(" emitter ") + EmitterName.ToString() + TEXT(" attribute ") + Var.Variable.GetName().ToString() + ": ");
+					CacheCompare::AddError(Errors, ValueDiff.ToString());
+					CacheCompare::AddError(Errors, "\n");
 					bEqual = false;
 				}
 			}
