Niagara - Additional fixes for data interface routing.
+ Fix renderer data interface bindings.
+ Update the translator so that data interface reads will correctly find a previously written data interface in the case where the DI names don't match.
+ Update the translator so that a data interface that is written to multiple parameters generates mutiple entries in the cached default data interfaces.
+ Update the grid 2d and grid 3d collections to use the data interface itself to find attributes, rather than the data interface name.

diff --git a/Engine/Plugins/FX/Niagara/Source/Niagara/Classes/NiagaraDataInterfaceGrid2DCollection.h b/Engine/Plugins/FX/Niagara/Source/Niagara/Classes/NiagaraDataInterfaceGrid2DCollection.h
--- a/Engine/Plugins/FX/Niagara/Source/Niagara/Classes/NiagaraDataInterfaceGrid2DCollection.h
+++ b/Engine/Plugins/FX/Niagara/Source/Niagara/Classes/NiagaraDataInterfaceGrid2DCollection.h
@@ -188,12 +188,10 @@ public:
 	//~ UNiagaraDataInterface interface END
 
 private:
-	static void CollectAttributesForScript(UNiagaraScript* Script, FName VariableName, TArray<FNiagaraVariableBase>& OutVariables, TArray<uint32>& OutVariableOffsets, int32& TotalAttributes, TArray<FText>* OutWarnings = nullptr, bool UseReader = false);
+	static void CollectAttributesForScript(UNiagaraScript* Script, const UNiagaraDataInterface* DataInterFace, TArray<FNiagaraVariableBase>& OutVariables, TArray<uint32>& OutVariableOffsets, int32& TotalAttributes, TArray<FText>* OutWarnings = nullptr);
 public:
-	/** Finds all attributes by locating the variable name inside the parameter stores. */
-	void FindAttributesByName(FName DataInterfaceName, TArray<FNiagaraVariableBase>& OutVariables, TArray<uint32>& OutVariableOffsets, int32& OutNumAttribChannelsFound, TArray<FText>* OutWarnings = nullptr) const;
 	/** Finds all attributes by locating the data interface amongst the parameter stores. */
-	void FindAttributes(TArray<FNiagaraVariableBase>& OutVariables, TArray<uint32>& OutVariableOffsets, int32& OutNumAttribChannelsFound, TArray<FText>* OutWarnings = nullptr, bool UseReader = false) const;
+	void FindAttributes(TArray<FNiagaraVariableBase>& OutVariables, TArray<uint32>& OutVariableOffsets, int32& OutNumAttribChannelsFound, TArray<FText>* OutWarnings = nullptr) const;
 
 	// Fills a texture render target 2d with the current data from the simulation
 	// #todo(dmp): this will eventually go away when we formalize how data makes it out of Niagara
diff --git a/Engine/Plugins/FX/Niagara/Source/Niagara/Classes/NiagaraDataInterfaceGrid3DCollection.h b/Engine/Plugins/FX/Niagara/Source/Niagara/Classes/NiagaraDataInterfaceGrid3DCollection.h
--- a/Engine/Plugins/FX/Niagara/Source/Niagara/Classes/NiagaraDataInterfaceGrid3DCollection.h
+++ b/Engine/Plugins/FX/Niagara/Source/Niagara/Classes/NiagaraDataInterfaceGrid3DCollection.h
@@ -200,9 +200,6 @@ public:
 	//~ UNiagaraDataInterface interface END
 
 public:
-	/** Finds all attributes by locating the variable name inside the parameter stores. */
-	void FindAttributesByName(FName DataInterfaceName, TArray<FNiagaraVariableBase>& OutVariables, TArray<uint32>& OutVariableOffsets, int32& OutNumAttribChannelsFound, TArray<FText>* OutWarnings = nullptr) const;
-
 	// Fills a texture render target 2d with the current data from the simulation
 	// #todo(dmp): this will eventually go away when we formalize how data makes it out of Niagara
 	// #todo(dmp): reimplement for 3d
@@ -349,7 +346,7 @@ protected:
 
 public:
 	/** Finds all attributes by locating the data interface amongst the parameter stores. */
-	void FindAttributes(TArray<FNiagaraVariableBase>& OutVariables, TArray<uint32>& OutVariableOffsets, int32& OutNumAttribChannelsFound, TArray<FText>* OutWarnings = nullptr, bool UseReader = false) const;
+	void FindAttributes(TArray<FNiagaraVariableBase>& OutVariables, TArray<uint32>& OutVariableOffsets, int32& OutNumAttribChannelsFound, TArray<FText>* OutWarnings = nullptr) const;
 protected:
-	static void CollectAttributesForScript(UNiagaraScript* Script, FName VariableName, TArray<FNiagaraVariableBase>& OutVariables, TArray<uint32>& OutVariableOffsets, int32& TotalAttributes, TArray<FText>* OutWarnings = nullptr, bool UseReader = false);
+	static void CollectAttributesForScript(UNiagaraScript* Script, const UNiagaraDataInterface* DataInterface, TArray<FNiagaraVariableBase>& OutVariables, TArray<uint32>& OutVariableOffsets, int32& TotalAttributes, TArray<FText>* OutWarnings = nullptr);
 };
\ No newline at end of file
diff --git a/Engine/Plugins/FX/Niagara/Source/Niagara/Classes/NiagaraEmitter.h b/Engine/Plugins/FX/Niagara/Source/Niagara/Classes/NiagaraEmitter.h
--- a/Engine/Plugins/FX/Niagara/Source/Niagara/Classes/NiagaraEmitter.h
+++ b/Engine/Plugins/FX/Niagara/Source/Niagara/Classes/NiagaraEmitter.h
@@ -365,6 +365,9 @@ struct NIAGARA_API FVersionedNiagaraEmitterData
 	UPROPERTY()
 	FNiagaraParameterStore RendererBindings;
 
+	UPROPERTY()
+	TMap<FNiagaraVariableBase, FNiagaraVariableBase> ResolvedDIBindings;
+
 	void CopyFrom(const FVersionedNiagaraEmitterData& Source);
 	void PostLoad(UNiagaraEmitter& Emitter, int32 NiagaraVer);
 
diff --git a/Engine/Plugins/FX/Niagara/Source/Niagara/Private/NiagaraCommon.cpp b/Engine/Plugins/FX/Niagara/Source/Niagara/Private/NiagaraCommon.cpp
--- a/Engine/Plugins/FX/Niagara/Source/Niagara/Private/NiagaraCommon.cpp
+++ b/Engine/Plugins/FX/Niagara/Source/Niagara/Private/NiagaraCommon.cpp
@@ -883,9 +883,21 @@ void FNiagaraMaterialAttributeBinding::CacheValues(const UNiagaraEmitter* InEmit
 {
 	if (InEmitter != nullptr)
 	{
+		// Resolve emitter alias.
 		ResolvedNiagaraVariable = FNiagaraUtilities::ResolveAliases(NiagaraVariable, 
 			FNiagaraAliasContext(FNiagaraAliasContext::ERapidIterationParameterMode::EmitterOrParticleScript)
 			.ChangeEmitterToEmitterName(InEmitter->GetUniqueEmitterName()));
+		
+		// Resolve DI bindings.
+		if (ResolvedNiagaraVariable.IsDataInterface())
+		{
+			
+			const FNiagaraVariableBase* ResolvedDIBinding = InEmitter->GetLatestEmitterData()->ResolvedDIBindings.Find(ResolvedNiagaraVariable);
+			if (ResolvedDIBinding != nullptr)
+			{
+				ResolvedNiagaraVariable = *ResolvedDIBinding;
+			}
+		}
 	}
 	else
 	{
diff --git a/Engine/Plugins/FX/Niagara/Source/Niagara/Private/NiagaraDataInterfaceGrid2DCollection.cpp b/Engine/Plugins/FX/Niagara/Source/Niagara/Private/NiagaraDataInterfaceGrid2DCollection.cpp
--- a/Engine/Plugins/FX/Niagara/Source/Niagara/Private/NiagaraDataInterfaceGrid2DCollection.cpp
+++ b/Engine/Plugins/FX/Niagara/Source/Niagara/Private/NiagaraDataInterfaceGrid2DCollection.cpp
@@ -1982,26 +1982,15 @@ bool UNiagaraDataInterfaceGrid2DCollection::GetExposedVariableValue(const FNiaga
 	return false;
 }
 
-void UNiagaraDataInterfaceGrid2DCollection::CollectAttributesForScript(UNiagaraScript* Script, FName VariableName, TArray<FNiagaraVariableBase>& OutVariables, TArray<uint32>& OutVariableOffsets, int32& TotalAttributes, TArray<FText>* OutWarnings, bool UseReader)
+void UNiagaraDataInterfaceGrid2DCollection::CollectAttributesForScript(UNiagaraScript* Script, const UNiagaraDataInterface* DataInterface, TArray<FNiagaraVariableBase>& OutVariables, TArray<uint32>& OutVariableOffsets, int32& TotalAttributes, TArray<FText>* OutWarnings)
 {
 	if (const FNiagaraScriptExecutionParameterStore* ParameterStore = Script->GetExecutionReadyParameterStore(ENiagaraSimTarget::GPUComputeSim))
 	{
-		FNiagaraVariableBase DataInterfaceVariable(FNiagaraTypeDefinition(UNiagaraDataInterfaceGrid2DCollection::StaticClass()), VariableName);
-
-		if (UseReader)
-		{
-			DataInterfaceVariable = FNiagaraVariableBase(FNiagaraTypeDefinition(UNiagaraDataInterfaceGrid2DCollectionReader::StaticClass()), VariableName);
-		}
-		else
-		{
-			DataInterfaceVariable = FNiagaraVariableBase(FNiagaraTypeDefinition(UNiagaraDataInterfaceGrid2DCollection::StaticClass()), VariableName);
-		}
-
-		const int32 * IndexOfDataInterface = ParameterStore->FindParameterOffset(DataInterfaceVariable);
-		if (IndexOfDataInterface != nullptr)
+		int32 IndexOfDataInterface = ParameterStore->GetDataInterfaces().IndexOfByKey(DataInterface);
+		if (IndexOfDataInterface != INDEX_NONE)
 		{
 			TConstArrayView<FNiagaraDataInterfaceGPUParamInfo> ParamInfoArray = Script->GetDataInterfaceGPUParamInfos();
-			for (const FNiagaraDataInterfaceGeneratedFunction& Func : ParamInfoArray[*IndexOfDataInterface].GeneratedFunctions)
+			for (const FNiagaraDataInterfaceGeneratedFunction& Func : ParamInfoArray[IndexOfDataInterface].GeneratedFunctions)
 			{
 				if (const FName* AttributeName = Func.FindSpecifierValue(UNiagaraDataInterfaceRWBase::NAME_Attribute))
 				{
@@ -2034,7 +2023,7 @@ void UNiagaraDataInterfaceGrid2DCollection::CollectAttributesForScript(UNiagaraS
 	}
 }
 
-void UNiagaraDataInterfaceGrid2DCollection::FindAttributesByName(FName VariableName, TArray<FNiagaraVariableBase>& OutVariables, TArray<uint32>& OutVariableOffsets, int32& OutNumAttribChannelsFound, TArray<FText>* OutWarnings) const
+void UNiagaraDataInterfaceGrid2DCollection::FindAttributes(TArray<FNiagaraVariableBase>& OutVariables, TArray<uint32>& OutVariableOffsets, int32& OutNumAttribChannelsFound, TArray<FText>* OutWarnings) const
 {
 	OutNumAttribChannelsFound = 0;
 
@@ -2050,57 +2039,7 @@ void UNiagaraDataInterfaceGrid2DCollection::FindAttributesByName(FName VariableN
 		FVersionedNiagaraEmitterData* EmitterData = EmitterHandle.GetEmitterData();
 		if (EmitterData && EmitterHandle.GetIsEnabled() && EmitterData->IsValid() && (EmitterData->SimTarget == ENiagaraSimTarget::GPUComputeSim))
 		{
-			CollectAttributesForScript(EmitterData->GetGPUComputeScript(), VariableName, OutVariables, OutVariableOffsets, TotalAttributes, OutWarnings);
-		}
-	}
-	OutNumAttribChannelsFound = TotalAttributes - NumAttributes;
-}
-
-void UNiagaraDataInterfaceGrid2DCollection::FindAttributes(TArray<FNiagaraVariableBase>& OutVariables, TArray<uint32>& OutVariableOffsets, int32& OutNumAttribChannelsFound, TArray<FText>* OutWarnings, bool UseReader) const
-{
-	OutNumAttribChannelsFound = 0;
-
-	UNiagaraSystem* OwnerSystem = GetTypedOuter<UNiagaraSystem>();
-	if (OwnerSystem == nullptr)
-	{
-		return;
-	}
-
-	int32 TotalAttributes = NumAttributes;
-	for (const FNiagaraEmitterHandle& EmitterHandle : OwnerSystem->GetEmitterHandles())
-	{
-		FVersionedNiagaraEmitterData* EmitterData = EmitterHandle.GetEmitterData();
-		if (EmitterData && EmitterHandle.GetIsEnabled() && EmitterData->IsValid() && (EmitterData->SimTarget == ENiagaraSimTarget::GPUComputeSim))
-		{
-			// Search scripts for this data interface so we get the variable name
-			auto FindDataInterfaceVariable =
-				[&OwnerSystem, &EmitterData](const UNiagaraDataInterface* DataInterface) -> FName
-				{
-					UNiagaraScript* Scripts[] =
-					{
-						OwnerSystem->GetSystemSpawnScript(),
-						OwnerSystem->GetSystemUpdateScript(),
-						EmitterData->GetGPUComputeScript(),
-					};
-
-					for (UNiagaraScript* Script : Scripts)
-					{
-						for (const FNiagaraScriptResolvedDataInterfaceInfo& DataInterfaceInfo : Script->GetResolvedDataInterfaces())
-						{
-							if (DataInterfaceInfo.ResolvedDataInterface == DataInterface && DataInterfaceInfo.bIsInternal == false)
-							{
-								return DataInterfaceInfo.ResolvedVariable.GetName();
-							}
-						}
-					}
-					return NAME_None;
-				};
-
-			const FName VariableName = FindDataInterfaceVariable(this);
-			if (!VariableName.IsNone() )
-			{
-				CollectAttributesForScript(EmitterData->GetGPUComputeScript(), VariableName, OutVariables, OutVariableOffsets, TotalAttributes, OutWarnings, UseReader);
-			}
+			CollectAttributesForScript(EmitterData->GetGPUComputeScript(), this, OutVariables, OutVariableOffsets, TotalAttributes, OutWarnings);
 		}
 	}
 	OutNumAttribChannelsFound = TotalAttributes - NumAttributes;
diff --git a/Engine/Plugins/FX/Niagara/Source/Niagara/Private/NiagaraDataInterfaceGrid2DCollectionReader.cpp b/Engine/Plugins/FX/Niagara/Source/Niagara/Private/NiagaraDataInterfaceGrid2DCollectionReader.cpp
--- a/Engine/Plugins/FX/Niagara/Source/Niagara/Private/NiagaraDataInterfaceGrid2DCollectionReader.cpp
+++ b/Engine/Plugins/FX/Niagara/Source/Niagara/Private/NiagaraDataInterfaceGrid2DCollectionReader.cpp
@@ -108,7 +108,7 @@ bool UNiagaraDataInterfaceGrid2DCollectionReader::InitPerInstanceData(void* PerI
 			int32 NumNamedAttribChannelsFound = 0;
 			TArray<FNiagaraVariableBase> Vars;
 			TArray<uint32> Offsets;
-			FindAttributes(Vars, Offsets, NumNamedAttribChannelsFound, nullptr, true);
+			FindAttributes(Vars, Offsets, NumNamedAttribChannelsFound, nullptr);
 
 			FGrid2DCollectionRWInstanceData_GameThread* Grid2DInstanceData = GridInterface->GetSystemInstancesToProxyData_GT().FindRef(SystemInstance->GetId());
 			InstanceData->OtherDI = GridInterface;
diff --git a/Engine/Plugins/FX/Niagara/Source/Niagara/Private/NiagaraDataInterfaceGrid3DCollection.cpp b/Engine/Plugins/FX/Niagara/Source/Niagara/Private/NiagaraDataInterfaceGrid3DCollection.cpp
--- a/Engine/Plugins/FX/Niagara/Source/Niagara/Private/NiagaraDataInterfaceGrid3DCollection.cpp
+++ b/Engine/Plugins/FX/Niagara/Source/Niagara/Private/NiagaraDataInterfaceGrid3DCollection.cpp
@@ -3617,30 +3617,6 @@ bool UNiagaraDataInterfaceGrid3DCollection::GetExposedVariableValue(const FNiaga
 	return false;
 }
 
-
-
-void UNiagaraDataInterfaceGrid3DCollection::FindAttributesByName(FName VariableName, TArray<FNiagaraVariableBase>& OutVariables, TArray<uint32>& OutVariableOffsets, int32& OutNumAttribChannelsFound, TArray<FText>* OutWarnings) const
-{
-	OutNumAttribChannelsFound = 0;
-
-	UNiagaraSystem* OwnerSystem = GetTypedOuter<UNiagaraSystem>();
-	if (OwnerSystem == nullptr)
-	{
-		return;
-	}
-
-	int32 TotalAttributes = NumAttributes;
-	for (const FNiagaraEmitterHandle& EmitterHandle : OwnerSystem->GetEmitterHandles())
-	{
-		FVersionedNiagaraEmitterData* EmitterData = EmitterHandle.GetEmitterData();
-		if (EmitterData && EmitterHandle.GetIsEnabled() && EmitterData->IsValid() && (EmitterData->SimTarget == ENiagaraSimTarget::GPUComputeSim))
-		{
-			CollectAttributesForScript(EmitterData->GetGPUComputeScript(), VariableName, OutVariables, OutVariableOffsets, TotalAttributes, OutWarnings);
-		}
-	}
-	OutNumAttribChannelsFound = TotalAttributes - NumAttributes;
-}
-
 bool UNiagaraDataInterfaceGrid3DCollection::FillVolumeTexture(const UNiagaraComponent* Component, UVolumeTexture* Dest, int AttributeIndex)
 {
 	/*
@@ -4244,7 +4220,7 @@ void FNiagaraDataInterfaceProxyGrid3DCollectionProxy::GetDispatchArgs(const FNDI
 	}
 }
 
-void UNiagaraDataInterfaceGrid3DCollection::FindAttributes(TArray<FNiagaraVariableBase>& OutVariables, TArray<uint32>& OutVariableOffsets, int32& OutNumAttribChannelsFound, TArray<FText>* OutWarnings, bool UseReader) const
+void UNiagaraDataInterfaceGrid3DCollection::FindAttributes(TArray<FNiagaraVariableBase>& OutVariables, TArray<uint32>& OutVariableOffsets, int32& OutNumAttribChannelsFound, TArray<FText>* OutWarnings) const
 {
 	{
 		OutNumAttribChannelsFound = 0;
@@ -4261,61 +4237,22 @@ void UNiagaraDataInterfaceGrid3DCollection::FindAttributes(TArray<FNiagaraVariab
 			FVersionedNiagaraEmitterData* EmitterData = EmitterHandle.GetEmitterData();
 			if (EmitterData && EmitterHandle.GetIsEnabled() && EmitterData->IsValid() && (EmitterData->SimTarget == ENiagaraSimTarget::GPUComputeSim))
 			{
-				// Search scripts for this data interface so we get the variable name
-				auto FindDataInterfaceVariable =
-					[&OwnerSystem, &EmitterData](const UNiagaraDataInterface* DataInterface) -> FName
-				{
-					UNiagaraScript* Scripts[] =
-					{
-						OwnerSystem->GetSystemSpawnScript(),
-						OwnerSystem->GetSystemUpdateScript(),
-						EmitterData->GetGPUComputeScript(),
-					};
-
-					for (UNiagaraScript* Script : Scripts)
-					{
-						for (const FNiagaraScriptResolvedDataInterfaceInfo& DataInterfaceInfo : Script->GetResolvedDataInterfaces())
-						{
-							if (DataInterfaceInfo.ResolvedDataInterface == DataInterface && DataInterfaceInfo.bIsInternal == false)
-							{
-								return DataInterfaceInfo.ResolvedVariable.GetName();
-							}
-						}
-					}
-					return NAME_None;
-				};
-
-				const FName VariableName = FindDataInterfaceVariable(this);
-				if (!VariableName.IsNone())
-				{
-					CollectAttributesForScript(EmitterData->GetGPUComputeScript(), VariableName, OutVariables, OutVariableOffsets, TotalAttributes, OutWarnings, UseReader);
-				}
+				CollectAttributesForScript(EmitterData->GetGPUComputeScript(), this, OutVariables, OutVariableOffsets, TotalAttributes, OutWarnings);
 			}
 		}
 		OutNumAttribChannelsFound = TotalAttributes - NumAttributes;
 	}
 }
 
-void UNiagaraDataInterfaceGrid3DCollection::CollectAttributesForScript(UNiagaraScript* Script, FName VariableName, TArray<FNiagaraVariableBase>& OutVariables, TArray<uint32>& OutVariableOffsets, int32& TotalAttributes, TArray<FText>* OutWarnings, bool UseReader)
+void UNiagaraDataInterfaceGrid3DCollection::CollectAttributesForScript(UNiagaraScript* Script, const UNiagaraDataInterface* DataInterface, TArray<FNiagaraVariableBase>& OutVariables, TArray<uint32>& OutVariableOffsets, int32& TotalAttributes, TArray<FText>* OutWarnings)
 {
 	if (const FNiagaraScriptExecutionParameterStore* ParameterStore = Script->GetExecutionReadyParameterStore(ENiagaraSimTarget::GPUComputeSim))
 	{
-		FNiagaraVariableBase DataInterfaceVariable;
-
-		if (UseReader)
-		{
-			DataInterfaceVariable = FNiagaraVariableBase(FNiagaraTypeDefinition(UNiagaraDataInterfaceGrid3DCollectionReader::StaticClass()), VariableName);
-		}
-		else
-		{
-			DataInterfaceVariable = FNiagaraVariableBase(FNiagaraTypeDefinition(UNiagaraDataInterfaceGrid3DCollection::StaticClass()), VariableName);
-		}
-
-		const int32* IndexOfDataInterface = ParameterStore->FindParameterOffset(DataInterfaceVariable);
-		if (IndexOfDataInterface != nullptr)
+		int32 IndexOfDataInterface = ParameterStore->GetDataInterfaces().IndexOfByKey(DataInterface);
+		if (IndexOfDataInterface != INDEX_NONE)
 		{
 			TConstArrayView<FNiagaraDataInterfaceGPUParamInfo> ParamInfoArray = Script->GetDataInterfaceGPUParamInfos();
-			for (const FNiagaraDataInterfaceGeneratedFunction& Func : ParamInfoArray[*IndexOfDataInterface].GeneratedFunctions)
+			for (const FNiagaraDataInterfaceGeneratedFunction& Func : ParamInfoArray[IndexOfDataInterface].GeneratedFunctions)
 			{
 				if (const FName* AttributeName = Func.FindSpecifierValue(UNiagaraDataInterfaceRWBase::NAME_Attribute))
 				{
diff --git a/Engine/Plugins/FX/Niagara/Source/Niagara/Private/NiagaraDataInterfaceGrid3DCollectionReader.cpp b/Engine/Plugins/FX/Niagara/Source/Niagara/Private/NiagaraDataInterfaceGrid3DCollectionReader.cpp
--- a/Engine/Plugins/FX/Niagara/Source/Niagara/Private/NiagaraDataInterfaceGrid3DCollectionReader.cpp
+++ b/Engine/Plugins/FX/Niagara/Source/Niagara/Private/NiagaraDataInterfaceGrid3DCollectionReader.cpp
@@ -107,7 +107,7 @@ bool UNiagaraDataInterfaceGrid3DCollectionReader::InitPerInstanceData(void* PerI
 			int32 NumNamedAttribChannelsFound = 0;
 			TArray<FNiagaraVariableBase> Vars;
 			TArray<uint32> Offsets;
-			FindAttributes(Vars, Offsets, NumNamedAttribChannelsFound, nullptr, true);
+			FindAttributes(Vars, Offsets, NumNamedAttribChannelsFound, nullptr);
 
 			// #todo(dmp): slight hack here - these aren't always the same, but we generally can't use unnamed attrs very well
 			NumAttribChannelsFound = NumNamedAttribChannelsFound;
@@ -221,7 +221,7 @@ bool UNiagaraDataInterfaceGrid3DCollectionReader::PerInstanceTickPostSimulate(vo
 		int32 NumNamedAttribChannelsFound = 0;
 		TArray<FNiagaraVariableBase> Vars;
 		TArray<uint32> Offsets;
-		FindAttributes(Vars, Offsets, NumNamedAttribChannelsFound, nullptr, true);
+		FindAttributes(Vars, Offsets, NumNamedAttribChannelsFound, nullptr);
 
 		// #todo(dmp): slight hack here - these aren't always the same, but we generally can't use unnamed attrs very well
 		NumAttribChannelsFound = NumNamedAttribChannelsFound;
diff --git a/Engine/Plugins/FX/Niagara/Source/Niagara/Private/NiagaraEmitterInstance.cpp b/Engine/Plugins/FX/Niagara/Source/Niagara/Private/NiagaraEmitterInstance.cpp
--- a/Engine/Plugins/FX/Niagara/Source/Niagara/Private/NiagaraEmitterInstance.cpp
+++ b/Engine/Plugins/FX/Niagara/Source/Niagara/Private/NiagaraEmitterInstance.cpp
@@ -478,6 +478,12 @@ void FNiagaraEmitterInstance::Init(int32 InEmitterIdx, FNiagaraSystemInstanceID
 			if (ParentSystemInstance)
 			{
 				ParentSystemInstance->GetInstanceParameters().Bind(&RendererBindings);
+				FNiagaraSystemSimulationPtr ParentSystemSimulation = ParentSystemInstance->GetSystemSimulation();
+				if (ParentSystemSimulation.IsValid())
+				{
+					ParentSystemSimulation->GetSpawnExecutionContext()->Parameters.Bind(&RendererBindings);
+					ParentSystemSimulation->GetUpdateExecutionContext()->Parameters.Bind(&RendererBindings);
+				}
 			}
 
 			if (GPUExecContext && EmitterData->SimTarget == ENiagaraSimTarget::GPUComputeSim)
diff --git a/Engine/Plugins/FX/Niagara/Source/Niagara/Private/NiagaraResolveDIHelpers.cpp b/Engine/Plugins/FX/Niagara/Source/Niagara/Private/NiagaraResolveDIHelpers.cpp
--- a/Engine/Plugins/FX/Niagara/Source/Niagara/Private/NiagaraResolveDIHelpers.cpp
+++ b/Engine/Plugins/FX/Niagara/Source/Niagara/Private/NiagaraResolveDIHelpers.cpp
@@ -37,17 +37,17 @@ namespace FNiagaraResolveDIHelpers
 
 		for (const FNiagaraScriptDataInterfaceInfo& CachedDefaultDataInterface : TargetScript->GetCachedDefaultDataInterfaces())
 		{
-			for (FName RegisteredParameterMapWrite : CachedDefaultDataInterface.RegisteredParameterMapWrites)
+			if (CachedDefaultDataInterface.RegisteredParameterMapWrite != NAME_None)
 			{
-				FNiagaraVariable WriteVariable(CachedDefaultDataInterface.Type, RegisteredParameterMapWrite);
+				FNiagaraVariable WriteVariable(CachedDefaultDataInterface.Type, CachedDefaultDataInterface.RegisteredParameterMapWrite);
 				if (EmitterName.IsEmpty() == false)
 				{
 					WriteVariable = FNiagaraUtilities::ResolveAliases(WriteVariable, FNiagaraAliasContext().ChangeEmitterNameToEmitter(EmitterName));
 				}
 
+				// Handle bindings.
 				if (CachedDefaultDataInterface.RegisteredParameterMapRead != NAME_None)
 				{
-					// Handle bindings.
 					FNiagaraVariable ReadVariable(CachedDefaultDataInterface.Type, CachedDefaultDataInterface.RegisteredParameterMapRead);
 					if (EmitterName.IsEmpty() == false)
 					{
@@ -76,28 +76,26 @@ namespace FNiagaraResolveDIHelpers
 						}
 					}
 				}
-				else
+		
+				// Handle assignments
+				FDataInterfaceSourceEmitterNamePair* CurrentAssignment = OutVariableAssignmentMap.Find(WriteVariable);
+				if (CurrentAssignment != nullptr)
 				{
-					// Handle assignments
-					FDataInterfaceSourceEmitterNamePair* CurrentAssignment = OutVariableAssignmentMap.Find(WriteVariable);
-					if (CurrentAssignment != nullptr)
-					{
-						if (CurrentAssignment->DataInterface != CachedDefaultDataInterface.DataInterface)
-						{
-							OutErrorMessages.Add(FText::Format(
-								LOCTEXT("MultipleAssignmentsFormat", "A data interface parameter was thh target of an assignment multiple times in a single system.  The data interface used in the simulation may be incorrect.  Target Parameter: {0} First Assignment: {1} Current Assignment: {2}"),
-								FText::FromName(WriteVariable.GetName()),
-								FText::FromString(CurrentAssignment->DataInterface->GetName()),
-								FText::FromString(CachedDefaultDataInterface.DataInterface->GetName())));
-						}
-					}
-					else
+					if (CurrentAssignment->DataInterface != CachedDefaultDataInterface.DataInterface)
 					{
-						FDataInterfaceSourceEmitterNamePair& NewAssignment = OutVariableAssignmentMap.Add(WriteVariable);
-						NewAssignment.DataInterface = CachedDefaultDataInterface.DataInterface;
-						NewAssignment.SourceEmitterName = CachedDefaultDataInterface.SourceEmitterName;
+						OutErrorMessages.Add(FText::Format(
+							LOCTEXT("MultipleAssignmentsFormat", "A data interface parameter was the target of an assignment multiple times in a single system.  The data interface used in the simulation may be incorrect.  Target Parameter: {0} First Assignment: {1} Current Assignment: {2}"),
+							FText::FromName(WriteVariable.GetName()),
+							FText::FromString(CurrentAssignment->DataInterface->GetName()),
+							FText::FromString(CachedDefaultDataInterface.DataInterface->GetName())));
 					}
 				}
+				else
+				{
+					FDataInterfaceSourceEmitterNamePair& NewAssignment = OutVariableAssignmentMap.Add(WriteVariable);
+					NewAssignment.DataInterface = CachedDefaultDataInterface.DataInterface;
+					NewAssignment.SourceEmitterName = CachedDefaultDataInterface.SourceEmitterName;
+				}
 			}
 		}
 	}
@@ -172,18 +170,48 @@ namespace FNiagaraResolveDIHelpers
 		}
 	}
 
+	void CompactBindings(const TMap<FNiagaraVariableBase, FNiagaraVariableBase> InVariableBindingMap, TMap<FNiagaraVariableBase, FNiagaraVariableBase>& OutCompactedVariableBindingMap, TArray<FText>& OutErrorMessages)
+	{
+		TArray<FNiagaraVariableBase> SourceVariables;
+		InVariableBindingMap.GetKeys(SourceVariables);
+		TMap<FNiagaraVariableBase, FNiagaraVariableBase> CompactedVariableBindingMap;
+		for (const FNiagaraVariableBase& SourceVariable : SourceVariables)
+		{
+			FNiagaraVariableBase BoundVariable;
+			const FNiagaraVariableBase* CurrentBoundVariable = &SourceVariable;
+			TSet<FNiagaraVariableBase> SeenBoundVariables;
+			bool bCircularReferenceDetected = false;
+			while (CurrentBoundVariable != nullptr && bCircularReferenceDetected == false)
+			{
+				SeenBoundVariables.Add(*CurrentBoundVariable);
+				BoundVariable = *CurrentBoundVariable;
+				CurrentBoundVariable = InVariableBindingMap.Find(BoundVariable);
+				if (CurrentBoundVariable != nullptr && SeenBoundVariables.Contains(*CurrentBoundVariable))
+				{
+					bCircularReferenceDetected = true;
+					OutErrorMessages.Add(FText::Format(
+						LOCTEXT("CircularDependencyWarningFormat", "A data interface parameter circular dependency was found.  The data interface used in the simulation may be incorrect.  Target Parameter: {0} Resolved Parameter: {1}"),
+						FText::FromName(SourceVariable.GetName()),
+						FText::FromName(CurrentBoundVariable->GetName())));
+				}
+			}
+			OutCompactedVariableBindingMap.Add(SourceVariable, BoundVariable);
+		}
+	}
+
 	void ResolveDIsForScript(
 		const UNiagaraSystem* System,
 		UNiagaraScript* TargetScript,
 		const FString& EmitterName,
 		const TMap<FNiagaraVariableBase, FDataInterfaceSourceEmitterNamePair>& VariableAssignmentMap,
-		const TMap<FNiagaraVariableBase, FNiagaraVariableBase>& VariableBindingMap,
+		const TMap<FNiagaraVariableBase, FNiagaraVariableBase>& CompactedVariableBindingMap,
 		TArray<FText>& OutErrorMessages)
 	{
 		TArray<FNiagaraScriptResolvedDataInterfaceInfo> ResolvedDataInterfaces;
 		ResolvedDataInterfaces.Reserve(TargetScript->GetCachedDefaultDataInterfaces().Num());
 		TArray<FNiagaraResolvedUserDataInterfaceBinding> UserDataInterfaceBindings;
 		int32 ResolvedDataInterfaceIndex = 0;
+		TSet<FNiagaraVariableBase> ParameterStoreVariables;
 		for (const FNiagaraScriptDataInterfaceInfo& CachedDefaultDataInterface : TargetScript->GetCachedDefaultDataInterfaces())
 		{
 			FNiagaraScriptResolvedDataInterfaceInfo& ResolvedDataInterface = ResolvedDataInterfaces.AddDefaulted_GetRef();
@@ -200,46 +228,47 @@ namespace FNiagaraResolveDIHelpers
 					ReadVariable = FNiagaraUtilities::ResolveAliases(ReadVariable, FNiagaraAliasContext().ChangeEmitterToEmitterName(EmitterName));
 				}
 
-				FNiagaraVariableBase BoundVariable;
-				const FNiagaraVariableBase* CurrentBoundVariable = &ReadVariable;
-				TSet<FNiagaraVariableBase> SeenBoundVariables;
-				bool bCircularReferenceDetected = false;
-				while (CurrentBoundVariable != nullptr && bCircularReferenceDetected == false)
+				const FNiagaraVariableBase* BoundVariable = CompactedVariableBindingMap.Find(ReadVariable);
+				if (BoundVariable == nullptr)
 				{
-					SeenBoundVariables.Add(*CurrentBoundVariable);
-					BoundVariable = *CurrentBoundVariable;
-					CurrentBoundVariable = VariableBindingMap.Find(BoundVariable);
-
-					if (CurrentBoundVariable != nullptr && SeenBoundVariables.Contains(*CurrentBoundVariable))
-					{
-						bCircularReferenceDetected = true;
-						OutErrorMessages.Add(FText::Format(
-							LOCTEXT("CircularDependencyWarningFormat", "A data interface parameter circular dependency found.  The data interface used in the simulation may be incorrect.  Target Parameter: {0} Resolved Parameter: {1}"),
-							FText::FromName(ReadVariable.GetName()),
-							FText::FromName(CurrentBoundVariable->GetName())));
-					}
+					BoundVariable = &ReadVariable;
 				}
 
-				const FDataInterfaceSourceEmitterNamePair* DataInterfaceAssignment = VariableAssignmentMap.Find(BoundVariable);
-				if (DataInterfaceAssignment != nullptr)
+				const FDataInterfaceSourceEmitterNamePair* DataInterfaceAssignment = VariableAssignmentMap.Find(*BoundVariable);
+				if (DataInterfaceAssignment != nullptr && ParameterStoreVariables.Contains(ReadVariable) == false)
 				{
-					ResolvedDataInterface.ResolvedVariable = BoundVariable;
-					ResolvedDataInterface.SourceVariable = ReadVariable;
+					ResolvedDataInterface.ResolvedVariable = *BoundVariable;
+					ResolvedDataInterface.ParameterStoreVariable = ReadVariable;
 					ResolvedDataInterface.bIsInternal = false;
 					ResolvedDataInterface.ResolvedDataInterface = DataInterfaceAssignment->DataInterface;
 					ResolvedDataInterface.ResolvedSourceEmitterName = DataInterfaceAssignment->SourceEmitterName;
+					ParameterStoreVariables.Add(ResolvedDataInterface.ParameterStoreVariable);
 				}
 
-				if (FNiagaraUserRedirectionParameterStore::IsUserParameter(BoundVariable))
+				if (FNiagaraUserRedirectionParameterStore::IsUserParameter(*BoundVariable))
 				{
 					const FNiagaraParameterStore& InstanceParameters = System->GetSystemCompiledData().InstanceParamStore;
-					const int32* UserParameterIndex = InstanceParameters.FindParameterOffset(BoundVariable);
+					const int32* UserParameterIndex = InstanceParameters.FindParameterOffset(*BoundVariable);
 					if (UserParameterIndex != nullptr && *UserParameterIndex != INDEX_NONE)
 					{
 						UserDataInterfaceBindings.Add(FNiagaraResolvedUserDataInterfaceBinding(*UserParameterIndex, ResolvedDataInterfaceIndex));
 					}
 				}
 			}
+			else if (CachedDefaultDataInterface.RegisteredParameterMapWrite != NAME_None)
+			{
+				FNiagaraVariable WriteVariable(CachedDefaultDataInterface.Type, CachedDefaultDataInterface.RegisteredParameterMapWrite);
+				const FDataInterfaceSourceEmitterNamePair* DataInterfaceAssignment = VariableAssignmentMap.Find(WriteVariable);
+				if (DataInterfaceAssignment != nullptr && ParameterStoreVariables.Contains(WriteVariable) == false)
+				{
+					ResolvedDataInterface.ResolvedVariable = WriteVariable;
+					ResolvedDataInterface.ParameterStoreVariable = WriteVariable;
+					ResolvedDataInterface.bIsInternal = false;
+					ResolvedDataInterface.ResolvedDataInterface = DataInterfaceAssignment->DataInterface;
+					ResolvedDataInterface.ResolvedSourceEmitterName = DataInterfaceAssignment->SourceEmitterName;
+					ParameterStoreVariables.Add(ResolvedDataInterface.ParameterStoreVariable);
+				}
+			}
 
 			if (ResolvedDataInterface.ResolvedDataInterface == nullptr)
 			{
@@ -252,7 +281,7 @@ namespace FNiagaraResolveDIHelpers
 
 				FNiagaraVariable InternalVariable(CachedDefaultDataInterface.Type, FName(NameBuilder.ToString()));
 				ResolvedDataInterface.ResolvedVariable = InternalVariable;
-				ResolvedDataInterface.SourceVariable = InternalVariable;
+				ResolvedDataInterface.ParameterStoreVariable = InternalVariable;
 				ResolvedDataInterface.bIsInternal = true;
 				ResolvedDataInterface.ResolvedDataInterface = CachedDefaultDataInterface.DataInterface;
 				ResolvedDataInterface.ResolvedSourceEmitterName = CachedDefaultDataInterface.SourceEmitterName;
@@ -268,16 +297,13 @@ namespace FNiagaraResolveDIHelpers
 	{
 		for (const FNiagaraScriptResolvedDataInterfaceInfo& SourceResolvedDataInterface : SourceResolvedDataInterfaces)
 		{
-			if (SourceResolvedDataInterface.bIsInternal)
+			for (FNiagaraScriptResolvedDataInterfaceInfo& TargetResolvedDataInterface : TargetResolvedDataInterfaces)
 			{
-				for (FNiagaraScriptResolvedDataInterfaceInfo& TargetResolvedDataInterface : TargetResolvedDataInterfaces)
+				if (TargetResolvedDataInterface.bIsInternal && TargetResolvedDataInterface.Name == SourceResolvedDataInterface.Name)
 				{
-					if (TargetResolvedDataInterface.bIsInternal && TargetResolvedDataInterface.ResolvedVariable == SourceResolvedDataInterface.ResolvedVariable)
-					{
-						TargetResolvedDataInterface.ResolvedDataInterface = SourceResolvedDataInterface.ResolvedDataInterface;
-						TargetResolvedDataInterface.ResolvedSourceEmitterName = SourceResolvedDataInterface.ResolvedSourceEmitterName;
-						break;
-					}
+					TargetResolvedDataInterface.ResolvedDataInterface = SourceResolvedDataInterface.ResolvedDataInterface;
+					TargetResolvedDataInterface.ResolvedSourceEmitterName = SourceResolvedDataInterface.ResolvedSourceEmitterName;
+					break;
 				}
 			}
 		}
@@ -295,21 +321,6 @@ namespace FNiagaraResolveDIHelpers
 		ParticleSpawnScript->SetResolvedDataInterfaces(ResolvedSpawnDataInterfaces);
 	}
 
-	/** Handles the special case where internal data interfaces defined in particle spawn or particle update need to be copied to the gpu script so that they're using the same instance. */
-	void ResolveInternalDIsForGpuScripts(UNiagaraScript* ParticleSpawnScript, UNiagaraScript* ParticleUpdateScript, UNiagaraScript* ParticleGpuScript)
-	{
-		TArray<FNiagaraScriptResolvedDataInterfaceInfo> ResolvedGpuDataInterfaces;
-		ResolvedGpuDataInterfaces.Append(ParticleGpuScript->GetResolvedDataInterfaces());
-
-		TArrayView<const FNiagaraScriptResolvedDataInterfaceInfo> ResolvedSpawnDataInterfaces = ParticleSpawnScript->GetResolvedDataInterfaces();
-		SynchronizeMatchingInternalResolvedDataInterfaces(ResolvedSpawnDataInterfaces, ResolvedGpuDataInterfaces);
-
-		TArrayView<const FNiagaraScriptResolvedDataInterfaceInfo> ResolvedUpdateDataInterfaces = ParticleUpdateScript->GetResolvedDataInterfaces();
-		SynchronizeMatchingInternalResolvedDataInterfaces(ResolvedUpdateDataInterfaces, ResolvedGpuDataInterfaces);
-
-		ParticleGpuScript->SetResolvedDataInterfaces(ResolvedGpuDataInterfaces);
-	}
-
 	void ResolveDIsInternal(
 		UNiagaraSystem* System,
 		const TMap<FGuid, TMap<FNiagaraVariableBase, FDataInterfaceSourceEmitterNamePair>>& EmitterIdToVariableAssignmentsMap,
@@ -318,8 +329,10 @@ namespace FNiagaraResolveDIHelpers
 	{
 		TMap<FNiagaraVariableBase, FDataInterfaceSourceEmitterNamePair> VariableAssignmentMap = EmitterIdToVariableAssignmentsMap[FGuid()];
 		TMap<FNiagaraVariableBase, FNiagaraVariableBase> VariableBindingMap = EmitterIdToVariableBindingsMap[FGuid()];
-		ResolveDIsForScript(System, System->GetSystemSpawnScript(), FString(), VariableAssignmentMap, VariableBindingMap, OutErrorMessages);
-		ResolveDIsForScript(System, System->GetSystemUpdateScript(), FString(), VariableAssignmentMap, VariableBindingMap, OutErrorMessages);
+		TMap<FNiagaraVariableBase, FNiagaraVariableBase> CompactedVariableBindingMap;
+		CompactBindings(VariableBindingMap, CompactedVariableBindingMap, OutErrorMessages);
+		ResolveDIsForScript(System, System->GetSystemSpawnScript(), FString(), VariableAssignmentMap, CompactedVariableBindingMap, OutErrorMessages);
+		ResolveDIsForScript(System, System->GetSystemUpdateScript(), FString(), VariableAssignmentMap, CompactedVariableBindingMap, OutErrorMessages);
 
 		for (const FNiagaraEmitterHandle& EmitterHandle : System->GetEmitterHandles())
 		{
@@ -336,24 +349,20 @@ namespace FNiagaraResolveDIHelpers
 
 			TMap<FNiagaraVariableBase, FDataInterfaceSourceEmitterNamePair> EmitterVariableAssignmentMap = EmitterIdToVariableAssignmentsMap[EmitterHandle.GetId()];
 			TMap<FNiagaraVariableBase, FNiagaraVariableBase> EmitterVariableBindingMap = EmitterIdToVariableBindingsMap[EmitterHandle.GetId()];
-
-			ResolveDIsForScript(System, VersionedNiagaraEmitterData->SpawnScriptProps.Script, EmitterHandle.GetUniqueInstanceName(), EmitterVariableAssignmentMap, EmitterVariableBindingMap, OutErrorMessages);
-			ResolveDIsForScript(System, VersionedNiagaraEmitterData->UpdateScriptProps.Script, EmitterHandle.GetUniqueInstanceName(), EmitterVariableAssignmentMap, EmitterVariableBindingMap, OutErrorMessages);
-			ResolveDIsForScript(System, VersionedNiagaraEmitterData->GetGPUComputeScript(), EmitterHandle.GetUniqueInstanceName(), EmitterVariableAssignmentMap, EmitterVariableBindingMap, OutErrorMessages);
+			TMap<FNiagaraVariableBase, FNiagaraVariableBase> EmitterCompactedVariableBindingMap;
+			CompactBindings(EmitterVariableBindingMap, EmitterCompactedVariableBindingMap, OutErrorMessages);
+			ResolveDIsForScript(System, VersionedNiagaraEmitterData->SpawnScriptProps.Script, EmitterHandle.GetUniqueInstanceName(), EmitterVariableAssignmentMap, EmitterCompactedVariableBindingMap, OutErrorMessages);
+			ResolveDIsForScript(System, VersionedNiagaraEmitterData->UpdateScriptProps.Script, EmitterHandle.GetUniqueInstanceName(), EmitterVariableAssignmentMap, EmitterCompactedVariableBindingMap, OutErrorMessages);
+			ResolveDIsForScript(System, VersionedNiagaraEmitterData->GetGPUComputeScript(), EmitterHandle.GetUniqueInstanceName(), EmitterVariableAssignmentMap, EmitterCompactedVariableBindingMap, OutErrorMessages);
 
 			if (VersionedNiagaraEmitterData->bInterpolatedSpawning)
 			{
 				ResolveInternalDIsForInterpolatedSpawning(VersionedNiagaraEmitterData->SpawnScriptProps.Script, VersionedNiagaraEmitterData->UpdateScriptProps.Script);
 			}
 
-			if (VersionedNiagaraEmitterData->SimTarget == ENiagaraSimTarget::GPUComputeSim)
-			{
-				ResolveInternalDIsForGpuScripts(VersionedNiagaraEmitterData->SpawnScriptProps.Script, VersionedNiagaraEmitterData->EmitterUpdateScriptProps.Script, VersionedNiagaraEmitterData->GetGPUComputeScript());
-			}
-
 			for (const FNiagaraEventScriptProperties& EventHandler : VersionedNiagaraEmitterData->GetEventHandlers())
 			{
-				ResolveDIsForScript(System, EventHandler.Script, EmitterHandle.GetUniqueInstanceName(), EmitterVariableAssignmentMap, EmitterVariableBindingMap, OutErrorMessages);
+				ResolveDIsForScript(System, EventHandler.Script, EmitterHandle.GetUniqueInstanceName(), EmitterVariableAssignmentMap, EmitterCompactedVariableBindingMap, OutErrorMessages);
 			}
 
 			for (const UNiagaraSimulationStageBase* SimulationStage : VersionedNiagaraEmitterData->GetSimulationStages())
@@ -362,7 +371,7 @@ namespace FNiagaraResolveDIHelpers
 				{
 					continue;
 				}
-				ResolveDIsForScript(System, SimulationStage->Script, EmitterHandle.GetUniqueInstanceName(), EmitterVariableAssignmentMap, EmitterVariableBindingMap, OutErrorMessages);
+				ResolveDIsForScript(System, SimulationStage->Script, EmitterHandle.GetUniqueInstanceName(), EmitterVariableAssignmentMap, EmitterCompactedVariableBindingMap, OutErrorMessages);
 			}
 		}
 	}
@@ -371,9 +380,8 @@ namespace FNiagaraResolveDIHelpers
 	{
 		TMap<FGuid, TMap<FNiagaraVariableBase, FDataInterfaceSourceEmitterNamePair>> EmitterIdToVariableAssignmentsMap;
 		TMap<FGuid, TMap<FNiagaraVariableBase, FNiagaraVariableBase>> EmitterIdToVariableBindingsMap;
-		TArray<FText> ErrorMessages;
-		FNiagaraResolveDIHelpers::CollectDIBindingsAndAssignments(System, EmitterIdToVariableAssignmentsMap, EmitterIdToVariableBindingsMap, ErrorMessages);
-		FNiagaraResolveDIHelpers::ResolveDIsInternal(System, EmitterIdToVariableAssignmentsMap, EmitterIdToVariableBindingsMap, ErrorMessages);
+		FNiagaraResolveDIHelpers::CollectDIBindingsAndAssignments(System, EmitterIdToVariableAssignmentsMap, EmitterIdToVariableBindingsMap, OutErrorMessages);
+		FNiagaraResolveDIHelpers::ResolveDIsInternal(System, EmitterIdToVariableAssignmentsMap, EmitterIdToVariableBindingsMap, OutErrorMessages);
 	}
 }
 
diff --git a/Engine/Plugins/FX/Niagara/Source/Niagara/Private/NiagaraScript.cpp b/Engine/Plugins/FX/Niagara/Source/Niagara/Private/NiagaraScript.cpp
--- a/Engine/Plugins/FX/Niagara/Source/Niagara/Private/NiagaraScript.cpp
+++ b/Engine/Plugins/FX/Niagara/Source/Niagara/Private/NiagaraScript.cpp
@@ -1541,7 +1541,7 @@ const FNiagaraScriptExecutionParameterStore* UNiagaraScript::GetExecutionReadyPa
 			}
 			return &ScriptExecutionParamStoreCPU;
 		}
-		else if (SimTarget == ENiagaraSimTarget::GPUComputeSim)
+		else if (SimTarget == ENiagaraSimTarget::GPUComputeSim && IsReadyToRun(ENiagaraSimTarget::GPUComputeSim))
 		{
 			if (!ScriptExecutionParamStoreGPU.bInitialized)
 			{
@@ -2141,7 +2141,10 @@ void UNiagaraScript::PostLoad()
 
 	for (FNiagaraScriptResolvedDataInterfaceInfo& ResolvedDataInterface : ResolvedDataInterfaces)
 	{
-		ResolvedDataInterface.ResolvedDataInterface->ConditionalPostLoad();
+		if (ResolvedDataInterface.ResolvedDataInterface != nullptr)
+		{
+			ResolvedDataInterface.ResolvedDataInterface->ConditionalPostLoad();
+		}
 	}
 
 #if WITH_EDITORONLY_DATA
@@ -2843,10 +2846,7 @@ void UNiagaraScript::SetVMCompilationResults(const FNiagaraVMExecutableDataId& I
 		CachedDefaultDataInterfaces[Idx].Type = Info.Type;
 		CachedDefaultDataInterfaces[Idx].SourceEmitterName = Info.SourceEmitterName;
 		CachedDefaultDataInterfaces[Idx].RegisteredParameterMapRead = ResolveEmitterAlias(Info.RegisteredParameterMapRead, EmitterUniqueName);
-		for (const FName& RegisteredParameterMapWrite : Info.RegisteredParameterMapWrites)
-		{
-			CachedDefaultDataInterfaces[Idx].RegisteredParameterMapWrites.Add(ResolveEmitterAlias(RegisteredParameterMapWrite, EmitterUniqueName));
-		}
+		CachedDefaultDataInterfaces[Idx].RegisteredParameterMapWrite = ResolveEmitterAlias(Info.RegisteredParameterMapWrite, EmitterUniqueName);
 
 		// We compiled it just a bit ago, so we should be able to resolve it from the table that we passed in.
 		UNiagaraDataInterface* FindDIById = ResolveDataInterface(ObjectNameMap, CachedDefaultDataInterfaces[Idx].Name);
@@ -3651,10 +3651,7 @@ void UNiagaraScript::SyncAliases(const FNiagaraAliasContext& ResolveAliasesConte
 
 		// also update the MapRead/MapWrite member variables, they seem to typically match the Name parameter, but may be None
 		UpdateName(DataInterfaceInfo.RegisteredParameterMapRead);
-		for (FName& RegisteredParameterMapWrite : DataInterfaceInfo.RegisteredParameterMapWrites)
-		{
-			UpdateName(RegisteredParameterMapWrite);
-		}
+		UpdateName(DataInterfaceInfo.RegisteredParameterMapWrite);
 	}
 }
 
diff --git a/Engine/Plugins/FX/Niagara/Source/Niagara/Private/NiagaraScriptExecutionParameterStore.cpp b/Engine/Plugins/FX/Niagara/Source/Niagara/Private/NiagaraScriptExecutionParameterStore.cpp
--- a/Engine/Plugins/FX/Niagara/Source/Niagara/Private/NiagaraScriptExecutionParameterStore.cpp
+++ b/Engine/Plugins/FX/Niagara/Source/Niagara/Private/NiagaraScriptExecutionParameterStore.cpp
@@ -314,7 +314,7 @@ void FNiagaraScriptExecutionParameterStore::AddScriptParams(UNiagaraScript* Scri
 	for (const FNiagaraScriptResolvedDataInterfaceInfo& ResolvedDataInterface : Script->GetResolvedDataInterfaces())
 	{
 		int32 VarOffset = INDEX_NONE;
-		bAdded |= AddParameter(ResolvedDataInterface.SourceVariable, false, false, &VarOffset);
+		bAdded |= AddParameter(ResolvedDataInterface.ParameterStoreVariable, false, false, &VarOffset);
 		SetDataInterface(ResolvedDataInterface.ResolvedDataInterface, VarOffset);
 	}
 
diff --git a/Engine/Plugins/FX/Niagara/Source/Niagara/Private/NiagaraSystem.cpp b/Engine/Plugins/FX/Niagara/Source/Niagara/Private/NiagaraSystem.cpp
--- a/Engine/Plugins/FX/Niagara/Source/Niagara/Private/NiagaraSystem.cpp
+++ b/Engine/Plugins/FX/Niagara/Source/Niagara/Private/NiagaraSystem.cpp
@@ -1571,7 +1571,7 @@ void UNiagaraSystem::FindDataInterfaceDependencies(FVersionedNiagaraEmitterData*
 					{
 						for (const FNiagaraScriptResolvedDataInterfaceInfo& DataInterfaceInfo : Script->GetResolvedDataInterfaces())
 						{
-							if (DataInterfaceInfo.ResolvedVariable == Variable)
+							if (DataInterfaceInfo.ParameterStoreVariable == Variable)
 							{
 								return DataInterfaceInfo.ResolvedDataInterface;
 							}
@@ -1599,7 +1599,10 @@ void UNiagaraSystem::FindDataInterfaceDependencies(FVersionedNiagaraEmitterData*
 						continue;
 					}
 
-					StoreDataInterfaces[Variable.Offset]->GetEmitterDependencies(this, Dependencies);
+					if (StoreDataInterfaces[Variable.Offset] != nullptr)
+					{
+						StoreDataInterfaces[Variable.Offset]->GetEmitterDependencies(this, Dependencies);
+					}
 				}
 			}
 		}
@@ -1901,7 +1904,7 @@ void UNiagaraSystem::CacheFromCompiledData()
 						{
 							if (DataInterfaceInfo.bIsInternal == false)
 							{
-								DataInterfaceGpuUsage.Add(DataInterfaceInfo.ResolvedVariable.GetName());
+								DataInterfaceGpuUsage.Add(DataInterfaceInfo.ParameterStoreVariable.GetName());
 							}
 						}
 					}
@@ -1937,7 +1940,7 @@ void UNiagaraSystem::CacheFromCompiledData()
 		{
 			if (UNiagaraDataInterface* DataInterface = DataInterfaceInfo.ResolvedDataInterface)
 			{
-				const bool bUsedByGPU = DataInterfaceGpuUsage.Contains(DataInterfaceInfo.ResolvedVariable.GetName());
+				const bool bUsedByGPU = DataInterfaceGpuUsage.Contains(DataInterfaceInfo.ParameterStoreVariable.GetName());
 				DataInterface->CacheStaticBuffers(*StaticBuffers.Get(), DataInterfaceInfo.ResolvedVariable, true, bUsedByGPU);
 			}
 		}
diff --git a/Engine/Plugins/FX/Niagara/Source/Niagara/Public/NiagaraCommon.h b/Engine/Plugins/FX/Niagara/Source/Niagara/Public/NiagaraCommon.h
--- a/Engine/Plugins/FX/Niagara/Source/Niagara/Public/NiagaraCommon.h
+++ b/Engine/Plugins/FX/Niagara/Source/Niagara/Public/NiagaraCommon.h
@@ -700,7 +700,7 @@ public:
 	FName RegisteredParameterMapRead;
 
 	UPROPERTY()
-	TArray<FName> RegisteredParameterMapWrites;
+	FName RegisteredParameterMapWrite;
 
 	UPROPERTY()
 	FString SourceEmitterName;
@@ -754,7 +754,7 @@ struct NIAGARA_API FNiagaraScriptResolvedDataInterfaceInfo
 	FNiagaraVariableBase ResolvedVariable;
 
 	UPROPERTY()
-	FNiagaraVariableBase SourceVariable;
+	FNiagaraVariableBase ParameterStoreVariable;
 
 	UPROPERTY()
 	bool bIsInternal;
@@ -801,7 +801,7 @@ public:
 	FName RegisteredParameterMapRead;
 
 	UPROPERTY()
-	TArray<FName> RegisteredParameterMapWrites;
+	FName RegisteredParameterMapWrite;
 
 	UPROPERTY()
 	bool bIsPlaceholder;
diff --git a/Engine/Plugins/FX/Niagara/Source/NiagaraEditor/Private/Customizations/NiagaraTypeCustomizations.cpp b/Engine/Plugins/FX/Niagara/Source/NiagaraEditor/Private/Customizations/NiagaraTypeCustomizations.cpp
--- a/Engine/Plugins/FX/Niagara/Source/NiagaraEditor/Private/Customizations/NiagaraTypeCustomizations.cpp
+++ b/Engine/Plugins/FX/Niagara/Source/NiagaraEditor/Private/Customizations/NiagaraTypeCustomizations.cpp
@@ -975,7 +975,7 @@ TArray<TPair<FNiagaraVariableBase, FNiagaraVariableBase> > FNiagaraMaterialAttri
 					TArray<FNiagaraScriptDataInterfaceInfo>& CachedDIs = Script->GetCachedDefaultDataInterfaces();
 					for (const FNiagaraScriptDataInterfaceInfo& Info : CachedDIs)
 					{
-						if (Info.RegisteredParameterMapWrites.Contains(VariableName))
+						if (Info.RegisteredParameterMapWrite == VariableName)
 						{
 							return Info.DataInterface;
 						}
diff --git a/Engine/Plugins/FX/Niagara/Source/NiagaraEditor/Private/NiagaraHlslTranslator.cpp b/Engine/Plugins/FX/Niagara/Source/NiagaraEditor/Private/NiagaraHlslTranslator.cpp
--- a/Engine/Plugins/FX/Niagara/Source/NiagaraEditor/Private/NiagaraHlslTranslator.cpp
+++ b/Engine/Plugins/FX/Niagara/Source/NiagaraEditor/Private/NiagaraHlslTranslator.cpp
@@ -2303,7 +2303,7 @@ void FHlslNiagaraTranslator::HandleSimStageSetupAndTeardown(int32 InWhichStage,
 		// If it wasn't previously added, let's go ahead and do so. Maybe they are solely using the StackContext namespace.
 		if (DataInterfaceOwnerIndex == INDEX_NONE && bNeedsDIOwner)
 		{
-			DataInterfaceOwnerIndex = RegisterDataInterface(IterationSourceVar, CDO, true, false);
+			DataInterfaceOwnerIndex = RegisterDataInterface(IterationSourceVar, CDO, true, true);
 		}
 
 		// If we haven't created it by now, bail out. 
@@ -5493,7 +5493,25 @@ void FHlslNiagaraTranslator::ParameterMapSet(UNiagaraNodeParameterMapSet* SetNod
 					}
 
 					FNiagaraScriptDataInterfaceCompileInfo& Info = CompilationOutput.ScriptData.DataInterfaceInfo[Input];
-					Info.RegisteredParameterMapWrites.Add(UsageName);
+					if (Info.RegisteredParameterMapWrite == NAME_None)
+					{
+						Info.RegisteredParameterMapWrite = UsageName;
+					}
+					else
+					{
+						// This is a special case where a single data interface is written to multiple parameters.  In order to handle this correctly the data interface
+						// needs it's own registration slot so that it ends up in the parameter store under the correct name so that it can be looked up correctly in various places
+						// and can be bound properly to the renderer parameter stores.
+						UNiagaraDataInterface* CDO = CompileDuplicateData->GetDuplicatedDataInterfaceCDOForClass(const_cast<UClass*>(Var.GetType().GetClass()));
+						if (CDO)
+						{
+							FString PlaceholderName = FString::Printf(TEXT("Placeholder.Placeholder%i"), CompilationOutput.ScriptData.DataInterfaceInfo.Num());
+							FNiagaraVariable PlaceholderVar(Var.GetType(), *PlaceholderName);
+							int32 PlaceholderIndex = RegisterDataInterface(PlaceholderVar, CDO, true, false);
+							CompilationOutput.ScriptData.DataInterfaceInfo[PlaceholderIndex].RegisteredParameterMapRead = CompilationOutput.ScriptData.DataInterfaceInfo[Input].RegisteredParameterMapWrite;
+							CompilationOutput.ScriptData.DataInterfaceInfo[PlaceholderIndex].RegisteredParameterMapWrite = UsageName;
+						}
+					}
 				}
 			}
 			else
@@ -6941,9 +6959,17 @@ int32 FHlslNiagaraTranslator::RegisterDataInterface(FNiagaraVariable& Var, UNiag
 		CompilationOutput.ScriptData.bNeedsGPUContextInit = true;
 	}
 
+	FNiagaraVariable ReadVariable;
+	if (bAddParameterMapRead)
+	{
+		ReadVariable = FNiagaraParameterMapHistory::IsAliasedEmitterParameter(Var.GetName().ToString())
+			? ActiveHistoryForFunctionCalls.ResolveAliases(Var) 
+			: Var;
+	}
+
 	int32 Idx = CompilationOutput.ScriptData.DataInterfaceInfo.IndexOfByPredicate([&](const FNiagaraScriptDataInterfaceCompileInfo& OtherInfo)
 	{
-		return OtherInfo.Name == DataInterfaceName;
+		return OtherInfo.Name == DataInterfaceName || (ReadVariable.IsValid() && OtherInfo.RegisteredParameterMapWrite == ReadVariable.GetName());
 	});
 
 	if (Idx == INDEX_NONE)
@@ -6969,27 +6995,17 @@ int32 FHlslNiagaraTranslator::RegisterDataInterface(FNiagaraVariable& Var, UNiag
 		{
 			CompilationOutput.ScriptData.DataInterfaceInfo[Idx].UserPtrIdx = CompilationOutput.ScriptData.NumUserPtrs++;
 		}
+
+		if (bAddParameterMapRead)
+		{
+			CompilationOutput.ScriptData.DataInterfaceInfo[Idx].RegisteredParameterMapRead = ReadVariable.GetName();
+		}
 	}
 	else
 	{
 		check(CompilationOutput.ScriptData.DataInterfaceInfo[Idx].Type == Var.GetType());
 	}
 
-	if (bAddParameterMapRead)
-	{
-		FName UsageName;
-		if (FNiagaraParameterMapHistory::IsAliasedEmitterParameter(Var.GetName().ToString()))
-		{
-			FNiagaraVariable AliasedVar = ActiveHistoryForFunctionCalls.ResolveAliases(Var);
-			UsageName = AliasedVar.GetName();
-		}
-		else
-		{
-			UsageName = Var.GetName();
-		}
-		CompilationOutput.ScriptData.DataInterfaceInfo[Idx].RegisteredParameterMapRead = UsageName;
-	}
-
 	return Idx;
 }
 
diff --git a/Engine/Plugins/FX/Niagara/Source/NiagaraEditorWidgets/Private/DetailCustomizations/NiagaraDataInterfaceGrid2DCollectionDetails.cpp b/Engine/Plugins/FX/Niagara/Source/NiagaraEditorWidgets/Private/DetailCustomizations/NiagaraDataInterfaceGrid2DCollectionDetails.cpp
--- a/Engine/Plugins/FX/Niagara/Source/NiagaraEditorWidgets/Private/DetailCustomizations/NiagaraDataInterfaceGrid2DCollectionDetails.cpp
+++ b/Engine/Plugins/FX/Niagara/Source/NiagaraEditorWidgets/Private/DetailCustomizations/NiagaraDataInterfaceGrid2DCollectionDetails.cpp
@@ -11,6 +11,7 @@
 #include "NiagaraConstants.h"
 #include "NiagaraNodeInput.h"
 #include "SNiagaraNamePropertySelector.h"
+#include "NiagaraEditorUtilities.h"
 
 #define LOCTEXT_NAMESPACE "FNiagaraDataInterfaceGrid2DCollectionDetails"
 
@@ -122,7 +123,16 @@ void FNiagaraDataInterfaceGrid2DCollectionDetails::GeneratePreviewAttributes(TAr
 			TArray<FNiagaraVariableBase> FoundVariables;
 			TArray<uint32> FoundVariableOffsets;
 			int32 FoundNumAttribChannelsFound;
-			Grid2DInterface->FindAttributesByName(VariableName, FoundVariables, FoundVariableOffsets, FoundNumAttribChannelsFound);
+
+			// Finding attributes requires that we use the runtime instance of the data interface, so get that here.
+			UNiagaraSystem* OwningSystem = Grid2DInterface->GetTypedOuter<UNiagaraSystem>();
+			UNiagaraDataInterfaceGrid2DCollection* RuntimeGrid2DInterface = OwningSystem != nullptr 
+				? Cast<UNiagaraDataInterfaceGrid2DCollection>(FNiagaraEditorUtilities::GetResolvedRuntimeInstanceForEditorDataInterfaceInstance(*OwningSystem, *Grid2DInterface))
+				: nullptr;
+			if (RuntimeGrid2DInterface != nullptr)
+			{
+				RuntimeGrid2DInterface->FindAttributes(FoundVariables, FoundVariableOffsets, FoundNumAttribChannelsFound);
+			}
 
 			for (const FNiagaraVariableBase& Variable : FoundVariables)
 			{
diff --git a/Engine/Plugins/FX/Niagara/Source/NiagaraEditorWidgets/Private/DetailCustomizations/NiagaraDataInterfaceGrid3DCollectionDetails.cpp b/Engine/Plugins/FX/Niagara/Source/NiagaraEditorWidgets/Private/DetailCustomizations/NiagaraDataInterfaceGrid3DCollectionDetails.cpp
--- a/Engine/Plugins/FX/Niagara/Source/NiagaraEditorWidgets/Private/DetailCustomizations/NiagaraDataInterfaceGrid3DCollectionDetails.cpp
+++ b/Engine/Plugins/FX/Niagara/Source/NiagaraEditorWidgets/Private/DetailCustomizations/NiagaraDataInterfaceGrid3DCollectionDetails.cpp
@@ -11,6 +11,7 @@
 #include "NiagaraConstants.h"
 #include "NiagaraNodeInput.h"
 #include "SNiagaraNamePropertySelector.h"
+#include "NiagaraEditorUtilities.h"
 
 #define LOCTEXT_NAMESPACE "FNiagaraDataInterfaceGrid3DCollectionDetails"
 
@@ -122,7 +123,16 @@ void FNiagaraDataInterfaceGrid3DCollectionDetails::GeneratePreviewAttributes(TAr
 			TArray<FNiagaraVariableBase> FoundVariables;
 			TArray<uint32> FoundVariableOffsets;
 			int32 FoundNumAttribChannelsFound;
-			Grid3DInterface->FindAttributesByName(VariableName, FoundVariables, FoundVariableOffsets, FoundNumAttribChannelsFound);
+
+			// Finding attributes requires that we use the runtime instance of the data interface, so get that here.
+			UNiagaraSystem* OwningSystem = Grid3DInterface->GetTypedOuter<UNiagaraSystem>();
+			UNiagaraDataInterfaceGrid3DCollection* RuntimeGrid3DInterface = OwningSystem != nullptr
+				? Cast<UNiagaraDataInterfaceGrid3DCollection>(FNiagaraEditorUtilities::GetResolvedRuntimeInstanceForEditorDataInterfaceInstance(*OwningSystem, *Grid3DInterface))
+				: nullptr;
+			if (RuntimeGrid3DInterface != nullptr)
+			{
+				RuntimeGrid3DInterface->FindAttributes(FoundVariables, FoundVariableOffsets, FoundNumAttribChannelsFound);
+			}
 
 			for (const FNiagaraVariableBase& Variable : FoundVariables)
 			{
diff --git a/Engine/Source/Runtime/Core/Private/UObject/FortniteMainBranchObjectVersion.cpp b/Engine/Source/Runtime/Core/Private/UObject/FortniteMainBranchObjectVersion.cpp
--- a/Engine/Source/Runtime/Core/Private/UObject/FortniteMainBranchObjectVersion.cpp
+++ b/Engine/Source/Runtime/Core/Private/UObject/FortniteMainBranchObjectVersion.cpp
@@ -11,7 +11,7 @@ TMap<FGuid, FGuid> FFortniteMainBranchObjectVersion::GetSystemGuids()
 	SystemGuids.Add(DevGuids.MATERIALSHADERMAP_DERIVEDDATA_VER, FGuid("CE7776E63E9E43FBB8C103EE1AB40B5D"));
 	SystemGuids.Add(DevGuids.NANITE_DERIVEDDATA_VER, FGuid("D28C003F825F4C178F698938FB0E47F7"));
 	SystemGuids.Add(DevGuids.NIAGARASHADERMAP_DERIVEDDATA_VER, FGuid("35D7960B72164B49B9B7B3FB5857E674"));
-	SystemGuids.Add(DevGuids.Niagara_LatestScriptCompileVersion, FGuid("3A499F83F8274468B50A487A3349C2BB"));
+	SystemGuids.Add(DevGuids.Niagara_LatestScriptCompileVersion, FGuid("BB3424A1A79F42F58D08412E7DAC6E57"));
 	SystemGuids.Add(DevGuids.SkeletalMeshDerivedDataVersion, FGuid("87620CAAE42ACAA2D32409D2AD939201"));
 	SystemGuids.Add(DevGuids.STATICMESH_DERIVEDDATA_VER, FGuid("543843B103794E0E9BA4BE80FB602F79"));
 	return SystemGuids;
