- Fix for Grid customization as categories were changed but customization was not updated
- Customize override format to show what the current format is when not overridden

diff --git a/Engine/Plugins/FX/Niagara/Source/NiagaraEditorWidgets/Private/DetailCustomizations/NiagaraDataInterfaceGrid2DCollectionDetails.cpp b/Engine/Plugins/FX/Niagara/Source/NiagaraEditorWidgets/Private/DetailCustomizations/NiagaraDataInterfaceGrid2DCollectionDetails.cpp
--- a/Engine/Plugins/FX/Niagara/Source/NiagaraEditorWidgets/Private/DetailCustomizations/NiagaraDataInterfaceGrid2DCollectionDetails.cpp
+++ b/Engine/Plugins/FX/Niagara/Source/NiagaraEditorWidgets/Private/DetailCustomizations/NiagaraDataInterfaceGrid2DCollectionDetails.cpp
@@ -10,6 +10,7 @@
 #include "NiagaraComponent.h"
 #include "NiagaraConstants.h"
 #include "NiagaraNodeInput.h"
+#include "NiagaraSettings.h"
 #include "SNiagaraNamePropertySelector.h"
 
 #define LOCTEXT_NAMESPACE "FNiagaraDataInterfaceGrid2DCollectionDetails"
@@ -20,9 +21,7 @@ FNiagaraDataInterfaceGrid2DCollectionDetails::~FNiagaraDataInterfaceGrid2DCollec
 
 void FNiagaraDataInterfaceGrid2DCollectionDetails::CustomizeDetails(IDetailLayoutBuilder& DetailBuilder)
 {
-	static const FName Grid2DCollectionCategoryName = TEXT("Grid2DCollection");
-
-	LayoutBuilder = &DetailBuilder;
+	static const FName GridCollectionCategoryName = TEXT("Grid");
 
 	TArray<TWeakObjectPtr<UObject>> SelectedObjects;
 	DetailBuilder.GetObjectsBeingCustomized(SelectedObjects);
@@ -37,24 +36,22 @@ void FNiagaraDataInterfaceGrid2DCollectionDetails::CustomizeDetails(IDetailLayou
 	Grid2DInterface->OnChanged().RemoveAll(this);
 	Grid2DInterface->OnChanged().AddSP(this, &FNiagaraDataInterfaceGrid2DCollectionDetails::OnInterfaceChanged);
 
-	Grid2DCollectionCategory = &DetailBuilder.EditCategory(Grid2DCollectionCategoryName, LOCTEXT("Grid2DCollectionCat", "Grid2DCollection"));
+	IDetailCategoryBuilder* GridCollectionCategory = &DetailBuilder.EditCategory(GridCollectionCategoryName);
 	{
 		TArray<TSharedRef<IPropertyHandle>> Properties;
-		Grid2DCollectionCategory->GetDefaultProperties(Properties, true, true);
-
-		TSharedPtr<IPropertyHandle> PreviewAttributeProperty = DetailBuilder.GetProperty(GET_MEMBER_NAME_CHECKED(UNiagaraDataInterfaceGrid2DCollection, PreviewAttribute));
+		GridCollectionCategory->GetDefaultProperties(Properties, true, true);
 
 		for (TSharedPtr<IPropertyHandle> Property : Properties)
 		{
 			FProperty* PropertyPtr = Property->GetProperty();
 
-			if (PropertyPtr == PreviewAttributeProperty->GetProperty())
+			if (PropertyPtr->GetFName() == GET_MEMBER_NAME_CHECKED(UNiagaraDataInterfaceGrid2DCollection, PreviewAttribute))
 			{
 				TArray<TSharedPtr<FName>> PossibleNames;
 				GeneratePreviewAttributes(PossibleNames);
 				PreviewAttributesBuilder = SNew(SNiagaraNamePropertySelector, Property.ToSharedRef(), PossibleNames);
 
-				IDetailPropertyRow& PropertyRow = Grid2DCollectionCategory->AddProperty(Property);
+				IDetailPropertyRow& PropertyRow = GridCollectionCategory->AddProperty(Property);
 				PropertyRow.CustomWidget(false)
 					.NameContent()
 					[
@@ -66,9 +63,36 @@ void FNiagaraDataInterfaceGrid2DCollectionDetails::CustomizeDetails(IDetailLayou
 						PreviewAttributesBuilder.ToSharedRef()
 					];
 			}
+			else if (PropertyPtr->GetFName() == GET_MEMBER_NAME_CHECKED(UNiagaraDataInterfaceGrid2DCollection, OverrideBufferFormat))
+			{
+				IDetailPropertyRow& PropertyRow = GridCollectionCategory->AddProperty(Property);
+				PropertyRow.CustomWidget(false)
+				.NameContent()
+				[
+					Property->CreatePropertyNameWidget()
+				]
+				.ValueContent()
+				[
+					SNew(SHorizontalBox)
+					+ SHorizontalBox::Slot()
+					[
+						SNew(STextBlock)
+						.Visibility(this, &FNiagaraDataInterfaceGrid2DCollectionDetails::IsOverideFormatInvisibile)
+						.Text(this, &FNiagaraDataInterfaceGrid2DCollectionDetails::GetDefaultFormatText)
+					]
+					+ SHorizontalBox::Slot()
+					[
+						SNew(SBox)
+						.Visibility(this, &FNiagaraDataInterfaceGrid2DCollectionDetails::IsOverideFormatVisibile)
+						[
+							Property->CreatePropertyValueWidget()
+						]
+					]
+				];
+			}
 			else
 			{
-				Grid2DCollectionCategory->AddProperty(Property);
+				GridCollectionCategory->AddProperty(Property);
 			}
 		}
 	}
@@ -139,4 +163,24 @@ void FNiagaraDataInterfaceGrid2DCollectionDetails::GeneratePreviewAttributes(TAr
 	}
 }
 
+EVisibility FNiagaraDataInterfaceGrid2DCollectionDetails::IsOverideFormatVisibile() const
+{
+	UNiagaraDataInterfaceGrid2DCollection* Grid2D = Grid2DInterfacePtr.Get();
+	return Grid2D && Grid2D->bOverrideFormat ? EVisibility::Visible : EVisibility::Collapsed;
+}
+
+EVisibility FNiagaraDataInterfaceGrid2DCollectionDetails::IsOverideFormatInvisibile() const
+{
+	UNiagaraDataInterfaceGrid2DCollection* Grid2D = Grid2DInterfacePtr.Get();
+	return Grid2D && !Grid2D->bOverrideFormat ? EVisibility::Visible : EVisibility::Collapsed;
+}
+
+FText FNiagaraDataInterfaceGrid2DCollectionDetails::GetDefaultFormatText() const
+{
+	return FText::Format(
+		LOCTEXT("DefaultGridFormatFormat", "Using Project Default Format - {0}"),
+		StaticEnum<ENiagaraGpuBufferFormat>()->GetDisplayValueAsText(GetDefault<UNiagaraSettings>()->DefaultGridFormat)
+	);
+}
+
 #undef LOCTEXT_NAMESPACE
diff --git a/Engine/Plugins/FX/Niagara/Source/NiagaraEditorWidgets/Private/DetailCustomizations/NiagaraDataInterfaceGrid2DCollectionDetails.h b/Engine/Plugins/FX/Niagara/Source/NiagaraEditorWidgets/Private/DetailCustomizations/NiagaraDataInterfaceGrid2DCollectionDetails.h
--- a/Engine/Plugins/FX/Niagara/Source/NiagaraEditorWidgets/Private/DetailCustomizations/NiagaraDataInterfaceGrid2DCollectionDetails.h
+++ b/Engine/Plugins/FX/Niagara/Source/NiagaraEditorWidgets/Private/DetailCustomizations/NiagaraDataInterfaceGrid2DCollectionDetails.h
@@ -29,9 +29,11 @@ private:
 	void OnDataChanged();
 	void GeneratePreviewAttributes(TArray<TSharedPtr<FName>>& SourceArray);
 
+	EVisibility IsOverideFormatVisibile() const;
+	EVisibility IsOverideFormatInvisibile() const;
+	FText GetDefaultFormatText() const;
+
 private:
 	TWeakObjectPtr<UNiagaraDataInterfaceGrid2DCollection> Grid2DInterfacePtr;
 	TSharedPtr<SNiagaraNamePropertySelector> PreviewAttributesBuilder;
-	IDetailLayoutBuilder* LayoutBuilder = nullptr;
-	IDetailCategoryBuilder* Grid2DCollectionCategory = nullptr;
 };
diff --git a/Engine/Plugins/FX/Niagara/Source/NiagaraEditorWidgets/Private/DetailCustomizations/NiagaraDataInterfaceGrid3DCollectionDetails.cpp b/Engine/Plugins/FX/Niagara/Source/NiagaraEditorWidgets/Private/DetailCustomizations/NiagaraDataInterfaceGrid3DCollectionDetails.cpp
--- a/Engine/Plugins/FX/Niagara/Source/NiagaraEditorWidgets/Private/DetailCustomizations/NiagaraDataInterfaceGrid3DCollectionDetails.cpp
+++ b/Engine/Plugins/FX/Niagara/Source/NiagaraEditorWidgets/Private/DetailCustomizations/NiagaraDataInterfaceGrid3DCollectionDetails.cpp
@@ -10,6 +10,7 @@
 #include "NiagaraComponent.h"
 #include "NiagaraConstants.h"
 #include "NiagaraNodeInput.h"
+#include "NiagaraSettings.h"
 #include "SNiagaraNamePropertySelector.h"
 
 #define LOCTEXT_NAMESPACE "FNiagaraDataInterfaceGrid3DCollectionDetails"
@@ -20,9 +21,7 @@ FNiagaraDataInterfaceGrid3DCollectionDetails::~FNiagaraDataInterfaceGrid3DCollec
 
 void FNiagaraDataInterfaceGrid3DCollectionDetails::CustomizeDetails(IDetailLayoutBuilder& DetailBuilder)
 {
-	static const FName Grid3DCollectionCategoryName = TEXT("Grid3DCollection");
-
-	LayoutBuilder = &DetailBuilder;
+	static const FName GridCollectionCategoryName = TEXT("Grid");
 
 	TArray<TWeakObjectPtr<UObject>> SelectedObjects;
 	DetailBuilder.GetObjectsBeingCustomized(SelectedObjects);
@@ -37,24 +36,22 @@ void FNiagaraDataInterfaceGrid3DCollectionDetails::CustomizeDetails(IDetailLayou
 	Grid3DInterface->OnChanged().RemoveAll(this);
 	Grid3DInterface->OnChanged().AddSP(this, &FNiagaraDataInterfaceGrid3DCollectionDetails::OnInterfaceChanged);
 
-	Grid3DCollectionCategory = &DetailBuilder.EditCategory(Grid3DCollectionCategoryName, LOCTEXT("Grid3DCollectionCat", "Grid3DCollection"));
+	IDetailCategoryBuilder* GridCollectionCategory = &DetailBuilder.EditCategory(GridCollectionCategoryName);
 	{
 		TArray<TSharedRef<IPropertyHandle>> Properties;
-		Grid3DCollectionCategory->GetDefaultProperties(Properties, true, true);
-
-		TSharedPtr<IPropertyHandle> PreviewAttributeProperty = DetailBuilder.GetProperty(GET_MEMBER_NAME_CHECKED(UNiagaraDataInterfaceGrid3DCollection, PreviewAttribute));
+		GridCollectionCategory->GetDefaultProperties(Properties, true, true);
 
 		for (TSharedPtr<IPropertyHandle> Property : Properties)
 		{
 			FProperty* PropertyPtr = Property->GetProperty();
 
-			if (PropertyPtr == PreviewAttributeProperty->GetProperty())
+			if (PropertyPtr->GetFName() == GET_MEMBER_NAME_CHECKED(UNiagaraDataInterfaceGrid3DCollection, PreviewAttribute))
 			{
 				TArray<TSharedPtr<FName>> PossibleNames;
 				GeneratePreviewAttributes(PossibleNames);
 				PreviewAttributesBuilder = SNew(SNiagaraNamePropertySelector, Property.ToSharedRef(), PossibleNames);
 
-				IDetailPropertyRow& PropertyRow = Grid3DCollectionCategory->AddProperty(Property);
+				IDetailPropertyRow& PropertyRow = GridCollectionCategory->AddProperty(Property);
 				PropertyRow.CustomWidget(false)
 					.NameContent()
 					[
@@ -66,9 +63,36 @@ void FNiagaraDataInterfaceGrid3DCollectionDetails::CustomizeDetails(IDetailLayou
 						PreviewAttributesBuilder.ToSharedRef()
 					];
 			}
+			else if (PropertyPtr->GetFName() == GET_MEMBER_NAME_CHECKED(UNiagaraDataInterfaceGrid3DCollection, OverrideBufferFormat))
+			{
+				IDetailPropertyRow& PropertyRow = GridCollectionCategory->AddProperty(Property);
+				PropertyRow.CustomWidget(false)
+				.NameContent()
+				[
+					Property->CreatePropertyNameWidget()
+				]
+				.ValueContent()
+				[
+					SNew(SHorizontalBox)
+					+ SHorizontalBox::Slot()
+					[
+						SNew(STextBlock)
+						.Visibility(this, &FNiagaraDataInterfaceGrid3DCollectionDetails::IsOverideFormatInvisibile)
+						.Text(this, &FNiagaraDataInterfaceGrid3DCollectionDetails::GetDefaultFormatText)
+					]
+					+ SHorizontalBox::Slot()
+					[
+						SNew(SBox)
+						.Visibility(this, &FNiagaraDataInterfaceGrid3DCollectionDetails::IsOverideFormatVisibile)
+						[
+							Property->CreatePropertyValueWidget()
+						]
+					]
+				];
+			}
 			else
 			{
-				Grid3DCollectionCategory->AddProperty(Property);
+				GridCollectionCategory->AddProperty(Property);
 			}
 		}
 	}
@@ -139,4 +163,24 @@ void FNiagaraDataInterfaceGrid3DCollectionDetails::GeneratePreviewAttributes(TAr
 	}
 }
 
+EVisibility FNiagaraDataInterfaceGrid3DCollectionDetails::IsOverideFormatVisibile() const
+{
+	UNiagaraDataInterfaceGrid3DCollection* Grid3D = Grid3DInterfacePtr.Get();
+	return Grid3D && Grid3D->bOverrideFormat ? EVisibility::Visible : EVisibility::Collapsed;
+}
+
+EVisibility FNiagaraDataInterfaceGrid3DCollectionDetails::IsOverideFormatInvisibile() const
+{
+	UNiagaraDataInterfaceGrid3DCollection* Grid3D = Grid3DInterfacePtr.Get();
+	return Grid3D && !Grid3D->bOverrideFormat ? EVisibility::Visible : EVisibility::Collapsed;
+}
+
+FText FNiagaraDataInterfaceGrid3DCollectionDetails::GetDefaultFormatText() const
+{
+	return FText::Format(
+		LOCTEXT("DefaultGridFormatFormat", "Using Project Default Format - {0}"),
+		StaticEnum<ENiagaraGpuBufferFormat>()->GetDisplayValueAsText(GetDefault<UNiagaraSettings>()->DefaultGridFormat)
+	);
+}
+
 #undef LOCTEXT_NAMESPACE
diff --git a/Engine/Plugins/FX/Niagara/Source/NiagaraEditorWidgets/Private/DetailCustomizations/NiagaraDataInterfaceGrid3DCollectionDetails.h b/Engine/Plugins/FX/Niagara/Source/NiagaraEditorWidgets/Private/DetailCustomizations/NiagaraDataInterfaceGrid3DCollectionDetails.h
--- a/Engine/Plugins/FX/Niagara/Source/NiagaraEditorWidgets/Private/DetailCustomizations/NiagaraDataInterfaceGrid3DCollectionDetails.h
+++ b/Engine/Plugins/FX/Niagara/Source/NiagaraEditorWidgets/Private/DetailCustomizations/NiagaraDataInterfaceGrid3DCollectionDetails.h
@@ -29,9 +29,11 @@ private:
 	void OnDataChanged();
 	void GeneratePreviewAttributes(TArray<TSharedPtr<FName>>& SourceArray);
 
+	EVisibility IsOverideFormatVisibile() const;
+	EVisibility IsOverideFormatInvisibile() const;
+	FText GetDefaultFormatText() const;
+
 private:
 	TWeakObjectPtr<UNiagaraDataInterfaceGrid3DCollection> Grid3DInterfacePtr;
 	TSharedPtr<SNiagaraNamePropertySelector> PreviewAttributesBuilder;
-	IDetailLayoutBuilder* LayoutBuilder = nullptr;
-	IDetailCategoryBuilder* Grid3DCollectionCategory = nullptr;
 };
