- Add support for none particle bindings when in particle mode to decal & light renderer

diff --git a/Engine/Plugins/FX/Niagara/Source/Niagara/Private/NiagaraDecalRendererProperties.cpp b/Engine/Plugins/FX/Niagara/Source/Niagara/Private/NiagaraDecalRendererProperties.cpp
--- a/Engine/Plugins/FX/Niagara/Source/Niagara/Private/NiagaraDecalRendererProperties.cpp
+++ b/Engine/Plugins/FX/Niagara/Source/Niagara/Private/NiagaraDecalRendererProperties.cpp
@@ -242,14 +242,14 @@ void UNiagaraDecalRendererProperties::CacheFromCompiledData(const FNiagaraDataSe
 	UpdateSourceModeDerivates(SourceMode);
 	UpdateMICs();
 
-	PositionDataSetAccessor.Init(CompiledData, PositionBinding.GetDataSetBindableVariable().GetName());
-	DecalOrientationDataSetAccessor.Init(CompiledData, DecalOrientationBinding.GetDataSetBindableVariable().GetName());
-	DecalSizeDataSetAccessor.Init(CompiledData, DecalSizeBinding.GetDataSetBindableVariable().GetName());
-	DecalFadeDataSetAccessor.Init(CompiledData, DecalFadeBinding.GetDataSetBindableVariable().GetName());
-	DecalSortOrderDataSetAccessor.Init(CompiledData, DecalSortOrderBinding.GetDataSetBindableVariable().GetName());
-	DecalColorDataSetAccessor.Init(CompiledData, DecalColorBinding.GetDataSetBindableVariable().GetName());
-	DecalVisibleAccessor.Init(CompiledData, DecalVisibleBinding.GetDataSetBindableVariable().GetName());
-	RendererVisibilityTagAccessor.Init(CompiledData, RendererVisibilityTagBinding.GetDataSetBindableVariable().GetName());
+	InitParticleDataSetAccessor(PositionDataSetAccessor, CompiledData, PositionBinding);
+	InitParticleDataSetAccessor(DecalOrientationDataSetAccessor, CompiledData, DecalOrientationBinding);
+	InitParticleDataSetAccessor(DecalSizeDataSetAccessor, CompiledData, DecalSizeBinding);
+	InitParticleDataSetAccessor(DecalFadeDataSetAccessor, CompiledData, DecalFadeBinding);
+	InitParticleDataSetAccessor(DecalSortOrderDataSetAccessor, CompiledData, DecalSortOrderBinding);
+	InitParticleDataSetAccessor(DecalColorDataSetAccessor, CompiledData, DecalColorBinding);
+	InitParticleDataSetAccessor(DecalVisibleAccessor, CompiledData, DecalVisibleBinding);
+	InitParticleDataSetAccessor(RendererVisibilityTagAccessor, CompiledData, RendererVisibilityTagBinding);
 }
 
 void UNiagaraDecalRendererProperties::UpdateSourceModeDerivates(ENiagaraRendererSourceDataMode InSourceMode, bool bFromPropertyEdit)
diff --git a/Engine/Plugins/FX/Niagara/Source/Niagara/Private/NiagaraLightRendererProperties.cpp b/Engine/Plugins/FX/Niagara/Source/Niagara/Private/NiagaraLightRendererProperties.cpp
--- a/Engine/Plugins/FX/Niagara/Source/Niagara/Private/NiagaraLightRendererProperties.cpp
+++ b/Engine/Plugins/FX/Niagara/Source/Niagara/Private/NiagaraLightRendererProperties.cpp
@@ -186,13 +186,13 @@ void UNiagaraLightRendererProperties::CacheFromCompiledData(const FNiagaraDataSe
 {
 	UpdateSourceModeDerivates(SourceMode);
 
-	PositionDataSetAccessor.Init(CompiledData, PositionBinding.GetDataSetBindableVariable().GetName());
-	ColorDataSetAccessor.Init(CompiledData, ColorBinding.GetDataSetBindableVariable().GetName());
-	RadiusDataSetAccessor.Init(CompiledData, RadiusBinding.GetDataSetBindableVariable().GetName());
-	ExponentDataSetAccessor.Init(CompiledData, LightExponentBinding.GetDataSetBindableVariable().GetName());
-	ScatteringDataSetAccessor.Init(CompiledData, VolumetricScatteringBinding.GetDataSetBindableVariable().GetName());
-	EnabledDataSetAccessor.Init(CompiledData, LightRenderingEnabledBinding.GetDataSetBindableVariable().GetName());
-	RendererVisibilityTagAccessor.Init(CompiledData, RendererVisibilityTagBinding.GetDataSetBindableVariable().GetName());
+	InitParticleDataSetAccessor(PositionDataSetAccessor, CompiledData, PositionBinding);
+	InitParticleDataSetAccessor(ColorDataSetAccessor, CompiledData, ColorBinding);
+	InitParticleDataSetAccessor(RadiusDataSetAccessor, CompiledData, RadiusBinding);
+	InitParticleDataSetAccessor(ExponentDataSetAccessor, CompiledData, LightExponentBinding);
+	InitParticleDataSetAccessor(ScatteringDataSetAccessor, CompiledData, VolumetricScatteringBinding);
+	InitParticleDataSetAccessor(EnabledDataSetAccessor, CompiledData, LightRenderingEnabledBinding);
+	InitParticleDataSetAccessor(RendererVisibilityTagAccessor, CompiledData, RendererVisibilityTagBinding);
 }
 
 void UNiagaraLightRendererProperties::UpdateSourceModeDerivates(ENiagaraRendererSourceDataMode InSourceMode, bool bFromPropertyEdit)
diff --git a/Engine/Plugins/FX/Niagara/Source/Niagara/Private/NiagaraRendererDecals.cpp b/Engine/Plugins/FX/Niagara/Source/Niagara/Private/NiagaraRendererDecals.cpp
--- a/Engine/Plugins/FX/Niagara/Source/Niagara/Private/NiagaraRendererDecals.cpp
+++ b/Engine/Plugins/FX/Niagara/Source/Niagara/Private/NiagaraRendererDecals.cpp
@@ -180,6 +180,17 @@ FNiagaraDynamicDataBase* FNiagaraRendererDecals::GenerateDynamicData(const FNiag
 	TArray<FDeferredDecalUpdateParams> DecalUpdates;
 	DecalUpdates.Reserve(DataToRender->GetNumInstances());
 
+	// Get default values which might be none particle parameters
+	const FNiagaraParameterStore& ParameterStore = Emitter->GetRendererBoundVariables();
+	const bool LocalDefaultVisible = ParameterStore.GetParameterValueOrDefault(RendererProperties->DecalVisibleBinding.GetParamMapBindableVariable(), DefaultVisible).GetValue();
+	const int32 LocalDefaultVisTag = ParameterStore.GetParameterValueOrDefault(RendererProperties->RendererVisibilityTagBinding.GetParamMapBindableVariable(), RendererProperties->RendererVisibility);
+	const FVector3f LocalDefaultSimPos = ParameterStore.GetParameterValueOrDefault(RendererProperties->PositionBinding.GetParamMapBindableVariable(), DefaultPos);
+	const FQuat4f LocalDefaultSimRot = ParameterStore.GetParameterValueOrDefault(RendererProperties->DecalOrientationBinding.GetParamMapBindableVariable(), DefaultRot).GetNormalized();
+	const FVector3f LocalDefaultSize = ParameterStore.GetParameterValueOrDefault(RendererProperties->DecalSizeBinding.GetParamMapBindableVariable(), DefaultSize);
+	const FLinearColor LocalDefaultDecalColor = ParameterStore.GetParameterValueOrDefault(RendererProperties->DecalColorBinding.GetParamMapBindableVariable(), FLinearColor::White);
+	const float LocalDefaultFade = ParameterStore.GetParameterValueOrDefault(RendererProperties->DecalFadeBinding.GetParamMapBindableVariable(), DefaultFade);
+	const int32 LocalDefaultSortOrder = ParameterStore.GetParameterValueOrDefault(RendererProperties->DecalSortOrderBinding.GetParamMapBindableVariable(), DefaultSortOrder);
+
 	// Particles Source mode?
 	if (RendererProperties->SourceMode == ENiagaraRendererSourceDataMode::Particles)
 	{
@@ -196,22 +207,22 @@ FNiagaraDynamicDataBase* FNiagaraRendererDecals::GenerateDynamicData(const FNiag
 		for (uint32 ParticleIndex = 0; ParticleIndex < DataToRender->GetNumInstances(); ++ParticleIndex)
 		{
 			// Check if the decal is visible
-			const bool bVisible = VisibleReader.GetSafe(ParticleIndex, DefaultVisible).GetValue();
-			const int32 VisTag = VisTagReader.GetSafe(ParticleIndex, RendererProperties->RendererVisibility);
+			const bool bVisible = VisibleReader.GetSafe(ParticleIndex, LocalDefaultVisible).GetValue();
+			const int32 VisTag = VisTagReader.GetSafe(ParticleIndex, LocalDefaultVisTag);
 			if (!bVisible || VisTag != RendererProperties->RendererVisibility)
 			{
 				continue;
 			}
 
 			// Grab Decal Attributes
-			const FVector3f SimPos = PositionReader.GetSafe(ParticleIndex, DefaultPos);
-			const FQuat4f SimRot = RotationReader.GetSafe(ParticleIndex, DefaultRot).GetNormalized();
+			const FVector3f SimPos = PositionReader.GetSafe(ParticleIndex, LocalDefaultSimPos);
+			const FQuat4f SimRot = RotationReader.GetSafe(ParticleIndex, LocalDefaultSimRot).GetNormalized();
 			const FVector Position = bUseLocalSpace ? LocalToWorld.TransformPosition(FVector(SimPos)) : LwcConverter.ConvertSimulationPositionToWorld(SimPos);
 			const FQuat Rotation = bUseLocalSpace ? LocalToWorld.TransformRotation(FQuat(SimRot)) : FQuat(SimRot);
-			const FVector Size = FVector(SizeReader.GetSafe(ParticleIndex, DefaultSize) * 0.5f);
-			const FLinearColor DecalColor = ColorReader.GetSafe(ParticleIndex, FLinearColor::White);
-			const float Fade = FadeReader.GetSafe(ParticleIndex, DefaultFade);
-			const int32 SortOrder = SortOrderReader.GetSafe(ParticleIndex, DefaultSortOrder);
+			const FVector Size = FVector(SizeReader.GetSafe(ParticleIndex, LocalDefaultSize) * 0.5f);
+			const FLinearColor DecalColor = ColorReader.GetSafe(ParticleIndex, LocalDefaultDecalColor);
+			const float Fade = FadeReader.GetSafe(ParticleIndex, LocalDefaultFade);
+			const int32 SortOrder = SortOrderReader.GetSafe(ParticleIndex, LocalDefaultSortOrder);
 
 			// Create Update Parameters
 			FDeferredDecalUpdateParams& UpdateParams = DecalUpdates.AddDefaulted_GetRef();
@@ -242,19 +253,16 @@ FNiagaraDynamicDataBase* FNiagaraRendererDecals::GenerateDynamicData(const FNiag
 	else
 	{
 		//-OPT: Emitter mode we should be able to cache the parameter offset rather than search each frame
-		const FNiagaraParameterStore& ParameterStore = Emitter->GetRendererBoundVariables();
-		const bool bVisible = ParameterStore.GetParameterValueOrDefault(RendererProperties->DecalVisibleBinding.GetParamMapBindableVariable(), DefaultVisible).GetValue();
-		const int32 VisTag = ParameterStore.GetParameterValueOrDefault(RendererProperties->RendererVisibilityTagBinding.GetParamMapBindableVariable(), RendererProperties->RendererVisibility);
-		if (bVisible && VisTag == RendererProperties->RendererVisibility)
+		if (LocalDefaultVisible && LocalDefaultVisTag == RendererProperties->RendererVisibility)
 		{
-			const FVector3f SimPos = ParameterStore.GetParameterValueOrDefault(RendererProperties->PositionBinding.GetParamMapBindableVariable(), DefaultPos);
-			const FQuat4f SimRot = ParameterStore.GetParameterValueOrDefault(RendererProperties->DecalOrientationBinding.GetParamMapBindableVariable(), DefaultRot).GetNormalized();
+			const FVector3f SimPos = LocalDefaultSimPos;
+			const FQuat4f SimRot = LocalDefaultSimRot;
 			const FVector Position = bUseLocalSpace ? LocalToWorld.TransformPosition(FVector(SimPos)) : LwcConverter.ConvertSimulationPositionToWorld(SimPos);
 			const FQuat Rotation = bUseLocalSpace ? LocalToWorld.TransformRotation(FQuat(SimRot)) : FQuat(SimRot);
-			const FVector Size = FVector(ParameterStore.GetParameterValueOrDefault(RendererProperties->DecalSizeBinding.GetParamMapBindableVariable(), DefaultSize) * 0.5f);
-			const FLinearColor DecalColor = ParameterStore.GetParameterValueOrDefault(RendererProperties->DecalColorBinding.GetParamMapBindableVariable(), FLinearColor::White);
-			const float Fade = ParameterStore.GetParameterValueOrDefault(RendererProperties->DecalFadeBinding.GetParamMapBindableVariable(), DefaultFade);
-			const int32 SortOrder = ParameterStore.GetParameterValueOrDefault(RendererProperties->DecalSortOrderBinding.GetParamMapBindableVariable(), DefaultSortOrder);
+			const FVector Size = FVector(LocalDefaultSize * 0.5f);
+			const FLinearColor DecalColor = LocalDefaultDecalColor;
+			const float Fade = LocalDefaultFade;
+			const int32 SortOrder = LocalDefaultSortOrder;
 
 			// Create Update Parameters
 			FDeferredDecalUpdateParams& UpdateParams = DecalUpdates.AddDefaulted_GetRef();
diff --git a/Engine/Plugins/FX/Niagara/Source/Niagara/Private/NiagaraRendererLights.cpp b/Engine/Plugins/FX/Niagara/Source/Niagara/Private/NiagaraRendererLights.cpp
--- a/Engine/Plugins/FX/Niagara/Source/Niagara/Private/NiagaraRendererLights.cpp
+++ b/Engine/Plugins/FX/Niagara/Source/Niagara/Private/NiagaraRendererLights.cpp
@@ -102,13 +102,15 @@ FNiagaraDynamicDataBase* FNiagaraRendererLights::GenerateDynamicData(const FNiag
 
 	const bool bUseLocalSpace = UseLocalSpace(Proxy);
 	FNiagaraLWCConverter LwcConverter = SystemInstance->GetLWCConverter(bUseLocalSpace);
-	const FLinearColor DefaultColor = Properties->ColorBinding.GetDefaultValue<FLinearColor>();
-	const FNiagaraPosition DefaultPos = bUseLocalSpace ? FVector::ZeroVector : LocalToWorld.GetLocation();
-	const float DefaultRadius = Properties->RadiusBinding.GetDefaultValue<float>();
-	const float DefaultScattering = Properties->VolumetricScatteringBinding.GetDefaultValue<float>();
-	const FNiagaraBool DefaultEnabled(true);
-	const int32 DefaultVisibilityTag(0);
-	const float DefaultExponent = Properties->DefaultExponent;
+
+	const FNiagaraParameterStore& ParameterStore = Emitter->GetRendererBoundVariables();
+	const FVector3f DefaultPos = ParameterStore.GetParameterValueOrDefault(Properties->PositionBinding.GetParamMapBindableVariable(), FVector3f::ZeroVector);
+	const FLinearColor DefaultColor = ParameterStore.GetParameterValueOrDefault(Properties->ColorBinding.GetParamMapBindableVariable(), Properties->ColorBinding.GetDefaultValue<FLinearColor>());
+	const float DefaultRadius = ParameterStore.GetParameterValueOrDefault(Properties->RadiusBinding.GetParamMapBindableVariable(), Properties->RadiusBinding.GetDefaultValue<float>());
+	const float DefaultScattering = ParameterStore.GetParameterValueOrDefault(Properties->VolumetricScatteringBinding.GetParamMapBindableVariable(), Properties->VolumetricScatteringBinding.GetDefaultValue<float>());
+	const FNiagaraBool DefaultEnabled = ParameterStore.GetParameterValueOrDefault(Properties->LightRenderingEnabledBinding.GetParamMapBindableVariable(), FNiagaraBool(true));
+	const int32 DefaultVisibilityTag = ParameterStore.GetParameterValueOrDefault(Properties->RendererVisibilityTagBinding.GetParamMapBindableVariable(), Properties->RendererVisibility);
+	const float DefaultExponent = ParameterStore.GetParameterValueOrDefault(Properties->LightExponentBinding.GetParamMapBindableVariable(), Properties->DefaultExponent);
 
 	const float InverseExposureBlend = Properties->bOverrideInverseExposureBlend ? Properties->InverseExposureBlend : GetDefault<UNiagaraSettings>()->DefaultLightInverseExposureBlend;
 
@@ -127,18 +129,16 @@ FNiagaraDynamicDataBase* FNiagaraRendererLights::GenerateDynamicData(const FNiag
 
 		for (uint32 ParticleIndex = 0; ParticleIndex < DataToRender->GetNumInstances(); ParticleIndex++)
 		{
-			bool bShouldRenderParticleLight = EnabledReader.GetSafe(ParticleIndex, DefaultEnabled).GetValue();
-			if (bShouldRenderParticleLight && VisTagReader.IsValid())
-			{
-				bShouldRenderParticleLight = VisTagReader.GetSafe(ParticleIndex, DefaultVisibilityTag) == Properties->RendererVisibility;
-			}
-			float LightRadius = RadiusReader.GetSafe(ParticleIndex, DefaultRadius) * Properties->RadiusScale;
-			if (bShouldRenderParticleLight && LightRadius > 0)
+			const int32 VisTag = VisTagReader.GetSafe(ParticleIndex, DefaultVisibilityTag);
+			const bool bShouldRenderParticleLight = EnabledReader.GetSafe(ParticleIndex, DefaultEnabled).GetValue() && (VisTag == Properties->RendererVisibility);
+			const float LightRadius = RadiusReader.GetSafe(ParticleIndex, DefaultRadius) * Properties->RadiusScale;
+			if (bShouldRenderParticleLight && (LightRadius > 0.0f))
 			{
 				SimpleLightData& LightData = DynamicData->LightArray.AddDefaulted_GetRef();
 
 				const FLinearColor Color = ColorReader.GetSafe(ParticleIndex, DefaultColor);
 				const float Brightness = Properties->bAlphaScalesBrightness ? Color.A : 1.0f;
+				const FVector3f SimPos = PositionReader.GetSafe(ParticleIndex, DefaultPos);
 
 				LightData.LightEntry.Radius = LightRadius;
 				LightData.LightEntry.Color = FVector3f(Color) * Brightness + Properties->ColorAdd;
@@ -146,26 +146,21 @@ FNiagaraDynamicDataBase* FNiagaraRendererLights::GenerateDynamicData(const FNiag
 				LightData.LightEntry.InverseExposureBlend = InverseExposureBlend;
 				LightData.LightEntry.bAffectTranslucency = Properties->bAffectsTranslucency;
 				LightData.LightEntry.VolumetricScatteringIntensity = ScatteringReader.GetSafe(ParticleIndex, DefaultScattering);
-				LightData.PerViewEntry.Position = LwcConverter.ConvertSimulationPositionToWorld(PositionReader.GetSafe(ParticleIndex, DefaultPos));
-				if (bUseLocalSpace)
-				{
-					LightData.PerViewEntry.Position = LocalToWorld.TransformPosition(LightData.PerViewEntry.Position);
-				}
+				LightData.PerViewEntry.Position = bUseLocalSpace ? LocalToWorld.TransformPosition(FVector(SimPos)) : LwcConverter.ConvertSimulationPositionToWorld(SimPos);
 			}
 		}
 	}
 	else
 	{
-		const FNiagaraParameterStore& ParameterStore = Emitter->GetRendererBoundVariables();
-		const bool bEnabled = ParameterStore.GetParameterValueOrDefault(Properties->LightRenderingEnabledBinding.GetParamMapBindableVariable(), FNiagaraBool(true)).GetValue();
-		const int32 VisTag = ParameterStore.GetParameterValueOrDefault(Properties->RendererVisibilityTagBinding.GetParamMapBindableVariable(), Properties->RendererVisibility);
-		const float LightRadius = ParameterStore.GetParameterValueOrDefault(Properties->RadiusBinding.GetParamMapBindableVariable(), DefaultRadius) * Properties->RadiusScale;
+		const bool bEnabled = DefaultEnabled.GetValue();
+		const int32 VisTag = DefaultVisibilityTag;
+		const float LightRadius = DefaultRadius * Properties->RadiusScale;
 		if (bEnabled && VisTag == Properties->RendererVisibility && LightRadius > 0.0f)
 		{
-			const FVector3f SimPos = ParameterStore.GetParameterValueOrDefault(Properties->PositionBinding.GetParamMapBindableVariable(), DefaultPos);
-			const FLinearColor LightColor = ParameterStore.GetParameterValueOrDefault(Properties->ColorBinding.GetParamMapBindableVariable(), DefaultColor);
-			const float LightExponent = ParameterStore.GetParameterValueOrDefault(Properties->LightExponentBinding.GetParamMapBindableVariable(), DefaultExponent);
-			const float LightScattering = ParameterStore.GetParameterValueOrDefault(Properties->VolumetricScatteringBinding.GetParamMapBindableVariable(), DefaultScattering);
+			const FVector3f SimPos = DefaultPos;
+			const FLinearColor LightColor = DefaultColor;
+			const float LightExponent = DefaultExponent;
+			const float LightScattering = DefaultScattering;
 			const float Brightness = Properties->bAlphaScalesBrightness ? LightColor.A : 1.0f;
 
 			SimpleLightData& LightData = DynamicData->LightArray.AddDefaulted_GetRef();
@@ -175,11 +170,7 @@ FNiagaraDynamicDataBase* FNiagaraRendererLights::GenerateDynamicData(const FNiag
 			LightData.LightEntry.InverseExposureBlend = InverseExposureBlend;
 			LightData.LightEntry.bAffectTranslucency = Properties->bAffectsTranslucency;
 			LightData.LightEntry.VolumetricScatteringIntensity = LightScattering;
-			LightData.PerViewEntry.Position = LwcConverter.ConvertSimulationPositionToWorld(SimPos);
-			if (bUseLocalSpace)
-			{
-				LightData.PerViewEntry.Position = LocalToWorld.TransformPosition(LightData.PerViewEntry.Position);
-			}
+			LightData.PerViewEntry.Position = bUseLocalSpace ? LocalToWorld.TransformPosition(FVector(SimPos)) : LwcConverter.ConvertSimulationPositionToWorld(SimPos);
 		}
 	}
 
diff --git a/Engine/Plugins/FX/Niagara/Source/Niagara/Public/NiagaraRendererProperties.h b/Engine/Plugins/FX/Niagara/Source/Niagara/Public/NiagaraRendererProperties.h
--- a/Engine/Plugins/FX/Niagara/Source/Niagara/Public/NiagaraRendererProperties.h
+++ b/Engine/Plugins/FX/Niagara/Source/Niagara/Public/NiagaraRendererProperties.h
@@ -444,6 +444,12 @@ protected:
 #if WITH_EDITOR
 	FOnPropertiesChanged OnPropertiesChangedDelegate;
 #endif
+
+	template<typename TAccessorType>
+	void InitParticleDataSetAccessor(TAccessorType& Accessor, const FNiagaraDataSetCompiledData* CompiledData, const FNiagaraVariableAttributeBinding& Binding)
+	{
+		Accessor.Init(Binding.IsParticleBinding() ? CompiledData : nullptr, Binding.GetDataSetBindableVariable().GetName());
+	}
 };
 
 template<typename TAction>
