- Optimize calls generating the override parameters
- Bulk copy and just set existing back again, this avoids adding parameters one at a time which can result in a lot of allocations, searching, etc, proportial to the number of exposed parameters

#rnx

diff --git a/Engine/Plugins/FX/Niagara/Source/Niagara/Private/NiagaraComponent.cpp b/Engine/Plugins/FX/Niagara/Source/Niagara/Private/NiagaraComponent.cpp
--- a/Engine/Plugins/FX/Niagara/Source/Niagara/Private/NiagaraComponent.cpp
+++ b/Engine/Plugins/FX/Niagara/Source/Niagara/Private/NiagaraComponent.cpp
@@ -3257,31 +3257,50 @@ void UNiagaraComponent::FixDataInterfaceOuters()
 
 void UNiagaraComponent::CopyParametersFromAsset(bool bResetExistingOverrideParameters)
 {
-	TArray<FNiagaraVariable> ExistingVars;
-	OverrideParameters.GetParameters(ExistingVars);
+	// Copy Asset Override Parameters into Component Overrides
+	const FNiagaraUserRedirectionParameterStore& AssetExposedParameters = Asset->GetExposedParameters();
 
-	TArray<FNiagaraVariable> SourceVars;
-	Asset->GetExposedParameters().GetParameters(SourceVars);
+	FNiagaraUserRedirectionParameterStore ExistingOverrideParameters;
+	Swap(ExistingOverrideParameters, OverrideParameters);
+	OverrideParameters = AssetExposedParameters;
+	OverrideParameters.SetOwner(this);
 
-	// insert new asset parameters (and keep existing overrides)
-	for (FNiagaraVariable& Param : SourceVars)
+	// Copy Existing Variables into new parameter data
+	if (bResetExistingOverrideParameters == false)
 	{
-		bool bNewParam = OverrideParameters.AddParameter(Param, true);
-
-		bool bCopyAssetValue = bResetExistingOverrideParameters || bNewParam;
-		if (bCopyAssetValue || Param.IsDataInterface())
+		for (const FNiagaraVariableWithOffset& ExistingOverrideVar : ExistingOverrideParameters.ReadParameterVariables())
 		{
-			Asset->GetExposedParameters().CopyParameterData(OverrideParameters, Param);
+			OverrideParameters.AddParameter(ExistingOverrideVar, true);
+			if (ExistingOverrideVar.IsDataInterface())
+			{
+				UNiagaraDataInterface* ExistingDataInterface = ExistingOverrideParameters.GetDataInterface(ExistingOverrideVar.Offset);
+				OverrideParameters.SetDataInterface(ExistingDataInterface, ExistingOverrideVar);
+			}
+			else if (ExistingOverrideVar.IsUObject())
+			{
+				UObject* ExistingUObject = ExistingOverrideParameters.GetUObject(ExistingOverrideVar.Offset);
+				OverrideParameters.SetUObject(ExistingUObject, ExistingOverrideVar);
+			}
+			else
+			{
+				ExistingOverrideParameters.CopyParameterData(OverrideParameters, ExistingOverrideVar);
+			}
 		}
 	}
 
-	// remove parameters that don't exist in the source asset
-	for (FNiagaraVariable& ExistingVar : ExistingVars)
+	// DataInterfaces are unique per component so we need to make new instances if they are not already
+	for (int32 i = 0; i < AssetExposedParameters.GetDataInterfaces().Num(); ++i)
 	{
-		if (!SourceVars.Contains(ExistingVar))
+		UNiagaraDataInterface* AssetDataInterface = AssetExposedParameters.GetDataInterface(i);
+		UNiagaraDataInterface* OverrideDataInterface = OverrideParameters.GetDataInterface(i);
+		if (AssetDataInterface != OverrideDataInterface)
 		{
-			OverrideParameters.RemoveParameter(ExistingVar);
+			continue;
 		}
+
+		OverrideDataInterface = NewObject<UNiagaraDataInterface>(this, AssetDataInterface->GetClass(), NAME_None, RF_Transactional | RF_Public);
+		AssetDataInterface->CopyTo(OverrideDataInterface);
+		OverrideParameters.SetDataInterface(OverrideDataInterface, i);
 	}
 }
 
