Further abstract compilation out of NiagaraSystem

diff --git a/Engine/Plugins/FX/Niagara/Source/Niagara/Classes/NiagaraSystem.h b/Engine/Plugins/FX/Niagara/Source/Niagara/Classes/NiagaraSystem.h
--- a/Engine/Plugins/FX/Niagara/Source/Niagara/Classes/NiagaraSystem.h
+++ b/Engine/Plugins/FX/Niagara/Source/Niagara/Classes/NiagaraSystem.h
@@ -30,7 +30,7 @@
 
 #include "NiagaraSystem.generated.h"
 
-class FNiagaraAsyncCompileTask;
+class FNiagaraActiveCompilation;
 class UNiagaraEffectType;
 enum class ENiagaraCullProxyMode : uint32;
 enum class ERendererStencilMask : uint8;
@@ -171,41 +171,6 @@ struct FNiagaraSystemCompiledData
 	TArray<FNiagaraParameterDataSetBindingCollection> UpdateInstanceEmitterBindings;
 };
 
-USTRUCT()
-struct FNiagaraSystemCompileRequest
-{
-	GENERATED_USTRUCT_BODY()
-
-	double StartTime = 0.0;
-
-	UPROPERTY()
-	TArray<TObjectPtr<UObject>> RootObjects;
-
-	using FAsyncTaskPtr = TSharedPtr<FNiagaraAsyncCompileTask, ESPMode::ThreadSafe>;
-
-#if WITH_EDITORONLY_DATA
-	void Abort();
-	FAsyncTaskPtr* FindTask(const UNiagaraScript* Script);
-	bool QueryCompileComplete(UNiagaraSystem* OwningSystem, bool bWait, const double& MaxDuration);
-	bool Resolve(UNiagaraSystem* OwningSystem, FNiagaraParameterStore& ExposedParameters);
-	void Reset();
-	void Launch(UNiagaraSystem* OwningSystem, TConstArrayView<UNiagaraScript*> ScriptsNeedingCompile, TConstArrayView<FAsyncTaskPtr> Tasks);
-
-	bool bIsValid = true;
-	bool bForced = false;
-	bool bAllScriptsSynchronized = false;
-	bool bEvaluateParametersPending = false;
-
-	float CombinedCompileTime = 0.0f;
-
-private:
-	TArray<FAsyncTaskPtr> DDCTasks;
-
-	bool bDDCGetCompleted = false;
-
-#endif
-};
-
 struct FNiagaraEmitterExecutionIndex
 {
 	FNiagaraEmitterExecutionIndex() { bStartNewOverlapGroup = false; EmitterIndex = 0; }
@@ -805,15 +770,12 @@ private:
 
 	/** Since the shader compilation is done in another process, this is used to check if the result for any ongoing compilations is done.
 	*   If bWait is true then this *blocks* the game thread (and ui) until all shader compilations are finished.
+	*	Results from the compilation will be applied to the system unless there's already another compilation queued up
 	*/
-	bool QueryCompileComplete(bool bWait, bool bDoPost, bool bDoNotApply = false);
+	bool QueryCompileComplete(bool bWait);
 
 	void BroadcastOnSystemCompiled();
 
-	void PreProcessWaitingDDCTasks(bool bProcessForWait);
-
-	bool CompilationResultsValid(FNiagaraSystemCompileRequest& CompileRequest) const;
-
 	void EvaluateCompileResultDependencies() const;
 
 	void InitEmitterCompiledData();
@@ -868,8 +830,7 @@ protected:
 #if WITH_EDITORONLY_DATA
 	bool bCompilationReentrantGuard = false;
 
-	UPROPERTY(Transient)
-	TArray<FNiagaraSystemCompileRequest> ActiveCompilations;
+	TArray<TUniquePtr<FNiagaraActiveCompilation>> ActiveCompilations;
 
 	mutable TSharedPtr<FNiagaraGraphCachedDataBase, ESPMode::ThreadSafe> CachedTraversalData;
 #endif
diff --git a/Engine/Plugins/FX/Niagara/Source/Niagara/Private/NiagaraAsyncCompile.cpp b/Engine/Plugins/FX/Niagara/Source/Niagara/Private/NiagaraAsyncCompile.cpp
--- a/Engine/Plugins/FX/Niagara/Source/Niagara/Private/NiagaraAsyncCompile.cpp
+++ b/Engine/Plugins/FX/Niagara/Source/Niagara/Private/NiagaraAsyncCompile.cpp
@@ -129,7 +129,7 @@ void FNiagaraAsyncCompileTask::ProcessCurrentState()
 
 	if (CurrentState == ENiagaraCompilationState::CheckDDC)
 	{
-		// if FNiagaraSystemCompileRequest is going to handle evaluating the CheckDDC state, then skip doing anything here
+		// if FNiagaraActiveCompilationDefault is going to handle evaluating the CheckDDC state, then skip doing anything here
 		if (bCheckDDCEnabled)
 		{
 			// check if the DDC data is ready
@@ -319,7 +319,12 @@ bool FNiagaraAsyncCompileTask::AwaitResult()
 	UE_LOG(LogNiagara, Verbose, TEXT("Waiting on compilation to finish for %s!"), *AssetPath);
 
 	INiagaraModule& NiagaraModule = FModuleManager::Get().LoadModuleChecked<INiagaraModule>(TEXT("Niagara"));
-	ExeData = NiagaraModule.GetCompileJobResult(ScriptPair.PendingJobID, bWaitForCompileJob);
+
+	if (!ExeData.IsValid())
+	{
+		ExeData = NiagaraModule.GetCompileJobResult(ScriptPair.PendingJobID, bWaitForCompileJob);
+	}
+
 	return ExeData.IsValid();
 }
 
@@ -478,29 +483,109 @@ bool FNiagaraAsyncCompileTask::CompilationIdMatchesRequest() const
 	return IdMatches;
 }
 
-void FNiagaraSystemCompileRequest::Abort()
+void FNiagaraActiveCompilationDefault::AddReferencedObjects(FReferenceCollector& Collector)
 {
-	for (auto& AsyncTask : DDCTasks)
+	Collector.AddStableReferenceSet(&RootObjects);
+}
+
+void FNiagaraActiveCompilationDefault::Abort()
+{
+	for (auto& AsyncTask : Tasks)
 	{
 		AsyncTask->AbortTask();
 	}
 }
 
-FNiagaraSystemCompileRequest::FAsyncTaskPtr* FNiagaraSystemCompileRequest::FindTask(const UNiagaraScript* Script)
+bool FNiagaraActiveCompilationDefault::Validate(const FNiagaraQueryCompilationOptions& Options) const
+{
+	// for now the only thing we're concerned about is if we've got results for SystemSpawn and SystemUpdate scripts
+	// then we need to make sure that they agree in terms of the dataset attributes
+	const FAsyncTaskPtr* SpawnScriptRequest = FindTask(Options.System->GetSystemSpawnScript());
+	const FAsyncTaskPtr* UpdateScriptRequest = FindTask(Options.System->GetSystemUpdateScript());
+
+	const bool SpawnScriptValid = SpawnScriptRequest
+		&& SpawnScriptRequest->Get()->GetScriptPair().CompileResults.IsValid()
+		&& SpawnScriptRequest->Get()->GetScriptPair().CompileResults->LastCompileStatus != ENiagaraScriptCompileStatus::NCS_Error;
+
+	const bool UpdateScriptValid = UpdateScriptRequest
+		&& UpdateScriptRequest->Get()->GetScriptPair().CompileResults.IsValid()
+		&& UpdateScriptRequest->Get()->GetScriptPair().CompileResults->LastCompileStatus != ENiagaraScriptCompileStatus::NCS_Error;
+
+	if (SpawnScriptValid && UpdateScriptValid)
+	{
+		const FEmitterCompiledScriptPair& SpawnScriptPair = SpawnScriptRequest->Get()->GetScriptPair();
+		const FEmitterCompiledScriptPair& UpdateScriptPair = UpdateScriptRequest->Get()->GetScriptPair();
+		if (SpawnScriptPair.CompileResults->Attributes != UpdateScriptPair.CompileResults->Attributes)
+		{
+			// if we had requested a full rebuild, then we've got a case where the generated scripts are not compatible.  This indicates
+			// a significant issue where we're allowing graphs to generate invalid collections of scripts.  One known example is using
+			// the Script.Context static switch that isn't fully processed in all scripts, leading to attributes differing between the
+			// SystemSpawnScript and the SystemUpdateScript
+			if (bForced)
+			{
+				FString MissingAttributes;
+				FString AdditionalAttributes;
+
+				for (const auto& SpawnAttrib : SpawnScriptPair.CompileResults->Attributes)
+				{
+					if (!UpdateScriptPair.CompileResults->Attributes.Contains(SpawnAttrib))
+					{
+						MissingAttributes.Appendf(TEXT("%s%s"), MissingAttributes.Len() ? TEXT(", ") : TEXT(""), *SpawnAttrib.GetName().ToString());
+					}
+				}
+
+				for (const auto& UpdateAttrib : UpdateScriptPair.CompileResults->Attributes)
+				{
+					if (!SpawnScriptPair.CompileResults->Attributes.Contains(UpdateAttrib))
+					{
+						AdditionalAttributes.Appendf(TEXT("%s%s"), AdditionalAttributes.Len() ? TEXT(", ") : TEXT(""), *UpdateAttrib.GetName().ToString());
+					}
+				}
+
+				FNiagaraCompileEvent AttributeMismatchEvent(
+					FNiagaraCompileEventSeverity::Error,
+					FText::Format(LOCTEXT("SystemScriptAttributeMismatchError", "System Spawn/Update scripts have attributes which don't match!\n\tMissing update attributes: {0}\n\tAdditional update attributes: {1}"),
+						FText::FromString(MissingAttributes),
+						FText::FromString(AdditionalAttributes))
+					.ToString());
+
+				SpawnScriptPair.CompileResults->LastCompileStatus = ENiagaraScriptCompileStatus::NCS_Error;
+				SpawnScriptPair.CompileResults->LastCompileEvents.Add(AttributeMismatchEvent);
+			}
+			else
+			{
+				UE_LOG(LogNiagara, Log, TEXT("Failed to generate consistent results for System spawn and update scripts for system %s."), *Options.System->GetFullName());
+			}
+
+			return false;
+		}
+	}
+	return true;
+}
+
+FNiagaraActiveCompilationDefault::FAsyncTaskPtr* FNiagaraActiveCompilationDefault::FindTask(const UNiagaraScript* Script)
 {
-	return DDCTasks.FindByPredicate([Script](const FAsyncTaskPtr& TaskPtr) -> bool
+	return Tasks.FindByPredicate([Script](const FAsyncTaskPtr& TaskPtr) -> bool
 	{
 		return TaskPtr->GetScriptPair().CompiledScript == Script;
 	});
 }
 
-bool FNiagaraSystemCompileRequest::QueryCompileComplete(UNiagaraSystem* OwningSystem, bool bWait, const double& MaxDuration)
+const FNiagaraActiveCompilationDefault::FAsyncTaskPtr* FNiagaraActiveCompilationDefault::FindTask(const UNiagaraScript* Script) const
+{
+	return Tasks.FindByPredicate([Script](const FAsyncTaskPtr& TaskPtr) -> bool
+	{
+		return TaskPtr->GetScriptPair().CompiledScript == Script;
+	});
+}
+
+bool FNiagaraActiveCompilationDefault::QueryCompileComplete(const FNiagaraQueryCompilationOptions& Options)
 {
 	const double QueryStartTime = FPlatformTime::Seconds();
 
-	if (bWait)
+	if (Options.bWait)
 	{
-		for (FAsyncTaskPtr& AsyncTask : DDCTasks)
+		for (FAsyncTaskPtr& AsyncTask : Tasks)
 		{
 			AsyncTask->bWaitForCompileJob = true;
 		}
@@ -513,7 +598,7 @@ bool FNiagaraSystemCompileRequest::QueryCompileComplete(UNiagaraSystem* OwningSy
 			bDDCGetCompleted = true;
 			bool bAllTasksFinalized = true;
 			{
-				for (FAsyncTaskPtr& AsyncTask : DDCTasks)
+				for (FAsyncTaskPtr& AsyncTask : Tasks)
 				{
 					if (AsyncTask->CurrentState == ENiagaraCompilationState::CheckDDC)
 					{
@@ -535,7 +620,7 @@ bool FNiagaraSystemCompileRequest::QueryCompileComplete(UNiagaraSystem* OwningSy
 			{
 				// run through the active scripts and make sure that they are still up to date before we try to do
 				// any additional processing on them
-				for (FAsyncTaskPtr& AsyncTask : DDCTasks)
+				for (FAsyncTaskPtr& AsyncTask : Tasks)
 				{
 					if (AsyncTask->CurrentState == ENiagaraCompilationState::Finished)
 					{
@@ -551,16 +636,16 @@ bool FNiagaraSystemCompileRequest::QueryCompileComplete(UNiagaraSystem* OwningSy
 				// note - by having this deferred (after we've checked DDC) we do risk that the prepared RI parameters could be
 				// out of sync with what was there at the time the compile was requested.  So, we'll only take this path if we're
 				// not in edit mode.  In edit mode we'll prepare the RI parameters during UNiagaraSystem::RequestCompile()
-				OwningSystem->PrepareRapidIterationParametersForCompilation();
+				Options.System->PrepareRapidIterationParametersForCompilation();
 
 				// go through our tasks and evaluate the rapid iteration parameters.  We'll need to store them in the ExeData
 				// so that subsequent builds can take advantage of the work done in PrepareRapidIterationParamtersForCompilation
 				// by storing it in the ExeData.  Additionally if this evaluation has changed the rapid iteration parameters the
 				// DDC key will be different, and we'll want to push the data both to the old key and to the new key.
-				const int32 TaskCount = DDCTasks.Num();
+				const int32 TaskCount = Tasks.Num();
 				for (int32 TaskIt = 0; TaskIt < TaskCount; ++TaskIt)
 				{
-					FAsyncTaskPtr& AsyncTask = DDCTasks[TaskIt];
+					FAsyncTaskPtr& AsyncTask = Tasks[TaskIt];
 					if (AsyncTask->CurrentState == ENiagaraCompilationState::Finished)
 					{
 						continue;
@@ -599,11 +684,9 @@ bool FNiagaraSystemCompileRequest::QueryCompileComplete(UNiagaraSystem* OwningSy
 		}
 	}
 
-	bool bCompileComplete = true;
-
-	if (bWait)
+	if (Options.bWait)
 	{
-		for (FAsyncTaskPtr& AsyncTask : DDCTasks)
+		for (FAsyncTaskPtr& AsyncTask : Tasks)
 		{
 			// before we start to wait for the compile results, we start the compilation of all remaining tasks
 			while (!AsyncTask->IsDone() && AsyncTask->CurrentState < ENiagaraCompilationState::AwaitResult)
@@ -614,53 +697,48 @@ bool FNiagaraSystemCompileRequest::QueryCompileComplete(UNiagaraSystem* OwningSy
 	}
 
 	// Check to see if ALL of the sub-requests have resolved.
+	for (auto& AsyncTask : Tasks)
 	{
-		for (auto& AsyncTask : DDCTasks)
+		// if a task is very expensive we bail and continue on the next frame  
+		if (FPlatformTime::Seconds() - QueryStartTime < Options.MaxWaitDuration)
 		{
-			// if a task is very expensive we bail and continue on the next frame  
-			if (FPlatformTime::Seconds() - QueryStartTime < MaxDuration)
-			{
-				AsyncTask->ProcessCurrentState();
-			}
+			AsyncTask->ProcessCurrentState();
+		}
 
-			if (AsyncTask->IsDone())
-			{
-				continue;
-			}
+		if (AsyncTask->IsDone())
+		{
+			continue;
+		}
 
-			if (!AsyncTask->bFetchedGCObjects && AsyncTask->CurrentState > ENiagaraCompilationState::Precompile && AsyncTask->CurrentState <= ENiagaraCompilationState::ProcessResult)
-			{
-				RootObjects.Append(AsyncTask->PrecompileReference->CompilationRootObjects);
-				AsyncTask->bFetchedGCObjects = true;
-			}
+		if (!AsyncTask->bFetchedGCObjects && AsyncTask->CurrentState > ENiagaraCompilationState::Precompile && AsyncTask->CurrentState <= ENiagaraCompilationState::ProcessResult)
+		{
+			// note that PrecompileReference is shared across all tasks and so the contents within are going to
+			// be added multiple times (once per task).  Should pull this out of the task processing.
+			RootObjects.Append(AsyncTask->PrecompileReference->CompilationRootObjects);
+			AsyncTask->bFetchedGCObjects = true;
+		}
 
-			if (bWait)
-			{
-				AsyncTask->WaitAndResolveResult();
-			}
-			else
-			{
-				bCompileComplete = false;
-			}
+		if (Options.bWait)
+		{
+			AsyncTask->WaitAndResolveResult();
+		}
+		else
+		{
+			return false;
 		}
 	}
 
-	return bCompileComplete;
+	return true;
 }
 
-bool FNiagaraSystemCompileRequest::Resolve(UNiagaraSystem* OwningSystem, FNiagaraParameterStore& ExposedParameters)
+void FNiagaraActiveCompilationDefault::Apply(const FNiagaraQueryCompilationOptions& Options)
 {
-	bool bHasCompiledJobs = false;
-
-	for (FAsyncTaskPtr& AsyncTask : DDCTasks)
+	for (FAsyncTaskPtr& AsyncTask : Tasks)
 	{
-		bHasCompiledJobs |= AsyncTask->bUsedShaderCompilerWorker;
-
 		FEmitterCompiledScriptPair& EmitterCompiledScriptPair = AsyncTask->ScriptPair;
 		if (EmitterCompiledScriptPair.bResultsReady)
 		{
 			TSharedPtr<FNiagaraVMExecutableData> ExeData = EmitterCompiledScriptPair.CompileResults;
-			CombinedCompileTime += EmitterCompiledScriptPair.CompileTime;
 			UNiagaraScript* CompiledScript = EmitterCompiledScriptPair.CompiledScript;
 
 			// generate the ObjectNameMap from the source (from the duplicated data if available).
@@ -669,7 +747,7 @@ bool FNiagaraSystemCompileRequest::Resolve(UNiagaraSystem* OwningSystem, FNiagar
 			{
 				if (const UNiagaraScriptSourceBase* ScriptSource = AsyncTask->ComputedPrecompileDuplicateData->GetScriptSource())
 				{
-					ObjectNameMap = ScriptSource->ComputeObjectNameMap(*OwningSystem, CompiledScript->GetUsage(), CompiledScript->GetUsageId(), AsyncTask->UniqueEmitterName);
+					ObjectNameMap = ScriptSource->ComputeObjectNameMap(*Options.System, CompiledScript->GetUsage(), CompiledScript->GetUsageId(), AsyncTask->UniqueEmitterName);
 				}
 				else
 				{
@@ -693,12 +771,14 @@ bool FNiagaraSystemCompileRequest::Resolve(UNiagaraSystem* OwningSystem, FNiagar
 			else
 			{
 				const UNiagaraScriptSourceBase* ScriptSource = CompiledScript->GetLatestSource();
-				ObjectNameMap = ScriptSource->ComputeObjectNameMap(*OwningSystem, CompiledScript->GetUsage(), CompiledScript->GetUsageId(), AsyncTask->UniqueEmitterName);
+				ObjectNameMap = ScriptSource->ComputeObjectNameMap(*Options.System, CompiledScript->GetUsage(), CompiledScript->GetUsageId(), AsyncTask->UniqueEmitterName);
 			}
 
 			EmitterCompiledScriptPair.CompiledScript->SetVMCompilationResults(EmitterCompiledScriptPair.CompileId, *ExeData, AsyncTask->UniqueEmitterName, ObjectNameMap, bEvaluateParametersPending);
 
 			// Synchronize the variables that we actually encountered during precompile so that we can expose them to the end user.
+			FNiagaraUserRedirectionParameterStore& ExposedParameters = Options.System->GetExposedParameters();
+
 			TArray<FNiagaraVariable> OriginalExposedParams;
 			ExposedParameters.GetParameters(OriginalExposedParams);
 			TArray<FNiagaraVariable>& EncounteredExposedVars = AsyncTask->EncounteredExposedVars;
@@ -712,14 +792,12 @@ bool FNiagaraSystemCompileRequest::Resolve(UNiagaraSystem* OwningSystem, FNiagar
 			}
 		}
 	}
-
-	return bHasCompiledJobs;
 }
 
-void FNiagaraSystemCompileRequest::Reset()
+void FNiagaraActiveCompilationDefault::Reset()
 {
 	// clean up the precompile data
-	for (FAsyncTaskPtr& AsyncTask : DDCTasks)
+	for (FAsyncTaskPtr& AsyncTask : Tasks)
 	{
 		AsyncTask->ComputedPrecompileData.Reset();
 		if (AsyncTask->ComputedPrecompileDuplicateData.IsValid())
@@ -730,49 +808,184 @@ void FNiagaraSystemCompileRequest::Reset()
 	}
 }
 
-void FNiagaraSystemCompileRequest::Launch(UNiagaraSystem* OwningSystem, TConstArrayView<UNiagaraScript*> ScriptsNeedingCompile, TConstArrayView<FAsyncTaskPtr> Tasks)
+bool FNiagaraActiveCompilationDefault::Launch(const FNiagaraCompilationOptions& Options)
 {
-	bEvaluateParametersPending = !OwningSystem->ShouldUseRapidIterationParameters();
-	DDCTasks = Tasks;
+	bForced = Options.bForced;
+	StartTime = FPlatformTime::Seconds();
+
+	check(Options.System->GetSystemSpawnScript()->GetLatestSource() == Options.System->GetSystemUpdateScript()->GetLatestSource());
 
-	// prepare data for any precompile the ddc tasks need to do
-	TSharedPtr<FNiagaraLazyPrecompileReference, ESPMode::ThreadSafe> PrecompileReference = MakeShared<FNiagaraLazyPrecompileReference, ESPMode::ThreadSafe>();
-	PrecompileReference->System = OwningSystem;
-	PrecompileReference->Scripts = ScriptsNeedingCompile;
+	// when not compiling for edit mode we will defer preparing rapid iteration parameters till after we get the
+	// results from the DDC search (FNiagaraPrepareCompileTask::Tick())
+	const bool UseRapidIterationParameters = Options.System->ShouldUseRapidIterationParameters();
 
-	const TArray<FNiagaraEmitterHandle>& EmitterHandles = OwningSystem->GetEmitterHandles();
-	for (int32 i = 0; i < EmitterHandles.Num(); i++)
+	if (UseRapidIterationParameters)
 	{
-		FNiagaraEmitterHandle Handle = EmitterHandles[i];
-		if (Handle.GetIsEnabled() && Handle.GetEmitterData())
+		Options.System->PrepareRapidIterationParametersForCompilation();
+	}
+
+	TArray<UNiagaraScript*> ScriptsNeedingCompile;
+	bool bAnyCompiled = false;
+	bool bAnyUnsynchronized = false;
+
+	// Pass one... determine if any need to be compiled.
+	{
+		const TArray<FNiagaraEmitterHandle>& EmitterHandles = Options.System->GetEmitterHandles();
+		for (int32 i = 0; i < EmitterHandles.Num(); i++)
 		{
-			TArray<UNiagaraScript*> EmitterScripts;
-			Handle.GetEmitterData()->GetScripts(EmitterScripts, false, true);
-			check(EmitterScripts.Num() > 0);
-			for (UNiagaraScript* EmitterScript : EmitterScripts)
+			FNiagaraEmitterHandle Handle = EmitterHandles[i];
+			FVersionedNiagaraEmitterData* EmitterData = Handle.GetEmitterData();
+			if (EmitterData && Handle.GetIsEnabled())
+			{
+				TArray<UNiagaraScript*> EmitterScripts;
+				EmitterData->GetScripts(EmitterScripts, false, true);
+				check(EmitterScripts.Num() > 0);
+				for (UNiagaraScript* EmitterScript : EmitterScripts)
+				{
+					FEmitterCompiledScriptPair Pair;
+					Pair.VersionedEmitter = Handle.GetInstance();
+					Pair.CompiledScript = EmitterScript;
+
+					TSharedPtr<FNiagaraAsyncCompileTask, ESPMode::ThreadSafe> AsyncTask = MakeShared<FNiagaraAsyncCompileTask, ESPMode::ThreadSafe>(Options.System, EmitterScript->GetPathName(), Pair);
+					if (EmitterScript->IsCompilable() && !EmitterScript->AreScriptAndSourceSynchronized())
+					{
+						// we need to compute the vmID here to check later in the ddc task before doing the precompile if anything has changed in the meantime.
+						// the compilation id was just calculated during AreScriptAndSourceSynchronized() so we reuse it here
+						AsyncTask->AssignInitialCompilationId(EmitterScript->GetComputedVMCompilationId());
+
+						ScriptsNeedingCompile.Add(EmitterScript);
+						bAnyUnsynchronized = true;
+					}
+					else
+					{
+						AsyncTask->CurrentState = ENiagaraCompilationState::Finished;
+					}
+					AsyncTask->UniqueEmitterName = Handle.GetInstance().Emitter->GetUniqueEmitterName();
+					Tasks.Add(AsyncTask);
+				}
+			}
+		}
+
+		bAnyCompiled = bAnyUnsynchronized || Options.bForced;
+
+		// Now add the system scripts for compilation...
+		{
+			FEmitterCompiledScriptPair Pair;
+			Pair.VersionedEmitter = FVersionedNiagaraEmitter();
+			Pair.CompiledScript = Options.System->GetSystemSpawnScript();
+
+			TSharedPtr<FNiagaraAsyncCompileTask, ESPMode::ThreadSafe> AsyncTask = MakeShared<FNiagaraAsyncCompileTask, ESPMode::ThreadSafe>(Options.System, Pair.CompiledScript->GetPathName(), Pair);
+			if (!Pair.CompiledScript->AreScriptAndSourceSynchronized())
+			{
+				// we need to compute the vmID here to check later in the ddc task before doing the precompile if anything has changed in the meantime.
+				// the compilation id was just calculated during AreScriptAndSourceSynchronized() so we reuse it here
+				AsyncTask->AssignInitialCompilationId(Pair.CompiledScript->GetComputedVMCompilationId());
+
+				ScriptsNeedingCompile.Add(Pair.CompiledScript);
+				bAnyCompiled = true;
+			}
+			else
+			{
+				AsyncTask->CurrentState = ENiagaraCompilationState::Finished;
+			}
+			Tasks.Add(AsyncTask);
+		}
+
+		{
+			FEmitterCompiledScriptPair Pair;
+			Pair.VersionedEmitter = FVersionedNiagaraEmitter();
+			Pair.CompiledScript = Options.System->GetSystemUpdateScript();
+
+			TSharedPtr<FNiagaraAsyncCompileTask, ESPMode::ThreadSafe> AsyncTask = MakeShared<FNiagaraAsyncCompileTask, ESPMode::ThreadSafe>(Options.System, Pair.CompiledScript->GetPathName(), Pair);
+			if (!Pair.CompiledScript->AreScriptAndSourceSynchronized())
+			{
+				// we need to compute the vmID here to check later in the ddc task before doing the precompile if anything has changed in the meantime.
+				// the compilation id was just calculated during AreScriptAndSourceSynchronized() so we reuse it here
+				AsyncTask->AssignInitialCompilationId(Pair.CompiledScript->GetComputedVMCompilationId());
+
+				ScriptsNeedingCompile.Add(Pair.CompiledScript);
+				bAnyCompiled = true;
+			}
+			else
 			{
-				PrecompileReference->EmitterScriptIndex.Add(EmitterScript, i);
+				AsyncTask->CurrentState = ENiagaraCompilationState::Finished;
 			}
+			Tasks.Add(AsyncTask);
 		}
 	}
 
-	for (UNiagaraScript* Script : ScriptsNeedingCompile)
 	{
-		if (FAsyncTaskPtr* AsyncTaskPtr = FindTask(Script))
+		bEvaluateParametersPending = !UseRapidIterationParameters;
+
+		// prepare data for any precompile the ddc tasks need to do
+		TSharedPtr<FNiagaraLazyPrecompileReference, ESPMode::ThreadSafe> PrecompileReference = MakeShared<FNiagaraLazyPrecompileReference, ESPMode::ThreadSafe>();
+		PrecompileReference->System = Options.System;
+		PrecompileReference->Scripts = ScriptsNeedingCompile;
+
+		const TArray<FNiagaraEmitterHandle>& EmitterHandles = Options.System->GetEmitterHandles();
+		for (int32 i = 0; i < EmitterHandles.Num(); i++)
 		{
-			if (ensure(*AsyncTaskPtr))
+			FNiagaraEmitterHandle Handle = EmitterHandles[i];
+			if (Handle.GetIsEnabled() && Handle.GetEmitterData())
 			{
-				FNiagaraAsyncCompileTask& CompileTask = *(*AsyncTaskPtr);
-				CompileTask.DDCKey = UNiagaraScript::BuildNiagaraDDCKeyString(CompileTask.GetScriptPair().CompileId, CompileTask.AssetPath);
-				CompileTask.PrecompileReference = PrecompileReference;
-				CompileTask.StartTaskTime = FPlatformTime::Seconds();
-				CompileTask.bCheckDDCEnabled = !bEvaluateParametersPending;
-
-				// fire off all the ddc tasks, which will trigger the compilation if the data is not in the ddc
-				UE_LOG(LogNiagara, Verbose, TEXT("Scheduling async get task for %s"), *CompileTask.AssetPath);
-				CompileTask.TaskHandle = GetDerivedDataCacheRef().GetAsynchronous(*CompileTask.DDCKey, CompileTask.AssetPath);
+				TArray<UNiagaraScript*> EmitterScripts;
+				Handle.GetEmitterData()->GetScripts(EmitterScripts, false, true);
+				check(EmitterScripts.Num() > 0);
+				for (UNiagaraScript* EmitterScript : EmitterScripts)
+				{
+					PrecompileReference->EmitterScriptIndex.Add(EmitterScript, i);
+				}
 			}
 		}
+
+		for (UNiagaraScript* Script : ScriptsNeedingCompile)
+		{
+			if (FAsyncTaskPtr* AsyncTaskPtr = FindTask(Script))
+			{
+				if (ensure(*AsyncTaskPtr))
+				{
+					FNiagaraAsyncCompileTask& CompileTask = *(*AsyncTaskPtr);
+					CompileTask.DDCKey = UNiagaraScript::BuildNiagaraDDCKeyString(CompileTask.GetScriptPair().CompileId, CompileTask.AssetPath);
+					CompileTask.PrecompileReference = PrecompileReference;
+					CompileTask.StartTaskTime = FPlatformTime::Seconds();
+					CompileTask.bCheckDDCEnabled = !bEvaluateParametersPending;
+
+					// fire off all the ddc tasks, which will trigger the compilation if the data is not in the ddc
+					UE_LOG(LogNiagara, Verbose, TEXT("Scheduling async get task for %s"), *CompileTask.AssetPath);
+					CompileTask.TaskHandle = GetDerivedDataCacheRef().GetAsynchronous(*CompileTask.DDCKey, CompileTask.AssetPath);
+				}
+			}
+		}
+	}
+
+	bAllScriptsSynchronized = !bAnyUnsynchronized;
+	return bAnyCompiled;
+}
+
+void FNiagaraActiveCompilationDefault::ReportResults(const FNiagaraQueryCompilationOptions& Options) const
+{
+	const float ElapsedWallTime = (float)(FPlatformTime::Seconds() - StartTime);
+	bool bHasCompiledJobs = false;
+	float CombinedCompileTime = 0.0f;
+
+	for (const FAsyncTaskPtr& AsyncTask : Tasks)
+	{
+		if (AsyncTask->bUsedShaderCompilerWorker)
+		{
+			bHasCompiledJobs = true;
+		}
+
+		CombinedCompileTime += AsyncTask->ScriptPair.CompileTime;
+	}
+
+	if (bHasCompiledJobs)
+	{
+		UE_LOG(LogNiagara, Log, TEXT("Compiling System %s took %f sec (time since issued), %f sec (combined shader worker time)."),
+			*Options.System->GetFullName(), ElapsedWallTime, CombinedCompileTime);
+	}
+	else if (bAllScriptsSynchronized == false)
+	{
+		UE_LOG(LogNiagara, Verbose, TEXT("Retrieving %s from DDC took %f sec."), *Options.System->GetFullName(), ElapsedWallTime);
 	}
 }
 
diff --git a/Engine/Plugins/FX/Niagara/Source/Niagara/Private/NiagaraAsyncCompile.h b/Engine/Plugins/FX/Niagara/Source/Niagara/Private/NiagaraAsyncCompile.h
--- a/Engine/Plugins/FX/Niagara/Source/Niagara/Private/NiagaraAsyncCompile.h
+++ b/Engine/Plugins/FX/Niagara/Source/Niagara/Private/NiagaraAsyncCompile.h
@@ -2,6 +2,7 @@
 
 #pragma once
 
+#include "NiagaraCompilationTypes.h"
 #include "NiagaraScript.h"
 
 #include "NiagaraAsyncCompile.generated.h"
@@ -112,3 +113,35 @@ public:
 	bool CompilationIdMatchesRequest() const;
 #endif
 };
+
+class FNiagaraActiveCompilationDefault : public FNiagaraActiveCompilation
+{
+public:
+	virtual void AddReferencedObjects(FReferenceCollector& Collector) override;
+
+	virtual bool Launch(const FNiagaraCompilationOptions& Options) override;
+	virtual void Abort() override;
+	virtual bool QueryCompileComplete(const FNiagaraQueryCompilationOptions& Options) override;
+	virtual bool Validate(const FNiagaraQueryCompilationOptions& Options) const override;
+	virtual void Apply(const FNiagaraQueryCompilationOptions& Options) override;
+	virtual void ReportResults(const FNiagaraQueryCompilationOptions& Options) const override;
+
+	double StartTime = 0.0;
+
+	TSet<TObjectPtr<UObject>> RootObjects;
+
+	using FAsyncTaskPtr = TSharedPtr<FNiagaraAsyncCompileTask, ESPMode::ThreadSafe>;
+
+	FAsyncTaskPtr* FindTask(const UNiagaraScript* Script);
+	const FAsyncTaskPtr* FindTask(const UNiagaraScript* Script) const;
+	void Reset();
+
+	bool bAllScriptsSynchronized = false;
+	bool bEvaluateParametersPending = false;
+
+	TArray<FAsyncTaskPtr> Tasks;
+
+private:
+	bool bDDCGetCompleted = false;
+};
+
diff --git a/Engine/Plugins/FX/Niagara/Source/Niagara/Private/NiagaraCompilationTypes.cpp b/Engine/Plugins/FX/Niagara/Source/Niagara/Private/NiagaraCompilationTypes.cpp
new file mode 100644
--- /dev/null
+++ b/Engine/Plugins/FX/Niagara/Source/Niagara/Private/NiagaraCompilationTypes.cpp
@@ -0,0 +1,21 @@
+// Copyright Epic Games, Inc. All Rights Reserved.
+
+#include "NiagaraCompilationTypes.h"
+
+#include "NiagaraAsyncCompile.h"
+
+#if WITH_EDITORONLY_DATA
+
+TUniquePtr<FNiagaraActiveCompilation> FNiagaraActiveCompilation::CreateCompilation()
+{
+	return MakeUnique<FNiagaraActiveCompilationDefault>();
+}
+
+#else // WITH_EDITORONLY_DATA
+
+TUniquePtr<FNiagaraActiveCompilation> FNiagaraActiveCompilation::CreateCompilation()
+{
+	return nullptr;
+}
+
+#endif // WITH_EDITORONLY_DATA
diff --git a/Engine/Plugins/FX/Niagara/Source/Niagara/Private/NiagaraSystem.cpp b/Engine/Plugins/FX/Niagara/Source/Niagara/Private/NiagaraSystem.cpp
--- a/Engine/Plugins/FX/Niagara/Source/Niagara/Private/NiagaraSystem.cpp
+++ b/Engine/Plugins/FX/Niagara/Source/Niagara/Private/NiagaraSystem.cpp
@@ -11,6 +11,7 @@
 #include "NiagaraAsyncCompile.h"
 #include "NiagaraConstants.h"
 #include "NiagaraBakerSettings.h"
+#include "NiagaraCompilationTypes.h"
 #include "NiagaraCustomVersion.h"
 #include "NiagaraDataInterfaceUtilities.h"
 #include "NiagaraEditorDataBase.h"
@@ -228,10 +229,7 @@ void UNiagaraSystem::BeginDestroy()
 	Super::BeginDestroy();
 
 #if WITH_EDITORONLY_DATA
-	while (ActiveCompilations.Num() > 0)
-	{
-		KillAllActiveCompilations();
-	}
+	KillAllActiveCompilations();
 #endif
 
 #if WITH_EDITOR
@@ -2258,9 +2256,9 @@ const UNiagaraScript* UNiagaraSystem::GetSystemUpdateScript() const
 void UNiagaraSystem::KillAllActiveCompilations()
 {
 	InvalidateActiveCompiles();
-	for (FNiagaraSystemCompileRequest& Request : ActiveCompilations)
+	for (TUniquePtr<FNiagaraActiveCompilation>& ActiveCompilation : ActiveCompilations)
 	{
-		Request.Abort();
+		ActiveCompilation->Abort();
 	}
 	ActiveCompilations.Empty();
 }
@@ -2347,9 +2345,9 @@ void UNiagaraSystem::WaitForCompilationComplete(bool bIncludingGPUShaders, bool
 	double TimeLogThreshold = GNiagaraCompileWaitLoggingThreshold;
 	uint32 NumLogIterations = GNiagaraCompileWaitLoggingTerminationCap;
 
-	while (ActiveCompilations.Num() > 0)
+	while (!ActiveCompilations.IsEmpty())
 	{
-		if (QueryCompileComplete(true, ActiveCompilations.Num() == 1))
+		if (QueryCompileComplete(true))
 		{
 			// make sure to only mark progress if we actually have accomplished something in the QueryCompileComplete
 			Progress.EnterProgressFrame();
@@ -2383,9 +2381,9 @@ void UNiagaraSystem::WaitForCompilationComplete(bool bIncludingGPUShaders, bool
 
 void UNiagaraSystem::InvalidateActiveCompiles()
 {
-	for (FNiagaraSystemCompileRequest& CompileRequest : ActiveCompilations)
+	for (TUniquePtr<FNiagaraActiveCompilation>& ActiveCompilation : ActiveCompilations)
 	{
-		CompileRequest.bIsValid = false;
+		ActiveCompilation->Invalidate();
 	}
 }
 
@@ -2396,78 +2394,7 @@ bool UNiagaraSystem::PollForCompilationComplete(bool bFlushRequestCompile)
 		RequestCompile(false);
 	}
 
-	if (ActiveCompilations.Num() > 0)
-	{
-		return QueryCompileComplete(false, true);
-	}
-	return true;
-}
-
-bool UNiagaraSystem::CompilationResultsValid(FNiagaraSystemCompileRequest& CompileRequest) const
-{
-	// for now the only thing we're concerned about is if we've got results for SystemSpawn and SystemUpdate scripts
-	// then we need to make sure that they agree in terms of the dataset attributes
-	FNiagaraSystemCompileRequest::FAsyncTaskPtr* SpawnScriptRequest = CompileRequest.FindTask(SystemSpawnScript);
-	FNiagaraSystemCompileRequest::FAsyncTaskPtr* UpdateScriptRequest = CompileRequest.FindTask(SystemUpdateScript);
-
-	const bool SpawnScriptValid = SpawnScriptRequest
-		&& SpawnScriptRequest->Get()->GetScriptPair().CompileResults.IsValid()
-		&& SpawnScriptRequest->Get()->GetScriptPair().CompileResults->LastCompileStatus != ENiagaraScriptCompileStatus::NCS_Error;
-
-	const bool UpdateScriptValid = UpdateScriptRequest
-		&& UpdateScriptRequest->Get()->GetScriptPair().CompileResults.IsValid()
-		&& UpdateScriptRequest->Get()->GetScriptPair().CompileResults->LastCompileStatus != ENiagaraScriptCompileStatus::NCS_Error;
-
-	if (SpawnScriptValid && UpdateScriptValid)
-	{
-		const FEmitterCompiledScriptPair& SpawnScriptPair = SpawnScriptRequest->Get()->GetScriptPair();
-		const FEmitterCompiledScriptPair& UpdateScriptPair = UpdateScriptRequest->Get()->GetScriptPair();
-		if (SpawnScriptPair.CompileResults->Attributes != UpdateScriptPair.CompileResults->Attributes)
-		{
-			// if we had requested a full rebuild, then we've got a case where the generated scripts are not compatible.  This indicates
-			// a significant issue where we're allowing graphs to generate invalid collections of scripts.  One known example is using
-			// the Script.Context static switch that isn't fully processed in all scripts, leading to attributes differing between the
-			// SystemSpawnScript and the SystemUpdateScript
-			if (CompileRequest.bForced)
-			{
-				FString MissingAttributes;
-				FString AdditionalAttributes;
-
-				for (const auto& SpawnAttrib : SpawnScriptPair.CompileResults->Attributes)
-				{
-					if (!UpdateScriptPair.CompileResults->Attributes.Contains(SpawnAttrib))
-					{
-						MissingAttributes.Appendf(TEXT("%s%s"), MissingAttributes.Len() ? TEXT(", ") : TEXT(""), *SpawnAttrib.GetName().ToString());
-					}
-				}
-
-				for (const auto& UpdateAttrib : UpdateScriptPair.CompileResults->Attributes)
-				{
-					if (!SpawnScriptPair.CompileResults->Attributes.Contains(UpdateAttrib))
-					{
-						AdditionalAttributes.Appendf(TEXT("%s%s"), AdditionalAttributes.Len() ? TEXT(", ") : TEXT(""), *UpdateAttrib.GetName().ToString());
-					}
-				}
-
-				FNiagaraCompileEvent AttributeMismatchEvent(
-					FNiagaraCompileEventSeverity::Error,
-					FText::Format(LOCTEXT("SystemScriptAttributeMismatchError", "System Spawn/Update scripts have attributes which don't match!\n\tMissing update attributes: {0}\n\tAdditional update attributes: {1}"),
-						FText::FromString(MissingAttributes),
-						FText::FromString(AdditionalAttributes))
-					.ToString());
-
-				SpawnScriptPair.CompileResults->LastCompileStatus = ENiagaraScriptCompileStatus::NCS_Error;
-				SpawnScriptPair.CompileResults->LastCompileEvents.Add(AttributeMismatchEvent);
-			}
-			else
-			{
-				UE_LOG(LogNiagara, Log, TEXT("Failed to generate consistent results for System spawn and update scripts for system %s."), *GetFullName());
-			}
-
-			return false;
-		}
-	}
-	return true;
+	return QueryCompileComplete(false);
 }
 
 void UNiagaraSystem::EvaluateCompileResultDependencies() const
@@ -2710,110 +2637,100 @@ void UNiagaraSystem::EvaluateCompileResultDependencies() const
 	}
 }
 
-bool UNiagaraSystem::QueryCompileComplete(bool bWait, bool bDoPost, bool bDoNotApply)
+bool UNiagaraSystem::QueryCompileComplete(bool bWait)
 {
 	check(IsInGameThread());
-	if (ActiveCompilations.Num() > 0 && !bCompilationReentrantGuard)
+
+	if (ActiveCompilations.IsEmpty() || bCompilationReentrantGuard)
 	{
-		bCompilationReentrantGuard = true;
-		ON_SCOPE_EXIT { bCompilationReentrantGuard = false; };
+		return false;
+	}
 
-		FNiagaraSystemCompileRequest& CompileRequest = ActiveCompilations[0];
+	TGuardValue<bool> ReentrantGuard(bCompilationReentrantGuard, true);
 
-		constexpr double MaxDurationPerIteration = 0.125;
-		bool bAreWeWaitingForAnyResults = !CompileRequest.QueryCompileComplete(this, bWait, MaxDurationPerIteration);
+	FNiagaraQueryCompilationOptions Options;
+	Options.System = this;
+	Options.bWait = bWait;
+	Options.MaxWaitDuration = 0.125;
 
-		// Make sure that we aren't waiting for any results to come back.
-		if (bAreWeWaitingForAnyResults)
-		{
-			return false;
-		}
-		// if we've gotten all the results, run a quick check to see if the data is valid, if it's not then that indicates that
-		// we've run into a compatibility issue and so we should see if we should issue a full rebuild
-		const bool ResultsValid = CompilationResultsValid(CompileRequest);
-		if (!ResultsValid && !CompileRequest.bForced)
-		{
-			CompileRequest.RootObjects.Empty();
-			ActiveCompilations.RemoveAt(0);
-			RequestCompile(true, nullptr);
-			return false;
-		}
+	const bool bCompileComplete = ActiveCompilations[0]->QueryCompileComplete(Options);
 
-		// In the world of do not apply, we're exiting the system completely so let's just kill any active compilations altogether.
-		if (bDoNotApply || CompileRequest.bIsValid == false)
-		{
-			CompileRequest.RootObjects.Empty();
-			ActiveCompilations.RemoveAt(0);
-			return true;
-		}
+	// Make sure that we aren't waiting for any results to come back.
+	if (!bCompileComplete)
+	{
+		return false;
+	}
 
-		// Now that the above code says they are all complete, go ahead and resolve them all at once.
-		bool bHasCompiledJobs = CompileRequest.Resolve(this, ExposedParameters);
-		CompileRequest.Reset();
+	// now that the request is complete we can pull it from the array of ActiveCompilations and
+	// finish processing it
+	TUniquePtr<FNiagaraActiveCompilation> CurrentCompilation = MoveTemp(ActiveCompilations[0]);
+	ActiveCompilations.RemoveAt(0);
 
-		// Once compile results have been set, check dependencies found during compile and evaluate whether dependency compile events should be added.
-		if (FNiagaraCVarUtilities::GetShouldEmitMessagesForFailIfNotSet())
-		{
-			EvaluateCompileResultDependencies();
-		}
+	// if we've gotten all the results, run a quick check to see if the data is valid, if it's not then that indicates that
+	// we've run into a compatibility issue and so we should see if we should issue a full rebuild
+	const bool ResultsValid = CurrentCompilation->Validate(Options);
+	if (!ResultsValid && !CurrentCompilation->WasForced())
+	{
+		RequestCompile(true, nullptr);
+		return false;
+	}
 
-		if (bDoPost)
-		{
-			for (FNiagaraEmitterHandle Handle : EmitterHandles)
-			{
-				if (Handle.GetInstance().Emitter)
-				{
-					if (Handle.GetIsEnabled())
-					{
-						Handle.GetInstance().Emitter->OnPostCompile(Handle.GetInstance().Version);
-					}
-					else
-					{
-						Handle.GetEmitterData()->InvalidateCompileResults();
-					}
-				}
-			}
-		}
+	// In the world of do not apply, we're exiting the system completely so let's just kill any active compilations altogether.
+	if (!CurrentCompilation->ShouldApply())
+	{
+		return true;
+	}
 
-		InitEmitterCompiledData();
-		InitSystemCompiledData();
+	CurrentCompilation->Apply(Options);
 
-		CompileRequest.RootObjects.Empty();
+	// Once compile results have been set, check dependencies found during compile and evaluate whether dependency compile events should be added.
+	if (FNiagaraCVarUtilities::GetShouldEmitMessagesForFailIfNotSet())
+	{
+		EvaluateCompileResultDependencies();
+	}
 
-		ResolveDIBindings();
-		UpdatePostCompileDIInfo();
-		ComputeEmittersExecutionOrder();
-		ComputeRenderersDrawOrder();
-		CacheFromCompiledData();
-		UpdateHasGPUEmitters();
-		UpdateDITickFlags();
-		ResolveScalabilitySettings();
+	// we want to run the post compile steps (specifically broadcasting compilations completing) unless
+	// we're in the middle of waiting and we still have more compilations to get through
+	const bool bDoPostCompile = !bWait || ActiveCompilations.IsEmpty();
 
-		const float ElapsedWallTime = (float)(FPlatformTime::Seconds() - CompileRequest.StartTime);
-		
-		if (bHasCompiledJobs)
-		{
-			UE_LOG(LogNiagara, Log, TEXT("Compiling System %s took %f sec (time since issued), %f sec (combined shader worker time)."),
-				*GetFullName(), ElapsedWallTime, CompileRequest.CombinedCompileTime);
-		}
-		else if (CompileRequest.bAllScriptsSynchronized == false)
+	if (bDoPostCompile)
+	{
+		for (FNiagaraEmitterHandle Handle : EmitterHandles)
 		{
-			UE_LOG(LogNiagara, Verbose, TEXT("Retrieving %s from DDC took %f sec."), *GetFullName(), ElapsedWallTime);
+			if (Handle.GetInstance().Emitter)
+			{
+				if (Handle.GetIsEnabled())
+				{
+					Handle.GetInstance().Emitter->OnPostCompile(Handle.GetInstance().Version);
+				}
+				else
+				{
+					Handle.GetEmitterData()->InvalidateCompileResults();
+				}
+			}
 		}
+	}
 
-		ActiveCompilations.RemoveAt(0);
+	InitEmitterCompiledData();
+	InitSystemCompiledData();
 
-		if (bDoPost)
-		{
-			SCOPE_CYCLE_COUNTER(STAT_Niagara_System_CompileScriptResetAfter);
+	ResolveDIBindings();
+	UpdatePostCompileDIInfo();
+	ComputeEmittersExecutionOrder();
+	ComputeRenderersDrawOrder();
+	CacheFromCompiledData();
+	UpdateHasGPUEmitters();
+	UpdateDITickFlags();
+	ResolveScalabilitySettings();
 
-			BroadcastOnSystemCompiled();
-		}
+	CurrentCompilation->ReportResults(Options);
 
-		return true;
+	if (bDoPostCompile)
+	{
+		BroadcastOnSystemCompiled();
 	}
 
-	return false;
+	return true;
 }
 
 void UNiagaraSystem::BroadcastOnSystemCompiled()
@@ -3011,140 +2928,37 @@ bool UNiagaraSystem::RequestCompile(bool bForce, FNiagaraSystemUpdateContext* Op
 	// we don't want to stack compilations, so we remove requests that have not started processing yet
 	for (int i = ActiveCompilations.Num() - 1; i >= 1; i--)
 	{
-		FNiagaraSystemCompileRequest& Request = ActiveCompilations[i];
-		Request.Abort();
+		ActiveCompilations[i]->Abort();
 		ActiveCompilations.RemoveAt(i);
-	}	
+	}
 
-	// Record that we entered this function already.
 	SCOPE_CYCLE_COUNTER(STAT_Niagara_System_CompileScript);
-	bCompilationReentrantGuard = true;
-	ON_SCOPE_EXIT { check(bCompilationReentrantGuard == false); };
 
-	FNiagaraSystemCompileRequest& ActiveCompilation = ActiveCompilations.AddDefaulted_GetRef();
-	ActiveCompilation.bForced = bForce;
-	ActiveCompilation.StartTime = FPlatformTime::Seconds();
-
-	check(SystemSpawnScript->GetLatestSource() == SystemUpdateScript->GetLatestSource());
+	bool bLaunchedCompilations = false;
 
-	// when not compiling for edit mode we will defer preparing rapid iteration parameters till after we get the
-	// results from the DDC search (FNiagaraPrepareCompileTask::Tick())
-	if (ShouldUseRapidIterationParameters())
+	// Record that we entered this function already.
 	{
-		PrepareRapidIterationParametersForCompilation();
-	}
+		TGuardValue<bool> ReentrantGuard(bCompilationReentrantGuard, true);
 
-	TArray<UNiagaraScript*> ScriptsNeedingCompile;
-	bool bAnyCompiled = false;
-	{
 		COOK_STAT(auto Timer = NiagaraScriptCookStats::UsageStats.TimeSyncWork());
 		COOK_STAT(Timer.TrackCyclesOnly());
 
-		//Compile all emitters
-		TArray<FNiagaraSystemCompileRequest::FAsyncTaskPtr> Tasks;
-
-		// Pass one... determine if any need to be compiled.
-		{
-			bool bAnyUnsynchronized = false;
-			for (int32 i = 0; i < EmitterHandles.Num(); i++)
-			{
-				FNiagaraEmitterHandle Handle = EmitterHandles[i];
-				FVersionedNiagaraEmitterData* EmitterData = Handle.GetEmitterData();
-				if (EmitterData && Handle.GetIsEnabled())
-				{
-					TArray<UNiagaraScript*> EmitterScripts;
-					EmitterData->GetScripts(EmitterScripts, false, true);
-					check(EmitterScripts.Num() > 0);
-					for (UNiagaraScript* EmitterScript : EmitterScripts)
-					{
-						FEmitterCompiledScriptPair Pair;
-						Pair.VersionedEmitter = Handle.GetInstance();
-						Pair.CompiledScript = EmitterScript;
-
-						TSharedPtr<FNiagaraAsyncCompileTask, ESPMode::ThreadSafe> AsyncTask = MakeShared<FNiagaraAsyncCompileTask, ESPMode::ThreadSafe>(this, EmitterScript->GetPathName(), Pair);
-						if (EmitterScript->IsCompilable() && !EmitterScript->AreScriptAndSourceSynchronized())
-						{
-							// we need to compute the vmID here to check later in the ddc task before doing the precompile if anything has changed in the meantime.
-							// the compilation id was just calculated during AreScriptAndSourceSynchronized() so we reuse it here
-							AsyncTask->AssignInitialCompilationId(EmitterScript->GetComputedVMCompilationId());
-
-							ScriptsNeedingCompile.Add(EmitterScript);
-							bAnyUnsynchronized = true;
-						}
-						else
-						{
-							AsyncTask->CurrentState = ENiagaraCompilationState::Finished;
-						}
-						AsyncTask->UniqueEmitterName = Handle.GetInstance().Emitter->GetUniqueEmitterName();
-						Tasks.Add(AsyncTask);
-					}
-				}
-			}
-
-			bAnyCompiled = bAnyUnsynchronized || bForce;
-
-			// Now add the system scripts for compilation...
-			{
-				FEmitterCompiledScriptPair Pair;
-				Pair.VersionedEmitter = FVersionedNiagaraEmitter();
-				Pair.CompiledScript = SystemSpawnScript;
-				
-				TSharedPtr<FNiagaraAsyncCompileTask, ESPMode::ThreadSafe> AsyncTask = MakeShared<FNiagaraAsyncCompileTask, ESPMode::ThreadSafe>(this, SystemSpawnScript->GetPathName(), Pair);
-				if (!SystemSpawnScript->AreScriptAndSourceSynchronized())
-				{
-					// we need to compute the vmID here to check later in the ddc task before doing the precompile if anything has changed in the meantime.
-					// the compilation id was just calculated during AreScriptAndSourceSynchronized() so we reuse it here
-					AsyncTask->AssignInitialCompilationId(SystemSpawnScript->GetComputedVMCompilationId());
-
-					ScriptsNeedingCompile.Add(SystemSpawnScript);
-					bAnyCompiled = true;
-				}
-				else
-				{
-					AsyncTask->CurrentState = ENiagaraCompilationState::Finished;
-				}
-				Tasks.Add(AsyncTask);
-			}
+		TUniquePtr<FNiagaraActiveCompilation>& CurrentCompilation = ActiveCompilations.Emplace_GetRef(
+			FNiagaraActiveCompilation::CreateCompilation());
 
-			{
-				FEmitterCompiledScriptPair Pair;
-				Pair.VersionedEmitter = FVersionedNiagaraEmitter();
-				Pair.CompiledScript = SystemUpdateScript;
-				
-				TSharedPtr<FNiagaraAsyncCompileTask, ESPMode::ThreadSafe> AsyncTask = MakeShared<FNiagaraAsyncCompileTask, ESPMode::ThreadSafe>(this, SystemUpdateScript->GetPathName(), Pair);
-				if (!SystemUpdateScript->AreScriptAndSourceSynchronized())
-				{
-					// we need to compute the vmID here to check later in the ddc task before doing the precompile if anything has changed in the meantime.
-					// the compilation id was just calculated during AreScriptAndSourceSynchronized() so we reuse it here
-					AsyncTask->AssignInitialCompilationId(SystemUpdateScript->GetComputedVMCompilationId());
-
-					ScriptsNeedingCompile.Add(SystemUpdateScript);
-					bAnyCompiled = true;
-				}
-				else
-				{
-					AsyncTask->CurrentState = ENiagaraCompilationState::Finished;
-				}
-				Tasks.Add(AsyncTask);
-			}
-		}
+		FNiagaraCompilationOptions Options;
+		Options.System = this;
+		Options.bForced = bForce;
 
-		{
-			ActiveCompilation.Launch(this, ScriptsNeedingCompile, Tasks);
-		}
+		bLaunchedCompilations = CurrentCompilation->Launch(Options);
 	}
 
-	// Now record that we are done with this function.
-	bCompilationReentrantGuard = false;
-
 	// We might be able to just complete compilation right now if nothing needed compilation.
-	if (ScriptsNeedingCompile.Num() == 0 && ActiveCompilations.Num() == 1)
+	if (!bLaunchedCompilations && ActiveCompilations.Num() == 1)
 	{
-		ActiveCompilation.bAllScriptsSynchronized = true;
 		PollForCompilationComplete();
 	}
 
-	
 	if (OptionalUpdateContext)
 	{
 		OptionalUpdateContext->Add(this, true);
@@ -3154,7 +2968,7 @@ bool UNiagaraSystem::RequestCompile(bool bForce, FNiagaraSystemUpdateContext* Op
 		FNiagaraSystemUpdateContext UpdateCtx(this, true);
 	}
 
-	return bAnyCompiled;
+	return bLaunchedCompilations;
 }
 
 void UNiagaraSystem::InitEmitterCompiledData()
diff --git a/Engine/Plugins/FX/Niagara/Source/Niagara/Public/NiagaraCompilationTypes.h b/Engine/Plugins/FX/Niagara/Source/Niagara/Public/NiagaraCompilationTypes.h
new file mode 100644
--- /dev/null
+++ b/Engine/Plugins/FX/Niagara/Source/Niagara/Public/NiagaraCompilationTypes.h
@@ -0,0 +1,55 @@
+// Copyright Epic Games, Inc. All Rights Reserved.
+
+#pragma once
+
+#include "NiagaraCommon.h"
+
+class UNiagaraSystem;
+
+struct FNiagaraCompilationOptions
+{
+	UNiagaraSystem* System = nullptr;
+	bool bForced = false;
+};
+
+struct FNiagaraQueryCompilationOptions
+{
+	UNiagaraSystem* System = nullptr;
+	double MaxWaitDuration = 0.125;
+	bool bWait = false;
+};
+
+class FNiagaraActiveCompilation : public FGCObject
+{
+public:
+	static TUniquePtr<FNiagaraActiveCompilation> CreateCompilation();
+
+	virtual bool Launch(const FNiagaraCompilationOptions& Options) = 0;
+	virtual void Abort() = 0;
+	virtual bool QueryCompileComplete(const FNiagaraQueryCompilationOptions& Options) = 0;
+	virtual bool Validate(const FNiagaraQueryCompilationOptions& Options) const = 0;
+	virtual void Apply(const FNiagaraQueryCompilationOptions& Options) = 0;
+	virtual void ReportResults(const FNiagaraQueryCompilationOptions& Options) const = 0;
+
+	void Invalidate()
+	{
+		bShouldApply = false;
+	}
+
+	bool ShouldApply() const
+	{
+		return bShouldApply;
+	}
+
+	bool WasForced() const
+	{
+		return bForced;
+	}
+
+protected:
+	bool bForced = false;
+
+private:
+	bool bShouldApply = true;
+};
+
