Fixed a cook content warning when cooking emitters that were saved with deprecated summary view properties

diff --git a/Engine/Plugins/FX/Niagara/Source/NiagaraEditor/Private/NiagaraEmitterEditorData.cpp b/Engine/Plugins/FX/Niagara/Source/NiagaraEditor/Private/NiagaraEmitterEditorData.cpp
--- a/Engine/Plugins/FX/Niagara/Source/NiagaraEditor/Private/NiagaraEmitterEditorData.cpp
+++ b/Engine/Plugins/FX/Niagara/Source/NiagaraEditor/Private/NiagaraEmitterEditorData.cpp
@@ -245,7 +245,9 @@ void UNiagaraEmitterEditorData::PostLoad_TransferSummaryDataToNewFormat()
 				
 				if(!It.Value().DisplayName.IsNone() && ScriptVariable != nullptr)
 				{
-					SummaryModuleInput->SetDisplayNameOverride(FText::FromName(It.Value().DisplayName));
+					// We use FText::Format here because otherwise loading the editor data during a cook will result in a warning (FText should not be initialized from strings during a cook).
+					FText NameOverride = FText::Format(LOCTEXT("NameOverrideFormat", "{0}"), FText::FromName(It.Value().DisplayName));
+					SummaryModuleInput->SetDisplayNameOverride(NameOverride);
 				}
 				SummarySortOrder.Add(SummaryModuleInput, It.Value().SortIndex);
 
