- Add simple motion option to Niagara viewport

diff --git a/Engine/Plugins/FX/Niagara/Source/NiagaraEditor/Private/NiagaraEditorCommands.cpp b/Engine/Plugins/FX/Niagara/Source/NiagaraEditor/Private/NiagaraEditorCommands.cpp
--- a/Engine/Plugins/FX/Niagara/Source/NiagaraEditor/Private/NiagaraEditorCommands.cpp
+++ b/Engine/Plugins/FX/Niagara/Source/NiagaraEditor/Private/NiagaraEditorCommands.cpp
@@ -39,6 +39,7 @@ void FNiagaraEditorCommands::RegisterCommands()
 	UI_COMMAND(ToggleBounds_SetFixedBounds_SelectedEmitters, "Set Fixed Bounds (Emitters)", "Set Fixed Bounds on emitters (only the selected emitters when in a System asset) ", EUserInterfaceActionType::Button, FInputChord());
 	UI_COMMAND(ToggleBounds_SetFixedBounds_System, "Set Fixed Bounds (System)", "Set Fixed Bounds on the system", EUserInterfaceActionType::Button, FInputChord());
 	UI_COMMAND(ToggleOrbit, "Orbit Mode", "Toggle Orbit Navigation", EUserInterfaceActionType::ToggleButton, FInputChord());
+	UI_COMMAND(ToggleMotion, "Motion", "Toggle Motion On Component", EUserInterfaceActionType::ToggleButton, FInputChord());
 	UI_COMMAND(SaveThumbnailImage, "Thumbnail", "Generate Thumbnail", EUserInterfaceActionType::Button, FInputChord());
 	UI_COMMAND(ToggleAutoPlay, "Auto-play", "Toggles whether or not simulations auto-play when their asset editor is opened, and when the asset is modified.", EUserInterfaceActionType::ToggleButton, FInputChord());
 	UI_COMMAND(ToggleResetSimulationOnChange, "Reset on change", "Toggles whether or not the simulation is reset whenever a change is made in the asset editor.", EUserInterfaceActionType::ToggleButton, FInputChord());
diff --git a/Engine/Plugins/FX/Niagara/Source/NiagaraEditor/Private/Widgets/SNiagaraSystemViewport.cpp b/Engine/Plugins/FX/Niagara/Source/NiagaraEditor/Private/Widgets/SNiagaraSystemViewport.cpp
--- a/Engine/Plugins/FX/Niagara/Source/NiagaraEditor/Private/Widgets/SNiagaraSystemViewport.cpp
+++ b/Engine/Plugins/FX/Niagara/Source/NiagaraEditor/Private/Widgets/SNiagaraSystemViewport.cpp
@@ -583,6 +583,20 @@ void SNiagaraSystemViewport::ToggleOrbit()
 	}
 }
 
+void SNiagaraSystemViewport::ToggleMotion()
+{
+	bMotionEnabled = !bMotionEnabled;
+	if (PreviewComponent && bMotionEnabled == false)
+	{
+		PreviewComponent->SetComponentToWorld(FTransform::Identity);
+	}
+}
+
+bool SNiagaraSystemViewport::IsMotionEnabled() const
+{
+	return bMotionEnabled;
+}
+
 void SNiagaraSystemViewport::RefreshViewport()
 {
 	//reregister the preview components, so if the preview material changed it will be propagated to the render thread
@@ -605,6 +619,18 @@ void SNiagaraSystemViewport::Tick( const FGeometry& AllottedGeometry, const doub
 		bShouldActivateOrbitAfterTransitioning = false;
 		bIsViewTransitioning = false;
 	}
+
+	if (PreviewComponent && bMotionEnabled)
+	{
+		const float MotionTime = PreviewComponent->GetDesiredAge();
+
+		UWorld* World = GetWorld();
+		FVector Location;
+		Location.X = MotionRadius * FMath::Sin(FMath::DegreesToRadians(MotionRate) * MotionTime);
+		Location.Y = 0.0f;
+		Location.Z = MotionRadius * FMath::Cos(FMath::DegreesToRadians(MotionRate) * MotionTime);
+		PreviewComponent->SetComponentToWorld(FTransform(Location));
+	}
 }
 
 void SNiagaraSystemViewport::SetPreviewComponent(UNiagaraComponent* NiagaraComponent)
@@ -736,6 +762,11 @@ void SNiagaraSystemViewport::BindCommands()
 		FCanExecuteAction(),
 		FIsActionChecked::CreateSP(this, &SNiagaraSystemViewport::IsToggleOrbitChecked));
 
+	CommandList->MapAction(
+		Commands.ToggleMotion,
+		FExecuteAction::CreateSP(this, &SNiagaraSystemViewport::ToggleMotion),
+		FCanExecuteAction(),
+		FIsActionChecked::CreateSP(this, &SNiagaraSystemViewport::IsMotionEnabled));
 }
 
 void SNiagaraSystemViewport::OnFocusViewportToSelection()
diff --git a/Engine/Plugins/FX/Niagara/Source/NiagaraEditor/Private/Widgets/SNiagaraSystemViewport.h b/Engine/Plugins/FX/Niagara/Source/NiagaraEditor/Private/Widgets/SNiagaraSystemViewport.h
--- a/Engine/Plugins/FX/Niagara/Source/NiagaraEditor/Private/Widgets/SNiagaraSystemViewport.h
+++ b/Engine/Plugins/FX/Niagara/Source/NiagaraEditor/Private/Widgets/SNiagaraSystemViewport.h
@@ -53,9 +53,6 @@ public:
 	
 	void ToggleRealtime();
 	
-	/** @return The list of commands known by the niagara editor */
-	TSharedRef<FUICommandList> GetNiagaraSystemEditorCommands() const;
-	
 	/** If true, render background object in the preview scene. */
 	bool bShowBackground;
 
@@ -91,10 +88,18 @@ public:
 	void ToggleDrawElement(EDrawElements Element);
 	void CreateThumbnail(UObject* InScreenShotOwner);
 
-
 	bool IsToggleOrbitChecked() const;
 	void ToggleOrbit();
 
+	bool IsMotionEnabled() const;
+	void ToggleMotion();
+
+	float GetMotionRate() const { return MotionRate; }
+	void SetMotionRate(float Rate) { MotionRate = Rate; }
+
+	float GetMotionRadius() const { return MotionRadius; }
+	void SetMotionRadius(float Radius) { MotionRadius = Radius; }
+
 protected:
 	/** SEditorViewport interface */
 	virtual TSharedRef<FEditorViewportClient> MakeEditorViewportClient() override;
@@ -139,6 +144,11 @@ private:
 	/** True if orbit mode was active before we started a view transition. Used to restore orbit mode at the the end of the transition */
 	bool bShouldActivateOrbitAfterTransitioning = false;
 
+	/** Motion Parameters */
+	bool bMotionEnabled = false;
+	float MotionRate = 90.0f;
+	float MotionRadius = 200.0f;
+
 	FDelegateHandle OnPreviewFeatureLevelChangedHandle;
 };
 
diff --git a/Engine/Plugins/FX/Niagara/Source/NiagaraEditor/Private/Widgets/SNiagaraSystemViewportToolBar.cpp b/Engine/Plugins/FX/Niagara/Source/NiagaraEditor/Private/Widgets/SNiagaraSystemViewportToolBar.cpp
--- a/Engine/Plugins/FX/Niagara/Source/NiagaraEditor/Private/Widgets/SNiagaraSystemViewportToolBar.cpp
+++ b/Engine/Plugins/FX/Niagara/Source/NiagaraEditor/Private/Widgets/SNiagaraSystemViewportToolBar.cpp
@@ -1,6 +1,7 @@
 // Copyright Epic Games, Inc. All Rights Reserved.
 
 #include "SNiagaraSystemViewportToolBar.h"
+#include "Widgets/Input/SSpinBox.h"
 #include "Widgets/Layout/SBorder.h"
 #include "Styling/AppStyle.h"
 #include "NiagaraEditorCommands.h"
@@ -17,6 +18,8 @@
 
 void SNiagaraSystemViewportToolBar::Construct(const FArguments& InArgs, TSharedPtr<class SNiagaraSystemViewport> InViewport)
 {
+	WeakViewport = InViewport;
+
 	// we don't want a realtime button here as we create a custom one by extending the left menu
 	SCommonEditorViewportToolbarBase::Construct(SCommonEditorViewportToolbarBase::FArguments().AddRealtimeButton(true).PreviewProfileController(MakeShared<FPreviewProfileController>()), InViewport);
 	Sequencer = InArgs._Sequencer;
@@ -93,6 +96,12 @@ void SNiagaraSystemViewportToolBar::ExtendOptionsMenu(FMenuBuilder& OptionsMenuB
 {
 	OptionsMenuBuilder.BeginSection("LevelViewportNavigationOptions", LOCTEXT("NavOptionsMenuHeader", "Navigation Options"));
 	OptionsMenuBuilder.AddMenuEntry(FNiagaraEditorCommands::Get().ToggleOrbit);
+	OptionsMenuBuilder.AddSubMenu(
+		LOCTEXT("MotionOptions", "Motion Options"),
+		LOCTEXT("MotionOptionsToolTip", "Set Motion Options For the Niagara Component"),
+		FNewMenuDelegate::CreateSP(this, &SNiagaraSystemViewportToolBar::MakeMotionMenu)
+	);
+
 	OptionsMenuBuilder.EndSection();
 }
 
@@ -118,4 +127,38 @@ void SNiagaraSystemViewportToolBar::ExtendLeftAlignedToolbarSlots(TSharedPtr<SHo
 	];
 }
 
+void SNiagaraSystemViewportToolBar::MakeMotionMenu(FMenuBuilder& MenuBuilder) const
+{
+	MenuBuilder.AddMenuEntry(FNiagaraEditorCommands::Get().ToggleMotion);
+
+	if (TSharedPtr<SNiagaraSystemViewport> ViewportSP = WeakViewport.Pin())
+	{
+		TSharedRef<SNiagaraSystemViewport> Viewport = ViewportSP.ToSharedRef();
+
+		MenuBuilder.AddWidget(
+			SNew(SSpinBox<float>)
+				.IsEnabled(Viewport, &SNiagaraSystemViewport::IsMotionEnabled)
+				.Font(FAppStyle::GetFontStyle(TEXT("MenuItem.Font")))
+				.MinSliderValue(0.0f)
+				.MaxSliderValue(360.0f)
+				.Value(Viewport, &SNiagaraSystemViewport::GetMotionRate)
+				.OnValueChanged(Viewport, &SNiagaraSystemViewport::SetMotionRate)
+			,
+			LOCTEXT("MotionSpeed", "Motion Speed")
+		);
+
+		MenuBuilder.AddWidget(
+			SNew(SSpinBox<float>)
+			.IsEnabled(Viewport, &SNiagaraSystemViewport::IsMotionEnabled)
+			.Font(FAppStyle::GetFontStyle(TEXT("MenuItem.Font")))
+			.MinSliderValue(0.0f)
+			.MaxSliderValue(1000.0f)
+			.Value(Viewport, &SNiagaraSystemViewport::GetMotionRadius)
+			.OnValueChanged(Viewport, &SNiagaraSystemViewport::SetMotionRadius)
+			,
+			LOCTEXT("MotionRadius", "Motion Radius")
+		);
+	}
+}
+
 #undef LOCTEXT_NAMESPACE
diff --git a/Engine/Plugins/FX/Niagara/Source/NiagaraEditor/Private/Widgets/SNiagaraSystemViewportToolBar.h b/Engine/Plugins/FX/Niagara/Source/NiagaraEditor/Private/Widgets/SNiagaraSystemViewportToolBar.h
--- a/Engine/Plugins/FX/Niagara/Source/NiagaraEditor/Private/Widgets/SNiagaraSystemViewportToolBar.h
+++ b/Engine/Plugins/FX/Niagara/Source/NiagaraEditor/Private/Widgets/SNiagaraSystemViewportToolBar.h
@@ -33,6 +33,7 @@ public:
 	virtual bool IsViewModeSupported(EViewModeIndex ViewModeIndex) const override;
 
 private:
+	TWeakPtr<class SNiagaraSystemViewport> WeakViewport = nullptr;
 	TWeakPtr<ISequencer> Sequencer = nullptr;
 	bool bIsSimulationRealtime = true;
 	
@@ -40,4 +41,6 @@ private:
 	FReply OnSimulationRealtimeWarningClicked() const;
 
 	FText GetSimulationSpeedText() const;
-};
\ No newline at end of file
+
+	void MakeMotionMenu(FMenuBuilder& MenuBuilder) const;
+};
diff --git a/Engine/Plugins/FX/Niagara/Source/NiagaraEditor/Public/NiagaraEditorCommands.h b/Engine/Plugins/FX/Niagara/Source/NiagaraEditor/Public/NiagaraEditorCommands.h
--- a/Engine/Plugins/FX/Niagara/Source/NiagaraEditor/Public/NiagaraEditorCommands.h
+++ b/Engine/Plugins/FX/Niagara/Source/NiagaraEditor/Public/NiagaraEditorCommands.h
@@ -54,6 +54,8 @@ public:
 	TSharedPtr<FUICommandInfo> ToggleBounds_SetFixedBounds_System;
 	TSharedPtr<FUICommandInfo> SaveThumbnailImage;
 
+	TSharedPtr<FUICommandInfo> ToggleMotion;
+
 	TSharedPtr<FUICommandInfo> ToggleStatPerformance;
 	TSharedPtr<FUICommandInfo> ToggleStatPerformanceGPU;
 	TSharedPtr<FUICommandInfo> ClearStatPerformance;
