Impact FX Aggregation via Data Channels.

Primary CL goal:
Prepare a Data Channels based system that could be used in games and elsewhere to create an optimized impact/oneshot FX system.
The high level operation being that BP and Game code can add items into the NDC with position, orientation, Surface parameters etc and a handler system will read those events and their data to spawn FX into large, uber simulations.
Improving perf by reduing overhead and increasing instance counts for the VM & GPU to operate more efficiently.

Sub sections of this CL in support of the Primary Goal.

 - The Islands Data Channel Handler.
	A new handler type that will subdivide the world and channel data into discrete islands so that spacially localized FX and game code can operate only on the data channel data in their local region.

 - Improved archetecture around handing out data to game code and DIs.
	We now have a common FNiagaraDataChannelDataPtr shared ptr which Channels will provide to users on request. A Channel may have > 1 of these. e.g. Each Island in the Islands channel type has it's own FNiagaraDataChannelDataPtr.
	DIs will grab this ptr for I/O with the channel.
	Each FNiagaraDataChannelDataPtr has an RTProxy when needing to expose data to the RT for GPU sims.

 - Data Channel Read DI now has a Spawn Conditional function which is similar to the existing Spawn function in that it allows the DI to spawn particles into an emitter for every entry in the bound Data Channel.
	The differences:
	 - Takes an explicit min/max count to spawn.
	 - Can handle subsequent calls to either accumulate spawn counts or overwrite them based on a Mode param.
	 - Can add additional parameters to the function that act as conditions on spawning.
		These conditional params are compared to the corresponding variable in each data channel entry based on an operator enum input.
		Only if the comparison is true will an element in the data channel trigger spawning.

 - Improve tick order issues.
	Begin/EndFrame for handlers calls allow islands to marshal data correctly such that all previous frame data is available when reading in a subsequent frame.
	Reworked the game level data store slightly

 - Optional Pre/Post Stage update functions for DIs,
	This is required for some new Spawning functions but will also be needed for furture work such as Sim Stages on CPU.
	Each VM exec can have it's own Pre/Post stage call in the DI.
	This works by hooking a new FNDIStageTickHandler into the exec contexts directly.
	These build a tick list for DIs needing the pre/post call at each stage.

 - Refactored much out of Data Channel Global and into a general DataChannel Data class that hanlders can now own one or more of.

 - Improved BP and C++ I/O api. Still needs more work before full release. Ideally a custom K2 node and some smarter helper classes for C++.
	We now have a single WriteTo/ReadFromDataChannel call which handles all of the setup work.
	Though we still have a clunky need to have typed WriteXXX calls for every value with manual name typing :(

 - Removed specific spawn DI in favour of adding Spawn functions to the Read DI. While this bloats the read DI it will reduce the number of DIs required in actual content.
	Also renamed the Spawn function to SpawnFromSpawnInfo to be more clear.

- Asset Refactor
	Removed NiagaraDataChannelDefinitions.
	All Niagara Data Channels have their own asset which are now Primary assets so that they can be cooked and loaded even if only referenced by code.
	When creating DataChannel assets, their directory will automatically be added to ini config files. Specifically, the PrimaryAsset search paths so they are cooked.
	All Data Channel References are now direct to these assets rather than obfruscated via a Channel Name.

- Niagara DI functions with variadic pins can also now provide pin actions for the Add pin menu.
	This allows you to populate the add parameter menu as you see fit with categories of variables with specific tooltips etc.
	The Data Channels are using this to add all Data Channels as categories to the pin menu for the Read and Write variadic pins.
	Guiding usage and allowing you to select a particular parameter in a more discoverable way.

- Allowing FNiagaraPositions to be used in data channels + special case code to allow full fat LWC at game level but SWC simulation types elsewhere.

- Refactored VM runtime handling of Variadics params out to some helper classes so other DIs should be able to use them fairly easily.

 - Misc:
	- Some special case handling for Enum values in Data Channels to avoid some things converting them to Int32.
	- Modified Customization for FNiagaraVariableBase so that it works better in arrays etc.
	- Widened the range of types allowed by [REDACTED_PERSON].
	- Various details customizaitons
	- VM External Function input handlers now store their start address so they can be reset an iterate over the same register data again.
	- Some boiler plate change with UNiagaraNodeWithDynamicPins::AddParameter
	- Misc Tooltip and comments improvements.
	- Fixed issue with variadic functions and signature equality checks
	- Refactored some files/classes between private and public.
	- World manger now has a pre actor tick required to call DataChannel manaager BeginFrame.
	- Improved reporting and feedback for errors.
	- Read DI now has a SpawnInfo Variable property allowing you to customize
	- Quick and dirty Total's graph to the debug hud perf graphs.
	- Exposing EPhysicalSurface to Niagara.
	- Fixed issue when writing structs to from Game level data channels to the simulation DataSets.
	- Preventing Data Channel DIs from overlapping frames as they need to sync up correctly with the manager's begin/end frame calls.
	- FNDIInputParam can now be initialized after it's created which makes certain things easier to write in VM DI functions.
	- Enabled Particle Caputre in non shipping builds to support cooked/client side code tests in upcomming CL.

diff --git a/Engine/Plugins/FX/Niagara/Shaders/Private/DataChannel/NiagaraDataInterfaceDataChannelTemplate_Read.ush b/Engine/Plugins/FX/Niagara/Shaders/Private/DataChannel/NiagaraDataInterfaceDataChannelTemplate_Read.ush
--- a/Engine/Plugins/FX/Niagara/Shaders/Private/DataChannel/NiagaraDataInterfaceDataChannelTemplate_Read.ush
+++ b/Engine/Plugins/FX/Niagara/Shaders/Private/DataChannel/NiagaraDataInterfaceDataChannelTemplate_Read.ush
@@ -5,7 +5,7 @@
  * Contained in it's own shader file to allow reuse between different invocations with different params.
  */
 
-void Read_{ParameterName}(int ElementIndex, out bool bOutSuccess, {FunctionParameters})
+void {FunctionSymbol}(int ElementIndex, out bool bOutSuccess, {FunctionParameters})
 {
 	if(ElementIndex >= 0 && ElementIndex < {ParameterName}_Num)
 	{
diff --git a/Engine/Plugins/FX/Niagara/Source/Niagara/Classes/NiagaraSystem.h b/Engine/Plugins/FX/Niagara/Source/Niagara/Classes/NiagaraSystem.h
--- a/Engine/Plugins/FX/Niagara/Source/Niagara/Classes/NiagaraSystem.h
+++ b/Engine/Plugins/FX/Niagara/Source/Niagara/Classes/NiagaraSystem.h
@@ -48,7 +48,6 @@ struct FNiagaraGraphCachedDataBase;
 #endif
 
 class UNiagaraDataChannel;
-class UNiagaraDataChannelDefinitions;
 
 USTRUCT()
 struct FNiagaraEmitterCompiledData
@@ -87,10 +86,11 @@ struct FNiagaraEmitterCompiledData
 	UPROPERTY()
 	FNiagaraDataSetCompiledData DataSetCompiledData;
 
-#if WITH_EDITORONLY_DATA
-	UPROPERTY()
-	FNiagaraDataSetCompiledData GPUCaptureDataSetCompiledData;
-#endif
+	#if NIAGARA_SYSTEM_CAPTURE
+	/** Debug only data for caputuring from GPU sims, laziliy initialized when needed via GetGPUCaptureDataSetCompiledData().*/
+	mutable FNiagaraDataSetCompiledData GPUCaptureDataSetCompiledData;
+	const FNiagaraDataSetCompiledData& GetGPUCaptureDataSetCompiledData()const;
+	#endif
 };
 
 USTRUCT()
@@ -754,9 +754,6 @@ public:
 	/** Resets internal data leaving it in a state which would have minimal cost to exist in headless builds (servers) */
 	void ResetToEmptySystem();
 
-	/** Registers that this System uses the passed DI. No symmetrical unregister as this is cleared and rebuilt on compile. */
-	void RegisterDataChannelUse(const UNiagaraDataChannel* DataChannel);
-
 	/** Updates any post compile data based upon data interfaces. */
 	void OnCompiledDataInterfaceChanged();
 
@@ -933,9 +930,6 @@ protected:
 	UPROPERTY()
 	bool bNeedsGPUContextInitForDataInterfaces;
 
-	UPROPERTY()
-	TArray<TObjectPtr<UNiagaraDataChannelDefinitions>> ReferencedDataChannelDefinitions;
-
 	/** Array of emitter indices sorted by execution priority. The emitters will be ticked in this order. Please note that some indices may have the top bit set (kStartNewOverlapGroupBit)
 	* to indicate synchronization points in parallel execution, so mask it out before using the values as indices in the emitters array.
 	*/
diff --git a/Engine/Plugins/FX/Niagara/Source/Niagara/Internal/DataInterface/NiagaraDataInterfaceDataChannelCommon.h b/Engine/Plugins/FX/Niagara/Source/Niagara/Internal/DataInterface/NiagaraDataInterfaceDataChannelCommon.h
--- a/Engine/Plugins/FX/Niagara/Source/Niagara/Internal/DataInterface/NiagaraDataInterfaceDataChannelCommon.h
+++ b/Engine/Plugins/FX/Niagara/Source/Niagara/Internal/DataInterface/NiagaraDataInterfaceDataChannelCommon.h
@@ -6,6 +6,7 @@
 #include "NiagaraCompileHash.h"
 #include "Misc/LazySingleton.h"
 #include "NiagaraDataChannelCommon.h"
+#include "NiagaraDataInterface.h"
 #include "NiagaraDataInterfaceDataChannelCommon.generated.h"
 
 UENUM()
@@ -94,7 +95,7 @@ struct FNDIDataChannel_FunctionToDataSetBinding
 	FNiagaraDataSetCompiledData DebugCompiledData;
 	#endif
 
-	FNDIDataChannel_FunctionToDataSetBinding(const FNDIDataChannelFunctionInfo& FunctionInfo, const FNiagaraDataSetCompiledData& DataSetLayout);
+	FNDIDataChannel_FunctionToDataSetBinding(const FNDIDataChannelFunctionInfo& FunctionInfo, const FNiagaraDataSetCompiledData& DataSetLayout, TArray<FNiagaraVariableBase>& OutMissingParams);
 
 	bool IsValid()const { return DataSetLayoutHash != 0; }
 	void GenVMBindings(const FNiagaraVariableBase& Var, const UStruct* Struct, uint32& FuncFloatRegister, uint32& FuncIntRegister, uint32& FuncHalfRegister, uint32& DataSetFloatRegister, uint32& DataSetIntRegister, uint32& DataSetHalfRegister);
@@ -136,7 +137,7 @@ public:
 	/** 
 	Retrieves, or generates, the layout information that maps from the given function to the data in the given dataset.
 	*/
-	FNDIDataChannel_FuncToDataSetBindingPtr GetLayoutInfo(const FNDIDataChannelFunctionInfo& FunctionInfo, const FNiagaraDataSetCompiledData& DataSetLayout);
+	FNDIDataChannel_FuncToDataSetBindingPtr GetLayoutInfo(const FNDIDataChannelFunctionInfo& FunctionInfo, const FNiagaraDataSetCompiledData& DataSetLayout, TArray<FNiagaraVariableBase>& OutMissingParams);
 };
 
 /** A sorted table of parameters accessed by each GPU script */
@@ -203,4 +204,151 @@ protected:
 namespace NDIDataChannelUtilities
 {
 	void SortParameters(TArray<FNiagaraVariableBase>& Parameters);	
-}
\ No newline at end of file
+}
+
+
+
+/** Handles any number of variadic parameter inputs. */
+template<int32 EXPECTED_NUM_INPUTS>
+struct FNDIVariadicInputHandler
+{
+	TArray<FNDIInputParam<float>, TInlineAllocator<EXPECTED_NUM_INPUTS>> FloatInputs;
+	TArray<FNDIInputParam<int32>, TInlineAllocator<EXPECTED_NUM_INPUTS>> IntInputs;
+	TArray<FNDIInputParam<FFloat16>, TInlineAllocator<EXPECTED_NUM_INPUTS>> HalfInputs;
+
+	FNDIVariadicInputHandler(FVectorVMExternalFunctionContext& Context, const FNDIDataChannel_FunctionToDataSetBinding* BindingPtr)
+	{
+		if(BindingPtr)
+		{
+			FloatInputs.Reserve(BindingPtr->NumFloatComponents);
+			IntInputs.Reserve(BindingPtr->NumInt32Components);
+			HalfInputs.Reserve(BindingPtr->NumHalfComponents);
+			for (const FNDIDataChannelRegisterBinding& VMBinding : BindingPtr->VMRegisterBindings)
+			{
+				switch (VMBinding.DataType)
+				{
+				case (int32)ENiagaraBaseTypes::Float: FloatInputs.Emplace(Context); break;
+				case (int32)ENiagaraBaseTypes::Int32: IntInputs.Emplace(Context); break;
+				case (int32)ENiagaraBaseTypes::Bool: IntInputs.Emplace(Context); break;
+				case (int32)ENiagaraBaseTypes::Half: HalfInputs.Emplace(Context); break;
+				default: check(0);
+				};
+			}
+		}
+	}
+
+	void Reset()
+	{
+		for (FNDIInputParam<float>& Input : FloatInputs) { Input.Reset(); }
+		for (FNDIInputParam<int32>& Input : IntInputs) { Input.Reset(); }
+		for (FNDIInputParam<FFloat16>& Input : HalfInputs) { Input.Reset(); }
+	}
+
+	void Advance()
+	{
+		for (FNDIInputParam<float>& Input : FloatInputs) { Input.Advance(); }
+		for (FNDIInputParam<int32>& Input : IntInputs) { Input.Advance(); }
+		for (FNDIInputParam<FFloat16>& Input : HalfInputs) { Input.Advance(); }
+	}
+
+	template<typename TFloatAction, typename TIntAction, typename THalfAction>
+	bool Process(bool bProcess, const FNDIDataChannel_FunctionToDataSetBinding* BindingInfo, TFloatAction FloatFunc, TIntAction IntFunc, THalfAction HalfFunc)
+	{
+		if (BindingInfo && bProcess)
+		{
+			//TODO: Optimize for long runs of writes to reduce binding/lookup overhead.
+			for (const FNDIDataChannelRegisterBinding& VMBinding : BindingInfo->VMRegisterBindings)
+			{
+				switch(VMBinding.DataType)
+				{
+					case (int32)ENiagaraBaseTypes::Float: FloatFunc(VMBinding, FloatInputs[VMBinding.FunctionRegisterIndex].GetAndAdvance()); break;
+					case (int32)ENiagaraBaseTypes::Int32: IntFunc(VMBinding, IntInputs[VMBinding.FunctionRegisterIndex].GetAndAdvance()); break;
+					case (int32)ENiagaraBaseTypes::Bool: IntFunc(VMBinding, IntInputs[VMBinding.FunctionRegisterIndex].GetAndAdvance()); break;
+					case (int32)ENiagaraBaseTypes::Half: HalfFunc(VMBinding, HalfInputs[VMBinding.FunctionRegisterIndex].GetAndAdvance()); break;
+					default: check(0);
+				};
+			}
+			return true;
+		}
+		
+		Advance();
+		return false;
+	}
+};
+
+
+template<int32 EXPECTED_NUM_INPUTS>
+struct FNDIVariadicOutputHandler
+{
+	TArray<VectorVM::FExternalFuncRegisterHandler<float>, TInlineAllocator<EXPECTED_NUM_INPUTS>> FloatOutputs;
+	TArray<VectorVM::FExternalFuncRegisterHandler<int32>, TInlineAllocator<EXPECTED_NUM_INPUTS>> IntOutputs;
+	TArray<VectorVM::FExternalFuncRegisterHandler<FFloat16>, TInlineAllocator<EXPECTED_NUM_INPUTS>> HalfOutputs;
+
+	FNDIVariadicOutputHandler(FVectorVMExternalFunctionContext& Context, const FNDIDataChannel_FunctionToDataSetBinding* BindingPtr)
+	{
+		//Parse the VM bytecode inputs in order, mapping them to the correct DataChannel data 
+		if(BindingPtr)
+		{
+			FloatOutputs.Reserve(BindingPtr->NumFloatComponents);
+			IntOutputs.Reserve(BindingPtr->NumInt32Components);
+			HalfOutputs.Reserve(BindingPtr->NumHalfComponents);
+			for (const FNDIDataChannelRegisterBinding& VMBinding : BindingPtr->VMRegisterBindings)
+			{
+				switch (VMBinding.DataType)
+				{
+				case (int32)ENiagaraBaseTypes::Float: FloatOutputs.Emplace(Context); break;
+				case (int32)ENiagaraBaseTypes::Int32: IntOutputs.Emplace(Context); break;
+				case (int32)ENiagaraBaseTypes::Bool: IntOutputs.Emplace(Context); break;
+				case (int32)ENiagaraBaseTypes::Half: HalfOutputs.Emplace(Context); break;
+				default: check(0);
+				};
+			}
+		}
+	}
+
+	template<typename TFloatAction, typename TIntAction, typename THalfAction>
+	bool Process(bool bProcess, const FNDIDataChannel_FunctionToDataSetBinding* BindingInfo, TFloatAction FloatFunc, TIntAction IntFunc, THalfAction HalfFunc)
+	{
+		if (BindingInfo && bProcess)
+		{
+			//TODO: Optimize for long runs of writes to reduce binding/lookup overhead.
+			for (const FNDIDataChannelRegisterBinding& VMBinding : BindingInfo->VMRegisterBindings)
+			{
+				switch (VMBinding.DataType)
+				{
+				case (int32)ENiagaraBaseTypes::Float: FloatFunc(VMBinding, FloatOutputs[VMBinding.FunctionRegisterIndex].GetDestAndAdvance()); break;
+				case (int32)ENiagaraBaseTypes::Int32: IntFunc(VMBinding, IntOutputs[VMBinding.FunctionRegisterIndex].GetDestAndAdvance()); break;
+				case (int32)ENiagaraBaseTypes::Bool: IntFunc(VMBinding, IntOutputs[VMBinding.FunctionRegisterIndex].GetDestAndAdvance()); break;
+				case (int32)ENiagaraBaseTypes::Half: HalfFunc(VMBinding, HalfOutputs[VMBinding.FunctionRegisterIndex].GetDestAndAdvance()); break;
+				default: check(0);
+				};
+			}
+			return true;
+		}
+
+		Fallback(1);
+
+		return false;
+	}
+
+	void Fallback(int32 Count)
+	{
+		for (int32 OutIdx = 0; OutIdx < FloatOutputs.Num(); ++OutIdx)
+		{
+			FMemory::Memzero(FloatOutputs[OutIdx].GetDest(), sizeof(float) * Count);
+			FloatOutputs[OutIdx].Advance(Count);
+		}
+
+		for (int32 OutIdx = 0; OutIdx < IntOutputs.Num(); ++OutIdx)
+		{
+			FMemory::Memzero(IntOutputs[OutIdx].GetDest(), sizeof(int32) * Count);
+			IntOutputs[OutIdx].Advance(Count);
+		}
+
+		for (int32 OutIdx = 0; OutIdx < HalfOutputs.Num(); ++OutIdx)
+		{
+			FMemory::Memzero(HalfOutputs[OutIdx].GetDest(), sizeof(FFloat16) * Count);
+			HalfOutputs[OutIdx].Advance(Count);
+		}
+	}
+};
\ No newline at end of file
diff --git a/Engine/Plugins/FX/Niagara/Source/Niagara/Internal/DataInterface/NiagaraDataInterfaceDataChannelRead.h b/Engine/Plugins/FX/Niagara/Source/Niagara/Internal/DataInterface/NiagaraDataInterfaceDataChannelRead.h
--- a/Engine/Plugins/FX/Niagara/Source/Niagara/Internal/DataInterface/NiagaraDataInterfaceDataChannelRead.h
+++ b/Engine/Plugins/FX/Niagara/Source/Niagara/Internal/DataInterface/NiagaraDataInterfaceDataChannelRead.h
@@ -24,6 +24,8 @@ There is working code for this already but the concept/API needs more fleshing o
 
 #include "NiagaraDataInterfaceDataChannelCommon.h"
 #include "NiagaraDataInterface.h"
+#include "NiagaraDataInterfaceEmitterBinding.h"
+#include "NiagaraDataSetAccessor.h"
 #include "NiagaraDataInterfaceDataChannelRead.generated.h"
 
 class UNiagaraDataInterfaceDataChannelWrite;
@@ -31,6 +33,24 @@ struct FNDIDataChannelWriteInstanceData;
 class UNiagaraDataChannelHandler;
 class FNiagaraDataBuffer;
 
+struct FNiagaraDataChannelDataProxy;
+
+/** Mode controlling the behavior of the Spawn function for Niagara's Data Channel Read Data Interface.*/
+UENUM()
+enum class ENDIDataChannelSpawnMode
+{
+	/** This mode will override any previous spawn counts set or accumulated this tick. */
+	Override,
+
+	/** This mode will accumulate spawn counts with other spawns. */
+	Accumulate,
+
+	/** Does nothing. This mode is useful for dynamically skipping spawning based on script logic. */
+	None,
+
+	Max UMETA(Hidden),
+};
+
 UCLASS(Experimental, EditInlineNew, Category = "Data Channels", CollapseCategories, meta = (DisplayName = "Data Channel Reader"))
 class NIAGARA_API UNiagaraDataInterfaceDataChannelRead : public UNiagaraDataInterface
 {
@@ -50,15 +70,37 @@ public:
 
 	/** When reading from external, the channel to consume. */
 	UPROPERTY(EditAnywhere, Category="Data Channel")
-	FNiagaraDataChannelReference Channel;
+	TObjectPtr<UNiagaraDataChannelAsset> Channel;
 	
+	/** A bounds emitter instance when using functions like Spawn. Defaults to Self. */
+	UPROPERTY(EditAnywhere, Category = "Spawning", AdvancedDisplay)
+	FNiagaraDataInterfaceEmitterBinding EmitterBinding;
+
 	/** True if this reader will read the current frame's data. If false, we read the previous frame.
 	* Reading the current frame introduces a tick order dependency but allows for zero latency reads. Any data channel elements that are generated after this reader is used are missed.
 	* Reading the previous frame's data introduces a frame of latency but ensures we never miss any data as we have access to the whole frame.
 	*/
-	UPROPERTY(EditAnywhere, Category = "Data Channel")
+	UPROPERTY(EditAnywhere, Category = "Data Channel", AdvancedDisplay)
 	bool bReadCurrentFrame = false;
 
+	/**
+	Whether this DI should request updated source data from the Data Channel each tick.
+	Some Data Channels have multiple separate source data elements for things such as spatial subdivision. 
+	Each DI will request the correct one for it's owning system instance from the data channel. 
+	Depending on the data channel this could be an expensive search so we should avoid doing this every tick if possible.
+	*/
+	UPROPERTY(EditAnywhere, Category = "Data Channel", AdvancedDisplay)
+	bool bUpdateSourceDataEveryTick = true;
+
+
+	/** When true, Emitter.Spawn group for any spawned particles will be overridden to the index of the data channel element that generated that spawn. This allows particles to read further data from the data channel element that generated them. */
+	UPROPERTY(EditAnywhere, Category = "Spawning", AdvancedDisplay)
+	bool bOverrideSpawnGroupToDataChannelIndex = true;
+
+	/** The spawn info variable we'll read from the data channel to control spawning using the SpawnFromSpawnInfo DI function.*/
+	UPROPERTY(EditAnywhere, Category = "Spawning", AdvancedDisplay)
+	FName SpawnInfoName;
+	
 	//UObject Interface
 	virtual void PostInitProperties() override;
 	virtual void BeginDestroy() override;
@@ -92,17 +134,26 @@ public:
 	virtual int32 PerInstanceDataSize() const override;
 	virtual bool HasPreSimulateTick() const override { return true; }
 	virtual bool HasPostSimulateTick() const override { return true; }
+	virtual bool HasPostStageTick(ENiagaraScriptUsage Usage) const override { return Usage == ENiagaraScriptUsage::EmitterUpdateScript || Usage == ENiagaraScriptUsage::SystemUpdateScript; }
 	virtual bool PerInstanceTick(void* PerInstanceData, FNiagaraSystemInstance* SystemInstance, float DeltaSeconds) override;
 	virtual bool PerInstanceTickPostSimulate(void* PerInstanceData, FNiagaraSystemInstance* SystemInstance, float DeltaSeconds) override;
+	virtual void PostStageTick(FNDICpuPostStageContext& Context)override;
 	virtual void ProvidePerInstanceDataForRenderThread(void* DataForRenderThread, void* PerInstanceData, const FNiagaraSystemInstanceID& SystemInstance) override;
 	virtual void GetEmitterDependencies(UNiagaraSystem* Asset, TArray<FVersionedNiagaraEmitter>& Dependencies) const;
+	
+	//We cannot overlap frames as we must correctly sync up with the data channel manager on Begin/End frame etc.
+	virtual bool PostSimulateCanOverlapFrames() const { return false; }
 	//UNiagaraDataInterface Interface
 
 	//Functions usable anywhere.
 	void Num(FVectorVMExternalFunctionContext& Context);
 	void Read(FVectorVMExternalFunctionContext& Context, int32 FuncIdx);
 	void Consume(FVectorVMExternalFunctionContext& Context, int32 FuncIdx);
-	
+
+	//Emitter only functions.
+	void SpawnFromSpawnInfo(FVectorVMExternalFunctionContext& Context);
+	void SpawnConditional(FVectorVMExternalFunctionContext& Context, int32 FuncIndex);
+
 	FNDIDataChannelCompiledData& GetCompiledData() { return CompiledData; }
 
 protected:
@@ -125,14 +176,9 @@ struct FNDIDataChannelReadInstanceData
 	/** Pointer to the world DataChannel Channel we're reading from, if any. */
 	TWeakObjectPtr<UNiagaraDataChannelHandler> DataChannel;
 
-	/** Pointer to the cpu simulation buffer from the DataChannel if we're reading from there. */
-	FNiagaraDataBufferRef ExternalCPU = nullptr;
+	/** Shared pointer to the actual data channel data pulled from the data channel. */
+	FNiagaraDataChannelDataPtr DataChannelData = nullptr;
 	
-	/** The GPU data set from the DataChannel. We must grab the whole set rather than a buffer here as we don't know which buffer will be contain this frame's channel data. 
-	*	TODO: We'll likely want to improve this GPU sim flow soon. Possibly grab an RT proxy object from the data channel handler and pass that over to the GPU instead of a direct dataset like this.
-	*/
-	FNiagaraDataSet* ExternalGPU = nullptr;
-
 	/** Cached hash to check if the layout of our source data has changed. */
 	uint64 ChachedDataSetLayoutHash = INDEX_NONE;
 
@@ -154,8 +200,22 @@ struct FNDIDataChannelReadInstanceData
 	*/
 	std::atomic<int32> ConsumeIndex = 0;
 
+	/** 
+	Conditional spawns to inject into the bound emitter based on our DI function calls.
+	Each data channel element we read in a particular execution will have an entry in this array and can spawn(or not) independently.
+	*/
+	TArray<FNiagaraSpawnInfo> ConditionalSpawns;
+
+	/*
+	We read a spawn info for future proofing here but currently we only use the spawn count.
+	NOTE: This is 100% dependent on SpawnInfo class not changing. If it does we'll need to update this.
+	*/
+	FNiagaraDataSetAccessor<FNiagaraSpawnInfo> SpawnInfoAccessor;
+	uint32 CachedLayoutHash = INDEX_NONE;
+	FNiagaraEmitterInstance* EmitterInstance = nullptr;
+
 	virtual ~FNDIDataChannelReadInstanceData();
-	FNiagaraDataBuffer* GetReadBufferCPU();
+	FNiagaraDataBuffer* GetReadBufferCPU(bool bPrevFrame);
 	bool Init(UNiagaraDataInterfaceDataChannelRead* Interface, FNiagaraSystemInstance* Instance);
 	bool Tick(UNiagaraDataInterfaceDataChannelRead* Interface, FNiagaraSystemInstance* Instance);
 	bool PostTick(UNiagaraDataInterfaceDataChannelRead* Interface, FNiagaraSystemInstance* Instance);
@@ -171,7 +231,9 @@ struct FNiagaraDataInterfaceProxy_DataChannelRead : public FNiagaraDataInterface
 	{
 		//GPU Dataset from the channel handler. We'll grab the current buffer from this on the RT.
 		//This must be grabbed fresh from the handler each frame as it's lifetime cannot be ensured.
-		FNiagaraDataSet* ChannelDataSet = nullptr;
+		FNiagaraDataChannelDataProxy* ChannelDataRTProxy = nullptr;
+
+		bool bReadPrevFrame = false;
 
 		/**
 		A buffer containing layout information needed to access parameters for each script using this DI.
diff --git a/Engine/Plugins/FX/Niagara/Source/Niagara/Internal/DataInterface/NiagaraDataInterfaceDataChannelWrite.h b/Engine/Plugins/FX/Niagara/Source/Niagara/Internal/DataInterface/NiagaraDataInterfaceDataChannelWrite.h
--- a/Engine/Plugins/FX/Niagara/Source/Niagara/Internal/DataInterface/NiagaraDataInterfaceDataChannelWrite.h
+++ b/Engine/Plugins/FX/Niagara/Source/Niagara/Internal/DataInterface/NiagaraDataInterfaceDataChannelWrite.h
@@ -60,10 +60,19 @@ public:
 	/** Whether the data generated by this DI should be published to a world data channel. This is required to allow GPU emitters in other Niagara Systems to read this data.  */
 	UPROPERTY(EditAnywhere, Category = "Data Channel")
 	bool bPublishToGPU = false;
+	
+	/**
+	Whether this DI should request updated destination data from the Data Channel each tick.
+	Some Data Channels have multiple separate data elements for things such as spacial subdivision. 
+	Each DI will request the correct one for it's owning system instance from the data channel. 
+	Depending on the data channel this could be an expensive search so we should avoid doing this every tick if possible.
+	*/
+	UPROPERTY(EditAnywhere, Category = "Data Channel", AdvancedDisplay)
+	bool bUpdateDestinationDataEveryTick = true;
 
 	/** When writing externally, the channel to use. */
 	UPROPERTY(EditAnywhere, Category = "Data Channel", meta = (EditCondition = "bPublishToWorld"))
-	FNiagaraDataChannelReference Channel;
+	TObjectPtr<UNiagaraDataChannelAsset> Channel;
 
 	//UObject Interface
 	virtual void PostInitProperties() override;
@@ -84,6 +93,11 @@ public:
 	virtual void BuildShaderParameters(FNiagaraShaderParametersBuilder& ShaderParametersBuilder) const override;
 	virtual void SetShaderParameters(const FNiagaraDataInterfaceSetShaderParametersContext& Context) const override;
 
+#if WITH_EDITOR
+	virtual void GetFeedback(UNiagaraSystem* InAsset, UNiagaraComponent* InComponent, TArray<FNiagaraDataInterfaceError>& OutErrors, TArray<FNiagaraDataInterfaceFeedback>& OutWarnings, TArray<FNiagaraDataInterfaceFeedback>& OutInfo);
+	virtual void ValidateFunction(const FNiagaraFunctionSignature& Function, TArray<FText>& OutValidationErrors);
+#endif
+
 	virtual bool Equals(const UNiagaraDataInterface* Other) const override;
 
 	virtual bool InitPerInstanceData(void* PerInstanceData, FNiagaraSystemInstance* SystemInstance) override;
@@ -94,6 +108,9 @@ public:
 	virtual bool PerInstanceTick(void* PerInstanceData, FNiagaraSystemInstance* SystemInstance, float DeltaSeconds) override;
 	virtual bool PerInstanceTickPostSimulate(void* PerInstanceData, FNiagaraSystemInstance* SystemInstance, float DeltaSeconds) override;
 	virtual void ProvidePerInstanceDataForRenderThread(void* DataForRenderThread, void* PerInstanceData, const FNiagaraSystemInstanceID& SystemInstance) override;
+
+	//We cannot overlapt frames as we must correctly sync up with the data channel manager on Begin/End frame etc.
+	virtual bool PostSimulateCanOverlapFrames() const { return false; }
 	//UNiagaraDataInterface Interface
 
 	void Num(FVectorVMExternalFunctionContext& Context);
diff --git a/Engine/Plugins/FX/Niagara/Source/Niagara/Internal/NiagaraDataChannelCommon.h b/Engine/Plugins/FX/Niagara/Source/Niagara/Internal/NiagaraDataChannelCommon.h
--- a/Engine/Plugins/FX/Niagara/Source/Niagara/Internal/NiagaraDataChannelCommon.h
+++ b/Engine/Plugins/FX/Niagara/Source/Niagara/Internal/NiagaraDataChannelCommon.h
@@ -2,8 +2,10 @@
 
 #pragma once
 
+#include "CoreMinimal.h"
+#include "UObject/ObjectMacros.h"
 #include "NiagaraDataSet.h"
-#include "NiagaraDataChannelCommon.generated.h"
+#include "NiagaraDataChannelPublic.h"
 
 class UWorld;
 class UActorComponentc;
@@ -12,23 +14,121 @@ class FNiagaraDataBuffer;
 class FNiagaraSystemInstance;
 class UNiagaraDataChannel;
 class UNiagaraDataChannelHandler;
-class UNiagaraDataChannelDefinitions;
 class UNiagaraDataInterfaceDataChannelWrite;
 class UNiagaraDataInterfaceDataChannelRead;
-class UNiagaraDataChannelWriter;
-class UNiagaraDataChannelReader;
-struct FNiagaraDataChannelGameData;
+struct FNiagaraDataChannelDataProxy;
 
-typedef TSharedPtr<FNiagaraDataChannelGameData> FNiagaraDataChannelGameDataPtr;
 
-/** Wrapper struct for FNames referencing Niagara Data Channels allowing a type customization. */
-USTRUCT(BlueprintType)
-struct FNiagaraDataChannelReference
+/** Buffer containing a single FNiagaraVariable's data at the game level. AoS layout. LWC Types. */
+struct FNiagaraDataChannelVariableBuffer
 {
-	GENERATED_BODY()
+	TArray<uint8> Data;
+	TArray<uint8> PrevData;
+	int32 Size = 0;
 
-	UPROPERTY(EditAnywhere, Category = "Data Channel")
-	FName ChannelName = NAME_None;
+	void Init(const FNiagaraVariableBase& Var)
+	{
+		//Position types are a special case where we have to store an LWC Vector in game level data and convert to a simulation friendly FVector3f as it enters Niagara.
+		if (Var.GetType() == FNiagaraTypeDefinition::GetPositionDef())
+		{
+			Size = FNiagaraTypeHelper::GetVectorDef().GetSize();
+		}
+		else
+		{
+			Size = Var.GetSizeInBytes();
+		}
+	}
+
+	void Empty() { Data.Empty(); }
+
+	void BeginFrame(bool bKeepPrevious)
+	{
+		if (bKeepPrevious)
+		{
+			Swap(Data, PrevData);
+		}
+		Data.Reset();
+	}
+
+	template<typename T>
+	bool Write(int32 Index, const T& InData)
+	{
+		check(sizeof(T) == Size);
+		if (Index >= 0 && Index < Num())
+		{
+			T* Dest = ((T*)Data.GetData()) + Index;
+			*Dest = InData;
+
+			return true;
+		}
+		return false;
+	}
+
+	template<typename T>
+	bool Read(int32 Index, T& OutData, bool bPreviousFrameData)const
+	{
+		check(sizeof(T) == Size);
+
+		//Should we fallback to previous if we don't have current?
+		//bPreviousFrameData |= Data.Num() == 0;
+
+		int32 NumElems = bPreviousFrameData ? PrevNum() : Num();
+		if (Index >= 0 && Index < NumElems)
+		{
+			const uint8* DataPtr = bPreviousFrameData ? PrevData.GetData() : Data.GetData();
+			T* Src = ((T*)DataPtr) + Index;
+			OutData = *Src;
+
+			return true;
+		}
+		OutData = T();
+		return false;
+	}
+
+	void SetNum(int32 Num)
+	{
+		Data.SetNumZeroed(Size * Num);
+	}
+
+	int32 Num()const { return Data.Num() / Size; }
+	int32 PrevNum()const { return PrevData.Num() / Size; }
+
+	int32 GetElementSize()const { return Size; }
+};
+
+/** Storage for game level DataChannels generated by BP / C++ */
+struct FNiagaraDataChannelGameData : public TSharedFromThis<FNiagaraDataChannelGameData>
+{
+private:
+
+	/** Per variable storage buffers. */
+	TArray<FNiagaraDataChannelVariableBuffer> VariableData;
+
+	int32 NumElements = 0;
+	int32 PrevNumElements = 0;
+
+	TWeakObjectPtr<const UNiagaraDataChannel> DataChannel;
+
+public:
+
+	void Init(const UNiagaraDataChannel* InDataChannel);
+	void Empty();
+	void BeginFrame();
+	void SetNum(int32 NewNum);
+	int32 Num() const { return NumElements; }
+	int32 PrevNum()const { return PrevNumElements; }
+
+	FNiagaraDataChannelVariableBuffer* FindVariableBuffer(const FNiagaraVariableBase& Var);
+
+	void WriteToDataSet(FNiagaraDataBuffer* DestBuffer, int32 DestStartIdx, FVector3f SimulationLwcTile);
+
+	void AppendFromGameData(const FNiagaraDataChannelGameData& GameData);
+
+	void AppendFromDataSet(const FNiagaraDataBuffer* SrcBuffer, FVector3f SimulationLwcTile);
+
+	const UNiagaraDataChannel* GetDataChannel()const { return DataChannel.Get(); }
+
+	const TConstArrayView<FNiagaraDataChannelVariableBuffer> GetVariableBuffers()const { return VariableData; }
 };
 
 /** A request to publish data into a Niagara Data Channel.  */
@@ -38,7 +138,7 @@ struct FNiagaraDataChannelPublishRequest
 	FNiagaraDataBufferRef Data;
 
 	/** Game level data if this request comes from the game code. */
-	FNiagaraDataChannelGameDataPtr GameData;
+	TSharedPtr<FNiagaraDataChannelGameData> GameData = nullptr;
 
 	/** If true, data in this request will be made visible to BP and C++ game code.*/
 	bool bVisibleToGame = false;
@@ -66,3 +166,85 @@ struct FNiagaraDataChannelPublishRequest
 	{
 	}
 };
+
+
+/**
+Underlying storage class for data channel data.
+Some data channels will have many of these and can distribute them as needed to different accessing systems.
+For example, some data channel handlers may subdivide the scene such that distant systems are not interacting.
+In this case, each subdivision would have it's own FNiagaraDataChannelData and distribute these to the relevant NiagaraSystems.
+*/
+struct NIAGARA_API FNiagaraDataChannelData final
+{
+	UE_NONCOPYABLE(FNiagaraDataChannelData)
+	explicit FNiagaraDataChannelData(UNiagaraDataChannelHandler* Owner);
+	~FNiagaraDataChannelData();
+
+	void Reset();
+
+	void BeginFrame(UNiagaraDataChannelHandler* Owner);
+	void EndFrame(UNiagaraDataChannelHandler* Owner);
+	void Tick(UNiagaraDataChannelHandler* Owner);
+	void ConsumePublishRequests(UNiagaraDataChannelHandler* Owner);
+
+	FNiagaraDataChannelGameData* GetGameData();
+	FNiagaraDataBufferRef GetCPUData(bool bPreviousFrame);
+	FNiagaraDataChannelDataProxy* GetRTProxy(){ return RTProxy.Get(); }
+	
+	/** Adds a request to publish some data into the channel on the next tick. */
+	void Publish(const FNiagaraDataChannelPublishRequest& Request);
+		
+	/**
+	 *Removes all publish requests involving the given dataset.
+	 *TODO: REMOVE
+	 *This is a hack to get around lifetime issues wrt data buffers/datasets and their compiled data.
+	 *We should rework things such that data buffers can exist beyond their owning dataset, including the compiled data detailing their layout.
+	 **/
+	void RemovePublishRequests(const FNiagaraDataSet* DataSet);
+
+	const FNiagaraDataSetCompiledData& GetCompiledData(ENiagaraSimTarget SimTarget);
+
+	void SetLwcTile(FVector3f InLwcTile){ LwcTile = InLwcTile; }
+	FVector3f GetLwcTile()const { return LwcTile; }
+private:
+
+	/** DataChannel data accessible from Game/BP. AoS Layout. LWC types. */
+	FNiagaraDataChannelGameDataPtr GameData;
+
+	//		▲	CPU Sim Data can optionally be made visible to the Game Data.
+	//		|
+	//		|
+	//		|
+	//		▼	Game/BP Data can optionally be made visible to CPU sims.
+
+	/** DataChannel data accessible to Niagara CPU sims. SoA layout. Non LWC types. */
+	FNiagaraDataSet* CPUSimData = nullptr;
+
+	// Cached off buffer with the previous frame's CPU Sim accessible data.
+	// Some systems can choose to read this to avoid any current frame tick ordering issues.
+	FNiagaraDataBufferRef PrevCPUSimData = nullptr;
+
+	/** Dataset we use for staging game data for the consumption by RT/GPU sims. */
+	FNiagaraDataSet* GameDataStaging = nullptr;
+
+	//		▲	GPU Sim Data can optionally be made visible to CPU Sims and Game Data.
+	//		|
+	//		|
+	//		|
+	//		▼	CPU Sim and Game Data can optionally be made visible to the GPU.
+
+	/** DataChannel data accessible to Niagara GPU sims. SoA layout. Non LWC types. */
+	FNiagaraDataSet* GPUSimData = nullptr;
+
+
+	/** Data buffers we'll be passing to the RT proxy for uploading to the GPU */
+	TArray<FNiagaraDataBufferRef> BuffersForGPU;
+
+	/** Render thread proxy for this data. Owns all RT side data meant for GPU simulations. */
+	TUniquePtr<FNiagaraDataChannelDataProxy> RTProxy;
+
+	/** Pending requests to publish data into this data channel. These requests are consumed at tick tick group. */
+	TArray<FNiagaraDataChannelPublishRequest> PublishRequests;
+
+	FVector3f LwcTile = FVector3f::ZeroVector;
+};
\ No newline at end of file
diff --git a/Engine/Plugins/FX/Niagara/Source/Niagara/Private/DataInterface/NiagaraDataInterfaceDataChannelCommon.cpp b/Engine/Plugins/FX/Niagara/Source/Niagara/Private/DataInterface/NiagaraDataInterfaceDataChannelCommon.cpp
--- a/Engine/Plugins/FX/Niagara/Source/Niagara/Private/DataInterface/NiagaraDataInterfaceDataChannelCommon.cpp
+++ b/Engine/Plugins/FX/Niagara/Source/Niagara/Private/DataInterface/NiagaraDataInterfaceDataChannelCommon.cpp
@@ -20,7 +20,7 @@ FAutoConsoleCommandWithWorldAndArgs ResetDataChannelLayouts(
 );
 
 
-FNDIDataChannel_FunctionToDataSetBinding::FNDIDataChannel_FunctionToDataSetBinding(const FNDIDataChannelFunctionInfo& FunctionInfo, const FNiagaraDataSetCompiledData& DataSetLayout)
+FNDIDataChannel_FunctionToDataSetBinding::FNDIDataChannel_FunctionToDataSetBinding(const FNDIDataChannelFunctionInfo& FunctionInfo, const FNiagaraDataSetCompiledData& DataSetLayout, TArray<FNiagaraVariableBase>& OutMissingParams)
 {
 #if DEBUG_NDI_DATACHANNEL
 	DebugFunctionInfo = FunctionInfo;
@@ -47,10 +47,16 @@ FNDIDataChannel_FunctionToDataSetBinding::FNDIDataChannel_FunctionToDataSetBindi
 			else
 			{
 				DataSetLayoutHash = 0;
+
+				#if !UE_BUILD_SHIPPING && !UE_BUILD_TEST
+				OutMissingParams.Emplace(Param);
+				#else
 				return;
+				#endif
 			}
 		}
 	};
+
 	DoGenVMBindings(FunctionInfo.Inputs);
 	DoGenVMBindings(FunctionInfo.Outputs);
 }
@@ -145,7 +151,7 @@ void FNDIDataChannelLayoutManager::Reset()
 	FunctionToDataSetLayoutMap.Reset();
 }
 
-FNDIDataChannel_FuncToDataSetBindingPtr FNDIDataChannelLayoutManager::GetLayoutInfo(const FNDIDataChannelFunctionInfo& FunctionInfo, const FNiagaraDataSetCompiledData& DataSetLayout)
+FNDIDataChannel_FuncToDataSetBindingPtr FNDIDataChannelLayoutManager::GetLayoutInfo(const FNDIDataChannelFunctionInfo& FunctionInfo, const FNiagaraDataSetCompiledData& DataSetLayout, TArray<FNiagaraVariableBase>& OutMissingParams)
 {
 	uint32 Key = GetLayoutKey(FunctionInfo, DataSetLayout);
 	
@@ -161,10 +167,10 @@ FNDIDataChannel_FuncToDataSetBindingPtr FNDIDataChannelLayoutManager::GetLayoutI
 
 	//No valid existing layout so generate a new one.
 	FRWScopeLock WriteLock(FunctionToDataSetMapLock, SLT_Write);
-	FNDIDataChannel_FuncToDataSetBindingPtr& FuncLayout = FunctionToDataSetLayoutMap.FindOrAdd(Key);
+	FNDIDataChannel_FuncToDataSetBindingPtr FuncLayout = MakeShared<FNDIDataChannel_FunctionToDataSetBinding, ESPMode::ThreadSafe>(FunctionInfo, DataSetLayout, OutMissingParams);
 	if (FuncLayout.IsValid() == false)
 	{
-		FuncLayout = MakeShared<FNDIDataChannel_FunctionToDataSetBinding, ESPMode::ThreadSafe>(FunctionInfo, DataSetLayout);
+		FunctionToDataSetLayoutMap.Add(Key) = FuncLayout;
 	}
 #if DEBUG_NDI_DATACHANNEL
 	else
diff --git a/Engine/Plugins/FX/Niagara/Source/Niagara/Private/DataInterface/NiagaraDataInterfaceDataChannelRead.cpp b/Engine/Plugins/FX/Niagara/Source/Niagara/Private/DataInterface/NiagaraDataInterfaceDataChannelRead.cpp
--- a/Engine/Plugins/FX/Niagara/Source/Niagara/Private/DataInterface/NiagaraDataInterfaceDataChannelRead.cpp
+++ b/Engine/Plugins/FX/Niagara/Source/Niagara/Private/DataInterface/NiagaraDataInterfaceDataChannelRead.cpp
@@ -13,7 +13,6 @@
 
 #include "NiagaraDataChannel.h"
 #include "NiagaraDataChannelHandler.h"
-#include "NiagaraDataChannelDefinitions.h"
 #include "NiagaraDataChannelManager.h"
 
 #include "NiagaraRenderer.h"
@@ -21,12 +20,19 @@
 
 #include "NiagaraDataInterfaceUtilities.h"
 
+#if WITH_EDITOR
+#include "INiagaraEditorOnlyDataUtlities.h"
+#include "Modules/ModuleManager.h"
+#include "NiagaraModule.h"
+#endif
+
 //////////////////////////////////////////////////////////////////////////
 
 #define LOCTEXT_NAMESPACE "NiagaraDataInterfaceDataChannelRead"
 
 DECLARE_CYCLE_STAT(TEXT("NDIDataChannelRead Read"), STAT_NDIDataChannelRead_Read, STATGROUP_NiagaraDataChannels);
 DECLARE_CYCLE_STAT(TEXT("NDIDataChannelRead Consume"), STAT_NDIDataChannelRead_Consume, STATGROUP_NiagaraDataChannels);
+DECLARE_CYCLE_STAT(TEXT("NDIDataChannelRead Spawn"), STAT_NDIDataChannelRead_Spawn, STATGROUP_NiagaraDataChannels);
 DECLARE_CYCLE_STAT(TEXT("NDIDataChannelRead Tick"), STAT_NDIDataChannelRead_Tick, STATGROUP_NiagaraDataChannels);
 DECLARE_CYCLE_STAT(TEXT("NDIDataChannelRead PostTick"), STAT_NDIDataChannelRead_PostTick, STATGROUP_NiagaraDataChannels);
 
@@ -40,6 +46,10 @@ namespace NDIDataChannelReadLocal
 	static const FName NumName(TEXT("Num"));
 	static const FName ReadName(TEXT("Read"));
 	static const FName ConsumeName(TEXT("Consume"));
+	static const FName SpawnFromSpawnInfoName(TEXT("SpawnFromSpawnInfo"));
+	static const FName SpawnConditionalName(TEXT("SpawnConditional"));
+
+	static const FName SpawnInfoName(TEXT("SpawnInfo"));
 
 	const TCHAR* GetFunctionTemplate(FName FunctionName)
 	{
@@ -65,12 +75,10 @@ namespace NDIDataChannelReadLocal
 /** Render thread copy of current instance data. */
 struct FNDIDataChannelReadInstanceData_RT
 {
-	//GPU Dataset from the channel handler. We'll grab the current buffer from this on the RT.
-	//This must be grabbed fresh from the handler each frame as it's lifetime cannot be ensured.
-	//TODO: Improve this.
-	//TODO: Likely better to have the handler provide some RT proxy class for the correct data from which we grab the correct buffer.
-	FNiagaraDataSet* ChannelDataSet = nullptr;
+	//RT proxy for game channel data from which we're reading.
+	FNiagaraDataChannelDataProxy* ChannelDataRTProxy = nullptr;
 
+	bool bReadPrevFrame = false;
 	/**
 	Table of all parameter offsets used by each GPU script using this DI.
 	Each script has to have it's own section of this table as the offsets into this table are embedded in the hlsl.
@@ -98,11 +106,11 @@ FNDIDataChannelReadInstanceData::~FNDIDataChannelReadInstanceData()
 
 }
 
-FNiagaraDataBuffer* FNDIDataChannelReadInstanceData::GetReadBufferCPU()
+FNiagaraDataBuffer* FNDIDataChannelReadInstanceData::GetReadBufferCPU(bool bPrevFrame)
 {
-	if (ExternalCPU)
+	if (DataChannelData)
 	{
-		return ExternalCPU;
+		return DataChannelData->GetCPUData(bPrevFrame);
 	}
 
 	//TODO: Local reads.
@@ -119,6 +127,8 @@ FNiagaraDataBuffer* FNDIDataChannelReadInstanceData::GetReadBufferCPU()
 
 bool FNDIDataChannelReadInstanceData::Init(UNiagaraDataInterfaceDataChannelRead* Interface, FNiagaraSystemInstance* Instance)
 {
+	EmitterInstance = Interface->EmitterBinding.Resolve(Instance, Interface);
+
 	bool bSuccess = Tick(Interface, Instance);
 	bSuccess &= PostTick(Interface, Instance);
 	return bSuccess;
@@ -127,6 +137,7 @@ bool FNDIDataChannelReadInstanceData::Init(UNiagaraDataInterfaceDataChannelRead*
 bool FNDIDataChannelReadInstanceData::Tick(UNiagaraDataInterfaceDataChannelRead* Interface, FNiagaraSystemInstance* Instance)
 {
 	ConsumeIndex = 0;
+	ConditionalSpawns.Reset();
 
 	const FNDIDataChannelCompiledData& CompiledData = Interface->GetCompiledData();
 
@@ -136,6 +147,9 @@ bool FNDIDataChannelReadInstanceData::Tick(UNiagaraDataInterfaceDataChannelRead*
 		UE_LOG(LogNiagara, Warning, TEXT("Data Channel Interface is being initialized but it is never used.\nSystem: %s\nInterface: %s"), *Instance->GetSystem()->GetFullName(), *Interface->GetFullName());
 		return true;
 	}
+	
+	//In non test/shipping builds we gather and log and missing parameters that cause us to fail to find correct bindings.
+	TArray<FNiagaraVariableBase> MissingParams;
 
 	//TODO: Reads directly from a local writer DI.
 	//For local readers, we find the source DI inside the same system and bind our functions to it's data layout.
@@ -197,55 +211,39 @@ bool FNDIDataChannelReadInstanceData::Tick(UNiagaraDataInterfaceDataChannelRead*
 		UNiagaraDataChannelHandler* DataChannelPtr = DataChannel.Get();
 		if (DataChannelPtr == nullptr)
 		{
-			//DataChannel channel we had could have been destroy so clear these.
-			ExternalCPU = nullptr;
-			ExternalGPU = nullptr;
+			DataChannelData = nullptr;
 			ChachedDataSetLayoutHash = INDEX_NONE;
-			if (UNiagaraSystem* System = Interface->GetTypedOuter<UNiagaraSystem>())
+			UWorld* World = Instance->GetWorld();
+			if (FNiagaraWorldManager* WorldMan = FNiagaraWorldManager::Get(World))
 			{
-				UWorld* World = Instance->GetWorld();
-				if (FNiagaraWorldManager* WorldMan = FNiagaraWorldManager::Get(World))
+				if (UNiagaraDataChannelHandler* NewChannelHandler = WorldMan->GetDataChannelManager().FindDataChannelHandler(Interface->Channel))
 				{
-					if (UNiagaraDataChannelHandler* NewChannelHandler = WorldMan->GetDataChannelManager().FindDataChannelHandler(Interface->Channel.ChannelName))
-					{
-						DataChannelPtr = NewChannelHandler;
-						DataChannel = NewChannelHandler;
-					}
-					else
-					{
-						UE_LOG(LogNiagara, Warning, TEXT("Failed to find or add Naigara DataChannel Channel: %s"), *Interface->Channel.ChannelName.ToString());
-						return false;
-					}
+					DataChannelPtr = NewChannelHandler;
+					DataChannel = NewChannelHandler;
+				}
+				else
+				{
+					UE_LOG(LogNiagara, Warning, TEXT("Failed to find or add Naigara DataChannel Channel: %s"), *Interface->Channel.GetName());
+					return false;
 				}
 			}
 		}
 
-		check(ExternalCPU == nullptr);
-
 		//Grab the world DataChannel data if we're reading from there.	
 		if (DataChannelPtr)
 		{
-			const FNiagaraDataSetCompiledData* CPUSourceDataCompiledData = nullptr;
-			const FNiagaraDataSetCompiledData* GPUSourceDataCompiledData = nullptr;
-			if (CompiledData.UsedByCPU())
+			if(DataChannelData == nullptr || Interface->bUpdateSourceDataEveryTick)
 			{
 				//TODO: Automatically modify tick group if we have DIs that require current frame info?
-				ExternalCPU = DataChannelPtr->GetData(Instance, Interface->bReadCurrentFrame == false);
-				CPUSourceDataCompiledData = &ExternalCPU->GetOwner()->GetCompiledData();
+				FNiagaraDataChannelSearchParameters SearchParams;
+				SearchParams.OwningComponent = Instance->GetAttachComponent();
+				DataChannelData = DataChannelPtr->FindData(SearchParams, ENiagaraResourceAccess::ReadOnly);//TODO: Maybe should have two paths, one for system instances and another for SceneComponents...
 			}
 
-			if (CompiledData.UsedByGPU())
-			{
-				//TODO: Detect in post compile if we need GPU/CPU or both.
-				ExternalGPU = DataChannelPtr->GetDataGPU(Instance);
-				GPUSourceDataCompiledData = &ExternalGPU->GetCompiledData();
-			}
-
-			const FNiagaraDataSetCompiledData* SourceDataCompiledData = CPUSourceDataCompiledData ? CPUSourceDataCompiledData : GPUSourceDataCompiledData;
-			check(SourceDataCompiledData);
-			check(CPUSourceDataCompiledData == nullptr || GPUSourceDataCompiledData == nullptr || CPUSourceDataCompiledData->GetLayoutHash() == GPUSourceDataCompiledData->GetLayoutHash());
-
-			uint64 SourceDataLayoutHash = SourceDataCompiledData->GetLayoutHash();
+			const FNiagaraDataSetCompiledData& CPUSourceDataCompiledData = DataChannelPtr->GetDataChannel()->GetCompiledData(ENiagaraSimTarget::CPUSim);
+			const FNiagaraDataSetCompiledData& GPUSourceDataCompiledData = DataChannelPtr->GetDataChannel()->GetCompiledData(ENiagaraSimTarget::GPUComputeSim);
+			check(CPUSourceDataCompiledData.GetLayoutHash() && CPUSourceDataCompiledData.GetLayoutHash() == GPUSourceDataCompiledData.GetLayoutHash());			
+			uint64 SourceDataLayoutHash = CPUSourceDataCompiledData.GetLayoutHash();
 			bool bChanged = SourceDataLayoutHash != ChachedDataSetLayoutHash;
 
 			//If our CPU or GPU source data has changed then regenerate our binding info.
@@ -255,6 +253,8 @@ bool FNDIDataChannelReadInstanceData::Tick(UNiagaraDataInterfaceDataChannelRead*
 			{
 				ChachedDataSetLayoutHash = SourceDataLayoutHash;
 
+				SpawnInfoAccessor = FNiagaraDataSetAccessor<FNiagaraSpawnInfo>(CPUSourceDataCompiledData, Interface->SpawnInfoName);
+
 				//We can likely be more targeted here.
 				//Could probably only update the RT when the GPU data changes and only update the bindings if the function hashes change etc.
 				bUpdateRTData = CompiledData.UsedByGPU();
@@ -266,16 +266,35 @@ bool FNDIDataChannelReadInstanceData::Tick(UNiagaraDataInterfaceDataChannelRead*
 					const FNDIDataChannelFunctionInfo& FuncInfo = CompiledData.GetFunctionInfo()[BindingIdx];
 
 					FNDIDataChannel_FuncToDataSetBindingPtr& BindingPtr = FuncToDataSetBindingInfo[BindingIdx];
-					BindingPtr = FNDIDataChannelLayoutManager::Get().GetLayoutInfo(FuncInfo, *SourceDataCompiledData/*, FuncToDataSetLayoutKeys[BindingIdx]*/);
+					BindingPtr = FNDIDataChannelLayoutManager::Get().GetLayoutInfo(FuncInfo, CPUSourceDataCompiledData, MissingParams);
 				}
 			}
 		}
 	}
 
+#if !UE_BUILD_SHIPPING && !UE_BUILD_TEST
+	if(MissingParams.Num() > 0)
+	{
+		FString MissingParamsString;
+		for (FNiagaraVariableBase& MissingParam : MissingParams)
+		{
+			MissingParamsString += FString::Printf(TEXT("%s %s\n"), *MissingParam.GetType().GetName(), *MissingParam.GetName().ToString());
+		}
+
+		UE_LOG(LogNiagara, Warning, TEXT("Niagara Data Channel Reader Interface is trying to read parameters that do not exist in this channel.\nIt's likely that the Data Channel Definition has been changed and this system needs to be updated.\nData Channel: %s\nSystem: %s\nComponent:%s\nMissing Parameters:\n%s\n")
+			, *DataChannel->GetDataChannel()->GetName()
+			, *Instance->GetSystem()->GetPathName()
+			, *Instance->GetAttachComponent()->GetPathName()
+			, *MissingParamsString);		
+	}	
+#endif
 
 	if (!DataChannel.IsValid() /*&& !SourceDI.IsValid()*/)//TODO: Local reads
 	{
-		//No valid source for data so bail
+		UE_LOG(LogNiagara, Warning, TEXT("Niagara Data Channel Reader Interface could not find a valid data channel.\nData Channel: %s\nSystem: %s\nComponent:%s\n")
+			, Interface->Channel ? *Interface->Channel->GetName() : TEXT("None")
+			, *Instance->GetSystem()->GetPathName()
+			, *Instance->GetAttachComponent()->GetPathName());
 		return false;
 	}
 
@@ -298,11 +317,6 @@ bool FNDIDataChannelReadInstanceData::Tick(UNiagaraDataInterfaceDataChannelRead*
 
 bool FNDIDataChannelReadInstanceData::PostTick(UNiagaraDataInterfaceDataChannelRead* Interface, FNiagaraSystemInstance* Instance)
 {
-	if (ExternalCPU)
-	{
-		ExternalCPU = nullptr;
-	}
-
 	return true;
 }
 
@@ -312,6 +326,8 @@ UNiagaraDataInterfaceDataChannelRead::UNiagaraDataInterfaceDataChannelRead(FObje
 	: Super(ObjectInitializer)
 {
 	Proxy.Reset(new FNiagaraDataInterfaceProxy_DataChannelRead());
+	EmitterBinding.BindingMode = ENiagaraDataInterfaceEmitterBindingMode::Self;
+	SpawnInfoName = NDIDataChannelReadLocal::SpawnInfoName;
 }
 
 void UNiagaraDataInterfaceDataChannelRead::PostInitProperties()
@@ -415,19 +431,41 @@ bool UNiagaraDataInterfaceDataChannelRead::PerInstanceTickPostSimulate(void* Per
 	return false;
 }
 
+void UNiagaraDataInterfaceDataChannelRead::PostStageTick(FNDICpuPostStageContext& Context)
+{
+	FNDIDataChannelReadInstanceData* InstanceData = Context.GetPerInstanceData<FNDIDataChannelReadInstanceData>();
+
+	check(InstanceData);
+	for(int32 i=0; i<InstanceData->ConditionalSpawns.Num(); ++i)
+	{
+		if(InstanceData->ConditionalSpawns[i].Count > 0)
+		{
+			if(bOverrideSpawnGroupToDataChannelIndex)
+			{
+				InstanceData->ConditionalSpawns[i].SpawnGroup = i;
+			}
+			InstanceData->EmitterInstance->GetSpawnInfo().Emplace(InstanceData->ConditionalSpawns[i]);
+		}
+	}
+	InstanceData->ConditionalSpawns.Reset();
+}
+
 void UNiagaraDataInterfaceDataChannelRead::ProvidePerInstanceDataForRenderThread(void* DataForRenderThread, void* PerInstanceData, const FNiagaraSystemInstanceID& SystemInstance)
 {
 	const FNDIDataChannelReadInstanceData& SourceData = *reinterpret_cast<const FNDIDataChannelReadInstanceData*>(PerInstanceData);
 	FNDIDataChannelReadInstanceData_RT* TargetData = new(DataForRenderThread) FNDIDataChannelReadInstanceData_RT();
 
 	//Always update the dataset, this may change without triggering a full update if it's layout is the same.
-	TargetData->ChannelDataSet = SourceData.ExternalGPU;
+	TargetData->ChannelDataRTProxy = SourceData.DataChannelData ? SourceData.DataChannelData->GetRTProxy() : nullptr;
+	TargetData->bReadPrevFrame = bReadCurrentFrame == false;
 
 	if (SourceData.bUpdateRTData && INiagaraModule::DataChannelsEnabled())
 	{
 		SourceData.bUpdateRTData = false;
 
 		TargetData->bHasUpdate = true;
+		
+		const FNiagaraDataSetCompiledData& GPUCompiledData = SourceData.DataChannel->GetDataChannel()->GetCompiledData(ENiagaraSimTarget::GPUComputeSim);
 
 		//For every GPU script, we append it's parameter access info to the table.
 		TargetData->GPUScriptParameterTableOffsets.Reset();
@@ -443,7 +481,7 @@ void UNiagaraDataInterfaceDataChannelRead::ProvidePerInstanceDataForRenderThread
 			//Now fill the table for this script
 			for (const FNiagaraVariableBase& Param : ParamAccessInfo.SortedParameters)
 			{
-				if (const FNiagaraVariableLayoutInfo* LayoutInfo = SourceData.ExternalGPU->GetVariableLayout(Param))
+				if (const FNiagaraVariableLayoutInfo* LayoutInfo = GPUCompiledData.FindVariableLayoutInfo(Param))
 				{
 					TargetData->GPUScriptParameterOffsetTable.Add(LayoutInfo->GetNumFloatComponents() > 0 ? LayoutInfo->FloatComponentStart : INDEX_NONE);
 					TargetData->GPUScriptParameterOffsetTable.Add(LayoutInfo->GetNumInt32Components() > 0 ? LayoutInfo->Int32ComponentStart : INDEX_NONE);
@@ -498,11 +536,6 @@ void UNiagaraDataInterfaceDataChannelRead::PostCompile()
 {
 	UNiagaraSystem* OwnerSystem = GetTypedOuter<UNiagaraSystem>();
 	CompiledData.Init(OwnerSystem, this);
-
-	if (const UNiagaraDataChannel* DataChannel = UNiagaraDataChannelDefinitions::FindDataChannel(Channel.ChannelName))
-	{
-		OwnerSystem->RegisterDataChannelUse(DataChannel);
-	}
 }
 
 #endif
@@ -512,19 +545,40 @@ void UNiagaraDataInterfaceDataChannelRead::PostCompile()
 
 void UNiagaraDataInterfaceDataChannelRead::GetFeedback(UNiagaraSystem* InAsset, UNiagaraComponent* InComponent, TArray<FNiagaraDataInterfaceError>& OutErrors, TArray<FNiagaraDataInterfaceFeedback>& OutWarnings, TArray<FNiagaraDataInterfaceFeedback>& OutInfo)
 {
-	//We have to post compile here manually as these graph node DIs don't get done elsewhere but we rely on the post compile data for reporting errors.
-	PostCompile();
-
 	Super::GetFeedback(InAsset, InComponent, OutErrors, OutWarnings, OutInfo);
 
+	INiagaraModule& NiagaraModule = FModuleManager::GetModuleChecked<INiagaraModule>("Niagara");
+	const INiagaraEditorOnlyDataUtilities& EditorOnlyDataUtilities = NiagaraModule.GetEditorOnlyDataUtilities();
+	UNiagaraDataInterface* RuntimeInstanceOfThis = EditorOnlyDataUtilities.IsEditorDataInterfaceInstance(this)
+		? EditorOnlyDataUtilities.GetResolvedRuntimeInstanceForEditorDataInterfaceInstance(*InAsset, *this)
+		: this;
+
+	UNiagaraDataInterfaceDataChannelRead* RuntimeReadDI = Cast<UNiagaraDataInterfaceDataChannelRead>(RuntimeInstanceOfThis);
+	
+	if(!RuntimeReadDI)
+	{
+		return;	
+	}
+
+	if (RuntimeReadDI->Channel == nullptr)
+	{
+		OutErrors.Emplace(LOCTEXT("DataChannelMissingFmt", "Data Channel Interface has no valid Data Channel."),
+			LOCTEXT("DataChannelMissingErrorSummaryFmt", "Missing Data Channel."),
+			FNiagaraDataInterfaceFix());
+
+		return;
+	}
+
 	//if(Scope == ENiagaraDataChannelScope::World)
 	{
-		if (const UNiagaraDataChannel* DataChannel = UNiagaraDataChannelDefinitions::FindDataChannel(Channel.ChannelName))
+		if (const UNiagaraDataChannel* DataChannel = RuntimeReadDI->Channel->Get())
 		{
 			//Ensure the data channel contains all the parameters this function is requesting.
 			TConstArrayView<FNiagaraVariable> ChannelVars = DataChannel->GetVariables();
-			for (const FNDIDataChannelFunctionInfo& FuncInfo : CompiledData.GetFunctionInfo())
+			for (const FNDIDataChannelFunctionInfo& FuncInfo : RuntimeReadDI->GetCompiledData().GetFunctionInfo())
 			{
+				TArray<FNiagaraVariableBase> MissingParams;
+
 				auto VerifyChannelContainsParams = [&](const TArray<FNiagaraVariableBase>& Parameters)
 				{
 					for (const FNiagaraVariableBase& FuncParam : Parameters)
@@ -532,35 +586,50 @@ void UNiagaraDataInterfaceDataChannelRead::GetFeedback(UNiagaraSystem* InAsset,
 						bool bParamFound = false;
 						for (const FNiagaraVariable& ChannelVar : ChannelVars)
 						{
-							//We have to convert each channel var to SWC for comparrison with the function variables as there is no reliable way to go back from the SWC function var to the originating LWC var.
-							UScriptStruct* ChannelSWCStruct = FNiagaraTypeHelper::GetSWCStruct(ChannelVar.GetType().GetScriptStruct());
-							if (ChannelSWCStruct)
+							//We have to convert each channel var to SWC for comparison with the function variables as there is no reliable way to go back from the SWC function var to the originating LWC var.
+							FNiagaraVariable SWCVar(ChannelVar);
+							if(ChannelVar.GetType().IsEnum() == false)
 							{
-								FNiagaraTypeDefinition SWCType(ChannelSWCStruct, FNiagaraTypeDefinition::EAllowUnfriendlyStruct::Deny);
-								FNiagaraVariable SWCVar(SWCType, ChannelVar.GetName());
-								if (SWCVar == FuncParam)
+								UScriptStruct* ChannelSWCStruct = FNiagaraTypeHelper::GetSWCStruct(ChannelVar.GetType().GetScriptStruct());
+								if (ChannelSWCStruct)
 								{
-									bParamFound = true;
-									break;
+									FNiagaraTypeDefinition SWCType(ChannelSWCStruct, FNiagaraTypeDefinition::EAllowUnfriendlyStruct::Deny);
+									SWCVar = FNiagaraVariable(SWCType, ChannelVar.GetName());
 								}
 							}
+						
+							if (SWCVar == FuncParam)
+							{
+								bParamFound = true;
+								break;
+							}
 						}
 
 						if (bParamFound == false)
 						{
-							OutErrors.Emplace(FText::Format(LOCTEXT("FuncParamMissingFromDataChannelErrorFmt", "Accessing variable {0} {1} which does not exist in Data Channel {2}."), FuncParam.GetType().GetNameText(), FText::FromName(FuncParam.GetName()), FText::FromName(Channel.ChannelName)),
-								LOCTEXT("FuncParamMissingFromDataChannelErrorSummaryFmt", "Data Channel DI function is accessing and invalid parameter."),
-								FNiagaraDataInterfaceFix());
+							MissingParams.Add(FuncParam);
 						}
 					}
 				};
 				VerifyChannelContainsParams(FuncInfo.Inputs);
 				VerifyChannelContainsParams(FuncInfo.Outputs);
+
+				if(MissingParams.Num() > 0)
+				{
+					FTextBuilder Builder;
+					Builder.AppendLineFormat(LOCTEXT("FuncParamMissingFromDataChannelErrorFmt", "Accessing variables that do not exist in Data Channel {0}."), FText::FromString(Channel.GetName()));
+					for(FNiagaraVariableBase& Param : MissingParams)
+					{
+						Builder.AppendLineFormat(LOCTEXT("FuncParamMissingFromDataChannelErrorLineFmt", "{0} {1}"), Param.GetType().GetNameText(), FText::FromName(Param.GetName()));
+					}					
+					
+					OutErrors.Emplace(Builder.ToText(), LOCTEXT("FuncParamMissingFromDataChannelErrorSummaryFmt", "Data Channel DI function is accessing invalid parameters."), FNiagaraDataInterfaceFix());
+				}
 			}
 		}
 		else
 		{
-			OutErrors.Emplace(FText::Format(LOCTEXT("DataChannelDoesNotExistErrorFmt", "Data Channel {0} does not exist. It may have been deleted."), FText::FromName(Channel.ChannelName)),
+			OutErrors.Emplace(FText::Format(LOCTEXT("DataChannelDoesNotExistErrorFmt", "Data Channel {0} does not exist. It may have been deleted."), FText::FromString(Channel.GetName())),
 				LOCTEXT("DataChannelDoesNotExistErrorSummaryFmt", "Data Channel DI is accesssinga a Data Channel that doesn't exist."),
 				FNiagaraDataInterfaceFix());
 		}
@@ -632,8 +701,12 @@ bool UNiagaraDataInterfaceDataChannelRead::Equals(const UNiagaraDataInterface* O
 		if (Super::Equals(Other) &&
 			//Scope == OtherTyped->Scope &&
 			//Source == OtherTyped->Source &&
-			Channel.ChannelName == OtherTyped->Channel.ChannelName &&
-			bReadCurrentFrame == OtherTyped->bReadCurrentFrame)
+			Channel == OtherTyped->Channel &&
+			bReadCurrentFrame == OtherTyped->bReadCurrentFrame &&
+			bUpdateSourceDataEveryTick == OtherTyped->bUpdateSourceDataEveryTick &&
+			EmitterBinding == OtherTyped->EmitterBinding &&
+			bOverrideSpawnGroupToDataChannelIndex == OtherTyped->bOverrideSpawnGroupToDataChannelIndex &&
+			SpawnInfoName == OtherTyped->SpawnInfoName)
 		{
 			return true;
 		}
@@ -657,6 +730,10 @@ bool UNiagaraDataInterfaceDataChannelRead::CopyToInternal(UNiagaraDataInterface*
 		DestTyped->Channel = Channel;
 		DestTyped->CompiledData = CompiledData;
 		DestTyped->bReadCurrentFrame = bReadCurrentFrame;
+		DestTyped->bUpdateSourceDataEveryTick = bUpdateSourceDataEveryTick;
+		DestTyped->EmitterBinding = EmitterBinding;
+		DestTyped->bOverrideSpawnGroupToDataChannelIndex = bOverrideSpawnGroupToDataChannelIndex;
+		DestTyped->SpawnInfoName = SpawnInfoName;
 		return true;
 	}
 
@@ -706,6 +783,49 @@ void UNiagaraDataInterfaceDataChannelRead::GetFunctions(TArray<FNiagaraFunctionS
 		Sig.RequiredOutputs = Sig.Outputs.Num();//The user defines what we read in the graph.
 		OutFunctions.Add(Sig);
 	}
+
+	FNiagaraVariable EnabledVar(FNiagaraTypeDefinition::GetBoolDef(), TEXT("Enable"));
+	EnabledVar.SetValue(FNiagaraBool(true));
+
+	{
+		FNiagaraFunctionSignature Sig;
+		Sig.Name = NDIDataChannelReadLocal::SpawnFromSpawnInfoName;
+#if WITH_EDITORONLY_DATA
+		Sig.Description = LOCTEXT("SpawnFromSpawnInfolFunctionDescription", "Will Spawn particles into the bound Emitter based on a spawn info variable read from each Data Channel element. ");
+#endif
+		Sig.bMemberFunction = true;
+		Sig.bExperimental = true;
+		Sig.bRequiresExecPin = true;
+		Sig.ModuleUsageBitmask = ENiagaraScriptUsageMask::Emitter | ENiagaraScriptUsageMask::System;
+		Sig.AddInput(FNiagaraVariable(FNiagaraTypeDefinition(GetClass()), TEXT("DataChannel interface")));
+		Sig.AddInput(EnabledVar, LOCTEXT("SpawnEnableInputDesc", "Enable or disable this function call. If false, this call with have no effetcs."));
+		OutFunctions.Add(Sig);
+	}
+
+	{
+		FNiagaraFunctionSignature Sig;
+		Sig.Name = NDIDataChannelReadLocal::SpawnConditionalName;
+#if WITH_EDITORONLY_DATA
+		Sig.Description = LOCTEXT("SpawnCustomFunctionDescription", "Will Spawn particles into the bound Emitter between Min and Max counts for every element in the bound Data Channel.\n\
+		Can take optional additional parameters as conditions on spawning that will be compared against the contents of each data channel element.\n\
+		For example you could spawn only for a particular value of an enum.\n\
+		For compound data types that contain multiple component floats or ints, comparissons are done on a per component basis.\n\
+		For example if you add a Vector condition parameter it will be compared against each component of the corresponding Vector in the Data Channel.\n\
+		Result = (Param.X == ChannelValue.X) && (Param.Y == ChannelValue.Y) && (Param.Z == ChannelValue.Z)");
+#endif
+		Sig.bMemberFunction = true;
+		Sig.bExperimental = true;
+		Sig.bRequiresExecPin = true;
+		Sig.ModuleUsageBitmask = ENiagaraScriptUsageMask::Emitter | ENiagaraScriptUsageMask::System;
+		Sig.AddInput(FNiagaraVariable(FNiagaraTypeDefinition(GetClass()), TEXT("DataChannel interface")));
+		Sig.AddInput(EnabledVar, LOCTEXT("SpawnEnableInputDesc", "Enable or disable this function call. If false, this call with have no effetcs."));
+		Sig.AddInput(FNiagaraVariable(StaticEnum<ENDIDataChannelSpawnMode>(), TEXT("Mode")), LOCTEXT("SpawnCondModeInputDesc", "A mode switch that controls how this funciton will behave."));
+		Sig.AddInput(FNiagaraVariable(StaticEnum<ENiagaraConditionalOperator>(), TEXT("Operator")), LOCTEXT("SpawnCondOpInputDesc", "The comparison operator to use when comparing values in the data channel to conditional parameters."));
+		Sig.AddInput(FNiagaraVariable(FNiagaraTypeDefinition::GetIntDef(), TEXT("Min Spawn Count")), LOCTEXT("MinSpawnCountInputDesc", "Minimum number of particles to spawn for each element in the data channel."));
+		Sig.AddInput(FNiagaraVariable(FNiagaraTypeDefinition::GetIntDef(), TEXT("Max Spawn Count")), LOCTEXT("MaxSpawnCountInputDesc", "Maximum number of particles to spawn for each element in the data channel."));
+		Sig.RequiredInputs = Sig.Inputs.Num();
+		OutFunctions.Add(Sig);
+	}
 }
 
 void UNiagaraDataInterfaceDataChannelRead::GetVMExternalFunction(const FVMExternalFunctionBindingInfo& BindingInfo, void* InstanceData, FVMExternalFunction& OutFunc)
@@ -717,135 +837,28 @@ void UNiagaraDataInterfaceDataChannelRead::GetVMExternalFunction(const FVMExtern
 	else
 	{
 		int32 FuncIndex = CompiledData.FindFunctionInfoIndex(BindingInfo.Name, BindingInfo.VariadicInputs, BindingInfo.VariadicOutputs);
-		if (FuncIndex != INDEX_NONE)
+		if (BindingInfo.Name == NDIDataChannelReadLocal::SpawnFromSpawnInfoName)
 		{
-			if (BindingInfo.Name == NDIDataChannelReadLocal::ReadName)
-			{
-				OutFunc = FVMExternalFunction::CreateLambda([this, FuncIndex](FVectorVMExternalFunctionContext& Context) { this->Read(Context, FuncIndex); });
-			}
-			else if (BindingInfo.Name == NDIDataChannelReadLocal::ConsumeName)
-			{
-				OutFunc = FVMExternalFunction::CreateLambda([this, FuncIndex](FVectorVMExternalFunctionContext& Context) { this->Consume(Context, FuncIndex); });
-			}
-			else
-			{
-				UE_LOG(LogTemp, Display, TEXT("Could not find data interface external function in %s. Received Name: %s"), *GetPathNameSafe(this), *BindingInfo.Name.ToString());
-			}
+			OutFunc = FVMExternalFunction::CreateLambda([this](FVectorVMExternalFunctionContext& Context) { this->SpawnFromSpawnInfo(Context); });
 		}
-		else
+		else if (BindingInfo.Name == NDIDataChannelReadLocal::ReadName)
 		{
-			UE_LOG(LogTemp, Display, TEXT("Could not find data interface external function in %s. Received Name: %s"), *GetPathNameSafe(this), *BindingInfo.Name.ToString());
+			OutFunc = FVMExternalFunction::CreateLambda([this, FuncIndex](FVectorVMExternalFunctionContext& Context) { this->Read(Context, FuncIndex); });
 		}
-	}
-}
-
-template<int32 EXPECTED_NUM_INPUTS>
-struct FNDIVariadicOutputHandler
-{
-	TArray<VectorVM::FExternalFuncRegisterHandler<float>, TInlineAllocator<EXPECTED_NUM_INPUTS>> FloatOutputs;
-	TArray<VectorVM::FExternalFuncRegisterHandler<int32>, TInlineAllocator<EXPECTED_NUM_INPUTS>> IntOutputs;
-	TArray<VectorVM::FExternalFuncRegisterHandler<FFloat16>, TInlineAllocator<EXPECTED_NUM_INPUTS>> HalfOutputs;
-
-	FNDIVariadicOutputHandler(FVectorVMExternalFunctionContext& Context, const FNDIDataChannel_FuncToDataSetBindingPtr& BindingPtr)
-	{
-		//Parse the VM bytecode inputs in order, mapping them to the correct DataChannel data 
-		FloatOutputs.Reserve(BindingPtr->NumFloatComponents);
-		IntOutputs.Reserve(BindingPtr->NumInt32Components);
-		HalfOutputs.Reserve(BindingPtr->NumHalfComponents);
-		for (FNDIDataChannelRegisterBinding& VMBinding : BindingPtr->VMRegisterBindings)
+		else if (BindingInfo.Name == NDIDataChannelReadLocal::ConsumeName)
 		{
-			if (VMBinding.DataType == (int32)ENiagaraBaseTypes::Float)
-			{
-				FloatOutputs.Emplace(Context);
-			}
-			else if (VMBinding.DataType == (int32)ENiagaraBaseTypes::Int32 || VMBinding.DataType == (int32)ENiagaraBaseTypes::Bool)
-			{
-				IntOutputs.Emplace(Context);
-			}
-			else if (VMBinding.DataType == (int32)ENiagaraBaseTypes::Half)
-			{
-				HalfOutputs.Emplace(Context);
-			}
-			else
-			{
-				checkf(false, TEXT("Didn't find a binding for this function input. Likely an error in the building of the binding data and the Float/Int/Half Bindings arrays."));
-			}
+			OutFunc = FVMExternalFunction::CreateLambda([this, FuncIndex](FVectorVMExternalFunctionContext& Context) { this->Consume(Context, FuncIndex); });
 		}
-	}
-
-	bool Process(FNiagaraDataBuffer* Data, uint32 Index, const FNDIDataChannel_FuncToDataSetBindingPtr& BindingPtr)
-	{
-		checkSlow(Data);
-
-		//Process this index.
-		//TODO: Can be optimized to pull N values out at a time, reducing the binding/lookup overhead.
-		if ((uint32)Index < Data->GetNumInstances())
+		else if (BindingInfo.Name == NDIDataChannelReadLocal::SpawnConditionalName)
 		{
-			for (FNDIDataChannelRegisterBinding& VMBinding : BindingPtr->VMRegisterBindings)
-			{
-				if (VMBinding.DataType == (int32)ENiagaraBaseTypes::Float)
-				{
-					if (VMBinding.DataSetRegisterIndex == INDEX_NONE)
-					{
-						*FloatOutputs[VMBinding.FunctionRegisterIndex].GetDestAndAdvance() = 0.0f;
-					}
-					else
-					{
-						*FloatOutputs[VMBinding.FunctionRegisterIndex].GetDestAndAdvance() = *Data->GetInstancePtrFloat(VMBinding.DataSetRegisterIndex, (uint32)Index);
-					}
-				}
-				else if (VMBinding.DataType == (int32)ENiagaraBaseTypes::Int32 || VMBinding.DataType == (int32)ENiagaraBaseTypes::Bool)
-				{
-					if (VMBinding.DataSetRegisterIndex == INDEX_NONE)
-					{
-						*IntOutputs[VMBinding.FunctionRegisterIndex].GetDestAndAdvance() = 0;
-					}
-					else
-					{
-						*IntOutputs[VMBinding.FunctionRegisterIndex].GetDestAndAdvance() = *Data->GetInstancePtrInt32(VMBinding.DataSetRegisterIndex, (uint32)Index);
-					}
-				}
-				else if (VMBinding.DataType == (int32)ENiagaraBaseTypes::Half)
-				{
-					if (VMBinding.DataSetRegisterIndex == INDEX_NONE)
-					{
-						*HalfOutputs[VMBinding.FunctionRegisterIndex].GetDestAndAdvance() = FFloat16(0.0f);
-					}
-					else
-					{
-						*HalfOutputs[VMBinding.FunctionRegisterIndex].GetDestAndAdvance() = *Data->GetInstancePtrHalf(VMBinding.DataSetRegisterIndex, (uint32)Index);
-					}
-				}
-			}
-			return true;
+			OutFunc = FVMExternalFunction::CreateLambda([this, FuncIndex](FVectorVMExternalFunctionContext& Context) { this->SpawnConditional(Context, FuncIndex); });
 		}
-
-		Fallback(1);
-
-		return false;
-	}
-
-	void Fallback(int32 Count)
-	{
-		for (int32 OutIdx = 0; OutIdx < FloatOutputs.Num(); ++OutIdx)
-		{
-			FMemory::Memzero(FloatOutputs[OutIdx].GetDest(), sizeof(float) * Count);
-			FloatOutputs[OutIdx].Advance(Count);
-		}
-
-		for (int32 OutIdx = 0; OutIdx < IntOutputs.Num(); ++OutIdx)
-		{
-			FMemory::Memzero(IntOutputs[OutIdx].GetDest(), sizeof(int32) * Count);
-			IntOutputs[OutIdx].Advance(Count);
-		}
-
-		for (int32 OutIdx = 0; OutIdx < HalfOutputs.Num(); ++OutIdx)
+		else
 		{
-			FMemory::Memzero(HalfOutputs[OutIdx].GetDest(), sizeof(FFloat16) * Count);
-			HalfOutputs[OutIdx].Advance(Count);
+			UE_LOG(LogTemp, Display, TEXT("Could not find data interface external function in %s. Received Name: %s"), *GetPathNameSafe(this), *BindingInfo.Name.ToString());
 		}
 	}
-};
+}
 
 void UNiagaraDataInterfaceDataChannelRead::Num(FVectorVMExternalFunctionContext& Context)
 {
@@ -853,7 +866,7 @@ void UNiagaraDataInterfaceDataChannelRead::Num(FVectorVMExternalFunctionContext&
 
 	FNDIOutputParam<int32> OutNum(Context);
 
-	FNiagaraDataBuffer* Buffer = InstData->GetReadBufferCPU();
+	FNiagaraDataBuffer* Buffer = InstData->GetReadBufferCPU(bReadCurrentFrame == false);
 
 	int32 Num = 0;
 	if (Buffer && INiagaraModule::DataChannelsEnabled())
@@ -873,21 +886,46 @@ void UNiagaraDataInterfaceDataChannelRead::Read(FVectorVMExternalFunctionContext
 	VectorVM::FUserPtrHandler<FNDIDataChannelReadInstanceData> InstData(Context);
 	FNDIInputParam<int32> InIndex(Context);
 
-	FNDIOutputParam<bool> OutSuccess(Context);
+	FNDIOutputParam<FNiagaraBool> OutSuccess(Context);
 
 	const FNDIDataChannelFunctionInfo& FuncInfo = CompiledData.GetFunctionInfo()[FuncIdx];
-	const FNDIDataChannel_FuncToDataSetBindingPtr& BindingInfo = InstData->FuncToDataSetBindingInfo.IsValidIndex(FuncIdx) ? InstData->FuncToDataSetBindingInfo[FuncIdx] : nullptr;
+	const FNDIDataChannel_FunctionToDataSetBinding* BindingInfo = InstData->FuncToDataSetBindingInfo.IsValidIndex(FuncIdx) ? InstData->FuncToDataSetBindingInfo[FuncIdx].Get() : nullptr;
 	FNDIVariadicOutputHandler<16> VariadicOutputs(Context, BindingInfo);//TODO: Make static / avoid allocation
 
-	FNiagaraDataBuffer* Buffer = InstData->GetReadBufferCPU();
-	if (Buffer && BindingInfo && INiagaraModule::DataChannelsEnabled())
-	{
+	FNiagaraDataBuffer* Data = InstData->GetReadBufferCPU(bReadCurrentFrame == false);
+	if (Data && BindingInfo && INiagaraModule::DataChannelsEnabled())
+ 	{
+// 		FString Label = TEXT("NDIDataChannelRead::Read() - ");
+// 		Data->Dump(0, Data->GetNumInstances(), Label);
+
 		for (int32 i = 0; i < Context.GetNumInstances(); ++i)
 		{
 			int32 Index = InIndex.GetAndAdvance();
 
-			//TODO: Optimize for long runs of reads.
-			bool bSuccess = VariadicOutputs.Process(Buffer, Index, BindingInfo);
+			bool bProcess = (uint32)Index < Data->GetNumInstances();
+
+			auto FloatFunc = [Data, Index](const FNDIDataChannelRegisterBinding& VMBinding, float* FloatData)
+			{
+				if (VMBinding.DataSetRegisterIndex != INDEX_NONE)
+				{
+					float* Src = Data->GetInstancePtrFloat(VMBinding.DataSetRegisterIndex, (uint32)Index);
+					*FloatData = *Src;
+				}
+			};
+			auto IntFunc = [Data, Index](const FNDIDataChannelRegisterBinding& VMBinding, int32* IntData)
+			{
+				if (VMBinding.DataSetRegisterIndex != INDEX_NONE)
+				{
+					int32* Src = Data->GetInstancePtrInt32(VMBinding.DataSetRegisterIndex, (uint32)Index);
+					*IntData = *Src;
+				}
+			};
+			auto HalfFunc = [Data, Index](const FNDIDataChannelRegisterBinding& VMBinding, FFloat16* HalfData)
+			{
+				if (VMBinding.DataSetRegisterIndex != INDEX_NONE)
+					*HalfData = *Data->GetInstancePtrHalf(VMBinding.DataSetRegisterIndex, (uint32)Index);
+			};
+			bool bSuccess = VariadicOutputs.Process(bProcess, BindingInfo, FloatFunc, IntFunc, HalfFunc);
 
 			if (OutSuccess.IsValid())
 			{
@@ -915,28 +953,45 @@ void UNiagaraDataInterfaceDataChannelRead::Consume(FVectorVMExternalFunctionCont
 	FNDIOutputParam<int32> OutIndex(Context);
 
 	const FNDIDataChannelFunctionInfo& FuncInfo = CompiledData.GetFunctionInfo()[FuncIdx];
-	const FNDIDataChannel_FuncToDataSetBindingPtr& BindingInfo = InstData->FuncToDataSetBindingInfo.IsValidIndex(FuncIdx) ? InstData->FuncToDataSetBindingInfo[FuncIdx] : nullptr;
+	const FNDIDataChannel_FunctionToDataSetBinding* BindingInfo = InstData->FuncToDataSetBindingInfo.IsValidIndex(FuncIdx) ? InstData->FuncToDataSetBindingInfo[FuncIdx].Get() : nullptr;
 	FNDIVariadicOutputHandler<16> VariadicOutputs(Context, BindingInfo);//TODO: Make static / avoid allocation
 
 	//TODO: Optimize for constant bConsume.
 	//TODO: Optimize for long runs of bConsume==true;
-	FNiagaraDataBuffer* Buffer = InstData->GetReadBufferCPU();
-	if (Buffer && BindingInfo && INiagaraModule::DataChannelsEnabled())
+	FNiagaraDataBuffer* Data = InstData->GetReadBufferCPU(bReadCurrentFrame == false);
+	if (Data && BindingInfo && INiagaraModule::DataChannelsEnabled())
 	{
 		for (int32 i = 0; i < Context.GetNumInstances(); ++i)
 		{
 			bool bConsume = InConsume.GetAndAdvance();
 
 			bool bSuccess = false;
-			int32 ConsumeIdx = InstData->ConsumeIndex;
+			int32 Index = InstData->ConsumeIndex;
 			if (bConsume)
 			{
-				while (ConsumeIdx < (int32)Buffer->GetNumInstances() && InstData->ConsumeIndex.compare_exchange_strong(ConsumeIdx, ConsumeIdx + 1) == false)
+				while (Index < (int32)Data->GetNumInstances() && InstData->ConsumeIndex.compare_exchange_strong(Index, Index + 1) == false)
 				{
-					ConsumeIdx = InstData->ConsumeIndex;
+					Index = InstData->ConsumeIndex;
 				}
 
-				bSuccess = VariadicOutputs.Process(Buffer, ConsumeIdx, BindingInfo);
+				bConsume &= (uint32)Index < Data->GetNumInstances();
+
+				auto FloatFunc = [Data, Index](const FNDIDataChannelRegisterBinding& VMBinding, float* FloatData)
+				{
+					if (VMBinding.DataSetRegisterIndex != INDEX_NONE)
+						*FloatData = *Data->GetInstancePtrFloat(VMBinding.DataSetRegisterIndex, (uint32)Index);
+				};
+				auto IntFunc = [Data, Index](const FNDIDataChannelRegisterBinding& VMBinding, int32* IntData)
+				{
+					if (VMBinding.DataSetRegisterIndex != INDEX_NONE)
+						*IntData = *Data->GetInstancePtrInt32(VMBinding.DataSetRegisterIndex, (uint32)Index);
+				};
+				auto HalfFunc = [Data, Index](const FNDIDataChannelRegisterBinding& VMBinding, FFloat16* HalfData)
+				{
+					if (VMBinding.DataSetRegisterIndex != INDEX_NONE)
+						*HalfData = *Data->GetInstancePtrHalf(VMBinding.DataSetRegisterIndex, (uint32)Index);
+				};
+				bSuccess = VariadicOutputs.Process(bConsume, BindingInfo, FloatFunc, IntFunc, HalfFunc);
 			}
 			else
 			{
@@ -949,7 +1004,7 @@ void UNiagaraDataInterfaceDataChannelRead::Consume(FVectorVMExternalFunctionCont
 			}
 			if (OutIndex.IsValid())
 			{
-				OutIndex.SetAndAdvance(bSuccess ? ConsumeIdx : INDEX_NONE);
+				OutIndex.SetAndAdvance(bSuccess ? Index : INDEX_NONE);
 			}
 		}
 	}
@@ -969,6 +1024,136 @@ void UNiagaraDataInterfaceDataChannelRead::Consume(FVectorVMExternalFunctionCont
 	}
 }
 
+void UNiagaraDataInterfaceDataChannelRead::SpawnFromSpawnInfo(FVectorVMExternalFunctionContext& Context)
+{
+	SCOPE_CYCLE_COUNTER(STAT_NDIDataChannelRead_Spawn);
+
+	//This should only be called from emitter scripts and since it has per instance data then we process them individually.
+	check(Context.GetNumInstances() == 1);
+
+	FNDIRandomHelperFromStream RandHelper(Context);
+
+	VectorVM::FUserPtrHandler<FNDIDataChannelReadInstanceData> InstData(Context);
+		
+	FNDIInputParam<FNiagaraBool> InEnabled(Context);
+
+	FNiagaraEmitterInstance* EmittterInst = InstData->EmitterInstance;
+	FNiagaraDataChannelData* DataChannelData = InstData->DataChannelData.Get();
+	FNiagaraDataBuffer* Data = DataChannelData ? DataChannelData->GetCPUData(!bReadCurrentFrame) : nullptr;
+
+	bool bSpawn = INiagaraModule::DataChannelsEnabled() && Data && EmittterInst;
+
+	TArray<FNiagaraSpawnInfo>& EmitterSpawnInfos = EmittterInst->GetSpawnInfo();
+
+	bSpawn &= InEnabled.GetAndAdvance();
+	bSpawn &= InstData->SpawnInfoAccessor.IsValid();
+	if (bSpawn)
+	{
+		FNiagaraDataSetReaderStruct<FNiagaraSpawnInfo> Reader = InstData->SpawnInfoAccessor.GetReader(Data);
+
+		int32 NumDataChannelInstances = Data->GetNumInstances();
+
+		for (int32 DataChannelIdx = 0; DataChannelIdx < NumDataChannelInstances; ++DataChannelIdx)
+		{
+			FNiagaraSpawnInfo SpawnInfo(0, 0.0f, 0.0f, 0);
+			SpawnInfo = Reader.Get(DataChannelIdx);
+
+			if (SpawnInfo.Count > 0)
+			{
+				if (bOverrideSpawnGroupToDataChannelIndex)
+				{
+					SpawnInfo.SpawnGroup = DataChannelIdx;
+				}
+				EmitterSpawnInfos.Emplace(SpawnInfo);
+			}
+		}
+	}
+}
+
+void UNiagaraDataInterfaceDataChannelRead::SpawnConditional(FVectorVMExternalFunctionContext& Context, int32 FuncIdx)
+{
+	SCOPE_CYCLE_COUNTER(STAT_NDIDataChannelRead_Spawn);
+
+	//This should only be called from emitter scripts and since it has per instance data then we process them individually.
+	check(Context.GetNumInstances() == 1);
+
+	FNDIRandomHelperFromStream RandHelper(Context);
+
+	VectorVM::FUserPtrHandler<FNDIDataChannelReadInstanceData> InstData(Context);
+
+	//Binding info can be null here as we can be spawning without any conditions, i.e. no variadic parameters to the function.
+	const FNDIDataChannel_FunctionToDataSetBinding* BindingInfo = InstData->FuncToDataSetBindingInfo.IsValidIndex(FuncIdx) ? InstData->FuncToDataSetBindingInfo[FuncIdx].Get() : nullptr;
+
+	FNDIInputParam<FNiagaraBool> InEnabled(Context);
+	FNDIInputParam<int32> InMode(Context);
+	FNDIInputParam<int32> InOp(Context);
+	FNDIInputParam<int32> InSpawnMin(Context);
+	FNDIInputParam<int32> InSpawnMax(Context);
+
+	FNDIVariadicInputHandler<16> VariadicInputs(Context, BindingInfo);//TODO: Make static / avoid allocation
+
+	FNiagaraEmitterInstance* EmittterInst = InstData->EmitterInstance;
+	FNiagaraDataChannelData* DataChannelData = InstData->DataChannelData.Get();
+	FNiagaraDataBuffer* Data = DataChannelData ? DataChannelData->GetCPUData(!bReadCurrentFrame) : nullptr;
+
+	bool bSpawn = INiagaraModule::DataChannelsEnabled() && Data && EmittterInst && InEnabled.GetAndAdvance();
+	if(bSpawn)
+	{
+		TArray<FNiagaraSpawnInfo>& EmitterSpawnInfos = EmittterInst->GetSpawnInfo();
+		ENDIDataChannelSpawnMode Mode = (ENDIDataChannelSpawnMode)InMode.GetAndAdvance();
+		ENiagaraConditionalOperator Op = (ENiagaraConditionalOperator)InOp.GetAndAdvance();
+
+		//Is mode none or invalid?
+		if ((int32)Mode == (int32)ENDIDataChannelSpawnMode::None || (int32)Mode < 0 || (int32)Mode >= (int32)ENDIDataChannelSpawnMode::Max)
+		{
+			return;
+		}
+
+		int32 SpawnMin = InSpawnMin.GetAndAdvance();
+		int32 SpawnMax = InSpawnMax.GetAndAdvance();
+
+		int32 NumDataChannelInstances = Data->GetNumInstances();
+
+		//Each data channel element has an additional spawn entry which accumulates across all spawning calls and can be nulled independently by a suppression call.
+		InstData->ConditionalSpawns.SetNumZeroed(NumDataChannelInstances);
+
+		for (int32 DataChannelIdx = 0; DataChannelIdx < NumDataChannelInstances; ++DataChannelIdx)
+		{
+			bool bConditionsPass = true;
+			auto FloatFunc = [&bConditionsPass, Op, Data, DataChannelIdx](const FNDIDataChannelRegisterBinding& VMBinding, float FloatData)
+			{
+				if (VMBinding.DataSetRegisterIndex != INDEX_NONE)
+					bConditionsPass &= EvalConditional(Op, FloatData, *Data->GetInstancePtrFloat(VMBinding.DataSetRegisterIndex, (uint32)DataChannelIdx));
+			};
+			auto IntFunc = [&bConditionsPass, Op, Data, DataChannelIdx](const FNDIDataChannelRegisterBinding& VMBinding, int32 IntData)
+			{
+				if (VMBinding.DataSetRegisterIndex != INDEX_NONE)
+					bConditionsPass &= EvalConditional(Op, IntData, *Data->GetInstancePtrInt32(VMBinding.DataSetRegisterIndex, (uint32)DataChannelIdx));
+			};
+			auto HalfFunc = [&bConditionsPass, Op, Data, DataChannelIdx](const FNDIDataChannelRegisterBinding& VMBinding, FFloat16 HalfData)
+			{
+				if (VMBinding.DataSetRegisterIndex != INDEX_NONE)
+					bConditionsPass &= EvalConditional(Op, HalfData, *Data->GetInstancePtrHalf(VMBinding.DataSetRegisterIndex, (uint32)DataChannelIdx));
+			};
+			VariadicInputs.Process(true, BindingInfo, FloatFunc, IntFunc, HalfFunc);
+			VariadicInputs.Reset();
+
+			if (bConditionsPass)
+			{
+				int32 Count = RandHelper.RandRange(DataChannelIdx, SpawnMin, SpawnMax);
+				if (Mode == ENDIDataChannelSpawnMode::Accumulate)
+				{
+					InstData->ConditionalSpawns[DataChannelIdx].Count += Count;
+				}
+				else if (Mode == ENDIDataChannelSpawnMode::Override)
+				{
+					InstData->ConditionalSpawns[DataChannelIdx].Count = Count;
+				}
+			}
+		}
+	}
+}
+
 #if WITH_EDITORONLY_DATA
 bool UNiagaraDataInterfaceDataChannelRead::AppendCompileHash(FNiagaraCompileHashVisitor* InVisitor) const
 {
@@ -1015,7 +1200,9 @@ void UNiagaraDataInterfaceDataChannelRead::GetParameterDefinitionHLSL(FNiagaraDa
 	{
 		//Common args for all functions
 		{TEXT("ParameterName"),							HlslGenContext.ParameterInfo.DataInterfaceHLSLSymbol},
+		
 		//Per function args. These will be changed with each function written
+		{TEXT("FunctionSymbol"),						FString(TEXT("FunctionSymbol"))},//Function symbol which will be a mangled form from the translator.
 		{TEXT("FunctionParameters"),					FString(TEXT("FunctionParameters"))},//Function parameters written into the function signature.
 		{TEXT("PerFunctionParameterShaderCode"),		FString(TEXT("PerFunctionParameterShaderCode"))},//Combined string of all code dealing with each of the function parameters.
 		//Per function parameter args. These will be changed with each parameter written
@@ -1030,6 +1217,7 @@ void UNiagaraDataInterfaceDataChannelRead::GetParameterDefinitionHLSL(FNiagaraDa
 	};
 
 	//Grab refs to per function args we'll change with each function written.
+	FStringFormatArg& FunctionSymbol = HlslTemplateArgs.FindChecked(TEXT("FunctionSymbol"));
 	FStringFormatArg& FunctionParameters = HlslTemplateArgs.FindChecked(TEXT("FunctionParameters"));
 	FStringFormatArg& PerFunctionParameterShaderCode = HlslTemplateArgs.FindChecked(TEXT("PerFunctionParameterShaderCode"));
 
@@ -1127,6 +1315,7 @@ void UNiagaraDataInterfaceDataChannelRead::GetParameterDefinitionHLSL(FNiagaraDa
 		const FNiagaraFunctionSignature& Signature = HlslGenContext.Signatures[FuncIdx];
 
 		//Init/Reset our per function hlsl template args.
+		FunctionSymbol.StringValue = HlslGenContext.GetFunctionSignatureSymbol(Signature);
 		FunctionParameters.StringValue.Reset();//Reset function parameters ready to rebuild when we iterate over the parameters.
 		PerFunctionParameterShaderCode.StringValue.Reset(); //Reset function parameters ready to rebuild when we iterate over the parameters.
 
@@ -1138,6 +1327,12 @@ void UNiagaraDataInterfaceDataChannelRead::GetParameterDefinitionHLSL(FNiagaraDa
 			//Function that will recurse down a parameter's type and generate the appropriate IO code for all of it's members.
 			TFunction<void(bool, UScriptStruct*, FString&)> GenerateRWParameterCode = [&](bool bRead, UScriptStruct* Struct, FString& OutCode)
 			{
+				//Intercept positions and replace with FVector3fs
+				if(Struct == FNiagaraPosition::StaticStruct())
+				{
+					Struct = FNiagaraTypeDefinition::GetVec3Struct();
+				}
+
 				if (Struct == FNiagaraTypeDefinition::GetFloatStruct() || Struct == FNiagaraTypeDefinition::GetVec2Struct() || Struct == FNiagaraTypeDefinition::GetVec3Struct()
 					|| Struct == FNiagaraTypeDefinition::GetVec4Struct() || Struct == FNiagaraTypeDefinition::GetColorStruct() || Struct == FNiagaraTypeDefinition::GetQuatStruct())
 				{
@@ -1277,9 +1472,10 @@ void UNiagaraDataInterfaceDataChannelRead::SetShaderParameters(const FNiagaraDat
 			ParameterOffsetTableIndex = *ParameterOffsetTableIndexPtr;
 		}
 
-		if (InstanceData->ChannelDataSet && ParameterOffsetTableIndex != INDEX_NONE)
+		if (InstanceData->ChannelDataRTProxy && ParameterOffsetTableIndex != INDEX_NONE)
 		{
-			if (FNiagaraDataBuffer* Data = InstanceData->ChannelDataSet->GetCurrentData())
+			FNiagaraDataBuffer* Data = InstanceData->bReadPrevFrame ? InstanceData->ChannelDataRTProxy->PrevFrameData.GetReference() : InstanceData->ChannelDataRTProxy->GPUDataSet->GetCurrentData();
+			if (Data)
 			{
 				const FReadBuffer& ParameterLayoutBuffer = InstanceData->ParameterLayoutBuffer;
 
@@ -1320,7 +1516,8 @@ void FNiagaraDataInterfaceProxy_DataChannelRead::ConsumePerInstanceDataFromGameT
 	FNDIDataChannelReadInstanceData_RT& SourceData = *reinterpret_cast<FNDIDataChannelReadInstanceData_RT*>(PerInstanceData);
 	FInstanceData& InstData = SystemInstancesToProxyData_RT.FindOrAdd(Instance);
 
-	InstData.ChannelDataSet = SourceData.ChannelDataSet;
+	InstData.ChannelDataRTProxy = SourceData.ChannelDataRTProxy;
+	InstData.bReadPrevFrame = SourceData.bReadPrevFrame;
 
 	if (SourceData.bHasUpdate)
 	{
@@ -1354,4 +1551,4 @@ int32 FNiagaraDataInterfaceProxy_DataChannelRead::PerInstanceDataPassedToRenderT
 	return sizeof(FNDIDataChannelReadInstanceData_RT);
 }
 
-#undef LOCTEXT_NAMESPACE
+#undef LOCTEXT_NAMESPACE
\ No newline at end of file
