Impact FX Aggregation via Data Channels.

diff --git a/Engine/Plugins/FX/Niagara/Source/Niagara/Private/DataInterface/NiagaraDataInterfaceDataChannelWrite.cpp b/Engine/Plugins/FX/Niagara/Source/Niagara/Private/DataInterface/NiagaraDataInterfaceDataChannelWrite.cpp
--- a/Engine/Plugins/FX/Niagara/Source/Niagara/Private/DataInterface/NiagaraDataInterfaceDataChannelWrite.cpp
+++ b/Engine/Plugins/FX/Niagara/Source/Niagara/Private/DataInterface/NiagaraDataInterfaceDataChannelWrite.cpp
@@ -14,9 +14,13 @@
 
 #include "NiagaraDataChannel.h"
 #include "NiagaraDataChannelHandler.h"
-#include "NiagaraDataChannelDefinitions.h"
 #include "NiagaraDataChannelManager.h"
 
+#if WITH_EDITOR
+#include "INiagaraEditorOnlyDataUtlities.h"
+#include "Modules/ModuleManager.h"
+#include "NiagaraModule.h"
+#endif
 
 #define LOCTEXT_NAMESPACE "NiagaraDataInterfaceDataChannelWrite"
 
@@ -100,20 +104,26 @@ struct FNDIDataChannelWriteInstanceData
 	/** Pointer to the world DataChannel Channel we'll push our DataChannel into. Can be null if DI is not set to publish it's DataChannel. */
 	TWeakObjectPtr<UNiagaraDataChannelHandler> DataChannel;
 
-	FNiagaraDataSet* Data = nullptr;
+	//Shared pointer to the actual data we'll be pushing into for this data channel.
+	FNiagaraDataChannelDataPtr DataChannelData;
 
-	bool bNeedsLastFrameData = true;
+	/** Local dataset we write into. 
+	TODO: alternate write modes.
+		- Crit sec access to the data channel buffer and write direct?
+		- 
+	*/
+	FNiagaraDataSet* Data = nullptr;
 
 	TArray<FNDIDataChannel_FuncToDataSetBindingPtr, TInlineAllocator<8>> FunctionToDatSetBindingInfo;
 
 	~FNDIDataChannelWriteInstanceData()
 	{
-		if (Data && DataChannel.IsValid())
+		if (Data && DataChannelData.IsValid())
 		{
-			DataChannel->RemovePublishRequests(Data);
+			DataChannelData->RemovePublishRequests(Data);
 		}
 
-		/** We defer the deletion of the datase to the RT to be sure all in-flight RT commands have finished using it.*/
+		/** We defer the deletion of the dataset to the RT to be sure all in-flight RT commands have finished using it.*/
 		ENQUEUE_RENDER_COMMAND(FDeleteContextCommand)(
 			[DataChannelDataSet = Data](FRHICommandListImmediate& RHICmdList)
 			{
@@ -131,12 +141,15 @@ struct FNDIDataChannelWriteInstanceData
 		Data = new FNiagaraDataSet();
 		Data->Init(&Interface->GetCompiledData().DataLayout);
 
+		//In non test/shipping builds we gather and log and missing parameters that cause us to fail to find correct bindings.
+		TArray<FNiagaraVariableBase> MissingParams;
+
 		//Grab the correct function binding infos for this DI.
 		const FNDIDataChannelCompiledData& CompiledData = Interface->GetCompiledData();
 		FunctionToDatSetBindingInfo.Reset(CompiledData.GetFunctionInfo().Num());
 		for (const FNDIDataChannelFunctionInfo& FuncInfo : CompiledData.GetFunctionInfo())
 		{
-			FunctionToDatSetBindingInfo.Add(FNDIDataChannelLayoutManager::Get().GetLayoutInfo(FuncInfo, Data->GetCompiledData()));
+			FunctionToDatSetBindingInfo.Add(FNDIDataChannelLayoutManager::Get().GetLayoutInfo(FuncInfo, Data->GetCompiledData(), MissingParams));			
 		}
 
 		return true;
@@ -152,14 +165,39 @@ struct FNDIDataChannelWriteInstanceData
 				UWorld* World = Instance->GetWorld();
 				if (FNiagaraWorldManager* WorldMan = FNiagaraWorldManager::Get(World))
 				{
-					if (UNiagaraDataChannelHandler* NewChannelHandler = WorldMan->GetDataChannelManager().FindDataChannelHandler(Interface->Channel.ChannelName))
+					if (UNiagaraDataChannelHandler* NewChannelHandler = WorldMan->GetDataChannelManager().FindDataChannelHandler(Interface->Channel))
 					{
 						DataChannelPtr = NewChannelHandler;
 						DataChannel = NewChannelHandler;
+
+#if !UE_BUILD_SHIPPING && !UE_BUILD_TEST
+						//In non test/shipping builds we gather and log and missing parameters that cause us to fail to find correct bindings.
+						TArray<FNiagaraVariableBase> MissingParams;
+						const FNDIDataChannelCompiledData& CompiledData = Interface->GetCompiledData();
+						for (const FNDIDataChannelFunctionInfo& FuncInfo : CompiledData.GetFunctionInfo())
+						{
+							FNDIDataChannelLayoutManager::Get().GetLayoutInfo(FuncInfo, DataChannelPtr->GetDataChannel()->GetCompiledData(ENiagaraSimTarget::CPUSim), MissingParams);
+						}
+
+						if (MissingParams.Num() > 0)
+						{
+							FString MissingParamsString;
+							for (FNiagaraVariableBase& MissingParam : MissingParams)
+							{
+								MissingParamsString += FString::Printf(TEXT("%s %s\n"), *MissingParam.GetType().GetName(), *MissingParam.GetName().ToString());
+							}
+
+							UE_LOG(LogNiagara, Warning, TEXT("Niagara Data Channel Writer Interface is trying to write parameters that do not exist in this channel.\nIt's likely that the Data Channel Definition has been changed and this system needs to be updated.\nData Channel: %s\nSystem: %s\nComponent:%s\nMissing Parameters:\n%s\n")
+								, *DataChannel->GetDataChannel()->GetName()
+								, *Instance->GetSystem()->GetPathName()
+								, *Instance->GetAttachComponent()->GetPathName()
+								, *MissingParamsString);
+						}
+#endif
 					}
 					else
 					{
-						UE_LOG(LogNiagara, Warning, TEXT("Failed to find or add Naigara DataChannel Channel: %s"), *Interface->Channel.ChannelName.ToString());
+						UE_LOG(LogNiagara, Warning, TEXT("Failed to find or add Naigara DataChannel Channel: %s"), *Interface->Channel.GetName());
 						return false;
 					}
 				}
@@ -191,7 +229,7 @@ struct FNDIDataChannelWriteInstanceData
 		//Verify our function info.
 		if(!ensure(Interface->GetCompiledData().GetFunctionInfo().Num() == FunctionToDatSetBindingInfo.Num()))
 		{
-			UE_LOG(LogNiagara, Warning, TEXT("Invalid Bindings for Niagara Data Interface Data Channel Write: %s"), *Interface->Channel.ChannelName.ToString());
+			UE_LOG(LogNiagara, Warning, TEXT("Invalid Bindings for Niagara Data Interface Data Channel Write: %s"), *Interface->Channel.GetName());
 			return false;			
 		}
 
@@ -199,7 +237,7 @@ struct FNDIDataChannelWriteInstanceData
 		{
 			if(Binding.IsValid() == false)
 			{
-				UE_LOG(LogNiagara, Warning, TEXT("Invalid Bindings for Niagara Data Interface Data Channel Write: %s"), *Interface->Channel.ChannelName.ToString());
+				UE_LOG(LogNiagara, Warning, TEXT("Invalid Bindings for Niagara Data Interface Data Channel Write: %s"), *Interface->Channel.GetName());
 				return false;			
 			}
 		}
@@ -216,20 +254,27 @@ struct FNDIDataChannelWriteInstanceData
 			if (GbDebugDumpWriter)
 			{
 				FNiagaraDataBuffer& Buffer = Data->GetCurrentDataChecked();
-				Buffer.Dump(0, Buffer.GetNumInstances(), FString::Printf(TEXT("=== Data Channle Write: %d Elements --> %s ==="), Buffer.GetNumInstances(), *Interface->Channel.ChannelName.ToString()));
+				Buffer.Dump(0, Buffer.GetNumInstances(), FString::Printf(TEXT("=== Data Channle Write: %d Elements --> %s ==="), Buffer.GetNumInstances(), *Interface->Channel.GetName()));
 			}
 
 			if (Interface->ShouldPublish() && Data->GetCurrentData()->GetNumInstances() > 0)
 			{
 				if (UNiagaraDataChannelHandler* Channel = DataChannel.Get())
 				{
+					if(DataChannelData == nullptr || Interface->bUpdateDestinationDataEveryTick)
+					{
+						FNiagaraDataChannelSearchParameters SearchParams;
+						SearchParams.OwningComponent = Instance->GetAttachComponent();
+						DataChannelData = Channel->FindData(SearchParams, ENiagaraResourceAccess::WriteOnly);
+					}
+
 					FNiagaraDataChannelPublishRequest PublishRequest(Data->GetCurrentData());
 					PublishRequest.bVisibleToGame = Interface->bPublishToGame;
 					PublishRequest.bVisibleToCPUSims = Interface->bPublishToCPU;
 					PublishRequest.bVisibleToGPUSims = Interface->bPublishToGPU;
 					PublishRequest.Data = Data->GetCurrentData();
 					PublishRequest.LwcTile = Instance->GetLWCTile();
-					Channel->Publish(PublishRequest);
+					DataChannelData->Publish(PublishRequest);
 				}
 			}
 		}
@@ -351,11 +396,120 @@ void UNiagaraDataInterfaceDataChannelWrite::PostCompile()
 {
 	UNiagaraSystem* OwnerSystem = GetTypedOuter<UNiagaraSystem>();
 	CompiledData.Init(OwnerSystem, this);
+}
+
+#endif
+
+
+
+#if WITH_EDITOR	
+
+void UNiagaraDataInterfaceDataChannelWrite::GetFeedback(UNiagaraSystem* InAsset, UNiagaraComponent* InComponent, TArray<FNiagaraDataInterfaceError>& OutErrors, TArray<FNiagaraDataInterfaceFeedback>& OutWarnings, TArray<FNiagaraDataInterfaceFeedback>& OutInfo)
+{
+	INiagaraModule& NiagaraModule = FModuleManager::GetModuleChecked<INiagaraModule>("Niagara");
+	const INiagaraEditorOnlyDataUtilities& EditorOnlyDataUtilities = NiagaraModule.GetEditorOnlyDataUtilities();
+	UNiagaraDataInterface* RuntimeInstanceOfThis = EditorOnlyDataUtilities.IsEditorDataInterfaceInstance(this)
+		? EditorOnlyDataUtilities.GetResolvedRuntimeInstanceForEditorDataInterfaceInstance(*InAsset, *this)
+		: this;
+
+	UNiagaraDataInterfaceDataChannelWrite* RuntimeDI = Cast<UNiagaraDataInterfaceDataChannelWrite>(RuntimeInstanceOfThis);
+
+	if (!RuntimeDI)
+	{
+		return;
+	}
+
+	Super::GetFeedback(InAsset, InComponent, OutErrors, OutWarnings, OutInfo);
+
+	if (Channel == nullptr)
+	{
+		OutErrors.Emplace(LOCTEXT("DataChannelMissingFmt", "Data Channel Interface has no valid Data Channel."),
+			LOCTEXT("DataChannelMissingErrorSummaryFmt", "Missing Data Channel."),
+			FNiagaraDataInterfaceFix());	
+
+		return;
+	}
+
+	if(ShouldPublish() == false)
+	{
+		OutErrors.Emplace(FText::Format(LOCTEXT("DataChannelDoesNotExistErrorFmt", "Data Channel {0} does not publish it's data to the Game, CPU Simulations or GPU simulations."), FText::FromName(Channel.GetFName())),
+			LOCTEXT("DataChannelDoesNotPublishErrorSummaryFmt", "Data Channel DI does not publish."),
+			FNiagaraDataInterfaceFix());
+	}
 
-	if (const UNiagaraDataChannel* DataChannel = UNiagaraDataChannelDefinitions::FindDataChannel(Channel.ChannelName))
+	if (const UNiagaraDataChannel* DataChannel = RuntimeDI->Channel->Get())
 	{
-		OwnerSystem->RegisterDataChannelUse(DataChannel);
+		//Ensure the data channel contains all the parameters this function is requesting.
+		TConstArrayView<FNiagaraVariable> ChannelVars = DataChannel->GetVariables();
+		for (const FNDIDataChannelFunctionInfo& FuncInfo : RuntimeDI->GetCompiledData().GetFunctionInfo())
+		{
+			TArray<FNiagaraVariableBase> MissingParams;
+
+			auto VerifyChannelContainsParams = [&](const TArray<FNiagaraVariableBase>& Parameters)
+			{
+				for (const FNiagaraVariableBase& FuncParam : Parameters)
+				{
+					bool bParamFound = false;
+					for (const FNiagaraVariable& ChannelVar : ChannelVars)
+					{
+						FNiagaraVariable SWCVar(ChannelVar);
+
+						//We have to convert each channel var to SWC for comparison with the function variables as there is no reliable way to go back from the SWC function var to the originating LWC var.
+						if (ChannelVar.GetType().IsEnum() == false)
+						{
+							UScriptStruct* ChannelSWCStruct = FNiagaraTypeHelper::GetSWCStruct(ChannelVar.GetType().GetScriptStruct());
+							if (ChannelSWCStruct)
+							{
+								FNiagaraTypeDefinition SWCType(ChannelSWCStruct, FNiagaraTypeDefinition::EAllowUnfriendlyStruct::Deny);
+								SWCVar = FNiagaraVariable(SWCType, ChannelVar.GetName());
+							}
+						}
+
+						if (SWCVar == FuncParam)
+						{
+							bParamFound = true;
+							break;
+						}
+					}
+
+					if (bParamFound == false)
+					{
+						MissingParams.Add(FuncParam);
+					}
+				}
+			};
+			VerifyChannelContainsParams(FuncInfo.Inputs);
+			VerifyChannelContainsParams(FuncInfo.Outputs);
+
+			if (MissingParams.Num() > 0)
+			{
+				FTextBuilder Builder;
+				Builder.AppendLineFormat(LOCTEXT("FuncParamMissingFromDataChannelWriteErrorFmt", "Accessing variables that do not exist in Data Channel {0}."), FText::FromName(Channel.GetFName()));
+				for (FNiagaraVariableBase& Param : MissingParams)
+				{
+					Builder.AppendLineFormat(LOCTEXT("FuncParamMissingFromDataChannelWriteErrorLineFmt", "{0} {1}"), Param.GetType().GetNameText(), FText::FromName(Param.GetName()));
+				}
+
+				OutErrors.Emplace(Builder.ToText(), LOCTEXT("FuncParamMissingFromDataChannelWriteErrorSummaryFmt", "Data Channel DI function is accessing invalid parameters."), FNiagaraDataInterfaceFix());
+			}
+		}
 	}
+	else
+	{
+		OutErrors.Emplace(FText::Format(LOCTEXT("DataChannelDoesNotExistErrorFmt", "Data Channel {0} does not exist. It may have been deleted."), FText::FromName(Channel.GetFName())),
+			LOCTEXT("DataChannelDoesNotExistErrorSummaryFmt", "Data Channel DI is accesssinga a Data Channel that doesn't exist."),
+			FNiagaraDataInterfaceFix());
+	}
+}
+
+void UNiagaraDataInterfaceDataChannelWrite::ValidateFunction(const FNiagaraFunctionSignature& Function, TArray<FText>& OutValidationErrors)
+{
+	Super::ValidateFunction(Function, OutValidationErrors);
+
+	//It would be great to be able to validate the parameters on the function calls here but this is only called on the DI CDO. We don't have the context of which data channel we'll be accessing.
+	//The translator should have all the required data to use the actual DIs when validating functions. We just need to do some wrangling to pull it from the pre compiled data correctly.
+	//This would probably also allow us to actually call hlsl generation functions on the actual DIs rather than their CDOs. Which would allow for a bunch of better optimized code gen for things like fluids.
+	//TODO!!!
 }
 
 #endif
@@ -370,7 +524,8 @@ bool UNiagaraDataInterfaceDataChannelWrite::Equals(const UNiagaraDataInterface*
 			bPublishToGame == OtherTyped->bPublishToGame &&
 			bPublishToCPU == OtherTyped->bPublishToCPU &&
 			bPublishToGPU == OtherTyped->bPublishToGPU &&
-			Channel.ChannelName == OtherTyped->Channel.ChannelName)
+			Channel == OtherTyped->Channel &&
+			bUpdateDestinationDataEveryTick == OtherTyped->bUpdateDestinationDataEveryTick)
 		{
 			return true;
 		}
@@ -392,8 +547,9 @@ bool UNiagaraDataInterfaceDataChannelWrite::CopyToInternal(UNiagaraDataInterface
 		DestTyped->bPublishToGame = bPublishToGame;
 		DestTyped->bPublishToCPU = bPublishToCPU;
 		DestTyped->bPublishToGPU = bPublishToGPU;
-		DestTyped->Channel.ChannelName = Channel.ChannelName;
+		DestTyped->Channel = Channel;
 		DestTyped->CompiledData = CompiledData;
+		DestTyped->bUpdateDestinationDataEveryTick = bUpdateDestinationDataEveryTick;
 		return true;
 	}
 
@@ -414,6 +570,10 @@ void UNiagaraDataInterfaceDataChannelWrite::GetFunctions(TArray<FNiagaraFunction
 		Sig.AddOutput(FNiagaraVariable(FNiagaraTypeDefinition::GetIntDef(), TEXT("Num")));
 		OutFunctions.Add(Sig);
 	}
+
+	static FNiagaraVariable EmitVar(FNiagaraTypeDefinition::GetBoolDef(), TEXT("Emit"));
+	EmitVar.SetValue(FNiagaraBool(true));
+
 	{
 		FNiagaraFunctionSignature Sig;
 		Sig.Name = NDIDataChannelWriteLocal::WriteName;
@@ -424,7 +584,7 @@ void UNiagaraDataInterfaceDataChannelWrite::GetFunctions(TArray<FNiagaraFunction
 		Sig.bRequiresExecPin = true;
 		Sig.bExperimental = true;
 		Sig.AddInput(FNiagaraVariable(FNiagaraTypeDefinition(GetClass()), TEXT("DataChannel interface")));
-		Sig.AddInput(FNiagaraVariable(FNiagaraTypeDefinition::GetBoolDef(), TEXT("Emit")));
+		Sig.AddInput(EmitVar);
 		Sig.AddInput(FNiagaraVariable(FNiagaraTypeDefinition::GetIntDef(), TEXT("Index")));
 		Sig.AddOutput(FNiagaraVariable(FNiagaraTypeDefinition::GetBoolDef(), TEXT("Success")));
 		Sig.RequiredInputs = Sig.Inputs.Num();//The user defines what we write in the graph.
@@ -440,7 +600,7 @@ void UNiagaraDataInterfaceDataChannelWrite::GetFunctions(TArray<FNiagaraFunction
 		Sig.bRequiresExecPin = true;
 		Sig.bExperimental = true;
 		Sig.AddInput(FNiagaraVariable(FNiagaraTypeDefinition(GetClass()), TEXT("DataChannel interface")));
-		Sig.AddInput(FNiagaraVariable(FNiagaraTypeDefinition::GetBoolDef(), TEXT("Emit")));
+		Sig.AddInput(EmitVar);
 		Sig.AddOutput(FNiagaraVariable(FNiagaraTypeDefinition::GetBoolDef(), TEXT("Success")));
 		Sig.RequiredInputs = Sig.Inputs.Num();//The user defines what we write in the graph.
 		OutFunctions.Add(Sig);
@@ -456,21 +616,13 @@ void UNiagaraDataInterfaceDataChannelWrite::GetVMExternalFunction(const FVMExter
 	else
 	{
 		int32 FuncIndex = CompiledData.FindFunctionInfoIndex(BindingInfo.Name, BindingInfo.VariadicInputs, BindingInfo.VariadicOutputs);
-
-		if (FuncIndex != INDEX_NONE)
+		if (BindingInfo.Name == NDIDataChannelWriteLocal::WriteName)
 		{
-			if (BindingInfo.Name == NDIDataChannelWriteLocal::WriteName)
-			{
-				OutFunc = FVMExternalFunction::CreateLambda([this, FuncIndex](FVectorVMExternalFunctionContext& Context) { this->Write(Context, FuncIndex); });
-			}
-			else if (BindingInfo.Name == NDIDataChannelWriteLocal::AppendName)
-			{
-				OutFunc = FVMExternalFunction::CreateLambda([this, FuncIndex](FVectorVMExternalFunctionContext& Context) { this->Append(Context, FuncIndex); });
-			}
-			else
-			{
-				UE_LOG(LogTemp, Display, TEXT("Could not find data interface external function in %s. Received Name: %s"), *GetPathNameSafe(this), *BindingInfo.Name.ToString());
-			}
+			OutFunc = FVMExternalFunction::CreateLambda([this, FuncIndex](FVectorVMExternalFunctionContext& Context) { this->Write(Context, FuncIndex); });
+		}
+		else if (BindingInfo.Name == NDIDataChannelWriteLocal::AppendName)
+		{
+			OutFunc = FVMExternalFunction::CreateLambda([this, FuncIndex](FVectorVMExternalFunctionContext& Context) { this->Append(Context, FuncIndex); });
 		}
 		else
 		{
@@ -479,90 +631,6 @@ void UNiagaraDataInterfaceDataChannelWrite::GetVMExternalFunction(const FVMExter
 	}
 }
 
-/** Handles any number of variadic parameter inputs. */
-template<int32 EXPECTED_NUM_INPUTS>
-struct FNDIVariadicInputHandler
-{
-	TArray<FNDIInputParam<float>, TInlineAllocator<EXPECTED_NUM_INPUTS>> FloatInputs;
-	TArray<FNDIInputParam<int32>, TInlineAllocator<EXPECTED_NUM_INPUTS>> IntInputs;
-	TArray<FNDIInputParam<FFloat16>, TInlineAllocator<EXPECTED_NUM_INPUTS>> HalfInputs;
-
-	FNDIVariadicInputHandler(FVectorVMExternalFunctionContext& Context, const FNDIDataChannel_FuncToDataSetBindingPtr& BindingPtr)
-	{
-		FloatInputs.Reserve(BindingPtr->NumFloatComponents);
-		IntInputs.Reserve(BindingPtr->NumInt32Components);
-		HalfInputs.Reserve(BindingPtr->NumHalfComponents);
-		for (FNDIDataChannelRegisterBinding& VMBinding : BindingPtr->VMRegisterBindings)
-		{
-			if (VMBinding.DataType == (int32)ENiagaraBaseTypes::Float)
-			{
-				FloatInputs.Emplace(Context);
-			}
-			else if (VMBinding.DataType == (int32)ENiagaraBaseTypes::Int32 || VMBinding.DataType == (int32)ENiagaraBaseTypes::Bool)
-			{
-				IntInputs.Emplace(Context);
-			}
-			else if (VMBinding.DataType == (int32)ENiagaraBaseTypes::Half)
-			{
-				HalfInputs.Emplace(Context);
-			}
-			else
-			{
-				checkf(false, TEXT("Didn't find a binding for this function input. Likely an error in the building of the binding data and the Float/Int/Half Bindings arrays."));
-			}
-		}
-	}
-
-	void Advance()
-	{
-		for (FNDIInputParam<float>& Input : FloatInputs) { Input.Advance(); }
-		for (FNDIInputParam<int32>& Input : IntInputs) { Input.Advance(); }
-		for (FNDIInputParam<FFloat16>& Input : HalfInputs) { Input.Advance(); }
-	}
-
-	bool Process(FNiagaraDataBuffer* Data, uint32 Index, const FNDIDataChannel_FuncToDataSetBindingPtr& BindingInfo)
-	{
-		if (Data && BindingInfo)
-		{
-			//TODO: Optimize for long runs of writes to reduce binding/lookup overhead.
-			if (Index < Data->GetNumInstances())
-			{
-				for (FNDIDataChannelRegisterBinding& VMBinding : BindingInfo->VMRegisterBindings)
-				{
-					if (VMBinding.DataType == (int32)ENiagaraBaseTypes::Float)
-					{
-						if (VMBinding.DataSetRegisterIndex != INDEX_NONE)
-						{
-							*Data->GetInstancePtrFloat(VMBinding.DataSetRegisterIndex, (uint32)Index) = FloatInputs[VMBinding.FunctionRegisterIndex].Get();
-						}
-					}
-					else if (VMBinding.DataType == (int32)ENiagaraBaseTypes::Int32 || VMBinding.DataType == (int32)ENiagaraBaseTypes::Bool)
-					{
-						if (VMBinding.DataSetRegisterIndex != INDEX_NONE)
-						{
-							*Data->GetInstancePtrInt32(VMBinding.DataSetRegisterIndex, (uint32)Index) = IntInputs[VMBinding.FunctionRegisterIndex].Get();
-						}
-					}
-					else if (VMBinding.DataType == (int32)ENiagaraBaseTypes::Half)
-					{
-						if (VMBinding.DataSetRegisterIndex != INDEX_NONE)
-						{
-							*Data->GetInstancePtrHalf(VMBinding.DataSetRegisterIndex, (uint32)Index) = HalfInputs[VMBinding.FunctionRegisterIndex].Get();
-						}
-					}
-					else
-					{
-						checkf(false, TEXT("Didn't find a binding for this function input. Likely an error in the building of the binding data and the Float/Int/Half Bindings arrays."));
-					}
-				}
-
-				return true;
-			}
-		}
-		return false;
-	}
-};
-
 void UNiagaraDataInterfaceDataChannelWrite::Num(FVectorVMExternalFunctionContext& Context)
 {
 	VectorVM::FUserPtrHandler<FNDIDataChannelWriteInstanceData> InstData(Context);
@@ -590,36 +658,46 @@ void UNiagaraDataInterfaceDataChannelWrite::Write(FVectorVMExternalFunctionConte
 	FNDIInputParam<int32> InIndex(Context);
 
 	const FNDIDataChannelFunctionInfo& FuncInfo = CompiledData.GetFunctionInfo()[FuncIdx];
-	const FNDIDataChannel_FuncToDataSetBindingPtr BindingInfo = InstData->FunctionToDatSetBindingInfo.IsValidIndex(FuncIdx) ? InstData->FunctionToDatSetBindingInfo[FuncIdx] : nullptr;
+	const FNDIDataChannel_FunctionToDataSetBinding* BindingInfo = InstData->FunctionToDatSetBindingInfo.IsValidIndex(FuncIdx) ? InstData->FunctionToDatSetBindingInfo[FuncIdx].Get() : nullptr;
 	FNDIVariadicInputHandler<16> VariadicInputs(Context, BindingInfo);//TODO: Make static / avoid allocation
 
 	FNDIOutputParam<bool> OutSuccess(Context);
 
-	if (InstData->Data && INiagaraModule::DataChannelsEnabled())
+	if (InstData->Data && BindingInfo && INiagaraModule::DataChannelsEnabled())
 	{
-		FNiagaraDataBuffer* Buffer = InstData->Data->GetDestinationData();
+		FNiagaraDataBuffer* Data = InstData->Data->GetDestinationData();
 
 		for (int32 i = 0; i < Context.GetNumInstances(); ++i)
 		{
 			int32 Index = InIndex.GetAndAdvance();
-
 			bool bEmit = InEmit.GetAndAdvance();
+			bEmit &= (uint32)Index < Data->GetNumInstances();
 
 			bool bSuccess = false;
 
 			//TODO: Optimize case where emit is constant
 			//TODO: Optimize for runs of sequential true emits.
-			if (bEmit)
+			auto FloatFunc = [Data, Index](const FNDIDataChannelRegisterBinding& VMBinding, float FloatData)
 			{
-				bSuccess = VariadicInputs.Process(Buffer, Index, BindingInfo);
-			}
+				if(VMBinding.DataSetRegisterIndex != INDEX_NONE)
+					*Data->GetInstancePtrFloat(VMBinding.DataSetRegisterIndex, (uint32)Index) = FloatData; 
+			};
+			auto IntFunc = [Data, Index](const FNDIDataChannelRegisterBinding& VMBinding, int32 IntData)
+			{
+				if(VMBinding.DataSetRegisterIndex != INDEX_NONE)
+					*Data->GetInstancePtrInt32(VMBinding.DataSetRegisterIndex, (uint32)Index) = IntData; 
+			};
+			auto HalfFunc = [Data, Index](const FNDIDataChannelRegisterBinding& VMBinding, FFloat16 HalfData)
+			{
+				if(VMBinding.DataSetRegisterIndex != INDEX_NONE)
+					*Data->GetInstancePtrHalf(VMBinding.DataSetRegisterIndex, (uint32)Index) = HalfData; 
+			};
+			bSuccess = VariadicInputs.Process(bEmit, BindingInfo, FloatFunc, IntFunc, HalfFunc);
 
 			if (OutSuccess.IsValid())
 			{
 				OutSuccess.SetAndAdvance(bSuccess);
 			}
-
-			VariadicInputs.Advance();
 		}
 	}
 	else
@@ -641,14 +719,14 @@ void UNiagaraDataInterfaceDataChannelWrite::Append(FVectorVMExternalFunctionCont
 	FNDIInputParam<bool> InEmit(Context);
 
 	const FNDIDataChannelFunctionInfo& FuncInfo = CompiledData.GetFunctionInfo()[FuncIdx];
-	const FNDIDataChannel_FuncToDataSetBindingPtr BindingInfo = InstData->FunctionToDatSetBindingInfo.IsValidIndex(FuncIdx) ? InstData->FunctionToDatSetBindingInfo[FuncIdx] : nullptr;
+	const FNDIDataChannel_FunctionToDataSetBinding* BindingInfo = InstData->FunctionToDatSetBindingInfo.IsValidIndex(FuncIdx) ? InstData->FunctionToDatSetBindingInfo[FuncIdx].Get() : nullptr;
 	FNDIVariadicInputHandler<16> VariadicInputs(Context, BindingInfo);//TODO: Make static / avoid allocation
 
 	FNDIOutputParam<bool> OutSuccess(Context);
 
-	if (InstData->Data && INiagaraModule::DataChannelsEnabled())
+	if (InstData->Data && BindingInfo &&  INiagaraModule::DataChannelsEnabled())
 	{
-		FNiagaraDataBuffer* Buffer = InstData->Data->GetDestinationData();
+		FNiagaraDataBuffer* Data = InstData->Data->GetDestinationData();
 		for (int32 i = 0; i < Context.GetNumInstances(); ++i)
 		{
 			bool bEmit = InEmit.GetAndAdvance();
@@ -657,22 +735,39 @@ void UNiagaraDataInterfaceDataChannelWrite::Append(FVectorVMExternalFunctionCont
 
 			//TODO: Optimize case where emit is constant
 			//TODO: Optimize for runs of sequential true emits.
-			if (bEmit && Buffer)
+			if (bEmit && Data)
 			{
-				int32 Num = Buffer->GetNumInstances();
-				if ((int32)Buffer->GetNumInstancesAllocated() > Num)
+				int32 Num = Data->GetNumInstances();
+				bEmit &= (int32)Data->GetNumInstancesAllocated() > Num;
+
+				if (bEmit)
 				{
-					Buffer->SetNumInstances(Num + 1);
-					bSuccess = VariadicInputs.Process(Buffer, Num, BindingInfo);
+					Data->SetNumInstances(Num + 1);
 				}
+
+				int32 Index = Num;
+				auto FloatFunc = [Data, Index](const FNDIDataChannelRegisterBinding& VMBinding, float FloatData)
+				{
+					if (VMBinding.DataSetRegisterIndex != INDEX_NONE)
+						*Data->GetInstancePtrFloat(VMBinding.DataSetRegisterIndex, (uint32)Index) = FloatData;
+				};
+				auto IntFunc = [Data, Index](const FNDIDataChannelRegisterBinding& VMBinding, int32 IntData)
+				{
+					if (VMBinding.DataSetRegisterIndex != INDEX_NONE)
+						*Data->GetInstancePtrInt32(VMBinding.DataSetRegisterIndex, (uint32)Index) = IntData;
+				};
+				auto HalfFunc = [Data, Index](const FNDIDataChannelRegisterBinding& VMBinding, FFloat16 HalfData)
+				{
+					if (VMBinding.DataSetRegisterIndex != INDEX_NONE)
+						*Data->GetInstancePtrHalf(VMBinding.DataSetRegisterIndex, (uint32)Index) = HalfData;
+				};
+				bSuccess = VariadicInputs.Process(bEmit, BindingInfo, FloatFunc, IntFunc, HalfFunc);
 			}
 
 			if (OutSuccess.IsValid())
 			{
 				OutSuccess.SetAndAdvance(bSuccess);
 			}
-
-			VariadicInputs.Advance();
 		}
 	}
 	else
diff --git a/Engine/Plugins/FX/Niagara/Source/Niagara/Private/NiagaraDataChannel.cpp b/Engine/Plugins/FX/Niagara/Source/Niagara/Private/NiagaraDataChannel.cpp
--- a/Engine/Plugins/FX/Niagara/Source/Niagara/Private/NiagaraDataChannel.cpp
+++ b/Engine/Plugins/FX/Niagara/Source/Niagara/Private/NiagaraDataChannel.cpp
@@ -3,11 +3,14 @@
 #include "NiagaraDataChannel.h"
 
 #include "NiagaraWorldManager.h"
+#include "NiagaraDataChannelCommon.h"
 #include "NiagaraDataChannelHandler.h"
 #include "NiagaraDataChannelManager.h"
 #include "Engine/Engine.h"
 #include "UObject/UObjectGlobals.h"
 #include "UObject/Package.h"
+#include "NiagaraDataChannelAccessor.h"
+#include "NiagaraGpuComputeDispatchInterface.h"
 
 #include UE_INLINE_GENERATED_CPP_BY_NAME(NiagaraDataChannel)
 
@@ -21,11 +24,16 @@ void FNiagaraDataChannelGameDataLayout::Init(const TArray<FNiagaraVariable>& Var
 	LwcConverters.Reserve(Variables.Num());
 	for (const FNiagaraVariable& Var : Variables)
 	{
-		int32& VarIdx = VariableIndices.Add(Var);
+		//Sigh.
+		//We must convert from the variable stored var in the data channels definition as we currently cannot serialize/store actual LWC types in FNiagawraTypeDefinitions.
+		FNiagaraTypeDefinition LWCType = FNiagaraTypeHelper::GetLWCType(Var.GetType());
+		FNiagaraVariableBase LWCVar(LWCType, Var.GetName());
+
+		int32& VarIdx = VariableIndices.Add(LWCVar);
 		VarIdx = VariableIndices.Num() - 1;
 
 		FNiagaraLwcStructConverter& Converter = LwcConverters.AddDefaulted_GetRef();
-		Converter = FNiagaraTypeRegistry::GetStructConverter(Var.GetType());
+		Converter = FNiagaraTypeRegistry::GetStructConverter(LWCType);
 	}
 }
 
@@ -55,12 +63,14 @@ void FNiagaraDataChannelGameData::Empty()
 	}
 }
 
-void FNiagaraDataChannelGameData::Reset()
+void FNiagaraDataChannelGameData::BeginFrame()
 {
+	bool bKeepPrevious = DataChannel->KeepPreviousFrameData();
+	PrevNumElements = bKeepPrevious ? NumElements : 0;
 	NumElements = 0;
 	for (FNiagaraDataChannelVariableBuffer& VarData : VariableData)
 	{
-		VarData.Reset();
+		VarData.BeginFrame(bKeepPrevious);
 	}
 }
 
@@ -76,14 +86,26 @@ void FNiagaraDataChannelGameData::SetNum(int32 NewNum)
 FNiagaraDataChannelVariableBuffer* FNiagaraDataChannelGameData::FindVariableBuffer(const FNiagaraVariableBase& Var)
 {
 	const FNiagaraDataChannelGameDataLayout& Layout = DataChannel->GetGameDataLayout();
-	if(const int32* Idx = Layout.VariableIndices.Find(Var))
+	const FNiagaraTypeDefinition& VarType = Var.GetType();
+	for(auto& Pair : Layout.VariableIndices)
 	{
-		return &VariableData[*Idx];
+		const FNiagaraVariableBase& LayoutVar = Pair.Key;
+		const FNiagaraTypeDefinition& LayoutVarType = LayoutVar.GetType();
+		int32 Index = Pair.Value;
+		
+		if(Var.GetName() == LayoutVar.GetName())
+		{
+			//For enum variables we'll hack things a little so that correctly named ints also match. This gets around some limitations in calling code not being able to provide the correct enum types.
+			if(VarType == LayoutVarType || (LayoutVarType.IsEnum() && VarType == FNiagaraTypeDefinition::GetIntDef()))
+			{
+				return &VariableData[Index];
+			}
+		}
 	}
 	return nullptr;
 }
 
-void FNiagaraDataChannelGameData::WriteToDataSet(FNiagaraDataBuffer* DestBuffer, int32 DestStartIdx)
+void FNiagaraDataChannelGameData::WriteToDataSet(FNiagaraDataBuffer* DestBuffer, int32 DestStartIdx, FVector3f SimulationLwcTile)
 {
 	const FNiagaraDataSetCompiledData& CompiledData = DestBuffer->GetOwner()->GetCompiledData();
 
@@ -96,16 +118,16 @@ void FNiagaraDataChannelGameData::WriteToDataSet(FNiagaraDataBuffer* DestBuffer,
 		return;
 	}
 	
-	DestBuffer->SetNumInstances(NumInstances);
+	DestBuffer->SetNumInstances(DestStartIdx + NumInstances);
 
 	const FNiagaraDataChannelGameDataLayout& Layout = DataChannel->GetGameDataLayout();
 
 	for (const TPair<FNiagaraVariableBase, int32>& VarIndexPair : Layout.VariableIndices)
 	{
-		const FNiagaraVariableBase& Var = VarIndexPair.Key;
+		FNiagaraVariableBase Var = VarIndexPair.Key;
 		int32 VarIndex = VarIndexPair.Value;
 		FNiagaraDataChannelVariableBuffer& VarBuffer = VariableData[VarIndex];
-		uint8* SrcData = VarBuffer.Data.GetData();
+		uint8* SrcDataBase = VarBuffer.Data.GetData();
 		
 		FNiagaraVariableBase SimVar = Var;
 
@@ -115,23 +137,18 @@ void FNiagaraDataChannelGameData::WriteToDataSet(FNiagaraDataBuffer* DestBuffer,
 			SimVar.SetType(FNiagaraTypeDefinition(FNiagaraTypeHelper::FindNiagaraFriendlyTopLevelStruct(CastChecked<UScriptStruct>(Var.GetType().GetStruct()), ENiagaraStructConversion::Simulation)));
 		}
 
+		//Niagara Positions are a special case where we're actually storing them as FVectors in the game level data and must convert down to actual Positions/FVector3f in the sim data.
+		if(Var.GetType() == FNiagaraTypeDefinition::GetPositionDef())
+		{
+			Var.SetType(FNiagaraTypeHelper::GetVectorDef());
+		}
+
 		int32 SimVarIndex = CompiledData.Variables.IndexOfByKey(SimVar);
 		if (SimVarIndex == INDEX_NONE)
 		{
 			continue; //Did not find this variable in the dataset. Warn?
 		}
 
-		//Convert from LWC types to Niagara Simulation Types where required.
-// 		if (FNiagaraTypeHelper::IsLWCType(Var.GetType()))
-// 		{
-// 			LWCConversionBuffer.SetNumUninitialized(SimVar.GetSizeInBytes() * NumInstances, false);
-// 			
-// 			//TODO: Can we just avoid all this and convert any doubles direct to floats when we transfer the date into the dest buffer?
-// 			FNiagaraLwcStructConverter& Converter = LwcConverters[VarIndex];
-// 			Converter.ConvertDataToSimulation(LWCConversionBuffer.GetData(), SrcData, NumInstances);
-// 			SrcData = LWCConversionBuffer.GetData();
-// 		}
-
 		int32 SrcVarSize = Var.GetSizeInBytes();
 		int32 DestVarSize = SimVar.GetSizeInBytes();
 		const FNiagaraVariableLayoutInfo& SimLayout = CompiledData.VariableLayouts[SimVarIndex];
@@ -140,75 +157,100 @@ void FNiagaraDataChannelGameData::WriteToDataSet(FNiagaraDataBuffer* DestBuffer,
 		int32 IntCompIdx = SimLayout.Int32ComponentStart;
 		int32 HalfCompIdx = SimLayout.HalfComponentStart;
 
-		TFunction<void(UScriptStruct*, UScriptStruct*)> WriteData;
-		WriteData = [&](UScriptStruct* SrcStruct, UScriptStruct* DestStruct)
+		TFunction<void(UScriptStruct*, UScriptStruct*, uint8*)> WriteData;
+		WriteData = [&](UScriptStruct* SrcStruct, UScriptStruct* DestStruct, uint8* SrcPropertyBase)
 		{
-			TFieldIterator<FProperty> SrcPropertyIt(SrcStruct, EFieldIteratorFlags::IncludeSuper);
-			TFieldIterator<FProperty> DestPropertyIt(DestStruct, EFieldIteratorFlags::IncludeSuper);
-			for (; SrcPropertyIt; ++SrcPropertyIt, ++DestPropertyIt)
-			{
-				FProperty* SrcProperty = *SrcPropertyIt;
-				FProperty* DestProperty = *DestPropertyIt;
-				int32 SrcOffset = SrcProperty->GetOffset_ForInternal();
-
-				//Convert any LWC doubles to floats. //TODO: Insert LWCTile... probably need to explicitly check for vectors etc.
-				if (SrcProperty->IsA(FDoubleProperty::StaticClass()))
+			//Write all data from the data channel into the data set. Converting into from LWC into the local LWC Tile space as we go.
+
+			uint8* SrcData = SrcPropertyBase;
+
+			//Positions are a special case that are stored as FVectors in game data but converted to an LWCTile local FVector3f in simulation data.
+			if (DestStruct == FNiagaraTypeDefinition::GetPositionStruct())
+			{	
+				float* DestX = (float*)DestBuffer->GetInstancePtrFloat(FloatCompIdx++, DestStartIdx);
+				float* DestY = (float*)DestBuffer->GetInstancePtrFloat(FloatCompIdx++, DestStartIdx);
+				float* DestZ = (float*)DestBuffer->GetInstancePtrFloat(FloatCompIdx++, DestStartIdx);
+								
+				for (int32 i = 0; i < NumInstances; ++i)
 				{
-					check(DestProperty->IsA(FFloatProperty::StaticClass()));
-					float* Dest = (float*)DestBuffer->GetComponentPtrFloat(FloatCompIdx++);
-					//Write all instances for this component.
-					for (int32 i = 0; i < NumInstances; ++i)
-					{
-						double* Src = (double*)((SrcData + i * SrcVarSize) + SrcOffset);
-						*Dest++ = *Src;
-					}
+					FVector* Src = (FVector*)((SrcData + i * SrcVarSize));
+
+					FVector3f SimLocalSWC = FVector3f((*Src) - FVector(SimulationLwcTile) * FLargeWorldRenderScalar::GetTileSize());
+					*DestX++ = SimLocalSWC.X;
+					*DestY++ = SimLocalSWC.Y;
+					*DestZ++ = SimLocalSWC.Z;
 				}
-				else if (SrcProperty->IsA(FFloatProperty::StaticClass()))
+			}
+			else
+			{
+				TFieldIterator<FProperty> SrcPropertyIt(SrcStruct, EFieldIteratorFlags::IncludeSuper);
+				TFieldIterator<FProperty> DestPropertyIt(DestStruct, EFieldIteratorFlags::IncludeSuper);
+				for (; SrcPropertyIt; ++SrcPropertyIt, ++DestPropertyIt)
 				{
-					float* Dest = (float*)DestBuffer->GetComponentPtrFloat(FloatCompIdx++);
+					FProperty* SrcProperty = *SrcPropertyIt;
+					FProperty* DestProperty = *DestPropertyIt;
+					int32 SrcOffset = SrcProperty->GetOffset_ForInternal();
+					SrcData = SrcPropertyBase + SrcProperty->GetOffset_ForInternal();
 
-					//Write all instances for this component.
-					for (int32 i = 0; i < NumInstances; ++i)
+					//Convert any LWC doubles to floats. //TODO: Insert LWCTile... probably need to explicitly check for vectors etc.
+					if (SrcProperty->IsA(FDoubleProperty::StaticClass()))
 					{
-						float* Src = (float*)((SrcData + i * SrcVarSize) + SrcOffset);
-						*Dest++ = *Src;
+						check(DestProperty->IsA(FFloatProperty::StaticClass()));
+						float* Dest = (float*)DestBuffer->GetInstancePtrFloat(FloatCompIdx++, DestStartIdx);
+						//Write all instances for this component.
+						for (int32 i = 0; i < NumInstances; ++i)
+						{
+							double* Src = (double*)((SrcData + i * SrcVarSize));
+							*Dest++ = *Src;
+						}
 					}
-				}
-				else if (SrcProperty->IsA(FUInt16Property::StaticClass()))
-				{
-					FFloat16* Dest = (FFloat16*)DestBuffer->GetComponentPtrHalf(HalfCompIdx++);
-
-					//Write all instances for this component.
-					for (int32 i = 0; i < NumInstances; ++i)
+					else if (SrcProperty->IsA(FFloatProperty::StaticClass()))
 					{
-						FFloat16* Src = (FFloat16*)((SrcData + i * SrcVarSize) + SrcOffset);
-						*Dest++ = *Src;						
+						float* Dest = (float*)DestBuffer->GetInstancePtrFloat(FloatCompIdx++, DestStartIdx);
+
+						//Write all instances for this component.
+						for (int32 i = 0; i < NumInstances; ++i)
+						{
+							float* Src = (float*)((SrcData + i * SrcVarSize));
+							*Dest++ = *Src;
+						}
 					}
-				}
-				else if (SrcProperty->IsA(FIntProperty::StaticClass()) || SrcProperty->IsA(FBoolProperty::StaticClass()))
-				{
-					int32* Dest = (int32*)DestBuffer->GetComponentPtrInt32(IntCompIdx++);
-
-					//Write all instances for this component.
-					for (int32 i = 0; i < NumInstances; ++i)
+					else if (SrcProperty->IsA(FUInt16Property::StaticClass()))
 					{
-						int32* Src = (int32*)((SrcData + i * SrcVarSize) + SrcOffset);
-						*Dest++ = *Src;						
+						FFloat16* Dest = (FFloat16*)DestBuffer->GetInstancePtrHalf(HalfCompIdx++, DestStartIdx);
+
+						//Write all instances for this component.
+						for (int32 i = 0; i < NumInstances; ++i)
+						{
+							FFloat16* Src = (FFloat16*)((SrcData + i * SrcVarSize));
+							*Dest++ = *Src;
+						}
+					}
+					else if (SrcProperty->IsA(FIntProperty::StaticClass()) || SrcProperty->IsA(FBoolProperty::StaticClass()))
+					{
+						int32* Dest = (int32*)DestBuffer->GetInstancePtrInt32(IntCompIdx++, DestStartIdx);
+
+						//Write all instances for this component.
+						for (int32 i = 0; i < NumInstances; ++i)
+						{
+							int32* Src = (int32*)((SrcData + i * SrcVarSize));
+							*Dest++ = *Src;
+						}
+					}
+					//Should be able to support double easily enough
+					else if (FStructProperty* StructProp = CastField<FStructProperty>(SrcProperty))
+					{
+						FStructProperty* DestStructProp = CastField<FStructProperty>(DestProperty);
+						WriteData(StructProp->Struct, DestStructProp->Struct, SrcData);
+					}
+					else
+					{
+						checkf(false, TEXT("Property(%s) Class(%s) is not a supported type"), *SrcProperty->GetName(), *SrcProperty->GetClass()->GetName());
 					}
-				}
-				//Should be able to support double easily enough
-				else if (FStructProperty* StructProp = CastField<FStructProperty>(SrcProperty))
-				{
-					FStructProperty* DestStructProp = CastField<FStructProperty>(DestProperty);
-					WriteData(StructProp->Struct, DestStructProp->Struct);
-				}
-				else
-				{
-					checkf(false, TEXT("Property(%s) Class(%s) is not a supported type"), *SrcProperty->GetName(), *SrcProperty->GetClass()->GetName());
 				}
 			}
 		};
-		WriteData(Var.GetType().GetScriptStruct(), SimVar.GetType().GetScriptStruct());
+		WriteData(Var.GetType().GetScriptStruct(), SimVar.GetType().GetScriptStruct(), SrcDataBase);
 	}
 }
 void FNiagaraDataChannelGameData::AppendFromGameData(const FNiagaraDataChannelGameData& GameData)
@@ -225,7 +267,7 @@ void FNiagaraDataChannelGameData::AppendFromGameData(const FNiagaraDataChannelGa
 	}
 }
 
-void FNiagaraDataChannelGameData::AppendFromDataSet(const FNiagaraDataBuffer* SrcBuffer, FVector3f LwcTile)
+void FNiagaraDataChannelGameData::AppendFromDataSet(const FNiagaraDataBuffer* SrcBuffer, FVector3f SimulationLwcTile)
 {
 	const FNiagaraDataSetCompiledData& CompiledData = SrcBuffer->GetOwner()->GetCompiledData();
 
@@ -238,22 +280,26 @@ void FNiagaraDataChannelGameData::AppendFromDataSet(const FNiagaraDataBuffer* Sr
 	const FNiagaraDataChannelGameDataLayout& Layout = DataChannel->GetGameDataLayout();
 	for (const TPair<FNiagaraVariableBase, int32>& VarIndexPair : Layout.VariableIndices)
 	{
-		const FNiagaraVariableBase& Var = VarIndexPair.Key;
+		FNiagaraVariableBase Var = VarIndexPair.Key;
 		int32 VarIndex = VarIndexPair.Value;
 		FNiagaraDataChannelVariableBuffer& VarBuffer = VariableData[VarIndex];
 
-		int32 VarSize = Var.GetSizeInBytes();
-
 		VarBuffer.SetNum(NumElements);
 
 		uint8* DestDataBase = VarBuffer.Data.GetData();
-		uint8* DestData = DestDataBase;
 		
 		FNiagaraVariableBase SimVar = Var;
 		if (FNiagaraTypeHelper::IsLWCType(Var.GetType()))
 		{
 			SimVar.SetType(FNiagaraTypeDefinition(FNiagaraTypeHelper::FindNiagaraFriendlyTopLevelStruct(Var.GetType().GetScriptStruct(), ENiagaraStructConversion::Simulation)));
 		}
+
+		//Niagara Positions are a special case where we're actually storing them as FVectors in the game level data and must convert down to actual Positions/FVector3f in the sim data and visa versa
+		if(Var.GetType() == FNiagaraTypeDefinition::GetPositionDef())
+		{
+			Var.SetType(FNiagaraTypeHelper::GetVectorDef());
+		}
+		int32 VarSize = Var.GetSizeInBytes();
 		
 		int32 SimVarIndex = CompiledData.Variables.IndexOfByKey(SimVar);
 		if (SimVarIndex == INDEX_NONE)
@@ -261,151 +307,433 @@ void FNiagaraDataChannelGameData::AppendFromDataSet(const FNiagaraDataBuffer* Sr
 			continue; //Did not find this variable in the dataset. Warn?
 		}
 
-		//Convert from Niagara Simulation to LWC types where required.
-// 		if (FNiagaraTypeHelper::IsLWCType(Var.GetType()))
-// 		{
-// 			LWCConversionBuffer.SetNumUninitialized(Var.GetSizeInBytes() * NumInstances, false);
-// 
-// 			//Write instead to our intermediate conversion buffer and do the LWC convert after.
-// 			DestData = LWCConversionBuffer.GetData();
-// 		}
-
 		const FNiagaraVariableLayoutInfo& SimLayout = CompiledData.VariableLayouts[SimVarIndex];
 		
 		int32 FloatCompIdx = SimLayout.FloatComponentStart;
 		int32 IntCompIdx = SimLayout.Int32ComponentStart;
 		int32 HalfCompIdx = SimLayout.HalfComponentStart;
 		
-		TFunction<void(UScriptStruct*, UScriptStruct*)> ReadData;
-		ReadData = [&](UScriptStruct* SrcStruct, UScriptStruct* DestStruct)
+		TFunction<void(UScriptStruct*, UScriptStruct*, uint8*)> ReadData;
+		ReadData = [&](UScriptStruct* SrcStruct, UScriptStruct* DestStruct, uint8* DestData)
 		{
-			TFieldIterator<FProperty> SrcPropertyIt(SrcStruct, EFieldIteratorFlags::IncludeSuper);
-			TFieldIterator<FProperty> DestPropertyIt(DestStruct, EFieldIteratorFlags::IncludeSuper);			
-			for (; SrcPropertyIt; ++SrcPropertyIt, ++DestPropertyIt)
+			//Read all data from the simulation and place in the game level buffers. Converting from LWC local tile space where needed.
+
+			//Special case for writing to Niagara Positions. We must offset the data into the simulation local LWC space.
+			if (SrcStruct == FNiagaraTypeDefinition::GetPositionStruct())
 			{
-				FProperty* SrcProperty = *SrcPropertyIt;
-				FProperty* DestProperty = *DestPropertyIt;
-				DestData = DestDataBase + DestProperty->GetOffset_ForInternal();
-				if (DestPropertyIt->IsA(FDoubleProperty::StaticClass()))//TODO: Insert LWCTile... probably need to explicitly check for vectors etc.
-				{
-					double* Dest = (double*)(DestData);
-					float* Src = (float*)SrcBuffer->GetComponentPtrFloat(FloatCompIdx++);
+				float* SrcX = (float*)SrcBuffer->GetComponentPtrFloat(FloatCompIdx++);
+				float* SrcY = (float*)SrcBuffer->GetComponentPtrFloat(FloatCompIdx++);
+				float* SrcZ = (float*)SrcBuffer->GetComponentPtrFloat(FloatCompIdx++);
 
-					//Write all instances for this component.
-					for (int32 i = 0; i < NumInstances; ++i)
-					{
-						Dest = (double*)(DestData + VarSize * i);
-						*Dest = *Src++;
-					}
+				for (int32 i = 0; i < NumInstances; ++i)
+				{
+					FVector* Dest = (FVector*)(DestData + VarSize * i);
+					*Dest = FVector(*SrcX++, *SrcY++, *SrcZ++) + FVector(SimulationLwcTile) * FLargeWorldRenderScalar::GetTileSize();
 				}
-				else if (DestPropertyIt->IsA(FFloatProperty::StaticClass()))
+			}
+			else
+			{
+				TFieldIterator<FProperty> SrcPropertyIt(SrcStruct, EFieldIteratorFlags::IncludeSuper);
+				TFieldIterator<FProperty> DestPropertyIt(DestStruct, EFieldIteratorFlags::IncludeSuper);
+				for (; SrcPropertyIt; ++SrcPropertyIt, ++DestPropertyIt)
 				{
-					float* Dest = (float*)(DestData);
-					float* Src = (float*)SrcBuffer->GetComponentPtrFloat(FloatCompIdx++);
-
-					//Write all instances for this component.
-					for (int32 i = 0; i < NumInstances; ++i)
+					FProperty* SrcProperty = *SrcPropertyIt;
+					FProperty* DestProperty = *DestPropertyIt;
+					DestData += DestProperty->GetOffset_ForInternal();
+					if (DestPropertyIt->IsA(FDoubleProperty::StaticClass()))
 					{
-						Dest = (float*)(DestData + VarSize * i);
-						*Dest = *Src++;
+						double* Dest = (double*)(DestData);
+						float* Src = (float*)SrcBuffer->GetComponentPtrFloat(FloatCompIdx++);
+
+						//Write all instances for this component.
+						for (int32 i = 0; i < NumInstances; ++i)
+						{
+							Dest = (double*)(DestData + VarSize * i);
+							*Dest = *Src++;
+						}
 					}
-				}
-				else if (DestPropertyIt->IsA(FUInt16Property::StaticClass()))
-				{
-					FFloat16* Dest = (FFloat16*)(DestData);
-					FFloat16* Src = (FFloat16*)SrcBuffer->GetComponentPtrHalf(HalfCompIdx++);
-
-					//Write all instances for this component.
-					for (int32 i = 0; i < NumInstances; ++i)
+					else if (DestPropertyIt->IsA(FFloatProperty::StaticClass()))
 					{
-						Dest = (FFloat16*)(DestData + VarSize * i);
-						*Dest = *Src++;
+						float* Dest = (float*)(DestData);
+						float* Src = (float*)SrcBuffer->GetComponentPtrFloat(FloatCompIdx++);
+
+						//Write all instances for this component.
+						for (int32 i = 0; i < NumInstances; ++i)
+						{
+							Dest = (float*)(DestData + VarSize * i);
+							*Dest = *Src++;
+						}
 					}
-				}
-				else if (DestPropertyIt->IsA(FIntProperty::StaticClass()) || DestPropertyIt->IsA(FBoolProperty::StaticClass()))
-				{
-					int32* Dest = (int32*)(DestData);
-					int32* Src = (int32*)SrcBuffer->GetComponentPtrInt32(IntCompIdx++);
-
-					//Write all instances for this component.
-					for (int32 i = 0; i < NumInstances; ++i)
+					else if (DestPropertyIt->IsA(FUInt16Property::StaticClass()))
 					{
-						Dest = (int32*)(DestData + VarSize * i);
-						*Dest = *Src++;
+						FFloat16* Dest = (FFloat16*)(DestData);
+						FFloat16* Src = (FFloat16*)SrcBuffer->GetComponentPtrHalf(HalfCompIdx++);
+
+						//Write all instances for this component.
+						for (int32 i = 0; i < NumInstances; ++i)
+						{
+							Dest = (FFloat16*)(DestData + VarSize * i);
+							*Dest = *Src++;
+						}
+					}
+					else if (DestPropertyIt->IsA(FIntProperty::StaticClass()) || DestPropertyIt->IsA(FBoolProperty::StaticClass()))
+					{
+						int32* Dest = (int32*)(DestData);
+						int32* Src = (int32*)SrcBuffer->GetComponentPtrInt32(IntCompIdx++);
+
+						//Write all instances for this component.
+						for (int32 i = 0; i < NumInstances; ++i)
+						{
+							Dest = (int32*)(DestData + VarSize * i);
+							*Dest = *Src++;
+						}
+					}
+					else if (FStructProperty* SrcStructProp = CastField<FStructProperty>(SrcProperty))
+					{
+						FStructProperty* DestStructProp = CastField<FStructProperty>(DestProperty);
+						ReadData(SrcStructProp->Struct, DestStructProp->Struct, DestData);
+					}
+					else
+					{
+						checkf(false, TEXT("Property(%s) Class(%s) is not a supported type"), *SrcProperty->GetName(), *SrcProperty->GetClass()->GetName());
 					}
-				}
-				//Should be able to support double easily enough
-				else if (FStructProperty* SrcStructProp = CastField<FStructProperty>(SrcProperty))
-				{
-					FStructProperty* DestStructProp = CastField<FStructProperty>(DestProperty);
-					ReadData(SrcStructProp->Struct, DestStructProp->Struct);
-				}
-				else
-				{
-					checkf(false, TEXT("Property(%s) Class(%s) is not a supported type"), *SrcProperty->GetName(), *SrcProperty->GetClass()->GetName());
 				}
 			}
 		};
-		ReadData(SimVar.GetType().GetScriptStruct(), Var.GetType().GetScriptStruct());
+		ReadData(SimVar.GetType().GetScriptStruct(), Var.GetType().GetScriptStruct(), DestDataBase);
+	}
+}
 
-		//For LWC types we're written into an intermediate structure for the sim var. Now convert that into LWC proper.
-		if (FNiagaraTypeHelper::IsLWCType(Var.GetType()))
-		{
-			const FNiagaraLwcStructConverter& Converter = Layout.LwcConverters[VarIndex];
-			VarBuffer.SetNum(NumElements);
-			//TODO: Incorporate the LWC Tile
-			Converter.ConvertDataFromSimulation(VarBuffer.Data.GetData() + Var.GetSizeInBytes() * OrigNumDataChannel, LWCConversionBuffer.GetData(), NumInstances);
-		}
+//////////////////////////////////////////////////////////////////////////
+
+void FNiagaraDataChannelDataProxy::BeginFrame(bool bKeepPreviousFrameData)
+{
+	check(GPUDataSet);
+	check(GPUDataSet->GetSimTarget() == ENiagaraSimTarget::GPUComputeSim);
+
+}
+
+void FNiagaraDataChannelDataProxy::EndFrame(FNiagaraGpuComputeDispatchInterface* DispathInterface, FRHICommandListImmediate& CmdList, const TArray<FNiagaraDataBufferRef>& BuffersForGPU)
+{
+	check(GPUDataSet);
+	check(GPUDataSet->GetSimTarget() == ENiagaraSimTarget::GPUComputeSim);
+
+	PrevFrameData = GPUDataSet->GetCurrentData();
+
+	uint32 NumInstance = 0;
+	for (auto& Buffer : BuffersForGPU)
+	{
+		NumInstance += Buffer->GetNumInstances();
 	}
+
+	GPUDataSet->BeginSimulate();
+	FNiagaraDataBuffer* DestBuffer = GPUDataSet->GetDestinationData();	
+	DestBuffer->PushCPUBuffersToGPU(BuffersForGPU, true, CmdList, DispathInterface->GetFeatureLevel(), GetDebugName());
+	GPUDataSet->EndSimulate();
+
+	//For now we need not deal with the instance count manager but when we do GPU->GPU writes we will
+	//DestBuffer->SetGPUInstanceCountBufferOffset()
+}
+
+void FNiagaraDataChannelDataProxy::Reset()
+{
+	PrevFrameData = nullptr;
 }
 
 //////////////////////////////////////////////////////////////////////////
 
-FNiagaraWorldDataChannelStore::~FNiagaraWorldDataChannelStore()
+FNiagaraDataChannelData::FNiagaraDataChannelData(UNiagaraDataChannelHandler* Owner)
 {
-	Empty();
+	const UNiagaraDataChannel* DataChannel = Owner->GetDataChannel();
+	check(DataChannel);
+
+	GameData = DataChannel->CreateGameData();
+
+	CPUSimData = new FNiagaraDataSet();
+	GPUSimData = new FNiagaraDataSet();
+
+
+	const FNiagaraDataSetCompiledData& CompiledData = DataChannel->GetCompiledData(ENiagaraSimTarget::CPUSim);
+	const FNiagaraDataSetCompiledData& CompiledDataGPU = DataChannel->GetCompiledData(ENiagaraSimTarget::GPUComputeSim);
+	CPUSimData->Init(&CompiledData);
+	GPUSimData->Init(&CompiledDataGPU);
+
+	//TODO: Send game data to GPU direct without staging.
+	GameDataStaging = new FNiagaraDataSet();
+	GameDataStaging->Init(&CompiledData);
+
+	CPUSimData->BeginSimulate();
+	CPUSimData->EndSimulate();
+
+	RTProxy.Reset(new FNiagaraDataChannelDataProxy());
+	RTProxy->GPUDataSet = GPUSimData;
+	#if !UE_BUILD_SHIPPING
+	RTProxy->DebugName = FString::Printf(TEXT("%s__GPUData"), *DataChannel->GetName());
+	#endif
 }
 
-void FNiagaraWorldDataChannelStore::Empty()
+FNiagaraDataChannelData::~FNiagaraDataChannelData()
 {
 	GameData.Reset();
 
 	/** We defer the deletion of the dataset to the RT to be sure all in-flight RT commands have finished using it.*/
 	ENQUEUE_RENDER_COMMAND(FDeleteContextCommand)(
-		[CPUDataChannelDataSet = CPUSimData, GPUDataChannelDataSet = GPUSimData](FRHICommandListImmediate& RHICmdList)
+		[CPUDataChannelDataSet = CPUSimData, GPUDataChannelDataSet = GPUSimData, ReleasedRTProxy = RTProxy.Release()](FRHICommandListImmediate& RHICmdList)
+	{
+		if (CPUDataChannelDataSet != nullptr)
+		{
+			delete CPUDataChannelDataSet;
+		}
+		if (GPUDataChannelDataSet != nullptr)
+		{
+			delete GPUDataChannelDataSet;
+		}
+		if (ReleasedRTProxy)
 		{
-			if (CPUDataChannelDataSet != nullptr)
+			delete ReleasedRTProxy;
+		}
+	}
+	);
+	CPUSimData = nullptr;
+	GPUSimData = nullptr;
+}
+
+void FNiagaraDataChannelData::Reset()
+{
+	PublishRequests.Reset();
+	BuffersForGPU.Reset();
+	
+	PrevCPUSimData = nullptr;	
+
+	GameData->Empty();
+	CPUSimData->ResetBuffers();
+	GPUSimData->ResetBuffers();
+	GameDataStaging->ResetBuffers();
+	ENQUEUE_RENDER_COMMAND(FResetProxyCommand)(
+		[RT_Proxy = RTProxy.Get()](FRHICommandListImmediate& RHICmdList)
+	{
+		if (RT_Proxy)
+		{
+			RT_Proxy->Reset();
+		}
+	}
+	);
+}
+
+void FNiagaraDataChannelData::BeginFrame(UNiagaraDataChannelHandler* Owner)
+{
+	GameData->BeginFrame();
+
+	bool bRequirePreviousData = Owner->GetDataChannel()->KeepPreviousFrameData();
+	if (bRequirePreviousData)
+	{
+		PrevCPUSimData = CPUSimData->GetCurrentData();
+	}
+
+	//Grab a new buffer to store the CPU data.
+	CPUSimData->BeginSimulate();
+	CPUSimData->EndSimulate();
+	
+	check(RTProxy);
+	FNiagaraGpuComputeDispatchInterface* DispathInterface = FNiagaraGpuComputeDispatchInterface::Get(Owner->GetWorld());
+	ENQUEUE_RENDER_COMMAND(FDataChannelProxyBeginFrame) (
+		[RT_Proxy = RTProxy.Get(), bRequirePreviousData](FRHICommandListImmediate& CmdList)
+	{
+		RT_Proxy->BeginFrame(bRequirePreviousData);
+	});
+}
+
+void FNiagaraDataChannelData::EndFrame(UNiagaraDataChannelHandler* Owner)
+{
+	//We must do one final tick to process any final items generated by the last things to tick this frame.
+	ConsumePublishRequests(Owner);
+
+	check(RTProxy);
+	FNiagaraGpuComputeDispatchInterface* DispathInterface = FNiagaraGpuComputeDispatchInterface::Get(Owner->GetWorld());
+	ENQUEUE_RENDER_COMMAND(FDataChannelProxyEndFrame) (
+		[DispathInterface, RT_Proxy = RTProxy.Get(), RT_BuffersForGPU = MoveTemp(BuffersForGPU)](FRHICommandListImmediate& CmdList)
+	{
+		RT_Proxy->EndFrame(DispathInterface, CmdList, RT_BuffersForGPU);
+	});
+	BuffersForGPU.Reset();
+}
+
+void FNiagaraDataChannelData::ConsumePublishRequests(UNiagaraDataChannelHandler* Owner)
+{
+	check(IsValid(Owner));
+
+	if(PublishRequests.Num() == 0)
+	{
+		return;
+	}
+
+	UWorld* World = Owner->GetWorld();
+	check(IsValid(World));
+
+	int32 GameDataOrigSize = GameData->Num();
+	int32 CPUDataOrigSize = CPUSimData->GetCurrentData()->GetNumInstances();
+	int32 NewGameDataChannel = GameDataOrigSize;
+	int32 NewCPUDataChannel = CPUDataOrigSize;
+
+	//Do a pass to gather the new total size for our DataChannel data.
+
+	//Each DI that generates DataChannel can control whether it's pushed to Game/CPU/GPU.
+	BuffersForGPU.Reserve(BuffersForGPU.Num() + PublishRequests.Num());
+
+	for (FNiagaraDataChannelPublishRequest& PublishRequest : PublishRequests)
+	{
+		uint32 NumInsts = 0;
+
+		if (FNiagaraDataChannelGameData* RequestGameData = PublishRequest.GameData.Get())
+		{
+			NumInsts = RequestGameData->Num();
+			if (RequestGameData->Num() > 0 && PublishRequest.bVisibleToGPUSims)
 			{
-				delete CPUDataChannelDataSet;
+				//We stage the game data into a data set to facilitate easier copy over to the GPU. TODO: Send and copy the game data directly.			
+				FNiagaraDataBuffer* StagingBuf = &GameDataStaging->BeginSimulate();
+				StagingBuf->Allocate(NumInsts);				
+				RequestGameData->WriteToDataSet(StagingBuf, 0, LwcTile);
+				GameDataStaging->EndSimulate();				
+				BuffersForGPU.Emplace(StagingBuf);	
 			}
-			if (GPUDataChannelDataSet != nullptr)
+		}
+		else if (ensure(PublishRequest.Data))
+		{
+			NumInsts = PublishRequest.Data->GetNumInstances();
+
+			if (NumInsts > 0 && PublishRequest.bVisibleToGPUSims)
 			{
-				delete GPUDataChannelDataSet;
+				BuffersForGPU.Add(PublishRequest.Data);
 			}
 		}
-	);
-	CPUSimData = nullptr;
-	GPUSimData = nullptr;
+
+		if (PublishRequest.bVisibleToGame)
+		{
+			NewGameDataChannel += NumInsts;
+		}
+		if (PublishRequest.bVisibleToCPUSims)
+		{
+			NewCPUDataChannel += NumInsts;
+		}
+	}
+
+	//Allocate Sim buffers ready for gather.
+	CPUSimData->BeginSimulate();
+	CPUSimData->Allocate(NewCPUDataChannel, true);
+
+	//Now do the actual data collection.
+
+	GameData->SetNum(NewGameDataChannel);
+
+	for (auto It = PublishRequests.CreateIterator(); It ; ++It)
+	{
+		FNiagaraDataChannelPublishRequest& PublishRequest = *It;
+		if (FNiagaraDataChannelGameData* RequestGameData = PublishRequest.GameData.Get())
+		{
+			if (RequestGameData->Num())
+			{
+				if(PublishRequest.bVisibleToGame)
+				{
+					GameData->AppendFromGameData(*RequestGameData);
+				}
+				if(PublishRequest.bVisibleToCPUSims)
+				{
+					RequestGameData->WriteToDataSet(CPUSimData->GetDestinationData(), CPUSimData->GetDestinationDataChecked().GetNumInstances(), LwcTile);
+				}
+			}
+		}
+
+		if (PublishRequest.Data)
+		{
+			const ENiagaraSimTarget SimTarget = PublishRequest.Data->GetOwner()->GetSimTarget();
+			if (SimTarget == ENiagaraSimTarget::CPUSim)
+			{
+				if (PublishRequest.bVisibleToGame)
+				{
+					GameData->AppendFromDataSet(PublishRequest.Data, PublishRequest.LwcTile);
+				}
+				if (PublishRequest.bVisibleToCPUSims)
+				{
+					PublishRequest.Data->CopyToUnrelated(CPUSimData->GetDestinationDataChecked(), 0, CPUSimData->GetDestinationDataChecked().GetNumInstances(), PublishRequest.Data->GetNumInstances());
+				}
+			}
+			else if (SimTarget == ENiagaraSimTarget::GPUComputeSim)
+			{
+				//TODO: GPU->GPU handling will be done all RT side. GPU->CPU may be done here in future but we'll have to pull the data from the GPU on the RT and pass beck to GT here.
+			}
+			else
+			{
+				check(0);
+			}
+		}
+		
+		It.RemoveCurrentSwap();
+	}
+
+	CPUSimData->EndSimulate();
+
+#if !UE_BUILD_SHIPPING
+	if(Owner->GetDataChannel()->GetVerboseLogging())
+	{
+		FString Label = FString::Printf(TEXT("Data Channel %s"), *Owner->GetDataChannel()->GetName());
+		CPUSimData->GetCurrentData()->Dump(0, CPUSimData->GetCurrentData()->GetNumInstances(), Label);
+	}
+#endif
+
+	PublishRequests.Reset();
 }
 
-void FNiagaraWorldDataChannelStore::Init(UNiagaraDataChannelHandler* Owner)
+FNiagaraDataBufferRef FNiagaraDataChannelData::GetCPUData(bool bPreviousFrame)
 {
-	GameData = Owner->GetDataChannel()->CreateGameData();
+	if(bPreviousFrame)
+	{
+		return PrevCPUSimData;
+	}
+	else
+	{
+		return CPUSimData->GetCurrentData();
+	}
+}
 
-	CPUSimData = new FNiagaraDataSet();
-	GPUSimData = new FNiagaraDataSet();
+FNiagaraDataChannelGameData* FNiagaraDataChannelData::GetGameData()
+{
+	return GameData.Get();
+}
 
-	const FNiagaraDataSetCompiledData& CompiledData = Owner->GetDataChannel()->GetCompiledData(ENiagaraSimTarget::CPUSim);
-	const FNiagaraDataSetCompiledData& CompiledDataGPU = Owner->GetDataChannel()->GetCompiledData(ENiagaraSimTarget::GPUComputeSim);
-	CPUSimData->Init(&CompiledData);
-	GPUSimData->Init(&CompiledDataGPU);
+void FNiagaraDataChannelData::Publish(const FNiagaraDataChannelPublishRequest& Request)
+{
+	PublishRequests.Add(Request);
 }
 
-void FNiagaraWorldDataChannelStore::Tick(UNiagaraDataChannelHandler* Owner)
+void FNiagaraDataChannelData::RemovePublishRequests(const FNiagaraDataSet* DataSet)
 {
+	//TODO: Have to ensure lifetime of anything we're holding
+	for (auto It = PublishRequests.CreateIterator(); It; ++It)
+	{
+		const FNiagaraDataChannelPublishRequest& PublishRequest = *It;
 
+		//First remove any buffers we've queued up to push to the GPU
+		for (auto BufferIt = BuffersForGPU.CreateIterator(); BufferIt; ++BufferIt)
+		{
+			FNiagaraDataBuffer* Buffer = *BufferIt;
+			if (Buffer == nullptr || Buffer->GetOwner() == DataSet)
+			{
+				BufferIt.RemoveCurrentSwap();
+			}
+		}
+
+		if (PublishRequest.Data && PublishRequest.Data->GetOwner() == DataSet)
+		{
+			It.RemoveCurrentSwap();
+		}
+	}
 }
 
+const FNiagaraDataSetCompiledData& FNiagaraDataChannelData::GetCompiledData(ENiagaraSimTarget SimTarget)
+{
+	check(CPUSimData && GPUSimData)
+	return SimTarget == ENiagaraSimTarget::CPUSim ? CPUSimData->GetCompiledData() : GPUSimData->GetCompiledData();
+}
 
 //////////////////////////////////////////////////////////////////////////
 
@@ -447,6 +775,18 @@ namespace NiagaraDataChannel
 	}
 }
 
+void UNiagaraDataChannel::PostInitProperties()
+{
+	Super::PostInitProperties();
+
+	FNiagaraWorldManager::ForAllWorldManagers(
+		[DataChannel = this](FNiagaraWorldManager& WorldMan)
+		{
+			WorldMan.WaitForAsyncWork();
+			WorldMan.GetDataChannelManager().InitDataChannel(DataChannel, true);
+		});
+}
+
 void UNiagaraDataChannel::PostLoad()
 {
 	Super::PostLoad();
@@ -484,6 +824,7 @@ void UNiagaraDataChannel::PostLoad()
 	FNiagaraWorldManager::ForAllWorldManagers(
 		[DataChannel = this](FNiagaraWorldManager& WorldMan)
 		{
+			WorldMan.WaitForAsyncWork();
 			WorldMan.GetDataChannelManager().InitDataChannel(DataChannel, true);
 		});
 }
@@ -496,7 +837,7 @@ void UNiagaraDataChannel::BeginDestroy()
 		[DataChannel = this](FNiagaraWorldManager& WorldMan)
 		{
 			WorldMan.WaitForAsyncWork();
-			WorldMan.GetDataChannelManager().RemoveDataChannel(DataChannel->GetChannelName());
+			WorldMan.GetDataChannelManager().RemoveDataChannel(DataChannel);
 		});
 }
 
@@ -504,14 +845,12 @@ void UNiagaraDataChannel::BeginDestroy()
 
 void UNiagaraDataChannel::PreEditChange(FProperty* PropertyAboutToChange)
 {
-	//We'll do this for all channels inside the definitions pre edit change.
-	// 
-// 	FNiagaraWorldManager::ForAllWorldManagers(
-// 		[DataChannel = this](FNiagaraWorldManager& WorldMan)
-// 		{
-// 			WorldMan.WaitForAsyncWork();
-// 			WorldMan.GetDataChannelManager().RemoveDataChannel(DataChannel->GetChannelName());
-// 		});
+	FNiagaraWorldManager::ForAllWorldManagers(
+		[DataChannel = this](FNiagaraWorldManager& WorldMan)
+		{
+			WorldMan.WaitForAsyncWork();
+			WorldMan.GetDataChannelManager().RemoveDataChannel(DataChannel);
+		});
 }
 
 void UNiagaraDataChannel::PostEditChangeProperty(FPropertyChangedEvent& PropertyChangedEvent)
@@ -528,9 +867,11 @@ void UNiagaraDataChannel::PostEditChangeProperty(FPropertyChangedEvent& Property
 
 	GameDataLayout.Init(Variables);
 
-	OnChangedDelegate.Broadcast(this);
-
-	//We'll reinit this channel with the world managers etc in the PostEditChange of the data channel definitions.
+	FNiagaraWorldManager::ForAllWorldManagers(
+		[DataChannel = this](FNiagaraWorldManager& WorldMan)
+		{
+			WorldMan.GetDataChannelManager().InitDataChannel(DataChannel, true);
+		});
 }
 
 #endif//WITH_EDITOR
@@ -546,7 +887,10 @@ const FNiagaraDataSetCompiledData& UNiagaraDataChannel::GetCompiledData(ENiagara
 			CompiledData.SimTarget = ENiagaraSimTarget::CPUSim;
 			for (FNiagaraVariableBase Var : Variables)
 			{
-				Var.SetType(FNiagaraTypeDefinition(FNiagaraTypeHelper::FindNiagaraFriendlyTopLevelStruct(Var.GetType().GetScriptStruct(), ENiagaraStructConversion::Simulation)));
+				if(Var.GetType().IsEnum() == false)
+				{
+					Var.SetType(FNiagaraTypeDefinition(FNiagaraTypeHelper::FindNiagaraFriendlyTopLevelStruct(Var.GetType().GetScriptStruct(), ENiagaraStructConversion::Simulation)));
+				}
 				CompiledData.Variables.Add(Var);
 			}
 			CompiledData.BuildLayout();
@@ -563,7 +907,10 @@ const FNiagaraDataSetCompiledData& UNiagaraDataChannel::GetCompiledData(ENiagara
 			CompiledDataGPU.SimTarget = ENiagaraSimTarget::GPUComputeSim;
 			for (FNiagaraVariableBase Var : Variables)
 			{
-				Var.SetType(FNiagaraTypeDefinition(FNiagaraTypeHelper::FindNiagaraFriendlyTopLevelStruct(Var.GetType().GetScriptStruct(), ENiagaraStructConversion::Simulation)));
+				if (Var.GetType().IsEnum() == false)
+				{
+					Var.SetType(FNiagaraTypeDefinition(FNiagaraTypeHelper::FindNiagaraFriendlyTopLevelStruct(Var.GetType().GetScriptStruct(), ENiagaraStructConversion::Simulation)));
+				}
 				CompiledDataGPU.Variables.Add(Var);
 			}
 			CompiledDataGPU.BuildLayout();
@@ -574,14 +921,14 @@ const FNiagaraDataSetCompiledData& UNiagaraDataChannel::GetCompiledData(ENiagara
 
 FNiagaraDataChannelGameDataPtr UNiagaraDataChannel::CreateGameData()const
 {
-	FNiagaraDataChannelGameDataPtr NewGameData = MakeShared<FNiagaraDataChannelGameData>();
-	NewGameData->Init(this);
-	return NewGameData;
+	FNiagaraDataChannelGameDataPtr NewData = MakeShared<FNiagaraDataChannelGameData>();
+	NewData->Init(this);
+	return NewData;
 }
 
 bool UNiagaraDataChannel::IsValid()const
 {
-	return ChannelName != NAME_None && Variables.Num() > 0 && CompiledData.Variables.Num() == Variables.Num() && CompiledDataGPU.Variables.Num() == Variables.Num();
+	return Variables.Num() > 0 && CompiledData.Variables.Num() == Variables.Num() && CompiledDataGPU.Variables.Num() == Variables.Num();
 }
 
 //////////////////////////////////////////////////////////////////////////
@@ -591,7 +938,7 @@ UNiagaraDataChannelLibrary::UNiagaraDataChannelLibrary(const FObjectInitializer&
 {
 }
 
-UNiagaraDataChannelHandler* UNiagaraDataChannelLibrary::GetNiagaraDataChannel(const UObject* WorldContextObject, FName Channel)
+UNiagaraDataChannelHandler* UNiagaraDataChannelLibrary::GetNiagaraDataChannel(const UObject* WorldContextObject, UNiagaraDataChannelAsset* Channel)
 {
 	UWorld* World = (WorldContextObject != nullptr) ? GEngine->GetWorldFromContextObject(WorldContextObject, EGetWorldErrorMode::LogAndReturnNull) : nullptr;
 	if (World)
@@ -604,5 +951,83 @@ UNiagaraDataChannelHandler* UNiagaraDataChannelLibrary::GetNiagaraDataChannel(co
 	return nullptr;
 }
 
+UNiagaraDataChannelWriter* UNiagaraDataChannelLibrary::WriteToNiagaraDataChannel(const UObject* WorldContextObject, UNiagaraDataChannelAsset* Channel, FNiagaraDataChannelSearchParameters SearchParams, int32 Count, bool bVisibleToGame, bool bVisibleToCPU, bool bVisibleToGPU)
+{
+	check(IsInGameThread());
+	UWorld* World = (WorldContextObject != nullptr) ? GEngine->GetWorldFromContextObject(WorldContextObject, EGetWorldErrorMode::LogAndReturnNull) : nullptr;
+	if (World && Count > 0)
+	{
+		if (FNiagaraWorldManager* WorldMan = FNiagaraWorldManager::Get(World))
+		{
+			if(UNiagaraDataChannelHandler* Handler = WorldMan->GetDataChannelManager().FindDataChannelHandler(Channel))
+			{
+				if(UNiagaraDataChannelWriter* Writer = Handler->GetDataChannelWriter())
+				{
+					if(Writer->InitWrite(SearchParams, Count, bVisibleToGame, bVisibleToCPU, bVisibleToGPU))
+					{
+						return Writer;
+					}
+				}
+			}
+		}
+	}
+	return nullptr;
+}
+
+UNiagaraDataChannelReader* UNiagaraDataChannelLibrary::ReadFromNiagaraDataChannel(const UObject* WorldContextObject, UNiagaraDataChannelAsset* Channel, FNiagaraDataChannelSearchParameters SearchParams, bool bReadPreviousFrame)
+{
+	UWorld* World = (WorldContextObject != nullptr) ? GEngine->GetWorldFromContextObject(WorldContextObject, EGetWorldErrorMode::LogAndReturnNull) : nullptr;
+	if (World)
+	{
+		if (FNiagaraWorldManager* WorldMan = FNiagaraWorldManager::Get(World))
+		{
+			if (UNiagaraDataChannelHandler* Handler = WorldMan->GetDataChannelManager().FindDataChannelHandler(Channel))
+			{
+				if (UNiagaraDataChannelReader* Reader = Handler->GetDataChannelReader())
+				{
+					if(Reader->InitAccess(SearchParams, bReadPreviousFrame))
+					{
+						return Reader;
+					}
+				}
+			}
+		}
+	}
+	return nullptr;
+}
+
+FPrimaryAssetId UNiagaraDataChannelAsset::GetPrimaryAssetId()const 
+{
+	return FPrimaryAssetId(GetClass()->GetFName(), GetFName());
+}
+
+//////////////////////////////////////////////////////////////////////////
+
+#if !UE_BUILD_SHIPPING
+void FNiagaraDataChannelDebugUtilities::BeginFrame(FNiagaraWorldManager* WorldMan, float DeltaSeconds)
+{
+	check(WorldMan);
+	WorldMan->GetDataChannelManager().BeginFrame(DeltaSeconds);
+}
+
+void FNiagaraDataChannelDebugUtilities::EndFrame(FNiagaraWorldManager* WorldMan, float DeltaSeconds)
+{
+	check(WorldMan);
+	WorldMan->GetDataChannelManager().EndFrame(DeltaSeconds);
+}
+
+void FNiagaraDataChannelDebugUtilities::Tick(FNiagaraWorldManager* WorldMan, float DeltaSeconds, ETickingGroup TickGroup)
+{
+	check(WorldMan);
+	WorldMan->GetDataChannelManager().Tick(DeltaSeconds, TickGroup);
+}
+
+UNiagaraDataChannelHandler* FNiagaraDataChannelDebugUtilities::FindDataChannelHandler(FNiagaraWorldManager* WorldMan, UNiagaraDataChannel* DataChannel)
+{
+	check(WorldMan);
+	return WorldMan->GetDataChannelManager().FindDataChannelHandler(DataChannel);
+}
+#endif //!UE_BUILD_SHIPPING
 
 #undef LOCTEXT_NAMESPACE
+
diff --git a/Engine/Plugins/FX/Niagara/Source/Niagara/Private/NiagaraDataChannelAccessor.cpp b/Engine/Plugins/FX/Niagara/Source/Niagara/Private/NiagaraDataChannelAccessor.cpp
--- a/Engine/Plugins/FX/Niagara/Source/Niagara/Private/NiagaraDataChannelAccessor.cpp
+++ b/Engine/Plugins/FX/Niagara/Source/Niagara/Private/NiagaraDataChannelAccessor.cpp
@@ -4,6 +4,8 @@
 
 #include "UObject/UObjectGlobals.h"
 #include "UObject/Package.h"
+
+#include "NiagaraDataChannelCommon.h"
 #include "NiagaraDataChannel.h"
 #include "NiagaraDataChannelHandler.h"
 
@@ -65,83 +67,99 @@ namespace NiagaraDataChannelAccesorLocal
 
 //////////////////////////////////////////////////////////////////////////
 
-bool UNiagaraDataChannelReader::InitAccess(UActorComponent* OwningComponent)
+bool UNiagaraDataChannelReader::InitAccess(FNiagaraDataChannelSearchParameters SearchParams, bool bReadPreviousFrameData)
 {
 	Data = nullptr;
+	bReadingPreviousFrame = bReadPreviousFrameData;
 	check(Owner);
 
-	if (OwningComponent == nullptr)
-	{
-		UE_LOG(LogNiagara, Warning, TEXT("Call to UNiagaraDataChannelReader::InitAccess did not provide a valid Owning Componet."));
-		return false;
-	}
-
-	Data = Owner->GetGameData(OwningComponent);
-	return true;
+	Data = Owner->FindData(SearchParams, ENiagaraResourceAccess::ReadOnly);
+	return Data.IsValid();
 }
 
 int32 UNiagaraDataChannelReader::Num()const
 {
 	if (Data.IsValid())
 	{
-		return Data->Num();
+		return Data->GetGameData()->Num();
 	}
 	return 0;
 }
 
 template<typename T>
-T UNiagaraDataChannelReader::ReadData(const FNiagaraVariableBase& Var, int32 Index)const
+bool UNiagaraDataChannelReader::ReadData(const FNiagaraVariableBase& Var, int32 Index, T& OutData)const
 {
 	if (ensure(Data.IsValid()))
 	{
-		if (FNiagaraDataChannelVariableBuffer* VarBuffer = Data->FindVariableBuffer(Var))
+		if (FNiagaraDataChannelVariableBuffer* VarBuffer = Data->GetGameData()->FindVariableBuffer(Var))
 		{
-			T OutData;
-			VarBuffer->Read<T>(Index, OutData);
-			return OutData;
+			return VarBuffer->Read<T>(Index, OutData, bReadingPreviousFrame);
 		}
 	}
-	return T();
+	return false;
 }
 
 double UNiagaraDataChannelReader::ReadFloat(FName VarName, int32 Index)const
 {
-	return ReadData<double>(FNiagaraVariableBase(NiagaraDataChannelAccesorLocal::GetDoubleDef(), VarName), Index);
+	double RetVal = 0.0f;
+	ReadData<double>(FNiagaraVariableBase(NiagaraDataChannelAccesorLocal::GetDoubleDef(), VarName), Index, RetVal);
+	return RetVal;
 }
 
 FVector2D UNiagaraDataChannelReader::ReadVector2D(FName VarName, int32 Index)const
 {
-	return ReadData<FVector2D>(FNiagaraVariableBase(NiagaraDataChannelAccesorLocal::GetVector2DDef(), VarName), Index);
+	FVector2D RetVal = FVector2D::ZeroVector;
+	ReadData<FVector2D>(FNiagaraVariableBase(NiagaraDataChannelAccesorLocal::GetVector2DDef(), VarName), Index, RetVal);
+	return RetVal;
 }
 
 FVector UNiagaraDataChannelReader::ReadVector(FName VarName, int32 Index)const
 {
-	return ReadData<FVector>(FNiagaraVariableBase(NiagaraDataChannelAccesorLocal::GetVectorDef(), VarName), Index);
+	FVector RetVal = FVector::ZeroVector;
+	ReadData<FVector>(FNiagaraVariableBase(NiagaraDataChannelAccesorLocal::GetVectorDef(), VarName), Index, RetVal);
+	return RetVal;
 }
 
 FVector4 UNiagaraDataChannelReader::ReadVector4(FName VarName, int32 Index)const
 {
-	return ReadData<FVector4>(FNiagaraVariableBase(NiagaraDataChannelAccesorLocal::GetVector4Def(), VarName), Index);
+	FVector4 RetVal = FVector4(0.0f);
+	ReadData<FVector4>(FNiagaraVariableBase(NiagaraDataChannelAccesorLocal::GetVector4Def(), VarName), Index, RetVal);
+	return RetVal;
 }
 
 FQuat UNiagaraDataChannelReader::ReadQuat(FName VarName, int32 Index)const
 {
-	return ReadData<FQuat>(FNiagaraVariableBase(NiagaraDataChannelAccesorLocal::GetQuatDef(), VarName), Index);
+	FQuat RetVal = FQuat::Identity;
+	ReadData<FQuat>(FNiagaraVariableBase(NiagaraDataChannelAccesorLocal::GetQuatDef(), VarName), Index, RetVal);
+	return RetVal;
 }
 
 FLinearColor UNiagaraDataChannelReader::ReadLinearColor(FName VarName, int32 Index)const
 {
-	return ReadData<FLinearColor>(FNiagaraVariableBase(FNiagaraTypeDefinition::GetColorDef(), VarName), Index);
+	FLinearColor RetVal = FLinearColor::White;
+	ReadData<FLinearColor>(FNiagaraVariableBase(FNiagaraTypeDefinition::GetColorDef(), VarName), Index, RetVal);
+	return RetVal;
 }
 
 int32 UNiagaraDataChannelReader::ReadInt(FName VarName, int32 Index)const
 {
-	return ReadData<int32>(FNiagaraVariableBase(FNiagaraTypeDefinition::GetIntDef(), VarName), Index);
+	int32 RetVal = 0;
+	ReadData<int32>(FNiagaraVariableBase(FNiagaraTypeDefinition::GetIntDef(), VarName), Index, RetVal);
+	return RetVal;
 }
 
 bool UNiagaraDataChannelReader::ReadBool(FName VarName, int32 Index)const
 {
-	return ReadData<bool>(FNiagaraVariableBase(FNiagaraTypeDefinition::GetBoolDef(), VarName), Index);
+	bool RetVal = false;
+	ReadData<bool>(FNiagaraVariableBase(FNiagaraTypeDefinition::GetBoolDef(), VarName), Index, RetVal);
+	return RetVal;
+}
+
+FVector UNiagaraDataChannelReader::ReadPosition(FName VarName, int32 Index)const
+{
+	FVector RetVal = FVector::ZeroVector;
+	ReadData<FVector>(FNiagaraVariableBase(FNiagaraTypeDefinition::GetPositionDef(), VarName), Index, RetVal);
+	return RetVal;
 }
 
 //////////////////////////////////////////////////////////////////////////
@@ -159,14 +177,8 @@ void UNiagaraDataChannelWriter::WriteData(const FNiagaraVariableBase& Var, int32
 	}
 }
 
-bool UNiagaraDataChannelWriter::InitWrite(UActorComponent* OwningComponent, int32 Count, bool bVisibleToGame, bool bVisibleToCPU, bool bVisibleToGPU)
+bool UNiagaraDataChannelWriter::InitWrite(FNiagaraDataChannelSearchParameters SearchParams, int32 Count, bool bVisibleToGame, bool bVisibleToCPU, bool bVisibleToGPU)
 {
-	if( OwningComponent == nullptr )
-	{
-		UE_LOG(LogNiagara, Warning, TEXT("Call to UNiagaraDataChannelWriter::InitWrite did not provide a valid Owning Componet."));
-		return false;
-	}
-
 	if (Count == 0)
 	{
 		UE_LOG(LogNiagara, Warning, TEXT("Call to UNiagaraDataChannelWriter::InitWrite with Count == 0. Ignored."));
@@ -175,18 +187,23 @@ bool UNiagaraDataChannelWriter::InitWrite(UActorComponent* OwningComponent, int3
 
 	check(Owner);
 
-	Data = Owner->GetDataChannel()->CreateGameData();
-
-	FNiagaraDataChannelPublishRequest PublishRequest;
-	PublishRequest.bVisibleToGame = bVisibleToGame;
-	PublishRequest.bVisibleToCPUSims = bVisibleToCPU;
-	PublishRequest.bVisibleToGPUSims = bVisibleToGPU;
-	PublishRequest.GameData = Data;
-	Owner->Publish(PublishRequest);
-
-	Data->SetNum(Count);
+	FNiagaraDataChannelDataPtr DestData = Owner->FindData(SearchParams, ENiagaraResourceAccess::WriteOnly);
+	if(DestData)
+	{
+		//TODO- Dont create a whole new game Data here. Rather grab a pre-made staging/input One from the Data PTR With the relevant publish Flags.
+		Data = Owner->GetDataChannel()->CreateGameData();
+		Data->SetNum(Count);
+
+		FNiagaraDataChannelPublishRequest PublishRequest;
+		PublishRequest.bVisibleToGame = bVisibleToGame;
+		PublishRequest.bVisibleToCPUSims = bVisibleToCPU;
+		PublishRequest.bVisibleToGPUSims = bVisibleToGPU;
+		PublishRequest.GameData = Data;
+		DestData->Publish(PublishRequest);
+		return true;
+	}
 
-	return true;
+	return false;
 }
 
 int32 UNiagaraDataChannelWriter::Num()const
@@ -235,7 +252,7 @@ void UNiagaraDataChannelWriter::WriteInt(FName VarName, int32 Index, int32 InDat
 
 void UNiagaraDataChannelWriter::WriteBool(FName VarName, int32 Index, bool InData)
 {
-	WriteData(FNiagaraVariableBase(FNiagaraTypeDefinition::GetBoolDef(), VarName), Index, InData);
+	WriteData(FNiagaraVariableBase(FNiagaraTypeDefinition::GetBoolDef(), VarName), Index, FNiagaraBool(InData));
 }
 
 void UNiagaraDataChannelWriter::WriteSpawnInfo(FName VarName, int32 Index, FNiagaraSpawnInfo InData)
@@ -243,3 +260,7 @@ void UNiagaraDataChannelWriter::WriteSpawnInfo(FName VarName, int32 Index, FNiag
 	WriteData(FNiagaraVariableBase(FNiagaraTypeDefinition(FNiagaraSpawnInfo::StaticStruct()), VarName), Index, InData);
 }
 
+void UNiagaraDataChannelWriter::WritePosition(FName VarName, int32 Index, FVector InData)
+{
+	WriteData(FNiagaraVariableBase(FNiagaraTypeDefinition::GetPositionDef(), VarName), Index, InData);
+}
diff --git a/Engine/Plugins/FX/Niagara/Source/Niagara/Private/NiagaraDataChannelCommon.cpp b/Engine/Plugins/FX/Niagara/Source/Niagara/Private/NiagaraDataChannelCommon.cpp
--- a/Engine/Plugins/FX/Niagara/Source/Niagara/Private/NiagaraDataChannelCommon.cpp
+++ b/Engine/Plugins/FX/Niagara/Source/Niagara/Private/NiagaraDataChannelCommon.cpp
@@ -1,3 +1,9 @@
 // Copyright Epic Games, Inc. All Rights Reserved.
 
-#include "NiagaraDataChannelCommon.h"
\ No newline at end of file
+#include "NiagaraDataChannelCommon.h"
+#include "Components/SceneComponent.h"
+
+FVector FNiagaraDataChannelSearchParameters::GetLocation()const
+{
+	return OwningComponent ? OwningComponent->GetComponentLocation() : Location;
+}
\ No newline at end of file
diff --git a/Engine/Plugins/FX/Niagara/Source/Niagara/Private/NiagaraDataChannelHandler.cpp b/Engine/Plugins/FX/Niagara/Source/Niagara/Private/NiagaraDataChannelHandler.cpp
--- a/Engine/Plugins/FX/Niagara/Source/Niagara/Private/NiagaraDataChannelHandler.cpp
+++ b/Engine/Plugins/FX/Niagara/Source/Niagara/Private/NiagaraDataChannelHandler.cpp
@@ -2,19 +2,10 @@
 
 #include "NiagaraDataChannelHandler.h"
 #include "NiagaraDataChannelAccessor.h"
+#include "NiagaraDataChannelCommon.h"
 
 UNiagaraDataChannelHandler::~UNiagaraDataChannelHandler()
 {
-	if (FNiagaraDataChannelHandlerRTProxyBase* ReleasedProxy = RTProxy.Release())
-	{
-		ENQUEUE_RENDER_COMMAND(FDeleteRTProxyCmd) (
-			[RT_Proxy = ReleasedProxy](FRHICommandListImmediate& CmdList) mutable
-			{
-				delete RT_Proxy;
-			}
-		);
-		check(RTProxy.IsValid() == false);
-	}
 }
 
 void UNiagaraDataChannelHandler::Init(const UNiagaraDataChannel* InChannel)
@@ -22,38 +13,19 @@ void UNiagaraDataChannelHandler::Init(const UNiagaraDataChannel* InChannel)
 	DataChannel = InChannel;
 }
 
-void UNiagaraDataChannelHandler::Tick(float DeltaTime, ETickingGroup TickGroup, FNiagaraWorldManager* OwningWorld)
+void UNiagaraDataChannelHandler::BeginFrame(float DeltaTime, FNiagaraWorldManager* OwningWorld)
 {
-}
 
-void UNiagaraDataChannelHandler::Publish(const FNiagaraDataChannelPublishRequest& Request)
-{
-	PublishRequests.Add(Request);
 }
 
-void UNiagaraDataChannelHandler::RemovePublishRequests(const FNiagaraDataSet* DataSet)
+void UNiagaraDataChannelHandler::EndFrame(float DeltaTime, FNiagaraWorldManager* OwningWorld)
 {
-	for (auto It = PublishRequests.CreateIterator(); It; ++It)
-	{
-		const FNiagaraDataChannelPublishRequest& PublishRequest = *It;
-
-		//First remove any buffers we've queued up to push to the GPU
-		for (auto BufferIt = BuffersForGPU.CreateIterator(); BufferIt; ++BufferIt)
-		{
-			FNiagaraDataBuffer* Buffer = *BufferIt;
-			if (Buffer == nullptr || Buffer->GetOwner() == DataSet)
-			{
-				It.RemoveCurrentSwap();
-			}
-		}
 
-		if (PublishRequest.Data && PublishRequest.Data->GetOwner() == DataSet)
-		{
-			It.RemoveCurrentSwap();
-		}
-	}
 }
 
+void UNiagaraDataChannelHandler::Tick(float DeltaTime, ETickingGroup TickGroup, FNiagaraWorldManager* OwningWorld)
+{
+}
 
 UNiagaraDataChannelWriter* UNiagaraDataChannelHandler::GetDataChannelWriter()
 {
@@ -73,4 +45,9 @@ UNiagaraDataChannelReader* UNiagaraDataChannelHandler::GetDataChannelReader()
 		Reader->Owner = this;
 	}
 	return Reader;
+}
+
+FNiagaraDataChannelDataPtr UNiagaraDataChannelHandler::CreateData()
+{
+	return MakeShared<FNiagaraDataChannelData>(this);
 }
\ No newline at end of file
diff --git a/Engine/Plugins/FX/Niagara/Source/Niagara/Private/NiagaraDataChannelManager.cpp b/Engine/Plugins/FX/Niagara/Source/Niagara/Private/NiagaraDataChannelManager.cpp
--- a/Engine/Plugins/FX/Niagara/Source/Niagara/Private/NiagaraDataChannelManager.cpp
+++ b/Engine/Plugins/FX/Niagara/Source/Niagara/Private/NiagaraDataChannelManager.cpp
@@ -6,8 +6,10 @@
 #include "NiagaraWorldManager.h"
 #include "NiagaraDataChannel.h"
 #include "NiagaraDataChannelHandler.h"
-#include "NiagaraDataChannelDefinitions.h"
 
+DECLARE_CYCLE_STAT(TEXT("FNiagaraDataChannelManager::BeginFrame"), STAT_DataChannelManager_BeginFrame, STATGROUP_NiagaraDataChannels);
+DECLARE_CYCLE_STAT(TEXT("FNiagaraDataChannelManager::EndFrame"), STAT_DataChannelManager_EndFrame, STATGROUP_NiagaraDataChannels);
+DECLARE_CYCLE_STAT(TEXT("FNiagaraDataChannelManager::Tick"), STAT_DataChannelManager_Tick, STATGROUP_NiagaraDataChannels);
 
 FNiagaraDataChannelManager::FNiagaraDataChannelManager(FNiagaraWorldManager* InWorldMan)
 	: WorldMan(InWorldMan)
@@ -21,13 +23,14 @@ void FNiagaraDataChannelManager::AddReferencedObjects(FReferenceCollector& Colle
 
 void FNiagaraDataChannelManager::Init()
 {
-	for (const UNiagaraDataChannelDefinitions* DataChannelDefs : UNiagaraDataChannelDefinitions::GetDataChannelDefinitions(false, false))
+	//Initialize any existing data channels, more may be initialized later as they are loaded.
+	UNiagaraDataChannel::ForEachDataChannel([&](UNiagaraDataChannel* DataChannel)
 	{
-		for (const UNiagaraDataChannel* DataChannel : DataChannelDefs->DataChannels)
+		if (DataChannel->HasAnyFlags(RF_ClassDefaultObject) == false)
 		{
 			InitDataChannel(DataChannel, true);
 		}
-	}
+	});
 }
 
 void FNiagaraDataChannelManager::Cleanup()
@@ -35,10 +38,37 @@ void FNiagaraDataChannelManager::Cleanup()
 	Channels.Empty();
 }
 
+void FNiagaraDataChannelManager::BeginFrame(float DeltaSeconds)
+{
+	if (INiagaraModule::DataChannelsEnabled())
+	{
+		SCOPE_CYCLE_COUNTER(STAT_DataChannelManager_BeginFrame);
+		//Tick all DataChannel channel handlers.
+		for (auto& ChannelPair : Channels)
+		{
+			ChannelPair.Value->BeginFrame(DeltaSeconds, WorldMan);
+		}
+	}
+}
+
+void FNiagaraDataChannelManager::EndFrame(float DeltaSeconds)
+{
+	if (INiagaraModule::DataChannelsEnabled())
+	{
+		SCOPE_CYCLE_COUNTER(STAT_DataChannelManager_EndFrame);
+		//Tick all DataChannel channel handlers.
+		for (auto& ChannelPair : Channels)
+		{
+			ChannelPair.Value->EndFrame(DeltaSeconds, WorldMan);
+		}
+	}
+}
+
 void FNiagaraDataChannelManager::Tick(float DeltaSeconds, ETickingGroup TickGroup)
 {
 	if(INiagaraModule::DataChannelsEnabled())
 	{
+		SCOPE_CYCLE_COUNTER(STAT_DataChannelManager_Tick);
 		//Tick all DataChannel channel handlers.
 		for (auto& ChannelPair : Channels)
 		{
@@ -51,9 +81,9 @@ void FNiagaraDataChannelManager::Tick(float DeltaSeconds, ETickingGroup TickGrou
 	}
 }
 
-UNiagaraDataChannelHandler* FNiagaraDataChannelManager::FindDataChannelHandler(FName ChannelName)
+UNiagaraDataChannelHandler* FNiagaraDataChannelManager::FindDataChannelHandler(const UNiagaraDataChannel* Channel)
 {
-	if (TObjectPtr<UNiagaraDataChannelHandler>* Found = Channels.Find(ChannelName))
+	if (TObjectPtr<UNiagaraDataChannelHandler>* Found = Channels.Find(Channel))
 	{
 		return (*Found).Get();
 	}
@@ -62,21 +92,21 @@ UNiagaraDataChannelHandler* FNiagaraDataChannelManager::FindDataChannelHandler(F
 
 void FNiagaraDataChannelManager::InitDataChannel(const UNiagaraDataChannel* InChannel, bool bForce)
 {
-	if(INiagaraModule::DataChannelsEnabled() && InChannel->IsValid())
+	UWorld* World = GetWorld();
+	if(INiagaraModule::DataChannelsEnabled() && World && !World->IsNetMode(NM_DedicatedServer) && InChannel->IsValid())
 	{
-		FName Channel = InChannel->GetChannelName();
-		TObjectPtr<UNiagaraDataChannelHandler>& Handler = Channels.FindOrAdd(Channel);
+		TObjectPtr<UNiagaraDataChannelHandler>& Handler = Channels.FindOrAdd(InChannel);
 
 		if (bForce || Handler == nullptr)
 		{
-			Handler = InChannel->CreateHandler();
+			Handler = InChannel->CreateHandler(WorldMan->GetWorld());
 		}
 	}
 }
 
-void FNiagaraDataChannelManager::RemoveDataChannel(FName ChannelName)
+void FNiagaraDataChannelManager::RemoveDataChannel(const UNiagaraDataChannel* InChannel)
 {
-	Channels.Remove(ChannelName);
+	Channels.Remove(InChannel);
 }
 
 UWorld* FNiagaraDataChannelManager::GetWorld()const
diff --git a/Engine/Plugins/FX/Niagara/Source/Niagara/Private/NiagaraDataChannelManager.h b/Engine/Plugins/FX/Niagara/Source/Niagara/Private/NiagaraDataChannelManager.h
--- a/Engine/Plugins/FX/Niagara/Source/Niagara/Private/NiagaraDataChannelManager.h
+++ b/Engine/Plugins/FX/Niagara/Source/Niagara/Private/NiagaraDataChannelManager.h
@@ -3,6 +3,7 @@
 #pragma once
 
 #include "UObject/ObjectPtr.h"
+#include "NiagaraDataChannelPublic.h"
 
 class FNiagaraWorldManager;
 class UNiagaraDataChannel;
@@ -24,14 +25,18 @@ public:
 	void Cleanup();
 
 	void InitDataChannel(const UNiagaraDataChannel* Channel, bool bForce);
-	void RemoveDataChannel(FName ChannelName);
+	void RemoveDataChannel(const UNiagaraDataChannel* Channel);
+
+	void BeginFrame(float DeltaSeconds);
+	void EndFrame(float DeltaSeconds);
 
 	void Tick(float DeltaSeconds, ETickingGroup TickGroup);
 	
 	/**
 	Return the DataChannel handler for the given channel.
 	*/
-	UNiagaraDataChannelHandler* FindDataChannelHandler(FName ChannelName);
+	UNiagaraDataChannelHandler* FindDataChannelHandler(const UNiagaraDataChannel* Channel);
+	UNiagaraDataChannelHandler* FindDataChannelHandler(const UNiagaraDataChannelAsset* Channel) { return Channel ? FindDataChannelHandler(Channel->Get()) : nullptr; }
 
 	UWorld* GetWorld()const;
 
@@ -39,5 +44,5 @@ private:
 	FNiagaraWorldManager* WorldMan = nullptr;
 
 	/** Runtime handlers for each DataChannel channel. */
-	TMap<FName, TObjectPtr<UNiagaraDataChannelHandler>> Channels;
+	TMap<TObjectPtr<const UNiagaraDataChannel>, TObjectPtr<UNiagaraDataChannelHandler>> Channels;
 };
diff --git a/Engine/Plugins/FX/Niagara/Source/Niagara/Private/NiagaraDataChannel_Global.cpp b/Engine/Plugins/FX/Niagara/Source/Niagara/Private/NiagaraDataChannel_Global.cpp
--- a/Engine/Plugins/FX/Niagara/Source/Niagara/Private/NiagaraDataChannel_Global.cpp
+++ b/Engine/Plugins/FX/Niagara/Source/Niagara/Private/NiagaraDataChannel_Global.cpp
@@ -1,10 +1,8 @@
 // Copyright Epic Games, Inc. All Rights Reserved.
 
 #include "NiagaraDataChannel_Global.h"
-
-#include "NiagaraSystemInstance.h"
-#include "NiagaraGpuComputeDispatchInterface.h"
 #include "NiagaraWorldManager.h"
+#include "NiagaraDataChannelCommon.h"
 
 
 int32 GbDebugDumpDumpHandlerTick = 0;
@@ -19,9 +17,9 @@ DECLARE_CYCLE_STAT(TEXT("UNiagaraDataChannelHandler_Global::Tick"), STAT_DataCha
 
 //////////////////////////////////////////////////////////////////////////
 
-UNiagaraDataChannelHandler* UNiagaraDataChannel_Global::CreateHandler()const
+UNiagaraDataChannelHandler* UNiagaraDataChannel_Global::CreateHandler(UWorld* OwningWorld)const
 {
-	UNiagaraDataChannelHandler* NewHandler = NewObject<UNiagaraDataChannelHandler_Global>();
+	UNiagaraDataChannelHandler* NewHandler = NewObject<UNiagaraDataChannelHandler_Global>(OwningWorld);
 	NewHandler->Init(this);
 	return NewHandler;
 }
@@ -32,78 +30,33 @@ UNiagaraDataChannelHandler* UNiagaraDataChannel_Global::CreateHandler()const
 
 //////////////////////////////////////////////////////////////////////////
 
-class FNiagaraDataChannelHandlerRTProxy_Global : public FNiagaraDataChannelHandlerRTProxyBase
-{	
-public:
-	FNiagaraDataSet* GPUDataSet = nullptr;
-	FString DebugName;
-
-	void BeginFrame(FNiagaraGpuComputeDispatchInterface* DispathInterface, FRHICommandListImmediate& CmdList)
-	{
-	}
-
-	void EndFrame(FNiagaraGpuComputeDispatchInterface* DispathInterface, FRHICommandListImmediate& CmdList, const TArray<FNiagaraDataBufferRef>& BuffersForGPU)
-	{
-		uint32 NumInstance = 0;
-		for(auto& Buffer : BuffersForGPU)
-		{
-			NumInstance += Buffer->GetNumInstances();
-		}		
-
-		GPUDataSet->BeginSimulate();
-		FNiagaraDataBuffer* DestBuffer = GPUDataSet->GetDestinationData();
-		DestBuffer->PushCPUBuffersToGPU(BuffersForGPU, true, CmdList, DispathInterface->GetFeatureLevel(), *DebugName);
-		GPUDataSet->EndSimulate();
-		
-		//For now we need not deal with the instance count manager but when we do GPU->GPU writes we will
-		//DestBuffer->SetGPUInstanceCountBufferOffset()
-	}
-};
 
 UNiagaraDataChannelHandler_Global::UNiagaraDataChannelHandler_Global(FObjectInitializer const& ObjectInitializer)
 	: Super(ObjectInitializer)
 {
-	RTProxy.Reset(new FNiagaraDataChannelHandlerRTProxy_Global());
 }
 
 void UNiagaraDataChannelHandler_Global::BeginDestroy()
 {
 	Super::BeginDestroy();
-	Data.Empty();
-	
-	ENQUEUE_RENDER_COMMAND(FDeleteContextCommand)(
-		[RT_GameDataStaging = GameDataStaging](FRHICommandListImmediate& RHICmdList)
-		{
-			if (RT_GameDataStaging != nullptr)
-			{
-				delete RT_GameDataStaging;
-			}
-		}
-	);
-	GameDataStaging = nullptr;
+	Data.Reset();
 }
 
 void UNiagaraDataChannelHandler_Global::Init(const UNiagaraDataChannel* InChannel)
 {
 	Super::Init(InChannel);
-	Data.Init(this);
-
-	const FNiagaraDataSetCompiledData& CompiledData = GetDataChannel()->GetCompiledData(ENiagaraSimTarget::CPUSim);
-	GameDataStaging = new FNiagaraDataSet();
-	GameDataStaging->Init(&CompiledData);
+	check(DataChannel);
+	Data = CreateData();
+}
 
-	FString GPUDebugName;
-#if !UE_BUILD_SHIPPING
-	GPUDebugName = FString::Printf(TEXT("%s__GPUData"), *InChannel->GetChannelName().ToString());
-#endif
+void UNiagaraDataChannelHandler_Global::BeginFrame(float DeltaTime, FNiagaraWorldManager* OwningWorld)
+{
+	Data->BeginFrame(this);
+}
 
-	ENQUEUE_RENDER_COMMAND(FInitProxy) (
-		[RT_Proxy = GetRTProxyAs<FNiagaraDataChannelHandlerRTProxy_Global>(), GPUDataSet = Data.GPUSimData, DebugName=MoveTemp(GPUDebugName)](FRHICommandListImmediate& CmdList)mutable
-	{
-		RT_Proxy->GPUDataSet = GPUDataSet;
-		RT_Proxy->DebugName = MoveTemp(DebugName);
-	}
-	);
+void UNiagaraDataChannelHandler_Global::EndFrame(float DeltaTime, FNiagaraWorldManager* OwningWorld)
+{
+	Data->EndFrame(this);
 }
 
 void UNiagaraDataChannelHandler_Global::Tick(float DeltaSeconds, ETickingGroup TickGroup, FNiagaraWorldManager* OwningWorld)
@@ -111,196 +64,11 @@ void UNiagaraDataChannelHandler_Global::Tick(float DeltaSeconds, ETickingGroup T
 	SCOPE_CYCLE_COUNTER(STAT_DataChannelHandler_Global_Tick);
 	Super::Tick(DeltaSeconds, TickGroup, OwningWorld);
 
-	if (Data.IsInitialized() == false)
-	{
-		Data.Init(this);
-	}
-
-	if(GameDataStaging == nullptr)
-	{
-		const FNiagaraDataSetCompiledData CompiledData = GetDataChannel()->GetCompiledData(ENiagaraSimTarget::CPUSim);
-		GameDataStaging = new FNiagaraDataSet();
-		GameDataStaging->Init(&CompiledData);
-	}
-
-	bool bIsFirstTick = TickGroup == NiagaraFirstTickGroup;
-	bool bIsLastTick = TickGroup == NiagaraLastTickGroup;
-
-	if (bIsFirstTick)
-	{
-		//TODO: Add reference struct to handle these ref counts and make add/release private.
-		if(DataChannel->KeepPreviousFrameData())
-		{
-			Data.PrevCPUSimData = Data.CPUSimData->GetCurrentData();
-		}
-
-		FNiagaraGpuComputeDispatchInterface* DispathInterface = FNiagaraGpuComputeDispatchInterface::Get(OwningWorld->GetWorld());
-		ENQUEUE_RENDER_COMMAND(FDataChannelProxyBeginFrame) (
-			[DispathInterface, RT_Proxy = GetRTProxyAs<FNiagaraDataChannelHandlerRTProxy_Global>()](FRHICommandListImmediate& CmdList)
-		{
-			if(RT_Proxy)
-			{
-				RT_Proxy->BeginFrame(DispathInterface, CmdList);
-			}
-		});
-	}
-
-	//Gather various sources of DataChannel and ship them around as heeded.
-
-	int32 GameDataOrigSize = bIsFirstTick ? 0 : Data.GameData->Num();
-	int32 CPUDataOrigSize = bIsFirstTick ? 0 : Data.CPUSimData->GetCurrentData()->GetNumInstances();
-	int32 NewGameDataChannel = GameDataOrigSize;
-	int32 NewCPUDataChannel = CPUDataOrigSize;
-
-	//Do a pass to gather the new total size for our DataChannel data.
-
-	//Each DI that generates DataChannel can control whether it's pushed to Game/CPU/GPU.
-	BuffersForGPU.Reserve(BuffersForGPU.Num() + PublishRequests.Num());
-	for (FNiagaraDataChannelPublishRequest& PublishRequest : PublishRequests)
-	{
-		uint32 NumInsts = 0;
-
-		if(FNiagaraDataChannelGameData* RequestGameData = PublishRequest.GameData.Get())
-		{
-			NumInsts = RequestGameData->Num();
-			if (RequestGameData->Num() > 0 && (PublishRequest.bVisibleToCPUSims || PublishRequest.bVisibleToGPUSims))
-			{
-				//For now we copy game data into a staging buffer for pushing to CPU & GPU buffers.
-				//We can avoid this in future but will add complexity to the transfer code.
-				GameDataStaging->BeginSimulate();
-				GameDataStaging->Allocate(RequestGameData->Num());
-				RequestGameData->WriteToDataSet(GameDataStaging->GetDestinationData(), 0);
-				GameDataStaging->EndSimulate();
-
-				if (GbDebugDumpDumpHandlerTick)
-				{
-					FNiagaraDataBuffer& Buffer = GameDataStaging->GetCurrentDataChecked();
-					FString DebugString = FString::Printf(TEXT("%s Game Data Staging"), *DataChannel->GetChannelName().ToString());
-					Buffer.Dump(0, Buffer.GetNumInstances(), FString::Printf(TEXT("=== Data Channle Staging Game Data: %d Elements --> %s ==="), Buffer.GetNumInstances(), *DebugString));
-				}
-
-				PublishRequest.Data = GameDataStaging->GetCurrentData();
-			}
-		}
-		else if( ensure(PublishRequest.Data) )
-		{
-			NumInsts = PublishRequest.Data->GetNumInstances();
-		}
-	
-		if (PublishRequest.bVisibleToGame)
-		{
-			NewGameDataChannel += NumInsts;
-		}
-		if (PublishRequest.bVisibleToCPUSims)
-		{
-			NewCPUDataChannel += NumInsts;
-		}
-		if(PublishRequest.bVisibleToGPUSims)
-		{
-			BuffersForGPU.Add(PublishRequest.Data);
-		}
-	}
-
-	//Allocate Sim buffers ready for gather.
-	Data.CPUSimData->BeginSimulate();
-	Data.CPUSimData->Allocate(NewCPUDataChannel, bIsFirstTick == false);
-
-	//Now do the actual data collection.
-	
-	//Allocate game buffers ready for gather.
-	if(bIsFirstTick)
-	{
-		Data.GameData->Reset();
-	}
-	
-	Data.GameData->SetNum(NewGameDataChannel);
-
-	for (FNiagaraDataChannelPublishRequest& PublishRequest : PublishRequests)
-	{
-		if (FNiagaraDataChannelGameData* RequestGameData = PublishRequest.GameData.Get())
-		{
-			if(RequestGameData->Num() && PublishRequest.bVisibleToGame)
-			{
-				Data.GameData->AppendFromGameData(*RequestGameData);
-			}
-		}
-
-		if (PublishRequest.Data)
-		{
-			const ENiagaraSimTarget SimTarget = PublishRequest.Data->GetOwner()->GetSimTarget();
-			if (SimTarget == ENiagaraSimTarget::CPUSim)
-			{
-				if (PublishRequest.bVisibleToGame)
-				{
-					Data.GameData->AppendFromDataSet(PublishRequest.Data, PublishRequest.LwcTile);
-				}
-				if (PublishRequest.bVisibleToCPUSims)
-				{
-					PublishRequest.Data->CopyToUnrelated(Data.CPUSimData->GetDestinationDataChecked(), 0, Data.CPUSimData->GetDestinationDataChecked().GetNumInstances(), PublishRequest.Data->GetNumInstances());
-				}
-			}
-			else if (SimTarget == ENiagaraSimTarget::GPUComputeSim)
-			{
-				//TODO: GPU->GPU handling will be done all RT side. GPU->CPU may be done here in future but we'll have to pull the data from the GPU on the RT and pass beck to GT here.
-			}
-			else
-			{
-				check(0);
-			}
-		}
-	}
-
-	Data.CPUSimData->EndSimulate();
-
-	if (GbDebugDumpDumpHandlerTick)
-	{
-		FNiagaraDataBuffer& Buffer = Data.CPUSimData->GetCurrentDataChecked();
-		FString DebugString = FString::Printf(TEXT("%s CPU Sim Data"), *DataChannel->GetChannelName().ToString());
-		Buffer.Dump(0, Buffer.GetNumInstances(), FString::Printf(TEXT("=== Updated CPU Sim Data: %d Elements --> %s ==="), CPUDataOrigSize, *DebugString));
-	}
-
-	PublishRequests.Reset();
-
-	//Now we push the CPU sim data to the RT for upload to the GPU.
-	//We may need to do a final tick from PostActorTick in the world manager?
-	if(bIsLastTick)
-	{
-		FNiagaraDataChannelHandlerRTProxy_Global* Proxy = GetRTProxyAs<FNiagaraDataChannelHandlerRTProxy_Global>();
-
-		check(Proxy);
-		check(Proxy->GPUDataSet);
-		check(Proxy->GPUDataSet->GetSimTarget() == ENiagaraSimTarget::GPUComputeSim);
-		FNiagaraGpuComputeDispatchInterface* DispathInterface = FNiagaraGpuComputeDispatchInterface::Get(OwningWorld->GetWorld());
-		ENQUEUE_RENDER_COMMAND(FDataChannelProxyEndFrame) (
-			[DispathInterface, RT_Proxy = Proxy, RT_BuffersForGPU=MoveTemp(BuffersForGPU)](FRHICommandListImmediate& CmdList) 
-			{
-				RT_Proxy->EndFrame(DispathInterface, CmdList, RT_BuffersForGPU);				
-			});
-		BuffersForGPU.Reset();
-	}
-}
-
-FNiagaraDataBufferRef UNiagaraDataChannelHandler_Global::GetData(FNiagaraSystemInstance* SystemInstance, bool bGetPrevFrameData)
-{
-	if(bGetPrevFrameData && Data.PrevCPUSimData)
-	{
-		return Data.PrevCPUSimData;
-	}
-	else
-	{
-		return Data.CPUSimData->GetCurrentData();	
-	}
-	//For more complicated channels we could check the location + bounds of the system instance etc to return some spatially localized data.
+	Data->ConsumePublishRequests(this);
 }
 
-FNiagaraDataSet* UNiagaraDataChannelHandler_Global::GetDataGPU(FNiagaraSystemInstance* SystemInstance)
+FNiagaraDataChannelDataPtr UNiagaraDataChannelHandler_Global::FindData(FNiagaraDataChannelSearchParameters SearchParams, ENiagaraResourceAccess AccessType)
 {
-	return Data.GPUSimData;	
+	return Data;
 	//For more complicated channels we could check the location + bounds of the system instance etc to return some spatially localized data.
 }
-
-
-FNiagaraDataChannelGameDataPtr UNiagaraDataChannelHandler_Global::GetGameData(UActorComponent* OwningComponent)
-{
-	return Data.GameData;
-}
