Impact FX Aggregation via Data Channels.

diff --git a/Engine/Plugins/FX/Niagara/Source/Niagara/Private/NiagaraDataChannel_Islands.cpp b/Engine/Plugins/FX/Niagara/Source/Niagara/Private/NiagaraDataChannel_Islands.cpp
new file mode 100644
--- /dev/null
+++ b/Engine/Plugins/FX/Niagara/Source/Niagara/Private/NiagaraDataChannel_Islands.cpp
@@ -0,0 +1,413 @@
+// Copyright Epic Games, Inc. All Rights Reserved.
+
+#include "NiagaraDataChannel_Islands.h"
+
+#include "Engine/AssetManager.h"
+#include "UObject/UObjectGlobals.h"
+
+#include "NiagaraFunctionLibrary.h"
+#include "NiagaraComponent.h"
+#include "NiagaraWorldManager.h"
+#include "NiagaraDataChannelCommon.h"
+
+UNiagaraDataChannelHandler* UNiagaraDataChannel_Islands::CreateHandler(UWorld* OwningWorld)const
+{
+	UNiagaraDataChannelHandler* NewHandler = NewObject<UNiagaraDataChannelHandler_Islands>(OwningWorld);
+	NewHandler->Init(this);
+	return NewHandler;
+}
+
+UNiagaraDataChannelHandler_Islands::UNiagaraDataChannelHandler_Islands(FObjectInitializer const& ObjectInitializer)
+	: Super(ObjectInitializer)
+{
+}
+
+void UNiagaraDataChannelHandler_Islands::BeginDestroy()
+{
+	Super::BeginDestroy();
+	ActiveIslands.Empty();
+	FreeIslands.Empty();
+	IslandPool.Empty();
+}
+
+void UNiagaraDataChannelHandler_Islands::Init(const UNiagaraDataChannel* InChannel)
+{
+	Super::Init(InChannel);
+	if (const UNiagaraDataChannel_Islands* IslandChannel = CastChecked<UNiagaraDataChannel_Islands>(InChannel))
+	{
+		//For now immediately load handler systems but in future we should look at deferring the load until a channel is used.
+		for (const TSoftObjectPtr<UNiagaraSystem>& SoftSys : IslandChannel->Systems)
+		{
+			UAssetManager::GetStreamableManager().RequestAsyncLoad(SoftSys.ToSoftObjectPath()
+				, FStreamableDelegate()
+				, FStreamableManager::AsyncLoadHighPriority
+				, false
+				, false
+				, TEXT("NiagaraDataChannelIsland_HandlerSystems"));
+		}
+
+		int32 InitialPoolCount = IslandChannel->IslandPoolSize;
+		ActiveIslands.Reserve(InitialPoolCount);
+		FreeIslands.Reserve(InitialPoolCount);
+		IslandPool.SetNum(InitialPoolCount);
+		for (int32 i = 0; i < InitialPoolCount; ++i)
+		{
+			IslandPool[i].Init(this);
+			FreeIslands.Emplace(i);
+		}
+	}
+}
+
+void UNiagaraDataChannelHandler_Islands::BeginFrame(float DeltaTime, FNiagaraWorldManager* OwningWorld)
+{
+	Super::BeginFrame(DeltaTime, OwningWorld);
+	for (auto It = ActiveIslands.CreateIterator(); It; ++It)
+	{
+		FNDCIsland& Island = IslandPool[*It];
+		Island.BeginFrame();
+	}
+}
+
+void UNiagaraDataChannelHandler_Islands::EndFrame(float DeltaTime, FNiagaraWorldManager* OwningWorld)
+{
+	Super::EndFrame(DeltaTime, OwningWorld);
+	for (auto It = ActiveIslands.CreateIterator(); It; ++It)
+	{
+		FNDCIsland& Island = IslandPool[*It];
+		Island.EndFrame();
+	}
+}
+
+void UNiagaraDataChannelHandler_Islands::Tick(float DeltaTime, ETickingGroup TickGroup, FNiagaraWorldManager* OwningWorld)
+{
+	UNiagaraDataChannel_Islands* IslandChannel = GetChannelTyped<UNiagaraDataChannel_Islands>();
+	for (auto It = ActiveIslands.CreateIterator(); It; ++It)
+	{
+		FNDCIsland& Island = IslandPool[*It];
+		Island.Tick();
+
+		//Free up islands that are not being used.
+		//We do this immediately as we leave it up to the handler systems to delay their cleanup as long as they want to.
+		if (Island.IsBeingUsed() == false)
+		{
+			FreeIslands.Emplace(*It);
+			Island.OnReleased();
+			It.RemoveCurrentSwap();
+		}
+	}
+}
+
+FNiagaraDataChannelDataPtr UNiagaraDataChannelHandler_Islands::FindData(FNiagaraDataChannelSearchParameters SearchParams, ENiagaraResourceAccess AccessType)
+{
+	if (FNDCIsland* Island = FindOrCreateIsland(SearchParams, AccessType))
+	{
+		return Island->GetData();
+	}
+	return nullptr;
+}
+
+FNDCIsland* UNiagaraDataChannelHandler_Islands::FindOrCreateIsland(const FNiagaraDataChannelSearchParameters& SearchParams, ENiagaraResourceAccess AccessType)
+{
+	UNiagaraDataChannel_Islands* IslandChannel = CastChecked<UNiagaraDataChannel_Islands>(DataChannel);
+
+	FVector Location = SearchParams.GetLocation();
+	if (AccessType == ENiagaraResourceAccess::ReadOnly)
+	{
+		//If we're reading only then we just try to find the right island to read from.
+		//When writing we'll grow/spawn islands to accommodate the data.
+
+		//First we see if this is an islands handler system.
+		if(SearchParams.OwningComponent)
+		{
+			for (int32 i : ActiveIslands)
+			{
+				FNDCIsland& Island = IslandPool[i];
+				if (Island.IsHandlerSystem(SearchParams.OwningComponent))
+				{
+					return &Island;
+				}
+			}
+			//Failing that, we'll just return the first island that intersects the owners bounds.
+			for (int32 i : ActiveIslands)
+			{
+				FNDCIsland& Island = IslandPool[i];
+				if (Island.Intersects(SearchParams.OwningComponent->Bounds))
+				{
+					return &Island;
+				}
+			}
+		}
+		else
+		{
+			//We don't have an owning component so just return the first island to contain the location.
+			for (int32 i : ActiveIslands)
+			{
+				FNDCIsland& Island = IslandPool[i];
+				if (Island.Contains(Location))
+				{
+					return &Island;
+				}
+			}
+		}
+		return nullptr;
+	}
+	else
+	{
+		UWorld* World = GetWorld();
+		check(World);
+		double WorldTime = World->GetTimeSeconds();
+
+		//Find the first island that could contain this point.
+		//For now we do a linear search of active islands.
+		//Assuming that overall active island count will be low.
+		//If this is not the case then we'll want to add an acceleration structure to speed up this search.
+		FNDCIsland* IslandToUse = nullptr;
+		FVector MaxExtents = IslandChannel->MaxExtents;
+		FVector PerElementExtents = IslandChannel->PerElementExtents;
+		for (int32 i : ActiveIslands)
+		{
+			FNDCIsland& Island = IslandPool[i];
+			FBoxSphereBounds GrowthBounds;
+			if (Island.TryGrow(Location, PerElementExtents, MaxExtents))
+			{
+				IslandToUse = &Island;
+			}
+		}
+
+		if (IslandToUse)
+		{
+			return IslandToUse;
+		}
+
+		//Failing that, get a init a new island from the pool.
+		int32 NewIslandIndex = ActivateNewIsland(Location);
+		if (ensure(NewIslandIndex != INDEX_NONE))
+		{
+			IslandToUse = &IslandPool[NewIslandIndex];
+		}
+
+		if (IslandToUse)
+		{
+			IslandToUse->OnAcquired(Location);
+		}
+
+		return IslandToUse;
+	}
+}
+
+int32 UNiagaraDataChannelHandler_Islands::ActivateNewIsland(FVector Location)
+{
+	int32 NewIndex = INDEX_NONE;
+	FNDCIsland* NewIsland = nullptr;
+	if (FreeIslands.Num() > 0)
+	{
+		NewIndex = FreeIslands.Pop(false);
+		NewIsland = &IslandPool[NewIndex];
+	}
+
+	if (NewIsland == nullptr)
+	{
+		NewIndex = IslandPool.Num();
+		NewIsland = &IslandPool.AddDefaulted_GetRef();
+	}
+
+	if (NewIndex != INDEX_NONE)
+	{
+		check(NewIsland);
+		ActiveIslands.Add(NewIndex);
+		NewIsland->Init(this);
+	}
+	else
+	{
+		UE_LOG(LogNiagara, Error, TEXT("Failed to allocate new island for Island Data Channel %s"), *GetDataChannel()->GetName());
+	}
+
+	return NewIndex;
+}
+
+//////////////////////////////////////////////////////////////////////////
+
+FNDCIsland::FNDCIsland()
+{
+}
+
+FNDCIsland::~FNDCIsland()
+{
+	for (UNiagaraComponent* Comp : NiagaraSystems)
+	{
+		if (IsValid(Comp))
+		{
+			Comp->ReleaseToPool();
+		}
+	}
+	NiagaraSystems.Empty();
+}
+
+void FNDCIsland::Init(UNiagaraDataChannelHandler_Islands* InOwner)
+{
+	Owner = InOwner;
+
+	Data = Owner->CreateData();
+
+	const UNiagaraDataChannel_Islands* Channel = CastChecked<UNiagaraDataChannel_Islands>(Owner->GetDataChannel());
+
+	Bounds.Origin = FVector::ZeroVector;
+	Bounds.BoxExtent = Channel->InitialExtents;
+}
+
+void FNDCIsland::OnAcquired(FVector Location)
+{
+	const UNiagaraDataChannel_Islands* Channel = CastChecked<UNiagaraDataChannel_Islands>(Owner->GetDataChannel());
+
+	if(Channel->Mode == ENiagraDataChannel_IslandMode::AlignedStatic)
+	{
+		//Align our location down to the nearest lower multiple of the max island size to avoid overlapping islands.
+		//TODO: make our padding extents come from the owning component or the handler systems?
+		FVector ActualMaxSize = Channel->MaxExtents + Channel->PerElementExtents;
+		FVector BoxSize = Channel->MaxExtents * 2;
+		FVector BoxN = Location / BoxSize;
+		BoxN.X = FMath::RoundToNegativeInfinity(BoxN.X);
+		BoxN.Y = FMath::RoundToNegativeInfinity(BoxN.Y);
+		BoxN.Z = FMath::RoundToNegativeInfinity(BoxN.Z);
+		Bounds.Origin = BoxN * BoxSize + Channel->MaxExtents;
+		Bounds.BoxExtent = ActualMaxSize;
+	}
+	else
+	{
+		Bounds.Origin = Location;
+		Bounds.BoxExtent = Channel->InitialExtents;
+	}
+
+	Data->SetLwcTile(FLargeWorldRenderScalar::GetTileFor(Bounds.Origin));
+	//Spawn our handler systems.
+	FBox HandlerSystemBounds(-Bounds.BoxExtent, Bounds.BoxExtent);
+	for (auto& SoftSys : Channel->Systems)
+	{
+		if (UNiagaraSystem* Sys = SoftSys.Get())
+		{
+			FFXSystemSpawnParameters SpawnParams;
+			SpawnParams.bAutoActivate = false;//We must activate AFTER adding this component to our handle system array.
+			SpawnParams.bAutoDestroy = false;
+			SpawnParams.Location = Bounds.Origin;
+			SpawnParams.PoolingMethod = EPSCPoolMethod::ManualRelease;
+			SpawnParams.bPreCullCheck = false;
+			SpawnParams.SystemTemplate = Sys;
+			SpawnParams.WorldContextObject = Owner->GetWorld();
+			if(UNiagaraComponent* NewComp = UNiagaraFunctionLibrary::SpawnSystemAtLocationWithParams(SpawnParams))
+			{
+				NiagaraSystems.Add(NewComp);
+				NewComp->SetSystemFixedBounds(HandlerSystemBounds);
+				NewComp->SetTickBehavior(ENiagaraTickBehavior::ForceTickLast);
+				NewComp->Activate();
+			}
+		}
+	}
+	
+	if (Channel->Mode == ENiagraDataChannel_IslandMode::Dynamic)
+	{
+		ensure(TryGrow(Location, Channel->PerElementExtents, Channel->MaxExtents));
+	}
+}
+
+void FNDCIsland::OnReleased()
+{
+	Data->Reset();
+	for (UNiagaraComponent* Comp : NiagaraSystems)
+	{
+		Comp->ReleaseToPool();
+	}
+	NiagaraSystems.Reset();
+}
+
+void FNDCIsland::BeginFrame()
+{
+	Data->BeginFrame(Owner);
+}
+
+void FNDCIsland::EndFrame()
+{
+	Data->EndFrame(Owner);
+}
+
+void FNDCIsland::Tick()
+{
+	Data->ConsumePublishRequests(Owner);
+}
+
+bool FNDCIsland::Contains(FVector Point)
+{
+	if (DistanceToPoint(Point) > 0)
+	{
+		return false;
+	}
+
+	return true;
+}
+
+double FNDCIsland::DistanceToPoint(FVector Point)
+{
+	return Bounds.ComputeSquaredDistanceFromBoxToPoint(Point);
+}
+
+bool FNDCIsland::Intersects(const FBoxSphereBounds& CheckBounds)
+{
+	return FBoxSphereBounds::BoxesIntersect(Bounds, CheckBounds);
+}
+
+bool FNDCIsland::IsHandlerSystem(USceneComponent* Component)
+{
+	return NiagaraSystems.Contains(Component);
+}
+
+bool FNDCIsland::IsBeingUsed()const
+{
+	return Data.IsUnique() == false;//If someone is making use of our data then we're still in use.
+}
+
+bool FNDCIsland::TryGrow(FVector Point, FVector PerElementExtents, FVector MaxIslandExtents)
+{
+	FVector CurrMax = Bounds.Origin + Bounds.BoxExtent;
+	FVector CurrMin = Bounds.Origin - Bounds.BoxExtent;
+
+	FVector ElementMax = Point + PerElementExtents;
+	FVector ElementMin = Point - PerElementExtents;
+
+	//The max size of the island, we pad these out with the per element extents to ensure bounds cover anything right at the boundary.
+	FVector Max = Bounds.Origin + MaxIslandExtents + PerElementExtents;
+	FVector Min = Bounds.Origin - MaxIslandExtents - PerElementExtents;
+
+	//Can this island contain this element at all?
+
+	//Calculate the require bounds to encompass the point + the per element bounds.
+	FVector NewMax(FMath::Max(ElementMax.X, CurrMax.X), FMath::Max(ElementMax.Y, CurrMax.Y), FMath::Max(ElementMax.Z, CurrMax.Z));
+	FVector NewMin(FMath::Min(ElementMin.X, CurrMin.X), FMath::Min(ElementMin.Y, CurrMin.Y), FMath::Min(ElementMin.Z, CurrMin.Z));
+
+	//Return false if the new bounds are outside the max allowed for this island.
+	if (NewMax.X > Max.X || NewMax.Y > Max.Y || NewMax.Z > Max.Z || NewMin.X < Min.X || NewMin.Y < Min.Y || NewMin.Z < Min.Z)
+	{
+		return false;
+	}
+
+	FVector OldOrigin = Bounds.Origin;
+	FBox LocalBounds(NewMin - Bounds.Origin, NewMax - Bounds.Origin);
+	Bounds = LocalBounds;
+	Bounds.Origin += OldOrigin;
+
+	//If new bounds are valid, grow our bounds and update the 
+	FBox HandlerSystemBounds(-Bounds.BoxExtent, Bounds.BoxExtent);
+	for (auto It = NiagaraSystems.CreateIterator(); It; ++It)
+	{
+		UNiagaraComponent* Comp = *It;
+		if (IsValid(Comp))
+		{
+			Comp->SetSystemFixedBounds(LocalBounds);
+			Comp->SetWorldLocation(Bounds.Origin);
+		}
+		else
+		{
+			It.RemoveCurrentSwap();
+		}
+	}
+
+	return true;
+}
diff --git a/Engine/Plugins/FX/Niagara/Source/Niagara/Private/NiagaraEmitterInstance.cpp b/Engine/Plugins/FX/Niagara/Source/Niagara/Private/NiagaraEmitterInstance.cpp
--- a/Engine/Plugins/FX/Niagara/Source/Niagara/Private/NiagaraEmitterInstance.cpp
+++ b/Engine/Plugins/FX/Niagara/Source/Niagara/Private/NiagaraEmitterInstance.cpp
@@ -1421,7 +1421,7 @@ void FNiagaraEmitterInstance::Tick(float DeltaSeconds)
 		check(GPUExecContext->GPUScript_RT == EmitterData->GetGPUComputeScript()->GetRenderThreadScript());
 		GPUExecContext->GPUScript_RT = EmitterData->GetGPUComputeScript()->GetRenderThreadScript();
 
-#if WITH_EDITOR
+#if NIAGARA_SYSTEM_CAPTURE
 		if (ParentSystemInstance->ShouldCaptureThisFrame())
 		{
 			TSharedPtr<FNiagaraScriptDebuggerInfo, ESPMode::ThreadSafe> DebugInfo = ParentSystemInstance->GetActiveCaptureWrite(CachedIDName, ENiagaraScriptUsage::ParticleGPUComputeScript, FGuid());
@@ -1432,7 +1432,7 @@ void FNiagaraEmitterInstance::Tick(float DeltaSeconds)
 				DebugInfo->Parameters = GPUExecContext->CombinedParamStore;
 				
 				//TODO: This layout info can be pulled into the emitter/systems etc and all sets just refer to them. They are becoming an annoyance here.
-				DebugInfo->Frame.Init(&CachedEmitterCompiledData->GPUCaptureDataSetCompiledData);
+				DebugInfo->Frame.Init(&CachedEmitterCompiledData->GetGPUCaptureDataSetCompiledData());
 
 				// Execute a readback
 				ENQUEUE_RENDER_COMMAND(NiagaraReadbackGpuSim)(
@@ -1867,7 +1867,7 @@ void FNiagaraEmitterInstance::Tick(float DeltaSeconds)
 	}
 
 	//Now pull out any debug info we need.
-#if WITH_EDITORONLY_DATA
+#if NIAGARA_SYSTEM_CAPTURE
 	if (ParentSystemInstance->ShouldCaptureThisFrame())
 	{
 		//Pull out update data.
@@ -1948,7 +1948,7 @@ void FNiagaraEmitterInstance::Tick(float DeltaSeconds)
 							EventInstanceData->EventExecContexts[EventScriptIdx].Parameters.Dump();
 						}
 
-	#if WITH_EDITORONLY_DATA
+	#if NIAGARA_SYSTEM_CAPTURE
 						if (ParentSystemInstance->ShouldCaptureThisFrame())
 						{
 							FGuid EventGuid = EventInstanceData->EventExecContexts[EventScriptIdx].Script->GetUsageId();
@@ -2014,7 +2014,7 @@ void FNiagaraEmitterInstance::Tick(float DeltaSeconds)
 							EventInstanceData->EventExecContexts[EventScriptIdx].Parameters.Dump();
 						}
 
-#if WITH_EDITORONLY_DATA
+#if NIAGARA_SYSTEM_CAPTURE
 						if (ParentSystemInstance->ShouldCaptureThisFrame())
 						{
 							FGuid EventGuid = EventInstanceData->EventExecContexts[EventScriptIdx].Script->GetUsageId();
@@ -2073,7 +2073,7 @@ void FNiagaraEmitterInstance::Tick(float DeltaSeconds)
 // 					}
 // 
 // 
-// #if WITH_EDITORONLY_DATA
+// #if NIAGARA_SYSTEM_CAPTURE
 // 					if (ParentSystemInstance->ShouldCaptureThisFrame())
 // 					{
 // 						FGuid EventGuid = EventExecContexts[EventScriptIdx].Script->GetUsageId();
diff --git a/Engine/Plugins/FX/Niagara/Source/Niagara/Private/NiagaraModule.cpp b/Engine/Plugins/FX/Niagara/Source/Niagara/Private/NiagaraModule.cpp
--- a/Engine/Plugins/FX/Niagara/Source/Niagara/Private/NiagaraModule.cpp
+++ b/Engine/Plugins/FX/Niagara/Source/Niagara/Private/NiagaraModule.cpp
@@ -43,7 +43,6 @@
 #include "DataInterface/NiagaraDataInterfaceDataChannelCommon.h"
 #include "DataInterface/NiagaraDataInterfaceDataChannelWrite.h"
 #include "DataInterface/NiagaraDataInterfaceDataChannelRead.h"
-#include "DataInterface/NiagaraDataInterfaceDataChannelSpawn.h"
 
 #if UE_USE_OPENVDB
 #include "NiagaraOpenVDB.h"
@@ -107,7 +106,7 @@ void INiagaraModule::OnUseGlobalFXBudgetChanged(IConsoleVariable* Variable)
 	}
 }
 
-bool INiagaraModule::bDataChannelsEnabled = false;
+bool INiagaraModule::bDataChannelsEnabled = true;
 
 static FAutoConsoleVariableRef CVarDataChannelEnabled(
 	TEXT("fx.Niagara.DataChannels.Enabled"),
@@ -124,7 +123,6 @@ void INiagaraModule::OnDataChannelsEnabledChanged(IConsoleVariable* Variable)
 		ENiagaraTypeRegistryFlags Flags = ENiagaraTypeRegistryFlags::AllowAnyVariable | ENiagaraTypeRegistryFlags::AllowParameter;
 		FNiagaraTypeRegistry::Register(UNiagaraDataInterfaceDataChannelWrite::StaticClass(), Flags);
 		FNiagaraTypeRegistry::Register(UNiagaraDataInterfaceDataChannelRead::StaticClass(), Flags);
-		FNiagaraTypeRegistry::Register(UNiagaraDataInterfaceDataChannelSpawn::StaticClass(), Flags);
 
 		FNiagaraWorldManager::ForAllWorldManagers(
 			[](FNiagaraWorldManager& WorldMan)
@@ -268,6 +266,8 @@ FNiagaraVariable INiagaraModule::Translator_CallID;
 void INiagaraModule::StartupModule()
 {
 	VectorVM::Init();
+	FNiagaraTypeHelper::InitStaticTypes();
+
 #if VECTORVM_SUPPORTS_EXPERIMENTAL
 	InitVectorVM();
 #endif
@@ -1114,7 +1114,9 @@ void FNiagaraTypeDefinition::RecreateUserDefinedTypeRegistry()
 	FNiagaraTypeRegistry::Register(USimCacheClassDef, VarFlags);
 	FNiagaraTypeRegistry::Register(StaticEnum<ENiagaraLegacyTrailWidthMode>(), ParamFlags);
 
-	
+
+	FNiagaraTypeRegistry::Register(FNiagaraTypeDefinition(StaticEnum<EPhysicalSurface>()), ParamFlags);
+
 	if (!IsRunningCommandlet())
 	{
 		TArray<FSoftObjectPath> DenyList;
diff --git a/Engine/Plugins/FX/Niagara/Source/Niagara/Private/NiagaraSystem.cpp b/Engine/Plugins/FX/Niagara/Source/Niagara/Private/NiagaraSystem.cpp
--- a/Engine/Plugins/FX/Niagara/Source/Niagara/Private/NiagaraSystem.cpp
+++ b/Engine/Plugins/FX/Niagara/Source/Niagara/Private/NiagaraSystem.cpp
@@ -41,7 +41,6 @@
 #include "UObject/Package.h"
 #include "PipelineStateCache.h"
 #include "NiagaraDataChannel.h"
-#include "NiagaraDataChannelDefinitions.h"
 #include "UObject/UObjectIterator.h"
 
 #include UE_INLINE_GENERATED_CPP_BY_NAME(NiagaraSystem)
@@ -2013,21 +2012,17 @@ void UNiagaraSystem::UpdatePostCompileDIInfo()
 			}
 		}
 	}
-	
-	ReferencedDataChannelDefinitions.Reset();
 
-	//We must include internal DIs in this search so they get correctly PostCompiled.
-	FNiagaraDataInterfaceUtilities::FDataInterfaceSearchOptions SearchOptions;
-	SearchOptions.bIncludeInternal = true;
-	FNiagaraDataInterfaceUtilities::ForEachDataInterface(this, [](const FNiagaraDataInterfaceUtilities::FDataInterfaceUsageContext& Context)
+	ForEachScript([&](UNiagaraScript* Script) 
 	{
-		if (Context.DataInterface)
+		for(const FNiagaraScriptResolvedDataInterfaceInfo& ResolvedDIInfo : Script->GetResolvedDataInterfaces())
 		{
-			Context.DataInterface->PostCompile();
-		}
-		return true;
-	}
-	, SearchOptions);
+			if(ResolvedDIInfo.ResolvedDataInterface)
+			{
+				ResolvedDIInfo.ResolvedDataInterface->PostCompile();
+			}
+		}		
+	});
 }
 
 void UNiagaraSystem::UpdateDITickFlags()
@@ -2086,14 +2081,6 @@ void UNiagaraSystem::UpdateHasGPUEmitters()
 
 #if WITH_EDITORONLY_DATA
 
-void UNiagaraSystem::RegisterDataChannelUse(const UNiagaraDataChannel* DataChannel)
-{
-	if(UNiagaraDataChannelDefinitions* DefinitionsAsset = DataChannel->GetTypedOuter<UNiagaraDataChannelDefinitions>())
-	{
-		ReferencedDataChannelDefinitions.AddUnique(DefinitionsAsset);
-	}
-}
-
 void UNiagaraSystem::OnCompiledDataInterfaceChanged()
 {
 	// We need to make sure our execution ready parameter stores are invalidated so that the DI's will rebind functions
@@ -3021,16 +3008,10 @@ void UNiagaraSystem::InitEmitterCompiledData()
 			const FNiagaraEmitterHandle& EmitterHandle = EmitterHandles[EmitterIdx];
 			const UNiagaraEmitter* Emitter = EmitterHandle.GetInstance().Emitter;
 			FNiagaraDataSetCompiledData& EmitterDataSetCompiledData = NewEmitterCompiledData[EmitterIdx]->DataSetCompiledData;
-			FNiagaraDataSetCompiledData& GPUCaptureCompiledData = NewEmitterCompiledData[EmitterIdx]->GPUCaptureDataSetCompiledData;
 			if (ensureMsgf(Emitter != nullptr, TEXT("Failed to get Emitter Instance from Emitter Handle in post compile, please investigate.")))
 			{
-				static FName GPUCaptureDataSetName = TEXT("GPU Capture Dataset");
 				InitEmitterVariableAliasNames(NewEmitterCompiledData[EmitterIdx].Get(), Emitter);
-				InitEmitterDataSetCompiledData(EmitterDataSetCompiledData, EmitterHandle);
-				GPUCaptureCompiledData.ID = FNiagaraDataSetID(GPUCaptureDataSetName, ENiagaraDataSetType::ParticleData);
-				GPUCaptureCompiledData.Variables = EmitterDataSetCompiledData.Variables;
-				GPUCaptureCompiledData.SimTarget = ENiagaraSimTarget::CPUSim;
-				GPUCaptureCompiledData.BuildLayout();				
+				InitEmitterDataSetCompiledData(EmitterDataSetCompiledData, EmitterHandle);			
 			}
 		}
 
@@ -3640,5 +3621,22 @@ void UNiagaraSystem::SetBakerGeneratedSettings(UNiagaraBakerSettings* Settings)
 
 #endif
 
+
+#if NIAGARA_SYSTEM_CAPTURE
+const FNiagaraDataSetCompiledData& FNiagaraEmitterCompiledData::GetGPUCaptureDataSetCompiledData()const
+{
+	if (GPUCaptureDataSetCompiledData.Variables.Num() == 0)
+	{
+		static FName GPUCaptureDataSetName = TEXT("GPU Capture Dataset");
+		GPUCaptureDataSetCompiledData.ID = FNiagaraDataSetID(GPUCaptureDataSetName, ENiagaraDataSetType::ParticleData);
+		GPUCaptureDataSetCompiledData.Variables = DataSetCompiledData.Variables;
+		GPUCaptureDataSetCompiledData.SimTarget = ENiagaraSimTarget::CPUSim;
+		GPUCaptureDataSetCompiledData.BuildLayout();
+	}
+
+	return GPUCaptureDataSetCompiledData;
+}
+#endif
+
 #undef LOCTEXT_NAMESPACE // NiagaraSystem
 
diff --git a/Engine/Plugins/FX/Niagara/Source/Niagara/Private/NiagaraSystemInstance.cpp b/Engine/Plugins/FX/Niagara/Source/Niagara/Private/NiagaraSystemInstance.cpp
--- a/Engine/Plugins/FX/Niagara/Source/Niagara/Private/NiagaraSystemInstance.cpp
+++ b/Engine/Plugins/FX/Niagara/Source/Niagara/Private/NiagaraSystemInstance.cpp
@@ -116,7 +116,7 @@ FNiagaraSystemInstance::FNiagaraSystemInstance(UWorld& InWorld, UNiagaraSystem&
 	  , RandomSeedOffset(0)
 	  , LODDistance(0.0f)
 	  , MaxLODDistance(FLT_MAX)
-#if WITH_EDITORONLY_DATA
+#if NIAGARA_SYSTEM_CAPTURE
 	  , bWasSoloPriorToCaptureRequest(false)
 #endif
 	  , GlobalParameters{}
@@ -334,7 +334,7 @@ void FNiagaraSystemInstance::DumpTickInfo(FOutputDevice& Ar)
 	Ar.Logf(TEXT("\t\t\tInstance%s"), *PrereqInfo);
 }
 
-#if WITH_EDITORONLY_DATA
+#if NIAGARA_SYSTEM_CAPTURE
 bool FNiagaraSystemInstance::RequestCapture(const FGuid& RequestId)
 {
 	if (IsComplete() || CurrentCapture.IsValid())
@@ -476,7 +476,7 @@ bool FNiagaraSystemInstance::ShouldCaptureThisFrame() const
 {
 	return CurrentCapture.IsValid();
 }
-#endif
+#endif//NIAGARA_SYSTEM_CAPTURE
 
 void FNiagaraSystemInstance::SetSolo(bool bInSolo)
 {
diff --git a/Engine/Plugins/FX/Niagara/Source/Niagara/Private/NiagaraSystemSimulation.cpp b/Engine/Plugins/FX/Niagara/Source/Niagara/Private/NiagaraSystemSimulation.cpp
--- a/Engine/Plugins/FX/Niagara/Source/Niagara/Private/NiagaraSystemSimulation.cpp
+++ b/Engine/Plugins/FX/Niagara/Source/Niagara/Private/NiagaraSystemSimulation.cpp
@@ -598,7 +598,7 @@ FNiagaraSystemSimulationTickContext::FNiagaraSystemSimulationTickContext(class F
 
 	bRunningAsync = bAllowAsync && GNiagaraSystemSimulationAllowASync && FApp::ShouldUseThreadingForPerformance();
 
-#if WITH_EDITORONLY_DATA
+#if NIAGARA_SYSTEM_CAPTURE
 	if ( Owner->GetIsSolo() && Instances.Num() == 1 )
 	{
 		bRunningAsync &= Instances[0]->ShouldCaptureThisFrame() == false;
@@ -1778,7 +1778,7 @@ void FNiagaraSystemSimulation::Tick_Concurrent(FNiagaraSystemSimulationTickConte
 			}
 		}
 
-	#if WITH_EDITORONLY_DATA
+	#if NIAGARA_SYSTEM_CAPTURE
 		if (SoloSystemInstance)
 		{
 			SoloSystemInstance->FinishCapture();
@@ -1916,7 +1916,7 @@ void FNiagaraSystemSimulation::SpawnSystemInstances(FNiagaraSystemSimulationTick
 
 	Context.DataSet.EndSimulate();
 
-#if WITH_EDITORONLY_DATA
+#if NIAGARA_SYSTEM_CAPTURE
 	if (SoloSystemInstance && SoloSystemInstance->ShouldCaptureThisFrame())
 	{
 		TSharedPtr<struct FNiagaraScriptDebuggerInfo, ESPMode::ThreadSafe> DebugInfo = SoloSystemInstance->GetActiveCaptureWrite(NAME_None, ENiagaraScriptUsage::SystemSpawnScript, FGuid());
@@ -2014,7 +2014,7 @@ void FNiagaraSystemSimulation::UpdateSystemInstances(FNiagaraSystemSimulationTic
 
 		Context.DataSet.EndSimulate();
 
-#if WITH_EDITORONLY_DATA
+#if NIAGARA_SYSTEM_CAPTURE
 		if (SoloSystemInstance && SoloSystemInstance->ShouldCaptureThisFrame())
 		{
 			TSharedPtr<struct FNiagaraScriptDebuggerInfo, ESPMode::ThreadSafe> DebugInfo = SoloSystemInstance->GetActiveCaptureWrite(NAME_None, ENiagaraScriptUsage::SystemUpdateScript, FGuid());
diff --git a/Engine/Plugins/FX/Niagara/Source/Niagara/Private/NiagaraTypes.cpp b/Engine/Plugins/FX/Niagara/Source/Niagara/Private/NiagaraTypes.cpp
--- a/Engine/Plugins/FX/Niagara/Source/Niagara/Private/NiagaraTypes.cpp
+++ b/Engine/Plugins/FX/Niagara/Source/Niagara/Private/NiagaraTypes.cpp
@@ -279,6 +279,11 @@ void FNiagaraLwcStructConverter::AddConversionStep(int32 InSourceBytes, int32 In
 
 //////////////////////////////////////////////////////////////////////////
 
+FNiagaraTypeDefinition FNiagaraTypeHelper::Vector2DDef;
+FNiagaraTypeDefinition FNiagaraTypeHelper::VectorDef;
+FNiagaraTypeDefinition FNiagaraTypeHelper::Vector4Def;
+FNiagaraTypeDefinition FNiagaraTypeHelper::QuatDef;
+
 FRWLock FNiagaraTypeHelper::RemapTableLock;
 TMap<TWeakObjectPtr<UScriptStruct>, FNiagaraTypeHelper::FRemapEntry> FNiagaraTypeHelper::RemapTable;
 std::atomic<bool> FNiagaraTypeHelper::RemapTableDirty;
@@ -634,6 +639,11 @@ UScriptStruct* FNiagaraTypeHelper::GetLWCStruct(UScriptStruct* SWCStruct)
 
 FNiagaraTypeDefinition FNiagaraTypeHelper::GetSWCType(const FNiagaraTypeDefinition& InType)
 {
+	if (InType.IsEnum() || InType.IsUObject() || InType.IsDataInterface())
+	{
+		return InType;
+	}
+
 	if (IsLWCType(InType))
 	{
 		return FNiagaraTypeDefinition(FindNiagaraFriendlyTopLevelStruct(CastChecked<UScriptStruct>(InType.GetStruct()), ENiagaraStructConversion::Simulation));
@@ -641,6 +651,79 @@ FNiagaraTypeDefinition FNiagaraTypeHelper::GetSWCType(const FNiagaraTypeDefiniti
 	return InType;
 }
 
+FNiagaraTypeDefinition FNiagaraTypeHelper::GetLWCType(const FNiagaraTypeDefinition& InType)
+{
+	if(InType.IsEnum() || InType.IsUObject() || InType.IsDataInterface())
+	{
+		return InType;
+	}
+
+	if(InType == FNiagaraTypeDefinition::GetFloatDef())
+	{
+		return FNiagaraTypeDefinition(FNiagaraDouble::StaticStruct(), FNiagaraTypeDefinition::EAllowUnfriendlyStruct::Allow);
+	}
+	else if (InType == FNiagaraTypeDefinition::GetVec2Def())
+	{
+		return Vector2DDef;
+	}
+	else if (InType == FNiagaraTypeDefinition::GetVec3Def())
+	{
+		return VectorDef;
+	}
+	else if (InType == FNiagaraTypeDefinition::GetVec4Def())
+	{
+		return Vector4Def;
+	}
+	else if (InType == FNiagaraTypeDefinition::GetQuatDef())
+	{
+		return QuatDef;
+	}
+
+	if (IsLWCType(InType))
+	{
+		return InType;
+	}
+
+	if(UScriptStruct* LWCStruct = GetLWCStruct(InType.GetScriptStruct()))
+	{
+		return FNiagaraTypeDefinition(LWCStruct, FNiagaraTypeDefinition::EAllowUnfriendlyStruct::Allow);
+	}
+	
+	return InType;
+}
+
+void FNiagaraTypeHelper::InitStaticTypes()
+{
+	if (Vector2DDef.IsValid() == false)
+	{
+		UPackage* CoreUObjectPkg = FindObjectChecked<UPackage>(nullptr, TEXT("/Script/CoreUObject"));
+		Vector2DDef = FNiagaraTypeDefinition(FindObjectChecked<UScriptStruct>(CoreUObjectPkg, TEXT("Vector2D")), FNiagaraTypeDefinition::EAllowUnfriendlyStruct::Allow);
+		VectorDef = FNiagaraTypeDefinition(FindObjectChecked<UScriptStruct>(CoreUObjectPkg, TEXT("Vector")), FNiagaraTypeDefinition::EAllowUnfriendlyStruct::Allow);
+		Vector4Def = FNiagaraTypeDefinition(FindObjectChecked<UScriptStruct>(CoreUObjectPkg, TEXT("Vector4")), FNiagaraTypeDefinition::EAllowUnfriendlyStruct::Allow);
+		QuatDef = FNiagaraTypeDefinition(FindObjectChecked<UScriptStruct>(CoreUObjectPkg, TEXT("Quat")), FNiagaraTypeDefinition::EAllowUnfriendlyStruct::Allow);
+	}
+}
+
+FNiagaraTypeDefinition FNiagaraTypeHelper::GetVector2DDef()
+{
+	return Vector2DDef;
+}
+
+FNiagaraTypeDefinition FNiagaraTypeHelper::GetVectorDef()
+{
+	return VectorDef;
+}
+
+FNiagaraTypeDefinition FNiagaraTypeHelper::GetVector4Def()
+{
+	return Vector4Def;
+}
+
+FNiagaraTypeDefinition FNiagaraTypeHelper::GetQuatDef()
+{
+	return QuatDef;
+}
+
 UScriptStruct* FNiagaraTypeHelper::FindNiagaraFriendlyTopLevelStruct(UScriptStruct* InStruct, ENiagaraStructConversion StructConversion)
 {
 	if (!InStruct)
diff --git a/Engine/Plugins/FX/Niagara/Source/Niagara/Private/NiagaraWorldManager.cpp b/Engine/Plugins/FX/Niagara/Source/Niagara/Private/NiagaraWorldManager.cpp
--- a/Engine/Plugins/FX/Niagara/Source/Niagara/Private/NiagaraWorldManager.cpp
+++ b/Engine/Plugins/FX/Niagara/Source/Niagara/Private/NiagaraWorldManager.cpp
@@ -234,6 +234,7 @@ FDelegateHandle FNiagaraWorldManager::OnPostWorldCleanupHandle;
 FDelegateHandle FNiagaraWorldManager::OnPreWorldFinishDestroyHandle;
 FDelegateHandle FNiagaraWorldManager::OnWorldBeginTearDownHandle;
 FDelegateHandle FNiagaraWorldManager::TickWorldHandle;
+FDelegateHandle FNiagaraWorldManager::OnWorldPreActorTickHandle;
 FDelegateHandle FNiagaraWorldManager::OnWorldPreSendAllEndOfFrameUpdatesHandle;
 FDelegateHandle FNiagaraWorldManager::PreGCHandle;
 FDelegateHandle FNiagaraWorldManager::PostReachabilityAnalysisHandle;
@@ -366,6 +367,7 @@ void FNiagaraWorldManager::OnStartup()
 	OnPreWorldFinishDestroyHandle = FWorldDelegates::OnPreWorldFinishDestroy.AddStatic(&FNiagaraWorldManager::OnPreWorldFinishDestroy);
 	OnWorldBeginTearDownHandle = FWorldDelegates::OnWorldBeginTearDown.AddStatic(&FNiagaraWorldManager::OnWorldBeginTearDown);
 	TickWorldHandle = FWorldDelegates::OnWorldPostActorTick.AddStatic(&FNiagaraWorldManager::TickWorld);
+	OnWorldPreActorTickHandle = FWorldDelegates::OnWorldPreActorTick.AddStatic(&FNiagaraWorldManager::OnWorldPreActorTick);
 	OnWorldPreSendAllEndOfFrameUpdatesHandle = FWorldDelegates::OnWorldPreSendAllEndOfFrameUpdates.AddLambda(
 		[](UWorld* InWorld)
 		{
@@ -838,6 +840,11 @@ void FNiagaraWorldManager::TickWorld(UWorld* World, ELevelTick TickType, float D
 	Get(World)->PostActorTick(DeltaSeconds);
 }
 
+void FNiagaraWorldManager::OnWorldPreActorTick(UWorld* InWorld, ELevelTick InLevelTick, float InDeltaSeconds)
+{
+	Get(InWorld)->PreActorTick(InLevelTick, InDeltaSeconds);
+}
+
 void FNiagaraWorldManager::OnPreGarbageCollect()
 {
 	for (TPair<UWorld*, FNiagaraWorldManager*>& Pair : WorldManagers)
@@ -878,6 +885,11 @@ void FNiagaraWorldManager::OnRefreshOwnerAllowsScalability()
 	}
 }
 
+void FNiagaraWorldManager::PreActorTick(ELevelTick InLevelTick, float InDeltaSeconds)
+{	
+	DataChannelManager->BeginFrame(InDeltaSeconds);
+}
+
 void FNiagaraWorldManager::PostActorTick(float DeltaSeconds)
 {
 	CSV_SCOPED_TIMING_STAT_EXCLUSIVE(Effects);
@@ -1012,6 +1024,8 @@ void FNiagaraWorldManager::PostActorTick(float DeltaSeconds)
 		}
 	}
 #endif 
+
+	DataChannelManager->EndFrame(DeltaSeconds);
 }
 
 void FNiagaraWorldManager::PreSendAllEndOfFrameUpdates()
diff --git a/Engine/Plugins/FX/Niagara/Source/Niagara/Public/NiagaraCommon.h b/Engine/Plugins/FX/Niagara/Source/Niagara/Public/NiagaraCommon.h
--- a/Engine/Plugins/FX/Niagara/Source/Niagara/Public/NiagaraCommon.h
+++ b/Engine/Plugins/FX/Niagara/Source/Niagara/Public/NiagaraCommon.h
@@ -50,6 +50,8 @@ using FNiagaraDataBufferRef = TRefCountPtr<FNiagaraDataBuffer>;
 
 #define NIAGARA_MEMORY_TRACKING	!UE_BUILD_SHIPPING
 
+#define NIAGARA_SYSTEM_CAPTURE (!UE_BUILD_SHIPPING)
+
 #ifndef NIAGARA_COMPUTEDEBUG_ENABLED
 	#define NIAGARA_COMPUTEDEBUG_ENABLED WITH_EDITOR
 #endif
@@ -1706,3 +1708,40 @@ namespace ENiagaraGpuComputeTickStage
 		Last = PostOpaqueRender UMETA(hidden),
 	};
 }
+
+enum class ENiagaraResourceAccess : uint8
+{
+	ReadOnly,
+	WriteOnly,
+	ReadWrite,
+};
+
+UENUM()
+enum class ENiagaraConditionalOperator
+{
+	Equals,
+	NotEqual,
+	LessThan,
+	LessThanOrEqual,
+	GreaterThan,
+	GreaterThanOrEqual,
+
+	Max UMETA(Hidden),
+};
+
+template<typename T>
+bool EvalConditional(ENiagaraConditionalOperator Op, const T& A, const T& B)
+{
+	switch (Op)
+	{
+	case ENiagaraConditionalOperator::Equals: return A == B;
+	case ENiagaraConditionalOperator::NotEqual: return A != B;
+	case ENiagaraConditionalOperator::LessThan: return A < B;
+	case ENiagaraConditionalOperator::LessThanOrEqual: return A <= B;
+	case ENiagaraConditionalOperator::GreaterThan: return A > B;
+	case ENiagaraConditionalOperator::GreaterThanOrEqual: return A >= B;
+	default:check(0);
+	};
+
+	return false;
+};
\ No newline at end of file
diff --git a/Engine/Plugins/FX/Niagara/Source/Niagara/Public/NiagaraDataChannel.h b/Engine/Plugins/FX/Niagara/Source/Niagara/Public/NiagaraDataChannel.h
--- a/Engine/Plugins/FX/Niagara/Source/Niagara/Public/NiagaraDataChannel.h
+++ b/Engine/Plugins/FX/Niagara/Source/Niagara/Public/NiagaraDataChannel.h
@@ -38,154 +38,48 @@ When GPU->GPU is supported, the frame dependency issue have more meaning for GPU
 
 #pragma once
 
-#include "NiagaraDataChannelCommon.h"
+#include "NiagaraDataChannelPublic.h"
 #include "Kismet/BlueprintFunctionLibrary.h"
+#include "NiagaraDataSetCompiledData.h"
+#include "NiagaraCommon.h"
+#include "NiagaraDataSet.h"
+#include "UObject/UObjectIterator.h"
 #include "NiagaraDataChannel.generated.h"
 
 DECLARE_STATS_GROUP(TEXT("Niagara Data Channels"), STATGROUP_NiagaraDataChannels, STATCAT_Advanced);
 
-class UNiagaraDataChannelHandler;
-
-/** Buffer containing a single FNiagaraVariable's data at the game level. AoS layout. LWC Types. */
-struct FNiagaraDataChannelVariableBuffer
-{
-	TArray<uint8> Data;
-	int32 Size = 0;
-	
-	void Init(const FNiagaraVariableBase& Var)
-	{
-		Size = Var.GetSizeInBytes();
-	}
-
-	void Empty(){ Data.Empty(); }
-	void Reset(){ Data.Reset(); }
-	
-	template<typename T>
-	bool Write(int32 Index, const T& InData)
-	{
-		check(sizeof(T) == Size);
-		if (Index >= 0 && Index < Num())
-		{
-			T* Dest = ((T*)Data.GetData()) + Index;
-			*Dest = InData;
-
-			return true;
-		}
-		return false;
-	}
-
-	template<typename T>
-	bool Read(int32 Index, T& OutData)const
-	{
-		check(sizeof(T) == Size);
-		if (Index >= 0 && Index < Num())
-		{
-			T* Src = ((T*)Data.GetData()) + Index;
-			OutData = *Src;
-
-			return true;
-		}
-		return false;
-	}
-
-	void SetNum(int32 Num)
-	{
-		Data.SetNumZeroed(Size * Num);
-	}
-
-	int32 Num()const { return Data.Num() / Size; }
-};
-
-USTRUCT()
-struct FNiagaraDataChannelGameDataLayout
-{
-	GENERATED_BODY();
-
-	/** Map of all variables contained in this DataChannel data and the indices into data arrays for game data. */
-	UPROPERTY()
-	TMap<FNiagaraVariableBase, int32> VariableIndices;
-
-	/** Helpers for converting LWC types into Niagara simulation SWC types. */
-	UPROPERTY()
-	TArray<FNiagaraLwcStructConverter> LwcConverters;
+class FRHICommandListImmediate;
 
-	void Init(const TArray<FNiagaraVariable>& Variables);
-};
-
-/** Storage for game level DataChannels generated by BP / C++ */
-struct FNiagaraDataChannelGameData : public TSharedFromThis<FNiagaraDataChannelGameData>
-{
-private:
-	
-	/** Per variable storage buffers. */
-	TArray<FNiagaraDataChannelVariableBuffer> VariableData;
+class FNiagaraGpuComputeDispatchInterface;
+struct FNiagaraDataSetCompiledData;
 
-	int32 NumElements = 0;
-
-	TWeakObjectPtr<const UNiagaraDataChannel> DataChannel;
-
-public:
-
-	void Init(const UNiagaraDataChannel* InDataChannel);
-	void Empty();
-	void Reset();
-	void SetNum(int32 NewNum);
-	int32 Num() const { return NumElements; }
-
-	FNiagaraDataChannelVariableBuffer* FindVariableBuffer(const FNiagaraVariableBase& Var);
-
-	void WriteToDataSet(FNiagaraDataBuffer* DestBuffer, int32 DestStartIdx);
-
-	void AppendFromGameData(const FNiagaraDataChannelGameData& GameData);
-
-	void AppendFromDataSet(const FNiagaraDataBuffer* SrcBuffer, FVector3f LwcTile);
-
-	const UNiagaraDataChannel* GetDataChannel()const { return DataChannel.Get(); }
-
-	const TConstArrayView<FNiagaraDataChannelVariableBuffer> GetVariableBuffers()const { return VariableData; }
-};
+class UNiagaraDataChannelHandler;
+class UNiagaraDataChannelWriter;
+class UNiagaraDataChannelReader;
+struct FNiagaraDataChannelPublishRequest;
+struct FNiagaraDataChannelGameDataLayout;
 
 //////////////////////////////////////////////////////////////////////////
 
-/** Storage for DataChannels the world level. */
-struct NIAGARA_API FNiagaraWorldDataChannelStore
+/** Render thread proxy of FNiagaraDataChannelData. */
+struct FNiagaraDataChannelDataProxy
 {
-
-	/** DataChannel data accessible from Game/BP. AoS Layout. LWC types. */
-	FNiagaraDataChannelGameDataPtr GameData;
-
-	//		▲	CPU Sim Data can optionally be made visible to the Game Data.
-	//		|
-	//		|
-	//		|
-	//		▼	Game/BP Data can optionally be made visible to CPU sims.
-
-	/** DataChannel data accessible to Niagara CPU sims. SoA layout. Non LWC types. */
-	FNiagaraDataSet* CPUSimData = nullptr;
-
-	// Cached off buffer with the previous frame's CPU Sim accessible data.
-	// Some systems can choose to read this to avoid any current frame tick ordering issues.
-	FNiagaraDataBufferRef PrevCPUSimData = nullptr;
-	
-	//		▲	GPU Sim Data can optionally be made visible to CPU Sims and Game Data.
-	//		|
-	//		|
-	//		|
-	//		▼	CPU Sim and Game Data can optionally be made visible to the GPU.
-
-	/** DataChannel data accessible to Niagara GPU sims. SoA layout. Non LWC types. */
-	FNiagaraDataSet* GPUSimData = nullptr;
-	
-	~FNiagaraWorldDataChannelStore();
-
-	bool IsInitialized()const {return CPUSimData != nullptr; }
-	void Empty();
-	void Init(UNiagaraDataChannelHandler* Owner);
-	void Tick(UNiagaraDataChannelHandler* Owner);
+	FNiagaraDataSet* GPUDataSet = nullptr;
+	FNiagaraDataBufferRef PrevFrameData = nullptr;
+
+	#if !UE_BUILD_SHIPPING
+	FString DebugName;
+	const TCHAR* GetDebugName()const{return *DebugName;}
+	#else
+	const TCHAR* GetDebugName()const{return nullptr;}
+	#endif
+
+	void BeginFrame(bool bKeepPreviousFrameData);
+	void EndFrame(FNiagaraGpuComputeDispatchInterface* DispathInterface, FRHICommandListImmediate& CmdList, const TArray<FNiagaraDataBufferRef>& BuffersForGPU);
+	void Reset();
 };
 
-
-DECLARE_MULTICAST_DELEGATE_OneParam(FOnDataChannelChanged, const UNiagaraDataChannel*);
+DECLARE_MULTICAST_DELEGATE_OneParam(FOnDataChannelCreated, const UNiagaraDataChannel*);
 
 UCLASS(Experimental, abstract, EditInlineNew)
 class NIAGARA_API UNiagaraDataChannel : public UObject
@@ -194,6 +88,7 @@ public:
 	GENERATED_BODY()
 
 	//UObject Interface
+	virtual void PostInitProperties()override;
 	virtual void PostLoad();
 	virtual void BeginDestroy();
 #if WITH_EDITOR
@@ -202,7 +97,7 @@ public:
 #endif
 	//UObject Interface End.
 
-	FName GetChannelName()const { return ChannelName; }
+	const UNiagaraDataChannelAsset* GetAsset()const { return CastChecked<UNiagaraDataChannelAsset>(GetOuter()); }
 	TConstArrayView<FNiagaraVariable> GetVariables()const { return Variables; }
 
 	/** If true, we keep our previous frame's data. Some users will prefer a frame of latency to tick dependency. */
@@ -212,11 +107,7 @@ public:
 	const FNiagaraDataSetCompiledData& GetCompiledData(ENiagaraSimTarget SimTarget)const;
 
 	/** Create the appropriate handler object for this data channel. */
-	virtual UNiagaraDataChannelHandler* CreateHandler() const PURE_VIRTUAL(UNiagaraDataChannel::CreateHandler, {return nullptr;} );
-
-#if WITH_EDITOR
-	FOnDataChannelChanged& OnChanged() { return OnChangedDelegate; }
-#endif
+	virtual UNiagaraDataChannelHandler* CreateHandler(UWorld* OwningWorld) const PURE_VIRTUAL(UNiagaraDataChannel::CreateHandler, {return nullptr;} );
 	
 	const FNiagaraDataChannelGameDataLayout& GetGameDataLayout()const { return GameDataLayout; }
 
@@ -224,17 +115,23 @@ public:
 
 	bool IsValid()const;
 
+	#if !UE_BUILD_SHIPPING
+	void SetVerboseLogging(bool bValue){ bVerboseLogging = bValue; }
+	bool GetVerboseLogging()const { return bVerboseLogging; }
+	#endif
+
+	template<typename TFunc>
+	static void ForEachDataChannel(TFunc Func);
+
 private:
-	UPROPERTY(EditAnywhere, Category = "Data Channel")
-	FName ChannelName;
 
 	/** The variables that define the data contained in this Data Channel. */
 	UPROPERTY(EditAnywhere, Category = "Data Channel")
 	TArray<FNiagaraVariable> Variables;
 
-	/** If true, we keep our previous frame's data. Some users will prefer a frame of latency to tick dependency. */
+	/** If true, we keep our previous frame's data. This comes at a memory and performance cost but allows users to avoid tick order dependency by reading last frame's data. Some users will prefer a frame of latency to tick order dependency. */
 	UPROPERTY(EditAnywhere, Category = "Data Channel")
-	bool bKeepPreviousFrameData = false;
+	bool bKeepPreviousFrameData = true;
 		
 	/**
 	Data layout for payloads in Niagara datasets.
@@ -245,15 +142,26 @@ private:
 	UPROPERTY()
 	mutable FNiagaraDataSetCompiledData CompiledDataGPU;
 
-#if WITH_EDITOR
-	FOnDataChannelChanged OnChangedDelegate;
-#endif
-
 	/** Layout information for any data stored at the "Game" level. i.e. From game code/BP. AoS layout and LWC types. */
-	UPROPERTY()
 	FNiagaraDataChannelGameDataLayout GameDataLayout;
+	
+	#if !UE_BUILD_SHIPPING
+	mutable bool bVerboseLogging = false;
+	#endif
 };
 
+template<typename TAction>
+void UNiagaraDataChannel::ForEachDataChannel(TAction Func)
+{
+	for(TObjectIterator<UNiagaraDataChannel> It; It; ++It)
+	{
+		if(*It)
+		{
+			Func(*It);
+		}
+	}
+}
+
 /**
 * A C++ and Blueprint accessible library of utility functions for accessing Niagara DataChannel
 */
@@ -263,5 +171,13 @@ class NIAGARA_API UNiagaraDataChannelLibrary : public UBlueprintFunctionLibrary
 	GENERATED_UCLASS_BODY()
 
 	UFUNCTION(BlueprintCallable, Category = NiagaraDataChannel, meta = (Keywords = "niagara DataChannel", WorldContext = "WorldContextObject", UnsafeDuringActorConstruction = "true"))
-	static UNiagaraDataChannelHandler* GetNiagaraDataChannel(const UObject* WorldContextObject, FName Channel);
+	static UNiagaraDataChannelHandler* GetNiagaraDataChannel(const UObject* WorldContextObject, UNiagaraDataChannelAsset* Channel);
+
+	/** Initializes and returns the Niagara Data Channel writer to write N elements to the given data channel. */
+	UFUNCTION(BlueprintCallable, Category = NiagaraDataChannel, meta = (Keywords = "niagara DataChannel", WorldContext = "WorldContextObject", UnsafeDuringActorConstruction = "true"))
+	static UNiagaraDataChannelWriter* WriteToNiagaraDataChannel(const UObject* WorldContextObject, UNiagaraDataChannelAsset* Channel, FNiagaraDataChannelSearchParameters SearchParams, int32 Count, bool bVisibleToGame, bool bVisibleToCPU, bool bVisibleToGPU);
+
+	/** Initializes and returns the Niagara Data Channel reader for the given data channel. */
+	UFUNCTION(BlueprintCallable, Category = NiagaraDataChannel, meta = (Keywords = "niagara DataChannel", WorldContext = "WorldContextObject", UnsafeDuringActorConstruction = "true"))
+	static UNiagaraDataChannelReader* ReadFromNiagaraDataChannel(const UObject* WorldContextObject, UNiagaraDataChannelAsset* Channel, FNiagaraDataChannelSearchParameters SearchParams, bool bReadPreviousFrame);
 };
diff --git a/Engine/Plugins/FX/Niagara/Source/Niagara/Public/NiagaraDataChannelAccessor.h b/Engine/Plugins/FX/Niagara/Source/Niagara/Public/NiagaraDataChannelAccessor.h
new file mode 100644
--- /dev/null
+++ b/Engine/Plugins/FX/Niagara/Source/Niagara/Public/NiagaraDataChannelAccessor.h
@@ -0,0 +1,124 @@
+// Copyright Epic Games, Inc. All Rights Reserved.
+
+#pragma once
+
+#include "NiagaraDataChannelPublic.h"
+#include "NiagaraDataChannel.h"
+#include "NiagaraDataChannelAccessor.generated.h"
+
+class UNiagaraDataChannelHandler;
+struct FNiagaraSpawnInfo;
+struct FNiagaraVariableBase;
+
+/** 
+Initial simple API for reading and writing data in a data channel from game code / BP. 
+Likely to be replaced in the near future with a custom BP node and a helper struct.
+*/
+
+
+UCLASS(Experimental, BlueprintType)
+class NIAGARA_API UNiagaraDataChannelReader : public UObject
+{
+	GENERATED_BODY()
+private:
+
+	FNiagaraDataChannelDataPtr Data = nullptr;
+	bool bReadingPreviousFrame = false;
+
+	template<typename T>
+	bool ReadData(const FNiagaraVariableBase& Var, int32 Index, T& OutData)const;
+
+public:
+
+	UPROPERTY()
+	TObjectPtr<UNiagaraDataChannelHandler> Owner;
+	
+	/** Call before each access to the data channel to grab the correct data to read. */
+	UFUNCTION(BlueprintCallable, Category = NiagaraDataChannel, meta = (Keywords = "Niagara DataChannel"))
+	bool InitAccess(FNiagaraDataChannelSearchParameters SearchParams, bool bReadPrevFrameData);
+
+	UFUNCTION(BlueprintCallable, Category = NiagaraDataChannel, meta = (Keywords = "Niagara DataChannel"))
+	int32 Num()const;
+
+	UFUNCTION(BlueprintCallable, Category = NiagaraDataChannel, meta = (Keywords = "Niagara DataChannel"))
+	double ReadFloat(FName VarName, int32 Index)const;
+
+	UFUNCTION(BlueprintCallable, Category = NiagaraDataChannel, meta = (Keywords = "Niagara DataChannel"))
+	FVector2D ReadVector2D(FName VarName, int32 Index)const;
+
+	UFUNCTION(BlueprintCallable, Category = NiagaraDataChannel, meta = (Keywords = "Niagara DataChannel"))
+	FVector ReadVector(FName VarName, int32 Index)const;
+
+	UFUNCTION(BlueprintCallable, Category = NiagaraDataChannel, meta = (Keywords = "Niagara DataChannel"))
+	FVector4 ReadVector4(FName VarName, int32 Index)const;
+
+	UFUNCTION(BlueprintCallable, Category = NiagaraDataChannel, meta = (Keywords = "Niagara DataChannel"))
+	FQuat ReadQuat(FName VarName, int32 Index)const;
+
+	UFUNCTION(BlueprintCallable, Category = NiagaraDataChannel, meta = (Keywords = "Niagara DataChannel"))
+	FLinearColor ReadLinearColor(FName VarName, int32 Index)const;
+
+	UFUNCTION(BlueprintCallable, Category = NiagaraDataChannel, meta = (Keywords = "Niagara DataChannel"))
+	int32 ReadInt(FName VarName, int32 Index)const;
+
+	UFUNCTION(BlueprintCallable, Category = NiagaraDataChannel, meta = (Keywords = "Niagara DataChannel"))
+	bool ReadBool(FName VarName, int32 Index)const;
+
+	UFUNCTION(BlueprintCallable, Category = NiagaraDataChannel, meta = (Keywords = "Niagara DataChannel"))
+	FVector ReadPosition(FName VarName, int32 Index)const;
+};
+
+UCLASS(Experimental, BlueprintType)
+class NIAGARA_API UNiagaraDataChannelWriter : public UObject
+{
+	GENERATED_BODY()
+private:
+
+	/** Local data buffers we're writing into. */
+	FNiagaraDataChannelGameDataPtr Data = nullptr;
+
+	template<typename T>
+	void WriteData(const FNiagaraVariableBase& Var, int32 Index, const T& InData);
+
+public:
+
+	UPROPERTY()
+	TObjectPtr<UNiagaraDataChannelHandler> Owner;
+	
+	/** Call before each batch of writes to allocate the data we'll be writing to. */
+	UFUNCTION(BlueprintCallable, Category = NiagaraDataChannel, meta = (Keywords = "niagara DataChannel"))
+	bool InitWrite(FNiagaraDataChannelSearchParameters SearchParams, int32 Count, bool bVisibleToGame=true, bool bVisibleToCPU=true, bool bVisibleToGPU=true);
+
+	UFUNCTION(BlueprintCallable, Category = NiagaraDataChannel, meta = (Keywords = "niagara DataChannel"))
+	int32 Num()const;
+
+	UFUNCTION(BlueprintCallable, Category = NiagaraDataChannel, meta = (Keywords = "niagara DataChannel"))
+	void WriteFloat(FName VarName, int32 Index, double InData);
+	
+	UFUNCTION(BlueprintCallable, Category = NiagaraDataChannel, meta = (Keywords = "niagara DataChannel"))
+	void WriteVector2D(FName VarName, int32 Index, FVector2D InData);
+
+	UFUNCTION(BlueprintCallable, Category = NiagaraDataChannel, meta = (Keywords = "niagara DataChannel"))
+	void WriteVector(FName VarName, int32 Index, FVector InData);
+
+	UFUNCTION(BlueprintCallable, Category = NiagaraDataChannel, meta = (Keywords = "niagara DataChannel"))
+	void WriteVector4(FName VarName, int32 Index, FVector4 InData);
+
+	UFUNCTION(BlueprintCallable, Category = NiagaraDataChannel, meta = (Keywords = "niagara DataChannel"))
+	void WriteQuat(FName VarName, int32 Index, FQuat InData);
+
+	UFUNCTION(BlueprintCallable, Category = NiagaraDataChannel, meta = (Keywords = "niagara DataChannel"))
+	void WriteLinearColor(FName VarName, int32 Index, FLinearColor InData);
+	
+	UFUNCTION(BlueprintCallable, Category = NiagaraDataChannel, meta = (Keywords = "niagara DataChannel"))
+	void WriteInt(FName VarName, int32 Index, int32 InData);
+
+	UFUNCTION(BlueprintCallable, Category = NiagaraDataChannel, meta = (Keywords = "niagara DataChannel"))
+	void WriteBool(FName VarName, int32 Index, bool InData);
+
+	UFUNCTION(BlueprintCallable, Category = NiagaraDataChannel, meta = (Keywords = "niagara DataChannel"))
+	void WriteSpawnInfo(FName VarName, int32 Index, FNiagaraSpawnInfo InData);
+
+	UFUNCTION(BlueprintCallable, Category = NiagaraDataChannel, meta = (Keywords = "niagara DataChannel"))
+	void WritePosition(FName VarName, int32 Index, FVector InData);
+};
\ No newline at end of file
diff --git a/Engine/Plugins/FX/Niagara/Source/Niagara/Public/NiagaraDataChannelHandler.h b/Engine/Plugins/FX/Niagara/Source/Niagara/Public/NiagaraDataChannelHandler.h
--- a/Engine/Plugins/FX/Niagara/Source/Niagara/Public/NiagaraDataChannelHandler.h
+++ b/Engine/Plugins/FX/Niagara/Source/Niagara/Public/NiagaraDataChannelHandler.h
@@ -11,16 +11,9 @@ Some more complex handlers may want to divide up the scene in various different
 
 */
 
-
-#include "NiagaraDataChannelCommon.h"
+#include "NiagaraDataChannelPublic.h"
 #include "NiagaraDataChannelHandler.generated.h"
 
-/** Base class for all UNiagaraDataChannelHandler's Render Thread Proxies. */
-struct FNiagaraDataChannelHandlerRTProxyBase
-{
-	virtual ~FNiagaraDataChannelHandlerRTProxyBase() = default;
-};
-
 UCLASS(Experimental, abstract, BlueprintType)
 class NIAGARA_API UNiagaraDataChannelHandler : public UObject
 {
@@ -32,37 +25,17 @@ public:
 
 	virtual void Init(const UNiagaraDataChannel* InChannel);
 
-	/** Returns the DataChannel data for the given DataChannel name and location. */
+	virtual void BeginFrame(float DeltaTime, FNiagaraWorldManager* OwningWorld);
+
+	virtual void EndFrame(float DeltaTime, FNiagaraWorldManager* OwningWorld);
+
 	virtual void Tick(float DeltaTime, ETickingGroup TickGroup, FNiagaraWorldManager* OwningWorld);
 
-	/** Returns the data for this data channel given the system instance. Possibly some handlers will subdivide spacially etc so may not return the same data for the different systems. */
-	virtual FNiagaraDataBufferRef GetData(FNiagaraSystemInstance* SystemInstance, bool bGetLastFrameData)  PURE_VIRTUAL(UNiagaraDataChannelHandler::GetData, return nullptr;);	
-
-	/** Gets the GPU data set with data for the given system. This API will likely be reworked as GPU support is fleshed out and more complex handlers are created. */
-	virtual FNiagaraDataSet* GetDataGPU(FNiagaraSystemInstance* SystemInstance) PURE_VIRTUAL(UNiagaraDataChannelHandler::GetDataGPU, return nullptr;);	
-
-	/** Adds a request to publish some data into the channel on the next tick. */
-	virtual void Publish(const FNiagaraDataChannelPublishRequest& Request);
-		
-	/**
-	 *Removes all publish requests involving the given dataset.
-	 *TODO: REMOVE
-	 *This is a hack to get around lifetime issues wrt data buffers/datasets and their compiled data.
-	 *We should rework things such that data buffers can exist beyond their owning dataset, including the compiled data detailing their layout.
-	 **/
-	virtual void RemovePublishRequests(const FNiagaraDataSet* DataSet);
-
-	/** Returns the game level DataChannel data store. */
-	virtual FNiagaraDataChannelGameDataPtr GetGameData(UActorComponent* OwningComponent) PURE_VIRTUAL(UNiagaraDataChannelHandler::GetGameData, return nullptr; );
-
-	/** Utility allowing easy casting of proxy to it's proper derived type. */
-	template<typename T>
-	T* GetRTProxyAs() 
-	{
-		T* TypedProxy = static_cast<T*>(RTProxy.Get());
-		check(TypedProxy != nullptr);
-		return TypedProxy;
-	}
+	/** 
+	Finds the correct internal data for this data channel and the given search parameters. For example in some cases this may require a search of several elements of data that correspond to different spacial regions.
+	This shared ptr provides access to Game level data, CPU simulation data and a render thread proxy that can be given to the RT and provides access to GPUSimulaiton data.
+	*/
+	virtual FNiagaraDataChannelDataPtr FindData(FNiagaraDataChannelSearchParameters SearchParams, ENiagaraResourceAccess AccessType)  PURE_VIRTUAL(UNiagaraDataChannelHandler::FindData, return nullptr;);
 
 	const UNiagaraDataChannel* GetDataChannel()const { return DataChannel; }
 
@@ -72,20 +45,16 @@ public:
 	UFUNCTION(BlueprintCallable, Category = "Data Channel")
 	UNiagaraDataChannelReader* GetDataChannelReader();
 
+	template<typename T> 
+	T* GetChannelTyped()const{ return CastChecked<T>(DataChannel); }
+
+	FNiagaraDataChannelDataPtr CreateData();
+
 protected:
 
 	UPROPERTY()
 	TObjectPtr<const UNiagaraDataChannel> DataChannel;
 
-	/** Per frame requests to publish data from Niagara Systems into the main world DataChannels buffers. */
-	TArray<FNiagaraDataChannelPublishRequest> PublishRequests;
-
-	/** Data buffers we'll be passing to the RT for uploading to the GPU */
-	TArray<FNiagaraDataBufferRef> BuffersForGPU;
-
-	/** Render Thread Proxy for this handler. Each handler type must create it's proxy inside it's constructor. */
-	TUniquePtr<FNiagaraDataChannelHandlerRTProxyBase> RTProxy;
-
 	/** Helper object allowing BP to write data in this channel. */
 	UPROPERTY()
 	TObjectPtr<UNiagaraDataChannelWriter> Writer;
diff --git a/Engine/Plugins/FX/Niagara/Source/Niagara/Public/NiagaraDataChannelPublic.h b/Engine/Plugins/FX/Niagara/Source/Niagara/Public/NiagaraDataChannelPublic.h
new file mode 100644
--- /dev/null
+++ b/Engine/Plugins/FX/Niagara/Source/Niagara/Public/NiagaraDataChannelPublic.h
@@ -0,0 +1,93 @@
+// Copyright Epic Games, Inc. All Rights Reserved.
+
+#pragma once
+
+#include "Engine/EngineBaseTypes.h"
+#include "NiagaraCommon.h"
+#include "NiagaraDataChannelPublic.generated.h"
+
+class FNiagaraWorldManager;
+class UNiagaraDataChannelAsset;
+class UNiagaraDataChannel;
+class UNiagaraDataChannelHandler;
+struct FNiagaraDataChannelGameData;
+struct FNiagaraDataChannelData;
+using FNiagaraDataChannelGameDataPtr = TSharedPtr<FNiagaraDataChannelGameData>;
+using FNiagaraDataChannelDataPtr = TSharedPtr<FNiagaraDataChannelData>;
+
+/** Wrapper asset class for UNiagaraDataChannel which is instanced. */
+UCLASS(Experimental, DisplayName = "Niagara Data Channel")
+class NIAGARA_API UNiagaraDataChannelAsset : public UObject
+{
+	GENERATED_BODY()
+
+	UPROPERTY(EditAnywhere, Category = DataChannel, Instanced)
+	TObjectPtr<UNiagaraDataChannel> DataChannel;
+
+public:
+
+	//UObject Interface
+	virtual FPrimaryAssetId GetPrimaryAssetId() const override;
+	//END UObject Interface
+
+	UNiagaraDataChannel* Get() const { return DataChannel; }
+};
+
+/**
+Minimal set of types and declares required for external users of Niagara Data Channels.
+*/
+
+/**
+Parameters allowing users to search for the correct data channel data to read/write.
+Some data channels will sub divide their data internally in various ways, e.g., spacial partition.
+These parameters allow users to search for the correct internal data when reading and writing.
+*/
+USTRUCT(BlueprintType)
+struct NIAGARA_API FNiagaraDataChannelSearchParameters
+{
+	GENERATED_BODY()
+
+	/** In cases where there is an owning component such as an object spawning from itself etc, then we pass that component in. Some handlers may only use it's location but others may make use of more data. */
+	UPROPERTY(EditAnywhere, BlueprintReadWrite, Category = "Parameters")
+	TObjectPtr<USceneComponent> OwningComponent = nullptr;
+
+	/** In cases where there is no owning component for data being read or written to a data channel, we simply pass in a location. */
+	UPROPERTY(EditAnywhere, BlueprintReadWrite, Category = "Parameters")
+	FVector Location = FVector::ZeroVector;
+
+	FVector GetLocation()const;
+};
+
+USTRUCT()
+struct FNiagaraDataChannelGameDataLayout
+{
+	GENERATED_BODY();
+
+	/** Map of all variables contained in this DataChannel data and the indices into data arrays for game data. */
+	UPROPERTY()
+	TMap<FNiagaraVariableBase, int32> VariableIndices;
+
+	/** Helpers for converting LWC types into Niagara simulation SWC types. */
+	UPROPERTY()
+	TArray<FNiagaraLwcStructConverter> LwcConverters;
+
+	void Init(const TArray<FNiagaraVariable>& Variables);
+};
+
+
+#if !UE_BUILD_SHIPPING
+
+/** Hooks into internal NiagaraDataChannels code for debugging and testing purposes. */
+class NIAGARA_API FNiagaraDataChannelDebugUtilities
+{
+public: 
+
+	static void BeginFrame(FNiagaraWorldManager* WorldMan, float DeltaSeconds);
+	static void EndFrame(FNiagaraWorldManager* WorldMan, float DeltaSeconds);
+	static void Tick(FNiagaraWorldManager* WorldMan, float DeltaSeconds, ETickingGroup TickGroup);
+	
+	static UNiagaraDataChannelHandler* FindDataChannelHandler(FNiagaraWorldManager* WorldMan, UNiagaraDataChannel* DataChannel);
+};
+
+
+#endif
\ No newline at end of file
diff --git a/Engine/Plugins/FX/Niagara/Source/Niagara/Public/NiagaraDataChannel_Global.h b/Engine/Plugins/FX/Niagara/Source/Niagara/Public/NiagaraDataChannel_Global.h
new file mode 100644
--- /dev/null
+++ b/Engine/Plugins/FX/Niagara/Source/Niagara/Public/NiagaraDataChannel_Global.h
@@ -0,0 +1,53 @@
+// Copyright Epic Games, Inc. All Rights Reserved.
+
+#pragma once
+
+#include "NiagaraDataChannel.h"
+#include "NiagaraDataChannelHandler.h"
+#include "NiagaraDataChannel_Global.generated.h"
+
+
+UCLASS(Experimental)
+class UNiagaraDataChannel_Global : public UNiagaraDataChannel
+{
+	GENERATED_BODY()
+
+	virtual UNiagaraDataChannelHandler* CreateHandler(UWorld* OwningWorld)const override;
+};
+
+/**
+Basic DataChannel handler that makes all data visible globally.
+*/
+UCLASS(Experimental, BlueprintType)
+class UNiagaraDataChannelHandler_Global : public UNiagaraDataChannelHandler
+{
+	GENERATED_UCLASS_BODY()
+
+	FNiagaraDataChannelDataPtr Data;
+
+	//UObject Interface
+	virtual void BeginDestroy()override;
+	//UObject Interface End
+
+	virtual void Init(const UNiagaraDataChannel* InChannel) override;
+	virtual void BeginFrame(float DeltaTime, FNiagaraWorldManager* OwningWorld)override;
+	virtual void EndFrame(float DeltaTime, FNiagaraWorldManager* OwningWorld)override;
+	virtual void Tick(float DeltaTime, ETickingGroup TickGroup, FNiagaraWorldManager* OwningWorld) override;
+	virtual FNiagaraDataChannelDataPtr FindData(FNiagaraDataChannelSearchParameters SearchParams, ENiagaraResourceAccess AccessType) override;
+};
+
+/**
+DataChannel handler that separates DataChannel out by the hash of their position into spatially localized grid cells.
+Useful for DataChannel that should only be visible to nearby consumers.
+TODO: Likely unsuitable for LWC games as we need a defined min/max for the hash and storage. A sparse octree etc would be needed
+*/
+// UCLASS(EditInlineNew)
+// class UNiagaraDataChannelChannelHandler_WorldGrid : public UObject
+// {
+// };
+
+// UCLASS(EditInlineNew, abstract)
+// class UNiagaraDataChannelChannelHandler_Octree : UObject
+// {
+// 
+// };
diff --git a/Engine/Plugins/FX/Niagara/Source/Niagara/Public/NiagaraDataInterfaceUtilities.h b/Engine/Plugins/FX/Niagara/Source/Niagara/Public/NiagaraDataInterfaceUtilities.h
--- a/Engine/Plugins/FX/Niagara/Source/Niagara/Public/NiagaraDataInterfaceUtilities.h
+++ b/Engine/Plugins/FX/Niagara/Source/Niagara/Public/NiagaraDataInterfaceUtilities.h
@@ -53,6 +53,7 @@ namespace FNiagaraDataInterfaceUtilities
 /** Helper context object helping to facilitate data interfaces building their hlsl shader code for GPU simulations. */
 struct FNiagaraDataInterfaceHlslGenerationContext
 {
+	DECLARE_DELEGATE_RetVal_OneParam(FString, FGetFunctionSignatureSymbol, const FNiagaraFunctionSignature& /*Signature*/)
 	DECLARE_DELEGATE_RetVal_OneParam(FString, FGetStructHlslTypeName, const FNiagaraTypeDefinition& /*Type*/)
 	DECLARE_DELEGATE_RetVal_OneParam(FString, FGetPropertyHlslTypeName, const FProperty* /*Property*/)
 	DECLARE_DELEGATE_RetVal_TwoParams(FString, FGetSanitizedSymbolName, FStringView /*SymbolName*/, bool /*bCollapsNamespaces*/)
@@ -66,12 +67,14 @@ struct FNiagaraDataInterfaceHlslGenerationContext
 	TArrayView<const FNiagaraFunctionSignature> Signatures;
 	int32 FunctionInstanceIndex = INDEX_NONE;
 
+	FGetFunctionSignatureSymbol GetFunctionSignatureSymbolDelegate;
 	FGetStructHlslTypeName GetStructHlslTypeNameDelegate;
 	FGetPropertyHlslTypeName GetPropertyHlslTypeNameDelegate;
 	FGetSanitizedSymbolName GetSanitizedSymbolNameDelegate;
 
 	const FNiagaraDataInterfaceGeneratedFunction& GetFunctionInfo() const { return ParameterInfo.GeneratedFunctions[FunctionInstanceIndex]; }
 
+	FString GetFunctionSignatureSymbol(const FNiagaraFunctionSignature& Signature) { return GetFunctionSignatureSymbolDelegate.Execute(Signature); }
 	FString GetStructHlslTypeName(const FNiagaraTypeDefinition& Type) { return GetStructHlslTypeNameDelegate.Execute(Type); }
 	FString GetPropertyHlslTypeName(const FProperty* Property) { return GetPropertyHlslTypeNameDelegate.Execute(Property); }
 	FString GetSanitizedSymbolName(FStringView SymbolName, bool bCollapsNamespaces = false) { return GetSanitizedSymbolNameDelegate.Execute(SymbolName, bCollapsNamespaces); }
diff --git a/Engine/Plugins/FX/Niagara/Source/Niagara/Public/NiagaraSystemInstance.h b/Engine/Plugins/FX/Niagara/Source/Niagara/Public/NiagaraSystemInstance.h
--- a/Engine/Plugins/FX/Niagara/Source/Niagara/Public/NiagaraSystemInstance.h
+++ b/Engine/Plugins/FX/Niagara/Source/Niagara/Public/NiagaraSystemInstance.h
@@ -352,7 +352,7 @@ public:
 	*/
 	void AdvanceSimulation(int32 TickCountToSimulate, float TickDeltaSeconds);
 
-#if WITH_EDITORONLY_DATA
+#if NIAGARA_SYSTEM_CAPTURE
 	/** Request that this simulation capture a frame. Cannot capture if disabled or already completed.*/
 	bool RequestCapture(const FGuid& RequestId);
 
@@ -515,7 +515,7 @@ private:
 	FOnDestroyed OnDestroyedDelegate;
 #endif
 
-#if WITH_EDITORONLY_DATA
+#if NIAGARA_SYSTEM_CAPTURE
 	TSharedPtr<TArray<TSharedPtr<struct FNiagaraScriptDebuggerInfo, ESPMode::ThreadSafe>>, ESPMode::ThreadSafe> CurrentCapture;
 	TSharedPtr<FGuid, ESPMode::ThreadSafe> CurrentCaptureGuid;
 	bool bWasSoloPriorToCaptureRequest;
diff --git a/Engine/Plugins/FX/Niagara/Source/Niagara/Public/NiagaraTypes.h b/Engine/Plugins/FX/Niagara/Source/Niagara/Public/NiagaraTypes.h
--- a/Engine/Plugins/FX/Niagara/Source/Niagara/Public/NiagaraTypes.h
+++ b/Engine/Plugins/FX/Niagara/Source/Niagara/Public/NiagaraTypes.h
@@ -415,6 +415,19 @@ public:
 	static FNiagaraTypeDefinition GetSWCType(const FNiagaraTypeDefinition& InType);
 	static void TickTypeRemap();
 
+	/**
+	Attempt to get the actual LWC type for a given type that would be used in the rest of the engine etc.
+	Used for cases where we actually need to store types in LWC format.
+	These cannot be properly serialized currently so must be constructed at runtime from their Niagara Simulation types.
+	*/
+	static FNiagaraTypeDefinition GetLWCType(const FNiagaraTypeDefinition& InType);
+	
+	static FNiagaraTypeDefinition GetVector2DDef();
+	static FNiagaraTypeDefinition GetVectorDef();
+	static FNiagaraTypeDefinition GetVector4Def();
+	static FNiagaraTypeDefinition GetQuatDef();
+
+	static void InitStaticTypes();
 private:
 	struct FRemapEntry
 	{
@@ -436,6 +449,12 @@ private:
 	static FRWLock RemapTableLock;
 	static TMap<TWeakObjectPtr<UScriptStruct>, FRemapEntry> RemapTable;
 	static std::atomic<bool> RemapTableDirty;
+
+	/** LWC Counterparts to Niagara's base simulation types. */
+	static FNiagaraTypeDefinition Vector2DDef;
+	static FNiagaraTypeDefinition VectorDef;
+	static FNiagaraTypeDefinition Vector4Def;
+	static FNiagaraTypeDefinition QuatDef;
 };
 
 /** Information about how this type should be laid out in an FNiagaraDataSet */
@@ -541,10 +560,14 @@ If this changes, all scripts must be recompiled by bumping the NiagaraCustomVers
 UENUM()
 enum class ENiagaraExecutionStateSource : uint32
 {
-	Scalability, //State set by Scalability logic. Lowest precedence.
-	Internal, //Misc internal state. For example becoming inactive after we finish our set loops.
-	Owner, //State requested by the owner. Takes precedence over everything but internal completion logic.
-	InternalCompletion, // Internal completion logic. Has to take highest precedence for completion to be ensured.
+	/** State set by Scalability logic. Lowest precedence. */
+	Scalability,
+	/** Misc internal state. For example becoming inactive after we finish our set loops. */
+	Internal, 
+	/** State requested by the owner. Takes precedence over everything but internal completion logic. */
+	Owner, 
+	/** Internal completion logic. Has to take highest precedence for completion to be ensured. */
+	InternalCompletion,
 };
 
 UENUM()
diff --git a/Engine/Plugins/FX/Niagara/Source/Niagara/Public/NiagaraWorldManager.h b/Engine/Plugins/FX/Niagara/Source/Niagara/Public/NiagaraWorldManager.h
--- a/Engine/Plugins/FX/Niagara/Source/Niagara/Public/NiagaraWorldManager.h
+++ b/Engine/Plugins/FX/Niagara/Source/Niagara/Public/NiagaraWorldManager.h
@@ -118,6 +118,8 @@ public:
 
 	void Tick(ETickingGroup TickGroup, float DeltaSeconds, ELevelTick TickType, ENamedThreads::Type CurrentThread, const FGraphEventRef& MyCompletionGraphEvent);
 
+	void PreActorTick(ELevelTick InLevelTick, float InDeltaSeconds);
+
 	/** Called after all actor tick groups are complete. */
 	void PostActorTick(float DeltaSeconds);
 
@@ -239,6 +241,8 @@ private:
 	// Called when the world begins to be torn down for example by level streaming.
 	static void OnWorldBeginTearDown(UWorld* World);
 
+	static void OnWorldPreActorTick(UWorld* InWorld, ELevelTick InLevelTick, float InDeltaSeconds);
+
 	// Callback for when a world is ticked.
 	static void TickWorld(UWorld* World, ELevelTick TickType, float DeltaSeconds);
 
@@ -277,6 +281,7 @@ private:
 	static FDelegateHandle OnPreWorldFinishDestroyHandle;
 	static FDelegateHandle OnWorldBeginTearDownHandle;
 	static FDelegateHandle TickWorldHandle;
+	static FDelegateHandle OnWorldPreActorTickHandle;
 	static FDelegateHandle OnWorldPreSendAllEndOfFrameUpdatesHandle;
 	static FDelegateHandle PreGCHandle;
 	static FDelegateHandle PostReachabilityAnalysisHandle;
diff --git a/Engine/Plugins/FX/Niagara/Source/NiagaraEditor/Private/Customizations/NiagaraTypeCustomizations.cpp b/Engine/Plugins/FX/Niagara/Source/NiagaraEditor/Private/Customizations/NiagaraTypeCustomizations.cpp
--- a/Engine/Plugins/FX/Niagara/Source/NiagaraEditor/Private/Customizations/NiagaraTypeCustomizations.cpp
+++ b/Engine/Plugins/FX/Niagara/Source/NiagaraEditor/Private/Customizations/NiagaraTypeCustomizations.cpp
@@ -2455,39 +2455,25 @@ void SNiagaraRawVariableTypeSelectMenu::CollectAllActions(FGraphActionListBuilde
 {
 	FNiagaraMenuActionCollector Collector;
 
-	static TArray<FNiagaraTypeDefinition> Types;
-	if(Types.IsEmpty())
-	{
-		//TODO: Build better list of possible user facing types.
-		UPackage* CorePkg = FindObjectChecked<UPackage>(nullptr, TEXT("/Script/CoreUObject"));
-		Types.Emplace(FNiagaraTypeDefinition::GetBoolDef());
-		Types.Emplace(FNiagaraTypeDefinition::GetIntDef());
-		Types.Emplace(FNiagaraTypeDefinition(FNiagaraDouble::StaticStruct(), FNiagaraTypeDefinition::EAllowUnfriendlyStruct::Allow));
-
-		Types.Emplace(FNiagaraTypeDefinition(FindObjectChecked<UScriptStruct>(CorePkg, TEXT("Vector2d")), FNiagaraTypeDefinition::EAllowUnfriendlyStruct::Allow));
-		Types.Emplace(FNiagaraTypeDefinition(FindObjectChecked<UScriptStruct>(CorePkg, TEXT("Vector")), FNiagaraTypeDefinition::EAllowUnfriendlyStruct::Allow));
-		Types.Emplace(FNiagaraTypeDefinition(FindObjectChecked<UScriptStruct>(CorePkg, TEXT("Vector4")), FNiagaraTypeDefinition::EAllowUnfriendlyStruct::Allow));
-		Types.Emplace(FNiagaraTypeDefinition(FindObjectChecked<UScriptStruct>(CorePkg, TEXT("Quat")), FNiagaraTypeDefinition::EAllowUnfriendlyStruct::Allow));
-
-		Types.Emplace(FNiagaraTypeDefinition::GetColorDef());
-		Types.Emplace(FNiagaraTypeDefinition::GetIDDef());
-		Types.Emplace(FNiagaraTypeDefinition(FNiagaraSpawnInfo::StaticStruct()));
-	};
-
-	Types.Sort([](const FNiagaraTypeDefinition& A, const FNiagaraTypeDefinition& B) { return (A.GetNameText().ToLower().ToString() < B.GetNameText().ToLower().ToString()); });
-
-	for (const FNiagaraTypeDefinition& RegisteredType : Types)
+	TArray<FNiagaraTypeDefinition> Types;
+	FNiagaraEditorUtilities::GetAllowedUserVariableTypes(Types);
+	for (const FNiagaraTypeDefinition& Type : Types)
 	{
-		FNiagaraVariable DefaultVar(RegisteredType, FName(*RegisteredType.GetName()));
+		if(Type.IsDataInterface() || Type.GetClass())
+		{
+			continue;
+		}
+
+		FNiagaraVariable DefaultVar(Type, FName(*Type.GetName()));
 		FNiagaraEditorUtilities::ResetVariableToDefaultValue(DefaultVar);
 
 		FText Category = FNiagaraEditorUtilities::GetVariableTypeCategory(DefaultVar);
-		const FText DisplayName = RegisteredType.GetNameText();
-		const FText Tooltip = RegisteredType.GetNameText();
+		const FText DisplayName = Type.GetNameText();
+		const FText Tooltip = Type.GetNameText();
 		TSharedPtr<FNiagaraMenuAction> Action(new FNiagaraMenuAction(
 			Category, DisplayName, Tooltip, 0, FText::GetEmpty(),
 			FNiagaraMenuAction::FOnExecuteStackAction::CreateLambda(
-			[Var=VarToModify, PropHandle=PropertyHandle, RegisteredType]()
+			[Var=VarToModify, PropHandle=PropertyHandle, Type]()
 			{
 				FScopedTransaction Transaction(LOCTEXT("Set Raw Niagara Variable Type", "Set Variable Type"));
 				check(Var);
@@ -2500,7 +2486,7 @@ void SNiagaraRawVariableTypeSelectMenu::CollectAllActions(FGraphActionListBuilde
 				}
 
 				PropHandle->NotifyPreChange();
-				Var->SetType(RegisteredType);
+				Var->SetType(Type);
 				PropHandle->NotifyPostChange(EPropertyChangeType::ValueSet);
 				PropHandle->NotifyFinishedChangingProperties();
 			}
@@ -2511,80 +2497,162 @@ void SNiagaraRawVariableTypeSelectMenu::CollectAllActions(FGraphActionListBuilde
 
 	Collector.AddAllActionsTo(OutAllActions);
 }
+// Copy of private SConstrainedBox defined in PropertyEditor/Private/SConstrainedBox.h
+class SNiagaraConstrainedBox : public SCompoundWidget
+{
+public:
+	SLATE_BEGIN_ARGS(SNiagaraConstrainedBox) {}
+	SLATE_DEFAULT_SLOT(FArguments, Content)
+		SLATE_ATTRIBUTE(TOptional<float>, MinWidth)
+		SLATE_ATTRIBUTE(TOptional<float>, MaxWidth)
+		SLATE_END_ARGS()
+
+		void Construct(const FArguments& InArgs)
+	{
+		MinWidth = InArgs._MinWidth;
+		MaxWidth = InArgs._MaxWidth;
+
+		ChildSlot
+			[
+				InArgs._Content.Widget
+			];
+	}
+
+	virtual FVector2D ComputeDesiredSize(float LayoutScaleMultiplier) const override
+	{
+		const float MinWidthVal = MinWidth.Get().Get(0.0f);
+		const float MaxWidthVal = MaxWidth.Get().Get(0.0f);
+
+		if (MinWidthVal == 0.0f && MaxWidthVal == 0.0f)
+		{
+			return SCompoundWidget::ComputeDesiredSize(LayoutScaleMultiplier);
+		}
+		else
+		{
+			FVector2D ChildSize = ChildSlot.GetWidget()->GetDesiredSize();
+
+			float XVal = FMath::Max(MinWidthVal, ChildSize.X);
+			if (MaxWidthVal > MinWidthVal)
+			{
+				XVal = FMath::Min(MaxWidthVal, XVal);
+			}
+
+			return FVector2D(XVal, ChildSize.Y);
+		}
+	}
+
+private:
+	TAttribute< TOptional<float> > MinWidth;
+	TAttribute< TOptional<float> > MaxWidth;
+};
 
 void FNiagaraVariableDetailsCustomization::CustomizeHeader(TSharedRef<IPropertyHandle> PropertyHandle, FDetailWidgetRow& HeaderRow, IPropertyTypeCustomizationUtils& CustomizationUtils)
+{
+	
+}
+
+void FNiagaraVariableDetailsCustomization::CustomizeChildren(TSharedRef<IPropertyHandle> PropertyHandle, IDetailChildrenBuilder& ChildBuilder, IPropertyTypeCustomizationUtils& CustomizationUtils)
 {
 	TSharedPtr<IPropertyHandle> NameHandle = PropertyHandle->GetChildHandle(TEXT("Name"));
 	TSharedPtr<IPropertyHandle> TypeDefHandleHandle = PropertyHandle->GetChildHandle(TEXT("TypeDefHandle"));
-	
+
+	FDetailWidgetRow& Row = ChildBuilder.AddCustomRow(FText::GetEmpty());
+
 	TArray<UObject*> Objects;
 	PropertyHandle->GetOuterObjects(Objects);
 	if (Objects.Num() > 1)
 	{
-		HeaderRow.NameContent()
-		[
-			PropertyHandle->CreatePropertyNameWidget()
-		];
-		HeaderRow.ValueContent()
-		[
-			PropertyHandle->CreatePropertyValueWidget()
-		];
+		Row.NameContent()
+			[
+				PropertyHandle->CreatePropertyNameWidget()
+			];
+		Row.ValueContent()
+			[
+				PropertyHandle->CreatePropertyValueWidget()
+			];
+
 	}
 	else
-	{		
+	{
 		void* VarPtr = nullptr;
 		if (PropertyHandle->GetValueData(VarPtr) == FPropertyAccess::Success)
-		{		
+		{
 			FNiagaraVariable* Var = (FNiagaraVariable*)VarPtr;
 
-			HeaderRow.NameContent()
-				.MaxDesiredWidth(200.f)
+			Row.NameContent()
 				[
+					//PropertyHandle->CreatePropertyNameWidget()
 					SNew(SHorizontalBox)
 					+ SHorizontalBox::Slot()
-					.FillWidth(0.9f)
-					.Padding(4.0f, 0.0f, 4.0f, 0.0f)
+					.AutoWidth()
+					.Padding(4.0f, 4.0f, 4.0f, 4.0f)
 					[
-						SNew(STextBlock)
-						.Text_Lambda([Var] {	return Var ? Var->GetType().GetNameText() : FText::GetEmpty(); })
+						SNew(SComboButton)
+						.ButtonStyle(FAppStyle::Get(), "RoundButton")
+						.ForegroundColor(FAppStyle::GetSlateColor("DefaultForeground"))
+						.ContentPadding(FMargin(0, 0))
+						.OnGetMenuContent(this, &FNiagaraVariableDetailsCustomization::GetTypeMenu, TypeDefHandleHandle, Var)
+						//.IsEnabled(this, &SNiagaraParameterPanel::GetCanAddParametersToCategory, Category)
+						.HAlign(HAlign_Center)
+						.VAlign(VAlign_Center)
+						.HasDownArrow(false)
+						.ButtonContent()
+						[
+							SNew(SHorizontalBox)
+							+ SHorizontalBox::Slot()
+							.AutoWidth()
+							.Padding(FMargin(0, 1))
+							[
+								SNew(SImage)
+								.Image(FAppStyle::GetBrush("Icons.Edit"))
+							]
+						]
 					]
-				+ SHorizontalBox::Slot()
-				.AutoWidth()
-				.Padding(4.0f, 0.0f, 4.0f, 0.0f)
-				[
-					SNew(SComboButton)
-					.ButtonStyle(FAppStyle::Get(), "RoundButton")
-					.ForegroundColor(FAppStyle::GetSlateColor("DefaultForeground"))
-					.ContentPadding(FMargin(2, 0))
-					.OnGetMenuContent(this, &FNiagaraVariableDetailsCustomization::GetTypeMenu, TypeDefHandleHandle, Var)
-					//.IsEnabled(this, &SNiagaraParameterPanel::GetCanAddParametersToCategory, Category)
+					+ SHorizontalBox::Slot()
+					.AutoWidth()
 					.HAlign(HAlign_Center)
 					.VAlign(VAlign_Center)
-					.HasDownArrow(false)
-					.ButtonContent()
+					.Padding(4.0f, 4.0f, 4.0f, 4.0f)
 					[
-						SNew(SHorizontalBox)
-						+ SHorizontalBox::Slot()
-						.AutoWidth()
-						.Padding(FMargin(0, 1))
+						SNew(SNiagaraConstrainedBox)
+						.MinWidth(75.0f)
 						[
-							SNew(SImage)
-							.Image(FAppStyle::GetBrush("Icons.Edit"))
+							SNew(STextBlock)
+							.Text_Lambda([Var] {return Var ? Var->GetType().GetNameText() : FText::GetEmpty(); })
 						]
 					]
+				];
+
+			Row.ValueContent()
+				[
+					SNew(SHorizontalBox)
+			+ SHorizontalBox::Slot()
+				.AutoWidth()
+				.HAlign(HAlign_Center)
+				.VAlign(VAlign_Center)
+				.Padding(4.0f, 4.0f, 4.0f, 4.0f)
+				[
+					SNew(SNiagaraConstrainedBox)
+					.MinWidth(150.0f)
+				[
+					NameHandle->CreatePropertyValueWidget()
+				]
+				]
+			+ SHorizontalBox::Slot()
+				.AutoWidth()
+				.HAlign(HAlign_Center)
+				.VAlign(VAlign_Center)
+				.Padding(4.0f, 4.0f, 4.0f, 4.0f)
+				[
+					SNew(SNiagaraConstrainedBox)
+					.MinWidth(150.0f)
+				[
+					PropertyHandle->CreateDefaultPropertyButtonWidgets()
+				]
 				]
 			];
-		
-			HeaderRow.ValueContent()
-			[
-				NameHandle->CreatePropertyValueWidget()
-			];					
 		}
-	}	
-}
-
-void FNiagaraVariableDetailsCustomization::CustomizeChildren(TSharedRef<IPropertyHandle> PropertyHandle, IDetailChildrenBuilder& ChildBuilder, IPropertyTypeCustomizationUtils& CustomizationUtils)
-{
+	}
 }
 
 TSharedRef<SWidget> FNiagaraVariableDetailsCustomization::GetTypeMenu(TSharedPtr<IPropertyHandle> InPropertyHandle, FNiagaraVariable* Var)
diff --git a/Engine/Plugins/FX/Niagara/Source/NiagaraEditor/Private/NiagaraActions.cpp b/Engine/Plugins/FX/Niagara/Source/NiagaraEditor/Private/NiagaraActions.cpp
--- a/Engine/Plugins/FX/Niagara/Source/NiagaraEditor/Private/NiagaraActions.cpp
+++ b/Engine/Plugins/FX/Niagara/Source/NiagaraEditor/Private/NiagaraActions.cpp
@@ -23,7 +23,6 @@
 #include "Widgets/SToolTip.h"
 
 #include "NiagaraDataChannel.h"
-#include "NiagaraDataChannelDefinitions.h"
 
 #include UE_INLINE_GENERATED_CPP_BY_NAME(NiagaraActions)
 
@@ -681,6 +680,20 @@ void INiagaraDataInterfaceNodeActionProvider::GetNodeContextMenuActions(UClass*
 	}
 }
 
+void INiagaraDataInterfaceNodeActionProvider::CollectAddPinActions(UClass* DIClass, FNiagaraMenuActionCollector& Collector, UEdGraphPin* AddPin)
+{
+	if (ensure(DIClass))
+	{
+		while (DIClass && DIClass != UObject::StaticClass())
+		{
+			if (TUniquePtr<INiagaraDataInterfaceNodeActionProvider>* Provider = RegisteredActionProviders.Find(DIClass->GetFName()))
+			{
+				(*Provider)->CollectAddPinActionsImpl(Collector, AddPin);
+			}
+			DIClass = DIClass->GetSuperClass();
+		}
+	}
+}
 ////////////////////////
 
 namespace NiagaraActionsLocal
@@ -701,45 +714,51 @@ void FNiagaraDataInterfaceNodeActionProvider_DataChannelWrite::GetNodeContextMen
 		return;
 	}
 
-	const TArray<UNiagaraDataChannelDefinitions*>& ChannelDefs = UNiagaraDataChannelDefinitions::GetDataChannelDefinitions(false, true);
-	if (ChannelDefs.Num() > 0)
-	{
-		return;
-	}
-
 	FToolMenuSection& Section = Menu->AddSection("DataChannelWrite", InitForDataChannelHeaderText);
 
-	TWeakObjectPtr<UNiagaraNodeFunctionCall> WeakFuncNode = CastChecked<UNiagaraNodeFunctionCall>(Context->Node);
+	UNiagaraNodeFunctionCall* FuncNode = CastChecked<UNiagaraNodeFunctionCall>(Context->Node);
+	TWeakObjectPtr<UNiagaraNodeFunctionCall> WeakNode = FuncNode;
 
-	auto CreateNodeContextMenu = [WeakFuncNode](UToolMenu* InNewToolMenu)
+	auto CreateNodeContextMenu = [WeakNode](UToolMenu* InNewToolMenu)
 	{
 		FToolMenuSection& SubSection = InNewToolMenu->AddSection("InitForDataChannelSection", InitForDataChannelSectionText);
 
-		const TArray<UNiagaraDataChannelDefinitions*>& DataChannelDefs = UNiagaraDataChannelDefinitions::GetDataChannelDefinitions(false, true);
-		for (const UNiagaraDataChannelDefinitions* Def : DataChannelDefs)
+		auto InitForDataChannelSection = [&SubSection, WeakNode](UNiagaraDataChannel* DataChannel)
 		{
-			for (TObjectPtr<const UNiagaraDataChannel> Channel : Def->DataChannels)
+			if (DataChannel->HasAnyFlags(RF_ClassDefaultObject) == false)
 			{
-				TWeakObjectPtr<const UNiagaraDataChannel> WeakChannel = Channel;
-				auto CreateDataChannelActionEntry = [WeakFuncNode, WeakChannel]()
+				TWeakObjectPtr<UNiagaraDataChannel> WeakChannel = DataChannel;
+				auto CreateDataChannelActionEntry = [WeakChannel, WeakNode]()
 				{
-					if (UNiagaraNodeFunctionCall* Node = WeakFuncNode.Get())
+					UNiagaraDataChannel* Channel = WeakChannel.Get();
+					UNiagaraNodeFunctionCall* Node = WeakNode.Get();
+					if(Channel && Node)
 					{
-						if (const UNiagaraDataChannel* ChannelPtr = WeakChannel.Get())
+						Node->RemoveAllDynamicPins();
+						TConstArrayView<FNiagaraVariable> ChannelVars = Channel->GetVariables();
+						for (const FNiagaraVariable& Var : ChannelVars)
 						{
-							TConstArrayView<FNiagaraVariable> ChannelVars = ChannelPtr->GetVariables();
-							for (const FNiagaraVariable& Var : ChannelVars)
+							FNiagaraTypeDefinition Type = Var.GetType();
+							if (Type.IsEnum() == false)
 							{
-								FNiagaraVariable SWCVar(FNiagaraTypeHelper::GetSWCStruct(Var.GetType().GetScriptStruct()), Var.GetName());
-								Node->AddParameter(SWCVar, Node->Pins.Last());
+								Type = FNiagaraTypeDefinition(FNiagaraTypeHelper::GetSWCStruct(Var.GetType().GetScriptStruct()));
 							}
+							FNiagaraVariable SWCVar(Type, Var.GetName());
+							Node->AddParameter(SWCVar, EEdGraphPinDirection::EGPD_Input);
 						}
 					}
 				};
-				SubSection.AddMenuEntry(NAME_None, FText::FromName(Channel->GetChannelName()), FText::FromName(Channel->GetChannelName()), FSlateIcon(FAppStyle::GetAppStyleSetName(), "Icons.Import"),
+				SubSection.AddMenuEntry(
+					NAME_None,
+					FText::FromString(DataChannel->GetAsset()->GetName()),
+					FText::FromString(DataChannel->GetAsset()->GetName()),
+					FSlateIcon(FAppStyle::GetAppStyleSetName(), "Icons.Import"),
 					FUIAction(FExecuteAction::CreateLambda(CreateDataChannelActionEntry), FCanExecuteAction()));
+
 			}
-		}
+		};
+
+		UNiagaraDataChannel::ForEachDataChannel(InitForDataChannelSection);
 	};
 
 	Section.AddSubMenu("InitForDataChannelMenu", InitForDataChannelMenuText, InitForDataChannelMenuTooltipText,	FNewToolMenuDelegate::CreateLambda(CreateNodeContextMenu));
@@ -756,52 +775,106 @@ void FNiagaraDataInterfaceNodeActionProvider_DataChannelRead::GetNodeContextMenu
 		return;
 	}
 
-	const TArray<UNiagaraDataChannelDefinitions*>& ChannelDefs = UNiagaraDataChannelDefinitions::GetDataChannelDefinitions(false, true);
-	if (ChannelDefs.Num() > 0)
-	{
-		return;
-	}
-
 	FToolMenuSection& Section = Menu->AddSection("DataChannelRead", InitForDataChannelHeaderText);
 
-	TWeakObjectPtr<UNiagaraNodeFunctionCall> WeakFuncNode = CastChecked<UNiagaraNodeFunctionCall>(Context->Node);
+	UNiagaraNodeFunctionCall* FuncNode = CastChecked<UNiagaraNodeFunctionCall>(Context->Node);
+	TWeakObjectPtr<UNiagaraNodeFunctionCall> WeakNode = FuncNode;
 
-	auto CreateNodeContextMenu = [WeakFuncNode](UToolMenu* InNewToolMenu)
+	auto CreateNodeContextMenu = [WeakNode](UToolMenu* InNewToolMenu)
 	{
 		FToolMenuSection& SubSection = InNewToolMenu->AddSection("InitForDataChannelSection", InitForDataChannelSectionText);
 
-		const TArray<UNiagaraDataChannelDefinitions*>& DataChannelDefs = UNiagaraDataChannelDefinitions::GetDataChannelDefinitions(false, true);
-		for (const UNiagaraDataChannelDefinitions* Def : DataChannelDefs)
+		auto InitForDataChannelSection = [&SubSection, WeakNode](UNiagaraDataChannel* DataChannel)
 		{
-			for (TObjectPtr<const UNiagaraDataChannel> Channel : Def->DataChannels)
+			if (DataChannel->HasAnyFlags(RF_ClassDefaultObject) == false)
 			{
-				TWeakObjectPtr<const UNiagaraDataChannel> WeakChannel = Channel;
-				auto CreateDataChannelActionEntry = [WeakFuncNode, WeakChannel]()
+				TWeakObjectPtr<UNiagaraDataChannel> WeakChannel = DataChannel;
+				auto CreateDataChannelActionEntry = [WeakChannel, WeakNode]()
 				{
-					if (UNiagaraNodeFunctionCall* Node = WeakFuncNode.Get())
+					UNiagaraDataChannel* Channel = WeakChannel.Get();
+					UNiagaraNodeFunctionCall* Node = WeakNode.Get();
+
+					if (Channel && Node)
 					{
 						Node->RemoveAllDynamicPins();
-						if (const UNiagaraDataChannel* ChannelPtr = WeakChannel.Get())
+						TConstArrayView<FNiagaraVariable> ChannelVars = Channel->GetVariables();
+						for (const FNiagaraVariable& Var : ChannelVars)
 						{
-							TConstArrayView<FNiagaraVariable> ChannelVars = ChannelPtr->GetVariables();
-							for (const FNiagaraVariable& Var : ChannelVars)
+							FNiagaraTypeDefinition Type = Var.GetType();
+							if (Type.IsEnum() == false)
 							{
-								FNiagaraVariable SWCVar(FNiagaraTypeHelper::GetSWCStruct(Var.GetType().GetScriptStruct()), Var.GetName());
-								Node->AddParameter(SWCVar, Node->Pins.Last());
+								Type = FNiagaraTypeDefinition(FNiagaraTypeHelper::GetSWCStruct(Var.GetType().GetScriptStruct()));
 							}
+							FNiagaraVariable SWCVar(Type, Var.GetName());
+							Node->AddParameter(SWCVar, EEdGraphPinDirection::EGPD_Output);
 						}
 					}
 				};
 
-				SubSection.AddMenuEntry(NAME_None, FText::FromName(Channel->GetChannelName()), FText::FromName(Channel->GetChannelName()), FSlateIcon(FAppStyle::GetAppStyleSetName(), "Icons.Import"), 
+				SubSection.AddMenuEntry(NAME_None, FText::FromString(DataChannel->GetAsset()->GetName()), FText::FromString(DataChannel->GetAsset()->GetName()), FSlateIcon(FAppStyle::GetAppStyleSetName(), "Icons.Import"),
 					FUIAction(FExecuteAction::CreateLambda(CreateDataChannelActionEntry), FCanExecuteAction()));
 			}
-		}
+		};
+
+		UNiagaraDataChannel::ForEachDataChannel(InitForDataChannelSection);
 	};
 
 	Section.AddSubMenu("InitForDataChannelMenu", InitForDataChannelMenuText, InitForDataChannelMenuTooltipText,	FNewToolMenuDelegate::CreateLambda(CreateNodeContextMenu));
 }
 
+void FNiagaraDataInterfaceNodeActionProvider_DataChannelRead::CollectAddPinActionsImpl(FNiagaraMenuActionCollector& Collector, UEdGraphPin* AddPin)const
+{
+	FNiagaraFunctionSignature* Sig = nullptr;
+	FString FuncName;
+	if (UNiagaraNodeFunctionCall* FuncNode = Cast<UNiagaraNodeFunctionCall>(AddPin->GetOwningNode()))
+	{
+		Sig = &FuncNode->Signature;
+		FuncName = Sig->GetNameString();
+	}
+
+	auto GatherAddPinsForChannel = [&](UNiagaraDataChannel* Channel)
+	{
+		TConstArrayView<FNiagaraVariable> ChannelVars = Channel->GetVariables();
+		for (const FNiagaraVariable& Var : ChannelVars)
+		{
+			FNiagaraTypeDefinition Type = Var.GetType();
+			if (Type.IsEnum() == false)
+			{
+				Type = FNiagaraTypeDefinition(FNiagaraTypeHelper::GetSWCStruct(Var.GetType().GetScriptStruct()));
+			}
+			FNiagaraVariable SWCVar(Type, Var.GetName());
+			FNiagaraEditorUtilities::ResetVariableToDefaultValue(SWCVar);
 
+			const UEdGraphPin* ConstAddPin = AddPin;
+
+			// The script variable is not a duplicate, add an entry for it.
+			FText Category;
+			FText DisplayName;
+			FText Tooltip;
+			if (FuncName == TEXT("SpawnConditional"))
+			{
+				Category = FText::Format(LOCTEXT("NDIDataChannleReadSpawnCategoryFmt", "Conditions on {0} variables"), FText::FromString(Channel->GetAsset()->GetName()));
+				DisplayName = FText::FromName(SWCVar.GetName());
+				Tooltip = FText::Format(
+					LOCTEXT("NDIDataChannleReadAddParameterFmt", "Make this function conditional on {0} from Data Channel {1}.\nThe function will only have an effect for Data Channel entries that pass the comparisson test with the passed value."),
+					FText::FromName(SWCVar.GetName()),
+					FText::FromString(Channel->GetName()));
+			}
+			else
+			{
+				Category = FText::Format(LOCTEXT("NDIDataChannleReadCategoryFmt", "{0}"), FText::FromString(Channel->GetAsset()->GetName()));
+				DisplayName = FText::FromName(SWCVar.GetName());
+				Tooltip = FText::Format(LOCTEXT("NDIDataChannleReadAddParameterFmt", "Read {0} from Data Channel {1}."), DisplayName, FText::FromString(Channel->GetName()));
+			}
+
+			TSharedPtr<FNiagaraMenuAction> Action(new FNiagaraMenuAction(
+				Category, DisplayName, Tooltip, 0, FText::GetEmpty(),
+				FNiagaraMenuAction::FOnExecuteStackAction::CreateUObject(CastChecked<UNiagaraNodeWithDynamicPins>(AddPin->GetOwningNode()), &UNiagaraNodeWithDynamicPins::AddParameter, SWCVar, ConstAddPin->Direction)));
+			Action->SetParameterVariable(SWCVar);
+			Collector.AddAction(Action, 3);
+		}
+	};
+	UNiagaraDataChannel::ForEachDataChannel(GatherAddPinsForChannel);
+}
 
 #undef LOCTEXT_NAMESPACE
\ No newline at end of file
diff --git a/Engine/Plugins/FX/Niagara/Source/NiagaraEditor/Private/NiagaraEditorModule.cpp b/Engine/Plugins/FX/Niagara/Source/NiagaraEditor/Private/NiagaraEditorModule.cpp
--- a/Engine/Plugins/FX/Niagara/Source/NiagaraEditor/Private/NiagaraEditorModule.cpp
+++ b/Engine/Plugins/FX/Niagara/Source/NiagaraEditor/Private/NiagaraEditorModule.cpp
@@ -161,9 +161,11 @@
 
 #include "NiagaraActions.h"
 
-#include "NiagaraDataChannelDefinitions.h"
 #include "NiagaraDataChannel.h"
 
+#include "Engine/AssetManager.h"
+#include "Engine/AssetManagerSettings.h"
+
 #include UE_INLINE_GENERATED_CPP_BY_NAME(NiagaraEditorModule)
 
 IMPLEMENT_MODULE( FNiagaraEditorModule, NiagaraEditor );
@@ -948,8 +950,7 @@ void FNiagaraEditorModule::OnPreExit()
 
 	const FAssetRegistryModule& AssetRegistryModule = FModuleManager::LoadModuleChecked<FAssetRegistryModule>(TEXT("AssetRegistry"));
 	IAssetRegistry& AssetRegistry = AssetRegistryModule.Get();
-	AssetRegistry.OnInMemoryAssetCreated().Remove(OnAssetCreatedHandle);
-	AssetRegistry.OnInMemoryAssetDeleted().Remove(OnAssetDeletedHandle);
+	AssetRegistry.OnFilesLoaded().Remove(AssetRegistryOnLoadCompleteHandle);
 }
 
 void FNiagaraEditorModule::PostGarbageCollect()
@@ -1566,8 +1567,7 @@ void FNiagaraEditorModule::OnPostEngineInit()
 
 	const FAssetRegistryModule& AssetRegistryModule = FModuleManager::LoadModuleChecked<FAssetRegistryModule>(TEXT("AssetRegistry"));
 	IAssetRegistry& AssetRegistry = AssetRegistryModule.Get();
-	OnAssetCreatedHandle = AssetRegistry.OnInMemoryAssetCreated().AddRaw(this, &FNiagaraEditorModule::OnAssetCreated);
-	OnAssetDeletedHandle = AssetRegistry.OnInMemoryAssetDeleted().AddRaw(this, &FNiagaraEditorModule::OnAssetDeleted);
+	AssetRegistryOnLoadCompleteHandle = AssetRegistry.OnFilesLoaded().AddRaw(this, &FNiagaraEditorModule::OnAssetRegistryLoadComplete);
 }
 
 void FNiagaraEditorModule::OnDeviceProfileManagerUpdated()
@@ -2084,19 +2084,21 @@ bool FNiagaraEditorModule::DeferredDestructObjects(float InDeltaTime)
 	return false;
 }
 
-void FNiagaraEditorModule::OnAssetCreated(UObject* CreatedObject)
+void FNiagaraEditorModule::OnAssetRegistryLoadComplete()
 {
-	if (UNiagaraDataChannelDefinitions* DataChannelDef = Cast<UNiagaraDataChannelDefinitions>(CreatedObject))
-	{
-		UNiagaraDataChannelDefinitions::OnAssetCreated(DataChannelDef);
-	}
-}
+	FAssetRegistryModule& AssetRegistryModule = FModuleManager::LoadModuleChecked<FAssetRegistryModule>("AssetRegistry");
+	IAssetRegistry& AssetRegistry = AssetRegistryModule.Get();
 
-void FNiagaraEditorModule::OnAssetDeleted(UObject* DeletedObject)
-{
-	if (UNiagaraDataChannelDefinitions* DataChannelDef = Cast<UNiagaraDataChannelDefinitions>(DeletedObject))
+	AssetRegistry.OnFilesLoaded().Remove(AssetRegistryOnLoadCompleteHandle);
+
+	check(AssetRegistry.IsLoadingAssets() == false);
+
+	//Ensure All Data Channel Assets are loaded and available for use in editor.
+	TArray<FAssetData> AllDataChannels;
+	AssetRegistry.GetAssetsByClass(UNiagaraDataChannelAsset::StaticClass()->GetClassPathName(), AllDataChannels);
+	for (FAssetData& DataChannelAsset : AllDataChannels)
 	{
-		UNiagaraDataChannelDefinitions::OnAssetDeleted(DataChannelDef);
+		UNiagaraDataChannelAsset* NewAsset = Cast<UNiagaraDataChannelAsset>(DataChannelAsset.GetAsset());
 	}
 }
 
diff --git a/Engine/Plugins/FX/Niagara/Source/NiagaraEditor/Private/NiagaraHlslTranslator.cpp b/Engine/Plugins/FX/Niagara/Source/NiagaraEditor/Private/NiagaraHlslTranslator.cpp
--- a/Engine/Plugins/FX/Niagara/Source/NiagaraEditor/Private/NiagaraHlslTranslator.cpp
+++ b/Engine/Plugins/FX/Niagara/Source/NiagaraEditor/Private/NiagaraHlslTranslator.cpp
@@ -2939,6 +2939,7 @@ void FHlslNiagaraTranslator::DefineDataInterfaceHLSL(FString& InHlslOutput)
 
 			FNiagaraDataInterfaceHlslGenerationContext DIHlslGenContext(DIInstanceInfo, GeneratedFunctionSignatures);
 
+			DIHlslGenContext.GetFunctionSignatureSymbolDelegate.BindStatic(&FHlslNiagaraTranslator::GetFunctionSignatureSymbol);
 			DIHlslGenContext.GetStructHlslTypeNameDelegate.BindStatic(&FHlslNiagaraTranslator::GetStructHlslTypeName);
 			DIHlslGenContext.GetPropertyHlslTypeNameDelegate.BindStatic(&FHlslNiagaraTranslator::GetPropertyHlslTypeName);
 			DIHlslGenContext.GetSanitizedSymbolNameDelegate.BindStatic(&FHlslNiagaraTranslator::GetSanitizedSymbolName);
diff --git a/Engine/Plugins/FX/Niagara/Source/NiagaraEditor/Private/NiagaraNodeFunctionCall.cpp b/Engine/Plugins/FX/Niagara/Source/NiagaraEditor/Private/NiagaraNodeFunctionCall.cpp
--- a/Engine/Plugins/FX/Niagara/Source/NiagaraEditor/Private/NiagaraNodeFunctionCall.cpp
+++ b/Engine/Plugins/FX/Niagara/Source/NiagaraEditor/Private/NiagaraNodeFunctionCall.cpp
@@ -1198,6 +1198,19 @@ void UNiagaraNodeFunctionCall::GetNodeContextMenuActions(class UToolMenu* Menu,
 	}
 }
 
+void UNiagaraNodeFunctionCall::CollectAddPinActions(FNiagaraMenuActionCollector& Collector, UEdGraphPin* AddPin)const
+{
+	Super::CollectAddPinActions(Collector, AddPin);
+	if (FunctionScript == nullptr && Signature.Inputs.Num() > 0)
+	{
+		UClass* DIClass = Cast<UClass>(Signature.Inputs[0].GetType().GetClass());
+		if (DIClass)
+		{
+			INiagaraDataInterfaceNodeActionProvider::CollectAddPinActions(DIClass, Collector, AddPin);
+		}
+	}
+}
+
 bool UNiagaraNodeFunctionCall::HasValidScriptAndGraph() const
 {
 	return FunctionScript != nullptr && GetCalledGraph() != nullptr;
@@ -1655,7 +1668,7 @@ void UNiagaraNodeFunctionCall::RefreshSignature()
 					UEdGraphPin* InputPin = FoundPins[i];
 					FNiagaraVariable InputVariable = UEdGraphSchema_Niagara::PinToNiagaraVariable(InputPin);
 
-					if(!BaseSig->Inputs.Contains(InputVariable))
+					if(!BaseSig->Inputs.Contains(InputVariable) && IsAddPin(InputPin) == false && IsExecPin(InputPin) == false)
 					{
 						Signature.AddInput(InputVariable, FText::FromString(InputPin->PinToolTip));
 					}
@@ -1668,7 +1681,7 @@ void UNiagaraNodeFunctionCall::RefreshSignature()
 					UEdGraphPin* OutputPin = FoundPins[i];
 					FNiagaraVariable InputVariable = UEdGraphSchema_Niagara::PinToNiagaraVariable(OutputPin);
 
-					if(!BaseSig->Outputs.Contains(InputVariable))
+					if(!BaseSig->Outputs.Contains(InputVariable) && IsAddPin(OutputPin) == false && IsExecPin(OutputPin) == false)
 					{
 						Signature.AddOutput(UEdGraphSchema_Niagara::PinToNiagaraVariable(OutputPin), FText::FromString(OutputPin->PinToolTip));
 					}
@@ -1693,6 +1706,11 @@ bool UNiagaraNodeFunctionCall::IsBaseSignatureOfDataInterfaceFunction(const UEdG
 			FPinCollectorArray FoundPins;
 			TArray<FNiagaraVariable> InputOrOutputVariables;
 
+			if(BaseSig->bRequiresExecPin && IsExecPin(Pin))
+			{
+				return true;
+			}
+
 			if(Pin->Direction == EGPD_Input)
 			{
 				GetInputPins(FoundPins);
diff --git a/Engine/Plugins/FX/Niagara/Source/NiagaraEditor/Private/NiagaraNodeWithDynamicPins.cpp b/Engine/Plugins/FX/Niagara/Source/NiagaraEditor/Private/NiagaraNodeWithDynamicPins.cpp
--- a/Engine/Plugins/FX/Niagara/Source/NiagaraEditor/Private/NiagaraNodeWithDynamicPins.cpp
+++ b/Engine/Plugins/FX/Niagara/Source/NiagaraEditor/Private/NiagaraNodeWithDynamicPins.cpp
@@ -396,7 +396,7 @@ void UNiagaraNodeWithDynamicPins::GetNodeContextMenuActions(UToolMenu* Menu, UGr
 	}
 }
 
-void UNiagaraNodeWithDynamicPins::AddParameter(FNiagaraVariable Parameter, const UEdGraphPin* AddPin)
+void UNiagaraNodeWithDynamicPins::AddParameter(FNiagaraVariable Parameter, TEnumAsByte<EEdGraphPinDirection> Direction)
 {
 	if (this->IsA<UNiagaraNodeParameterMapBase>())
 	{
@@ -416,15 +416,15 @@ void UNiagaraNodeWithDynamicPins::AddParameter(FNiagaraVariable Parameter, const
 		Graph->AddParameter(Parameter);
 
 		Modify();
-		UEdGraphPin* Pin = this->RequestNewTypedPin(AddPin->Direction, Parameter.GetType(), Parameter.GetName());
+		UEdGraphPin* Pin = this->RequestNewTypedPin(Direction, Parameter.GetType(), Parameter.GetName());
 	}
 	else
 	{
-		RequestNewTypedPin(AddPin->Direction, Parameter.GetType(), Parameter.GetName());
+		RequestNewTypedPin(Direction, Parameter.GetType(), Parameter.GetName());
 	}
 }
 
-void UNiagaraNodeWithDynamicPins::AddParameter(const UNiagaraScriptVariable* ScriptVar, const UEdGraphPin* AddPin)
+void UNiagaraNodeWithDynamicPins::AddParameter(const UNiagaraScriptVariable* ScriptVar, TEnumAsByte<EEdGraphPinDirection> Direction)
 {
 	const FNiagaraVariable& Parameter = ScriptVar->Variable;
 	if (this->IsA<UNiagaraNodeParameterMapBase>())
@@ -439,11 +439,11 @@ void UNiagaraNodeWithDynamicPins::AddParameter(const UNiagaraScriptVariable* Scr
 		Graph->AddParameter(ScriptVar);
 
 		Modify();
-		UEdGraphPin* Pin = this->RequestNewTypedPin(AddPin->Direction, Parameter.GetType(), Parameter.GetName());
+		UEdGraphPin* Pin = this->RequestNewTypedPin(Direction, Parameter.GetType(), Parameter.GetName());
 	}
 	else
 	{
-		RequestNewTypedPin(AddPin->Direction, Parameter.GetType(), Parameter.GetName());
+		RequestNewTypedPin(Direction, Parameter.GetType(), Parameter.GetName());
 	}
 }
 
diff --git a/Engine/Plugins/FX/Niagara/Source/NiagaraEditor/Private/Widgets/SNiagaraGraphPinAdd.cpp b/Engine/Plugins/FX/Niagara/Source/NiagaraEditor/Private/Widgets/SNiagaraGraphPinAdd.cpp
--- a/Engine/Plugins/FX/Niagara/Source/NiagaraEditor/Private/Widgets/SNiagaraGraphPinAdd.cpp
+++ b/Engine/Plugins/FX/Niagara/Source/NiagaraEditor/Private/Widgets/SNiagaraGraphPinAdd.cpp
@@ -126,9 +126,9 @@ TSharedRef<SWidget> SNiagaraGraphPinAdd::OnGetAddButtonMenuContent()
 		.Graphs(InGraphs)
 		.AvailableParameterDefinitions(ScriptViewModel->GetAvailableParameterDefinitions(bSkipSubscribedLibraries))
 		.SubscribedParameterDefinitions(ScriptViewModel->GetSubscribedParameterDefinitions())
-		.OnNewParameterRequested_UObject(OwningNode, &UNiagaraNodeWithDynamicPins::AddParameter, ConstAddPin)
+		.OnNewParameterRequested_UObject(OwningNode, &UNiagaraNodeWithDynamicPins::AddParameter, ConstAddPin->Direction)
 		.OnSpecificParameterRequested_UObject(OwningNode, &UNiagaraNodeWithDynamicPins::AddExistingParameter, ConstAddPin)
-		.OnAddScriptVar_UObject(OwningNode, &UNiagaraNodeWithDynamicPins::AddParameter, ConstAddPin)
+		.OnAddScriptVar_UObject(OwningNode, &UNiagaraNodeWithDynamicPins::AddParameter, ConstAddPin->Direction)
 		.OnAddParameterDefinitions(ScriptViewModel.ToSharedRef(), &FNiagaraScriptViewModel::SubscribeToParameterDefinitions)
 		.OnAllowMakeType_UObject(OwningNode, &UNiagaraNodeWithDynamicPins::AllowNiagaraTypeForAddPin) 
 		.AllowCreatingNew(true)
