Impact FX Aggregation via Data Channels.

diff --git a/Engine/Plugins/FX/Niagara/Source/NiagaraEditor/Public/NiagaraActions.h b/Engine/Plugins/FX/Niagara/Source/NiagaraEditor/Public/NiagaraActions.h
--- a/Engine/Plugins/FX/Niagara/Source/NiagaraEditor/Public/NiagaraActions.h
+++ b/Engine/Plugins/FX/Niagara/Source/NiagaraEditor/Public/NiagaraActions.h
@@ -396,8 +396,11 @@ public:
 	INiagaraDataInterfaceNodeActionProvider& operator=(INiagaraDataInterfaceNodeActionProvider&&) = delete;
 	virtual ~INiagaraDataInterfaceNodeActionProvider() = default;
 
-	/** Pure virtual that sub classes should implement to provide menu actions specific to DIs and their functions. */
-	virtual void GetNodeContextMenuActionsImpl(UToolMenu* Menu, UGraphNodeContextMenuContext* Context, FNiagaraFunctionSignature Signature) const = 0;
+	/** Allows DIs to add context menu actions for the whole node. */
+	virtual void GetNodeContextMenuActionsImpl(UToolMenu* Menu, UGraphNodeContextMenuContext* Context, FNiagaraFunctionSignature Signature) const {}
+
+	/** Allows DIs to add actions for add pins on the DI function call nodes. */
+	virtual void  CollectAddPinActionsImpl(FNiagaraMenuActionCollector& Collector, UEdGraphPin* AddPin) const {}
 
 	template<typename DIClass, typename ActionProviderClass>
 	static void Register();
@@ -408,9 +411,13 @@ public:
 	template<typename DIClass>
 	static void GetNodeContextMenuActions(UToolMenu* Menu, UGraphNodeContextMenuContext* Context, FNiagaraFunctionSignature Signature);
 
-	/**  */
 	static void GetNodeContextMenuActions(UClass* DIClass, UToolMenu* Menu, UGraphNodeContextMenuContext* Context, FNiagaraFunctionSignature Signature);
 
+	template<typename DIClass>
+	static void CollectAddPinActions(FNiagaraMenuActionCollector& Collector, UEdGraphPin* AddPin);
+
+	static void CollectAddPinActions(UClass* DIClass, FNiagaraMenuActionCollector& Collector, UEdGraphPin* AddPin);
+
 	/** All currently registered action providers. */
 	static TMap<FName, TUniquePtr<INiagaraDataInterfaceNodeActionProvider>> RegisteredActionProviders;
 };
@@ -437,6 +444,15 @@ void INiagaraDataInterfaceNodeActionProvider::GetNodeContextMenuActions(UToolMen
 	}
 }
 
+template<typename DIClass>
+void INiagaraDataInterfaceNodeActionProvider::CollectAddPinActions(FNiagaraMenuActionCollector& Collector, UEdGraphPin* AddPin)
+{
+	if (TUniquePtr<INiagaraDataInterfaceNodeActionProvider>* Provider = RegisteredActionProviders.Find(DIClass::StaticClass()->GetFName()))
+	{
+		(*Provider)->CollectAddPinActionsImpl(Collector, AddPin);
+	}
+}
+
 ////////////////////////////////
 /// Actions for engine data interfaces.
 /////////////////////////////////
@@ -453,4 +469,5 @@ class FNiagaraDataInterfaceNodeActionProvider_DataChannelRead : public INiagaraD
 public:
 
 	virtual void GetNodeContextMenuActionsImpl(UToolMenu* Menu, UGraphNodeContextMenuContext* Context, FNiagaraFunctionSignature Signature) const override;
+	virtual void CollectAddPinActionsImpl(FNiagaraMenuActionCollector& Collector, UEdGraphPin* AddPin)const override;
 };
\ No newline at end of file
diff --git a/Engine/Plugins/FX/Niagara/Source/NiagaraEditor/Public/NiagaraEditorModule.h b/Engine/Plugins/FX/Niagara/Source/NiagaraEditor/Public/NiagaraEditorModule.h
--- a/Engine/Plugins/FX/Niagara/Source/NiagaraEditor/Public/NiagaraEditorModule.h
+++ b/Engine/Plugins/FX/Niagara/Source/NiagaraEditor/Public/NiagaraEditorModule.h
@@ -335,8 +335,7 @@ private:
 		StackIssueGenerators.Add(StructName) = Generator;
 	}
 
-	void OnAssetCreated(UObject* DeletedObject);
-	void OnAssetDeleted(UObject* DeletedObject);
+	void OnAssetRegistryLoadComplete();
 
 private:
 	TSharedPtr<FExtensibilityManager> MenuExtensibilityManager;
@@ -371,8 +370,7 @@ private:
 	FDelegateHandle PreviewPlatformChangedHandle;
 	FDelegateHandle PreviewFeatureLevelChangedHandle;
 
-	FDelegateHandle OnAssetCreatedHandle;
-	FDelegateHandle OnAssetDeletedHandle;
+	FDelegateHandle AssetRegistryOnLoadCompleteHandle;
 	
 	USequencerSettings* SequencerSettings;
 
diff --git a/Engine/Plugins/FX/Niagara/Source/NiagaraEditor/Public/NiagaraNodeFunctionCall.h b/Engine/Plugins/FX/Niagara/Source/NiagaraEditor/Public/NiagaraNodeFunctionCall.h
--- a/Engine/Plugins/FX/Niagara/Source/NiagaraEditor/Public/NiagaraNodeFunctionCall.h
+++ b/Engine/Plugins/FX/Niagara/Source/NiagaraEditor/Public/NiagaraNodeFunctionCall.h
@@ -130,6 +130,10 @@ public:
 	virtual bool IsDeprecated() const override;
 	//~ End EdGraphNode Interface
 
+	//~ Begin UNiagaraNodeWithDynamicPins Interface
+	virtual void CollectAddPinActions(FNiagaraMenuActionCollector& Collector, UEdGraphPin* AddPin) const override;
+	//~ End UNiagaraNodeWithDynamicPins Interface
+
 	/** When overriding an input value, this updates which variable guid was bound to which input name, so it can be reassigned when the input is renamed.*/
 	void UpdateInputNameBinding(const FGuid& BoundVariableGuid, const FName& BoundName);
 
diff --git a/Engine/Plugins/FX/Niagara/Source/NiagaraEditor/Public/NiagaraNodeWithDynamicPins.h b/Engine/Plugins/FX/Niagara/Source/NiagaraEditor/Public/NiagaraNodeWithDynamicPins.h
--- a/Engine/Plugins/FX/Niagara/Source/NiagaraEditor/Public/NiagaraNodeWithDynamicPins.h
+++ b/Engine/Plugins/FX/Niagara/Source/NiagaraEditor/Public/NiagaraNodeWithDynamicPins.h
@@ -42,13 +42,16 @@ public:
 	virtual bool AllowNiagaraTypeForAddPin(const FNiagaraTypeDefinition& InType, EEdGraphPinDirection InDirection) const { return AllowNiagaraTypeForAddPin(InType); };
 
 	/** Request a new pin. */
-	void AddParameter(FNiagaraVariable Parameter, const UEdGraphPin* AddPin);
-	void AddParameter(const UNiagaraScriptVariable* ScriptVar, const UEdGraphPin* AddPin);
+	void AddParameter(FNiagaraVariable Parameter, TEnumAsByte<EEdGraphPinDirection> Direction);
+	void AddParameter(const UNiagaraScriptVariable* ScriptVar, TEnumAsByte<EEdGraphPinDirection> Direction);
 	void AddExistingParameter(FNiagaraVariable Parameter, const UEdGraphPin* AddPin);
 
 	/** Convenience method to determine whether this Node is a Map Get or Map Set when adding a parameter through the parameter panel. */
 	virtual EEdGraphPinDirection GetPinDirectionForNewParameters() { return EEdGraphPinDirection::EGPD_MAX; };
 
+	/** Virtual functions to allow for custom actions in Add Pin menus. */
+	virtual void CollectAddPinActions(FNiagaraMenuActionCollector& Collector, UEdGraphPin* AddPin)const {}
+
 protected:
 	virtual bool AllowDynamicPins() const { return true; }
 
diff --git a/Engine/Plugins/FX/Niagara/Source/NiagaraEditor/Public/Widgets/SNiagaraParameterMenu.cpp b/Engine/Plugins/FX/Niagara/Source/NiagaraEditor/Public/Widgets/SNiagaraParameterMenu.cpp
--- a/Engine/Plugins/FX/Niagara/Source/NiagaraEditor/Public/Widgets/SNiagaraParameterMenu.cpp
+++ b/Engine/Plugins/FX/Niagara/Source/NiagaraEditor/Public/Widgets/SNiagaraParameterMenu.cpp
@@ -828,6 +828,10 @@ void SNiagaraAddParameterFromPinMenu::Construct(const FArguments& InArgs)
 void SNiagaraAddParameterFromPinMenu::CollectAllActions(FGraphActionListBuilderBase& OutAllActions)
 {
 	FNiagaraMenuActionCollector Collector;
+	
+	//First allow the node to add it's own actions.
+	NiagaraNode->CollectAddPinActions(Collector, AddPin);	
+	
 	TArray<FNiagaraTypeDefinition> Types;
 	FNiagaraEditorUtilities::GetAllowedParameterTypes(Types);
 	Types.Sort([](const FNiagaraTypeDefinition& A, const FNiagaraTypeDefinition& B) { return (A.GetNameText().ToLower().ToString() < B.GetNameText().ToLower().ToString()); });
@@ -848,7 +852,7 @@ void SNiagaraAddParameterFromPinMenu::CollectAllActions(FGraphActionListBuilderB
 			const UEdGraphPin* ConstAddPin = AddPin;
 			TSharedPtr<FNiagaraMenuAction> Action(new FNiagaraMenuAction(
 				Category, DisplayName, Tooltip, 0, FText::GetEmpty(),
-				FNiagaraMenuAction::FOnExecuteStackAction::CreateUObject(NiagaraNode, &UNiagaraNodeWithDynamicPins::AddParameter, Var, ConstAddPin)));
+				FNiagaraMenuAction::FOnExecuteStackAction::CreateUObject(NiagaraNode, &UNiagaraNodeWithDynamicPins::AddParameter, Var, ConstAddPin->Direction)));
 
 			Collector.AddAction(Action, 0);
 		}
diff --git a/Engine/Plugins/FX/Niagara/Source/NiagaraEditorWidgets/Private/DetailCustomizations/NiagaraDataChannelDetails.cpp b/Engine/Plugins/FX/Niagara/Source/NiagaraEditorWidgets/Private/DetailCustomizations/NiagaraDataChannelDetails.cpp
--- a/Engine/Plugins/FX/Niagara/Source/NiagaraEditorWidgets/Private/DetailCustomizations/NiagaraDataChannelDetails.cpp
+++ b/Engine/Plugins/FX/Niagara/Source/NiagaraEditorWidgets/Private/DetailCustomizations/NiagaraDataChannelDetails.cpp
@@ -1,40 +1,269 @@
 // Copyright Epic Games, Inc. All Rights Reserved.
 
 #include "NiagaraDataChannelDetails.h"
+
+#include "IDetailChildrenBuilder.h"
+#include "DetailLayoutBuilder.h"
+#include "PropertyCustomizationHelpers.h"
+
+#include "Types/SlateEnums.h"
+#include "Widgets/SItemSelector.h"
+#include "Widgets/SCompoundWidget.h"
+#include "Widgets/DeclarativeSyntaxSupport.h"
+#include "Widgets/Input/SComboBox.h"
+#include "Widgets/Input/SEditableTextBox.h"
+#include "Widgets/Input/SSearchBox.h"
+#include "Widgets/Input/SComboButton.h"
+#include "SListViewSelectorDropdownMenu.h"
+
 #include "SNiagaraNamePropertySelector.h"
 #include "NiagaraDataChannel.h"
-#include "NiagaraDataChannelDefinitions.h"
+#include "DataInterface\NiagaraDataInterfaceDataChannelRead.h"
+#include "NiagaraDataChannel.h"
+#include "NiagaraDataChannel_Global.h"
+#include "NiagaraDataChannel_Islands.h"
+
+#include "NiagaraSystem.h"
+#include "NiagaraEditorUtilities.h"
+#include "NiagaraEditorWidgetsStyle.h"
+
+#include "NiagaraDataInterfaceDetails.h"
 
-void FNiagaraDataChannelReferenceDetailsCustomization::CustomizeHeader(TSharedRef<IPropertyHandle> PropertyHandle, FDetailWidgetRow& HeaderRow, IPropertyTypeCustomizationUtils& CustomizationUtils)
+#define LOCTEXT_NAMESPACE "NiagaraDataChannelDetails"
+
+class IPropertyHandle;
+class SSearchBox;
+
+//TODO: If more broadly useful, generalize and move to own files.
+class SNiagaraVariableSelector : public SCompoundWidget
 {
-	ChannelHandle = PropertyHandle->GetChildHandle(TEXT("ChannelName"));
+public:
+	SLATE_BEGIN_ARGS(SNiagaraVariableSelector) { }
+	SLATE_END_ARGS();
+
+	void Construct(const FArguments& InArgs, TSharedRef<IPropertyHandle> InBaseProperty, FNiagaraVariableBase* InTarget, const TArray<TSharedPtr<FNiagaraVariableBase>>& InOptionsSource);
+
+	void SetSourceArray(TArray<TSharedPtr<FNiagaraVariableBase>>& InOptionsSource);
+
+	FSimpleMulticastDelegate& GetOnTargetChanged() { return OnTargetChanged; }
+
+private:
+	void OnComboOpening();
+
+	void OnSelectionChanged(TSharedPtr<FNiagaraVariableBase> InNewSelection, ESelectInfo::Type);
+
+	TSharedRef<SWidget> MakeSelectionWidget(TSharedPtr<FNiagaraVariableBase> InItem);
+
+	TSharedRef<ITableRow> GenerateAddElementRow(TSharedPtr<FNiagaraVariableBase> Entry, const TSharedRef<STableViewBase>& OwnerTable) const;
+
+
+private:
+	TArray<TSharedPtr<FNiagaraVariableBase>> OptionsSourceList;
+
+	/** The current array property being edited */
+	TSharedPtr<IPropertyHandle> PropertyHandle;
+
+	/** The component list control - part of the combo drop down */
+	TSharedPtr< SListView<TSharedPtr<FNiagaraVariableBase>> > ElementsListView;
+
+	TSharedPtr<SComboButton> ElementButton;
+
+	FNiagaraVariableBase* Target = nullptr;
+
+	TSharedPtr<SHorizontalBox> ComboButtonWidget;
+
+	FSimpleMulticastDelegate OnTargetChanged;
+};
 
-	TArray<TSharedPtr<FName>> DataChannelList;
-	const TArray<UNiagaraDataChannelDefinitions*>& DataChannelDefs = UNiagaraDataChannelDefinitions::GetDataChannelDefinitions(false, true);
-	for (const UNiagaraDataChannelDefinitions* Def : DataChannelDefs)
+void SNiagaraVariableSelector::Construct(const FArguments& InArgs, TSharedRef<IPropertyHandle> InBaseProperty, FNiagaraVariableBase* InTarget, const TArray<TSharedPtr<FNiagaraVariableBase>>& InOptionsSource)
+{
+	PropertyHandle = InBaseProperty;
+	OptionsSourceList = InOptionsSource;
+	Target = InTarget;
+
+	SAssignNew(ElementsListView, SListView<TSharedPtr<FNiagaraVariableBase>>)
+		.ListItemsSource(&OptionsSourceList)
+		.OnSelectionChanged(this, &SNiagaraVariableSelector::OnSelectionChanged)
+		.OnGenerateRow(this, &SNiagaraVariableSelector::GenerateAddElementRow)
+		.SelectionMode(ESelectionMode::Single);
+
+	ElementsListView->RequestListRefresh();
+
+	ComboButtonWidget = SNew(SHorizontalBox)
+		+ SHorizontalBox::Slot()
+		.VAlign(VAlign_Center)
+		.AutoWidth()
+		.Padding(1.f, 1.f)
+		[
+			FNiagaraParameterUtilities::GetParameterWidget(*Target, true, false)
+		];
+
+	// Create the Construct arguments for SComboButton
+
+	SAssignNew(ElementButton, SComboButton)
+		.ButtonContent()
+		[
+			ComboButtonWidget.ToSharedRef()
+		]
+	.MenuContent()
+		[
+			SNew(SListViewSelectorDropdownMenu<TSharedPtr<FNiagaraVariableBase>>, nullptr, ElementsListView)
+			[
+				SNew(SBorder)
+				.BorderImage(FAppStyle::GetBrush("Menu.Background"))
+				.Padding(2)
+				[
+					SNew(SBox)
+					.WidthOverride(175)
+					[
+						SNew(SVerticalBox)
+						+ SVerticalBox::Slot()
+						.MaxHeight(400)
+						.Padding(8.f)
+						[
+							ElementsListView.ToSharedRef()
+						]
+					]
+				]
+			]
+		]
+	.IsFocusable(true)
+		.ContentPadding(FMargin(5, 0))
+		.OnComboBoxOpened(this, &SNiagaraVariableSelector::OnComboOpening);
+
+	ElementsListView->EnableToolTipForceField(true);
+
+	ChildSlot
+		[
+			ElementButton.ToSharedRef()
+		];
+};
+
+void SNiagaraVariableSelector::SetSourceArray(TArray<TSharedPtr<FNiagaraVariableBase>>& InOptionsSource)
+{
+	OptionsSourceList = InOptionsSource;
+	if (ElementsListView.IsValid())
 	{
-		for (TObjectPtr<UNiagaraDataChannel> Channel : Def->DataChannels)
-		{
-			if(Channel)
-			{
-				DataChannelList.Add(MakeShared<FName>(Channel->GetChannelName()));
-			}
-		}
+		ElementsListView->RequestListRefresh();
 	}
+}
 
-	// Build Widget
-	HeaderRow
-		.NameContent()
+TSharedRef<ITableRow> SNiagaraVariableSelector::GenerateAddElementRow(TSharedPtr<FNiagaraVariableBase> Entry, const TSharedRef<STableViewBase>& OwnerTable) const
+{
+	return
+		SNew(STableRow< TSharedPtr<FString> >, OwnerTable)
+		.Style(&FAppStyle::Get().GetWidgetStyle<FTableRowStyle>("TableView.NoHoverTableRow"))
+		.ShowSelection(true)
 		[
-			ChannelHandle->CreatePropertyNameWidget()
-		]
-		.ValueContent()
-		.MinDesiredWidth(200.0f)
+			SNew(SBox)
+			.Padding(2.f)
+			[
+				FNiagaraParameterUtilities::GetParameterWidget(*Entry, true, false)
+			]
+		];
+}
+
+void SNiagaraVariableSelector::OnSelectionChanged(TSharedPtr<FNiagaraVariableBase> InNewSelection, ESelectInfo::Type SelectInfo)
+{
+	if (InNewSelection.IsValid() && (SelectInfo != ESelectInfo::OnNavigation))
+	{
+		*Target = *InNewSelection;
+		ElementButton->SetIsOpen(false, false);
+
+		ComboButtonWidget->GetSlot(0)
 		[
-			SNew(SNiagaraNamePropertySelector, ChannelHandle.ToSharedRef(), DataChannelList)
-		];	
+			FNiagaraParameterUtilities::GetParameterWidget(*Target, true, false)
+		];
+
+		GetOnTargetChanged().Broadcast();
+	}
+}
+
+void SNiagaraVariableSelector::OnComboOpening()
+{
+	
+}
+
+//////////////////////////////////////////////////////////////////////////
+
+void FNiagaraDataChannelIslandsDetails::CustomizeDetails(IDetailLayoutBuilder& DetailBuilder)
+{
+	TArray<TWeakObjectPtr<UObject>> SelectedObjects;
+	DetailBuilder.GetObjectsBeingCustomized(SelectedObjects);
+	if (SelectedObjects.Num() != 1)
+	{
+		return;
+	}
+
+ 	UObject* Obj = SelectedObjects[0].Get();
+ 	UNiagaraDataChannel_Islands* Islands = Cast<UNiagaraDataChannel_Islands>(Obj);
+
+	if(Islands == nullptr)
+	{
+		return;
+	}
+
+	//Right now we're only interested in customizing the islands channel. In future we'll do others.
+	auto AddProperties = [](const TArray<TSharedRef<IPropertyHandle>>& Properties, IDetailCategoryBuilder& CategoryBuilder)
+	{
+		for(auto& Prop : Properties)
+		{
+			CategoryBuilder.AddProperty(Prop);
+		}
+	};
+
+	TArray<TSharedRef<IPropertyHandle>> Properties;
+
+	static const FName ChannelCategoryName = TEXT("Data Channel");
+	IDetailCategoryBuilder& ChannelCategoryBuilder = DetailBuilder.EditCategory(ChannelCategoryName);	
+	ChannelCategoryBuilder.GetDefaultProperties(Properties);
+	AddProperties(Properties, ChannelCategoryBuilder);
+
+	static const FName IslandsCategoryName = TEXT("Islands");
+	IDetailCategoryBuilder& IslandsCategoryBuilder = DetailBuilder.EditCategory(IslandsCategoryName);
+	IslandsCategoryBuilder.GetDefaultProperties(Properties);
+	AddProperties(Properties, IslandsCategoryBuilder);
+}
+
+TSharedRef<IDetailCustomization> FNiagaraDataChannelIslandsDetails::MakeInstance()
+{
+	return MakeShared<FNiagaraDataChannelIslandsDetails>();
+}
+
+//////////////////////////////////////////////////////////////////////////
+
+
+void FNiagaraDataInterfaceDataChannelReadDetails::CustomizeDetails(IDetailLayoutBuilder& DetailBuilder)
+{
+	FNiagaraDataInterfaceDetailsBase::CustomizeDetails(DetailBuilder);
+
+	TArray<TWeakObjectPtr<UObject>> SelectedObjects;
+	DetailBuilder.GetObjectsBeingCustomized(SelectedObjects);
+	if (SelectedObjects.Num() != 1 || SelectedObjects[0]->IsA<UNiagaraDataInterfaceDataChannelRead>() == false)
+	{
+		return;
+	}
+
+	UNiagaraDataInterfaceDataChannelRead* ReadDataInterface = CastChecked<UNiagaraDataInterfaceDataChannelRead>(SelectedObjects[0].Get());
+	ReadDataInterfaceWeak = ReadDataInterface;
+	NiagaraSystemWeak = ReadDataInterface->GetTypedOuter<UNiagaraSystem>();
+
+	static const FName ReadCategoryName = TEXT("DataChannelRead");
+	IDetailCategoryBuilder& CategoryBuilder = DetailBuilder.EditCategory(ReadCategoryName);
+
+	TArray<TSharedRef<IPropertyHandle>> Properties;
+	CategoryBuilder.GetDefaultProperties(Properties, true, true);
+	for (const TSharedRef<IPropertyHandle>& PropertyHandle : Properties)
+	{
+		FProperty* Property = PropertyHandle->GetProperty();
+		CategoryBuilder.AddProperty(PropertyHandle);		
+	}
 }
 
-void FNiagaraDataChannelReferenceDetailsCustomization::CustomizeChildren(TSharedRef<IPropertyHandle> PropertyHandle, IDetailChildrenBuilder& ChildBuilder, IPropertyTypeCustomizationUtils& CustomizationUtils)
+TSharedRef<IDetailCustomization> FNiagaraDataInterfaceDataChannelReadDetails::MakeInstance()
 {
+	return MakeShared<FNiagaraDataInterfaceDataChannelReadDetails>();
 }
+
+#undef LOCTEXT_NAMESPACE
+
diff --git a/Engine/Plugins/FX/Niagara/Source/NiagaraEditorWidgets/Private/DetailCustomizations/NiagaraDataChannelDetails.h b/Engine/Plugins/FX/Niagara/Source/NiagaraEditorWidgets/Private/DetailCustomizations/NiagaraDataChannelDetails.h
--- a/Engine/Plugins/FX/Niagara/Source/NiagaraEditorWidgets/Private/DetailCustomizations/NiagaraDataChannelDetails.h
+++ b/Engine/Plugins/FX/Niagara/Source/NiagaraEditorWidgets/Private/DetailCustomizations/NiagaraDataChannelDetails.h
@@ -5,21 +5,35 @@
 #include "CoreMinimal.h"
 #include "IPropertyTypeCustomization.h"
 #include "IDetailCustomization.h"
+#include "NiagaraDataInterfaceDetails.h"
 
-//** Properties customization for FNiagaraDataChannelReference. */
-class FNiagaraDataChannelReferenceDetailsCustomization : public IPropertyTypeCustomization
+class UNiagaraDataInterfaceDataChannelRead;
+class UNiagaraSystem;
+class UNiagaraDataChannel;
+class UNiagaraDataChannel_Islands;
+
+/** Details customization for Niagara data channels. */
+class FNiagaraDataChannelIslandsDetails : public IDetailCustomization
 {
 public:
-	static TSharedRef<IPropertyTypeCustomization> MakeInstance()
-	{
-		return MakeShared<FNiagaraDataChannelReferenceDetailsCustomization>();
-	}
+	virtual void CustomizeDetails(IDetailLayoutBuilder& DetailBuilder) override;
+	static TSharedRef<IDetailCustomization> MakeInstance();
+
+private:
+
+	TWeakObjectPtr<UNiagaraDataChannel_Islands> DataChannel;
+};
 
-	virtual void CustomizeHeader(TSharedRef<IPropertyHandle> PropertyHandle, FDetailWidgetRow& HeaderRow, IPropertyTypeCustomizationUtils& CustomizationUtils);
 
-	virtual void CustomizeChildren(TSharedRef<IPropertyHandle> PropertyHandle, IDetailChildrenBuilder& ChildBuilder, IPropertyTypeCustomizationUtils& CustomizationUtils);
+/** Details customization for Niagara data channel read data interface. */
+class FNiagaraDataInterfaceDataChannelReadDetails : public FNiagaraDataInterfaceDetailsBase
+{
+public:
+	virtual void CustomizeDetails(IDetailLayoutBuilder& DetailBuilder) override;
+	static TSharedRef<IDetailCustomization> MakeInstance();
 
 private:
-	
-	TSharedPtr<IPropertyHandle> ChannelHandle;
+
+	TWeakObjectPtr<UNiagaraDataInterfaceDataChannelRead> ReadDataInterfaceWeak;
+	TWeakObjectPtr<UNiagaraSystem> NiagaraSystemWeak;
 };
diff --git a/Engine/Plugins/FX/Niagara/Source/NiagaraEditorWidgets/Private/NiagaraEditorWidgetsModule.cpp b/Engine/Plugins/FX/Niagara/Source/NiagaraEditorWidgets/Private/NiagaraEditorWidgetsModule.cpp
--- a/Engine/Plugins/FX/Niagara/Source/NiagaraEditorWidgets/Private/NiagaraEditorWidgetsModule.cpp
+++ b/Engine/Plugins/FX/Niagara/Source/NiagaraEditorWidgets/Private/NiagaraEditorWidgetsModule.cpp
@@ -30,6 +30,7 @@
 #include "NiagaraEditorCommands.h"
 
 #include "NiagaraDataChannelCommon.h"
+#include "NiagaraDataChannel_Islands.h"
 
 IMPLEMENT_MODULE(FNiagaraEditorWidgetsModule, NiagaraEditorWidgets);
 
@@ -142,9 +143,11 @@ void FNiagaraEditorWidgetsModule::StartupModule()
 	PropertyModule.RegisterCustomClassLayout("NiagaraRibbonRendererProperties", FOnGetDetailCustomizationInstance::CreateStatic(&FNiagaraRibbonRendererDetails::MakeInstance));
 	PropertyModule.RegisterCustomClassLayout("NiagaraSpriteRendererProperties", FOnGetDetailCustomizationInstance::CreateStatic(&FNiagaraSpriteRendererDetails::MakeInstance));
 
-	PropertyModule.RegisterCustomPropertyTypeLayout(
-		FNiagaraDataChannelReference::StaticStruct()->GetFName(),
-		FOnGetPropertyTypeCustomizationInstance::CreateStatic(&FNiagaraDataChannelReferenceDetailsCustomization::MakeInstance));
+	PropertyModule.RegisterCustomClassLayout("NiagaraDataInterfaceDataChannelRead", FOnGetDetailCustomizationInstance::CreateStatic(&FNiagaraDataInterfaceDataChannelReadDetails::MakeInstance));
+
+	PropertyModule.RegisterCustomClassLayout(
+		UNiagaraDataChannel_Islands::StaticClass()->GetFName(),
+		FOnGetDetailCustomizationInstance::CreateStatic(&FNiagaraDataChannelIslandsDetails::MakeInstance));
 
 	ReinitializeStyleCommand = IConsoleManager::Get().RegisterConsoleCommand(
 		TEXT("fx.NiagaraEditorWidgets.ReinitializeStyle"),
