speculative linux fix.

diff --git a/Engine/Plugins/FX/Niagara/Source/NiagaraEditorWidgets/Private/DetailCustomizations/NiagaraDataChannelDetails.cpp b/Engine/Plugins/FX/Niagara/Source/NiagaraEditorWidgets/Private/DetailCustomizations/NiagaraDataChannelDetails.cpp
--- a/Engine/Plugins/FX/Niagara/Source/NiagaraEditorWidgets/Private/DetailCustomizations/NiagaraDataChannelDetails.cpp
+++ b/Engine/Plugins/FX/Niagara/Source/NiagaraEditorWidgets/Private/DetailCustomizations/NiagaraDataChannelDetails.cpp
@@ -18,7 +18,7 @@
 
 #include "SNiagaraNamePropertySelector.h"
 #include "NiagaraDataChannel.h"
-#include "DataInterface\NiagaraDataInterfaceDataChannelRead.h"
+#include "DataInterface/NiagaraDataInterfaceDataChannelRead.h"
 #include "NiagaraDataChannel.h"
 #include "NiagaraDataChannel_Global.h"
 #include "NiagaraDataChannel_Islands.h"
