Fix for shipping editor builds.

diff --git a/Engine/Plugins/FX/Niagara/Source/Niagara/Public/NiagaraCommon.h b/Engine/Plugins/FX/Niagara/Source/Niagara/Public/NiagaraCommon.h
--- a/Engine/Plugins/FX/Niagara/Source/Niagara/Public/NiagaraCommon.h
+++ b/Engine/Plugins/FX/Niagara/Source/Niagara/Public/NiagaraCommon.h
@@ -50,7 +50,7 @@ using FNiagaraDataBufferRef = TRefCountPtr<FNiagaraDataBuffer>;
 
 #define NIAGARA_MEMORY_TRACKING	!UE_BUILD_SHIPPING
 
-#define NIAGARA_SYSTEM_CAPTURE (!UE_BUILD_SHIPPING)
+#define NIAGARA_SYSTEM_CAPTURE (!UE_BUILD_SHIPPING || WITH_EDITORONLY_DATA)
 
 #ifndef NIAGARA_COMPUTEDEBUG_ENABLED
 	#define NIAGARA_COMPUTEDEBUG_ENABLED WITH_EDITOR
