Add a bFullyLoaded check to FNiagaraWorldManager::PrimePoolForAllSystems as a speculative fix to a hang which is suspected to be caused by creating UNiagaraComponents on systems that haven't been fully loaded.

diff --git a/Engine/Plugins/FX/Niagara/Source/Niagara/Private/NiagaraWorldManager.cpp b/Engine/Plugins/FX/Niagara/Source/Niagara/Private/NiagaraWorldManager.cpp
--- a/Engine/Plugins/FX/Niagara/Source/Niagara/Private/NiagaraWorldManager.cpp
+++ b/Engine/Plugins/FX/Niagara/Source/Niagara/Private/NiagaraWorldManager.cpp
@@ -1851,7 +1851,8 @@ void FNiagaraWorldManager::PrimePoolForAllSystems()
 		//Prime the pool for all currently loaded systems.
 		for (TObjectIterator<UNiagaraSystem> It; It; ++It)
 		{
-			if (UNiagaraSystem* Sys = *It)
+			UNiagaraSystem* Sys = *It;
+			if (Sys && Sys->bFullyLoaded)
 			{
 				ComponentPool->PrimePool(Sys, World);
 			}
