Add Niagara components created during pool priming to the InUseComponents_Manual array to avoid false positive errors

diff --git a/Engine/Plugins/FX/Niagara/Source/Niagara/Private/NiagaraComponentPool.cpp b/Engine/Plugins/FX/Niagara/Source/Niagara/Private/NiagaraComponentPool.cpp
--- a/Engine/Plugins/FX/Niagara/Source/Niagara/Private/NiagaraComponentPool.cpp
+++ b/Engine/Plugins/FX/Niagara/Source/Niagara/Private/NiagaraComponentPool.cpp
@@ -346,6 +346,9 @@ void UNiagaraComponentPool::PrimePool(UNiagaraSystem* Template, UWorld* World)
 			for (auto& Comp : NewComps)
 			{
 				Comp = Pool.Acquire(World, Template, ENCPoolMethod::ManualRelease, true);//Force the pool to create a new system.
+#if ENABLE_NC_POOL_DEBUGGING
+				InUseComponents_Manual.Add(Comp);
+#endif
 				Comp->InitializeSystem();
 	
 				if(GNiagaraKeepPooledComponentsRegistered)
