- Handle disable emitters with simulation cache

#rnx

diff --git a/Engine/Plugins/FX/Niagara/Source/Niagara/Classes/NiagaraSimCache.h b/Engine/Plugins/FX/Niagara/Source/Niagara/Classes/NiagaraSimCache.h
--- a/Engine/Plugins/FX/Niagara/Source/Niagara/Classes/NiagaraSimCache.h
+++ b/Engine/Plugins/FX/Niagara/Source/Niagara/Classes/NiagaraSimCache.h
@@ -335,6 +335,7 @@ struct FNiagaraSimCacheDataBuffersLayout
 	mutable bool bNeedsCacheBufferReadInfoUpdateForRT = false;
 	mutable FCacheBufferReadInfo CacheBufferReadInfo_RT;
 
+	bool IsLayoutValid() const { return !LayoutName.IsNone(); }
 	int32 IndexOfCacheVariable(const FNiagaraVariableBase& InVariable) const;
 	const FNiagaraSimCacheVariable* FindCacheVariable(const FNiagaraVariableBase& InVariable) const;
 };
diff --git a/Engine/Plugins/FX/Niagara/Source/Niagara/Private/NiagaraSimCache.cpp b/Engine/Plugins/FX/Niagara/Source/Niagara/Private/NiagaraSimCache.cpp
--- a/Engine/Plugins/FX/Niagara/Source/Niagara/Private/NiagaraSimCache.cpp
+++ b/Engine/Plugins/FX/Niagara/Source/Niagara/Private/NiagaraSimCache.cpp
@@ -546,8 +546,12 @@ bool UNiagaraSimCache::CanRead(UNiagaraSystem* NiagaraSystem)
 		{
 			const FNiagaraEmitterCompiledData& EmitterCompiledData = NiagaraSystem->GetEmitterCompiledData()[i].Get();
 			FNiagaraSimCacheDataBuffersLayout& EmitterLayout = CacheLayout.EmitterLayouts[i];
-			bCacheValid &= FNiagaraSimCacheHelper::BuildCacheReadMappings(EmitterLayout, EmitterCompiledData.DataSetCompiledData);
-			bCacheValid &= EmitterLayout.SimTarget == EmitterCompiledData.DataSetCompiledData.SimTarget;
+			bCacheValid &= EmitterLayout.IsLayoutValid() == NiagaraSystem->GetEmitterHandle(i).GetIsEnabled();
+			if (EmitterLayout.IsLayoutValid())
+			{
+				bCacheValid &= FNiagaraSimCacheHelper::BuildCacheReadMappings(EmitterLayout, EmitterCompiledData.DataSetCompiledData);
+				bCacheValid &= EmitterLayout.SimTarget == EmitterCompiledData.DataSetCompiledData.SimTarget;
+			}
 		}
 
 		if (bCacheValid == false)
@@ -659,6 +663,11 @@ bool UNiagaraSimCache::ReadFrame(int32 FrameIndex, float FrameFraction, FNiagara
 	const int32 NumEmitters = CacheLayout.EmitterLayouts.Num();
 	for (int32 i=0; i < NumEmitters; ++i)
 	{
+		if (CacheLayout.EmitterLayouts[i].IsLayoutValid() == false)
+		{
+			continue;
+		}
+
 		const FNiagaraSimCacheEmitterFrame& CacheEmitterFrame = CacheFrame.EmitterData[i];
 		FNiagaraEmitterInstance& EmitterInstance = Helper.SystemInstance->GetEmitters()[i].Get();
 		EmitterInstance.CachedBounds = CacheEmitterFrame.LocalBounds;
