- Remove optimization to test for LUT index not changing
- This can cause issues if the LUT ended up in the same location but the size was different

#rnx

diff --git a/Engine/Plugins/FX/Niagara/Source/Niagara/Private/NiagaraDataInterfaceCurveBase.cpp b/Engine/Plugins/FX/Niagara/Source/Niagara/Private/NiagaraDataInterfaceCurveBase.cpp
--- a/Engine/Plugins/FX/Niagara/Source/Niagara/Private/NiagaraDataInterfaceCurveBase.cpp
+++ b/Engine/Plugins/FX/Niagara/Source/Niagara/Private/NiagaraDataInterfaceCurveBase.cpp
@@ -440,18 +440,12 @@ bool UNiagaraDataInterfaceCurveBase::Equals(const UNiagaraDataInterface* Other)
 
 void UNiagaraDataInterfaceCurveBase::CacheStaticBuffers(struct FNiagaraSystemStaticBuffers& StaticBuffers, const FNiagaraVariable& ResolvedVariable, bool bUsedByCPU, bool bUsedByGPU)
 {
-	const uint32 PrevLUTOffset = LUTOffset;
-
 	LUTOffset = INDEX_NONE;
 	if (bUsedByGPU && ResolvedVariable.IsInNameSpace(FNiagaraConstants::UserNamespaceString) == false)
 	{
 		LUTOffset = StaticBuffers.AddGpuData(ShaderLUT);
 	}
-
-	if (PrevLUTOffset != LUTOffset)
-	{
-		MarkRenderDataDirty();
-	}
+	MarkRenderDataDirty();
 }
 
 void UNiagaraDataInterfaceCurveBase::BuildShaderParameters(FNiagaraShaderParametersBuilder& ShaderParametersBuilder) const
