- Fix missing initializer for bounds

#rnx

diff --git a/Engine/Plugins/FX/Niagara/Source/Niagara/Public/NiagaraDataChannel_Islands.h b/Engine/Plugins/FX/Niagara/Source/Niagara/Public/NiagaraDataChannel_Islands.h
--- a/Engine/Plugins/FX/Niagara/Source/Niagara/Public/NiagaraDataChannel_Islands.h
+++ b/Engine/Plugins/FX/Niagara/Source/Niagara/Public/NiagaraDataChannel_Islands.h
@@ -63,7 +63,7 @@ private:
 
 	/** Current bounds of this island. The bounds of any handler systems are modified to match these bounds. */
 	UPROPERTY()
-	FBoxSphereBounds Bounds;
+	FBoxSphereBounds Bounds = FBoxSphereBounds(EForceInit::ForceInit);
 
 	/** Niagara components spawned for this island. */
 	UPROPERTY()
