Fix for GC ensure when a script compilation is active

diff --git a/Engine/Plugins/FX/Niagara/Source/Niagara/Private/NiagaraAsyncCompile.h b/Engine/Plugins/FX/Niagara/Source/Niagara/Private/NiagaraAsyncCompile.h
--- a/Engine/Plugins/FX/Niagara/Source/Niagara/Private/NiagaraAsyncCompile.h
+++ b/Engine/Plugins/FX/Niagara/Source/Niagara/Private/NiagaraAsyncCompile.h
@@ -118,6 +118,10 @@ class FNiagaraActiveCompilationDefault : public FNiagaraActiveCompilation
 {
 public:
 	virtual void AddReferencedObjects(FReferenceCollector& Collector) override;
+	virtual FString GetReferencerName() const override
+	{
+		return TEXT("NiagaraActiveCompilationDefault");
+	}
 
 	virtual bool Launch(const FNiagaraCompilationOptions& Options) override;
 	virtual void Abort() override;
