Adding total num particles for a system instance to the engine provided system constants.
Used to allow systems to deactivate when it's unused for N seconds while allowing all emitters to stay alive listening for external input.

diff --git a/Engine/Plugins/FX/Niagara/Source/Niagara/Classes/NiagaraConstants.h b/Engine/Plugins/FX/Niagara/Source/Niagara/Classes/NiagaraConstants.h
--- a/Engine/Plugins/FX/Niagara/Source/Niagara/Classes/NiagaraConstants.h
+++ b/Engine/Plugins/FX/Niagara/Source/Niagara/Classes/NiagaraConstants.h
@@ -75,6 +75,7 @@
 #define SYS_PARAM_ENGINE_SYSTEM_CURRENT_TIME_STEP		 INiagaraModule::GetVar_Engine_System_CurrentTimeStep()
 #define SYS_PARAM_ENGINE_SYSTEM_NUM_TIME_STEPS			 INiagaraModule::GetVar_Engine_System_NumTimeSteps()
 #define SYS_PARAM_ENGINE_SYSTEM_TIME_STEP_FRACTION			 INiagaraModule::GetVar_Engine_System_TimeStepFraction()
+#define SYS_PARAM_ENGINE_SYSTEM_NUM_PARTICLES			INiagaraModule::GetVar_Engine_System_NumParticles()
 
 
 #define SYS_PARAM_ENGINE_SYSTEM_NUM_EMITTERS             INiagaraModule::GetVar_Engine_System_NumEmitters()
diff --git a/Engine/Plugins/FX/Niagara/Source/Niagara/Private/NiagaraConstants.cpp b/Engine/Plugins/FX/Niagara/Source/Niagara/Private/NiagaraConstants.cpp
--- a/Engine/Plugins/FX/Niagara/Source/Niagara/Private/NiagaraConstants.cpp
+++ b/Engine/Plugins/FX/Niagara/Source/Niagara/Private/NiagaraConstants.cpp
@@ -118,6 +118,7 @@ void FNiagaraConstants::Init()
 		SystemParameters.Add(SYS_PARAM_ENGINE_SYSTEM_CURRENT_TIME_STEP);
 		SystemParameters.Add(SYS_PARAM_ENGINE_SYSTEM_NUM_TIME_STEPS);
 		SystemParameters.Add(SYS_PARAM_ENGINE_SYSTEM_TIME_STEP_FRACTION);
+		SystemParameters.Add(SYS_PARAM_ENGINE_SYSTEM_NUM_PARTICLES);
 		SystemParameters.Add(SYS_PARAM_ENGINE_SYSTEM_NUM_EMITTERS);
 		SystemParameters.Add(SYS_PARAM_ENGINE_NUM_SYSTEM_INSTANCES);
 		SystemParameters.Add(SYS_PARAM_ENGINE_GLOBAL_SPAWN_COUNT_SCALE);
@@ -261,6 +262,7 @@ void FNiagaraConstants::Init()
 		SystemStrMap.Add(SYS_PARAM_ENGINE_SYSTEM_CURRENT_TIME_STEP, LOCTEXT("CurrentTimeStep", "The current time step."));
 		SystemStrMap.Add(SYS_PARAM_ENGINE_SYSTEM_NUM_TIME_STEPS, LOCTEXT("NumTimeSteps", "Number of steps used, can be > 1 when using fixed delta time."));
 		SystemStrMap.Add(SYS_PARAM_ENGINE_SYSTEM_TIME_STEP_FRACTION, LOCTEXT("TimeStepFraction", "The fraction of the world delta time integrated for the current sub step."));
+		SystemStrMap.Add(SYS_PARAM_ENGINE_SYSTEM_NUM_PARTICLES, LOCTEXT("NumParticles", "The total active particles for all emitters in this system. Note that counts for GPU emitters will be latent and may not be exactly accurate/up to date for a give frame."));
 
 		SystemStrMap.Add(SYS_PARAM_ENGINE_SYSTEM_NUM_EMITTERS, LOCTEXT("SystemNumEmitters", "The number of emitters attached to this system. Should only be used in System scripts."));
 		SystemStrMap.Add(SYS_PARAM_ENGINE_NUM_SYSTEM_INSTANCES, LOCTEXT("SystemNumInstances", "The number of instances of this system currently ticking. Should only be used in System scripts."));
diff --git a/Engine/Plugins/FX/Niagara/Source/Niagara/Private/NiagaraModule.cpp b/Engine/Plugins/FX/Niagara/Source/Niagara/Private/NiagaraModule.cpp
--- a/Engine/Plugins/FX/Niagara/Source/Niagara/Private/NiagaraModule.cpp
+++ b/Engine/Plugins/FX/Niagara/Source/Niagara/Private/NiagaraModule.cpp
@@ -195,6 +195,7 @@ FNiagaraVariable INiagaraModule::Engine_System_RandomSeed;
 FNiagaraVariable INiagaraModule::Engine_System_CurrentTimeStep;
 FNiagaraVariable INiagaraModule::Engine_System_NumTimeSteps;
 FNiagaraVariable INiagaraModule::Engine_System_TimeStepFraction;
+FNiagaraVariable INiagaraModule::Engine_System_NumParticles;
 
 FNiagaraVariable INiagaraModule::Engine_System_NumEmitters;
 FNiagaraVariable INiagaraModule::Engine_NumSystemInstances;
@@ -349,6 +350,7 @@ void INiagaraModule::StartupModule()
 	Engine_System_CurrentTimeStep = FNiagaraVariable(FNiagaraTypeDefinition::GetIntDef(), TEXT("Engine.System.CurrentTimeStep"));
 	Engine_System_NumTimeSteps = FNiagaraVariable(FNiagaraTypeDefinition::GetIntDef(), TEXT("Engine.System.NumTimeSteps"));
 	Engine_System_TimeStepFraction = FNiagaraVariable(FNiagaraTypeDefinition::GetFloatDef(), TEXT("Engine.System.TimeStepFraction"));
+	Engine_System_NumParticles = FNiagaraVariable(FNiagaraTypeDefinition::GetIntDef(), TEXT("Engine.System.NumParticles"));
 
 	Engine_System_NumEmitters = FNiagaraVariable(FNiagaraTypeDefinition::GetIntDef(), TEXT("Engine.System.NumEmitters"));
 	Engine_NumSystemInstances = FNiagaraVariable(FNiagaraTypeDefinition::GetIntDef(), TEXT("Engine.NumSystemInstances"));
@@ -1681,7 +1683,6 @@ const TArray<FNiagaraVariable>& FNiagaraSystemParameters::GetVariables()
 {
 	static const FName NAME_NiagaraStructPadding0 = "Engine.System.PaddingInt32_0";
 	static const FName NAME_NiagaraStructPadding1 = "Engine.System.PaddingInt32_1";
-	static const FName NAME_NiagaraStructPadding2 = "Engine.System.PaddingInt32_2";
 
 	static const TArray<FNiagaraVariable> Variables =
 	{
@@ -1698,9 +1699,10 @@ const TArray<FNiagaraVariable>& FNiagaraSystemParameters::GetVariables()
 		SYS_PARAM_ENGINE_SYSTEM_CURRENT_TIME_STEP,
 		SYS_PARAM_ENGINE_SYSTEM_NUM_TIME_STEPS,
 		SYS_PARAM_ENGINE_SYSTEM_TIME_STEP_FRACTION,			
+		SYS_PARAM_ENGINE_SYSTEM_NUM_PARTICLES,
+
 		FNiagaraVariable(FNiagaraTypeDefinition::GetIntDef(), NAME_NiagaraStructPadding0),
 		FNiagaraVariable(FNiagaraTypeDefinition::GetIntDef(), NAME_NiagaraStructPadding1),
-		FNiagaraVariable(FNiagaraTypeDefinition::GetIntDef(), NAME_NiagaraStructPadding2),
 	};
 
 	return Variables;
diff --git a/Engine/Plugins/FX/Niagara/Source/Niagara/Private/NiagaraSystemInstance.cpp b/Engine/Plugins/FX/Niagara/Source/Niagara/Private/NiagaraSystemInstance.cpp
--- a/Engine/Plugins/FX/Niagara/Source/Niagara/Private/NiagaraSystemInstance.cpp
+++ b/Engine/Plugins/FX/Niagara/Source/Niagara/Private/NiagaraSystemInstance.cpp
@@ -1937,6 +1937,22 @@ void FNiagaraSystemInstance::TickInstanceParameters_GameThread(float DeltaSecond
 	FlipParameterBuffers();
 	uint32 ParameterIndex = GetParameterIndex();
 
+	FNiagaraSystemParameters& CurrentSystemParameters = SystemParameters[ParameterIndex];
+	CurrentSystemParameters.EngineSystemAge = Age;
+	CurrentSystemParameters.EngineTickCount = TickCount;
+	CurrentSystemParameters.EngineTimeSinceRendered = FMath::Max(0.0f, GatheredInstanceParameters.TimeSeconds - LastRenderTime - GLastRenderTimeSafetyBias);
+	CurrentSystemParameters.EngineExecutionState = static_cast<uint32>(RequestedExecutionState);
+	CurrentSystemParameters.EngineLodDistance = GetLODDistance();
+	CurrentSystemParameters.EngineLodDistanceFraction = CurrentSystemParameters.EngineLodDistance / MaxLODDistance;
+	CurrentSystemParameters.SignificanceIndex = SignificanceIndex;
+	CurrentSystemParameters.RandomSeed = RandomSeed + RandomSeedOffset;
+
+	CurrentSystemParameters.CurrentTimeStep = GetSystemSimulation()->GetTickInfo().TickNumber;
+	CurrentSystemParameters.NumTimeSteps = GetSystemSimulation()->GetTickInfo().TickCount;
+	CurrentSystemParameters.TimeStepFraction = GetSystemSimulation()->GetTickInfo().TimeStepFraction;
+	
+	CurrentSystemParameters.NumParticles = 0;
+
 	for (int32 i = 0; i < GatheredInstanceParameters.EmitterCount; ++i)
 	{
 		auto& CurrentEmitterParameters = EditEmitterParameters(i);
@@ -1944,7 +1960,10 @@ void FNiagaraSystemInstance::TickInstanceParameters_GameThread(float DeltaSecond
 		const TSharedRef<FNiagaraEmitterInstance, ESPMode::ThreadSafe>& Emitter = Emitters[i];
 		if (Emitter->GetExecutionState() != ENiagaraExecutionState::Disabled)
 		{
-			CurrentEmitterParameters.EmitterNumParticles = Emitter->GetNumParticles();
+			int32 NumParticles = Emitter->GetNumParticles();
+			CurrentSystemParameters.NumParticles += NumParticles;
+
+			CurrentEmitterParameters.EmitterNumParticles = NumParticles;
 			CurrentEmitterParameters.EmitterTotalSpawnedParticles = Emitter->GetTotalSpawnedParticles();
 			CurrentEmitterParameters.EmitterRandomSeed = Emitter->GetRandomSeed();
 			CurrentEmitterParameters.EmitterInstanceSeed = Emitter->GetInstanceSeed();
@@ -1958,20 +1977,6 @@ void FNiagaraSystemInstance::TickInstanceParameters_GameThread(float DeltaSecond
 		}
 	}
 
-	FNiagaraSystemParameters& CurrentSystemParameters = SystemParameters[ParameterIndex];
-	CurrentSystemParameters.EngineSystemAge = Age;
-	CurrentSystemParameters.EngineTickCount = TickCount;
-	CurrentSystemParameters.EngineTimeSinceRendered = FMath::Max(0.0f, GatheredInstanceParameters.TimeSeconds - LastRenderTime - GLastRenderTimeSafetyBias);
-	CurrentSystemParameters.EngineExecutionState = static_cast<uint32>(RequestedExecutionState);
-	CurrentSystemParameters.EngineLodDistance = GetLODDistance();
-	CurrentSystemParameters.EngineLodDistanceFraction = CurrentSystemParameters.EngineLodDistance / MaxLODDistance;
-	CurrentSystemParameters.SignificanceIndex = SignificanceIndex;
-	CurrentSystemParameters.RandomSeed = RandomSeed + RandomSeedOffset;
-
-	CurrentSystemParameters.CurrentTimeStep = GetSystemSimulation()->GetTickInfo().TickNumber;
-	CurrentSystemParameters.NumTimeSteps = GetSystemSimulation()->GetTickInfo().TickCount;
-	CurrentSystemParameters.TimeStepFraction = GetSystemSimulation()->GetTickInfo().TimeStepFraction;	
-
 	if (OverrideParameters)
 	{
 		OverrideParameters->ResolvePositions(GetLWCConverter());
diff --git a/Engine/Plugins/FX/Niagara/Source/Niagara/Public/NiagaraModule.h b/Engine/Plugins/FX/Niagara/Source/Niagara/Public/NiagaraModule.h
--- a/Engine/Plugins/FX/Niagara/Source/Niagara/Public/NiagaraModule.h
+++ b/Engine/Plugins/FX/Niagara/Source/Niagara/Public/NiagaraModule.h
@@ -170,6 +170,7 @@ public:
 	FORCEINLINE static const FNiagaraVariable& GetVar_Engine_System_CurrentTimeStep() { return Engine_System_CurrentTimeStep; }
 	FORCEINLINE static const FNiagaraVariable& GetVar_Engine_System_NumTimeSteps() { return Engine_System_NumTimeSteps; }
 	FORCEINLINE static const FNiagaraVariable& GetVar_Engine_System_TimeStepFraction() { return Engine_System_TimeStepFraction; }
+	FORCEINLINE static const FNiagaraVariable& GetVar_Engine_System_NumParticles() { return Engine_System_NumParticles; }
 
 	FORCEINLINE static const FNiagaraVariable&  GetVar_Engine_System_NumEmitters() { return Engine_System_NumEmitters; }
 	FORCEINLINE static const FNiagaraVariable&  GetVar_Engine_NumSystemInstances() { return Engine_NumSystemInstances; }
@@ -298,6 +299,7 @@ private:
 	static FNiagaraVariable Engine_System_CurrentTimeStep;
 	static FNiagaraVariable Engine_System_NumTimeSteps;
 	static FNiagaraVariable Engine_System_TimeStepFraction;
+	static FNiagaraVariable Engine_System_NumParticles;
 	
 	static FNiagaraVariable Engine_System_NumEmitters;
 	static FNiagaraVariable Engine_NumSystemInstances;
diff --git a/Engine/Plugins/FX/Niagara/Source/Niagara/Public/NiagaraTypes.h b/Engine/Plugins/FX/Niagara/Source/Niagara/Public/NiagaraTypes.h
--- a/Engine/Plugins/FX/Niagara/Source/Niagara/Public/NiagaraTypes.h
+++ b/Engine/Plugins/FX/Niagara/Source/Niagara/Public/NiagaraTypes.h
@@ -2104,9 +2104,10 @@ struct alignas(16) FNiagaraSystemParameters
 	int32 NumTimeSteps = 0;
 	float TimeStepFraction = 0.0f;
 
+	uint32 NumParticles = 0;
+
 	int32 _Pad0;
 	int32 _Pad1;
-	int32 _Pad2;
 };
 
 // Any change to this structure, or it's GetVariables implementation will require a bump in the CustomNiagaraVersion so that we
diff --git a/Engine/Plugins/FX/Niagara/Source/NiagaraEditor/Private/NiagaraHlslTranslator.cpp b/Engine/Plugins/FX/Niagara/Source/NiagaraEditor/Private/NiagaraHlslTranslator.cpp
--- a/Engine/Plugins/FX/Niagara/Source/NiagaraEditor/Private/NiagaraHlslTranslator.cpp
+++ b/Engine/Plugins/FX/Niagara/Source/NiagaraEditor/Private/NiagaraHlslTranslator.cpp
@@ -1071,6 +1071,7 @@ static void ConvertFloatToHalf(const FNiagaraCompileOptions& InCompileOptions, T
 			SYS_PARAM_ENGINE_SYSTEM_RANDOM_SEED,
 			SYS_PARAM_ENGINE_SYSTEM_NUM_TIME_STEPS,
 			SYS_PARAM_ENGINE_SYSTEM_TIME_STEP_FRACTION,
+			SYS_PARAM_ENGINE_SYSTEM_NUM_PARTICLES,
 			SYS_PARAM_ENGINE_SYSTEM_NUM_EMITTERS,
 			SYS_PARAM_ENGINE_NUM_SYSTEM_INSTANCES,
 			SYS_PARAM_ENGINE_EMITTER_NUM_PARTICLES,
@@ -4535,6 +4536,7 @@ bool FHlslNiagaraTranslator::ShouldInterpolateParameter(const FNiagaraVariable&
 		Parameter == SYS_PARAM_EMITTER_RANDOM_SEED ||
 		Parameter == SYS_PARAM_ENGINE_SYSTEM_NUM_TIME_STEPS ||
 		Parameter == SYS_PARAM_ENGINE_SYSTEM_TIME_STEP_FRACTION ||
+		Parameter == SYS_PARAM_ENGINE_SYSTEM_NUM_PARTICLES ||
 		Parameter == SYS_PARAM_ENGINE_SYSTEM_NUM_EMITTERS ||
 		Parameter == SYS_PARAM_ENGINE_EMITTER_INSTANCE_SEED ||
 		Parameter == SYS_PARAM_ENGINE_SYSTEM_TICK_COUNT ||
diff --git a/Engine/Source/Runtime/Core/Private/UObject/FortniteMainBranchObjectVersion.cpp b/Engine/Source/Runtime/Core/Private/UObject/FortniteMainBranchObjectVersion.cpp
--- a/Engine/Source/Runtime/Core/Private/UObject/FortniteMainBranchObjectVersion.cpp
+++ b/Engine/Source/Runtime/Core/Private/UObject/FortniteMainBranchObjectVersion.cpp
@@ -10,8 +10,8 @@ TMap<FGuid, FGuid> FFortniteMainBranchObjectVersion::GetSystemGuids()
 	SystemGuids.Add(DevGuids.LANDSCAPE_MOBILE_COOK_VERSION, FGuid("32D02EF867C74B71A0D4E0FA41392732"));
 	SystemGuids.Add(DevGuids.MATERIALSHADERMAP_DERIVEDDATA_VER, FGuid("CE7776E63E9E43FBB8C103EE1AB40B5D"));
 	SystemGuids.Add(DevGuids.NANITE_DERIVEDDATA_VER, FGuid("D28C003F825F4C178F698938FB0E47F7"));
-	SystemGuids.Add(DevGuids.NIAGARASHADERMAP_DERIVEDDATA_VER, FGuid("35D7960B72164B49B9B7B3FB5857E674"));
-	SystemGuids.Add(DevGuids.Niagara_LatestScriptCompileVersion, FGuid("3DC7587E76B74CC3BA6FAD6B2F5EC31A"));
+	SystemGuids.Add(DevGuids.NIAGARASHADERMAP_DERIVEDDATA_VER, FGuid("A6FFCB143BBC4C83A72001468D23D5A4"));
+	SystemGuids.Add(DevGuids.Niagara_LatestScriptCompileVersion, FGuid("BEAE0A1A932F433CBA00AC233F085F22"));
 	SystemGuids.Add(DevGuids.SkeletalMeshDerivedDataVersion, FGuid("518D7C51D31D9BE4A9F4B81542AC8119"));
 	SystemGuids.Add(DevGuids.STATICMESH_DERIVEDDATA_VER, FGuid("543843B103794E0E9BA4BE80FB602F79"));
 	return SystemGuids;
