- Basic ISM support added to SM DI
- Allows you to specify and change which instance to read data from using either the BP or VVM API

diff --git a/Engine/Plugins/FX/Niagara/Source/Niagara/Internal/DataInterface/NiagaraDataInterfaceStaticMesh.h b/Engine/Plugins/FX/Niagara/Source/Niagara/Internal/DataInterface/NiagaraDataInterfaceStaticMesh.h
--- a/Engine/Plugins/FX/Niagara/Source/Niagara/Internal/DataInterface/NiagaraDataInterfaceStaticMesh.h
+++ b/Engine/Plugins/FX/Niagara/Source/Niagara/Internal/DataInterface/NiagaraDataInterfaceStaticMesh.h
@@ -110,6 +110,10 @@ public:
 	UPROPERTY(EditAnywhere, Category = "LOD")
 	FNiagaraUserParameterBinding LODIndexUserParameter;
 
+	/** When attached to an Instanced Static Mesh, which instance should be read from. */
+	UPROPERTY(EditAnywhere, Category = "Mesh")
+	int32 InstanceIndex = INDEX_NONE;
+
 	/** List of filtered sockets to use. */
 	UPROPERTY(EditAnywhere, Category = "Mesh")
 	TArray<FName> FilteredSockets;
@@ -180,6 +184,9 @@ public:
 	void SetSourceComponentFromBlueprints(UStaticMeshComponent* ComponentToUse);
 	void SetDefaultMeshFromBlueprints(UStaticMesh* MeshToUse);
 
+	UFUNCTION(BlueprintCallable, Category = Niagara, meta = (DisplayName = "Set Niagara Static Mesh DI Instance Index"))
+	static void SetNiagaraStaticMeshDIInstanceIndex(UNiagaraComponent* NiagaraSystem, const FName UserParameterName, int32 NewInstanceIndex);
+
 protected:
 	// Bind/unbind delegates to release references to the source actor & component.
 	void UnbindSourceDelegates();
@@ -277,6 +284,8 @@ protected:
 	void VMGetLocalToWorld(FVectorVMExternalFunctionContext& Context);
 	void VMGetLocalToWorldInverseTransposed(FVectorVMExternalFunctionContext& Context);
 	void VMGetWorldVelocity(FVectorVMExternalFunctionContext& Context);
+	void VMGetInstanceIndex(FVectorVMExternalFunctionContext& Context);
+	void VMSetInstanceIndex(FVectorVMExternalFunctionContext& Context);
 
 	// VM UV mapping functions
 	void VMGetTriangleCoordAtUV(FVectorVMExternalFunctionContext& Context);
diff --git a/Engine/Plugins/FX/Niagara/Source/Niagara/Private/DataInterface/NiagaraDataInterfaceStaticMesh.cpp b/Engine/Plugins/FX/Niagara/Source/Niagara/Private/DataInterface/NiagaraDataInterfaceStaticMesh.cpp
--- a/Engine/Plugins/FX/Niagara/Source/Niagara/Private/DataInterface/NiagaraDataInterfaceStaticMesh.cpp
+++ b/Engine/Plugins/FX/Niagara/Source/Niagara/Private/DataInterface/NiagaraDataInterfaceStaticMesh.cpp
@@ -2,6 +2,7 @@
 
 #include "DataInterface/NiagaraDataInterfaceStaticMesh.h"
 #include "Components/StaticMeshComponent.h"
+#include "Components/InstancedStaticMeshComponent.h"
 #include "NiagaraStats.h"
 #include "Containers/StridedView.h"
 #include "NiagaraSystemInstance.h"
@@ -10,6 +11,7 @@
 #include "Misc/LargeWorldRenderPosition.h"
 #include "NiagaraSettings.h"
 #include "NiagaraCompileHashVisitor.h"
+#include "NiagaraFunctionLibrary.h"
 #include "NiagaraShaderParametersBuilder.h"
 #include "NiagaraDataInterfaceMeshCommon.h"
 #include "NiagaraWorldManager.h"
@@ -225,6 +227,9 @@ namespace NDIStaticMeshLocal
 	static const FName	GetLocalToWorldInverseTransposedName("GetLocalToWorldInverseTransposed");
 	static const FName	GetWorldVelocityName("GetWorldVelocity");
 
+	static const FName	GetInstanceIndexName("GetInstanceIndex");
+	static const FName	SetInstanceIndexName("SetInstanceIndex");
+
 	//////////////////////////////////////////////////////////////////////////
 	// VM UV mapping functions
 	static const FName	GetTriangleCoordAtUVName("GetTriangleCoordAtUV");
@@ -580,6 +585,9 @@ namespace NDIStaticMeshLocal
 		/** Vector from system owners instance location to the sampled mesh location. */
 		FVector3f OwnerToMeshVector = FVector3f::ZeroVector;
 
+		/** Which instance index we are reading from. */
+		int32 InstanceIndex = INDEX_NONE;
+
 		/** Cached ComponentToWorld. (Falls back to WorldTransform of the system instance) */
 		FMatrix Transform;
 		/** InverseTranspose of above for transforming normals/tangents. */
@@ -590,6 +598,15 @@ namespace NDIStaticMeshLocal
 		/** Cached Previous InverseTranspose of above for transforming normals/tangents. */
 		FMatrix PrevTransformInverseTransposed;
 
+		/** When attached to an ISM contains all the transforms, so safe to use concurrently */
+		TArray<FTransform> ISMTransforms;
+
+		/** Cached SystemInstance World Transform */
+		FTransform SystemInstanceWorldTransform = FTransform::Identity;
+
+		/** Cached LWC tile offset */
+		FVector LWCTileOffset = FVector::ZeroVector;
+
 		/** Cached Rotation. */
 		FQuat4f Rotation;
 		/** Cached Previous Rotation. */
@@ -721,6 +738,80 @@ namespace NDIStaticMeshLocal
 		}
 #endif
 
+		void UpdateTransforms(USceneComponent* SceneComponent, FNiagaraSystemInstance* SystemInstance)
+		{
+			SystemInstanceWorldTransform = SystemInstance->GetWorldTransform();
+
+			FTransform ComponentTransform;
+			if (SceneComponent)
+			{
+				ComponentTransform = SceneComponent->GetComponentToWorld();
+				if (UInstancedStaticMeshComponent* ISMComponent = Cast<UInstancedStaticMeshComponent>(SceneComponent))
+				{
+					const int32 NumInstances = ISMComponent->PerInstanceSMData.Num();
+					ISMTransforms.SetNum(NumInstances);
+					for (int32 i = 0; i < NumInstances; ++i)
+					{
+						ISMComponent->GetInstanceTransform(i, ISMTransforms[i], true);
+					}
+
+					if (ISMTransforms.IsValidIndex(InstanceIndex))
+					{
+						ComponentTransform = ISMTransforms[InstanceIndex];
+					}
+				}
+				else
+				{
+					ISMTransforms.Empty();
+				}
+			}
+			else
+			{
+				ISMTransforms.Empty();
+				ComponentTransform = SystemInstanceWorldTransform;
+			}
+
+			OwnerToMeshVector = FVector3f(ComponentTransform.GetLocation() - SystemInstanceWorldTransform.GetLocation());
+			LWCTileOffset = FVector(SystemInstance->GetLWCTile()) * -FLargeWorldRenderScalar::GetTileSize();
+			ComponentTransform.AddToTranslation(LWCTileOffset);
+
+			Transform = ComponentTransform.ToMatrixWithScale();
+			TransformInverseTransposed = ComponentTransform.Inverse().ToMatrixWithScale().GetTransposed();
+			Rotation = FQuat4f(ComponentTransform.GetRotation());
+		}
+
+		void UpdateTransformsResetPrevious(USceneComponent* SceneComponent, FNiagaraSystemInstance* SystemInstance)
+		{
+			UpdateTransforms(SceneComponent, SystemInstance);
+			PrevTransform = Transform;
+			PrevTransformInverseTransposed = TransformInverseTransposed;
+			PrevRotation = Rotation;
+		}
+
+		void UpdateTransformsCopyPrevious(USceneComponent* SceneComponent, FNiagaraSystemInstance* SystemInstance)
+		{
+			PrevTransform = Transform;
+			PrevTransformInverseTransposed = TransformInverseTransposed;
+			PrevRotation = Rotation;
+			UpdateTransforms(SceneComponent, SystemInstance);
+		}
+
+		void UpdateTransformsInstanceIndexChanged()
+		{
+			FTransform ComponentTransform = ISMTransforms.IsValidIndex(InstanceIndex) ? ISMTransforms[InstanceIndex] : SystemInstanceWorldTransform;
+
+			OwnerToMeshVector = FVector3f(ComponentTransform.GetLocation() - SystemInstanceWorldTransform.GetLocation());
+			ComponentTransform.AddToTranslation(LWCTileOffset);
+
+			Transform = ComponentTransform.ToMatrixWithScale();
+			TransformInverseTransposed = ComponentTransform.Inverse().ToMatrixWithScale().GetTransposed();
+			Rotation = FQuat4f(ComponentTransform.GetRotation());
+
+			PrevTransform = Transform;
+			PrevTransformInverseTransposed = TransformInverseTransposed;
+			PrevRotation = Rotation;
+		}
+
 		bool Init(UNiagaraDataInterfaceStaticMesh* Interface, FNiagaraSystemInstance* SystemInstance)
 		{
 			check(SystemInstance);
@@ -764,18 +855,8 @@ namespace NDIStaticMeshLocal
 			SceneComponentWeakPtr = SceneComponent;
 
 			// Gather attached information
-			bComponentValid = SceneComponent != nullptr;
-			FTransform ComponentTransform = bComponentValid ? SceneComponent->GetComponentToWorld() : SystemInstance->GetWorldTransform();
-			OwnerToMeshVector = FVector3f(ComponentTransform.GetLocation() -  SystemInstance->GetWorldTransform().GetLocation());
-			ComponentTransform.AddToTranslation(FVector(SystemInstance->GetLWCTile()) * -FLargeWorldRenderScalar::GetTileSize());
-
-			Transform = ComponentTransform.ToMatrixWithScale();
-			TransformInverseTransposed = ComponentTransform.Inverse().ToMatrixWithScale().GetTransposed();
-			PrevTransform = Transform;
-			PrevTransformInverseTransposed = ComponentTransform.Inverse().ToMatrixWithScale().GetTransposed();
-
-			Rotation = FQuat4f(ComponentTransform.GetRotation());
-			PrevRotation = Rotation;
+			InstanceIndex = Interface->InstanceIndex;
+			UpdateTransformsResetPrevious(SceneComponent, SystemInstance);
 
 			if (bUsePhysicsVelocity)
 			{
@@ -996,17 +1077,7 @@ namespace NDIStaticMeshLocal
 			DeltaSeconds = InDeltaSeconds;
 
 			USceneComponent* SceneComponent = SceneComponentWeakPtr.Get();
-			FTransform ComponentTransform = SceneComponent != nullptr ? SceneComponent->GetComponentToWorld() : SystemInstance->GetWorldTransform();
-			OwnerToMeshVector = FVector3f(ComponentTransform.GetLocation() - SystemInstance->GetWorldTransform().GetLocation());
-			ComponentTransform.AddToTranslation(FVector(SystemInstance->GetLWCTile()) * -FLargeWorldRenderScalar::GetTileSize());
-
-			PrevTransform = Transform;
-			PrevTransformInverseTransposed = TransformInverseTransposed;
-			Transform = ComponentTransform.ToMatrixWithScale();
-			TransformInverseTransposed = ComponentTransform.Inverse().ToMatrixWithScale().GetTransposed();
-
-			PrevRotation = Rotation;
-			Rotation = FQuat4f(ComponentTransform.GetRotation());
+			UpdateTransformsCopyPrevious(SceneComponent, SystemInstance);
 
 			if (bUsePhysicsVelocity)
 			{
@@ -2225,6 +2296,19 @@ void UNiagaraDataInterfaceStaticMesh::GetMiscFunctions(TArray<FNiagaraFunctionSi
 		Sig.Name = GetWorldVelocityName;
 		Sig.Outputs.Emplace(FNiagaraTypeDefinition::GetVec3Def(), TEXT("Velocity"));
 	}
+	//-TBD: There's not really a benefit to getting this in the graph currently
+	//{
+	//	FNiagaraFunctionSignature& Sig = OutFunctions.Add_GetRef(BaseSignature);
+	//	Sig.Name = GetInstanceIndexName;
+	//	Sig.Outputs.Emplace(FNiagaraTypeDefinition::GetIntDef(), TEXT("Index"));
+	//}
+	{
+		FNiagaraFunctionSignature& Sig = OutFunctions.Add_GetRef(BaseSignature);
+		Sig.Name = SetInstanceIndexName;
+		Sig.Inputs.Emplace(FNiagaraTypeDefinition::GetIntDef(), TEXT("Index"));
+		Sig.bRequiresExecPin = true;
+		Sig.bSupportsGPU = false;
+	}
 }
 
 void UNiagaraDataInterfaceStaticMesh::GetUVMappingFunctions(TArray<FNiagaraFunctionSignature>& OutFunctions, const FNiagaraFunctionSignature& BaseSignature) const
@@ -2784,6 +2868,14 @@ void UNiagaraDataInterfaceStaticMesh::GetVMExternalFunction(const FVMExternalFun
 	{
 		OutFunc = FVMExternalFunction::CreateLambda([this](FVectorVMExternalFunctionContext& Context) { VMGetWorldVelocity(Context); });
 	}
+	else if (BindingInfo.Name == GetInstanceIndexName)
+	{
+		OutFunc = FVMExternalFunction::CreateLambda([this](FVectorVMExternalFunctionContext& Context) { VMGetInstanceIndex(Context); });
+	}
+	else if (BindingInfo.Name == SetInstanceIndexName)
+	{
+		OutFunc = FVMExternalFunction::CreateLambda([this](FVectorVMExternalFunctionContext& Context) { VMSetInstanceIndex(Context); });
+	}
 
 	//////////////////////////////////////////////////////////////////////////
 	// VM UV mapping functions
@@ -2982,7 +3074,12 @@ bool UNiagaraDataInterfaceStaticMesh::GetFunctionHLSL(const FNiagaraDataInterfac
 		OutHLSL += FString::Format(FormatSample, ArgsSample);
 	}
 
-	return true;
+	static const TSet<FName> UnsupportedGpuFunctions =
+	{
+		SetInstanceIndexName,
+	};
+
+	return UnsupportedGpuFunctions.Contains(FunctionInfo.DefinitionName) ? false : true;
 }
 #endif
 
@@ -3050,6 +3147,9 @@ void UNiagaraDataInterfaceStaticMesh::SetShaderParameters(const FNiagaraDataInte
 	ShaderParameters->SocketTransforms	= InstanceData.SocketTransforms.SRV.IsValid() ? InstanceData.SocketTransforms.SRV.GetReference() : FNiagaraRenderer::GetDummyFloat4Buffer();
 	ShaderParameters->FilteredAndUnfilteredSockets = InstanceData.FilteredAndUnfilteredSockets.SRV.IsValid() ? InstanceData.FilteredAndUnfilteredSockets.SRV.GetReference() : FNiagaraRenderer::GetDummyUIntBuffer();
 
+	ShaderParameters->SectionInfos = InstanceData.SectionInfos.SRV.IsValid() ? InstanceData.SectionInfos.SRV.GetReference() : FNiagaraRenderer::GetDummyUInt4Buffer();
+	ShaderParameters->SocketTransforms	= InstanceData.SocketTransforms.SRV.IsValid() ? InstanceData.SocketTransforms.SRV.GetReference() : FNiagaraRenderer::GetDummyFloat4Buffer();
+
 	// Set misc data
 	const float InvDeltaTime = InstanceData.DeltaSeconds > 0.0f ? 1.0f / InstanceData.DeltaSeconds : 0.0f;
 	const FVector3f DeltaPosition = InstanceData.Transform.GetOrigin() - InstanceData.PrevTransform.GetOrigin();
@@ -3144,6 +3244,7 @@ bool UNiagaraDataInterfaceStaticMesh::Equals(const UNiagaraDataInterface* Other)
 		OtherTyped->bAllowSamplingFromStreamingLODs == bAllowSamplingFromStreamingLODs &&
 		OtherTyped->LODIndex == LODIndex &&
 		OtherTyped->LODIndexUserParameter == LODIndexUserParameter &&
+		OtherTyped->InstanceIndex == InstanceIndex &&
 		OtherTyped->FilteredSockets == FilteredSockets;
 }
 
@@ -3169,6 +3270,7 @@ bool UNiagaraDataInterfaceStaticMesh::CopyToInternal(UNiagaraDataInterface* Dest
 	OtherTyped->bAllowSamplingFromStreamingLODs = bAllowSamplingFromStreamingLODs;
 	OtherTyped->LODIndex = LODIndex;
 	OtherTyped->LODIndexUserParameter = LODIndexUserParameter;
+	OtherTyped->InstanceIndex = InstanceIndex;
 	OtherTyped->BindSourceDelegates();
 	return true;
 }
@@ -3417,6 +3519,18 @@ void UNiagaraDataInterfaceStaticMesh::SetSourceComponentFromBlueprints(UStaticMe
 	BindSourceDelegates();
 }
 
+void UNiagaraDataInterfaceStaticMesh::SetNiagaraStaticMeshDIInstanceIndex(UNiagaraComponent* NiagaraSystem, const FName UserParameterName, int32 NewInstanceIndex)
+{
+	if (UNiagaraDataInterfaceStaticMesh* DataInterface = UNiagaraFunctionLibrary::GetDataInterface<UNiagaraDataInterfaceStaticMesh>(NiagaraSystem, UserParameterName))
+	{
+		if (DataInterface->InstanceIndex != NewInstanceIndex)
+		{
+			++DataInterface->ChangeId;
+			DataInterface->InstanceIndex = NewInstanceIndex;
+		}
+	}
+}
+
 void UNiagaraDataInterfaceStaticMesh::BindSourceDelegates()
 {
 	if (AActor* Source = SoftSourceActor.Get())
@@ -4900,6 +5014,35 @@ void UNiagaraDataInterfaceStaticMesh::VMGetWorldVelocity(FVectorVMExternalFuncti
 	}
 }
 
+void UNiagaraDataInterfaceStaticMesh::VMGetInstanceIndex(FVectorVMExternalFunctionContext& Context)
+{
+	VectorVM::FUserPtrHandler<NDIStaticMeshLocal::FInstanceData_GameThread> InstanceData(Context);
+	FNDIOutputParam<int32> OutInstanceIndex(Context);
+
+	for (int32 i = 0; i < Context.GetNumInstances(); ++i)
+	{
+		OutInstanceIndex.SetAndAdvance(InstanceData->InstanceIndex);
+	}
+}
+
+void UNiagaraDataInterfaceStaticMesh::VMSetInstanceIndex(FVectorVMExternalFunctionContext& Context)
+{
+	VectorVM::FUserPtrHandler<NDIStaticMeshLocal::FInstanceData_GameThread> InstanceData(Context);
+	FNDIInputParam<int32> InInstanceIndex(Context);
+
+	bool bUpdateTransforms = false;
+	for (int32 i = 0; i < Context.GetNumInstances(); ++i)
+	{
+		const int32 NewInstanceIndex = InInstanceIndex.GetAndAdvance();
+		bUpdateTransforms |= InstanceData->InstanceIndex != NewInstanceIndex;
+		InstanceData->InstanceIndex = NewInstanceIndex;
+	}
+	if (bUpdateTransforms)
+	{
+		InstanceData->UpdateTransformsInstanceIndexChanged();
+	}
+}
+
 //////////////////////////////////////////////////////////////////////////
 // VM UV mapping functions
 void UNiagaraDataInterfaceStaticMesh::VMGetTriangleCoordAtUV(FVectorVMExternalFunctionContext& Context)
