Attribute Spreadsheet menu entry in the Niagara Debug menu will now open the new sim cache based attribute spreadsheet.

diff --git a/Engine/Plugins/FX/Niagara/Source/NiagaraEditor/Private/Toolkits/NiagaraSystemToolkit.cpp b/Engine/Plugins/FX/Niagara/Source/NiagaraEditor/Private/Toolkits/NiagaraSystemToolkit.cpp
--- a/Engine/Plugins/FX/Niagara/Source/NiagaraEditor/Private/Toolkits/NiagaraSystemToolkit.cpp
+++ b/Engine/Plugins/FX/Niagara/Source/NiagaraEditor/Private/Toolkits/NiagaraSystemToolkit.cpp
@@ -916,7 +916,7 @@ void FNiagaraSystemToolkit::OpenDebugOutliner()
 
 void FNiagaraSystemToolkit::OpenAttributeSpreadsheet()
 {
-	InvokeTab(FNiagaraSystemToolkitModeBase::DebugSpreadsheetTabID);
+	InvokeTab(FNiagaraSystemToolkitModeBase::DebugCacheSpreadsheetTabID);
 }
 
 
