Partial implementation GetResourceSizeEx for NiagaraComponent and NiagaraSystem.

#tests PIE

diff --git a/Engine/Plugins/FX/Niagara/Source/Niagara/Classes/NiagaraSystem.h b/Engine/Plugins/FX/Niagara/Source/Niagara/Classes/NiagaraSystem.h
--- a/Engine/Plugins/FX/Niagara/Source/Niagara/Classes/NiagaraSystem.h
+++ b/Engine/Plugins/FX/Niagara/Source/Niagara/Classes/NiagaraSystem.h
@@ -228,6 +228,7 @@ public:
 	virtual void Serialize(FArchive& Ar) override;
 	virtual void PostLoad() override; 
 	virtual bool CanBeClusterRoot() const override { return true; }
+	virtual void GetResourceSizeEx(FResourceSizeEx& CumulativeResourceSize) override;
 #if WITH_EDITORONLY_DATA
 	static void DeclareConstructClasses(TArray<FTopLevelAssetPath>& OutConstructClasses, const UClass* SpecificSubclass);
 #endif
diff --git a/Engine/Plugins/FX/Niagara/Source/Niagara/Private/NiagaraComponent.cpp b/Engine/Plugins/FX/Niagara/Source/Niagara/Private/NiagaraComponent.cpp
--- a/Engine/Plugins/FX/Niagara/Source/Niagara/Private/NiagaraComponent.cpp
+++ b/Engine/Plugins/FX/Niagara/Source/Niagara/Private/NiagaraComponent.cpp
@@ -148,7 +148,7 @@ FAutoConsoleCommandWithWorldAndArgs DumpNiagaraComponentsCommand(
 
 				if (Component->PoolingMethod == ENCPoolMethod::FreeInPool)
 				{
-					UE_LOG(LogNiagara, Log, TEXT("Component '%s' Asset '%s' Free In Compnent Pool"), *GetNameSafe(Component), *GetNameSafe(NiagaraSystem));
+					UE_LOG(LogNiagara, Log, TEXT("Component '%s' Asset '%s' Free In Component Pool"), *GetNameSafe(Component), *GetNameSafe(NiagaraSystem));
 				}
 				else
 				{
@@ -1981,6 +1981,15 @@ void UNiagaraComponent::BeginDestroy()
 	Super::BeginDestroy();
 }
 
+void UNiagaraComponent::GetResourceSizeEx(FResourceSizeEx& CumulativeResourceSize)
+{
+	Super::GetResourceSizeEx(CumulativeResourceSize);
+	if (SystemInstanceController)
+	{
+		CumulativeResourceSize.AddDedicatedSystemMemoryBytes(SystemInstanceController->GetTotalBytesUsed());
+	}
+}
+
 void UNiagaraComponent::CreateCullProxy(bool bForce)
 {
 	if (Asset->GetCullProxyMode() != ENiagaraCullProxyMode::None && CullProxy == nullptr)
diff --git a/Engine/Plugins/FX/Niagara/Source/Niagara/Private/NiagaraDebugHud.cpp b/Engine/Plugins/FX/Niagara/Source/Niagara/Private/NiagaraDebugHud.cpp
--- a/Engine/Plugins/FX/Niagara/Source/Niagara/Private/NiagaraDebugHud.cpp
+++ b/Engine/Plugins/FX/Niagara/Source/Niagara/Private/NiagaraDebugHud.cpp
@@ -1094,21 +1094,9 @@ void FNiagaraDebugHud::GatherSystemInfo()
 		}
 
 		// Track rough memory usage
-		if(SystemInstance)
-		{
-			for (const TSharedRef<FNiagaraEmitterInstance, ESPMode::ThreadSafe>& EmitterInstance : SystemInstance->GetEmitters())
-			{
-				UNiagaraEmitter* NiagaraEmitter = EmitterInstance->GetCachedEmitter().Emitter;
-				if (NiagaraEmitter == nullptr)
-				{
-					continue;
-				}
-
-				const int64 BytesUsed = EmitterInstance->GetTotalBytesUsed();
-				SystemDebugInfo.TotalBytes += BytesUsed;
-				GlobalTotalBytes += BytesUsed;
-			}
-		}
+		const int64 BytesUsed = SystemInstanceController->GetTotalBytesUsed();
+		SystemDebugInfo.TotalBytes += BytesUsed;
+		GlobalTotalBytes += BytesUsed;
 
 		if (bIsActive)
 		{
@@ -2803,14 +2791,7 @@ void FNiagaraDebugHud::DrawComponents(FNiagaraWorldManager* WorldManager, UCanva
 						BuildGpuHudInformation(StringBuilder, NiagaraComponent, SystemInstance, World->GetFeatureLevel());
 					}
 
-					int64 TotalBytes = 0;
-					for (const TSharedRef<FNiagaraEmitterInstance, ESPMode::ThreadSafe>& EmitterInstance : SystemInstance->GetEmitters())
-					{
-						if ( UNiagaraEmitter* NiagaraEmitter = EmitterInstance->GetCachedEmitter().Emitter )
-						{
-							TotalBytes += EmitterInstance->GetTotalBytesUsed();
-						}
-					}
+					int64 TotalBytes = SystemInstanceController->GetTotalBytesUsed();
 					StringBuilder.Appendf(TEXT("Memory - %6.2fMB\n"), float(double(TotalBytes) / (1024.0*1024.0)));
 				}
 
diff --git a/Engine/Plugins/FX/Niagara/Source/Niagara/Private/NiagaraSystem.cpp b/Engine/Plugins/FX/Niagara/Source/Niagara/Private/NiagaraSystem.cpp
--- a/Engine/Plugins/FX/Niagara/Source/Niagara/Private/NiagaraSystem.cpp
+++ b/Engine/Plugins/FX/Niagara/Source/Niagara/Private/NiagaraSystem.cpp
@@ -1002,7 +1002,38 @@ void UNiagaraSystem::PostLoad()
 	PrecachePSOs();
 }
 
+void UNiagaraSystem::GetResourceSizeEx(FResourceSizeEx& CumulativeResourceSize)
+{
+	Super::GetResourceSizeEx(CumulativeResourceSize);
+
+	CumulativeResourceSize.AddDedicatedSystemMemoryBytes(SystemScalabilityOverrides.Overrides.GetAllocatedSize());
+	CumulativeResourceSize.AddDedicatedSystemMemoryBytes(EmitterHandles.GetAllocatedSize());
+	if (const UScriptStruct* Struct = FVersionedNiagaraEmitterData::StaticStruct())
+	{
+		CumulativeResourceSize.AddDedicatedSystemMemoryBytes(Struct->GetStructureSize() * EmitterHandles.Num());
+	}
+	CumulativeResourceSize.AddDedicatedSystemMemoryBytes(EmitterCompiledData.GetAllocatedSize());
+	if (const UScriptStruct* Struct = FNiagaraEmitterCompiledData::StaticStruct())
+	{
+		CumulativeResourceSize.AddDedicatedSystemMemoryBytes(Struct->GetStructureSize() * EmitterCompiledData.Num());
+	}
+	if (const UScriptStruct* Struct = FNiagaraSystemCompiledData::StaticStruct())
+	{
+		CumulativeResourceSize.AddDedicatedSystemMemoryBytes(Struct->GetStructureSize());
+	}
+	CumulativeResourceSize.AddDedicatedSystemMemoryBytes(EmitterExecutionOrder.GetAllocatedSize());
+	CumulativeResourceSize.AddDedicatedSystemMemoryBytes(RendererPostTickOrder.GetAllocatedSize());
+	CumulativeResourceSize.AddDedicatedSystemMemoryBytes(RendererCompletionOrder.GetAllocatedSize());
+	CumulativeResourceSize.AddDedicatedSystemMemoryBytes(RendererDrawOrder.GetAllocatedSize());
+	CumulativeResourceSize.AddDedicatedSystemMemoryBytes(EmitterExecutionStateAccessors.GetAllocatedSize());
+	if (StaticBuffers != nullptr)
+	{
+		CumulativeResourceSize.AddDedicatedSystemMemoryBytes(StaticBuffers->GetCpuFloatBuffer().Num() * sizeof(float));
+	}
+}
+
 #if WITH_EDITORONLY_DATA
+
 void UNiagaraSystem::DeclareConstructClasses(TArray<FTopLevelAssetPath>& OutConstructClasses, const UClass* SpecificSubclass)
 {
 	Super::DeclareConstructClasses(OutConstructClasses, SpecificSubclass);
diff --git a/Engine/Plugins/FX/Niagara/Source/Niagara/Private/NiagaraSystemInstanceController.cpp b/Engine/Plugins/FX/Niagara/Source/Niagara/Private/NiagaraSystemInstanceController.cpp
--- a/Engine/Plugins/FX/Niagara/Source/Niagara/Private/NiagaraSystemInstanceController.cpp
+++ b/Engine/Plugins/FX/Niagara/Source/Niagara/Private/NiagaraSystemInstanceController.cpp
@@ -258,6 +258,23 @@ void FNiagaraSystemInstanceController::DebugDump(bool bFullDump)
 	}
 }
 
+SIZE_T FNiagaraSystemInstanceController::GetTotalBytesUsed() const
+{
+	SIZE_T Size = 0;
+	if (IsValid())
+	{
+		for (const TSharedRef<FNiagaraEmitterInstance, ESPMode::ThreadSafe>& Emitter : GetSystemInstance_Unsafe()->GetEmitters())
+		{
+			if (Emitter->GetCachedEmitter().Emitter != nullptr)
+			{
+				Size += Emitter->GetTotalBytesUsed();
+			}
+		}
+	}
+
+	return Size;
+}
+
 void FNiagaraSystemInstanceController::AddReferencedObjects(FReferenceCollector& Collector)
 {
 	for (auto& Override : EmitterMaterials)
diff --git a/Engine/Plugins/FX/Niagara/Source/Niagara/Public/NiagaraComponent.h b/Engine/Plugins/FX/Niagara/Source/Niagara/Public/NiagaraComponent.h
--- a/Engine/Plugins/FX/Niagara/Source/Niagara/Public/NiagaraComponent.h
+++ b/Engine/Plugins/FX/Niagara/Source/Niagara/Public/NiagaraComponent.h
@@ -205,6 +205,7 @@ public:
 	virtual void Activate(bool bReset = false) override;
 	virtual void Deactivate() override;
 	virtual void DeactivateImmediate() override;
+	virtual void GetResourceSizeEx(FResourceSizeEx& CumulativeResourceSize) override;
 
 	ENiagaraExecutionState GetRequestedExecutionState() const;
 	ENiagaraExecutionState GetExecutionState() const;
diff --git a/Engine/Plugins/FX/Niagara/Source/Niagara/Public/NiagaraSystemInstanceController.h b/Engine/Plugins/FX/Niagara/Source/Niagara/Public/NiagaraSystemInstanceController.h
--- a/Engine/Plugins/FX/Niagara/Source/Niagara/Public/NiagaraSystemInstanceController.h
+++ b/Engine/Plugins/FX/Niagara/Source/Niagara/Public/NiagaraSystemInstanceController.h
@@ -88,6 +88,7 @@ public:
 
 	/** Dumps system instance state and info to the log (Used by fx.Niagara.DumpComponents commandlet) */
 	void DebugDump(bool bFullDump);
+	SIZE_T GetTotalBytesUsed() const;
 
 	// Deferred methods - With NIAGARA_SYSTEM_INSTANCE_CONTROLLER_ASYNC, these are executed in a deferred queue
 	NIAGARA_SYSTEM_INSTANCE_CONTROLLER_DEFFERED_METHOD(Reset)
