- Fix for potential nullptr access to EmittterInst

#rnx

diff --git a/Engine/Plugins/FX/Niagara/Source/Niagara/Private/DataInterface/NiagaraDataInterfaceDataChannelRead.cpp b/Engine/Plugins/FX/Niagara/Source/Niagara/Private/DataInterface/NiagaraDataInterfaceDataChannelRead.cpp
--- a/Engine/Plugins/FX/Niagara/Source/Niagara/Private/DataInterface/NiagaraDataInterfaceDataChannelRead.cpp
+++ b/Engine/Plugins/FX/Niagara/Source/Niagara/Private/DataInterface/NiagaraDataInterfaceDataChannelRead.cpp
@@ -1041,30 +1041,30 @@ void UNiagaraDataInterfaceDataChannelRead::SpawnFromSpawnInfo(FVectorVMExternalF
 	FNiagaraDataChannelData* DataChannelData = InstData->DataChannelData.Get();
 	FNiagaraDataBuffer* Data = DataChannelData ? DataChannelData->GetCPUData(!bReadCurrentFrame) : nullptr;
 
-	bool bSpawn = INiagaraModule::DataChannelsEnabled() && Data && EmittterInst;
-
-	TArray<FNiagaraSpawnInfo>& EmitterSpawnInfos = EmittterInst->GetSpawnInfo();
-
-	bSpawn &= InEnabled.GetAndAdvance();
-	bSpawn &= InstData->SpawnInfoAccessor.IsValid();
-	if (bSpawn)
+	if ( INiagaraModule::DataChannelsEnabled() && Data && EmittterInst )
 	{
-		FNiagaraDataSetReaderStruct<FNiagaraSpawnInfo> Reader = InstData->SpawnInfoAccessor.GetReader(Data);
-
-		int32 NumDataChannelInstances = Data->GetNumInstances();
+		TArray<FNiagaraSpawnInfo>& EmitterSpawnInfos = EmittterInst->GetSpawnInfo();
 
-		for (int32 DataChannelIdx = 0; DataChannelIdx < NumDataChannelInstances; ++DataChannelIdx)
+		const bool bSpawn = InEnabled.GetAndAdvance() && InstData->SpawnInfoAccessor.IsValid();
+		if (bSpawn)
 		{
-			FNiagaraSpawnInfo SpawnInfo(0, 0.0f, 0.0f, 0);
-			SpawnInfo = Reader.Get(DataChannelIdx);
+			FNiagaraDataSetReaderStruct<FNiagaraSpawnInfo> Reader = InstData->SpawnInfoAccessor.GetReader(Data);
+
+			int32 NumDataChannelInstances = Data->GetNumInstances();
 
-			if (SpawnInfo.Count > 0)
+			for (int32 DataChannelIdx = 0; DataChannelIdx < NumDataChannelInstances; ++DataChannelIdx)
 			{
-				if (bOverrideSpawnGroupToDataChannelIndex)
+				FNiagaraSpawnInfo SpawnInfo(0, 0.0f, 0.0f, 0);
+				SpawnInfo = Reader.Get(DataChannelIdx);
+
+				if (SpawnInfo.Count > 0)
 				{
-					SpawnInfo.SpawnGroup = DataChannelIdx;
+					if (bOverrideSpawnGroupToDataChannelIndex)
+					{
+						SpawnInfo.SpawnGroup = DataChannelIdx;
+					}
+					EmitterSpawnInfos.Emplace(SpawnInfo);
 				}
-				EmitterSpawnInfos.Emplace(SpawnInfo);
 			}
 		}
 	}
