Remove PrimaryAssetId for data channel asset

#rnx

diff --git a/Engine/Plugins/FX/Niagara/Source/Niagara/Private/NiagaraDataChannel.cpp b/Engine/Plugins/FX/Niagara/Source/Niagara/Private/NiagaraDataChannel.cpp
--- a/Engine/Plugins/FX/Niagara/Source/Niagara/Private/NiagaraDataChannel.cpp
+++ b/Engine/Plugins/FX/Niagara/Source/Niagara/Private/NiagaraDataChannel.cpp
@@ -996,11 +996,6 @@ UNiagaraDataChannelReader* UNiagaraDataChannelLibrary::ReadFromNiagaraDataChanne
 	return nullptr;
 }
 
-FPrimaryAssetId UNiagaraDataChannelAsset::GetPrimaryAssetId()const 
-{
-	return FPrimaryAssetId(GetClass()->GetFName(), GetFName());
-}
-
 //////////////////////////////////////////////////////////////////////////
 
 #if !UE_BUILD_SHIPPING
diff --git a/Engine/Plugins/FX/Niagara/Source/Niagara/Public/NiagaraDataChannelPublic.h b/Engine/Plugins/FX/Niagara/Source/Niagara/Public/NiagaraDataChannelPublic.h
--- a/Engine/Plugins/FX/Niagara/Source/Niagara/Public/NiagaraDataChannelPublic.h
+++ b/Engine/Plugins/FX/Niagara/Source/Niagara/Public/NiagaraDataChannelPublic.h
@@ -26,10 +26,6 @@ class NIAGARA_API UNiagaraDataChannelAsset : public UObject
 
 public:
 
-	//UObject Interface
-	virtual FPrimaryAssetId GetPrimaryAssetId() const override;
-	//END UObject Interface
-
 	UNiagaraDataChannel* Get() const { return DataChannel; }
 };
 
