nullptr check fixing crash in debug hud.

#rnx

diff --git a/Engine/Plugins/FX/Niagara/Source/Niagara/Private/NiagaraDebugHud.cpp b/Engine/Plugins/FX/Niagara/Source/Niagara/Private/NiagaraDebugHud.cpp
--- a/Engine/Plugins/FX/Niagara/Source/Niagara/Private/NiagaraDebugHud.cpp
+++ b/Engine/Plugins/FX/Niagara/Source/Niagara/Private/NiagaraDebugHud.cpp
@@ -1094,7 +1094,7 @@ void FNiagaraDebugHud::GatherSystemInfo()
 		}
 
 		// Track rough memory usage
-		const int64 BytesUsed = SystemInstanceController->GetTotalBytesUsed();
+		const int64 BytesUsed = SystemInstanceController.IsValid() ? SystemInstanceController->GetTotalBytesUsed() : 0;
 		SystemDebugInfo.TotalBytes += BytesUsed;
 		GlobalTotalBytes += BytesUsed;
 
