Fixed use of incorrect object name in UI.

#rnx

diff --git a/Engine/Plugins/FX/Niagara/Source/NiagaraEditor/Private/NiagaraActions.cpp b/Engine/Plugins/FX/Niagara/Source/NiagaraEditor/Private/NiagaraActions.cpp
--- a/Engine/Plugins/FX/Niagara/Source/NiagaraEditor/Private/NiagaraActions.cpp
+++ b/Engine/Plugins/FX/Niagara/Source/NiagaraEditor/Private/NiagaraActions.cpp
@@ -858,13 +858,13 @@ void FNiagaraDataInterfaceNodeActionProvider_DataChannelRead::CollectAddPinActio
 				Tooltip = FText::Format(
 					LOCTEXT("NDIDataChannleReadAddParameterFmt", "Make this function conditional on {0} from Data Channel {1}.\nThe function will only have an effect for Data Channel entries that pass the comparisson test with the passed value."),
 					FText::FromName(SWCVar.GetName()),
-					FText::FromString(Channel->GetName()));
+					FText::FromString(Channel->GetAsset()->GetName()));
 			}
 			else
 			{
 				Category = FText::Format(LOCTEXT("NDIDataChannleReadCategoryFmt", "{0}"), FText::FromString(Channel->GetAsset()->GetName()));
 				DisplayName = FText::FromName(SWCVar.GetName());
-				Tooltip = FText::Format(LOCTEXT("NDIDataChannleReadAddParameterFmt", "Read {0} from Data Channel {1}."), DisplayName, FText::FromString(Channel->GetName()));
+				Tooltip = FText::Format(LOCTEXT("NDIDataChannleReadAddParameterFmt", "Read {0} from Data Channel {1}."), DisplayName, FText::FromString(Channel->GetAsset()->GetName()));
 			}
 
 			TSharedPtr<FNiagaraMenuAction> Action(new FNiagaraMenuAction(
