- Improve CPU access check on static mesh DI to allow functions used in system / emitter scripts to not break sampling

#rnx

diff --git a/Engine/Plugins/FX/Niagara/Source/Niagara/Internal/DataInterface/NiagaraDataInterfaceStaticMesh.h b/Engine/Plugins/FX/Niagara/Source/Niagara/Internal/DataInterface/NiagaraDataInterfaceStaticMesh.h
--- a/Engine/Plugins/FX/Niagara/Source/Niagara/Internal/DataInterface/NiagaraDataInterfaceStaticMesh.h
+++ b/Engine/Plugins/FX/Niagara/Source/Niagara/Internal/DataInterface/NiagaraDataInterfaceStaticMesh.h
@@ -147,7 +147,7 @@ public:
 	void GetUVMappingFunctions(TArray<FNiagaraFunctionSignature>& OutFunctions, const FNiagaraFunctionSignature& BaseSignature) const;
 	void GetDistanceFieldFunctions(TArray<FNiagaraFunctionSignature>& OutFunctions, const FNiagaraFunctionSignature& BaseSignature) const;
 	void GetDeprecatedFunctions(TArray<FNiagaraFunctionSignature>& OutFunctions, const FNiagaraFunctionSignature& BaseSignature) const;
-	void GetCpuAccessFunctions(TArray<FNiagaraFunctionSignature>& OutFunctions);
+	static bool FunctionNeedsCpuAccess(FName InName);
 
 #if WITH_EDITORONLY_DATA
 	virtual void GetCommonHLSL(FString& OutHLSL) override;
diff --git a/Engine/Plugins/FX/Niagara/Source/Niagara/Private/DataInterface/NiagaraDataInterfaceStaticMesh.cpp b/Engine/Plugins/FX/Niagara/Source/Niagara/Private/DataInterface/NiagaraDataInterfaceStaticMesh.cpp
--- a/Engine/Plugins/FX/Niagara/Source/Niagara/Private/DataInterface/NiagaraDataInterfaceStaticMesh.cpp
+++ b/Engine/Plugins/FX/Niagara/Source/Niagara/Private/DataInterface/NiagaraDataInterfaceStaticMesh.cpp
@@ -890,11 +890,25 @@ namespace NDIStaticMeshLocal
 				}
 				else
 				{
-					if (Interface->IsUsedWithCPUEmitter())
+					bool bNeedsCpuAccess = false;
+					FNiagaraDataInterfaceUtilities::ForEachVMFunction(
+						Interface,
+						SystemInstance,
+						[&bNeedsCpuAccess](const UNiagaraScript* Script, const FVMExternalFunctionBindingInfo& VMFunction) -> bool
+						{
+							bNeedsCpuAccess |= UNiagaraDataInterfaceStaticMesh::FunctionNeedsCpuAccess(VMFunction.Name);
+							return bNeedsCpuAccess == false;
+						}
+					);
+
+					if (bNeedsCpuAccess)
 					{
-						UE_LOG(LogNiagara, Log, TEXT("NiagaraStaticMeshDataInterface used by CPU emitter and does not allow CPU access. System: %s, Mesh: %s"), *GetFullNameSafe(SystemInstance->GetSystem()), *GetFullNameSafe(StaticMesh));
+						if (Interface->IsUsedWithCPUEmitter())
+						{
+							UE_LOG(LogNiagara, Log, TEXT("NiagaraStaticMeshDataInterface used by CPU emitter and does not allow CPU access. System: %s, Mesh: %s"), *GetFullNameSafe(SystemInstance->GetSystem()), *GetFullNameSafe(StaticMesh));
+						}
+						StaticMesh = nullptr;
 					}
-					StaticMesh = nullptr;
 				}
 			}
 
@@ -2475,25 +2489,24 @@ void UNiagaraDataInterfaceStaticMesh::GetDeprecatedFunctions(TArray<FNiagaraFunc
 	}
 }
 
-void UNiagaraDataInterfaceStaticMesh::GetCpuAccessFunctions(TArray<FNiagaraFunctionSignature>& OutFunctions)
+bool UNiagaraDataInterfaceStaticMesh::FunctionNeedsCpuAccess(FName InName)
 {
 	using namespace NDIStaticMeshLocal;
 
-	// Setup base signature
-	FNiagaraFunctionSignature BaseSignature;
-	BaseSignature.Inputs.Emplace(FNiagaraTypeDefinition(GetClass()), TEXT("StaticMesh"));
-	BaseSignature.bMemberFunction = true;
-	BaseSignature.bRequiresContext = false;
-#if WITH_EDITORONLY_DATA
-	BaseSignature.FunctionVersion = EDIFunctionVersion::LatestVersion;
-#endif
+	static TSet<FName> SafeFunctions =
+	{
+		IsValidName,
+		GetPreSkinnedLocalBoundsName,
+		GetMeshBoundsName,
+		GetMeshBoundsWSName,
+		GetLocalToWorldName,
+		GetLocalToWorldInverseTransposedName,
+		GetWorldVelocityName,
+		//GetInstanceIndexName,
+		SetInstanceIndexName,
+	};
 
-	GetVertexSamplingFunctions(OutFunctions, BaseSignature);
-	GetTriangleSamplingFunctions(OutFunctions, BaseSignature);
-	GetSectionFunctions(OutFunctions, BaseSignature);
-	GetUVMappingFunctions(OutFunctions, BaseSignature);
-	GetDistanceFieldFunctions(OutFunctions, BaseSignature);
-	GetDeprecatedFunctions(OutFunctions, BaseSignature);
+	return SafeFunctions.Contains(InName) == false;
 }
 
 #if WITH_EDITORONLY_DATA
@@ -3318,33 +3331,29 @@ void UNiagaraDataInterfaceStaticMesh::GetFeedback(UNiagaraSystem* Asset, UNiagar
 	{
 		if (CurrentMesh->bAllowCPUAccess == false)
 		{
-			// Get list of functions that require CPU access, this is slightly conservative but reduces the level of false warnings
-			TArray<FNiagaraFunctionSignature> CpuAccessFunctions;
-			GetCpuAccessFunctions(CpuAccessFunctions);
-
-			bool bCpuAccessWarning = false;
+			bool bNeedsCpuAccess = false;
 
 			FNiagaraDataInterfaceUtilities::ForEachVMFunction(
 				RuntimeInstanceOfThis,
 				Asset,
-				[&](const UNiagaraScript* Script, const FVMExternalFunctionBindingInfo& VMFunction) -> bool
+				[&bNeedsCpuAccess](const UNiagaraScript* Script, const FVMExternalFunctionBindingInfo& VMFunction) -> bool
 				{
-					bCpuAccessWarning |= CpuAccessFunctions.ContainsByPredicate([&](const FNiagaraFunctionSignature& CheckFunction) { return CheckFunction.Name == VMFunction.Name; });
-					return bCpuAccessWarning == false;
+					bNeedsCpuAccess |= FunctionNeedsCpuAccess(VMFunction.Name);
+					return bNeedsCpuAccess == false;
 				}
 			);
 
 			FNiagaraDataInterfaceUtilities::ForEachGpuFunction(
 				RuntimeInstanceOfThis,
 				Asset,
-				[&](const UNiagaraScript* Script,const FNiagaraDataInterfaceGeneratedFunction& GpuFunction) -> bool
+				[&bNeedsCpuAccess](const UNiagaraScript* Script,const FNiagaraDataInterfaceGeneratedFunction& GpuFunction) -> bool
 				{
-					bCpuAccessWarning |= CpuAccessFunctions.ContainsByPredicate([&](const FNiagaraFunctionSignature& CheckFunction) { return CheckFunction.Name == GpuFunction.DefinitionName; });
-					return bCpuAccessWarning == false;
+					bNeedsCpuAccess |= FunctionNeedsCpuAccess(GpuFunction.DefinitionName);
+					return bNeedsCpuAccess == false;
 				}
 			);
 
-			if (bCpuAccessWarning)
+			if (bNeedsCpuAccess)
 			{
 				OutErrors.Emplace(
 					FText::Format(LOCTEXT("CPUAccessNotAllowedError", "This mesh needs CPU access in order to be used properly.({0})"), FText::FromString(CurrentMesh->GetName())),
