Fixed SA warnings

diff --git a/Engine/Plugins/FX/Niagara/Source/NiagaraEditor/Private/NiagaraScriptMergeManager.cpp b/Engine/Plugins/FX/Niagara/Source/NiagaraEditor/Private/NiagaraScriptMergeManager.cpp
--- a/Engine/Plugins/FX/Niagara/Source/NiagaraEditor/Private/NiagaraScriptMergeManager.cpp
+++ b/Engine/Plugins/FX/Niagara/Source/NiagaraEditor/Private/NiagaraScriptMergeManager.cpp
@@ -3718,13 +3718,16 @@ FNiagaraScriptMergeManager::FApplyDiffResults FNiagaraScriptMergeManager::ApplyE
 				if(BaseRoot->FindChildWithIdentity(AddedInput->GetPersistentIdentity(), true) == nullptr)
 				{
 					UNiagaraHierarchyItemBase* NewChild = BaseRoot->CopyAndAddItemUnderParentIdentity(*AddedInput, ParentIdentity);
-					// see above
-					NewChild->GetChildrenMutable().Empty();
 					
 					if(!ensure(NewChild != nullptr))
 					{
 						UE_LOG(LogNiagaraEditor, Log, TEXT("Item %s could not be added during merge process"), *AddedInput->ToString());
 					}
+					else
+					{
+						// see above
+						NewChild->GetChildrenMutable().Empty();
+					}
 				}
 			}
 		}
