Added multiline support for tooltip texts in the hierarchy data

diff --git a/Engine/Plugins/FX/Niagara/Source/NiagaraEditor/Public/ViewModels/HierarchyEditor/NiagaraHierarchyViewModelBase.h b/Engine/Plugins/FX/Niagara/Source/NiagaraEditor/Public/ViewModels/HierarchyEditor/NiagaraHierarchyViewModelBase.h
--- a/Engine/Plugins/FX/Niagara/Source/NiagaraEditor/Public/ViewModels/HierarchyEditor/NiagaraHierarchyViewModelBase.h
+++ b/Engine/Plugins/FX/Niagara/Source/NiagaraEditor/Public/ViewModels/HierarchyEditor/NiagaraHierarchyViewModelBase.h
@@ -287,7 +287,7 @@ private:
 	FName Category;
 
 	/** The tooltip used when the user is hovering this category */
-	UPROPERTY(EditAnywhere, Category = "Category")
+	UPROPERTY(EditAnywhere, Category = "Category", meta = (MultiLine = "true"))
 	FText Tooltip;
 
 	UPROPERTY()
@@ -316,7 +316,7 @@ private:
 	FName Section;
 
 	/** The tooltip used when the user is hovering this section */
-	UPROPERTY(EditAnywhere, Category = "Section")
+	UPROPERTY(EditAnywhere, Category = "Section", meta = (MultiLine = "true"))
 	FText Tooltip;
 };
 
diff --git a/Engine/Plugins/FX/Niagara/Source/NiagaraEditor/Public/ViewModels/HierarchyEditor/NiagaraSummaryViewViewModel.h b/Engine/Plugins/FX/Niagara/Source/NiagaraEditor/Public/ViewModels/HierarchyEditor/NiagaraSummaryViewViewModel.h
--- a/Engine/Plugins/FX/Niagara/Source/NiagaraEditor/Public/ViewModels/HierarchyEditor/NiagaraSummaryViewViewModel.h
+++ b/Engine/Plugins/FX/Niagara/Source/NiagaraEditor/Public/ViewModels/HierarchyEditor/NiagaraSummaryViewViewModel.h
@@ -53,7 +53,7 @@ private:
 	FText DisplayNameOverride;
 
 	/** If specified, will override how the tooltip of this input in the stack. */
-	UPROPERTY(EditAnywhere, Category="Niagara")
+	UPROPERTY(EditAnywhere, Category="Niagara", meta = (MultiLine = "true"))
 	FText TooltipOverride;
 };
 
@@ -71,7 +71,7 @@ public:
 	FText GetTooltipOverride() const { return TooltipOverride; }
 private:
 	/** If specified, will override how the tooltip of this input in the stack. */
-	UPROPERTY(EditAnywhere, Category="Niagara")
+	UPROPERTY(EditAnywhere, Category="Niagara", meta = (MultiLine = "true"))
 	FText TooltipOverride;
 };
 
