- Fix SA issue

#rnx

diff --git a/Engine/Plugins/FX/Niagara/Source/Niagara/Private/NiagaraScriptExecutionContext.cpp b/Engine/Plugins/FX/Niagara/Source/Niagara/Private/NiagaraScriptExecutionContext.cpp
--- a/Engine/Plugins/FX/Niagara/Source/Niagara/Private/NiagaraScriptExecutionContext.cpp
+++ b/Engine/Plugins/FX/Niagara/Source/Niagara/Private/NiagaraScriptExecutionContext.cpp
@@ -58,6 +58,11 @@ FNiagaraScriptExecutionContextBase::~FNiagaraScriptExecutionContextBase()
 bool FNiagaraScriptExecutionContextBase::Init(FNiagaraSystemInstance* Instance, UNiagaraScript* InScript, ENiagaraSimTarget InTarget)
 {
 	Script = InScript;
+	if (!ensure(Script != nullptr))
+	{
+		UE_DEBUG_BREAK();
+		return false;
+	}
 
 	Parameters.InitFromOwningContext(Script, InTarget, true);
 
@@ -68,14 +73,11 @@ bool FNiagaraScriptExecutionContextBase::Init(FNiagaraSystemInstance* Instance,
 	const bool ScriptSupportsExperimentalVM = (Script ? Script->GetVMExecutableData().SupportsExperimentalVM() : false);
 	bool SystemDisabledExperimentalVM = false;
 
-	if (InScript)
+	if (const UNiagaraSystem* OwnerSystem = Script->GetTypedOuter<UNiagaraSystem>())
 	{
-		if (const UNiagaraSystem* OwnerSystem = InScript->GetTypedOuter<UNiagaraSystem>())
+		if (OwnerSystem->ShouldDisableExperimentalVM())
 		{
-			if (OwnerSystem->ShouldDisableExperimentalVM())
-			{
-				SystemDisabledExperimentalVM = true;
-			}
+			SystemDisabledExperimentalVM = true;
 		}
 	}
 
@@ -87,28 +89,25 @@ bool FNiagaraScriptExecutionContextBase::Init(FNiagaraSystemInstance* Instance,
 	VectorVMState = AllocVectorVMState(&OptimizeContext);
 #endif
 
-	if(Script)
+	ENiagaraScriptUsage Usage = Script->GetUsage();
+	for(auto& ResolvedDIInfo : Script->GetResolvedDataInterfaces())
 	{
-		ENiagaraScriptUsage Usage = Script->GetUsage();
-		for(auto& ResolvedDIInfo : Script->GetResolvedDataInterfaces())
+		if(ResolvedDIInfo.ResolvedDataInterface)
 		{
-			if(ResolvedDIInfo.ResolvedDataInterface)
-			{
-				bHasDIsWithPreStageTick |= (ResolvedDIInfo.ResolvedDataInterface->HasPreStageTick(Usage));
-				bHasDIsWithPostStageTick |= (ResolvedDIInfo.ResolvedDataInterface->HasPostStageTick(Usage));
-			}
-			if(bHasDIsWithPreStageTick && bHasDIsWithPostStageTick)
-			{
-				break;
-			}
+			bHasDIsWithPreStageTick |= (ResolvedDIInfo.ResolvedDataInterface->HasPreStageTick(Usage));
+			bHasDIsWithPostStageTick |= (ResolvedDIInfo.ResolvedDataInterface->HasPostStageTick(Usage));
 		}
-
-		if(Instance && (bHasDIsWithPreStageTick || bHasDIsWithPostStageTick))
+		if(bHasDIsWithPreStageTick && bHasDIsWithPostStageTick)
 		{
-			DIStageTickHandler.Init(Script, Instance);
+			break;
 		}
 	}
 
+	if(Instance && (bHasDIsWithPreStageTick || bHasDIsWithPostStageTick))
+	{
+		DIStageTickHandler.Init(Script, Instance);
+	}
+
 	return true;
 }
 
