- Remove debug break I accidentally checked in

#rnx

diff --git a/Engine/Plugins/FX/Niagara/Source/Niagara/Private/NiagaraScriptExecutionContext.cpp b/Engine/Plugins/FX/Niagara/Source/Niagara/Private/NiagaraScriptExecutionContext.cpp
--- a/Engine/Plugins/FX/Niagara/Source/Niagara/Private/NiagaraScriptExecutionContext.cpp
+++ b/Engine/Plugins/FX/Niagara/Source/Niagara/Private/NiagaraScriptExecutionContext.cpp
@@ -60,7 +60,6 @@ bool FNiagaraScriptExecutionContextBase::Init(FNiagaraSystemInstance* Instance,
 	Script = InScript;
 	if (!ensure(Script != nullptr))
 	{
-		UE_DEBUG_BREAK();
 		return false;
 	}
 
