Added hierarchy editor commands to replace the default Delete command. This is so that backspace no longer works to delete.

diff --git a/Engine/Plugins/FX/Niagara/Source/NiagaraEditor/Private/NiagaraEditorModule.cpp b/Engine/Plugins/FX/Niagara/Source/NiagaraEditor/Private/NiagaraEditorModule.cpp
--- a/Engine/Plugins/FX/Niagara/Source/NiagaraEditor/Private/NiagaraEditorModule.cpp
+++ b/Engine/Plugins/FX/Niagara/Source/NiagaraEditor/Private/NiagaraEditorModule.cpp
@@ -165,6 +165,7 @@
 
 #include "Engine/AssetManager.h"
 #include "Engine/AssetManagerSettings.h"
+#include "ViewModels/HierarchyEditor/NiagaraHierarchyCommands.h"
 
 #include UE_INLINE_GENERATED_CPP_BY_NAME(NiagaraEditorModule)
 
@@ -1159,6 +1160,7 @@ void FNiagaraEditorModule::StartupModule()
 		FConsoleCommandDelegate::CreateRaw(this, &FNiagaraEditorModule::ReinitializeStyle));
 
 	FNiagaraEditorCommands::Register();
+	FNiagaraHierarchyEditorCommands::Register();
 
 	NiagaraComponentBroker = MakeShareable(new FNiagaraComponentBroker);
 	FComponentAssetBrokerage::RegisterBroker(NiagaraComponentBroker, UNiagaraComponent::StaticClass(), true, true);
diff --git a/Engine/Plugins/FX/Niagara/Source/NiagaraEditor/Private/ViewModels/HierarchyEditor/NiagaraHierarchyCommands.cpp b/Engine/Plugins/FX/Niagara/Source/NiagaraEditor/Private/ViewModels/HierarchyEditor/NiagaraHierarchyCommands.cpp
new file mode 100644
--- /dev/null
+++ b/Engine/Plugins/FX/Niagara/Source/NiagaraEditor/Private/ViewModels/HierarchyEditor/NiagaraHierarchyCommands.cpp
@@ -0,0 +1,12 @@
+// Copyright Epic Games, Inc. All Rights Reserved.
+
+#include "ViewModels/HierarchyEditor/NiagaraHierarchyCommands.h"
+
+#define LOCTEXT_NAMESPACE "NiagaraHierarchyEditorCommands"
+
+void FNiagaraHierarchyEditorCommands::RegisterCommands()
+{
+	UI_COMMAND(Delete, "Delete", "Delete current selection", EUserInterfaceActionType::Button, FInputChord(EKeys::Delete))
+}
+
+#undef LOCTEXT_NAMESPACE
diff --git a/Engine/Plugins/FX/Niagara/Source/NiagaraEditor/Private/Widgets/SNiagaraHierarchy.cpp b/Engine/Plugins/FX/Niagara/Source/NiagaraEditor/Private/Widgets/SNiagaraHierarchy.cpp
--- a/Engine/Plugins/FX/Niagara/Source/NiagaraEditor/Private/Widgets/SNiagaraHierarchy.cpp
+++ b/Engine/Plugins/FX/Niagara/Source/NiagaraEditor/Private/Widgets/SNiagaraHierarchy.cpp
@@ -19,6 +19,7 @@
 #include "Framework/Application/MenuStack.h"
 #include "Framework/Application/SlateApplication.h"
 #include "ToolMenus.h"
+#include "ViewModels/HierarchyEditor/NiagaraHierarchyCommands.h"
 #include "Widgets/Input/SSearchBox.h"
 
 #define LOCTEXT_NAMESPACE "NiagaraHierarchy"
@@ -55,7 +56,7 @@ TSharedRef<SWidget> SummonContextMenu(TArray<TSharedPtr<FNiagaraHierarchyItemVie
 
 		if(bCanDeleteAll)
 		{
-			BaseSection->AddMenuEntry(FGenericCommands::Get().Delete);
+			BaseSection->AddMenuEntry(FNiagaraHierarchyEditorCommands::Get().Delete);
 		}
 
 		if(Items.Num() == 1 && bFromHierarchy == false)
@@ -778,7 +779,7 @@ void SNiagaraHierarchy::Construct(const FArguments& InArgs, TObjectPtr<UNiagaraH
 		FExecuteAction::CreateSP(this, &SNiagaraHierarchy::RequestRenameSelectedItem),
 		FCanExecuteAction::CreateSP(this, &SNiagaraHierarchy::CanRequestRenameSelectedItem));
 
-	HierarchyViewModel->GetCommands()->MapAction(FGenericCommands::Get().Delete,
+	HierarchyViewModel->GetCommands()->MapAction(FNiagaraHierarchyEditorCommands::Get().Delete,
 		FExecuteAction::CreateSP(this, &SNiagaraHierarchy::DeleteSelectedItems),
 		FCanExecuteAction::CreateSP(this, &SNiagaraHierarchy::CanDeleteSelectedItems),
 		FIsActionChecked(), FIsActionButtonVisible::CreateSP(this, &SNiagaraHierarchy::CanDeleteSelectedItems));
diff --git a/Engine/Plugins/FX/Niagara/Source/NiagaraEditor/Public/ViewModels/HierarchyEditor/NiagaraHierarchyCommands.h b/Engine/Plugins/FX/Niagara/Source/NiagaraEditor/Public/ViewModels/HierarchyEditor/NiagaraHierarchyCommands.h
new file mode 100644
--- /dev/null
+++ b/Engine/Plugins/FX/Niagara/Source/NiagaraEditor/Public/ViewModels/HierarchyEditor/NiagaraHierarchyCommands.h
@@ -0,0 +1,25 @@
+// Copyright Epic Games, Inc. All Rights Reserved.
+
+#pragma once
+
+#include "CoreMinimal.h"
+#include "Styling/CoreStyle.h"
+#include "Framework/Commands/Commands.h"
+
+class NIAGARAEDITOR_API FNiagaraHierarchyEditorCommands : public TCommands<FNiagaraHierarchyEditorCommands>
+{
+public:
+	
+	FNiagaraHierarchyEditorCommands()
+		: TCommands<FNiagaraHierarchyEditorCommands>( TEXT("NiagaraHierarchyEditorCommands"), NSLOCTEXT("NiagaraHierarchyEditorCommands", "Hierarchy Editor Commands", "Hierarchy Editor Commands"), NAME_None, "NiagaraEditorStyle" )
+	{
+	}
+
+	virtual ~FNiagaraHierarchyEditorCommands() override
+	{
+	}
+
+	virtual void RegisterCommands() override;	
+
+	TSharedPtr<FUICommandInfo> Delete;
+};
