- Fix early out if already activated and in solo mode
- Before we would disable the tick leaving the system lingering forever

diff --git a/Engine/Plugins/FX/Niagara/Source/Niagara/Private/NiagaraComponent.cpp b/Engine/Plugins/FX/Niagara/Source/Niagara/Private/NiagaraComponent.cpp
--- a/Engine/Plugins/FX/Niagara/Source/Niagara/Private/NiagaraComponent.cpp
+++ b/Engine/Plugins/FX/Niagara/Source/Niagara/Private/NiagaraComponent.cpp
@@ -1249,6 +1249,9 @@ void UNiagaraComponent::ActivateInternal(bool bReset /* = false */, bool bIsScal
 		SystemInstanceController->GetRequestedExecutionState() == ENiagaraExecutionState::Active &&
 		SystemInstanceController->GetActualExecutionState() == ENiagaraExecutionState::Active)
 	{
+		// It would be incorrect for the component tick to be enabled and not be in solo mode but we will verify this
+		ensure(SystemInstanceController->IsSolo() == IsComponentTickEnabled());
+		ScopedComponentTickEnabled.bTickEnabled = IsComponentTickEnabled();
 		return;
 	}
 	
