- Add option to capture the mesh transforms per frame or not
- In some instances you want to capture the information initially and never again

#preflght 645952fbfd4b8f4e0dcf1475

diff --git a/Engine/Plugins/FX/Niagara/Source/Niagara/Internal/DataInterface/NiagaraDataInterfaceStaticMesh.h b/Engine/Plugins/FX/Niagara/Source/Niagara/Internal/DataInterface/NiagaraDataInterfaceStaticMesh.h
--- a/Engine/Plugins/FX/Niagara/Source/Niagara/Internal/DataInterface/NiagaraDataInterfaceStaticMesh.h
+++ b/Engine/Plugins/FX/Niagara/Source/Niagara/Internal/DataInterface/NiagaraDataInterfaceStaticMesh.h
@@ -90,6 +90,10 @@ public:
 	UPROPERTY(EditAnywhere, Category = "Mesh")
 	FNDIStaticMeshSectionFilter SectionFilter;
 
+	/** If true we capture the transforms from the mesh component each frame. */
+	UPROPERTY(EditAnywhere, AdvancedDisplay, Category = "Mesh")
+	bool bCaptureTransformsPerFrame = true;
+
 	/** If true then the mesh velocity is taken from the mesh component's physics data. Otherwise it will be calculated by diffing the component transforms between ticks, which is more reliable but won't work on the first frame. */
 	UPROPERTY(EditAnywhere, AdvancedDisplay, Category = "Mesh")
     bool bUsePhysicsBodyVelocity = false;
diff --git a/Engine/Plugins/FX/Niagara/Source/Niagara/Private/DataInterface/NiagaraDataInterfaceStaticMesh.cpp b/Engine/Plugins/FX/Niagara/Source/Niagara/Private/DataInterface/NiagaraDataInterfaceStaticMesh.cpp
--- a/Engine/Plugins/FX/Niagara/Source/Niagara/Private/DataInterface/NiagaraDataInterfaceStaticMesh.cpp
+++ b/Engine/Plugins/FX/Niagara/Source/Niagara/Private/DataInterface/NiagaraDataInterfaceStaticMesh.cpp
@@ -621,6 +621,10 @@ namespace NDIStaticMeshLocal
 
 		/** Velocity set by the physics body of the mesh component */
 		FVector PhysicsVelocity;
+
+		/** True we capture the transforms from the mesh component each frame. */
+		uint32 bCaptureTransformsPerFrame : 1;
+
 		/** True if velocity should not be calculated via the transforms, but rather read the physics data from the mesh component */
 		uint32 bUsePhysicsVelocity : 1;
 
@@ -830,6 +834,7 @@ namespace NDIStaticMeshLocal
 			PreSkinnedLocalBoundsExtents = FVector3f::ZeroVector;
 			OwnerToMeshVector = FVector3f::ZeroVector;
 			PhysicsVelocity = FVector::ZeroVector;
+			bCaptureTransformsPerFrame = Interface->bCaptureTransformsPerFrame;
 			bUsePhysicsVelocity = Interface->bUsePhysicsBodyVelocity;
 			bComponentValid = false;
 			bMeshValid = false;
@@ -1091,7 +1096,10 @@ namespace NDIStaticMeshLocal
 			DeltaSeconds = InDeltaSeconds;
 
 			USceneComponent* SceneComponent = SceneComponentWeakPtr.Get();
-			UpdateTransformsCopyPrevious(SceneComponent, SystemInstance);
+			if (bCaptureTransformsPerFrame)
+			{
+				UpdateTransformsCopyPrevious(SceneComponent, SystemInstance);
+			}
 
 			if (bUsePhysicsVelocity)
 			{
@@ -3253,6 +3261,7 @@ bool UNiagaraDataInterfaceStaticMesh::Equals(const UNiagaraDataInterface* Other)
 		OtherTyped->SoftSourceActor == SoftSourceActor &&
 		OtherTyped->SourceComponent == SourceComponent &&
 		OtherTyped->SectionFilter.AllowedMaterialSlots == SectionFilter.AllowedMaterialSlots &&
+		OtherTyped->bCaptureTransformsPerFrame == bCaptureTransformsPerFrame &&
 		OtherTyped->bUsePhysicsBodyVelocity == bUsePhysicsBodyVelocity &&
 		OtherTyped->bAllowSamplingFromStreamingLODs == bAllowSamplingFromStreamingLODs &&
 		OtherTyped->LODIndex == LODIndex &&
@@ -3278,6 +3287,7 @@ bool UNiagaraDataInterfaceStaticMesh::CopyToInternal(UNiagaraDataInterface* Dest
 	OtherTyped->SoftSourceActor = SoftSourceActor;
 	OtherTyped->SourceComponent = SourceComponent;
 	OtherTyped->SectionFilter = SectionFilter;
+	OtherTyped->bCaptureTransformsPerFrame = bCaptureTransformsPerFrame;
 	OtherTyped->bUsePhysicsBodyVelocity = bUsePhysicsBodyVelocity;
 	OtherTyped->FilteredSockets = FilteredSockets;
 	OtherTyped->bAllowSamplingFromStreamingLODs = bAllowSamplingFromStreamingLODs;
