- Temp fix to handle null EditorData

#rnx

diff --git a/Engine/Plugins/FX/Niagara/Source/NiagaraEditor/Private/NiagaraScriptMergeManager.cpp b/Engine/Plugins/FX/Niagara/Source/NiagaraEditor/Private/NiagaraScriptMergeManager.cpp
--- a/Engine/Plugins/FX/Niagara/Source/NiagaraEditor/Private/NiagaraScriptMergeManager.cpp
+++ b/Engine/Plugins/FX/Niagara/Source/NiagaraEditor/Private/NiagaraScriptMergeManager.cpp
@@ -3670,6 +3670,15 @@ FNiagaraScriptMergeManager::FApplyDiffResults FNiagaraScriptMergeManager::ApplyE
 	FVersionedNiagaraEmitterData* BaseEmitterData = BaseEmitter.GetEmitterData();
 	UNiagaraEmitterEditorData* EditorData = Cast<UNiagaraEmitterEditorData>(BaseEmitterData->GetEditorData());
 
+	//-TEMP: Temporary code until we figure out why we don't have valid editor data
+	if (EditorData == nullptr)
+	{
+		FApplyDiffResults Results;
+		Results.bSucceeded = false;
+		Results.bModifiedGraph = false;
+		return Results;
+	}
+
 	check(EditorData != GetDefault<UNiagaraEmitterEditorData>());
 	check(EditorData->GetOuter() == BaseEmitter.Emitter);
 	
