Assset definitions for new Niagara Data Channels.

#rnx

diff --git a/Engine/Plugins/FX/Niagara/Source/NiagaraEditor/Private/AssetDefinitions/AssetDefinition_NiagaraDataChannel.cpp b/Engine/Plugins/FX/Niagara/Source/NiagaraEditor/Private/AssetDefinitions/AssetDefinition_NiagaraDataChannel.cpp
new file mode 100644
--- /dev/null
+++ b/Engine/Plugins/FX/Niagara/Source/NiagaraEditor/Private/AssetDefinitions/AssetDefinition_NiagaraDataChannel.cpp
@@ -0,0 +1,28 @@
+// Copyright Epic Games, Inc. All Rights Reserved.
+
+#include "AssetDefinitions/AssetDefinition_NiagaraDataChannel.h"
+
+#include "NiagaraEditorStyle.h"
+#include "NiagaraDataChannel.h"
+
+#define LOCTEXT_NAMESPACE "UAssetDefinition_NiagaraDataChannel"
+
+FLinearColor UAssetDefinition_NiagaraDataChannel::GetAssetColor() const
+{
+	return FNiagaraEditorStyle::Get().GetColor("NiagaraEditor.AssetColors.DataChannel");
+}
+
+TSoftClassPtr<> UAssetDefinition_NiagaraDataChannel::GetAssetClass() const
+{
+	return UNiagaraDataChannelAsset::StaticClass();
+}
+
+TConstArrayView<FAssetCategoryPath> UAssetDefinition_NiagaraDataChannel::GetAssetCategories() const
+{
+	static FAssetCategoryPath AssetPaths[] = { FAssetCategoryPath(LOCTEXT("NiagaraAssetsCategory", "FX"), LOCTEXT("NiagaraDataChannel_SubCategory", "Advanced")) };
+	return AssetPaths;
+}
+
+#undef LOCTEXT_NAMESPACE
+
+
diff --git a/Engine/Plugins/FX/Niagara/Source/NiagaraEditor/Private/AssetDefinitions/AssetDefinition_NiagaraDataChannel.h b/Engine/Plugins/FX/Niagara/Source/NiagaraEditor/Private/AssetDefinitions/AssetDefinition_NiagaraDataChannel.h
new file mode 100644
--- /dev/null
+++ b/Engine/Plugins/FX/Niagara/Source/NiagaraEditor/Private/AssetDefinitions/AssetDefinition_NiagaraDataChannel.h
@@ -0,0 +1,21 @@
+// Copyright Epic Games, Inc. All Rights Reserved.
+
+#pragma once
+
+#include "CoreMinimal.h"
+#include "AssetDefinitionDefault.h"
+
+#include "AssetDefinition_NiagaraDataChannel.generated.h"
+
+UCLASS()
+class UAssetDefinition_NiagaraDataChannel : public UAssetDefinitionDefault
+{
+	GENERATED_BODY()
+
+public:
+	// UAssetDefinition Implementation
+	virtual FText GetAssetDisplayName() const override { return NSLOCTEXT("AssetDefinition", "NiagaraDataChannel", "Niagara Data Channel"); }
+	virtual FLinearColor GetAssetColor() const override;
+	virtual TSoftClassPtr<UObject> GetAssetClass() const override;
+	virtual TConstArrayView<FAssetCategoryPath> GetAssetCategories() const override;
+};
diff --git a/Engine/Plugins/FX/Niagara/Source/NiagaraEditor/Private/NiagaraDataChannelAssetFactoryNew.cpp b/Engine/Plugins/FX/Niagara/Source/NiagaraEditor/Private/NiagaraDataChannelAssetFactoryNew.cpp
new file mode 100644
--- /dev/null
+++ b/Engine/Plugins/FX/Niagara/Source/NiagaraEditor/Private/NiagaraDataChannelAssetFactoryNew.cpp
@@ -0,0 +1,29 @@
+// Copyright Epic Games, Inc. All Rights Reserved.
+
+#include "NiagaraDataChannelAssetFactoryNew.h"
+#include "CoreMinimal.h"
+#include "Misc/ConfigCacheIni.h"
+#include "NiagaraSettings.h"
+#include "NiagaraDataChannel.h"
+
+UNiagaraDataChannelAssetFactoryNew::UNiagaraDataChannelAssetFactoryNew(const FObjectInitializer& ObjectInitializer)
+	: Super(ObjectInitializer)
+{
+	SupportedClass = UNiagaraDataChannelAsset::StaticClass();
+	bEditAfterNew = true;
+	bCreateNew = true;
+}
+
+UObject* UNiagaraDataChannelAssetFactoryNew::FactoryCreateNew(UClass* Class, UObject* InParent, FName Name, EObjectFlags Flags, UObject* Context, FFeedbackContext* Warn)
+{
+	check(Class->IsChildOf(UNiagaraDataChannelAsset::StaticClass()));
+
+	UNiagaraDataChannelAsset* NewChannel = NewObject<UNiagaraDataChannelAsset>(InParent, Class, Name, Flags | RF_Transactional);
+
+	return NewChannel;
+}
+
+bool UNiagaraDataChannelAssetFactoryNew::CanCreateNew()const
+{
+	return INiagaraModule::DataChannelsEnabled();
+}
diff --git a/Engine/Plugins/FX/Niagara/Source/NiagaraEditor/Public/NiagaraDataChannelAssetFactoryNew.h b/Engine/Plugins/FX/Niagara/Source/NiagaraEditor/Public/NiagaraDataChannelAssetFactoryNew.h
new file mode 100644
--- /dev/null
+++ b/Engine/Plugins/FX/Niagara/Source/NiagaraEditor/Public/NiagaraDataChannelAssetFactoryNew.h
@@ -0,0 +1,20 @@
+// Copyright Epic Games, Inc. All Rights Reserved.
+
+#pragma once
+
+#include "CoreMinimal.h"
+#include "UObject/ObjectMacros.h"
+#include "Factories/Factory.h"
+#include "NiagaraDataChannelAssetFactoryNew.generated.h"
+
+UCLASS(hidecategories = Object)
+class UNiagaraDataChannelAssetFactoryNew : public UFactory
+{
+	GENERATED_UCLASS_BODY()
+
+	//~ Begin UFactory Interface
+	virtual UObject* FactoryCreateNew(UClass* Class, UObject* InParent, FName Name, EObjectFlags Flags, UObject* Context, FFeedbackContext* Warn) override;
+
+	virtual bool CanCreateNew()const override;
+	//~ Begin UFactory Interface
+};
