Fix loctext warnings.

diff --git a/Engine/Plugins/FX/Niagara/Source/Niagara/Private/DataInterface/NiagaraDataInterfaceDataChannelWrite.cpp b/Engine/Plugins/FX/Niagara/Source/Niagara/Private/DataInterface/NiagaraDataInterfaceDataChannelWrite.cpp
--- a/Engine/Plugins/FX/Niagara/Source/Niagara/Private/DataInterface/NiagaraDataInterfaceDataChannelWrite.cpp
+++ b/Engine/Plugins/FX/Niagara/Source/Niagara/Private/DataInterface/NiagaraDataInterfaceDataChannelWrite.cpp
@@ -432,7 +432,7 @@ void UNiagaraDataInterfaceDataChannelWrite::GetFeedback(UNiagaraSystem* InAsset,
 
 	if(ShouldPublish() == false)
 	{
-		OutErrors.Emplace(FText::Format(LOCTEXT("DataChannelDoesNotExistErrorFmt", "Data Channel {0} does not publish it's data to the Game, CPU Simulations or GPU simulations."), FText::FromName(Channel.GetFName())),
+		OutErrors.Emplace(FText::Format(LOCTEXT("DataChannelDoesNotPublishtErrorFmt", "Data Channel {0} does not publish it's data to the Game, CPU Simulations or GPU simulations."), FText::FromName(Channel.GetFName())),
 			LOCTEXT("DataChannelDoesNotPublishErrorSummaryFmt", "Data Channel DI does not publish."),
 			FNiagaraDataInterfaceFix());
 	}
diff --git a/Engine/Plugins/FX/Niagara/Source/NiagaraEditor/Private/NiagaraActions.cpp b/Engine/Plugins/FX/Niagara/Source/NiagaraEditor/Private/NiagaraActions.cpp
--- a/Engine/Plugins/FX/Niagara/Source/NiagaraEditor/Private/NiagaraActions.cpp
+++ b/Engine/Plugins/FX/Niagara/Source/NiagaraEditor/Private/NiagaraActions.cpp
@@ -853,18 +853,18 @@ void FNiagaraDataInterfaceNodeActionProvider_DataChannelRead::CollectAddPinActio
 			FText Tooltip;
 			if (FuncName == TEXT("SpawnConditional"))
 			{
-				Category = FText::Format(LOCTEXT("NDIDataChannleReadSpawnCategoryFmt", "Conditions on {0} variables"), FText::FromString(Channel->GetAsset()->GetName()));
+				Category = FText::Format(LOCTEXT("NDISpawnConditionalAddPinCatFmt", "Conditions on {0} variables"), FText::FromString(Channel->GetAsset()->GetName()));
 				DisplayName = FText::FromName(SWCVar.GetName());
 				Tooltip = FText::Format(
-					LOCTEXT("NDIDataChannleReadAddParameterFmt", "Make this function conditional on {0} from Data Channel {1}.\nThe function will only have an effect for Data Channel entries that pass the comparisson test with the passed value."),
+					LOCTEXT("NDISpawnConditionalAddPinTooltipFmt", "Make this function conditional on {0} from Data Channel {1}.\nThe function will only have an effect for Data Channel entries that pass the comparisson test with the passed value."),
 					FText::FromName(SWCVar.GetName()),
 					FText::FromString(Channel->GetAsset()->GetName()));
 			}
 			else
 			{
-				Category = FText::Format(LOCTEXT("NDIDataChannleReadCategoryFmt", "{0}"), FText::FromString(Channel->GetAsset()->GetName()));
+				Category = FText::Format(LOCTEXT("NDIReadAddPinCatFmt", "{0}"), FText::FromString(Channel->GetAsset()->GetName()));
 				DisplayName = FText::FromName(SWCVar.GetName());
-				Tooltip = FText::Format(LOCTEXT("NDIDataChannleReadAddParameterFmt", "Read {0} from Data Channel {1}."), DisplayName, FText::FromString(Channel->GetAsset()->GetName()));
+				Tooltip = FText::Format(LOCTEXT("NDIReadAddPinTooltipFmt", "Read {0} from Data Channel {1}."), DisplayName, FText::FromString(Channel->GetAsset()->GetName()));
 			}
 
 			TSharedPtr<FNiagaraMenuAction> Action(new FNiagaraMenuAction(
