Fixed reliance on faulty version guid in renderer & sim stages in UI customization.
Replaced with emitter name.
This fixes a crash when changing emitter parent version as the UI customization would have the correct UNiagaraEmitter but the wrong guid (guid that points to a parent version instead of the actual emitter)

diff --git a/Engine/Plugins/FX/Niagara/Source/NiagaraEditor/Private/Customizations/NiagaraTypeCustomizations.cpp b/Engine/Plugins/FX/Niagara/Source/NiagaraEditor/Private/Customizations/NiagaraTypeCustomizations.cpp
--- a/Engine/Plugins/FX/Niagara/Source/NiagaraEditor/Private/Customizations/NiagaraTypeCustomizations.cpp
+++ b/Engine/Plugins/FX/Niagara/Source/NiagaraEditor/Private/Customizations/NiagaraTypeCustomizations.cpp
@@ -205,7 +205,7 @@ TArray<FNiagaraVariableBase> FNiagaraStackAssetAction_VarBind::FindVariables(con
 
 FName FNiagaraVariableAttributeBindingCustomization::GetVariableName() const
 {
-	if (BaseEmitter.Emitter && TargetVariableBinding)
+	if (OwningVersionedEmitter.Emitter && TargetVariableBinding)
 	{
 		return (TargetVariableBinding->GetName());
 	}
@@ -214,7 +214,7 @@ FName FNiagaraVariableAttributeBindingCustomization::GetVariableName() const
 
 FText FNiagaraVariableAttributeBindingCustomization::GetCurrentText() const
 {
-	if (BaseEmitter.Emitter && TargetVariableBinding)
+	if (OwningVersionedEmitter.Emitter && TargetVariableBinding)
 	{
 		return FText::FromName(TargetVariableBinding->GetName());
 	}
@@ -223,7 +223,7 @@ FText FNiagaraVariableAttributeBindingCustomization::GetCurrentText() const
 
 FText FNiagaraVariableAttributeBindingCustomization::GetTooltipText() const
 {
-	if (BaseEmitter.Emitter && TargetVariableBinding)
+	if (OwningVersionedEmitter.Emitter && TargetVariableBinding)
 	{
 		FString DefaultValueStr = TargetVariableBinding->GetDefaultValueString();
 
@@ -273,11 +273,11 @@ TArray<FName> FNiagaraVariableAttributeBindingCustomization::GetNames(const FVer
 	
 	TArray<UNiagaraGraph*> GraphsToTraverse;
 
-	GraphsToTraverse.Add(Cast<UNiagaraScriptSource>(BaseEmitter.GetEmitterData()->GraphSource)->NodeGraph);
+	GraphsToTraverse.Add(Cast<UNiagaraScriptSource>(OwningVersionedEmitter.GetEmitterData()->GraphSource)->NodeGraph);
 
 	TSet<FNiagaraVariableBase> Vars;
 
-	if(UNiagaraSystem* System = BaseEmitter.Emitter->GetTypedOuter<UNiagaraSystem>())
+	if(UNiagaraSystem* System = OwningVersionedEmitter.Emitter->GetTypedOuter<UNiagaraSystem>())
 	{
     	GraphsToTraverse.Add(Cast<UNiagaraScriptSource>(System->GetSystemUpdateScript()->GetLatestSource())->NodeGraph);
 	
@@ -296,7 +296,11 @@ TArray<FName> FNiagaraVariableAttributeBindingCustomization::GetNames(const FVer
 		for(UNiagaraNodeOutput* NodeOutput : OutputNodes)
 		{
 			FNiagaraParameterMapHistoryBuilder Builder;
-			Builder.ExclusiveEmitterHandle = EmitterHandleGuid;
+			
+			if(EmitterHandleGuid.IsValid())
+			{
+				Builder.ExclusiveEmitterHandle = EmitterHandleGuid;
+			}
 			
 			Builder.BuildParameterMaps(NodeOutput);
 
@@ -347,7 +351,7 @@ TArray<FName> FNiagaraVariableAttributeBindingCustomization::GetNames(const FVer
 
 void FNiagaraVariableAttributeBindingCustomization::CollectAllActions(FGraphActionListBuilderBase& OutAllActions)
 {
-	TArray<FName> EventNames = GetNames(BaseEmitter);
+	TArray<FName> EventNames = GetNames(OwningVersionedEmitter);
 	for (FName EventName : EventNames)
 	{
 		FText CategoryName = FText();
@@ -403,12 +407,12 @@ void FNiagaraVariableAttributeBindingCustomization::ChangeSource(FName InVarName
 	{
 		Obj->Modify();
 	}
-	check(BaseEmitter.Emitter);
+	check(OwningVersionedEmitter.Emitter);
 	check(RenderProps || SimulationStage);
 
 	PropertyHandle->NotifyPreChange();
 	const ENiagaraRendererSourceDataMode BindingSourceMode = RenderProps ? RenderProps->GetCurrentSourceMode() : ENiagaraRendererSourceDataMode::Emitter;
-	TargetVariableBinding->SetValue(InVarName, BaseEmitter, BindingSourceMode);
+	TargetVariableBinding->SetValue(InVarName, OwningVersionedEmitter, BindingSourceMode);
 	PropertyHandle->NotifyPostChange(EPropertyChangeType::ValueSet);
 	PropertyHandle->NotifyFinishedChangingProperties();
 }
@@ -420,7 +424,7 @@ void FNiagaraVariableAttributeBindingCustomization::ResetToDefault()
 
 EVisibility FNiagaraVariableAttributeBindingCustomization::IsResetToDefaultsVisible() const
 {
-	check(BaseEmitter.Emitter);
+	check(OwningVersionedEmitter.Emitter);
 	check(RenderProps || SimulationStage);
 	check(TargetVariableBinding);
 	check(DefaultVariableBinding);
@@ -438,14 +442,14 @@ FReply FNiagaraVariableAttributeBindingCustomization::OnResetToDefaultsClicked()
 	{
 		Obj->Modify();
 	}
-	check(BaseEmitter.Emitter);
+	check(OwningVersionedEmitter.Emitter);
 	check(RenderProps || SimulationStage);
 	check(TargetVariableBinding);
 	check(DefaultVariableBinding);
 
 	PropertyHandle->NotifyPreChange();
 	const ENiagaraRendererSourceDataMode BindingSourceMode = RenderProps ? RenderProps->GetCurrentSourceMode() : ENiagaraRendererSourceDataMode::Emitter;
-	TargetVariableBinding->ResetToDefault(*DefaultVariableBinding, BaseEmitter, BindingSourceMode);
+	TargetVariableBinding->ResetToDefault(*DefaultVariableBinding, OwningVersionedEmitter, BindingSourceMode);
 	PropertyHandle->NotifyPostChange(EPropertyChangeType::ValueSet);
 	PropertyHandle->NotifyFinishedChangingProperties();
 	return FReply::Handled();
@@ -455,7 +459,7 @@ void FNiagaraVariableAttributeBindingCustomization::CustomizeHeader(TSharedRef<I
 {
 	RenderProps = nullptr;
 	SimulationStage = nullptr;
-	BaseEmitter = FVersionedNiagaraEmitter();
+	OwningVersionedEmitter = FVersionedNiagaraEmitter();
 	PropertyHandle = InPropertyHandle;
 	TArray<UObject*> Objects;
 	PropertyHandle->GetOuterObjects(Objects);
@@ -475,22 +479,40 @@ void FNiagaraVariableAttributeBindingCustomization::CustomizeHeader(TSharedRef<I
 	if (Objects.Num() == 1)
 	{
 		RenderProps = Cast<UNiagaraRendererProperties>(Objects[0]);
-		if (RenderProps)
+		if(RenderProps)
 		{
-			BaseEmitter = RenderProps->GetOuterEmitter();
+			OwningVersionedEmitter = RenderProps->GetOuterEmitter();
 		}
 
 		SimulationStage = Cast<UNiagaraSimulationStageBase>(Objects[0]);
-		if ( SimulationStage )
+		if(SimulationStage)
 		{
-			BaseEmitter = SimulationStage->GetOuterEmitter();
+			OwningVersionedEmitter = SimulationStage->GetOuterEmitter();
 		}
-
-		if (BaseEmitter.Emitter)
+		
+		if (OwningVersionedEmitter.Emitter)
 		{
-			UNiagaraSystem* System = BaseEmitter.Emitter->GetTypedOuter<UNiagaraSystem>();
+			UNiagaraSystem* System = OwningVersionedEmitter.Emitter->GetTypedOuter<UNiagaraSystem>();
 			TSharedPtr<FNiagaraSystemViewModel> SystemViewModel = TNiagaraViewModelManager<UNiagaraSystem, FNiagaraSystemViewModel>::GetExistingViewModelForObject(System);
-			EmitterHandleGuid = SystemViewModel->GetEmitterHandleViewModelForEmitter(BaseEmitter)->GetId();
+			
+			if(SystemViewModel.IsValid())
+			{
+				// the owning emitter we retrieve might have the wrong version guid after changing emitter versions, as the sim stage & render props 'OuterEmitterVersion' can be overwritten during merge.
+				// so we rely on matching names to find the correct handle guid instead of using the version guid
+				const TArray<TSharedRef<FNiagaraEmitterHandleViewModel>>& EmitterHandleViewModels = SystemViewModel->GetEmitterHandleViewModels();
+				for(TSharedRef<FNiagaraEmitterHandleViewModel> EmitterHandleViewModel : EmitterHandleViewModels)
+				{
+					if(EmitterHandleViewModel->GetName().IsEqual(FName(OwningVersionedEmitter.Emitter->GetUniqueEmitterName())))
+					{
+						EmitterHandleGuid = EmitterHandleViewModel->GetId();
+					}
+				}
+
+				if(EmitterHandleGuid.IsValid() == false)
+				{
+					UE_LOG(LogNiagaraEditor, Warning, TEXT("EmitterHandleViewModel was not valid. This shouldn't happen and implies some mismatch in emitter names."));
+				}
+			}
 			
 			TargetVariableBinding = (FNiagaraVariableAttributeBinding*)PropertyHandle->GetValueBaseAddress((uint8*)Objects[0]);
 			DefaultVariableBinding = (FNiagaraVariableAttributeBinding*)PropertyHandle->GetValueBaseAddress((uint8*)Objects[0]->GetClass()->GetDefaultObject());
diff --git a/Engine/Plugins/FX/Niagara/Source/NiagaraEditor/Private/Customizations/NiagaraTypeCustomizations.h b/Engine/Plugins/FX/Niagara/Source/NiagaraEditor/Private/Customizations/NiagaraTypeCustomizations.h
--- a/Engine/Plugins/FX/Niagara/Source/NiagaraEditor/Private/Customizations/NiagaraTypeCustomizations.h
+++ b/Engine/Plugins/FX/Niagara/Source/NiagaraEditor/Private/Customizations/NiagaraTypeCustomizations.h
@@ -185,7 +185,7 @@ private:
 
 	TSharedPtr<IPropertyHandle> PropertyHandle;
 	TSharedPtr<SComboButton> ComboButton;
-	FVersionedNiagaraEmitter BaseEmitter;
+	FVersionedNiagaraEmitter OwningVersionedEmitter;
 	UNiagaraRendererProperties* RenderProps = nullptr;
 	class UNiagaraSimulationStageBase* SimulationStage = nullptr;
 	struct FNiagaraVariableAttributeBinding* TargetVariableBinding = nullptr;
