- Add information to show if the system is in edit mode to the debug hud

diff --git a/Engine/Plugins/FX/Niagara/Source/Niagara/Private/NiagaraDebugHud.cpp b/Engine/Plugins/FX/Niagara/Source/Niagara/Private/NiagaraDebugHud.cpp
--- a/Engine/Plugins/FX/Niagara/Source/Niagara/Private/NiagaraDebugHud.cpp
+++ b/Engine/Plugins/FX/Niagara/Source/Niagara/Private/NiagaraDebugHud.cpp
@@ -1017,6 +1017,9 @@ void FNiagaraDebugHud::GatherSystemInfo()
 		{
 			SystemDebugInfo.SystemName = GetNameSafe(NiagaraComponent->GetAsset());
 		}
+	#if WITH_EDITORONLY_DATA
+		SystemDebugInfo.bCompileForEdit = NiagaraComponent->GetAsset()->bCompileForEdit;
+	#endif
 		SystemDebugInfo.bShowInWorld = Settings.bSystemFilterEnabled && SystemDebugInfo.SystemName.MatchesWildcard(Settings.SystemFilter);
 		SystemDebugInfo.bPassesSystemFilter = !Settings.bSystemFilterEnabled || SystemDebugInfo.SystemName.MatchesWildcard(Settings.SystemFilter);
 
@@ -1570,8 +1573,15 @@ void FNiagaraDebugHud::DrawOverview(class FNiagaraWorldManager* WorldManager, FC
 				RowBGColor.A = Settings.SystemColorTableOpacity;
 				Canvas->DrawTile(X, Y, Col.MaxWidth, fAdvanceHeight, 0,0,0,0, RowBGColor);
 				const FLinearColor RowColor = SystemInfo.bShowInWorld ? DetailHighlightColor : DetailColor;
-				Canvas->DrawShadowedString(X, Y, *SystemInfo.SystemName, Font, RowColor);
-			});
+
+				const bool bShowEditMode = 
+#if WITH_EDITORONLY_DATA
+					SystemInfo.bCompileForEdit;
+#else
+					false;
+#endif
+					Canvas->DrawShadowedString(X, Y, bShowEditMode ? *FString::Printf(TEXT("%s (Edit Mode)"), *SystemInfo.SystemName) : *SystemInfo.SystemName, Font, RowColor);
+		});
 
 		if (Settings.OverviewMode == ENiagaraDebugHUDOverviewMode::Overview)
 		{
@@ -2167,6 +2177,9 @@ void FNiagaraDebugHud::DrawGpuComputeOverriew(class FNiagaraWorldManager* WorldM
 
 		const bool bShowDetailed = Settings.bSystemFilterEnabled && OwnerSystem->GetName().MatchesWildcard(Settings.SystemFilter);
 		SystemIt.Value().bShowDetailed = bShowDetailed;
+#if WITH_EDITORONLY_DATA
+		SystemIt.Value().bCompileForEdit = OwnerSystem->bCompileForEdit;
+#endif
 		bHasDetailedView |= bShowDetailed;
 		bHasSimpleView |= !bShowDetailed;
 	}
@@ -2212,6 +2225,12 @@ void FNiagaraDebugHud::DrawGpuComputeOverriew(class FNiagaraWorldManager* WorldM
 				{
 					const FGpuUsagePerStage& StageUsage = StageIt.Value();
 					OwnerSystem->GetFName().AppendString(SimpleTable.GetColumnText(0));
+#if WITH_EDITORONLY_DATA
+					if (SystemIt.Value().bCompileForEdit)
+					{
+						SimpleTable.GetColumnText(0).Append(TEXT(" (Edit Mode)"));
+					}
+#endif
 					OwnerEmitter->GetFName().AppendString(SimpleTable.GetColumnText(1));
 					StageIt.Key().AppendString(SimpleTable.GetColumnText(2));
 					SimpleTable.GetColumnText(3).Appendf(TEXT("%4.1f"), StageUsage.InstanceCount.GetAverage<float>());
@@ -2256,6 +2275,12 @@ void FNiagaraDebugHud::DrawGpuComputeOverriew(class FNiagaraWorldManager* WorldM
 			const FGpuUsagePerSystem& SystemUsage = SystemIt.Value();
 
 			OwnerSystem->GetFName().AppendString(SimpleTable.GetColumnText(0));
+#if WITH_EDITORONLY_DATA
+			if (SystemIt.Value().bCompileForEdit)
+			{
+				SimpleTable.GetColumnText(0).Append(TEXT(" (Edit Mode)"));
+			}
+#endif
 			SimpleTable.GetColumnText(1).Appendf(TEXT("%4.1f"), SystemUsage.InstanceCount.GetAverage<float>());
 			SimpleTable.GetColumnText(2).Appendf(TEXT("%llu"), SystemUsage.Microseconds.GetAverage());
 			SimpleTable.GetColumnText(3).Appendf(TEXT("%llu"), SystemUsage.Microseconds.GetMax());
diff --git a/Engine/Plugins/FX/Niagara/Source/Niagara/Private/NiagaraDebugHud.h b/Engine/Plugins/FX/Niagara/Source/Niagara/Private/NiagaraDebugHud.h
--- a/Engine/Plugins/FX/Niagara/Source/Niagara/Private/NiagaraDebugHud.h
+++ b/Engine/Plugins/FX/Niagara/Source/Niagara/Private/NiagaraDebugHud.h
@@ -93,6 +93,9 @@ class FNiagaraDebugHud
 	struct FSystemDebugInfo
 	{
 		FString		SystemName;
+	#if WITH_EDITORONLY_DATA
+		bool		bCompileForEdit = false;
+	#endif
 
 		#if WITH_PARTICLE_PERF_STATS
 		TSharedPtr<FNiagaraDebugHUDPerfStats> PerfStats = nullptr;
@@ -297,6 +300,9 @@ private:
 
 	struct FGpuUsagePerSystem
 	{
+#if WITH_EDITORONLY_DATA
+		bool bCompileForEdit = false;
+#endif
 		bool bShowDetailed = false;
 		FSmoothedCounter<uint32> InstanceCount;
 		FSmoothedCounter<uint64> Microseconds;
