- Add editor option to control how renderer categories should be expanded
- This allows the user to decide if they want the default rules, all expanding, all contracting, or expanding if modified

diff --git a/Engine/Plugins/FX/Niagara/Source/NiagaraEditor/Public/NiagaraEditorSettings.h b/Engine/Plugins/FX/Niagara/Source/NiagaraEditor/Public/NiagaraEditorSettings.h
--- a/Engine/Plugins/FX/Niagara/Source/NiagaraEditor/Public/NiagaraEditorSettings.h
+++ b/Engine/Plugins/FX/Niagara/Source/NiagaraEditor/Public/NiagaraEditorSettings.h
@@ -49,6 +49,19 @@ struct FNiagaraNewAssetDialogConfig
 	}
 };
 
+UENUM()
+enum class ENiagaraCategoryExpandState : uint8
+{
+	/** Categories will use the default expand / collapse state. */
+	Default,
+	/** Categories will use the default expand / collapse state unless they contain modified properties in which case they will expand. */
+	DefaultExpandModified,
+	/** Categories will ignore the default state and be collapsed. */
+	CollapseAll,
+	/** Categories will ignore the default state and be expanded. */
+	ExpandAll,
+};
+
 UENUM()
 enum class ENiagaraNamespaceMetadataOptions
 {
@@ -285,6 +298,9 @@ public:
 	UPROPERTY(config, EditAnywhere, Category = Niagara)
 	bool bAlwaysZoomToFitSystemGraph = true;
 
+	UPROPERTY(config, EditAnywhere, Category = Niagara)
+	ENiagaraCategoryExpandState RendererCategoryExpandState = ENiagaraCategoryExpandState::Default;
+
 	TArray<float> GetPlaybackSpeeds() const;
 
 	/** Gets whether or not auto-compile is enabled in the editors. */
diff --git a/Engine/Plugins/FX/Niagara/Source/NiagaraEditorWidgets/Private/DetailCustomizations/NiagaraDecalRendererDetails.cpp b/Engine/Plugins/FX/Niagara/Source/NiagaraEditorWidgets/Private/DetailCustomizations/NiagaraDecalRendererDetails.cpp
new file mode 100644
--- /dev/null
+++ b/Engine/Plugins/FX/Niagara/Source/NiagaraEditorWidgets/Private/DetailCustomizations/NiagaraDecalRendererDetails.cpp
@@ -0,0 +1,32 @@
+// Copyright Epic Games, Inc. All Rights Reserved.
+
+#include "NiagaraDecalRendererDetails.h"
+#include "DetailLayoutBuilder.h"
+#include "DetailCategoryBuilder.h"
+#include "IDetailChildrenBuilder.h"
+#include "DetailWidgetRow.h"
+#include "PropertyCustomizationHelpers.h"
+#include "Widgets/Layout/SExpandableArea.h"
+
+#include "NiagaraEditorSettings.h"
+
+TSharedRef<IDetailCustomization> FNiagaraDecalRendererDetails::MakeInstance()
+{
+	return MakeShared<FNiagaraDecalRendererDetails>();
+}
+
+void FNiagaraDecalRendererDetails::CustomizeDetails(IDetailLayoutBuilder& DetailBuilder)
+{
+	static const FName OrderedCategories[] =
+	{
+		FName("Decal Rendering"),
+	};
+
+	static const FName CollapsedCategories[] =
+	{
+		FName("Bindings"),
+		FName("Rendering"),
+		FName("Scalability"),
+	};
+	SetupCategories(DetailBuilder, MakeArrayView(OrderedCategories), MakeArrayView(CollapsedCategories));
+}
diff --git a/Engine/Plugins/FX/Niagara/Source/NiagaraEditorWidgets/Private/DetailCustomizations/NiagaraDecalRendererDetails.h b/Engine/Plugins/FX/Niagara/Source/NiagaraEditorWidgets/Private/DetailCustomizations/NiagaraDecalRendererDetails.h
new file mode 100644
--- /dev/null
+++ b/Engine/Plugins/FX/Niagara/Source/NiagaraEditorWidgets/Private/DetailCustomizations/NiagaraDecalRendererDetails.h
@@ -0,0 +1,15 @@
+// Copyright Epic Games, Inc. All Rights Reserved.
+
+#pragma once
+
+#include "NiagaraRendererDetails.h"
+
+class IDetailLayoutBuilder;
+
+class FNiagaraDecalRendererDetails : public FNiagaraRendererDetails
+{
+public:	
+	static TSharedRef<IDetailCustomization> MakeInstance();
+
+	virtual void CustomizeDetails(IDetailLayoutBuilder& DetailBuilder) override;
+};
\ No newline at end of file
diff --git a/Engine/Plugins/FX/Niagara/Source/NiagaraEditorWidgets/Private/DetailCustomizations/NiagaraLightRendererDetails.cpp b/Engine/Plugins/FX/Niagara/Source/NiagaraEditorWidgets/Private/DetailCustomizations/NiagaraLightRendererDetails.cpp
--- a/Engine/Plugins/FX/Niagara/Source/NiagaraEditorWidgets/Private/DetailCustomizations/NiagaraLightRendererDetails.cpp
+++ b/Engine/Plugins/FX/Niagara/Source/NiagaraEditorWidgets/Private/DetailCustomizations/NiagaraLightRendererDetails.cpp
@@ -15,32 +15,17 @@ TSharedRef<IDetailCustomization> FNiagaraLightRendererDetails::MakeInstance()
 
 void FNiagaraLightRendererDetails::CustomizeDetails(IDetailLayoutBuilder& DetailBuilder)
 {
-	// Categories to put first
+	static const FName OrderedCategories[] =
 	{
-		static const FName CategoriesToCollapse[] =
-		{
-			FName("Light Rendering"),
-			FName("Bindings"),
-		};
+		FName("Light Rendering"),
+		FName("Bindings"),
+	};
 
-		for (const FName& Category : CategoriesToCollapse)
-		{
-			IDetailCategoryBuilder& CategoryBuilder = DetailBuilder.EditCategory(Category);
-		}
-	}
-
-	// Collapse default categories
+	static const FName CollapsedCategories[] =
 	{
-		static const FName CategoriesToCollapse[] =
-		{
-			FName("Rendering"),
-			FName("Scalability"),
-		};
+		FName("Rendering"),
+		FName("Scalability"),
+	};
 
-		for (const FName& Category : CategoriesToCollapse)
-		{
-			IDetailCategoryBuilder& CategoryBuilder = DetailBuilder.EditCategory(Category);
-			CategoryBuilder.InitiallyCollapsed(true);
-		}
-	}
+	SetupCategories(DetailBuilder, MakeArrayView(OrderedCategories), MakeArrayView(CollapsedCategories));
 }
diff --git a/Engine/Plugins/FX/Niagara/Source/NiagaraEditorWidgets/Private/DetailCustomizations/NiagaraLightRendererDetails.h b/Engine/Plugins/FX/Niagara/Source/NiagaraEditorWidgets/Private/DetailCustomizations/NiagaraLightRendererDetails.h
--- a/Engine/Plugins/FX/Niagara/Source/NiagaraEditorWidgets/Private/DetailCustomizations/NiagaraLightRendererDetails.h
+++ b/Engine/Plugins/FX/Niagara/Source/NiagaraEditorWidgets/Private/DetailCustomizations/NiagaraLightRendererDetails.h
@@ -2,11 +2,11 @@
 
 #pragma once
 
-#include "IDetailCustomization.h"
+#include "NiagaraRendererDetails.h"
 
 class IDetailLayoutBuilder;
 
-class FNiagaraLightRendererDetails : public IDetailCustomization
+class FNiagaraLightRendererDetails : public FNiagaraRendererDetails
 {
 public:
 	static TSharedRef<IDetailCustomization> MakeInstance();
diff --git a/Engine/Plugins/FX/Niagara/Source/NiagaraEditorWidgets/Private/DetailCustomizations/NiagaraMeshRendererDetails.cpp b/Engine/Plugins/FX/Niagara/Source/NiagaraEditorWidgets/Private/DetailCustomizations/NiagaraMeshRendererDetails.cpp
--- a/Engine/Plugins/FX/Niagara/Source/NiagaraEditorWidgets/Private/DetailCustomizations/NiagaraMeshRendererDetails.cpp
+++ b/Engine/Plugins/FX/Niagara/Source/NiagaraEditorWidgets/Private/DetailCustomizations/NiagaraMeshRendererDetails.cpp
@@ -113,37 +113,22 @@ void FNiagaraMeshRendererDetails::CustomizeDetails(IDetailLayoutBuilder& DetailB
 {
 	LayoutBuilder = &DetailBuilder;
 
-	// Categories to put first
+	static const FName OrderedCategories[] =
 	{
-		static const FName CategoriesToCollapse[] =
-		{
-			FName("Mesh Rendering"),
-			FName("Sorting"),
-			FName("Visibility"),
-		};
-
-		for (const FName& Category : CategoriesToCollapse)
-		{
-			IDetailCategoryBuilder& CategoryBuilder = DetailBuilder.EditCategory(Category);
-		}
-	}
+		FName("Mesh Rendering"),
+		FName("Sorting"),
+		FName("Visibility"),
+	};
 
-	// Collapse default categories
+	static const FName CollapsedCategories[] =
 	{
-		static const FName CategoriesToCollapse[] =
-		{
-			FName("Bindings"),
-			FName("Rendering"),
-			FName("Scalability"),
-			FName("SubUV"),
-		};
+		FName("Bindings"),
+		FName("Rendering"),
+		FName("Scalability"),
+		FName("SubUV"),
+	};
 
-		for (const FName& Category : CategoriesToCollapse)
-		{
-			IDetailCategoryBuilder& CategoryBuilder = DetailBuilder.EditCategory(Category);
-			CategoryBuilder.InitiallyCollapsed(true);
-		}
-	}
+	SetupCategories(DetailBuilder, MakeArrayView(OrderedCategories), MakeArrayView(CollapsedCategories));
 
 	{
 		static const FName MeshesName = TEXT("Meshes");
diff --git a/Engine/Plugins/FX/Niagara/Source/NiagaraEditorWidgets/Private/DetailCustomizations/NiagaraMeshRendererDetails.h b/Engine/Plugins/FX/Niagara/Source/NiagaraEditorWidgets/Private/DetailCustomizations/NiagaraMeshRendererDetails.h
--- a/Engine/Plugins/FX/Niagara/Source/NiagaraEditorWidgets/Private/DetailCustomizations/NiagaraMeshRendererDetails.h
+++ b/Engine/Plugins/FX/Niagara/Source/NiagaraEditorWidgets/Private/DetailCustomizations/NiagaraMeshRendererDetails.h
@@ -2,13 +2,13 @@
 
 #pragma once
 
-#include "IDetailCustomization.h"
+#include "NiagaraRendererDetails.h"
 
 class IPropertyHandle;
 class IDetailChildrenBuilder;
 class IDetailLayoutBuilder;
 
-class FNiagaraMeshRendererDetails : public IDetailCustomization
+class FNiagaraMeshRendererDetails : public FNiagaraRendererDetails
 {
 public:	
 	static TSharedRef<IDetailCustomization> MakeInstance();
diff --git a/Engine/Plugins/FX/Niagara/Source/NiagaraEditorWidgets/Private/DetailCustomizations/NiagaraRendererDetails.cpp b/Engine/Plugins/FX/Niagara/Source/NiagaraEditorWidgets/Private/DetailCustomizations/NiagaraRendererDetails.cpp
new file mode 100644
--- /dev/null
+++ b/Engine/Plugins/FX/Niagara/Source/NiagaraEditorWidgets/Private/DetailCustomizations/NiagaraRendererDetails.cpp
@@ -0,0 +1,99 @@
+// Copyright Epic Games, Inc. All Rights Reserved.
+
+#include "NiagaraRendererDetails.h"
+#include "NiagaraEditorSettings.h"
+#include "NiagaraParameterBinding.h"
+
+#include "DetailLayoutBuilder.h"
+#include "DetailCategoryBuilder.h"
+#include "IDetailChildrenBuilder.h"
+#include "DetailWidgetRow.h"
+#include "PropertyCustomizationHelpers.h"
+#include "Widgets/Layout/SExpandableArea.h"
+
+void FNiagaraRendererDetails::SetupCategories(IDetailLayoutBuilder& DetailBuilder, TConstArrayView<FName> OrderedCategories, TConstArrayView<FName> CollapsedCategories)
+{
+	// Force the order of catefories
+	for (const FName& Category : OrderedCategories)
+	{
+		DetailBuilder.EditCategory(Category);
+		//IDetailCategoryBuilder& CategoryBuilder = DetailBuilder.EditCategory(Category);
+	}
+
+	// Iterate over each category setting the collapse state
+	TArray<FName> Categories;
+	DetailBuilder.GetCategoryNames(Categories);
+
+	const ENiagaraCategoryExpandState CategoryExpandState = GetDefault<UNiagaraEditorSettings>()->RendererCategoryExpandState;
+	const bool bForceCollapse = CategoryExpandState == ENiagaraCategoryExpandState::CollapseAll;
+	const bool bForceExpand = CategoryExpandState == ENiagaraCategoryExpandState::ExpandAll;
+	const bool bExpandMmodified = CategoryExpandState == ENiagaraCategoryExpandState::DefaultExpandModified;
+
+	for (const FName& Category : Categories)
+	{
+		IDetailCategoryBuilder& CategoryBuilder = DetailBuilder.EditCategory(Category);
+
+		bool bShouldCollapse = !bForceExpand && (bForceCollapse || CollapsedCategories.Contains(Category));
+		if (bShouldCollapse && bExpandMmodified)
+		{
+			TArray<TWeakObjectPtr<UObject>> SelectedObjects = DetailBuilder.GetSelectedObjects();
+
+			TArray<TSharedRef<IPropertyHandle>> DefaultProperties;
+			CategoryBuilder.GetDefaultProperties(DefaultProperties);
+
+			for (const TSharedRef<IPropertyHandle>& PropertyHandle : DefaultProperties)
+			{
+				// Special compare because attribute bindings are special
+				if ( FStructProperty* StructProperty = CastField<FStructProperty>(PropertyHandle->GetProperty()) )
+				{
+					if (StructProperty->Struct == FNiagaraVariableAttributeBinding::StaticStruct())
+					{
+						for (TWeakObjectPtr<UObject> WeakSelectedObject : DetailBuilder.GetSelectedObjects())
+						{
+							UObject* SelectedObject = WeakSelectedObject.Get();
+							if (SelectedObject == nullptr)
+							{
+								continue;
+							}
+							const FNiagaraVariableAttributeBinding* TargetVariableBinding = (FNiagaraVariableAttributeBinding*)PropertyHandle->GetValueBaseAddress((uint8*)SelectedObject);
+							const FNiagaraVariableAttributeBinding* DefaultVariableBinding = (FNiagaraVariableAttributeBinding*)PropertyHandle->GetValueBaseAddress((uint8*)SelectedObject->GetClass()->GetDefaultObject());
+							if (TargetVariableBinding->MatchesDefault(*DefaultVariableBinding, ENiagaraRendererSourceDataMode::Particles) == false)
+							{
+								bShouldCollapse = false;
+								break;
+							}
+						}
+						continue;
+					}
+					//-TODO: Implement this when the code makes it into here that fixes the reset to default, etc.
+					//else if (StructProperty->Struct == FNiagaraParameterBinding::StaticStruct())
+					//{
+					//	for (TWeakObjectPtr<UObject> WeakSelectedObject : DetailBuilder.GetSelectedObjects())
+					//	{
+					//		UObject* SelectedObject = WeakSelectedObject.Get();
+					//		if (SelectedObject == nullptr)
+					//		{
+					//			continue;
+					//		}
+					//		const FNiagaraParameterBinding* ParameterBinding = (FNiagaraParameterBinding*)PropertyHandle->GetValueBaseAddress((uint8*)SelectedObject);
+
+					//		if ( ParameterBinding->IsSetoToDefault() == false )
+					//		{
+					//			bShouldCollapse = false;
+					//			break;
+					//		}
+					//	}
+					//}
+				}
+
+				if (PropertyHandle->DiffersFromDefault())
+				{
+					bShouldCollapse = false;
+					break;
+				}
+			}
+
+		}
+		CategoryBuilder.InitiallyCollapsed(bShouldCollapse);
+	}
+}
diff --git a/Engine/Plugins/FX/Niagara/Source/NiagaraEditorWidgets/Private/DetailCustomizations/NiagaraRendererDetails.h b/Engine/Plugins/FX/Niagara/Source/NiagaraEditorWidgets/Private/DetailCustomizations/NiagaraRendererDetails.h
new file mode 100644
--- /dev/null
+++ b/Engine/Plugins/FX/Niagara/Source/NiagaraEditorWidgets/Private/DetailCustomizations/NiagaraRendererDetails.h
@@ -0,0 +1,13 @@
+// Copyright Epic Games, Inc. All Rights Reserved.
+
+#pragma once
+
+#include "IDetailCustomization.h"
+
+class IDetailLayoutBuilder;
+
+class FNiagaraRendererDetails : public IDetailCustomization
+{
+public:	
+	void SetupCategories(IDetailLayoutBuilder& DetailBuilder, TConstArrayView<FName> CategoryOrder, TConstArrayView<FName> CollapsedCategories);
+};
diff --git a/Engine/Plugins/FX/Niagara/Source/NiagaraEditorWidgets/Private/DetailCustomizations/NiagaraRibbonRendererDetails.cpp b/Engine/Plugins/FX/Niagara/Source/NiagaraEditorWidgets/Private/DetailCustomizations/NiagaraRibbonRendererDetails.cpp
--- a/Engine/Plugins/FX/Niagara/Source/NiagaraEditorWidgets/Private/DetailCustomizations/NiagaraRibbonRendererDetails.cpp
+++ b/Engine/Plugins/FX/Niagara/Source/NiagaraEditorWidgets/Private/DetailCustomizations/NiagaraRibbonRendererDetails.cpp
@@ -15,36 +15,21 @@ TSharedRef<IDetailCustomization> FNiagaraRibbonRendererDetails::MakeInstance()
 
 void FNiagaraRibbonRendererDetails::CustomizeDetails(IDetailLayoutBuilder& DetailBuilder)
 {
-	// Categories to put first
+	static const FName OrderedCategories[] =
 	{
-		static const FName CategoriesToCollapse[] =
-		{
-			FName("Ribbon Rendering"),
-			FName("Ribbon Shape"),
-			FName("Ribbon Tessellation"),
-			FName("Sorting"),
-			FName("Visibility"),
-		};
+		FName("Ribbon Rendering"),
+		FName("Ribbon Shape"),
+		FName("Ribbon Tessellation"),
+		FName("Sorting"),
+		FName("Visibility"),
+	};
 
-		for (const FName& Category : CategoriesToCollapse)
-		{
-			IDetailCategoryBuilder& CategoryBuilder = DetailBuilder.EditCategory(Category);
-		}
-	}
-
-	// Collapse default categories
+	static const FName CollapsedCategories[] =
 	{
-		static const FName CategoriesToCollapse[] =
-		{
-			FName("Bindings"),
-			FName("Rendering"),
-			FName("Scalability"),
-		};
+		FName("Bindings"),
+		FName("Rendering"),
+		FName("Scalability"),
+	};
 
-		for (const FName& Category : CategoriesToCollapse)
-		{
-			IDetailCategoryBuilder& CategoryBuilder = DetailBuilder.EditCategory(Category);
-			CategoryBuilder.InitiallyCollapsed(true);
-		}
-	}
+	SetupCategories(DetailBuilder, MakeArrayView(OrderedCategories), MakeArrayView(CollapsedCategories));
 }
diff --git a/Engine/Plugins/FX/Niagara/Source/NiagaraEditorWidgets/Private/DetailCustomizations/NiagaraRibbonRendererDetails.h b/Engine/Plugins/FX/Niagara/Source/NiagaraEditorWidgets/Private/DetailCustomizations/NiagaraRibbonRendererDetails.h
--- a/Engine/Plugins/FX/Niagara/Source/NiagaraEditorWidgets/Private/DetailCustomizations/NiagaraRibbonRendererDetails.h
+++ b/Engine/Plugins/FX/Niagara/Source/NiagaraEditorWidgets/Private/DetailCustomizations/NiagaraRibbonRendererDetails.h
@@ -2,11 +2,11 @@
 
 #pragma once
 
-#include "IDetailCustomization.h"
+#include "NiagaraRendererDetails.h"
 
 class IDetailLayoutBuilder;
 
-class FNiagaraRibbonRendererDetails : public IDetailCustomization
+class FNiagaraRibbonRendererDetails : public FNiagaraRendererDetails
 {
 public:
 	static TSharedRef<IDetailCustomization> MakeInstance();
diff --git a/Engine/Plugins/FX/Niagara/Source/NiagaraEditorWidgets/Private/DetailCustomizations/NiagaraSpriteRendererDetails.cpp b/Engine/Plugins/FX/Niagara/Source/NiagaraEditorWidgets/Private/DetailCustomizations/NiagaraSpriteRendererDetails.cpp
--- a/Engine/Plugins/FX/Niagara/Source/NiagaraEditorWidgets/Private/DetailCustomizations/NiagaraSpriteRendererDetails.cpp
+++ b/Engine/Plugins/FX/Niagara/Source/NiagaraEditorWidgets/Private/DetailCustomizations/NiagaraSpriteRendererDetails.cpp
@@ -8,6 +8,8 @@
 #include "PropertyCustomizationHelpers.h"
 #include "Widgets/Layout/SExpandableArea.h"
 
+#include "NiagaraEditorSettings.h"
+
 TSharedRef<IDetailCustomization> FNiagaraSpriteRendererDetails::MakeInstance()
 {
 	return MakeShared<FNiagaraSpriteRendererDetails>();
@@ -15,36 +17,20 @@ TSharedRef<IDetailCustomization> FNiagaraSpriteRendererDetails::MakeInstance()
 
 void FNiagaraSpriteRendererDetails::CustomizeDetails(IDetailLayoutBuilder& DetailBuilder)
 {
-	// Categories to put first
+	static const FName OrderedCategories[] =
 	{
-		static const FName CategoriesToCollapse[] =
-		{
-			FName("Sprite Rendering"),
-			FName("Sorting"),
-			FName("Visibility"),
-		};
-
-		for (const FName& Category : CategoriesToCollapse)
-		{
-			IDetailCategoryBuilder& CategoryBuilder = DetailBuilder.EditCategory(Category);
-		}
-	}
+		FName("Sprite Rendering"),
+		FName("Sorting"),
+		FName("Visibility"),
+	};
 
-	// Collapse default categories
+	static const FName CollapsedCategories[] =
 	{
-		static const FName CategoriesToCollapse[] =
-		{
-			FName("SubUV"),
-			FName("Cutout"),
-			FName("Bindings"),
-			FName("Rendering"),
-			FName("Scalability"),
-		};
-
-		for (const FName& Category : CategoriesToCollapse)
-		{
-			IDetailCategoryBuilder& CategoryBuilder = DetailBuilder.EditCategory(Category);
-			CategoryBuilder.InitiallyCollapsed(true);
-		}
-	}
+		FName("SubUV"),
+		FName("Cutout"),
+		FName("Bindings"),
+		FName("Rendering"),
+		FName("Scalability"),
+	};
+	SetupCategories(DetailBuilder, MakeArrayView(OrderedCategories), MakeArrayView(CollapsedCategories));
 }
diff --git a/Engine/Plugins/FX/Niagara/Source/NiagaraEditorWidgets/Private/DetailCustomizations/NiagaraSpriteRendererDetails.h b/Engine/Plugins/FX/Niagara/Source/NiagaraEditorWidgets/Private/DetailCustomizations/NiagaraSpriteRendererDetails.h
--- a/Engine/Plugins/FX/Niagara/Source/NiagaraEditorWidgets/Private/DetailCustomizations/NiagaraSpriteRendererDetails.h
+++ b/Engine/Plugins/FX/Niagara/Source/NiagaraEditorWidgets/Private/DetailCustomizations/NiagaraSpriteRendererDetails.h
@@ -2,11 +2,11 @@
 
 #pragma once
 
-#include "IDetailCustomization.h"
+#include "NiagaraRendererDetails.h"
 
 class IDetailLayoutBuilder;
 
-class FNiagaraSpriteRendererDetails : public IDetailCustomization
+class FNiagaraSpriteRendererDetails : public FNiagaraRendererDetails
 {
 public:	
 	static TSharedRef<IDetailCustomization> MakeInstance();
diff --git a/Engine/Plugins/FX/Niagara/Source/NiagaraEditorWidgets/Private/NiagaraEditorWidgetsModule.cpp b/Engine/Plugins/FX/Niagara/Source/NiagaraEditorWidgets/Private/NiagaraEditorWidgetsModule.cpp
--- a/Engine/Plugins/FX/Niagara/Source/NiagaraEditorWidgets/Private/NiagaraEditorWidgetsModule.cpp
+++ b/Engine/Plugins/FX/Niagara/Source/NiagaraEditorWidgets/Private/NiagaraEditorWidgetsModule.cpp
@@ -12,6 +12,7 @@
 #include "DetailCustomizations/NiagaraDataInterfaceStaticMeshDetails.h"
 #include "DetailCustomizations/NiagaraDataInterfaceMeshRendererInfoDetails.h"
 #include "DetailCustomizations/NiagaraDataInterfaceSpriteRendererInfoDetails.h"
+#include "DetailCustomizations/NiagaraDecalRendererDetails.h"
 #include "DetailCustomizations/NiagaraLightRendererDetails.h"
 #include "DetailCustomizations/NiagaraMeshRendererDetails.h"
 #include "DetailCustomizations/NiagaraRibbonRendererDetails.h"
@@ -142,6 +143,7 @@ void FNiagaraEditorWidgetsModule::StartupModule()
 	PropertyModule.RegisterCustomClassLayout("NiagaraMeshRendererProperties", FOnGetDetailCustomizationInstance::CreateStatic(&FNiagaraMeshRendererDetails::MakeInstance));
 	PropertyModule.RegisterCustomClassLayout("NiagaraRibbonRendererProperties", FOnGetDetailCustomizationInstance::CreateStatic(&FNiagaraRibbonRendererDetails::MakeInstance));
 	PropertyModule.RegisterCustomClassLayout("NiagaraSpriteRendererProperties", FOnGetDetailCustomizationInstance::CreateStatic(&FNiagaraSpriteRendererDetails::MakeInstance));
+	PropertyModule.RegisterCustomClassLayout("NiagaraDecalRendererProperties", FOnGetDetailCustomizationInstance::CreateStatic(&FNiagaraDecalRendererDetails::MakeInstance));
 
 	PropertyModule.RegisterCustomClassLayout("NiagaraDataInterfaceDataChannelRead", FOnGetDetailCustomizationInstance::CreateStatic(&FNiagaraDataInterfaceDataChannelReadDetails::MakeInstance));
 
