Fix vector RGBA grids- allocated as wrong texture type

diff --git a/Engine/Plugins/FX/Niagara/Source/Niagara/Private/NiagaraDataInterfaceGrid3DCollection.cpp b/Engine/Plugins/FX/Niagara/Source/Niagara/Private/NiagaraDataInterfaceGrid3DCollection.cpp
--- a/Engine/Plugins/FX/Niagara/Source/Niagara/Private/NiagaraDataInterfaceGrid3DCollection.cpp
+++ b/Engine/Plugins/FX/Niagara/Source/Niagara/Private/NiagaraDataInterfaceGrid3DCollection.cpp
@@ -3390,7 +3390,7 @@ bool UNiagaraDataInterfaceGrid3DCollection::InitPerInstanceData(void* PerInstanc
 			{
 			case 1: InstanceData->PixelFormat = EPixelFormat::PF_R32_FLOAT; break;
 			case 2:	InstanceData->PixelFormat = EPixelFormat::PF_G32R32F; break;
-			case 3: InstanceData->PixelFormat = EPixelFormat::PF_R32G32B32F; break;
+			case 3: InstanceData->PixelFormat = EPixelFormat::PF_A32B32G32R32F; break;
 			case 4:	InstanceData->PixelFormat = EPixelFormat::PF_A32B32G32R32F; break;
 			}
 			break;
