- Remove parameter store stack code as we no longer use it

#rnx

diff --git a/Engine/Plugins/FX/Niagara/Source/NiagaraEditor/Private/ViewModels/Stack/NiagaraStackParameterStoreEntry.cpp b/Engine/Plugins/FX/Niagara/Source/NiagaraEditor/Private/ViewModels/Stack/NiagaraStackParameterStoreEntry.cpp
deleted file mode 100644
diff --git a/Engine/Plugins/FX/Niagara/Source/NiagaraEditor/Private/ViewModels/Stack/NiagaraStackSystemSettingsGroup.cpp b/Engine/Plugins/FX/Niagara/Source/NiagaraEditor/Private/ViewModels/Stack/NiagaraStackSystemSettingsGroup.cpp
--- a/Engine/Plugins/FX/Niagara/Source/NiagaraEditor/Private/ViewModels/Stack/NiagaraStackSystemSettingsGroup.cpp
+++ b/Engine/Plugins/FX/Niagara/Source/NiagaraEditor/Private/ViewModels/Stack/NiagaraStackSystemSettingsGroup.cpp
@@ -5,7 +5,6 @@
 #include "ViewModels/NiagaraScriptGraphViewModel.h"
 #include "NiagaraNodeFunctionCall.h"
 #include "ViewModels/Stack/NiagaraStackErrorItem.h"
-#include "ViewModels/Stack/NiagaraStackParameterStoreEntry.h"
 #include "Internationalization/Internationalization.h"
 #include "NiagaraStackEditorData.h"
 #include "ViewModels/Stack/NiagaraStackGraphUtilities.h"
@@ -142,77 +141,5 @@ FText UNiagaraStackSystemUserParametersGroup::GetIconText() const
 	return FText::FromString(FString(TEXT("\xf007")/* fa-user */));
 }
 
-void UNiagaraStackParameterStoreItem::Initialize(
-	FRequiredEntryData InRequiredEntryData,
-	UObject* InOwner,
-	FNiagaraParameterStore* InParameterStore,
-	INiagaraStackItemGroupAddUtilities* InGroupAddUtilities)
-{
-	Super::Initialize(InRequiredEntryData, TEXT("ParameterStoreItem"));
-
-	Owner = InOwner;
-	ParameterStore = InParameterStore;
-	ParameterStoreChangedHandle = ParameterStore->AddOnChangedHandler(
-		FNiagaraParameterStore::FOnChanged::FDelegate::CreateUObject(this, &UNiagaraStackParameterStoreItem::ParameterStoreChanged));
-	GroupAddUtilities = InGroupAddUtilities;
-}
-
-FText UNiagaraStackParameterStoreItem::GetDisplayName() const
-{
-	return LOCTEXT("ParameterItemDisplayName", "User Parameters");
-}
-
-FText UNiagaraStackParameterStoreItem::GetTooltipText() const
-{
-	return LOCTEXT("ParameterItemTooltip", "Displays the variables created in the User namespace. These variables are exposed to owning UComponents, blueprints, etc.");
-}
-
-void UNiagaraStackParameterStoreItem::RefreshChildrenInternal(const TArray<UNiagaraStackEntry*>& CurrentChildren, TArray<UNiagaraStackEntry*>& NewChildren, TArray<FStackIssue>& NewIssues)
-{
-	if (ParameterStore != nullptr && ParameterStore->ReadParameterVariables().Num() > 0)
-	{
-		TArray<FNiagaraVariable> Variables;
-		ParameterStore->GetParameters(Variables);
-
-		for (FNiagaraVariable& Var : Variables)
-		{
-			UNiagaraStackParameterStoreEntry* ValueObjectEntry = FindCurrentChildOfTypeByPredicate<UNiagaraStackParameterStoreEntry>(CurrentChildren,
-				[=](UNiagaraStackParameterStoreEntry* CurrentEntry) { 
-				bool bSameName = CurrentEntry->GetDisplayName().ToString() == Var.GetName().ToString();
-				bool bSameType = CurrentEntry->GetInputType() == Var.GetType();
-				return bSameName && bSameType;
-			});
-
-			if (ValueObjectEntry == nullptr)
-			{
-				ValueObjectEntry = NewObject<UNiagaraStackParameterStoreEntry>(this);
-				ValueObjectEntry->Initialize(CreateDefaultChildRequiredData(), Owner.Get(), ParameterStore, Var.GetName().ToString(), Var.GetType(), GetStackEditorDataKey());
-			}
-
-			NewChildren.Add(ValueObjectEntry);
-		}
-	}
-	Super::RefreshChildrenInternal(CurrentChildren, NewChildren, NewIssues);
-}
-
-void UNiagaraStackParameterStoreItem::FinalizeInternal()
-{
-	if (Owner.IsValid() && ParameterStore != nullptr)
-	{
-		ParameterStore->RemoveOnChangedHandler(ParameterStoreChangedHandle);
-		Owner.Reset();
-		ParameterStore = nullptr;
-	}
-	Super::FinalizeInternal();
-}
-
-void UNiagaraStackParameterStoreItem::ParameterStoreChanged()
-{
-	if (IsFinalized() == false && Owner.IsValid())
-	{
-		RefreshChildren();
-	}
-}
-
 #undef LOCTEXT_NAMESPACE
 
diff --git a/Engine/Plugins/FX/Niagara/Source/NiagaraEditor/Public/ViewModels/Stack/NiagaraStackParameterStoreEntry.h b/Engine/Plugins/FX/Niagara/Source/NiagaraEditor/Public/ViewModels/Stack/NiagaraStackParameterStoreEntry.h
deleted file mode 100644
diff --git a/Engine/Plugins/FX/Niagara/Source/NiagaraEditor/Public/ViewModels/Stack/NiagaraStackSystemSettingsGroup.h b/Engine/Plugins/FX/Niagara/Source/NiagaraEditor/Public/ViewModels/Stack/NiagaraStackSystemSettingsGroup.h
--- a/Engine/Plugins/FX/Niagara/Source/NiagaraEditor/Public/ViewModels/Stack/NiagaraStackSystemSettingsGroup.h
+++ b/Engine/Plugins/FX/Niagara/Source/NiagaraEditor/Public/ViewModels/Stack/NiagaraStackSystemSettingsGroup.h
@@ -4,8 +4,8 @@
 
 #include "ViewModels/Stack/NiagaraStackItemGroup.h"
 #include "NiagaraCommon.h"
+#include "NiagaraStackItem.h"
 #include "NiagaraParameterStore.h"
-#include "ViewModels/Stack/NiagaraStackParameterStoreEntry.h"
 #include "Delegates/IDelegateInstance.h"
 #include "NiagaraStackSystemSettingsGroup.generated.h"
 
@@ -45,32 +45,3 @@ private:
 	FDelegateHandle ParameterStoreChangedHandle;
 };
 
-UCLASS()
-class NIAGARAEDITOR_API UNiagaraStackParameterStoreItem : public UNiagaraStackItem
-{
-	GENERATED_BODY()
-
-public:
-	void Initialize(FRequiredEntryData InRequiredEntryData, UObject* InOwner, FNiagaraParameterStore* InParameterStore, INiagaraStackItemGroupAddUtilities* InGroupAddUtilities);
-
-	virtual FText GetDisplayName() const override;
-
-	virtual FText GetTooltipText() const override;
-
-	virtual bool GetShouldShowInOverview() const override { return false; }
-
-	INiagaraStackItemGroupAddUtilities* GetGroupAddUtilities() const { return GroupAddUtilities; }
-
-protected:
-	virtual void RefreshChildrenInternal(const TArray<UNiagaraStackEntry*>& CurrentChildren, TArray<UNiagaraStackEntry*>& NewChildren, TArray<FStackIssue>& NewIssues) override;
-	virtual void FinalizeInternal() override;
-
-private:
-	void ParameterStoreChanged();
-
-private:
-	TWeakObjectPtr<UObject> Owner;
-	FNiagaraParameterStore* ParameterStore;
-	INiagaraStackItemGroupAddUtilities* GroupAddUtilities;
-	FDelegateHandle ParameterStoreChangedHandle;
-};
\ No newline at end of file
diff --git a/Engine/Plugins/FX/Niagara/Source/NiagaraEditorWidgets/Private/SNiagaraStack.cpp b/Engine/Plugins/FX/Niagara/Source/NiagaraEditorWidgets/Private/SNiagaraStack.cpp
--- a/Engine/Plugins/FX/Niagara/Source/NiagaraEditorWidgets/Private/SNiagaraStack.cpp
+++ b/Engine/Plugins/FX/Niagara/Source/NiagaraEditorWidgets/Private/SNiagaraStack.cpp
@@ -23,7 +23,6 @@
 #include "ViewModels/Stack/NiagaraStackFunctionInputCollection.h"
 #include "ViewModels/Stack/NiagaraStackModuleItem.h"
 #include "ViewModels/Stack/NiagaraStackSystemSettingsGroup.h"
-#include "ViewModels/Stack/NiagaraStackParameterStoreEntry.h"
 #include "Framework/Commands/GenericCommands.h"
 #include "IContentBrowserSingleton.h"
 #include "ContentBrowserModule.h"
@@ -42,9 +41,6 @@
 #include "Stack/SNiagaraStackItemFooter.h"
 #include "Stack/SNiagaraStackItemGroup.h"
 #include "Stack/SNiagaraStackModuleItem.h"
-#include "Stack/SNiagaraStackParameterStoreItem.h"
-#include "Stack/SNiagaraStackParameterStoreEntryName.h"
-#include "Stack/SNiagaraStackParameterStoreEntryValue.h"
 #include "Stack/SNiagaraStackTableRow.h"
 #include "Stack/SNiagaraStackIssueIcon.h"
 #include "NiagaraEditorWidgetsUtilities.h"
@@ -927,19 +923,6 @@ SNiagaraStack::FRowWidgets SNiagaraStack::ConstructNameAndValueWidgetsForItem(UN
 		UNiagaraStackItemFooter* ItemExpander = CastChecked<UNiagaraStackItemFooter>(Item);
 		return FRowWidgets(SNew(SNiagaraStackItemFooter, *ItemExpander));
 	}
-	else if (Item->IsA<UNiagaraStackParameterStoreItem>())
-	{
-		UNiagaraStackParameterStoreItem* StackEntry = CastChecked<UNiagaraStackParameterStoreItem>(Item);
-		return FRowWidgets(SNew(SNiagaraStackParameterStoreItem, *StackEntry, StackViewModel));
-	}
-	else if (Item->IsA<UNiagaraStackParameterStoreEntry>())
-	{
-		UNiagaraStackParameterStoreEntry* StackEntry = CastChecked<UNiagaraStackParameterStoreEntry>(Item);
-		return FRowWidgets(
-			SNew(SNiagaraStackParameterStoreEntryName, StackEntry, StackViewModel)
-			.IsSelected(Container, &SNiagaraStackTableRow::IsSelected),
-			SNew(SNiagaraStackParameterStoreEntryValue, StackEntry));
-	}
 	else if (Item->IsA<UNiagaraStackInputCategory>())
 	{
 		Container->SetOverrideNameAlignment(HAlign_Left, VAlign_Center);
diff --git a/Engine/Plugins/FX/Niagara/Source/NiagaraEditorWidgets/Private/Stack/SNiagaraStackParameterStoreEntryName.cpp b/Engine/Plugins/FX/Niagara/Source/NiagaraEditorWidgets/Private/Stack/SNiagaraStackParameterStoreEntryName.cpp
deleted file mode 100644
diff --git a/Engine/Plugins/FX/Niagara/Source/NiagaraEditorWidgets/Private/Stack/SNiagaraStackParameterStoreEntryName.h b/Engine/Plugins/FX/Niagara/Source/NiagaraEditorWidgets/Private/Stack/SNiagaraStackParameterStoreEntryName.h
deleted file mode 100644
diff --git a/Engine/Plugins/FX/Niagara/Source/NiagaraEditorWidgets/Private/Stack/SNiagaraStackParameterStoreEntryValue.cpp b/Engine/Plugins/FX/Niagara/Source/NiagaraEditorWidgets/Private/Stack/SNiagaraStackParameterStoreEntryValue.cpp
deleted file mode 100644
diff --git a/Engine/Plugins/FX/Niagara/Source/NiagaraEditorWidgets/Private/Stack/SNiagaraStackParameterStoreEntryValue.h b/Engine/Plugins/FX/Niagara/Source/NiagaraEditorWidgets/Private/Stack/SNiagaraStackParameterStoreEntryValue.h
deleted file mode 100644
diff --git a/Engine/Plugins/FX/Niagara/Source/NiagaraEditorWidgets/Private/Stack/SNiagaraStackParameterStoreItem.cpp b/Engine/Plugins/FX/Niagara/Source/NiagaraEditorWidgets/Private/Stack/SNiagaraStackParameterStoreItem.cpp
deleted file mode 100644
diff --git a/Engine/Plugins/FX/Niagara/Source/NiagaraEditorWidgets/Private/Stack/SNiagaraStackParameterStoreItem.h b/Engine/Plugins/FX/Niagara/Source/NiagaraEditorWidgets/Private/Stack/SNiagaraStackParameterStoreItem.h
deleted file mode 100644
