NiagraAudioDataInterface fixes

diff --git a/Engine/Plugins/FX/Niagara/Source/Niagara/Private/NiagaraDataInterfaceAudio.cpp b/Engine/Plugins/FX/Niagara/Source/Niagara/Private/NiagaraDataInterfaceAudio.cpp
--- a/Engine/Plugins/FX/Niagara/Source/Niagara/Private/NiagaraDataInterfaceAudio.cpp
+++ b/Engine/Plugins/FX/Niagara/Source/Niagara/Private/NiagaraDataInterfaceAudio.cpp
@@ -131,13 +131,15 @@ FNiagaraDataInterfaceProxySubmix::FNiagaraDataInterfaceProxySubmix(int32 InNumSa
 	, NumSamplesToBuffer(InNumSamplesToBuffer)
 {
 	check(NumSamplesToBuffer > 0);
-
-	DeviceCreatedHandle = FAudioDeviceManagerDelegates::OnAudioDeviceCreated.AddRaw(this, &FNiagaraDataInterfaceProxySubmix::OnNewDeviceCreated);
-	DeviceDestroyedHandle = FAudioDeviceManagerDelegates::OnAudioDeviceDestroyed.AddRaw(this, &FNiagaraDataInterfaceProxySubmix::OnDeviceDestroyed);
 }
 
 FNiagaraDataInterfaceProxySubmix::~FNiagaraDataInterfaceProxySubmix()
 {
+	// The proxy is deleted on the GFX thread.
+	// All deletegate removal should be done in BeginDestroy, called by game thead.
+	check(!DeviceDestroyedHandle.IsValid());
+	check(!DeviceCreatedHandle.IsValid());
+	check(SubmixListeners.IsEmpty());
 }
 
 void FNiagaraDataInterfaceProxySubmix::RegisterToAllAudioDevices()
@@ -159,6 +161,17 @@ void FNiagaraDataInterfaceProxySubmix::UnregisterFromAllAudioDevices()
 
 void FNiagaraDataInterfaceProxySubmix::OnUpdateSubmix(USoundSubmix* Submix)
 {
+	check(IsInGameThread());
+	
+	if (!DeviceCreatedHandle.IsValid())
+	{
+		DeviceCreatedHandle = FAudioDeviceManagerDelegates::OnAudioDeviceCreated.AddRaw(this, &FNiagaraDataInterfaceProxySubmix::OnNewDeviceCreated);
+	}
+	if (!DeviceDestroyedHandle.IsValid())
+	{
+		DeviceDestroyedHandle = FAudioDeviceManagerDelegates::OnAudioDeviceDestroyed.AddRaw(this, &FNiagaraDataInterfaceProxySubmix::OnDeviceDestroyed);
+	}
+	
 	if (bIsSubmixListenerRegistered)
 	{
 		UnregisterFromAllAudioDevices();
@@ -243,6 +256,8 @@ void FNiagaraDataInterfaceProxySubmix::OnBeginDestroy()
 	check(IsInGameThread());
 	FAudioDeviceManagerDelegates::OnAudioDeviceCreated.Remove(DeviceCreatedHandle);
 	FAudioDeviceManagerDelegates::OnAudioDeviceDestroyed.Remove(DeviceDestroyedHandle);
+	DeviceCreatedHandle.Reset();
+	DeviceDestroyedHandle.Reset();
 
 	if (bIsSubmixListenerRegistered)
 	{
@@ -300,8 +315,7 @@ bool UNiagaraDataInterfaceAudioSubmix::Equals(const UNiagaraDataInterface* Other
 void UNiagaraDataInterfaceAudioSubmix::PostEditChangeProperty(struct FPropertyChangedEvent& PropertyChangedEvent)
 {
 	Super::PostEditChangeProperty(PropertyChangedEvent);
-	
-	static FName SubmixFName = GET_MEMBER_NAME_CHECKED(UNiagaraDataInterfaceAudioSubmix, Submix);
+	static const FName SubmixFName = GET_MEMBER_NAME_CHECKED(UNiagaraDataInterfaceAudioSubmix, Submix);
 
 	// Regenerate on save any compressed sound formats or if analysis needs to be re-done
 	if (FProperty* PropertyThatChanged = PropertyChangedEvent.Property)
diff --git a/Engine/Plugins/FX/Niagara/Source/Niagara/Private/NiagaraDataInterfaceAudioOscilloscope.cpp b/Engine/Plugins/FX/Niagara/Source/Niagara/Private/NiagaraDataInterfaceAudioOscilloscope.cpp
--- a/Engine/Plugins/FX/Niagara/Source/Niagara/Private/NiagaraDataInterfaceAudioOscilloscope.cpp
+++ b/Engine/Plugins/FX/Niagara/Source/Niagara/Private/NiagaraDataInterfaceAudioOscilloscope.cpp
@@ -32,6 +32,7 @@ FNiagaraDataInterfaceProxyOscilloscope::FNiagaraDataInterfaceProxyOscilloscope(i
 	, ScopeInMilliseconds(InScopeInMillseconds)
 	, NumChannelsInDownsampledBuffer(0)
 {
+	check(IsInGameThread());
 	DeviceCreatedHandle = FAudioDeviceManagerDelegates::OnAudioDeviceCreated.AddRaw(this, &FNiagaraDataInterfaceProxyOscilloscope::OnNewDeviceCreated);
 	DeviceDestroyedHandle = FAudioDeviceManagerDelegates::OnAudioDeviceDestroyed.AddRaw(this, &FNiagaraDataInterfaceProxyOscilloscope::OnDeviceDestroyed);
 
@@ -405,6 +406,8 @@ void FNiagaraDataInterfaceProxyOscilloscope::OnBeginDestroy()
 	check(IsInGameThread());
 	FAudioDeviceManagerDelegates::OnAudioDeviceCreated.Remove(DeviceCreatedHandle);
 	FAudioDeviceManagerDelegates::OnAudioDeviceDestroyed.Remove(DeviceDestroyedHandle);
+	DeviceCreatedHandle.Reset();
+	DeviceDestroyedHandle.Reset();
 
 	if (bIsSubmixListenerRegistered)
 	{
