- When we close the system editor ensure we recompile back into performance mode

diff --git a/Engine/Plugins/FX/Niagara/Source/NiagaraEditor/Private/ViewModels/NiagaraSystemViewModel.cpp b/Engine/Plugins/FX/Niagara/Source/NiagaraEditor/Private/ViewModels/NiagaraSystemViewModel.cpp
--- a/Engine/Plugins/FX/Niagara/Source/NiagaraEditor/Private/ViewModels/NiagaraSystemViewModel.cpp
+++ b/Engine/Plugins/FX/Niagara/Source/NiagaraEditor/Private/ViewModels/NiagaraSystemViewModel.cpp
@@ -255,7 +255,12 @@ void FNiagaraSystemViewModel::Cleanup()
 
 	if (System)
 	{
-		System->bCompileForEdit = false;
+		if ( System->bCompileForEdit )
+		{
+			System->bCompileForEdit = false;
+			System->InvalidateActiveCompiles();
+			System->RequestCompile(false);
+		}
 		System->OnSystemPostEditChange().Remove(SystemChangedDelegateHandle);
 	}
 
