Fixed mismatched copy constructor/assignment operator in FNiagaraVariableWithOffset.

diff --git a/Engine/Plugins/FX/Niagara/Source/Niagara/Public/NiagaraParameterStore.h b/Engine/Plugins/FX/Niagara/Source/Niagara/Public/NiagaraParameterStore.h
--- a/Engine/Plugins/FX/Niagara/Source/Niagara/Public/NiagaraParameterStore.h
+++ b/Engine/Plugins/FX/Niagara/Source/Niagara/Public/NiagaraParameterStore.h
@@ -140,7 +140,6 @@ struct FNiagaraVariableWithOffset : public FNiagaraVariableBase
 
 	// Those constructor enforce that there are no data allocated.
 	FORCEINLINE FNiagaraVariableWithOffset() : Offset(INDEX_NONE) {}
-	FORCEINLINE FNiagaraVariableWithOffset(const FNiagaraVariableWithOffset& InRef) : FNiagaraVariableBase(InRef.GetType(), InRef.GetName()), Offset(InRef.Offset), StructConverter(InRef.StructConverter) {}
 	FORCEINLINE FNiagaraVariableWithOffset(const FNiagaraVariableBase& InVariable, int32 InOffset, const FNiagaraLwcStructConverter& InStructConverter) : FNiagaraVariableBase(InVariable.GetType(), InVariable.GetName()), Offset(InOffset), StructConverter(InStructConverter) {}
 
 	bool Serialize(FArchive& Ar);
