Added nullptr check for GEditor

diff --git a/Engine/Plugins/FX/Niagara/Source/NiagaraEditor/Private/ViewModels/HierarchyEditor/NiagaraHierarchyViewModelBase.cpp b/Engine/Plugins/FX/Niagara/Source/NiagaraEditor/Private/ViewModels/HierarchyEditor/NiagaraHierarchyViewModelBase.cpp
--- a/Engine/Plugins/FX/Niagara/Source/NiagaraEditor/Private/ViewModels/HierarchyEditor/NiagaraHierarchyViewModelBase.cpp
+++ b/Engine/Plugins/FX/Niagara/Source/NiagaraEditor/Private/ViewModels/HierarchyEditor/NiagaraHierarchyViewModelBase.cpp
@@ -451,7 +451,7 @@ void UNiagaraHierarchyViewModelBase::ForceFullRefreshOnTimer()
 
 void UNiagaraHierarchyViewModelBase::RequestFullRefreshNextFrame()
 {
-	if(!FullRefreshNextFrameHandle.IsValid())
+	if(!FullRefreshNextFrameHandle.IsValid() && GEditor != nullptr)
 	{
 		FullRefreshNextFrameHandle = GEditor->GetTimerManager()->SetTimerForNextTick(FTimerDelegate::CreateUObject(this, &UNiagaraHierarchyViewModelBase::ForceFullRefreshOnTimer));
 	}
