Niagara - Fix a bug in the merge manager where input diffing was using the name and type of the input, but applying the diff was only using the name, so the diff could be applied incorrectly resulting in bad input values.

diff --git a/Engine/Plugins/FX/Niagara/Source/NiagaraEditor/Private/NiagaraScriptMergeManager.cpp b/Engine/Plugins/FX/Niagara/Source/NiagaraEditor/Private/NiagaraScriptMergeManager.cpp
--- a/Engine/Plugins/FX/Niagara/Source/NiagaraEditor/Private/NiagaraScriptMergeManager.cpp
+++ b/Engine/Plugins/FX/Niagara/Source/NiagaraEditor/Private/NiagaraScriptMergeManager.cpp
@@ -390,11 +390,11 @@ const TArray<TSharedRef<FNiagaraStackFunctionInputOverrideMergeAdapter>>& FNiaga
 	return InputOverrides;
 }
 
-TSharedPtr<FNiagaraStackFunctionInputOverrideMergeAdapter> FNiagaraStackFunctionMergeAdapter::GetInputOverrideByInputName(FString InputName) const
+TSharedPtr<FNiagaraStackFunctionInputOverrideMergeAdapter> FNiagaraStackFunctionMergeAdapter::GetInputOverrideByInputNameAndType(const FString& InputName, const FNiagaraTypeDefinition& InputType) const
 {
 	for (TSharedPtr<FNiagaraStackFunctionInputOverrideMergeAdapter> InputOverride : InputOverrides)
 	{
-		if (InputOverride->GetInputName() == InputName)
+		if (InputOverride->GetInputName() == InputName && InputOverride->GetType() == InputType)
 		{
 			return InputOverride;
 		}
@@ -3242,7 +3242,8 @@ FNiagaraScriptMergeManager::FApplyDiffResults FNiagaraScriptMergeManager::ApplyS
 		TSharedPtr<FNiagaraStackFunctionMergeAdapter> MatchingModuleAdapter = BaseScriptStackAdapter->GetModuleFunctionById(RemovedInputOverrideAdapter->GetOwningFunctionCall()->NodeGuid);
 		if (MatchingModuleAdapter.IsValid())
 		{
-			TSharedPtr<FNiagaraStackFunctionInputOverrideMergeAdapter> MatchingInputOverrideAdapter = MatchingModuleAdapter->GetInputOverrideByInputName(RemovedInputOverrideAdapter->GetInputName());
+			TSharedPtr<FNiagaraStackFunctionInputOverrideMergeAdapter> MatchingInputOverrideAdapter =
+				MatchingModuleAdapter->GetInputOverrideByInputNameAndType(RemovedInputOverrideAdapter->GetInputName(), RemovedInputOverrideAdapter->GetType());
 			if (MatchingInputOverrideAdapter.IsValid())
 			{
 				RemoveInputOverrides.Add(MatchingInputOverrideAdapter.ToSharedRef());
@@ -3255,7 +3256,8 @@ FNiagaraScriptMergeManager::FApplyDiffResults FNiagaraScriptMergeManager::ApplyS
 		TSharedPtr<FNiagaraStackFunctionMergeAdapter> MatchingModuleAdapter = BaseScriptStackAdapter->GetModuleFunctionById(AddedInputOverrideAdapter->GetOwningFunctionCall()->NodeGuid);
 		if (MatchingModuleAdapter.IsValid())
 		{
-			TSharedPtr<FNiagaraStackFunctionInputOverrideMergeAdapter> MatchingInputOverrideAdapter = MatchingModuleAdapter->GetInputOverrideByInputName(AddedInputOverrideAdapter->GetInputName());
+			TSharedPtr<FNiagaraStackFunctionInputOverrideMergeAdapter> MatchingInputOverrideAdapter = 
+				MatchingModuleAdapter->GetInputOverrideByInputNameAndType(AddedInputOverrideAdapter->GetInputName(), AddedInputOverrideAdapter->GetType());
 			if (MatchingInputOverrideAdapter.IsValid())
 			{
 				RemoveInputOverrides.AddUnique(MatchingInputOverrideAdapter.ToSharedRef());
@@ -3273,7 +3275,8 @@ FNiagaraScriptMergeManager::FApplyDiffResults FNiagaraScriptMergeManager::ApplyS
 		TSharedPtr<FNiagaraStackFunctionMergeAdapter> MatchingModuleAdapter = BaseScriptStackAdapter->GetModuleFunctionById(ModifiedInputOverrideAdapter->GetOwningFunctionCall()->NodeGuid);
 		if (MatchingModuleAdapter.IsValid())
 		{
-			TSharedPtr<FNiagaraStackFunctionInputOverrideMergeAdapter> MatchingInputOverrideAdapter = MatchingModuleAdapter->GetInputOverrideByInputName(ModifiedInputOverrideAdapter->GetInputName());
+			TSharedPtr<FNiagaraStackFunctionInputOverrideMergeAdapter> MatchingInputOverrideAdapter = 
+				MatchingModuleAdapter->GetInputOverrideByInputNameAndType(ModifiedInputOverrideAdapter->GetInputName(), ModifiedInputOverrideAdapter->GetType());
 			if (MatchingInputOverrideAdapter.IsValid())
 			{
 				RemoveInputOverrides.AddUnique(MatchingInputOverrideAdapter.ToSharedRef());
diff --git a/Engine/Plugins/FX/Niagara/Source/NiagaraEditor/Private/NiagaraScriptMergeManager.h b/Engine/Plugins/FX/Niagara/Source/NiagaraEditor/Private/NiagaraScriptMergeManager.h
--- a/Engine/Plugins/FX/Niagara/Source/NiagaraEditor/Private/NiagaraScriptMergeManager.h
+++ b/Engine/Plugins/FX/Niagara/Source/NiagaraEditor/Private/NiagaraScriptMergeManager.h
@@ -109,7 +109,7 @@ public:
 
 	const TArray<TSharedRef<FNiagaraStackFunctionInputOverrideMergeAdapter>>& GetInputOverrides() const;
 
-	TSharedPtr<FNiagaraStackFunctionInputOverrideMergeAdapter> GetInputOverrideByInputName(FString InputName) const;
+	TSharedPtr<FNiagaraStackFunctionInputOverrideMergeAdapter> GetInputOverrideByInputNameAndType(const FString& InputName, const FNiagaraTypeDefinition& InputType) const;
 
 	void GatherFunctionCallNodes(TArray<UNiagaraNodeFunctionCall*>& OutFunctionCallNodes) const;
 
