Expose data channel and handler classes to allow inheriting in other modules.
New accessor functions that take data channel objects directly and make existing ones taking their asset class just wrappers.

diff --git a/Engine/Plugins/FX/Niagara/Source/Niagara/Private/NiagaraDataChannel.cpp b/Engine/Plugins/FX/Niagara/Source/Niagara/Private/NiagaraDataChannel.cpp
--- a/Engine/Plugins/FX/Niagara/Source/Niagara/Private/NiagaraDataChannel.cpp
+++ b/Engine/Plugins/FX/Niagara/Source/Niagara/Private/NiagaraDataChannel.cpp
@@ -939,6 +939,21 @@ UNiagaraDataChannelLibrary::UNiagaraDataChannelLibrary(const FObjectInitializer&
 }
 
 UNiagaraDataChannelHandler* UNiagaraDataChannelLibrary::GetNiagaraDataChannel(const UObject* WorldContextObject, UNiagaraDataChannelAsset* Channel)
+{
+	return GetNiagaraDataChannel(WorldContextObject, Channel->Get());
+}
+
+UNiagaraDataChannelWriter* UNiagaraDataChannelLibrary::WriteToNiagaraDataChannel(const UObject* WorldContextObject, UNiagaraDataChannelAsset* Channel, FNiagaraDataChannelSearchParameters SearchParams, int32 Count, bool bVisibleToGame, bool bVisibleToCPU, bool bVisibleToGPU)
+{
+	return WriteToNiagaraDataChannel(WorldContextObject, Channel->Get(), SearchParams, Count, bVisibleToGame, bVisibleToCPU, bVisibleToGPU);
+}
+
+UNiagaraDataChannelReader* UNiagaraDataChannelLibrary::ReadFromNiagaraDataChannel(const UObject* WorldContextObject, UNiagaraDataChannelAsset* Channel, FNiagaraDataChannelSearchParameters SearchParams, bool bReadPreviousFrame)
+{
+	return ReadFromNiagaraDataChannel(WorldContextObject, Channel->Get(), SearchParams, bReadPreviousFrame);
+}
+
+UNiagaraDataChannelHandler* UNiagaraDataChannelLibrary::GetNiagaraDataChannel(const UObject* WorldContextObject, UNiagaraDataChannel* Channel)
 {
 	UWorld* World = (WorldContextObject != nullptr) ? GEngine->GetWorldFromContextObject(WorldContextObject, EGetWorldErrorMode::LogAndReturnNull) : nullptr;
 	if (World)
@@ -951,7 +966,7 @@ UNiagaraDataChannelHandler* UNiagaraDataChannelLibrary::GetNiagaraDataChannel(co
 	return nullptr;
 }
 
-UNiagaraDataChannelWriter* UNiagaraDataChannelLibrary::WriteToNiagaraDataChannel(const UObject* WorldContextObject, UNiagaraDataChannelAsset* Channel, FNiagaraDataChannelSearchParameters SearchParams, int32 Count, bool bVisibleToGame, bool bVisibleToCPU, bool bVisibleToGPU)
+UNiagaraDataChannelWriter* UNiagaraDataChannelLibrary::WriteToNiagaraDataChannel(const UObject* WorldContextObject, UNiagaraDataChannel* Channel, FNiagaraDataChannelSearchParameters SearchParams, int32 Count, bool bVisibleToGame, bool bVisibleToCPU, bool bVisibleToGPU)
 {
 	check(IsInGameThread());
 	UWorld* World = (WorldContextObject != nullptr) ? GEngine->GetWorldFromContextObject(WorldContextObject, EGetWorldErrorMode::LogAndReturnNull) : nullptr;
@@ -974,7 +989,7 @@ UNiagaraDataChannelWriter* UNiagaraDataChannelLibrary::WriteToNiagaraDataChannel
 	return nullptr;
 }
 
-UNiagaraDataChannelReader* UNiagaraDataChannelLibrary::ReadFromNiagaraDataChannel(const UObject* WorldContextObject, UNiagaraDataChannelAsset* Channel, FNiagaraDataChannelSearchParameters SearchParams, bool bReadPreviousFrame)
+UNiagaraDataChannelReader* UNiagaraDataChannelLibrary::ReadFromNiagaraDataChannel(const UObject* WorldContextObject, UNiagaraDataChannel* Channel, FNiagaraDataChannelSearchParameters SearchParams, bool bReadPreviousFrame)
 {
 	UWorld* World = (WorldContextObject != nullptr) ? GEngine->GetWorldFromContextObject(WorldContextObject, EGetWorldErrorMode::LogAndReturnNull) : nullptr;
 	if (World)
diff --git a/Engine/Plugins/FX/Niagara/Source/Niagara/Public/NiagaraDataChannel.h b/Engine/Plugins/FX/Niagara/Source/Niagara/Public/NiagaraDataChannel.h
--- a/Engine/Plugins/FX/Niagara/Source/Niagara/Public/NiagaraDataChannel.h
+++ b/Engine/Plugins/FX/Niagara/Source/Niagara/Public/NiagaraDataChannel.h
@@ -180,4 +180,8 @@ class NIAGARA_API UNiagaraDataChannelLibrary : public UBlueprintFunctionLibrary
 	/** Initializes and returns the Niagara Data Channel reader for the given data channel. */
 	UFUNCTION(BlueprintCallable, Category = NiagaraDataChannel, meta = (Keywords = "niagara DataChannel", WorldContext = "WorldContextObject", UnsafeDuringActorConstruction = "true"))
 	static UNiagaraDataChannelReader* ReadFromNiagaraDataChannel(const UObject* WorldContextObject, UNiagaraDataChannelAsset* Channel, FNiagaraDataChannelSearchParameters SearchParams, bool bReadPreviousFrame);
+
+	static UNiagaraDataChannelHandler* GetNiagaraDataChannel(const UObject* WorldContextObject, UNiagaraDataChannel* Channel);
+	static UNiagaraDataChannelWriter* WriteToNiagaraDataChannel(const UObject* WorldContextObject, UNiagaraDataChannel* Channel, FNiagaraDataChannelSearchParameters SearchParams, int32 Count, bool bVisibleToGame, bool bVisibleToCPU, bool bVisibleToGPU);
+	static UNiagaraDataChannelReader* ReadFromNiagaraDataChannel(const UObject* WorldContextObject, UNiagaraDataChannel* Channel, FNiagaraDataChannelSearchParameters SearchParams, bool bReadPreviousFrame);
 };
diff --git a/Engine/Plugins/FX/Niagara/Source/Niagara/Public/NiagaraDataChannel_Global.h b/Engine/Plugins/FX/Niagara/Source/Niagara/Public/NiagaraDataChannel_Global.h
--- a/Engine/Plugins/FX/Niagara/Source/Niagara/Public/NiagaraDataChannel_Global.h
+++ b/Engine/Plugins/FX/Niagara/Source/Niagara/Public/NiagaraDataChannel_Global.h
@@ -8,7 +8,7 @@
 
 
 UCLASS(Experimental)
-class UNiagaraDataChannel_Global : public UNiagaraDataChannel
+class NIAGARA_API UNiagaraDataChannel_Global : public UNiagaraDataChannel
 {
 	GENERATED_BODY()
 
@@ -19,7 +19,7 @@ class UNiagaraDataChannel_Global : public UNiagaraDataChannel
 Basic DataChannel handler that makes all data visible globally.
 */
 UCLASS(Experimental, BlueprintType)
-class UNiagaraDataChannelHandler_Global : public UNiagaraDataChannelHandler
+class NIAGARA_API UNiagaraDataChannelHandler_Global : public UNiagaraDataChannelHandler
 {
 	GENERATED_UCLASS_BODY()
 
diff --git a/Engine/Plugins/FX/Niagara/Source/Niagara/Public/NiagaraDataChannel_Islands.h b/Engine/Plugins/FX/Niagara/Source/Niagara/Public/NiagaraDataChannel_Islands.h
--- a/Engine/Plugins/FX/Niagara/Source/Niagara/Public/NiagaraDataChannel_Islands.h
+++ b/Engine/Plugins/FX/Niagara/Source/Niagara/Public/NiagaraDataChannel_Islands.h
@@ -27,6 +27,8 @@ USTRUCT()
 struct FNDCIsland
 {
 	GENERATED_BODY()
+
+	UE_NONCOPYABLE(FNDCIsland)
 public:
 
 	FNDCIsland();
@@ -74,11 +76,6 @@ private:
 	
 	/** Publish requests from game code/BP or other Niagara Systems wishing to write data into the data channel for this island. */
 	TArray<FNiagaraDataChannelPublishRequest> PublishRequests;
- 
- 	FNDCIsland(const FNDCIsland&) = delete;
- 	FNDCIsland& operator=(const FNDCIsland&) = delete;
- 	FNDCIsland(const FNDCIsland&&) = delete;
- 	FNDCIsland& operator=(const FNDCIsland&&) = delete;
 };
 
 template<>
@@ -87,8 +84,8 @@ struct TStructOpsTypeTraits<FNDCIsland> : public TStructOpsTypeTraitsBase2<FNDCI
 	enum{ WithCopy = false };
 };
 
-UCLASS(Experimental, MinimalAPI)
-class UNiagaraDataChannel_Islands : public UNiagaraDataChannel
+UCLASS(Experimental)
+class NIAGARA_API UNiagaraDataChannel_Islands : public UNiagaraDataChannel
 {
 	GENERATED_BODY()
 
@@ -139,7 +136,7 @@ public:
 };
 
 UCLASS(Experimental, BlueprintType)
-class UNiagaraDataChannelHandler_Islands : public UNiagaraDataChannelHandler
+class NIAGARA_API UNiagaraDataChannelHandler_Islands : public UNiagaraDataChannelHandler
 {
 	GENERATED_UCLASS_BODY()
 
