Removed outdated summary filter parameter

diff --git a/Engine/Plugins/FX/Niagara/Source/NiagaraEditor/Private/ViewModels/Stack/NiagaraStackFunctionInputCollection.cpp b/Engine/Plugins/FX/Niagara/Source/NiagaraEditor/Private/ViewModels/Stack/NiagaraStackFunctionInputCollection.cpp
--- a/Engine/Plugins/FX/Niagara/Source/NiagaraEditor/Private/ViewModels/Stack/NiagaraStackFunctionInputCollection.cpp
+++ b/Engine/Plugins/FX/Niagara/Source/NiagaraEditor/Private/ViewModels/Stack/NiagaraStackFunctionInputCollection.cpp
@@ -115,14 +115,14 @@ int32 UNiagaraStackValueCollection::GetChildIndentLevel() const
 	return GetIndentLevel();
 }
 
-void UNiagaraStackFunctionInputCollection::RefreshChildrenForFunctionCall(const TArray<UNiagaraStackEntry*>& CurrentChildren, TArray<UNiagaraStackEntry*>& NewChildren, TArray<FStackIssue>& NewIssues, bool bShouldApplySummaryFilter)
+void UNiagaraStackFunctionInputCollection::RefreshChildrenForFunctionCall(const TArray<UNiagaraStackEntry*>& CurrentChildren, TArray<UNiagaraStackEntry*>& NewChildren, TArray<FStackIssue>& NewIssues)
 {	
 	FFunctionCallNodesState State;
-	AppendInputsForFunctionCall(State, NewIssues, bShouldApplySummaryFilter);
-	ApplyAllFunctionInputsToChildren(State, CurrentChildren, NewChildren, NewIssues, bShouldApplySummaryFilter);
+	AppendInputsForFunctionCall(State, NewIssues);
+	ApplyAllFunctionInputsToChildren(State, CurrentChildren, NewChildren, NewIssues);
 }
 
-void UNiagaraStackFunctionInputCollection::AppendInputsForFunctionCall(FFunctionCallNodesState& State, TArray<FStackIssue>& NewIssues, bool bShouldApplySummaryFilter)
+void UNiagaraStackFunctionInputCollection::AppendInputsForFunctionCall(FFunctionCallNodesState& State, TArray<FStackIssue>& NewIssues)
 {
 	FVersionedNiagaraEmitter Emitter = GetEmitterViewModel().IsValid() ? GetEmitterViewModel()->GetEmitter() : FVersionedNiagaraEmitter();
 
@@ -275,7 +275,7 @@ void UNiagaraStackFunctionInputCollection::AppendInputsForFunctionCall(FFunction
 	RefreshIssues(DuplicateInputNames, ValidAliasedInputNames, InputsWithInvalidTypes, StaticSwitchInputs, NewIssues);
 }
 
-void UNiagaraStackFunctionInputCollection::ApplyAllFunctionInputsToChildren(FFunctionCallNodesState& State, const TArray<UNiagaraStackEntry*>& CurrentChildren, TArray<UNiagaraStackEntry*>& NewChildren, TArray<FStackIssue>& NewIssues, bool bShouldApplySummaryFilter)
+void UNiagaraStackFunctionInputCollection::ApplyAllFunctionInputsToChildren(FFunctionCallNodesState& State, const TArray<UNiagaraStackEntry*>& CurrentChildren, TArray<UNiagaraStackEntry*>& NewChildren, TArray<FStackIssue>& NewIssues)
 {
 	// resolve the parent/child relationships
 	for (auto& Entry : State.ParentMapping)
@@ -806,7 +806,7 @@ void UNiagaraStackFunctionInputCollection::OnScriptApplied(UNiagaraScript* Niaga
 
 void UNiagaraStackFunctionInputCollection::RefreshChildrenInternal(const TArray<UNiagaraStackEntry*>& CurrentChildren, TArray<UNiagaraStackEntry*>& NewChildren, TArray<FStackIssue>& NewIssues)
 {
-	RefreshChildrenForFunctionCall(CurrentChildren, NewChildren, NewIssues, false);
+	RefreshChildrenForFunctionCall(CurrentChildren, NewChildren, NewIssues);
 	Super::RefreshChildrenInternal(CurrentChildren, NewChildren, NewIssues);
 }
 
diff --git a/Engine/Plugins/FX/Niagara/Source/NiagaraEditor/Public/ViewModels/Stack/NiagaraStackFunctionInputCollection.h b/Engine/Plugins/FX/Niagara/Source/NiagaraEditor/Public/ViewModels/Stack/NiagaraStackFunctionInputCollection.h
--- a/Engine/Plugins/FX/Niagara/Source/NiagaraEditor/Public/ViewModels/Stack/NiagaraStackFunctionInputCollection.h
+++ b/Engine/Plugins/FX/Niagara/Source/NiagaraEditor/Public/ViewModels/Stack/NiagaraStackFunctionInputCollection.h
@@ -125,11 +125,11 @@ private:
 
 	virtual void GetSectionsInternal(TArray<FNiagaraStackSection>& OutStackSections) const override;
 	
-	void RefreshChildrenForFunctionCall(const TArray<UNiagaraStackEntry*>& CurrentChildren, TArray<UNiagaraStackEntry*>& NewChildren, TArray<FStackIssue>& NewIssues, bool bShouldApplySummaryFilter);
+	void RefreshChildrenForFunctionCall(const TArray<UNiagaraStackEntry*>& CurrentChildren, TArray<UNiagaraStackEntry*>& NewChildren, TArray<FStackIssue>& NewIssues);
 	
-	void AppendInputsForFunctionCall(FFunctionCallNodesState& State, TArray<FStackIssue>& NewIssues, bool bShouldApplySummaryFilter);
+	void AppendInputsForFunctionCall(FFunctionCallNodesState& State, TArray<FStackIssue>& NewIssues);
 	
-	void ApplyAllFunctionInputsToChildren(FFunctionCallNodesState& State, const TArray<UNiagaraStackEntry*>& CurrentChildren, TArray<UNiagaraStackEntry*>& NewChildren, TArray<FStackIssue>& NewIssues, bool bShouldApplySummaryFilter);
+	void ApplyAllFunctionInputsToChildren(FFunctionCallNodesState& State, const TArray<UNiagaraStackEntry*>& CurrentChildren, TArray<UNiagaraStackEntry*>& NewChildren, TArray<FStackIssue>& NewIssues);
 
 	void RefreshIssues(const TArray<FName>& DuplicateInputNames, const TArray<FName>& ValidAliasedInputNames, const TArray<FNiagaraVariable>& InputsWithInvalidTypes, const TMap<FName, UEdGraphPin*>& StaticSwitchInputs, TArray<FStackIssue>& NewIssues);
 
