Update logic around changing a system into CompileForEdit mode via the SystemViewModel to take into account the fact that the ViewModel might be just around for emitter merging

diff --git a/Engine/Plugins/FX/Niagara/Source/NiagaraEditor/Private/ViewModels/NiagaraSystemViewModel.cpp b/Engine/Plugins/FX/Niagara/Source/NiagaraEditor/Private/ViewModels/NiagaraSystemViewModel.cpp
--- a/Engine/Plugins/FX/Niagara/Source/NiagaraEditor/Private/ViewModels/NiagaraSystemViewModel.cpp
+++ b/Engine/Plugins/FX/Niagara/Source/NiagaraEditor/Private/ViewModels/NiagaraSystemViewModel.cpp
@@ -108,7 +108,6 @@ FNiagaraSystemViewModel::FNiagaraSystemViewModel()
 void FNiagaraSystemViewModel::Initialize(UNiagaraSystem& InSystem, FNiagaraSystemViewModelOptions InOptions)
 {
 	System = &InSystem;
-	System->bCompileForEdit = bSupportCompileForEdit;
 
 	RegisteredHandle = RegisterViewModelWithMap(System, this);
 
@@ -163,6 +162,11 @@ void FNiagaraSystemViewModel::Initialize(UNiagaraSystem& InSystem, FNiagaraSyste
 
 	SystemGraphSelectionViewModel = MakeShared<FNiagaraSystemGraphSelectionViewModel>();
 	SystemGraphSelectionViewModel->Initialize(this->AsShared());
+
+	if (CanImpactCompileForEdit())
+	{
+		System->bCompileForEdit = true;
+	}
 	
 	if (bIsForDataProcessingOnly == false)
 	{
@@ -255,7 +259,7 @@ void FNiagaraSystemViewModel::Cleanup()
 
 	if (System)
 	{
-		if ( System->bCompileForEdit )
+		if (CanImpactCompileForEdit())
 		{
 			System->bCompileForEdit = false;
 			System->InvalidateActiveCompiles();
@@ -1042,7 +1046,7 @@ void FNiagaraSystemViewModel::NotifyPreSave()
 	}
 
 	// we want to save without compile for edit turned on for best results, but make sure to turn in back on again after save is done
-	if (bSupportCompileForEdit)
+	if (CanImpactCompileForEdit())
 	{
 		check(System->bCompileForEdit);
 		System->bCompileForEdit = false;
@@ -1054,7 +1058,7 @@ void FNiagaraSystemViewModel::NotifyPreSave()
 void FNiagaraSystemViewModel::NotifyPostSave()
 {
 	// we compile for edit again after having it turned off for PreSave. This should fetch DDC data and must not update the ChangeID
-	if (bSupportCompileForEdit)
+	if (CanImpactCompileForEdit())
 	{
 		check(!System->bCompileForEdit);
 		System->bCompileForEdit = true;
@@ -3128,4 +3132,9 @@ TSharedPtr<FNiagaraOverviewGraphViewModel> FNiagaraSystemViewModel::GetOverviewG
 	return OverviewGraphViewModel;
 }
 
+bool FNiagaraSystemViewModel::CanImpactCompileForEdit() const
+{
+	return bSupportCompileForEdit && !bIsForDataProcessingOnly && EditMode == ENiagaraSystemViewModelEditMode::SystemAsset;
+}
+
 #undef LOCTEXT_NAMESPACE // NiagaraSystemViewModel
diff --git a/Engine/Plugins/FX/Niagara/Source/NiagaraEditor/Public/ViewModels/NiagaraSystemViewModel.h b/Engine/Plugins/FX/Niagara/Source/NiagaraEditor/Public/ViewModels/NiagaraSystemViewModel.h
--- a/Engine/Plugins/FX/Niagara/Source/NiagaraEditor/Public/ViewModels/NiagaraSystemViewModel.h
+++ b/Engine/Plugins/FX/Niagara/Source/NiagaraEditor/Public/ViewModels/NiagaraSystemViewModel.h
@@ -591,6 +591,9 @@ private:
 	/** Utility method to gather all emitter handle viewmodels marked as selected via the system selection viewmodel. */
 	TArray<TSharedRef<FNiagaraEmitterHandleViewModel>> GetSelectedEmitterHandleViewModels();
 
+	/** Indicates whether this instance of the view model should impact whether the System's CompileForEdit flag should be updated by the viewmodel. */
+	bool CanImpactCompileForEdit() const;
+
 private:
 	/** The System being viewed and edited by this view model. */
 	UNiagaraSystem* System;
