Fixed static switches not being correctly hidden in summary view

diff --git a/Engine/Plugins/FX/Niagara/Source/NiagaraEditor/Private/ViewModels/Stack/NiagaraStackInputCategory.cpp b/Engine/Plugins/FX/Niagara/Source/NiagaraEditor/Private/ViewModels/Stack/NiagaraStackInputCategory.cpp
--- a/Engine/Plugins/FX/Niagara/Source/NiagaraEditor/Private/ViewModels/Stack/NiagaraStackInputCategory.cpp
+++ b/Engine/Plugins/FX/Niagara/Source/NiagaraEditor/Private/ViewModels/Stack/NiagaraStackInputCategory.cpp
@@ -260,6 +260,14 @@ void UNiagaraStackSummaryCategory::RefreshChildrenInternal(const TArray<UNiagara
 		FCompileConstantResolver Resolver(GetEmitterViewModel()->GetEmitter(), FNiagaraStackGraphUtilities::GetOutputNodeUsage(*FunctionCall));
 		FNiagaraStackGraphUtilities::GetStackFunctionInputs(*FunctionCall, InputVariables, HiddenVariables, Resolver, ENiagaraGetStackFunctionInputPinsOptions::ModuleInputsOnly, true);
 		FunctionCallToHiddenVariablesMap.Add(FunctionCall->NodeGuid, HiddenVariables);
+		
+		TArray<UEdGraphPin*> OutInputPins;
+		TSet<UEdGraphPin*> OutHiddenPins;
+		FNiagaraStackGraphUtilities::GetStackFunctionStaticSwitchPins(*FunctionCall, OutInputPins, OutHiddenPins, Resolver);
+		for(UEdGraphPin* HiddenStaticSwitchPin : OutHiddenPins)
+		{
+			FunctionCallToHiddenVariablesMap[FunctionCall->NodeGuid].Add(UEdGraphSchema_Niagara::PinToNiagaraVariable(HiddenStaticSwitchPin));
+		}
 	}
 	
 	for(TSharedPtr<FNiagaraHierarchyItemViewModelBase> HierarchyViewModel : ChildrenViewModels)
diff --git a/Engine/Plugins/FX/Niagara/Source/NiagaraEditor/Private/ViewModels/Stack/NiagaraStackSummaryViewInputCollection.cpp b/Engine/Plugins/FX/Niagara/Source/NiagaraEditor/Private/ViewModels/Stack/NiagaraStackSummaryViewInputCollection.cpp
--- a/Engine/Plugins/FX/Niagara/Source/NiagaraEditor/Private/ViewModels/Stack/NiagaraStackSummaryViewInputCollection.cpp
+++ b/Engine/Plugins/FX/Niagara/Source/NiagaraEditor/Private/ViewModels/Stack/NiagaraStackSummaryViewInputCollection.cpp
@@ -114,6 +114,14 @@ void UNiagaraStackSummaryViewCollection::RefreshChildrenInternal(const TArray<UN
 		FCompileConstantResolver Resolver(GetEmitterViewModel()->GetEmitter(), FNiagaraStackGraphUtilities::GetOutputNodeUsage(*FunctionCall));
 		FNiagaraStackGraphUtilities::GetStackFunctionInputs(*FunctionCall, InputVariables, HiddenVariables, Resolver, ENiagaraGetStackFunctionInputPinsOptions::ModuleInputsOnly, true);
 		FunctionCallToHiddenVariablesMap.Add(FunctionCall->NodeGuid, HiddenVariables);
+
+		TArray<UEdGraphPin*> OutInputPins;
+		TSet<UEdGraphPin*> OutHiddenPins;
+		FNiagaraStackGraphUtilities::GetStackFunctionStaticSwitchPins(*FunctionCall, OutInputPins, OutHiddenPins, Resolver);
+		for(UEdGraphPin* HiddenStaticSwitchPin : OutHiddenPins)
+		{
+			FunctionCallToHiddenVariablesMap[FunctionCall->NodeGuid].Add(UEdGraphSchema_Niagara::PinToNiagaraVariable(HiddenStaticSwitchPin));
+		}
 	}
 	
 	for(TSharedPtr<FNiagaraHierarchyItemViewModelBase> HierarchyViewModel : ChildrenViewModels)
