- Reduce sizeof FNiagaraVariableAttributeBinding from 88 to 56 bytes

diff --git a/Engine/Plugins/FX/Niagara/Source/Niagara/Private/NiagaraCommon.cpp b/Engine/Plugins/FX/Niagara/Source/Niagara/Private/NiagaraCommon.cpp
--- a/Engine/Plugins/FX/Niagara/Source/Niagara/Private/NiagaraCommon.cpp
+++ b/Engine/Plugins/FX/Niagara/Source/Niagara/Private/NiagaraCommon.cpp
@@ -512,7 +512,8 @@ void FNiagaraVariableAttributeBinding::SetAsPreviousValue(const FNiagaraVariable
 	static const FString PreviousNamespace = FNiagaraConstants::PreviousNamespace.ToString();
 
 	ParamMapVariable = Src;
-	RootVariable = DataSetVariable = Src;
+	RootVariable = Src;
+	DataSetName = Src.GetName();
 
 	// Split out the name and it's namespace
 	TArray<FString> SplitName;
@@ -549,7 +550,8 @@ void FNiagaraVariableAttributeBinding::SetAsPreviousValue(const FNiagaraVariable
 	static const FString PreviousNamespace = FNiagaraConstants::PreviousNamespace.ToString();
 
 	ParamMapVariable = Src.RootVariable;
-	RootVariable = DataSetVariable = Src.RootVariable;
+	RootVariable = Src.RootVariable;
+	DataSetName = Src.RootVariable.GetName();
 
 	// Split out the name and it's namespace
 	TStringBuilder<128> VarName;
@@ -626,10 +628,14 @@ void FNiagaraVariableAttributeBinding::PostLoad(ENiagaraRendererSourceDataMode I
 #if WITH_EDITORONLY_DATA
 	if (BoundVariable.IsValid())
 	{
-		RootVariable.SetType(DataSetVariable.GetType()); //Sometimes the BoundVariable was bogus in the past. THe DataSet shouldn't be though.
+		RootVariable.SetType(ParamMapVariable.GetType()); //Sometimes the BoundVariable was bogus in the past. THe DataSet shouldn't be though.
 		SetValue(BoundVariable.GetName(), FVersionedNiagaraEmitter(), InSourceMode);
 		BoundVariable = FNiagaraVariable();
 	}
+	if (!DataSetVariable_DEPRECATED.GetName().IsNone())
+	{
+		DataSetName = DataSetVariable_DEPRECATED.GetName();
+	}
 #endif
 
 }
@@ -638,7 +644,7 @@ void FNiagaraVariableAttributeBinding::Dump() const
 {
 	UE_LOG(LogNiagara, Log, TEXT("PostLoad for FNiagaraVariableAttributeBinding...."));
 	UE_LOG(LogNiagara, Log, TEXT("ParamMapVariable: %s %s"), *ParamMapVariable.GetName().ToString(), *ParamMapVariable.GetType().GetName());
-	UE_LOG(LogNiagara, Log, TEXT("DataSetVariable: %s %s"), *DataSetVariable.GetName().ToString(), *DataSetVariable.GetType().GetName());
+	UE_LOG(LogNiagara, Log, TEXT("DataSetName: %s"), *DataSetName.ToString());
 	UE_LOG(LogNiagara, Log, TEXT("RootVariable: %s %s"), *RootVariable.GetName().ToString(), *RootVariable.GetType().GetName());
 #if WITH_EDITORONLY_DATA
 	UE_LOG(LogNiagara, Log, TEXT("BoundVariable: %s %s"), *BoundVariable.GetName().ToString(), *BoundVariable.GetType().GetName());
@@ -657,14 +663,14 @@ void FNiagaraVariableAttributeBinding::ResetToDefault(const FNiagaraVariableAttr
 		if ((InSourceMode == ENiagaraRendererSourceDataMode::Emitter && InOther.BindingSourceMode == ENiagaraBindingSource::ImplicitFromSource) ||
 			InOther.BindingSourceMode == ENiagaraBindingSource::ExplicitEmitter)
 		{
-			ensure(!InOther.DataSetVariable.IsInNameSpace(FNiagaraConstants::EmitterNamespaceString));
-			TempVar.SetNamespacedName(FNiagaraConstants::EmitterNamespaceString, InOther.DataSetVariable.GetName());
+			ensure(!InOther.GetDataSetBindableVariable().IsInNameSpace(FNiagaraConstants::EmitterNamespaceString));
+			TempVar.SetNamespacedName(FNiagaraConstants::EmitterNamespaceString, InOther.DataSetName);
 		}
 		else if ((InSourceMode == ENiagaraRendererSourceDataMode::Particles && InOther.BindingSourceMode == ENiagaraBindingSource::ImplicitFromSource) ||
 			InOther.BindingSourceMode == ENiagaraBindingSource::ExplicitParticles)
 		{
-			ensure(!InOther.DataSetVariable.IsInNameSpace(FNiagaraConstants::ParticleAttributeNamespaceString));
-			TempVar.SetNamespacedName(FNiagaraConstants::ParticleAttributeNamespaceString, InOther.DataSetVariable.GetName());
+			ensure(!InOther.GetDataSetBindableVariable().IsInNameSpace(FNiagaraConstants::ParticleAttributeNamespaceString));
+			TempVar.SetNamespacedName(FNiagaraConstants::ParticleAttributeNamespaceString, InOther.DataSetName);
 		}
 
 		SetValue(TempVar.GetName(), FVersionedNiagaraEmitter(), InSourceMode);
@@ -677,7 +683,7 @@ void FNiagaraVariableAttributeBinding::ResetToDefault(const FNiagaraVariableAttr
 
 bool FNiagaraVariableAttributeBinding::MatchesDefault(const FNiagaraVariableAttributeBinding& InOther, ENiagaraRendererSourceDataMode InSourceMode) const
 {
-	if (DataSetVariable.GetName() != InOther.DataSetVariable.GetName())
+	if (DataSetName != InOther.DataSetName)
 		return false;
 	if (RootVariable.GetName() != InOther.RootVariable.GetName())
 		return false;
@@ -748,7 +754,8 @@ void FNiagaraVariableAttributeBinding::CacheValues(const FVersionedNiagaraEmitte
 		RootVariable.SetName(FNiagaraConstants::GetAttributeAsEmitterDataSetKey(RootVariable).GetName());
 	}
 
-	DataSetVariable = ParamMapVariable = (const FNiagaraVariableBase&)RootVariable;
+	ParamMapVariable = (const FNiagaraVariableBase&)RootVariable;
+	DataSetName = ParamMapVariable.GetName();
 	bBindingExistsOnSource = false;
 
 	// Decide if this is going to be bound to a particle attribute (needed for use by the renderers, for instance)
@@ -762,21 +769,21 @@ void FNiagaraVariableAttributeBinding::CacheValues(const FVersionedNiagaraEmitte
 	}
 
 	// If this is one of the possible namespaces that is implicitly defined, go ahead and expand the full namespace. RootVariable should be non-namespaced at this point.
-	if (DataSetVariable.GetName().IsNone())
+	if (DataSetName.IsNone())
 	{
 		ParamMapVariable.SetName(NAME_None);
 	}
 	else if ((InSourceMode == ENiagaraRendererSourceDataMode::Emitter && BindingSourceMode == ENiagaraBindingSource::ImplicitFromSource) ||
 		BindingSourceMode == ENiagaraBindingSource::ExplicitEmitter)
 	{
-		ensure(!DataSetVariable.IsInNameSpace(FNiagaraConstants::EmitterNamespaceString));
-		ParamMapVariable.SetNamespacedName(FNiagaraConstants::EmitterNamespaceString, DataSetVariable.GetName());
+		ensure(!GetDataSetBindableVariable().IsInNameSpace(FNiagaraConstants::EmitterNamespaceString));
+		ParamMapVariable.SetNamespacedName(FNiagaraConstants::EmitterNamespaceString, DataSetName);
 	}
 	else if ((InSourceMode == ENiagaraRendererSourceDataMode::Particles && BindingSourceMode == ENiagaraBindingSource::ImplicitFromSource) ||
 		BindingSourceMode == ENiagaraBindingSource::ExplicitParticles)
 	{
-		ensure(!DataSetVariable.IsInNameSpace(FNiagaraConstants::ParticleAttributeNamespaceString));
-		ParamMapVariable.SetNamespacedName(FNiagaraConstants::ParticleAttributeNamespaceString, DataSetVariable.GetName());
+		ensure(!GetDataSetBindableVariable().IsInNameSpace(FNiagaraConstants::ParticleAttributeNamespaceString));
+		ParamMapVariable.SetNamespacedName(FNiagaraConstants::ParticleAttributeNamespaceString, DataSetName);
 	}
 
 #if WITH_EDITORONLY_DATA
@@ -792,13 +799,13 @@ void FNiagaraVariableAttributeBinding::CacheValues(const FVersionedNiagaraEmitte
 			FNiagaraAliasContext ResolveAliasesContext(FNiagaraAliasContext::ERapidIterationParameterMode::EmitterOrParticleScript);
 			ResolveAliasesContext.ChangeEmitterToEmitterName(Emitter->GetUniqueEmitterName());
 			ParamMapVariable = FNiagaraUtilities::ResolveAliases(ParamMapVariable, ResolveAliasesContext);
-			DataSetVariable = FNiagaraUtilities::ResolveAliases(DataSetVariable, ResolveAliasesContext);
+			DataSetName = FNiagaraUtilities::ResolveAliases(GetDataSetBindableVariable(), ResolveAliasesContext).GetName();
 		}
 
 		FNiagaraTypeDefinition BoundVarType = ParamMapVariable.GetType();
 		if (BindingSourceMode == ENiagaraBindingSource::ExplicitParticles || (InSourceMode == ENiagaraRendererSourceDataMode::Particles && BindingSourceMode == ENiagaraBindingSource::ImplicitFromSource))
 		{
-			bBindingExistsOnSource = Emitter->CanObtainParticleAttribute(DataSetVariable, InVersionedEmitter.Version, BoundVarType);
+			bBindingExistsOnSource = Emitter->CanObtainParticleAttribute(GetDataSetBindableVariable(), InVersionedEmitter.Version, BoundVarType);
 		}
 		else if (BindingSourceMode == ENiagaraBindingSource::ExplicitEmitter || (InSourceMode == ENiagaraRendererSourceDataMode::Emitter && BindingSourceMode == ENiagaraBindingSource::ImplicitFromSource))
 		{
diff --git a/Engine/Plugins/FX/Niagara/Source/Niagara/Public/NiagaraCommon.h b/Engine/Plugins/FX/Niagara/Source/Niagara/Public/NiagaraCommon.h
--- a/Engine/Plugins/FX/Niagara/Source/Niagara/Public/NiagaraCommon.h
+++ b/Engine/Plugins/FX/Niagara/Source/Niagara/Public/NiagaraCommon.h
@@ -1189,14 +1189,14 @@ struct FNiagaraVariableAttributeBinding
 
 #if WITH_EDITORONLY_DATA
 	NIAGARA_API const FName& GetName() const { return CachedDisplayName; }
-	NIAGARA_API FString GetDefaultValueString() const;
+	NIAGARA_API [[nodiscard]] FString GetDefaultValueString() const;
 #endif
 	NIAGARA_API const FNiagaraVariableBase& GetParamMapBindableVariable() const { return ParamMapVariable; }
-	NIAGARA_API const FNiagaraVariableBase& GetDataSetBindableVariable() const { return DataSetVariable; }
-	NIAGARA_API const FNiagaraTypeDefinition& GetType() const { return DataSetVariable.GetType(); }
+	NIAGARA_API [[nodiscard]] FNiagaraVariableBase GetDataSetBindableVariable() const { return FNiagaraVariableBase(ParamMapVariable.GetType(), DataSetName); }
+	NIAGARA_API const FNiagaraTypeDefinition& GetType() const { return ParamMapVariable.GetType(); }
 	NIAGARA_API ENiagaraBindingSource GetBindingSourceMode() const { return BindingSourceMode; }
 
-	NIAGARA_API bool IsValid() const {	return DataSetVariable.IsValid();}
+	NIAGARA_API bool IsValid() const { return !DataSetName.IsNone() && ParamMapVariable.GetType().IsValid(); }
 	
 	template<typename T>
 	T GetDefaultValue() const
@@ -1211,18 +1211,19 @@ struct FNiagaraVariableAttributeBinding
 
 	void NIAGARA_API ResetToDefault(const FNiagaraVariableAttributeBinding& InOther, const FVersionedNiagaraEmitter& InEmitter, ENiagaraRendererSourceDataMode InSourceMode);
 	bool NIAGARA_API MatchesDefault(const FNiagaraVariableAttributeBinding& InOther, ENiagaraRendererSourceDataMode InSourceMode) const;
+
 protected:
+	/** The namespace and default value explicitly set by the user. If meant to be derived from the source mode, it will be without a namespace.*/
+	UPROPERTY()
+	FNiagaraVariable RootVariable;
+
 	/** The fully expressed namespace for the variable. If an emitter namespace, this will include the Emitter's unique name.*/
 	UPROPERTY()
 	FNiagaraVariableBase ParamMapVariable;
 
 	/** The version of the namespace to be found in an attribute table lookup. I.e. without Particles or Emitter.*/
 	UPROPERTY()
-	FNiagaraVariable DataSetVariable;
-
-	/** The namespace and default value explicitly set by the user. If meant to be derived from the source mode, it will be without a namespace.*/
-	UPROPERTY()
-	FNiagaraVariable RootVariable;
+	FName DataSetName;
 
 #if WITH_EDITORONLY_DATA
 	/** Old variable brought in from previous setup. Generally ignored other than postload work.*/
@@ -1239,11 +1240,17 @@ protected:
 	
 	/** Determine if this varible is accessible by the associated emitter passed into CacheValues.*/
 	UPROPERTY()
-	uint32 bBindingExistsOnSource : 1; 
+	uint8 bBindingExistsOnSource : 1; 
 
 	/** When CacheValues is called, was this a particle attribute?*/
 	UPROPERTY()
-	uint32 bIsCachedParticleValue : 1;
+	uint8 bIsCachedParticleValue : 1;
+
+#if WITH_EDITORONLY_DATA
+	/** The version of the namespace to be found in an attribute table lookup. I.e. without Particles or Emitter.*/
+	UPROPERTY()
+	FNiagaraVariable DataSetVariable_DEPRECATED;
+#endif
 };
 
 USTRUCT()
