- Move nodiscard for Linx build

#rnx

diff --git a/Engine/Plugins/FX/Niagara/Source/Niagara/Public/NiagaraCommon.h b/Engine/Plugins/FX/Niagara/Source/Niagara/Public/NiagaraCommon.h
--- a/Engine/Plugins/FX/Niagara/Source/Niagara/Public/NiagaraCommon.h
+++ b/Engine/Plugins/FX/Niagara/Source/Niagara/Public/NiagaraCommon.h
@@ -1189,10 +1189,10 @@ struct FNiagaraVariableAttributeBinding
 
 #if WITH_EDITORONLY_DATA
 	NIAGARA_API const FName& GetName() const { return CachedDisplayName; }
-	NIAGARA_API [[nodiscard]] FString GetDefaultValueString() const;
+	[[nodiscard]] NIAGARA_API FString GetDefaultValueString() const;
 #endif
 	NIAGARA_API const FNiagaraVariableBase& GetParamMapBindableVariable() const { return ParamMapVariable; }
-	NIAGARA_API [[nodiscard]] FNiagaraVariableBase GetDataSetBindableVariable() const { return FNiagaraVariableBase(ParamMapVariable.GetType(), DataSetName); }
+	[[nodiscard]] NIAGARA_API FNiagaraVariableBase GetDataSetBindableVariable() const { return FNiagaraVariableBase(ParamMapVariable.GetType(), DataSetName); }
 	NIAGARA_API const FNiagaraTypeDefinition& GetType() const { return ParamMapVariable.GetType(); }
 	NIAGARA_API ENiagaraBindingSource GetBindingSourceMode() const { return BindingSourceMode; }
 
