fix rigid mesh collisions

diff --git a/Engine/Plugins/FX/Niagara/Shaders/Private/NiagaraDataInterfaceRigidMeshCollisionQuery.ush b/Engine/Plugins/FX/Niagara/Shaders/Private/NiagaraDataInterfaceRigidMeshCollisionQuery.ush
--- a/Engine/Plugins/FX/Niagara/Shaders/Private/NiagaraDataInterfaceRigidMeshCollisionQuery.ush
+++ b/Engine/Plugins/FX/Niagara/Shaders/Private/NiagaraDataInterfaceRigidMeshCollisionQuery.ush
@@ -482,6 +482,8 @@ void GetClosestPointMeshDistanceFieldAccurate_{ParameterName}(in float3 WorldPos
 // Given a world space position (WorldPosition) compute the static mesh closest point (position,normal,velocity)
 void GetClosestPointMeshDistanceFieldNoNormal_{ParameterName}(in float3 WorldPosition, in float DeltaTime, in float TimeFraction, in float MaxDistance, out float OutClosestDistance, out float3 OutClosestPosition, out float3 OutClosestVelocity)
 {	
+	const FLWCVector3 LWCWorldPosition = MakeLWCVector3({ParameterName}_SystemLWCTile, WorldPosition);
+
 	float3 CollisionPosition = float3(0,0,0);
 	float3 CollisionNormal = float3(0,0,0);
 	OutClosestDistance = MAX_DISTANCE;	
@@ -489,44 +491,20 @@ void GetClosestPointMeshDistanceFieldNoNormal_{ParameterName}(in float3 WorldPos
 	OutClosestPosition = float3(0,0,0);
 	OutClosestVelocity = float3(0,0,0);
 	
-	const int ElementIndex = GetClosestElement_{ParameterName}(WorldPosition,CollisionPosition,CollisionNormal,OutClosestDistance,TimeFraction);
+	float OutClosestDistanceTmp;
+	const int ElementIndex = GetClosestElement_{ParameterName}(WorldPosition,CollisionPosition,CollisionNormal,OutClosestDistanceTmp,TimeFraction);
 
 	if (ElementIndex != -1)
 	{
-		int DFIndex = {ParameterName}_DFIndexBuffer[ElementIndex];
+		uint DFIndex = {ParameterName}_DFIndexBuffer[ElementIndex];
 
 		if (DFIndex >= 0 && DFIndex < NumSceneObjects)
 		{
 			// #todo(dmp): this does a dynamic branch based on intersecting the bbox.  Maybe we can factor that out due to the broadphase here?
+			// #todo(dmp): Check the distance from the physics asset for a speedup?
 			MaxDistance = max(.001, MaxDistance);
 
-			FDFObjectData DFObjectData = LoadDFObjectData(DFIndex);
-
-			float3 WorldExtent = {ParameterName}_ElementExtentBuffer[ElementIndex].xyz;
-
-			float3 LocalPos = GetLocalPosition_{ParameterName}(WorldPosition, ElementIndex, TimeFraction);
-			LocalPos *= GetMeshScale_{ParameterName}(ElementIndex);
-			
-			LocalPos /= DFObjectData.VolumeToWorldScale;
-							
-			float3 VolumePosition = LocalPos;
-
-			float3 ClampedVolumePosition = clamp(VolumePosition, -DFObjectData.VolumePositionExtent, DFObjectData.VolumePositionExtent);
-			float DistanceToBox = length((VolumePosition - ClampedVolumePosition) * DFObjectData.VolumeToWorldScale);
-
-			OutClosestDistance = MaxDistance;
-
-			BRANCH
-			if (DistanceToBox < MaxDistance)
-			{
-				OutClosestDistance = DistanceToMeshSurfaceStandalone(ClampedVolumePosition, DFObjectData) * DFObjectData.VolumeScale + DistanceToBox;
-
-				// Don't allow negative SDF distance to stick out of the bounding box
-				if (DistanceToBox > 0.0f)
-				{
-					OutClosestDistance = max(OutClosestDistance, DistanceToBox);
-				}
-			}		
+			OutClosestDistance = DistanceToNearestSurfaceForObject(DFIndex, LWCWorldPosition, MaxDistance);			
 		}
 		
 		const float3 PreviousPosition = mul(GetPreviousTransform_{ParameterName}(ElementIndex), float4(CollisionPosition,1.0)).xyz;
