- Reorder some member variables to fix bad packing
- Implement GetResourceSizeEx to include renderer layout cost

diff --git a/Engine/Plugins/FX/Niagara/Source/Niagara/Private/NiagaraMeshRendererProperties.cpp b/Engine/Plugins/FX/Niagara/Source/Niagara/Private/NiagaraMeshRendererProperties.cpp
--- a/Engine/Plugins/FX/Niagara/Source/Niagara/Private/NiagaraMeshRendererProperties.cpp
+++ b/Engine/Plugins/FX/Niagara/Source/Niagara/Private/NiagaraMeshRendererProperties.cpp
@@ -242,17 +242,11 @@ bool FNiagaraMeshRendererMeshProperties::HasValidRenderableMesh() const
 }
 
 UNiagaraMeshRendererProperties::UNiagaraMeshRendererProperties()
-	: SourceMode(ENiagaraRendererSourceDataMode::Particles)
-	, SortMode(ENiagaraSortMode::None)
-	, bOverrideMaterials(false)
+	: bOverrideMaterials(false)
 	, bUseHeterogeneousVolumes(false)
 	, bSortOnlyWhenTranslucent(true)
 	, bSubImageBlend(false)
-	, SubImageSize(1.0f, 1.0f)
-	, FacingMode(ENiagaraMeshFacingMode::Default)
 	, bLockedAxisEnable(false)
-	, LockedAxis(0.0f, 0.0f, 1.0f)
-	, LockedAxisSpace(ENiagaraMeshLockedAxisSpace::Simulation)
 {
 	// Initialize the array with a single, defaulted entry
 	Meshes.AddDefaulted();
@@ -403,6 +397,14 @@ void UNiagaraMeshRendererProperties::Serialize(FArchive& Ar)
 	Super::Serialize(Ar);
 }
 
+void UNiagaraMeshRendererProperties::GetResourceSizeEx(FResourceSizeEx& CumulativeResourceSize)
+{
+	Super::GetResourceSizeEx(CumulativeResourceSize);
+
+	CumulativeResourceSize.AddDedicatedSystemMemoryBytes(RendererLayoutWithCustomSorting.GetAllocatedSize());
+	CumulativeResourceSize.AddDedicatedSystemMemoryBytes(RendererLayoutWithoutCustomSorting.GetAllocatedSize());
+}
+
 /** The bindings depend on variables that are created during the NiagaraModule startup. However, the CDO's are build prior to this being initialized, so we defer setting these values until later.*/
 void UNiagaraMeshRendererProperties::InitCDOPropertiesAfterModuleStartup()
 {
@@ -790,6 +792,7 @@ void UNiagaraMeshRendererProperties::PostLoad()
 {
 	Super::PostLoad();
 
+#if WITH_EDITORONLY_DATA
 	if (Meshes.Num() == 1 && Meshes[0].Mesh == nullptr && ParticleMesh_DEPRECATED != nullptr)
 	{
 		// Likely predates the mesh array ... just add ParticleMesh to the list of Meshes
@@ -798,6 +801,7 @@ void UNiagaraMeshRendererProperties::PostLoad()
 		Mesh.PivotOffset = PivotOffset_DEPRECATED;
 		Mesh.PivotOffsetSpace = PivotOffsetSpace_DEPRECATED;
 	}
+#endif
 
 	for (FNiagaraMeshRendererMeshProperties& MeshProperties : Meshes)
 	{
diff --git a/Engine/Plugins/FX/Niagara/Source/Niagara/Private/NiagaraRendererProperties.cpp b/Engine/Plugins/FX/Niagara/Source/Niagara/Private/NiagaraRendererProperties.cpp
--- a/Engine/Plugins/FX/Niagara/Source/Niagara/Private/NiagaraRendererProperties.cpp
+++ b/Engine/Plugins/FX/Niagara/Source/Niagara/Private/NiagaraRendererProperties.cpp
@@ -834,10 +834,12 @@ void UNiagaraRendererProperties::PostLoad()
 {
 	Super::PostLoad();
 
+#if WITH_EDITORONLY_DATA
 	if (bMotionBlurEnabled_DEPRECATED == false)
 	{
 		MotionVectorSetting = ENiagaraRendererMotionVectorSetting::Disable;
 	}
+#endif
 }
 
 #if WITH_EDITORONLY_DATA
diff --git a/Engine/Plugins/FX/Niagara/Source/Niagara/Private/NiagaraRibbonRendererProperties.cpp b/Engine/Plugins/FX/Niagara/Source/Niagara/Private/NiagaraRibbonRendererProperties.cpp
--- a/Engine/Plugins/FX/Niagara/Source/Niagara/Private/NiagaraRibbonRendererProperties.cpp
+++ b/Engine/Plugins/FX/Niagara/Source/Niagara/Private/NiagaraRibbonRendererProperties.cpp
@@ -38,18 +38,17 @@ FNiagaraRibbonUVSettings::FNiagaraRibbonUVSettings()
 	: DistributionMode(ENiagaraRibbonUVDistributionMode::ScaledUsingRibbonSegmentLength)
 	, LeadingEdgeMode(ENiagaraRibbonUVEdgeMode::Locked)
 	, TrailingEdgeMode(ENiagaraRibbonUVEdgeMode::Locked)
+	, bEnablePerParticleUOverride(false)
+	, bEnablePerParticleVRangeOverride(false)
 	, TilingLength(100.0f)
 	, Offset(FVector2D(0.0f, 0.0f))
 	, Scale(FVector2D(1.0f, 1.0f))
-	, bEnablePerParticleUOverride(false)
-	, bEnablePerParticleVRangeOverride(false)
 {
 }
 
 UNiagaraRibbonRendererProperties::UNiagaraRibbonRendererProperties()
 	: Material(nullptr)
 	, MaterialUserParamBinding(FNiagaraTypeDefinition(UMaterialInterface::StaticClass()))
-	, FacingMode(ENiagaraRibbonFacingMode::Screen)
 #if WITH_EDITORONLY_DATA
 	, UV0TilingDistance_DEPRECATED(0.0f)
 	, UV0Scale_DEPRECATED(FVector2D(1.0f, 1.0f))
@@ -59,18 +58,18 @@ UNiagaraRibbonRendererProperties::UNiagaraRibbonRendererProperties()
 	, UV1AgeOffsetMode_DEPRECATED(ENiagaraRibbonAgeOffsetMode::Scale)
 #endif
 	, MaxNumRibbons(0)
+	, bEnableAccurateGeometry(false)
 	, bUseGPUInit(false)
+	, bUseConstantFactor(false)
+	, bScreenSpaceTessellation(true)
 	, Shape(ENiagaraRibbonShapeMode::Plane)
-	, bEnableAccurateGeometry(false)
 	, WidthSegmentationCount(1)
 	, MultiPlaneCount(2)
 	, TubeSubdivisions(3)
 	, TessellationMode(ENiagaraRibbonTessellationMode::Automatic)
 	, CurveTension(0.f)
 	, TessellationFactor(16)
-	, bUseConstantFactor(false)
 	, TessellationAngle(15)
-	, bScreenSpaceTessellation(true)
 {
 	AttributeBindings.Reserve(19);
 	AttributeBindings.Add(&PositionBinding);
@@ -300,6 +299,13 @@ void UNiagaraRibbonRendererProperties::Serialize(FStructuredArchive::FRecord Rec
 	}
 }
 
+void UNiagaraRibbonRendererProperties::GetResourceSizeEx(FResourceSizeEx& CumulativeResourceSize)
+{
+	Super::GetResourceSizeEx(CumulativeResourceSize);
+
+	CumulativeResourceSize.AddDedicatedSystemMemoryBytes(RendererLayout.GetAllocatedSize());
+}
+
 /** The bindings depend on variables that are created during the NiagaraModule startup. However, the CDO's are build prior to this being initialized, so we defer setting these values until later.*/
 void UNiagaraRibbonRendererProperties::InitCDOPropertiesAfterModuleStartup()
 {
diff --git a/Engine/Plugins/FX/Niagara/Source/Niagara/Private/NiagaraSpriteRendererProperties.cpp b/Engine/Plugins/FX/Niagara/Source/Niagara/Private/NiagaraSpriteRendererProperties.cpp
--- a/Engine/Plugins/FX/Niagara/Source/Niagara/Private/NiagaraSpriteRendererProperties.cpp
+++ b/Engine/Plugins/FX/Niagara/Source/Niagara/Private/NiagaraSpriteRendererProperties.cpp
@@ -53,18 +53,10 @@ FCookStatsManager::FAutoRegisterCallback NiagaraCutoutCookStats::RegisterCookSta
 
 UNiagaraSpriteRendererProperties::UNiagaraSpriteRendererProperties()
 	: Material(nullptr)
-	, SourceMode(ENiagaraRendererSourceDataMode::Particles)
 	, MaterialUserParamBinding(FNiagaraTypeDefinition(UMaterialInterface::StaticClass()))
-	, Alignment(ENiagaraSpriteAlignment::Automatic)
-	, FacingMode(ENiagaraSpriteFacingMode::Automatic)
-	, PivotInUVSpace(0.5f, 0.5f)
-	, SortMode(ENiagaraSortMode::None)
-	, SubImageSize(1.0f, 1.0f)
 	, bSubImageBlend(false)
 	, bRemoveHMDRollInVR(false)
 	, bSortOnlyWhenTranslucent(true)
-	, MinFacingCameraBlendDistance(0.0f)
-	, MaxFacingCameraBlendDistance(0.0f)
 #if WITH_EDITORONLY_DATA
 	, BoundingMode(BVC_EightVertices)
 	, AlphaThreshold(0.1f)
@@ -256,6 +248,14 @@ void UNiagaraSpriteRendererProperties::Serialize(FStructuredArchive::FRecord Rec
 	}
 }
 
+void UNiagaraSpriteRendererProperties::GetResourceSizeEx(FResourceSizeEx& CumulativeResourceSize)
+{
+	Super::GetResourceSizeEx(CumulativeResourceSize);
+
+	CumulativeResourceSize.AddDedicatedSystemMemoryBytes(RendererLayoutWithCustomSort.GetAllocatedSize());
+	CumulativeResourceSize.AddDedicatedSystemMemoryBytes(RendererLayoutWithoutCustomSort.GetAllocatedSize());
+}
+
 /** The bindings depend on variables that are created during the NiagaraModule startup. However, the CDO's are build prior to this being initialized, so we defer setting these values until later.*/
 void UNiagaraSpriteRendererProperties::InitCDOPropertiesAfterModuleStartup()
 {
diff --git a/Engine/Plugins/FX/Niagara/Source/Niagara/Public/NiagaraMeshRendererProperties.h b/Engine/Plugins/FX/Niagara/Source/Niagara/Public/NiagaraMeshRendererProperties.h
--- a/Engine/Plugins/FX/Niagara/Source/Niagara/Public/NiagaraMeshRendererProperties.h
+++ b/Engine/Plugins/FX/Niagara/Source/Niagara/Public/NiagaraMeshRendererProperties.h
@@ -187,6 +187,7 @@ public:
 	virtual void PostLoad() override;
 	virtual void PostInitProperties() override;
 	virtual void Serialize(FArchive& Ar) override;
+	virtual void GetResourceSizeEx(FResourceSizeEx& CumulativeResourceSize) override;
 #if WITH_EDITORONLY_DATA
 	virtual void BeginDestroy() override;
 	virtual void PreEditChange(class FProperty* PropertyThatWillChange) override;
@@ -242,11 +243,23 @@ public:
 
 	/** Whether or not to draw a single element for the Emitter or to draw the particles.*/
 	UPROPERTY(EditAnywhere, Category = "Mesh Rendering")
-	ENiagaraRendererSourceDataMode SourceMode;
+	ENiagaraRendererSourceDataMode SourceMode = ENiagaraRendererSourceDataMode::Particles;
 
 	/** Determines how we sort the particles prior to rendering.*/
 	UPROPERTY(EditAnywhere, Category = "Sorting")
-	ENiagaraSortMode SortMode;
+	ENiagaraSortMode SortMode = ENiagaraSortMode::None;
+
+	/** Sort precision to use when sorting is active. */
+	UPROPERTY(EditAnywhere, Category = "Sorting", meta = (EditCondition = "SortMode != ENiagaraSortMode::None", EditConditionHides, DisplayAfter="bSortOnlyWhenTranslucent"))
+	ENiagaraRendererSortPrecision SortPrecision = ENiagaraRendererSortPrecision::Default;
+
+	/**
+	Gpu simulations run at different points in the frame depending on what features are used, i.e. depth buffer, distance fields, etc.
+	Opaque materials will run latent when these features are used.
+	Translucent materials can choose if they want to use this frames or the previous frames data to match opaque draws.
+	*/
+	UPROPERTY(EditAnywhere, AdvancedDisplay, Category = "Mesh Rendering")
+	ENiagaraRendererGpuTranslucentLatency GpuTranslucentLatency = ENiagaraRendererGpuTranslucentLatency::ProjectDefault;
 
 	/** Whether or not to use the OverrideMaterials array instead of the mesh's existing materials.*/
 	UPROPERTY(EditAnywhere, Category = "Mesh Rendering", DisplayName="Enable Material Overrides")
@@ -260,18 +273,6 @@ public:
 	UPROPERTY(EditAnywhere, Category = "Sorting", meta = (EditCondition = "SortMode != ENiagaraSortMode::None", EditConditionHides))
 	uint32 bSortOnlyWhenTranslucent : 1;
 
-	/** Sort precision to use when sorting is active. */
-	UPROPERTY(EditAnywhere, Category = "Sorting", meta = (EditCondition = "SortMode != ENiagaraSortMode::None", EditConditionHides))
-	ENiagaraRendererSortPrecision SortPrecision = ENiagaraRendererSortPrecision::Default;
-
-	/**
-	Gpu simulations run at different points in the frame depending on what features are used, i.e. depth buffer, distance fields, etc.
-	Opaque materials will run latent when these features are used.
-	Translucent materials can choose if they want to use this frames or the previous frames data to match opaque draws.
-	*/
-	UPROPERTY(EditAnywhere, AdvancedDisplay, Category = "Mesh Rendering")
-	ENiagaraRendererGpuTranslucentLatency GpuTranslucentLatency = ENiagaraRendererGpuTranslucentLatency::ProjectDefault;
-
 	/** If true, blends the sub-image UV lookup with its next adjacent member using the fractional part of the SubImageIndex float value as the linear interpolation factor.*/
 	UPROPERTY(EditAnywhere, Category = "SubUV", meta = (DisplayName = "Sub UV Blending Enabled"))
 	uint32 bSubImageBlend : 1;
@@ -288,6 +289,10 @@ public:
 	UPROPERTY(EditAnywhere, Category = "Mesh Rendering", meta = (DisplayAfter = "MeshBoundsScale"))
 	uint32 bEnableMeshFlipbook : 1;
 
+	/** If true and in a non-default facing mode, will lock facing direction to an arbitrary plane of rotation */
+	UPROPERTY(EditAnywhere, Category = "Mesh Rendering")
+	uint32 bLockedAxisEnable : 1;
+
 	/** The materials to be used instead of the StaticMesh's materials. Note that each material must have the Niagara Mesh Particles flag checked. If the ParticleMesh
 	requires more materials than exist in this array or any entry in this array is set to None, we will use the ParticleMesh's existing Material instead.*/
 	UPROPERTY(EditAnywhere, Category = "Mesh Rendering", meta = (EditCondition = "bOverrideMaterials", EditConditionHides))
@@ -298,23 +303,11 @@ public:
 
 	/** When using SubImage lookups for particles, this variable contains the number of columns in X and the number of rows in Y.*/
 	UPROPERTY(EditAnywhere, Category = "SubUV")
-	FVector2D SubImageSize;
-
-	/** Determines how the mesh orients itself relative to the camera. */
-	UPROPERTY(EditAnywhere, Category = "Mesh Rendering")
-	ENiagaraMeshFacingMode FacingMode;
-
-	/** If true and in a non-default facing mode, will lock facing direction to an arbitrary plane of rotation */
-	UPROPERTY(EditAnywhere, Category = "Mesh Rendering")
-	uint32 bLockedAxisEnable : 1;
+	FVector2D SubImageSize = FVector2D(1.0f, 1.0f);
 
 	/** Arbitrary axis by which to lock facing rotations */
-	UPROPERTY(EditAnywhere, Category = "Mesh Rendering", meta = (EditCondition = "bLockedAxisEnable", EditConditionHides))
-	FVector LockedAxis;
-
-	/** Specifies what space the locked axis is in */
-	UPROPERTY(EditAnywhere, Category = "Mesh Rendering", meta = (EditCondition = "bLockedAxisEnable", EditConditionHides))
-	ENiagaraMeshLockedAxisSpace LockedAxisSpace;
+	UPROPERTY(EditAnywhere, Category = "Mesh Rendering", meta = (EditCondition = "bLockedAxisEnable", EditConditionHides, DisplayAfter="bLockedAxisEnable"))
+	FVector LockedAxis = FVector(0.0f, 0.0f, 1.0f);
 
 	/**
 	Scale factor applied to all of the meshes bounds.
@@ -325,6 +318,14 @@ public:
 	UPROPERTY(EditAnywhere, Category = "Mesh Rendering", AdvancedDisplay)
 	FVector MeshBoundsScale = FVector::OneVector;
 
+	/** Determines how the mesh orients itself relative to the camera. */
+	UPROPERTY(EditAnywhere, Category = "Mesh Rendering")
+	ENiagaraMeshFacingMode FacingMode = ENiagaraMeshFacingMode::Default;
+
+	/** Specifies what space the locked axis is in */
+	UPROPERTY(EditAnywhere, Category = "Mesh Rendering", meta = (EditCondition = "bLockedAxisEnable", EditConditionHides, DisplayAfter = "LockedAxis"))
+	ENiagaraMeshLockedAxisSpace LockedAxisSpace = ENiagaraMeshLockedAxisSpace::Simulation;
+
 	UPROPERTY(EditAnywhere, Category = "Visibility", meta = (EditCondition = "bEnableCameraDistanceCulling", EditConditionHides, ClampMin = 0.0f))
 	float MinCameraDistance;
 
@@ -472,6 +473,7 @@ protected:
 private:
 	static TArray<TWeakObjectPtr<UNiagaraMeshRendererProperties>> MeshRendererPropertiesToDeferredInit;
 
+#if WITH_EDITORONLY_DATA
 	// These properties are deprecated and moved to FNiagaraMeshRendererMeshProperties
 	UPROPERTY()
 	TObjectPtr<UStaticMesh> ParticleMesh_DEPRECATED;
@@ -481,4 +483,5 @@ private:
 
 	UPROPERTY()
 	ENiagaraMeshPivotOffsetSpace PivotOffsetSpace_DEPRECATED;
+#endif
 };
diff --git a/Engine/Plugins/FX/Niagara/Source/Niagara/Public/NiagaraRendererProperties.h b/Engine/Plugins/FX/Niagara/Source/Niagara/Public/NiagaraRendererProperties.h
--- a/Engine/Plugins/FX/Niagara/Source/Niagara/Public/NiagaraRendererProperties.h
+++ b/Engine/Plugins/FX/Niagara/Source/Niagara/Public/NiagaraRendererProperties.h
@@ -146,6 +146,8 @@ struct NIAGARA_API FNiagaraRendererLayout
 	int32 GetTotalFloatComponents_RenderThread() const { check(IsInRenderingThread()); return TotalFloatComponents_RT; }
 	int32 GetTotalHalfComponents_RenderThread() const { check(IsInRenderingThread()); return TotalHalfComponents_RT; }
 
+	SIZE_T GetAllocatedSize() const { return VFVariables_GT.GetAllocatedSize() + VFVariables_RT.GetAllocatedSize(); }
+
 private:
 	TArray<FNiagaraRendererVariableInfo> VFVariables_GT;
 	int32 TotalFloatComponents_GT;
@@ -282,7 +284,6 @@ public:
 	UNiagaraRendererProperties()		
 		: bIsEnabled(true)
 		, bAllowInCullProxies(true)
-		, bMotionBlurEnabled_DEPRECATED(true)
 	{
 	}
 
@@ -427,10 +428,12 @@ public:
 
 	UPROPERTY()
 	FGuid OuterEmitterVersion;
-	
+
 protected:
+#if WITH_EDITORONLY_DATA
 	UPROPERTY()
-	bool bMotionBlurEnabled_DEPRECATED; // This has been rolled into MotionVectorSetting
+	bool bMotionBlurEnabled_DEPRECATED = true; // This has been rolled into MotionVectorSetting
+#endif
 
 	TArray<const FNiagaraVariableAttributeBinding*> AttributeBindings;
 
diff --git a/Engine/Plugins/FX/Niagara/Source/Niagara/Public/NiagaraRibbonRendererProperties.h b/Engine/Plugins/FX/Niagara/Source/Niagara/Public/NiagaraRibbonRendererProperties.h
--- a/Engine/Plugins/FX/Niagara/Source/Niagara/Public/NiagaraRibbonRendererProperties.h
+++ b/Engine/Plugins/FX/Niagara/Source/Niagara/Public/NiagaraRibbonRendererProperties.h
@@ -91,7 +91,7 @@ enum class ENiagaraRibbonTessellationMode : uint8
 
 /** Specifies options for handling UVs at the leading and trailing edges of ribbons. */
 UENUM()
-enum class ENiagaraRibbonUVEdgeMode
+enum class ENiagaraRibbonUVEdgeMode : uint8
 {
 	/** The UV value at the edge will smoothly transition across the segment using normalized age.
 	This will result in	UV values which are outside of the standard 0-1 range and works best with
@@ -104,7 +104,7 @@ enum class ENiagaraRibbonUVEdgeMode
 
 /** Specifies options for distributing UV values across ribbon segments. */
 UENUM()
-enum class ENiagaraRibbonUVDistributionMode
+enum class ENiagaraRibbonUVDistributionMode : uint8
 {
 	/** Ribbon UVs will stretch the length of the ribbon, without repeating, but distributed by segment, so can be uneven with unequal length segments. */
 	ScaledUniformly UMETA(DisplayName = "Uniform Scale (By Segment)"),
@@ -139,6 +139,14 @@ struct FNiagaraRibbonUVSettings
 	UPROPERTY(EditAnywhere, Category = UVs, meta = (DisplayName="Trailing Edge Transition", EditCondition = "!bEnablePerParticleUOverride && DistributionMode != ENiagaraRibbonUVDistributionMode::TiledOverRibbonLength && DistributionMode != ENiagaraRibbonUVDistributionMode::TiledFromStartOverRibbonLength", EditConditionHides))
 	ENiagaraRibbonUVEdgeMode TrailingEdgeMode;
 
+	/** Enables overriding the U component with values read from the particles. When enabled, edge behavior and distribution are ignored. */
+	UPROPERTY(EditAnywhere, Category = UVs)
+	uint8 bEnablePerParticleUOverride : 1;
+
+	/** Enables overriding the range of the V component with values read from the particles. */
+	UPROPERTY(EditAnywhere, Category = UVs)
+	uint8 bEnablePerParticleVRangeOverride : 1;
+
 	/** Specifies the length in world units to use when tiling UVs across the ribbon when using one of the tiled distribution modes. */
 	UPROPERTY(EditAnywhere, Category = UVs, meta = (EditCondition="!bEnablePerParticleUOverride && DistributionMode == ENiagaraRibbonUVDistributionMode::TiledOverRibbonLength || DistributionMode == ENiagaraRibbonUVDistributionMode::TiledFromStartOverRibbonLength", EditConditionHides))
 	float TilingLength;
@@ -150,14 +158,6 @@ struct FNiagaraRibbonUVSettings
 	/** Specifies an additional scaler which is applied to the UV range. */
 	UPROPERTY(EditAnywhere, Category = UVs)
 	FVector2D Scale;
-
-	/** Enables overriding the U component with values read from the particles. When enabled, edge behavior and distribution are ignored. */
-	UPROPERTY(EditAnywhere, Category = UVs)
-	bool bEnablePerParticleUOverride;
-
-	/** Enables overriding the range of the V component with values read from the particles. */
-	UPROPERTY(EditAnywhere, Category = UVs)
-	bool bEnablePerParticleVRangeOverride;
 };
 
 namespace ENiagaraRibbonVFLayout
@@ -202,6 +202,7 @@ public:
 	virtual void PostLoad() override;
 	virtual void PostInitProperties() override;
 	virtual void Serialize(FStructuredArchive::FRecord Record) override;
+	virtual void GetResourceSizeEx(FResourceSizeEx& CumulativeResourceSize) override;
 #if WITH_EDITOR
 	virtual void PostEditChangeProperty(struct FPropertyChangedEvent& PropertyChangedEvent) override;
 	virtual void RenameVariable(const FNiagaraVariableBase& OldVariable, const FNiagaraVariableBase& NewVariable, const FVersionedNiagaraEmitter& InEmitter) override;
@@ -246,15 +247,15 @@ public:
 	UPROPERTY(EditAnywhere, Category = "Ribbon Rendering")
 	FNiagaraUserParameterBinding MaterialUserParamBinding;
 
-	UPROPERTY(EditAnywhere, Category = "Ribbon Rendering")
-	ENiagaraRibbonFacingMode FacingMode;
-
 	UPROPERTY(EditAnywhere, Category = "Ribbon Rendering", meta=(DisplayName="UV0 Settings"))
 	FNiagaraRibbonUVSettings UV0Settings;
 
 	UPROPERTY(EditAnywhere, Category = "Ribbon Rendering", meta=(DisplayName="UV1 Settings"))
 	FNiagaraRibbonUVSettings UV1Settings;
 
+	UPROPERTY(EditAnywhere, Category = "Ribbon Rendering", meta=(DisplayAfter="MaterialUserParamBinding"))
+	ENiagaraRibbonFacingMode FacingMode = ENiagaraRibbonFacingMode::Screen;
+
 #if WITH_EDITORONLY_DATA
 private:
 	/** Tiles UV0 based on the distance traversed by the ribbon. Disables offsetting UVs by age. */
@@ -287,17 +288,6 @@ public:
 	UPROPERTY(EditAnywhere, Category = "Ribbon Rendering")	
 	int32 MaxNumRibbons;
 	
-	/**
-	*	Whether we use the CPU or GPU to generate ribbon geometry for CPU systems.
-	*	GPU systems will always use a fully GPU initialization pipeline,
-	*	Will fall back to CPU init when GPU init isn't available.
-	*/ 
-	UPROPERTY(EditAnywhere, Category = "Ribbon Rendering")	
-	bool bUseGPUInit;
-
-
-	
-
 	/** If true, the particles are only sorted when using a translucent material. */
 	UPROPERTY(EditAnywhere, Category = "Ribbon Rendering")
 	ENiagaraRibbonDrawDirection DrawDirection;
@@ -311,7 +301,23 @@ public:
 	  * 3D Ribbons will render like normal meshes with backface culling enabled.
 	  */
 	UPROPERTY(EditAnywhere, Category = "Ribbon Shape", meta = (EditCondition = "Shape != ENiagaraRibbonShapeMode::Plane", EditConditionHides))
-	bool bEnableAccurateGeometry;
+	uint8 bEnableAccurateGeometry : 1;
+
+	/**
+	*	Whether we use the CPU or GPU to generate ribbon geometry for CPU systems.
+	*	GPU systems will always use a fully GPU initialization pipeline,
+	*	Will fall back to CPU init when GPU init isn't available.
+	*/ 
+	UPROPERTY(EditAnywhere, Category = "Ribbon Rendering")	
+	uint8 bUseGPUInit : 1;
+
+	/** If checked, use the above constant factor. Otherwise, adaptively select the tessellation factor based on the below parameters. */
+	UPROPERTY(EditAnywhere, Category = "Ribbon Tessellation", meta = (EditCondition = "TessellationMode == ENiagaraRibbonTessellationMode::Custom", EditConditionHides, DisplayAfter="TessellationFactor"))
+	uint8 bUseConstantFactor : 1;
+
+	/** If checked, use the ribbon's screen space percentage to adaptively adjust the tessellation factor. */
+	UPROPERTY(EditAnywhere, Category = "Ribbon Tessellation", meta = (DisplayName = "Screen Space", EditCondition = "TessellationMode == ENiagaraRibbonTessellationMode::Custom && !bUseConstantFactor", EditConditionHides, DisplayAfter="TessellationAngle"))
+	uint8 bScreenSpaceTessellation : 1;
 
 	/** Tessellation factor to apply to the width of the ribbon.
 	* Ranges from 1 to 16. Greater values increase amount of tessellation.
@@ -347,10 +353,6 @@ public:
 	UPROPERTY(EditAnywhere, Category = "Ribbon Tessellation", meta = (DisplayName = "Max Tessellation Factor", EditCondition = "TessellationMode == ENiagaraRibbonTessellationMode::Custom", EditConditionHides, ClampMin = "1", ClampMax = "16"))
 	int32 TessellationFactor;
 
-	/** If checked, use the above constant factor. Otherwise, adaptively select the tessellation factor based on the below parameters. */
-	UPROPERTY(EditAnywhere, Category = "Ribbon Tessellation", meta = (EditCondition = "TessellationMode == ENiagaraRibbonTessellationMode::Custom", EditConditionHides))
-	bool bUseConstantFactor;
-
 	/** Defines the angle in degrees at which tessellation occurs.
 	  * Ranges from 1 to 180. Smaller values increase amount of tessellation.
 	  * If set to 0, use the maximum tessellation set above.
@@ -358,10 +360,6 @@ public:
 	UPROPERTY(EditAnywhere, Category = "Ribbon Tessellation", meta = (EditCondition = "TessellationMode == ENiagaraRibbonTessellationMode::Custom && !bUseConstantFactor", ClampMin = "0", ClampMax = "180", UIMin = "1", UIMax = "180", EditConditionHides))
 	float TessellationAngle;
 
-	/** If checked, use the ribbon's screen space percentage to adaptively adjust the tessellation factor. */
-	UPROPERTY(EditAnywhere, Category = "Ribbon Tessellation", meta = (DisplayName = "Screen Space", EditCondition = "TessellationMode == ENiagaraRibbonTessellationMode::Custom && !bUseConstantFactor", EditConditionHides))
-	bool bScreenSpaceTessellation;
-
 	/** Which attribute should we use for position when generating ribbons?*/
 	UPROPERTY(EditAnywhere, Category = "Bindings")
 	FNiagaraVariableAttributeBinding PositionBinding;
diff --git a/Engine/Plugins/FX/Niagara/Source/Niagara/Public/NiagaraSpriteRendererProperties.h b/Engine/Plugins/FX/Niagara/Source/Niagara/Public/NiagaraSpriteRendererProperties.h
--- a/Engine/Plugins/FX/Niagara/Source/Niagara/Public/NiagaraSpriteRendererProperties.h
+++ b/Engine/Plugins/FX/Niagara/Source/Niagara/Public/NiagaraSpriteRendererProperties.h
@@ -118,12 +118,12 @@ public:
 	virtual void PostLoad() override;
 	virtual void PostInitProperties() override;
 	virtual void Serialize(FStructuredArchive::FRecord Record) override;
+	virtual void GetResourceSizeEx(FResourceSizeEx& CumulativeResourceSize) override;
 #if WITH_EDITORONLY_DATA
 	virtual void PostEditChangeProperty(struct FPropertyChangedEvent& PropertyChangedEvent) override;
 	virtual void RenameVariable(const FNiagaraVariableBase& OldVariable, const FNiagaraVariableBase& NewVariable, const FVersionedNiagaraEmitter& InEmitter) override;
 	virtual void RemoveVariable(const FNiagaraVariableBase& OldVariable, const FVersionedNiagaraEmitter& InEmitter) override;
 	virtual TArray<FNiagaraVariable> GetBoundAttributes() const override;
-
 #endif // WITH_EDITORONLY_DATA
 	//UObject Interface END
 
@@ -162,52 +162,56 @@ public:
 	TObjectPtr<UMaterialInstanceConstant> MICMaterial;
 #endif
 
-	/** Whether or not to draw a single element for the Emitter or to draw the particles.*/
-	UPROPERTY(EditAnywhere, Category = "Sprite Rendering")
-	ENiagaraRendererSourceDataMode SourceMode;
-
 	/** Use the UMaterialInterface bound to this user variable if it is set to a valid value. If this is bound to a valid value and Material is also set, UserParamBinding wins.*/
 	UPROPERTY(EditAnywhere, Category = "Sprite Rendering")
 	FNiagaraUserParameterBinding MaterialUserParamBinding;
 
+	/** Whether or not to draw a single element for the Emitter or to draw the particles.*/
+	UPROPERTY(EditAnywhere, Category = "Sprite Rendering")
+	ENiagaraRendererSourceDataMode SourceMode = ENiagaraRendererSourceDataMode::Particles;
+
 	/** Imagine the particle texture having an arrow pointing up, these modes define how the particle aligns that texture to other particle attributes.*/
 	UPROPERTY(EditAnywhere, Category = "Sprite Rendering")
-	ENiagaraSpriteAlignment Alignment;
+	ENiagaraSpriteAlignment Alignment = ENiagaraSpriteAlignment::Automatic;
 
 	/** Determines how the particle billboard orients itself relative to the camera.*/
 	UPROPERTY(EditAnywhere, Category = "Sprite Rendering")
-	ENiagaraSpriteFacingMode FacingMode;
+	ENiagaraSpriteFacingMode FacingMode = ENiagaraSpriteFacingMode::Automatic;
 
-	/**
-	 * Determines the location of the pivot point of this particle. It follows Unreal's UV space, which has the upper left of the image at 0,0 and bottom right at 1,1. The middle is at 0.5, 0.5.
-	 * NOTE: This value is ignored if "Pivot Offset Binding" is bound to a valid attribute
-	 */
-	UPROPERTY(EditAnywhere, Category = "Sprite Rendering", meta = (DisplayName = "Default Pivot in UV Space"))
-	FVector2D PivotInUVSpace;
+	/** Determines how we sort the particles prior to rendering.*/
+	UPROPERTY(EditAnywhere, Category = "Sorting")
+	ENiagaraSortMode SortMode = ENiagaraSortMode::None;
 
 	/** World space radius that UVs generated with the ParticleMacroUV material node will tile based on. */
 	UPROPERTY(EditAnywhere, Category = "Sprite Rendering")
 	float MacroUVRadius = 0.0f;
 
-	/** Determines how we sort the particles prior to rendering.*/
-	UPROPERTY(EditAnywhere, Category = "Sorting")
-	ENiagaraSortMode SortMode;
+	/**
+	 * Determines the location of the pivot point of this particle. It follows Unreal's UV space, which has the upper left of the image at 0,0 and bottom right at 1,1. The middle is at 0.5, 0.5.
+	 * NOTE: This value is ignored if "Pivot Offset Binding" is bound to a valid attribute
+	 */
+	UPROPERTY(EditAnywhere, Category = "Sprite Rendering", meta = (DisplayName = "Default Pivot in UV Space"))
+	FVector2D PivotInUVSpace = FVector2D(0.5f, 0.5f);
 
 	/** When using SubImage lookups for particles, this variable contains the number of columns in X and the number of rows in Y.*/
 	UPROPERTY(EditAnywhere, Category = "SubUV", meta = (DisplayAfter = bSubImageBlend))
-	FVector2D SubImageSize;
+	FVector2D SubImageSize = FVector2D(1.0f, 1.0f);
 
 	/** If true, blends the sub-image UV lookup with its next adjacent member using the fractional part of the SubImageIndex float value as the linear interpolation factor.*/
 	UPROPERTY(EditAnywhere, Category = "SubUV", meta = (DisplayName = "Sub UV Blending Enabled"))
-	uint32 bSubImageBlend : 1;
+	uint8 bSubImageBlend : 1;
 
 	/** If true, removes the HMD view roll (e.g. in VR) */
 	UPROPERTY(EditAnywhere, Category = "Sprite Rendering", meta = (DisplayName = "Remove HMD Roll"))
-	uint32 bRemoveHMDRollInVR : 1;
+	uint8 bRemoveHMDRollInVR : 1;
 
 	/** If true, the particles are only sorted when using a translucent material. */
 	UPROPERTY(EditAnywhere, Category = "Sorting", meta = (EditCondition = "SortMode != ENiagaraSortMode::None", EditConditionHides))
-	uint32 bSortOnlyWhenTranslucent : 1;
+	uint8 bSortOnlyWhenTranslucent : 1;
+
+	/** Enables frustum culling of individual sprites */
+	UPROPERTY(EditAnywhere, Category = "Visibility")
+	uint8 bEnableCameraDistanceCulling : 1;
 
 	/** Sort precision to use when sorting is active. */
 	UPROPERTY(EditAnywhere, Category = "Sorting", meta = (EditCondition = "SortMode != ENiagaraSortMode::None", EditConditionHides))
@@ -241,15 +245,11 @@ public:
 
 	/** When FacingMode is FacingCameraDistanceBlend, the distance at which the sprite is fully facing the camera plane. */
 	UPROPERTY(EditAnywhere, Category = "Sprite Rendering", meta = (UIMin = "0"))
-	float MinFacingCameraBlendDistance;
+	float MinFacingCameraBlendDistance = 0.0f;
 
 	/** When FacingMode is FacingCameraDistanceBlend, the distance at which the sprite is fully facing the camera position */
 	UPROPERTY(EditAnywhere, Category = "Sprite Rendering", meta = (UIMin = "0"))
-	float MaxFacingCameraBlendDistance;
-
-	/** Enables frustum culling of individual sprites */
-	UPROPERTY(EditAnywhere, Category = "Visibility")
-	uint32 bEnableCameraDistanceCulling : 1;
+	float MaxFacingCameraBlendDistance = 0.0f;
 
 	UPROPERTY(EditAnywhere, Category = "Visibility", meta = (EditCondition = "bEnableCameraDistanceCulling", UIMin = 0.0f, EditConditionHides))
 	float MinCameraDistance;
