- Move performance graph into it's own overview mode to make discovery of both modes easier
- Fix poor formatting of performance values and ensure the columns don't move around

diff --git a/Engine/Plugins/FX/Niagara/Source/Niagara/Private/NiagaraDebugHud.cpp b/Engine/Plugins/FX/Niagara/Source/Niagara/Private/NiagaraDebugHud.cpp
--- a/Engine/Plugins/FX/Niagara/Source/Niagara/Private/NiagaraDebugHud.cpp
+++ b/Engine/Plugins/FX/Niagara/Source/Niagara/Private/NiagaraDebugHud.cpp
@@ -438,6 +438,17 @@ namespace NiagaraDebugLocal
 	#endif
 	}
 
+	FString FormatPerfValue(double dValue, int32 Length = 7)
+	{
+		const int32 Value = int32(dValue);
+		FString TempString = FString::FormatAsNumber(Value);
+		while (TempString.Len() < Length)
+		{
+			TempString.AppendChar(' ');
+		}
+		return TempString;
+	}
+
 	const FCachedVariables& GetCachedVariables(UNiagaraSystem* NiagaraSystem)
 	{
 		FCachedVariables* CachedVariables = GCachedSystemVariables.Find(NiagaraSystem);
@@ -953,7 +964,7 @@ void FNiagaraDebugHud::GatherSystemInfo()
 
 #if WITH_PARTICLE_PERF_STATS
 	bool bUpdateStats = false;
-	if (Settings.bOverviewEnabled && (Settings.OverviewMode == ENiagaraDebugHUDOverviewMode::Performance))
+	if (Settings.bOverviewEnabled && (Settings.OverviewMode == ENiagaraDebugHUDOverviewMode::Performance || Settings.OverviewMode == ENiagaraDebugHUDOverviewMode::PerformanceGraph))
 	{
 		if (StatsListener.IsValid() == false)
 		{
@@ -1151,7 +1162,7 @@ void FNiagaraDebugHud::GatherSystemInfo()
 
 		bool bWillBeVisible = false;
 		
-		if(Settings.OverviewMode == ENiagaraDebugHUDOverviewMode::Performance)
+		if (Settings.OverviewMode == ENiagaraDebugHUDOverviewMode::Performance || Settings.OverviewMode == ENiagaraDebugHUDOverviewMode::PerformanceGraph)
 		{
 			bWillBeVisible = SystemDebugInfo.TotalActive > 0;
 		}
@@ -1529,7 +1540,7 @@ void FNiagaraDebugHud::DrawOverview(class FNiagaraWorldManager* WorldManager, FC
 		typedef TFunction<void(FCanvas*, UFont*, float, float, FOverviewColumn&, const FSystemDebugInfo&)>  FSystemDrawFunc;
 		FSystemDrawFunc SystemDrawFunc;
 
-		FOverviewColumn(FString InGlobalHeader, FString InGlobalData, FString InSystemHeader, float& InOffset, UFont* Font, FString ExampleSystemString, FSystemDrawFunc InSystemDrawFunc)
+		FOverviewColumn(FString InGlobalHeader, FString InGlobalData, FString InSystemHeader, float& InOffset, UFont* Font, const TCHAR* ExampleSystemString, FSystemDrawFunc InSystemDrawFunc)
 			: GlobalHeader(InGlobalHeader)
 			, GlobalData(InGlobalData)
 			, SystemHeader(InSystemHeader)
@@ -1537,12 +1548,28 @@ void FNiagaraDebugHud::DrawOverview(class FNiagaraWorldManager* WorldManager, FC
 			, SystemDrawFunc(InSystemDrawFunc)
 		{
 			HeaderWidth = float(GetStringSize(Font, *GlobalHeader).X);
-			SystemWidth = float(FMath::Max(GetStringSize(Font, *SystemHeader).X, GetStringSize(Font, *ExampleSystemString).X));
+			SystemWidth = float(FMath::Max(GetStringSize(Font, *SystemHeader).X, GetStringSize(Font, ExampleSystemString).X));
 			MaxWidth = float(FMath::Max(HeaderWidth + GlobalDataSeparator + GetStringSize(Font, *GlobalData).X, SystemWidth)) + ColumnSeparator;
 
 			InOffset += MaxWidth;
 		}
 
+		FOverviewColumn(FString InGlobalHeader, FString InGlobalData, int32 GlobalDataStringSize, FString InSystemHeader, float& InOffset, UFont* Font, int32 SystemStringSize, FSystemDrawFunc InSystemDrawFunc)
+			: GlobalHeader(InGlobalHeader)
+			, GlobalData(InGlobalData)
+			, SystemHeader(InSystemHeader)
+			, Offset(InOffset)
+			, SystemDrawFunc(InSystemDrawFunc)
+		{
+			GlobalDataStringSize = FMath::Max(GetStringSize(Font, *GlobalData).X, GlobalDataStringSize);
+
+			HeaderWidth = float(GetStringSize(Font, *GlobalHeader).X);
+			SystemWidth = float(FMath::Max(GetStringSize(Font, *SystemHeader).X, SystemStringSize));
+			MaxWidth = float(FMath::Max(HeaderWidth + GlobalDataSeparator + GlobalDataStringSize, SystemWidth)) + ColumnSeparator;
+
+			InOffset += MaxWidth;
+		}
+
 		void DrawGlobalHeader(FCanvas* DrawCanvas, UFont* Font, float X, float Y, FLinearColor HeadingColor, FLinearColor DetailColor)
 		{
 			DrawCanvas->DrawShadowedString(X + Offset, Y, *GlobalHeader, Font, HeadingColor);
@@ -1735,79 +1762,79 @@ void FNiagaraDebugHud::DrawOverview(class FNiagaraWorldManager* WorldManager, FC
 		{
 			FNiagaraDebugHUDPerfStats& GlobalPerfStats = StatsListener->GetGlobalStats();
 
-			if (Settings.PerfGraphMode == ENiagaraDebugHUDPerfGraphMode::None)
-			{
-				OverviewColumns.Emplace(TEXT("Game Thread Avg:"), LexToSanitizedString(GlobalPerfStats.Avg.Time_GT), TEXT("GT Avg (us)"), ColumnOffset, Font, TEXT("000.0000"),
-					[&DetailColor, &DetailHighlightColor, &fAdvanceHeight](FCanvas* Canvas, UFont* Font, float X, float Y, FOverviewColumn& Col, const FSystemDebugInfo& SystemInfo)
-					{
-						FLinearColor RowBGColor = SystemInfo.UniqueColor;
-						RowBGColor.A = Settings.SystemColorTableOpacity;
-						Canvas->DrawTile(X, Y, Col.MaxWidth, fAdvanceHeight, 0, 0, 0, 0, RowBGColor);
-						const FLinearColor RowColor = SystemInfo.bShowInWorld ? DetailHighlightColor : DetailColor;
-						Canvas->DrawShadowedString(X, Y, *LexToSanitizedString(SystemInfo.PerfStats ? SystemInfo.PerfStats->Avg.Time_GT : 0.0), Font, RowColor);
-					});
-
-				OverviewColumns.Emplace(TEXT("Game Thread Max:"), LexToSanitizedString(GlobalPerfStats.Max.Time_GT), TEXT("GT Max (us)"), ColumnOffset, Font, TEXT("000.0000"),
-					[&DetailColor, &DetailHighlightColor, &fAdvanceHeight](FCanvas* Canvas, UFont* Font, float X, float Y, FOverviewColumn& Col, const FSystemDebugInfo& SystemInfo)
-					{
-						FLinearColor RowBGColor = SystemInfo.UniqueColor;
-						RowBGColor.A = Settings.SystemColorTableOpacity;
-						Canvas->DrawTile(X, Y, Col.MaxWidth, fAdvanceHeight, 0, 0, 0, 0, RowBGColor);
-						const FLinearColor RowColor = SystemInfo.bShowInWorld ? DetailHighlightColor : DetailColor;
-						Canvas->DrawShadowedString(X, Y, *LexToSanitizedString(SystemInfo.PerfStats ? SystemInfo.PerfStats->Max.Time_GT : 0.0), Font, RowColor);
-					});
-
-				OverviewColumns.Emplace_GetRef(TEXT("Render Thread Avg:"), LexToSanitizedString(GlobalPerfStats.Avg.Time_RT), TEXT("RT Avg (us)"), ColumnOffset, Font, TEXT("000.0000"),
-					[&DetailColor, &DetailHighlightColor, &fAdvanceHeight](FCanvas* Canvas, UFont* Font, float X, float Y, FOverviewColumn& Col, const FSystemDebugInfo& SystemInfo)
-					{
-						FLinearColor RowBGColor = SystemInfo.UniqueColor;
-						RowBGColor.A = Settings.SystemColorTableOpacity;
-						Canvas->DrawTile(X, Y, Col.MaxWidth, fAdvanceHeight, 0, 0, 0, 0, RowBGColor);
-						const FLinearColor RowColor = SystemInfo.bShowInWorld ? DetailHighlightColor : DetailColor;
-						Canvas->DrawShadowedString(X, Y, *LexToSanitizedString(SystemInfo.PerfStats ? SystemInfo.PerfStats->Avg.Time_RT : 0.0), Font, RowColor);
-					});
-
-				OverviewColumns.Emplace(TEXT("Render Thread Max:"), LexToSanitizedString(GlobalPerfStats.Max.Time_RT), TEXT("RT Max (us)"), ColumnOffset, Font, TEXT("000.0000"),
-					[&DetailColor, &DetailHighlightColor, &fAdvanceHeight](FCanvas* Canvas, UFont* Font, float X, float Y, FOverviewColumn& Col, const FSystemDebugInfo& SystemInfo)
-					{
-						FLinearColor RowBGColor = SystemInfo.UniqueColor;
-						RowBGColor.A = Settings.SystemColorTableOpacity;
-						Canvas->DrawTile(X, Y, Col.MaxWidth, fAdvanceHeight, 0, 0, 0, 0, RowBGColor);
-						const FLinearColor RowColor = SystemInfo.bShowInWorld ? DetailHighlightColor : DetailColor;
-						Canvas->DrawShadowedString(X, Y, *LexToSanitizedString(SystemInfo.PerfStats ? SystemInfo.PerfStats->Max.Time_RT : 0.0), Font, RowColor);
-					});
-
-				OverviewColumns.Emplace_GetRef(TEXT("Gpu Avg:"), LexToSanitizedString(GlobalPerfStats.Avg.Time_GPU), TEXT("Gpu Avg (us)"), ColumnOffset, Font, TEXT("000.0000"),
-					[&DetailColor, &DetailHighlightColor, &fAdvanceHeight](FCanvas* Canvas, UFont* Font, float X, float Y, FOverviewColumn& Col, const FSystemDebugInfo& SystemInfo)
-					{
-						FLinearColor RowBGColor = SystemInfo.UniqueColor;
-						RowBGColor.A = Settings.SystemColorTableOpacity;
-						Canvas->DrawTile(X, Y, Col.MaxWidth, fAdvanceHeight, 0, 0, 0, 0, RowBGColor);
-						const FLinearColor RowColor = SystemInfo.bShowInWorld ? DetailHighlightColor : DetailColor;
-						Canvas->DrawShadowedString(X, Y, *LexToSanitizedString(SystemInfo.PerfStats ? SystemInfo.PerfStats->Avg.Time_GPU : 0.0), Font, RowColor);
-					});
-
-				OverviewColumns.Emplace(TEXT("Gpu Max:"), LexToSanitizedString(GlobalPerfStats.Max.Time_GPU), TEXT("Gpu Max (us)"), ColumnOffset, Font, TEXT("000.0000"),
-					[&DetailColor, &DetailHighlightColor, &fAdvanceHeight](FCanvas* Canvas, UFont* Font, float X, float Y, FOverviewColumn& Col, const FSystemDebugInfo& SystemInfo)
-					{
-						FLinearColor RowBGColor = SystemInfo.UniqueColor;
-						RowBGColor.A = Settings.SystemColorTableOpacity;
-						Canvas->DrawTile(X, Y, Col.MaxWidth, fAdvanceHeight, 0, 0, 0, 0, RowBGColor);
-						const FLinearColor RowColor = SystemInfo.bShowInWorld ? DetailHighlightColor : DetailColor;
-						Canvas->DrawShadowedString(X, Y, *LexToSanitizedString(SystemInfo.PerfStats ? SystemInfo.PerfStats->Max.Time_GPU : 0.0), Font, RowColor);
-					});
-			}
-			else
-			{
-				OverviewColumns.Emplace(TEXT(""), TEXT(""), TEXT(""), ColumnOffset, Font, TEXT("------"),
-					[&DetailColor, &DetailHighlightColor, &fAdvanceHeight](FCanvas* Canvas, UFont* Font, float X, float Y, FOverviewColumn& Col, const FSystemDebugInfo& SystemInfo)
-					{
-						Canvas->DrawTile(X, Y, Col.MaxWidth - 6.0f, fAdvanceHeight - 2.0f, 0.0f, 0.0f, 0.0f, 0.0f, SystemInfo.UniqueColor);
-						FString SysCountString = LexToSanitizedString(SystemInfo.TotalActive);
-						float StringWidth = float(Font->GetStringSize(*SysCountString));
-						Canvas->DrawShadowedString((X + (Col.MaxWidth * 0.5f)) - StringWidth * 0.5f, Y, *SysCountString, Font, FLinearColor::Black);
-					});
-			}
+			const int32 GlobalDataStringSize = GetStringSize(Font, TEXT("000,000")).X;
+			const int32 SystemStringSize = GetStringSize(Font, TEXT("000,000")).X;
+
+			OverviewColumns.Emplace(TEXT("Game Thread Avg:"), FormatPerfValue(GlobalPerfStats.Avg.Time_GT), GlobalDataStringSize, TEXT("GT Avg (us)"), ColumnOffset, Font, SystemStringSize,
+				[&DetailColor, &DetailHighlightColor, &fAdvanceHeight](FCanvas* Canvas, UFont* Font, float X, float Y, FOverviewColumn& Col, const FSystemDebugInfo& SystemInfo)
+				{
+					FLinearColor RowBGColor = SystemInfo.UniqueColor;
+					RowBGColor.A = Settings.SystemColorTableOpacity;
+					Canvas->DrawTile(X, Y, Col.MaxWidth, fAdvanceHeight, 0, 0, 0, 0, RowBGColor);
+					const FLinearColor RowColor = SystemInfo.bShowInWorld ? DetailHighlightColor : DetailColor;
+					Canvas->DrawShadowedString(X, Y, *FormatPerfValue(SystemInfo.PerfStats ? SystemInfo.PerfStats->Avg.Time_GT : 0.0), Font, RowColor);
+				});
+
+			OverviewColumns.Emplace(TEXT("Game Thread Max:"), FormatPerfValue(GlobalPerfStats.Max.Time_GT), GlobalDataStringSize, TEXT("GT Max (us)"), ColumnOffset, Font, SystemStringSize,
+				[&DetailColor, &DetailHighlightColor, &fAdvanceHeight](FCanvas* Canvas, UFont* Font, float X, float Y, FOverviewColumn& Col, const FSystemDebugInfo& SystemInfo)
+				{
+					FLinearColor RowBGColor = SystemInfo.UniqueColor;
+					RowBGColor.A = Settings.SystemColorTableOpacity;
+					Canvas->DrawTile(X, Y, Col.MaxWidth, fAdvanceHeight, 0, 0, 0, 0, RowBGColor);
+					const FLinearColor RowColor = SystemInfo.bShowInWorld ? DetailHighlightColor : DetailColor;
+					Canvas->DrawShadowedString(X, Y, *FormatPerfValue(SystemInfo.PerfStats ? SystemInfo.PerfStats->Max.Time_GT : 0.0), Font, RowColor);
+				});
+
+			OverviewColumns.Emplace_GetRef(TEXT("Render Thread Avg:"), FormatPerfValue(GlobalPerfStats.Avg.Time_RT), GlobalDataStringSize, TEXT("RT Avg (us)"), ColumnOffset, Font, SystemStringSize,
+				[&DetailColor, &DetailHighlightColor, &fAdvanceHeight](FCanvas* Canvas, UFont* Font, float X, float Y, FOverviewColumn& Col, const FSystemDebugInfo& SystemInfo)
+				{
+					FLinearColor RowBGColor = SystemInfo.UniqueColor;
+					RowBGColor.A = Settings.SystemColorTableOpacity;
+					Canvas->DrawTile(X, Y, Col.MaxWidth, fAdvanceHeight, 0, 0, 0, 0, RowBGColor);
+					const FLinearColor RowColor = SystemInfo.bShowInWorld ? DetailHighlightColor : DetailColor;
+					Canvas->DrawShadowedString(X, Y, *FormatPerfValue(SystemInfo.PerfStats ? SystemInfo.PerfStats->Avg.Time_RT : 0.0), Font, RowColor);
+				});
+
+			OverviewColumns.Emplace(TEXT("Render Thread Max:"), FormatPerfValue(GlobalPerfStats.Max.Time_RT), GlobalDataStringSize, TEXT("RT Max (us)"), ColumnOffset, Font, SystemStringSize,
+				[&DetailColor, &DetailHighlightColor, &fAdvanceHeight](FCanvas* Canvas, UFont* Font, float X, float Y, FOverviewColumn& Col, const FSystemDebugInfo& SystemInfo)
+				{
+					FLinearColor RowBGColor = SystemInfo.UniqueColor;
+					RowBGColor.A = Settings.SystemColorTableOpacity;
+					Canvas->DrawTile(X, Y, Col.MaxWidth, fAdvanceHeight, 0, 0, 0, 0, RowBGColor);
+					const FLinearColor RowColor = SystemInfo.bShowInWorld ? DetailHighlightColor : DetailColor;
+					Canvas->DrawShadowedString(X, Y, *FormatPerfValue(SystemInfo.PerfStats ? SystemInfo.PerfStats->Max.Time_RT : 0.0), Font, RowColor);
+				});
+
+			OverviewColumns.Emplace_GetRef(TEXT("Gpu Avg:"), FormatPerfValue(GlobalPerfStats.Avg.Time_GPU), GlobalDataStringSize, TEXT("Gpu Avg (us)"), ColumnOffset, Font, SystemStringSize,
+				[&DetailColor, &DetailHighlightColor, &fAdvanceHeight](FCanvas* Canvas, UFont* Font, float X, float Y, FOverviewColumn& Col, const FSystemDebugInfo& SystemInfo)
+				{
+					FLinearColor RowBGColor = SystemInfo.UniqueColor;
+					RowBGColor.A = Settings.SystemColorTableOpacity;
+					Canvas->DrawTile(X, Y, Col.MaxWidth, fAdvanceHeight, 0, 0, 0, 0, RowBGColor);
+					const FLinearColor RowColor = SystemInfo.bShowInWorld ? DetailHighlightColor : DetailColor;
+					Canvas->DrawShadowedString(X, Y, *FormatPerfValue(SystemInfo.PerfStats ? SystemInfo.PerfStats->Avg.Time_GPU : 0.0), Font, RowColor);
+				});
+
+			OverviewColumns.Emplace(TEXT("Gpu Max:"), FormatPerfValue(GlobalPerfStats.Max.Time_GPU), GlobalDataStringSize, TEXT("Gpu Max (us)"), ColumnOffset, Font, SystemStringSize,
+				[&DetailColor, &DetailHighlightColor, &fAdvanceHeight](FCanvas* Canvas, UFont* Font, float X, float Y, FOverviewColumn& Col, const FSystemDebugInfo& SystemInfo)
+				{
+					FLinearColor RowBGColor = SystemInfo.UniqueColor;
+					RowBGColor.A = Settings.SystemColorTableOpacity;
+					Canvas->DrawTile(X, Y, Col.MaxWidth, fAdvanceHeight, 0, 0, 0, 0, RowBGColor);
+					const FLinearColor RowColor = SystemInfo.bShowInWorld ? DetailHighlightColor : DetailColor;
+					Canvas->DrawShadowedString(X, Y, *FormatPerfValue(SystemInfo.PerfStats ? SystemInfo.PerfStats->Max.Time_GPU : 0.0), Font, RowColor);
+				});
+		}
+		else if (Settings.OverviewMode == ENiagaraDebugHUDOverviewMode::PerformanceGraph && StatsListener)
+		{
+			OverviewColumns.Emplace(TEXT(""), TEXT(""), TEXT(""), ColumnOffset, Font, TEXT("------"),
+				[&DetailColor, &DetailHighlightColor, &fAdvanceHeight](FCanvas* Canvas, UFont* Font, float X, float Y, FOverviewColumn& Col, const FSystemDebugInfo& SystemInfo)
+				{
+					Canvas->DrawTile(X, Y, Col.MaxWidth - 6.0f, fAdvanceHeight - 2.0f, 0.0f, 0.0f, 0.0f, 0.0f, SystemInfo.UniqueColor);
+					FString SysCountString = LexToSanitizedString(SystemInfo.TotalActive);
+					float StringWidth = float(Font->GetStringSize(*SysCountString));
+					Canvas->DrawShadowedString((X + (Col.MaxWidth * 0.5f)) - StringWidth * 0.5f, Y, *SysCountString, Font, FLinearColor::Black);
+				});
 		}
 #endif//WITH_PARTICLE_PERF_STATS
 
@@ -1972,71 +1999,67 @@ void FNiagaraDebugHud::DrawOverview(class FNiagaraWorldManager* WorldManager, FC
 		
 		#if WITH_PARTICLE_PERF_STATS
 		//Draw graph
-		if (Settings.OverviewMode == ENiagaraDebugHUDOverviewMode::Performance && StatsListener)
+		if (Settings.OverviewMode == ENiagaraDebugHUDOverviewMode::PerformanceGraph && StatsListener)
 		{
-			//Render the graph
-			if (Settings.PerfGraphMode != ENiagaraDebugHUDPerfGraphMode::None)
-			{
-				FLinearColor GraphAxesColor = Settings.PerfGraphAxisColor;
+			FLinearColor GraphAxesColor = Settings.PerfGraphAxisColor;
 
-				TextLocation.Y += fAdvanceHeight;
+			TextLocation.Y += fAdvanceHeight;
 
-				FVector2f GraphLocation(TextLocation.X, SystemDataY + fAdvanceHeight);
+			FVector2f GraphLocation(TextLocation.X, SystemDataY + fAdvanceHeight);
 
-				if (OverviewColumns.Num())
-				{
-					GraphLocation.X += OverviewColumns.Last().Offset + OverviewColumns.Last().MaxWidth + 50.0f;
-				}
+			if (OverviewColumns.Num())
+			{
+				GraphLocation.X += OverviewColumns.Last().Offset + OverviewColumns.Last().MaxWidth + 50.0f;
+			}
 
-				FGraph<double> Graph(DrawCanvas, Font, GraphLocation, FVector2f(Settings.PerfGraphSize), FVector2f(float(Settings.PerfHistoryFrames), Settings.PerfGraphTimeRange), TEXT("Frame"), TEXT("Time(us)"));
+			FGraph<double> Graph(DrawCanvas, Font, GraphLocation, FVector2f(Settings.PerfGraphSize), FVector2f(float(Settings.PerfHistoryFrames), Settings.PerfGraphTimeRange), TEXT("Frame"), TEXT("Time(us)"));
 
-				Graph.Draw(Settings.PerfGraphAxisColor, BackgroundColor);
+			Graph.Draw(Settings.PerfGraphAxisColor, BackgroundColor);
 
-				//Add each line to the graph.
-				for (auto& Pair : PerSystemDebugInfo)
+			//Add each line to the graph.
+			for (auto& Pair : PerSystemDebugInfo)
+			{
+				FSystemDebugInfo& SysInfo = Pair.Value;
+				if (SysInfo.PerfStats && SysInfo.bPassesSystemFilter)
 				{
-					FSystemDebugInfo& SysInfo = Pair.Value;
-					if (SysInfo.PerfStats && SysInfo.bPassesSystemFilter)
+					TArray<double> Frames;
+					if (Settings.PerfGraphMode == ENiagaraDebugHUDPerfGraphMode::GameThread)
 					{
-						TArray<double> Frames;
-						if (Settings.PerfGraphMode == ENiagaraDebugHUDPerfGraphMode::GameThread)
-						{
-							SysInfo.PerfStats->History.GetHistoryFrames_GT(Frames);
-						}
-						else if (Settings.PerfGraphMode == ENiagaraDebugHUDPerfGraphMode::RenderThread)
-						{
-							SysInfo.PerfStats->History.GetHistoryFrames_RT(Frames);
-						}
-						else if (Settings.PerfGraphMode == ENiagaraDebugHUDPerfGraphMode::GPU)
-						{
-							SysInfo.PerfStats->History.GetHistoryFrames_GPU(Frames);
-						}
+						SysInfo.PerfStats->History.GetHistoryFrames_GT(Frames);
+					}
+					else if (Settings.PerfGraphMode == ENiagaraDebugHUDPerfGraphMode::RenderThread)
+					{
+						SysInfo.PerfStats->History.GetHistoryFrames_RT(Frames);
+					}
+					else if (Settings.PerfGraphMode == ENiagaraDebugHUDPerfGraphMode::GPU)
+					{
+						SysInfo.PerfStats->History.GetHistoryFrames_GPU(Frames);
+					}
 
-						if (Settings.bEnableSmoothing)
+					if (Settings.bEnableSmoothing)
+					{
+						TArray<double> Smoothed;
+						Smoothed.Reserve(Frames.Num());
+						for (int32 i = 0; i < Frames.Num(); ++i)
 						{
-							TArray<double> Smoothed;
-							Smoothed.Reserve(Frames.Num());
-							for (int32 i = 0; i < Frames.Num(); ++i)
+							int32 SmoothSamples = 0;
+							double SmoothTotal = 0.0f;
+							for (int32 j = i - Settings.SmoothingWidth; j < i + Settings.SmoothingWidth; ++j)
 							{
-								int32 SmoothSamples = 0;
-								double SmoothTotal = 0.0f;
-								for (int32 j = i - Settings.SmoothingWidth; j < i + Settings.SmoothingWidth; ++j)
+								if (j >= 0 && j < Frames.Num())
 								{
-									if (j >= 0 && j < Frames.Num())
-									{
-										++SmoothSamples;
-										SmoothTotal += Frames[j];
-									}
+									++SmoothSamples;
+									SmoothTotal += Frames[j];
 								}
-
-								Smoothed.Add(SmoothTotal / SmoothSamples);
 							}
 
-							Frames = MoveTemp(Smoothed);
+							Smoothed.Add(SmoothTotal / SmoothSamples);
 						}
 
-						Graph.DrawLine(SysInfo.SystemName, SysInfo.UniqueColor, Frames);
+						Frames = MoveTemp(Smoothed);
 					}
+
+					Graph.DrawLine(SysInfo.SystemName, SysInfo.UniqueColor, Frames);
 				}
 			}
 		}
diff --git a/Engine/Plugins/FX/Niagara/Source/Niagara/Public/NiagaraDebuggerCommon.h b/Engine/Plugins/FX/Niagara/Source/Niagara/Public/NiagaraDebuggerCommon.h
--- a/Engine/Plugins/FX/Niagara/Source/Niagara/Public/NiagaraDebuggerCommon.h
+++ b/Engine/Plugins/FX/Niagara/Source/Niagara/Public/NiagaraDebuggerCommon.h
@@ -387,13 +387,13 @@ enum class ENiagaraDebugHUDOverviewMode
 	Overview,
 	Scalability,	
 	Performance,
+	PerformanceGraph,
 	GpuComputePerformance,
 }; 
 
 UENUM()
 enum class ENiagaraDebugHUDPerfGraphMode
 {
-	None,
 	GameThread,
 	RenderThread,
 	GPU,
