Remove checkSlow in FNiagaraParameterStore::GetDataInterface() as returning null could be legitimate.

- UNiagaraComponent::FixInvalidUserParameterOverrideData() during PostLoad calls and fixes up invalid data interfaces, hitting this checkSlow in Debug builds.

diff --git a/Engine/Plugins/FX/Niagara/Source/Niagara/Private/NiagaraParameterStore.cpp b/Engine/Plugins/FX/Niagara/Source/Niagara/Private/NiagaraParameterStore.cpp
--- a/Engine/Plugins/FX/Niagara/Source/Niagara/Private/NiagaraParameterStore.cpp
+++ b/Engine/Plugins/FX/Niagara/Source/Niagara/Private/NiagaraParameterStore.cpp
@@ -1380,7 +1380,6 @@ UNiagaraDataInterface* FNiagaraParameterStore::GetDataInterface(const FNiagaraVa
 {
 	int32 Offset = IndexOf(Parameter);
 	UNiagaraDataInterface* Interface = GetDataInterface(Offset);
-	checkSlow(!Interface || Parameter.GetType().GetClass() == Interface->GetClass());
 	return Interface;
 }
 
