- Fix for MIC rename causing name conflicts
- Only allow a rename if it's safe and attempt to find one by name first

diff --git a/Engine/Plugins/FX/Niagara/Source/Niagara/Private/NiagaraRendererProperties.cpp b/Engine/Plugins/FX/Niagara/Source/Niagara/Private/NiagaraRendererProperties.cpp
--- a/Engine/Plugins/FX/Niagara/Source/Niagara/Private/NiagaraRendererProperties.cpp
+++ b/Engine/Plugins/FX/Niagara/Source/Niagara/Private/NiagaraRendererProperties.cpp
@@ -545,12 +545,23 @@ void UNiagaraRendererProperties::UpdateMaterialParametersMIC(const FNiagaraRende
 		NameBuilder.Append(TEXT("_MIC"));
 
 		UMaterialInstanceConstant* MIC = nullptr;
+		bool bNeedsRename = false;
 		if (MICPool.Num() > 0)
 		{
-			MIC = MICPool.Pop();
+			FName MICName(NameBuilder);
+			const int32 ExistingIndex = MICPool.IndexOfByPredicate([&MICName](UMaterialInstanceConstant* MIC) { return MIC->GetFName() == MICName; });
+			if (ExistingIndex != INDEX_NONE)
+			{
+				MIC = MICPool[ExistingIndex];
+				MICPool.RemoveAtSwap(ExistingIndex, 1, false);
+			}
+			else
+			{
+				bNeedsRename = true;
+				MIC = MICPool.Pop();
+			}
 			if (MIC->Parent != Material)
 			{
-				MIC->Rename(NameBuilder.ToString());
 				MIC->SetParentEditorOnly(Material);
 			}
 		}
@@ -568,8 +579,21 @@ void UNiagaraRendererProperties::UpdateMaterialParametersMIC(const FNiagaraRende
 		{
 			InOutMICs.SetNum(i + 1);
 			InOutMICs[i] = MIC;
+			if (bNeedsRename)
+			{
+				if (MIC->Rename(NameBuilder.ToString(), nullptr, REN_Test))
+				{
+					MIC->Rename(NameBuilder.ToString());
+				}
+			}
 		}
 	}
+
+	// Garabge the MICs we are no longer going to use
+	for (int i=0; i < MICPool.Num(); ++i)
+	{
+		MICPool[i]->MarkAsGarbage();
+	}
 }
 
 int32 UNiagaraRendererProperties::GetDynamicParameterChannelMask(const FVersionedNiagaraEmitterData* EmitterData, FName BindingName, int32 DefaultChannelMask) const
