MGPU:  Fixes for Niagara GPU particle sim cross-GPU serialization issues, potentially saving several ms in performance, and solving artifacts.   Implementation now runs without D3D12 validation errors on Valley test scene, including fixes for previously existing errors for nDisplay render target transfers, plus incorrect resource state and usage overlap errors for Niagara GPU particle systems.  A new cross-GPU transfer API was added to RHI with explicit lists of fences, rather than attaching fences to individual transfers, offering more flexibility. The new API properly handles transition states for destination resources of a copy, by transitioning the relevant destination resources before the fence.  The plan is to eventually deprecate the original API.

On the Niagara side, some refactoring was required to make it so extra nDisplay views running on other GPUs don't need to wait on data generated by the first view.  Niagara renders particle systems delayed by a frame by default, so additional views need to continue to render with the previous frame's data, rather than the data generated on the current frame.  Swaps of the current data to render are deferred in multi-view to accomplish this.  Non-delayed (low latency translucent / post opaque) particles can use the current frame's data on all GPUs.  To accomplish this, cross GPU transfers always occur after the post opaque stage, and are waited on at the post opaque stage on other views, to minimize the likelihood of needing to wait (given transfer and wait occur at the same place in Render).

An edge case needed to be fixed, where count clearing (once per frame housekeeping) needs to only occur for the first view, so that process doesn't run when the second view renders, which is frequently on a different GPU.  This avoids an expensive round trip of data back to the first GPU, and solves visual corruption bugs in previous versions of the code when switching GPU assignments with multi-view active, due to inadequate synchronization of the extra transfers generated by the clears.  Count clearing is disabled for views that don't run simulation (such as scene captures), and for views other than the first view with multi-view.

#rnx

diff --git a/Engine/Plugins/FX/Niagara/Source/Niagara/Classes/NiagaraComputeExecutionContext.h b/Engine/Plugins/FX/Niagara/Source/Niagara/Classes/NiagaraComputeExecutionContext.h
--- a/Engine/Plugins/FX/Niagara/Source/Niagara/Classes/NiagaraComputeExecutionContext.h
+++ b/Engine/Plugins/FX/Niagara/Source/Niagara/Classes/NiagaraComputeExecutionContext.h
@@ -72,8 +72,21 @@ struct FNiagaraComputeExecutionContext
 
 	void SetDataToRender(FNiagaraDataBuffer* InDataToRender);
 	void SetTranslucentDataToRender(FNiagaraDataBuffer* InTranslucentDataToRender);
+	void SetMultiViewPreviousDataToRender(FNiagaraDataBuffer* InMultiViewPreviousDataToRender);
 	bool HasTranslucentDataToRender() const { return TranslucentDataToRender != nullptr; }	
-	FNiagaraDataBuffer* GetDataToRender(bool bIsLowLatencyTranslucent) const { return bIsLowLatencyTranslucent && HasTranslucentDataToRender() ? TranslucentDataToRender : DataToRender; }
+	FNiagaraDataBuffer* GetDataToRender(bool bIsLowLatencyTranslucent) const
+	{
+		if (bIsLowLatencyTranslucent)
+		{
+			// Translucent rendering uses current frame (newest) data, which is TranslucentDataToRender if present
+			return TranslucentDataToRender ? TranslucentDataToRender : DataToRender;
+		}
+		else
+		{
+			// Non-translucent rendering uses previous frame (older) data, which is MultiViewPreviousDataToRender if present
+			return MultiViewPreviousDataToRender ? MultiViewPreviousDataToRender : DataToRender;
+		}
+	}
 
 	struct 
 	{
@@ -140,6 +153,9 @@ public:
 	// Optional buffer which can be used to render translucent data with no latency (i.e. this frames data)
 	FNiagaraDataBufferRef TranslucentDataToRender = nullptr;
 
+	// For MultiView, we need to track the previous frame's data for rendering systems on other view families
+	FNiagaraDataBufferRef MultiViewPreviousDataToRender = nullptr;
+
 	// Game thread spawn info will be sent to the render thread inside FNiagaraComputeInstanceData
 	FNiagaraGpuSpawnInfo GpuSpawnInfo_GT;
 
diff --git a/Engine/Plugins/FX/Niagara/Source/Niagara/Classes/NiagaraGPUInstanceCountManager.h b/Engine/Plugins/FX/Niagara/Source/Niagara/Classes/NiagaraGPUInstanceCountManager.h
--- a/Engine/Plugins/FX/Niagara/Source/Niagara/Classes/NiagaraGPUInstanceCountManager.h
+++ b/Engine/Plugins/FX/Niagara/Source/Niagara/Classes/NiagaraGPUInstanceCountManager.h
@@ -113,6 +113,8 @@ public:
 
 	bool HasEntriesPendingFree() const { return InstanceCountClearTasks.Num() > 0; }
 
+	void CopyToMultiViewCountBuffer(FRHICommandListImmediate& RHICmdList);
+
 protected:
 	struct FIndirectArgsPoolEntry
 	{
@@ -150,6 +152,10 @@ protected:
 	FRHIGPUMemoryReadback* CountReadback = nullptr;
 	int32 CountReadbackSize = 0;
 
+	/** Mirrored copy of CountBuffer for multi-view rendering. */
+	int32 MultiViewAllocatedInstanceCounts = 0;
+	FRWBuffer MultiViewCountBuffer;
+
 	/** The list of all draw indirected tasks that are to be run in UpdateDrawIndirectBuffers() */
 	TArray<FNiagaraDrawIndirectArgGenTaskInfo> DrawIndirectArgGenTasks[ENiagaraGPUCountUpdatePhase::Max];
 
diff --git a/Engine/Plugins/FX/Niagara/Source/Niagara/Private/NiagaraComputeExecutionContext.cpp b/Engine/Plugins/FX/Niagara/Source/Niagara/Private/NiagaraComputeExecutionContext.cpp
--- a/Engine/Plugins/FX/Niagara/Source/Niagara/Private/NiagaraComputeExecutionContext.cpp
+++ b/Engine/Plugins/FX/Niagara/Source/Niagara/Private/NiagaraComputeExecutionContext.cpp
@@ -251,6 +251,11 @@ void FNiagaraComputeExecutionContext::SetTranslucentDataToRender(FNiagaraDataBuf
 	TranslucentDataToRender = InTranslucentDataToRender;
 }
 
+void FNiagaraComputeExecutionContext::SetMultiViewPreviousDataToRender(FNiagaraDataBuffer* InMultiViewPreviousDataToRender)
+{
+	MultiViewPreviousDataToRender = InMultiViewPreviousDataToRender;
+}
+
 bool FNiagaraComputeInstanceData::IsOutputStage(FNiagaraDataInterfaceProxy* DIProxy, uint32 SimulationStageIndex) const
 {
 	return Context->IsOutputStage(DIProxy, SimulationStageIndex);
diff --git a/Engine/Plugins/FX/Niagara/Source/Niagara/Private/NiagaraGPUInstanceCountManager.cpp b/Engine/Plugins/FX/Niagara/Source/Niagara/Private/NiagaraGPUInstanceCountManager.cpp
--- a/Engine/Plugins/FX/Niagara/Source/Niagara/Private/NiagaraGPUInstanceCountManager.cpp
+++ b/Engine/Plugins/FX/Niagara/Source/Niagara/Private/NiagaraGPUInstanceCountManager.cpp
@@ -133,6 +133,7 @@ void FNiagaraGPUInstanceCountManager::ReleaseCounts()
 {
 	CountBuffer.Release();
 	CulledCountBuffer.Release();
+	MultiViewCountBuffer.Release();
 
 	AllocatedInstanceCounts = 0;
 	AllocatedCulledCounts = 0;
@@ -396,7 +397,7 @@ void FNiagaraGPUInstanceCountManager::UpdateDrawIndirectBuffers(FNiagaraGpuCompu
 {
 	// Anything to process?
 	TArray<FNiagaraDrawIndirectArgGenTaskInfo>& ArgTasks = DrawIndirectArgGenTasks[CountPhase];
-	const bool bClearCounts = (CountPhase == ENiagaraGPUCountUpdatePhase::PreOpaque) && (InstanceCountClearTasks.Num() > 0);
+	const bool bClearCounts = (CountPhase == ENiagaraGPUCountUpdatePhase::PreOpaque) && (InstanceCountClearTasks.Num() > 0) && ComputeDispatchInterface->IsFirstViewFamily();
 	if ( (ArgTasks.Num() > 0 || bClearCounts) && FNiagaraUtilities::AllowComputeShaders(GShaderPlatformForFeatureLevel[FeatureLevel]) )
 	{
 		INC_DWORD_STAT_BY(STAT_NiagaraIndirectDraws, ArgTasks.Num());
@@ -407,7 +408,7 @@ void FNiagaraGPUInstanceCountManager::UpdateDrawIndirectBuffers(FNiagaraGpuCompu
 		FReadBuffer TaskInfosBuffer;
 		{
 			const uint32 ArgGenSize = ArgTasks.Num() * sizeof(FNiagaraDrawIndirectArgGenTaskInfo);
-			const uint32 InstanceCountClearSize = InstanceCountClearTasks.Num() * sizeof(uint32);
+			const uint32 InstanceCountClearSize = ComputeDispatchInterface->IsFirstViewFamily() ? InstanceCountClearTasks.Num() * sizeof(uint32) : 0;
 			const uint32 TaskBufferSize = ArgGenSize + InstanceCountClearSize;
 			TaskInfosBuffer.Initialize(TEXT("NiagaraTaskInfosBuffer"), sizeof(uint32), TaskBufferSize / sizeof(uint32), EPixelFormat::PF_R32_UINT, BUF_Volatile);
 
@@ -421,15 +422,17 @@ void FNiagaraGPUInstanceCountManager::UpdateDrawIndirectBuffers(FNiagaraGpuCompu
 		TArray<FRHITransitionInfo, TInlineAllocator<10>> Transitions;
 		Transitions.Reserve(DrawIndirectPool.Num() + 2);
 
+		FRWBuffer& CurrentCountBuffer = ComputeDispatchInterface->IsFirstViewFamily() || CountPhase == ENiagaraGPUCountUpdatePhase::PostOpaque ? CountBuffer : MultiViewCountBuffer;
+
 		// Get counts buffer
 		FUnorderedAccessViewRHIRef CountsUAV = nullptr;
-		const bool bCountBufferIsValid = CountBuffer.UAV.IsValid();
+		const bool bCountBufferIsValid = CurrentCountBuffer.UAV.IsValid();
 		if (bCountBufferIsValid)
 		{
 			// treat the incoming UAV as being unknown to be sure a barrier is inserted in the case
 			// where the preceding dispatch wrote to the counts buffer
-			Transitions.Emplace(CountBuffer.UAV, ERHIAccess::Unknown, ERHIAccess::UAVCompute);
-			CountsUAV = CountBuffer.UAV;
+			Transitions.Emplace(CurrentCountBuffer.UAV, ERHIAccess::Unknown, ERHIAccess::UAVCompute);
+			CountsUAV = CurrentCountBuffer.UAV;
 		}
 		else
 		{
@@ -490,7 +493,7 @@ void FNiagaraGPUInstanceCountManager::UpdateDrawIndirectBuffers(FNiagaraGpuCompu
 			}
 
 			const bool bIsLastDispatch = DispatchIdx == (NumDispatches - 1);
-			const int32 NumInstanceCountClearTasks = bIsLastDispatch ? InstanceCountClearTasks.Num() : 0;
+			const int32 NumInstanceCountClearTasks = bIsLastDispatch && ComputeDispatchInterface->IsFirstViewFamily() ? InstanceCountClearTasks.Num() : 0;
 
 			// Do we have anything to do for this pool?
 			if (NumArgGenTasks + NumInstanceCountClearTasks == 0)
@@ -548,7 +551,7 @@ void FNiagaraGPUInstanceCountManager::UpdateDrawIndirectBuffers(FNiagaraGpuCompu
 		{
 			Transitions.Emplace(PoolEntry->Buffer.UAV, ERHIAccess::UAVCompute, kIndirectArgsDefaultState);
 		}
-		Transitions.Emplace(CountBuffer.UAV, ERHIAccess::UAVCompute, kCountBufferDefaultState);
+		Transitions.Emplace(CurrentCountBuffer.UAV, ERHIAccess::UAVCompute, kCountBufferDefaultState);
 		RHICmdList.Transition(Transitions);
 	}
 
@@ -649,3 +652,26 @@ bool FNiagaraGPUInstanceCountManager::HasPendingGPUReadback() const
 	return CountReadback && CountReadbackSize;
 }
 
+void FNiagaraGPUInstanceCountManager::CopyToMultiViewCountBuffer(FRHICommandListImmediate& RHICmdList)
+{
+	// Need to copy on all GPUs
+	SCOPED_GPU_MASK(RHICmdList, FRHIGPUMask::All());
+
+	// Set AllocatedInstanceCounts and copy CountBuffer
+	if (MultiViewAllocatedInstanceCounts != AllocatedInstanceCounts)
+	{
+		MultiViewAllocatedInstanceCounts = AllocatedInstanceCounts;
+		MultiViewCountBuffer.Initialize(TEXT("NiagaraGPUInstanceCounts"), sizeof(uint32), MultiViewAllocatedInstanceCounts, EPixelFormat::PF_R32_UINT, ERHIAccess::UAVCompute, BUF_Static | BUF_SourceCopy);
+	}
+	else
+	{
+		RHICmdList.Transition(FRHITransitionInfo(MultiViewCountBuffer.UAV, kCountBufferDefaultState, ERHIAccess::UAVCompute));
+	}
+
+	FRHIUnorderedAccessView* UAVs[] = { MultiViewCountBuffer.UAV };
+	int32 UsedIndexCounts[] = { MultiViewAllocatedInstanceCounts };
+	CopyUIntBufferToTargets(RHICmdList, FeatureLevel, CountBuffer.SRV, UAVs, UsedIndexCounts, 0, UE_ARRAY_COUNT(UAVs));
+
+	RHICmdList.Transition(FRHITransitionInfo(MultiViewCountBuffer.UAV, ERHIAccess::UAVCompute, kCountBufferDefaultState));
+}
+
diff --git a/Engine/Plugins/FX/Niagara/Source/Niagara/Private/NiagaraGpuComputeDispatch.cpp b/Engine/Plugins/FX/Niagara/Source/Niagara/Private/NiagaraGpuComputeDispatch.cpp
--- a/Engine/Plugins/FX/Niagara/Source/Niagara/Private/NiagaraGpuComputeDispatch.cpp
+++ b/Engine/Plugins/FX/Niagara/Source/Niagara/Private/NiagaraGpuComputeDispatch.cpp
@@ -82,6 +82,16 @@ static FAutoConsoleVariableRef CVarNiagaraBatcherFreeBufferEarly(
 	ECVF_Default
 );
 
+#if WITH_MGPU
+int32 GNiagaraOptimizedCrossGPUTransferEnabled = 1;
+static FAutoConsoleVariableRef CVarNiagaraOptimizedCrossGPUTransferEnabled(
+	TEXT("fx.NiagaraOptimizeCrossGPUTransfer"),
+	GNiagaraOptimizedCrossGPUTransferEnabled,
+	TEXT("Optimizes fence waits for cross GPU transfers when rendering views on multiple GPUs via nDisplay.  (Default = 1)"),
+	ECVF_Default
+);
+#endif
+
 const FName FNiagaraGpuComputeDispatch::Name(TEXT("FNiagaraGpuComputeDispatch"));
 
 namespace FNiagaraGpuComputeDispatchLocal
@@ -193,9 +203,6 @@ FNiagaraGpuComputeDispatch::FNiagaraGpuComputeDispatch(ERHIFeatureLevel::Type In
 				[this](FRHICommandListImmediate& RHICmdList)
 				{
 					GPUInstanceCounterManager.UpdateDrawIndirectBuffers(this, RHICmdList, ENiagaraGPUCountUpdatePhase::PreOpaque);
-				#if WITH_MGPU
-					TransferMultiGPUBufers(RHICmdList, ENiagaraGpuComputeTickStage::PreInitViews);
-				#endif // WITH_MGPU
 				}
 			);
 
@@ -205,7 +212,7 @@ FNiagaraGpuComputeDispatch::FNiagaraGpuComputeDispatch(ERHIFeatureLevel::Type In
 				{
 					GPUInstanceCounterManager.UpdateDrawIndirectBuffers(this, RHICmdList, ENiagaraGPUCountUpdatePhase::PostOpaque);
 				#if WITH_MGPU
-					TransferMultiGPUBufers(RHICmdList, ENiagaraGpuComputeTickStage::PostOpaqueRender);
+					TransferMultiGPUBuffers(RHICmdList);
 				#endif // WITH_MGPU
 				}
 			);
@@ -431,7 +438,7 @@ void FNiagaraGpuComputeDispatch::ProcessPendingTicksFlush(FRHICommandListImmedia
 			{
 				FRDGBuilder GraphBuilder(RHICmdList);
 				CreateSystemTextures(GraphBuilder);
-				PreInitViews(GraphBuilder, bAllowGPUParticleUpdate);
+				PreInitViews(GraphBuilder, bAllowGPUParticleUpdate, TArrayView<const FSceneViewFamily*>(), nullptr);
 				AddPass(GraphBuilder, RDG_EVENT_NAME("UpdateDrawIndirectBuffers - PreOpaque"),
 					[this](FRHICommandListImmediate& RHICmdList)
 					{
@@ -478,7 +485,7 @@ void FNiagaraGpuComputeDispatch::FinishDispatches()
 	{
 		for (FNiagaraSystemGpuComputeProxy* ComputeProxy : ProxiesPerStage[iTickStage])
 		{
-			ComputeProxy->ReleaseTicks(GPUInstanceCounterManager, MaxTicksToFlush);
+			ComputeProxy->ReleaseTicks(GPUInstanceCounterManager, MaxTicksToFlush, bIsLastViewFamily);
 		}
 	}
 
@@ -1015,16 +1022,53 @@ void FNiagaraGpuComputeDispatch::PrepareAllTicks(FRHICommandListImmediate& RHICm
 void FNiagaraGpuComputeDispatch::ExecuteTicks(FRDGBuilder& GraphBuilder, TConstStridedView<FSceneView> Views, ENiagaraGpuComputeTickStage::Type TickStage)
 {
 #if WITH_MGPU
-	if (StageToWaitForGPUTransfers == TickStage)
+	if (TickStage == ENiagaraGpuComputeTickStage::PostOpaqueRender)
 	{
- 		GraphBuilder.AddPass(
-			RDG_EVENT_NAME("Niagara::ExecuteTicksPre"),
-			ERDGPassFlags::None,
-			[this, TickStage, NumDispatchGroups=DispatchListPerStage[TickStage].DispatchGroups.Num()](FRHICommandListImmediate& RHICmdList)
-			{
-				WaitForMultiGPUBuffers(RHICmdList, TickStage);
-			}
-		);
+		// See if this view is rendering on a GPU that needs to wait on its cross GPU transfers
+		uint32 GPUIndex = Views[0].GPUMask.GetFirstIndex();
+		if (OptimizedCrossGPUTransferMask & (1u << GPUIndex))
+		{
+			// Track that we've waited on the buffers for this GPU
+			OptimizedCrossGPUTransferMask &= ~(1u << GPUIndex);
+
+			GraphBuilder.AddPass(
+				RDG_EVENT_NAME("Niagara::ExecuteTicksPre"),
+				ERDGPassFlags::None,
+				[this, TickStage, NumDispatchGroups = DispatchListPerStage[TickStage].DispatchGroups.Num(), GPUIndex](FRHICommandListImmediate& RHICmdList)
+				{
+					WaitForMultiGPUBuffers(RHICmdList, GPUIndex);
+				}
+			);
+		}
+
+		// Last view family when rendering multiple view families?
+		if (bIsLastViewFamily && !bIsFirstViewFamily)
+		{
+			GraphBuilder.AddPass(
+				RDG_EVENT_NAME("Niagara::MultiViewPreviousDataClear"),
+				ERDGPassFlags::None,
+				[this](FRHICommandListImmediate& RHICmdList)
+				{
+					// Handle bookkeeping for multi-view, clearing MultiViewPreviousDataToRender pointer and running a couple bits of cleanup logic,
+					// which we had to skip earlier to keep necessary values valid when generating render data for additional view families.
+					for (FNiagaraComputeExecutionContext* ComputeContext : NeedsMultiViewPreviousDataClear)
+					{
+						ComputeContext->SetMultiViewPreviousDataToRender(nullptr);
+
+						// Mark data as ready for anyone who picks up the buffer on the next frame (see FNiagaraGpuComputeDispatch::ExecuteTicks)
+						ComputeContext->GetDataToRender(false)->SetGPUDataReadyStage(ENiagaraGpuComputeTickStage::First);
+
+						// Clear instance count offsets (see FNiagaraSystemGpuComputeProxy::ReleaseTicks)
+						for (int i = 0; i < UE_ARRAY_COUNT(ComputeContext->DataBuffers_RT); ++i)
+						{
+							ComputeContext->DataBuffers_RT[i]->ClearGPUInstanceCount();
+						}
+					}
+
+					NeedsMultiViewPreviousDataClear.Reset();
+				}
+			);
+		}
 	}
 #endif
 
@@ -1221,9 +1265,6 @@ void FNiagaraGpuComputeDispatch::ExecuteTicks(FRDGBuilder& GraphBuilder, TConstS
 				ComputeContext->MainDataSet->CurrentData = CurrentData;
 				CurrentData->SwapGPU(FinalSimStageDataBuffer);
 
-				// Mark data as ready for anyone who picks up the buffer on the next frame
-				CurrentData->SetGPUDataReadyStage(ENiagaraGpuComputeTickStage::First);
-                    
 				ComputeContext->SetTranslucentDataToRender(nullptr);
 				ComputeContext->SetDataToRender(CurrentData);
 
@@ -1235,6 +1276,24 @@ void FNiagaraGpuComputeDispatch::ExecuteTicks(FRDGBuilder& GraphBuilder, TConstS
 					AddCrossGPUTransfer(GraphBuilder.RHICmdList, CurrentData->GetGPUBufferInt().Buffer);
 				}
 #endif // WITH_MGPU
+
+				if (bIsLastViewFamily)
+				{
+					// Mark data as ready for anyone who picks up the buffer on the next frame
+					CurrentData->SetGPUDataReadyStage(ENiagaraGpuComputeTickStage::First);
+				}
+				else
+				{
+					// Track the previous data to render for multi-view, so additional view families can render using consistent data with
+					// the first view family, and add it to array marking that we need to clean it up on the last view family.
+					ComputeContext->SetMultiViewPreviousDataToRender(FinalSimStageDataBuffer);
+					NeedsMultiViewPreviousDataClear.Emplace(ComputeContext);
+
+					// Mark CurrentData as not to be used until PostOpaqueRender.  Render data generated before PostOpaqueRender should use
+					// previous frame data (MultiViewPreviousDataToRender, set above), instead of CurrentData.  Setting the ready stage back to
+					// ENiagaraGpuComputeTickStage::First for the next frame happens later in the Niagara::MultiViewPreviousDataClear Pass.
+					CurrentData->SetGPUDataReadyStage(ENiagaraGpuComputeTickStage::PostOpaqueRender);
+				}
 			}
 			// If this is not the final tick of the final stage we need set our temporary buffer for data interfaces, etc, that may snoop from CurrentData
 			else
@@ -1761,20 +1820,75 @@ void FNiagaraGpuComputeDispatch::DispatchStage(FRDGBuilder& GraphBuilder, const
 	//}
 }
 
-void FNiagaraGpuComputeDispatch::PreInitViews(FRDGBuilder& GraphBuilder, bool bAllowGPUParticleUpdate)
+void FNiagaraGpuComputeDispatch::PreInitViews(FRDGBuilder& GraphBuilder, bool bAllowGPUParticleUpdate, const TArrayView<const FSceneViewFamily*>& ViewFamilies, const FSceneViewFamily* CurrentFamily)
 {
 	SCOPE_CYCLE_COUNTER(STAT_NiagaraGPUDispatchSetup_RT);
 	RDG_CSV_STAT_EXCLUSIVE_SCOPE(GraphBuilder, Niagara);
 
+	bIsFirstViewFamily = CurrentFamily ? CurrentFamily == ViewFamilies[0] : true;
+	bIsLastViewFamily = CurrentFamily ? CurrentFamily == ViewFamilies.Last() : true;
 	bRequiresReadback = false;
 #if WITH_EDITOR
 	bRaisedWarningThisFrame = false;
 #endif
 #if WITH_MGPU
 	bCrossGPUTransferEnabled = GNumExplicitGPUsForRendering > 1;
-	StageToTransferGPUBuffers = ENiagaraGpuComputeTickStage::Last;
-	StageToWaitForGPUTransfers = ENiagaraGpuComputeTickStage::First;
+
+	// We want to optimize GPU transfers from the simulating GPU to other view rendering GPUs, meaning we want to defer fence waits on
+	// the other GPUs until the data is needed.  To keep things simple, we only do this optimization if all the families are single GPU,
+	// which will be true in practice for multiple families rendered via nDisplay.  Here we compute a mask defining which other GPUs we
+	// want to optimize transfers for (ones that render other views).  GPUs which don't do any view rendering in the current frame don't
+	// incur any performance cost waiting on the fences immediately (the default behavior of cross GPU transfers), so we don't need to
+	// delay the fence wait for those.
+	if (bIsFirstViewFamily)
+	{
+		OptimizedCrossGPUTransferMask = 0;
+	}
+
+	if ((ViewFamilies.Num() > 1) && bIsFirstViewFamily && GNiagaraOptimizedCrossGPUTransferEnabled)
+	{
+		bool bOptimizeTransfers = true;
+		FRHIGPUMask SimGPUMask = ViewFamilies[0]->Views[0]->GPUMask;
+		FRHIGPUMask AllViewGPUMask = SimGPUMask;
+
+		for (const FSceneViewFamily* Family : ViewFamilies)
+		{
+			FRHIGPUMask FamilyFirstViewMask = Family->Views[0]->GPUMask;
+			AllViewGPUMask |= FamilyFirstViewMask;
+
+			if (!FamilyFirstViewMask.HasSingleIndex())
+			{
+				bOptimizeTransfers = false;
+				break;
+			}
+			for (const FSceneView* View : Family->Views)
+			{
+				if (View->GPUMask != FamilyFirstViewMask)
+				{
+					bOptimizeTransfers = false;
+					break;
+				}
+			}
+			if (!bOptimizeTransfers)
+			{
+				break;
+			}
+		}
+
+		if (bOptimizeTransfers)
+		{
+			// Figure out what destination GPUs we should optimize transfers to
+			OptimizedCrossGPUTransferMask = AllViewGPUMask.GetNative() & ~SimGPUMask.GetNative();
+		}
+	}
 #endif
+	if ((ViewFamilies.Num() > 1) && bIsFirstViewFamily)
+	{
+		AddPass(GraphBuilder, RDG_EVENT_NAME("Niagara::CopyToMultiViewCountBuffer"), [this](FRHICommandListImmediate& RHICmdList)
+		{
+			GPUInstanceCounterManager.CopyToMultiViewCountBuffer(RHICmdList);
+		});
+	}
 
 	GpuReadbackManagerPtr->Tick();
 #if NIAGARA_COMPUTEDEBUG_ENABLED
@@ -1808,10 +1922,6 @@ void FNiagaraGpuComputeDispatch::PreInitViews(FRDGBuilder& GraphBuilder, bool bA
 
 			UpdateInstanceCountManager(GraphBuilder.RHICmdList);
 			PrepareAllTicks(GraphBuilder.RHICmdList);
-		
-		#if WITH_MGPU
-			CalculateCrossGPUTransferLocation();
-		#endif
 
 			AsyncGpuTraceHelper->BeginFrame(GraphBuilder.RHICmdList, this);
 
@@ -1845,7 +1955,7 @@ void FNiagaraGpuComputeDispatch::PreInitViews(FRDGBuilder& GraphBuilder, bool bA
 			return false;
 		};
 
-	if (GPUInstanceCounterManager.HasEntriesPendingFree() || HasAnyWork())
+	if (bAllowGPUParticleUpdate && ((IsFirstViewFamily() && GPUInstanceCounterManager.HasEntriesPendingFree()) || HasAnyWork()))
 	{
 		MultiGPUResourceModified(GraphBuilder, GPUInstanceCounterManager.GetInstanceCountBuffer().Buffer, true, true);
 	}
@@ -1863,21 +1973,6 @@ void FNiagaraGpuComputeDispatch::PostInitViews(FRDGBuilder& GraphBuilder, TConst
 		RDG_CSV_STAT_EXCLUSIVE_SCOPE(GraphBuilder, Niagara);
 
 		ExecuteTicks(GraphBuilder, Views, ENiagaraGpuComputeTickStage::PostInitViews);
-
-	#if WITH_MGPU
-		// Queue a transfer request
-		if (StageToTransferGPUBuffers == ENiagaraGpuComputeTickStage::PostInitViews)
-		{
-			AddPass(
-				GraphBuilder,
-				RDG_EVENT_NAME("Niagara::TransferMultiGPUBufers"),
-				[this](FRHICommandListImmediate& RHICmdList)
-				{
-					TransferMultiGPUBufers(RHICmdList, ENiagaraGpuComputeTickStage::PostInitViews);
-				}
-			);
-		}
-	#endif // WITH_MGPU
 	}
 }
 
@@ -2377,7 +2472,19 @@ void FNiagaraGpuComputeDispatch::MultiGPUResourceModified(FRHICommandList& RHICm
 		{
 			if (!GPUMask.Contains(GPUIndex))
 			{
-				const_cast<FNiagaraGpuComputeDispatch*>(this)->CrossGPUTransferBuffers.Emplace(Texture, GPUMask.GetFirstIndex(), GPUIndex, bPullData, bLockStep);
+				if (OptimizedCrossGPUTransferMask & (1u << GPUIndex))
+				{
+					// Add a sync point for this destination GPU when we're adding the first transfer
+					if (OptimizedCrossGPUTransferBuffers[GPUIndex].Num() == 0)
+					{
+						OptimizedCrossGPUFences[GPUIndex] = RHICreateCrossGPUTransferFence();
+					}
+					OptimizedCrossGPUTransferBuffers[GPUIndex].Emplace(Texture, GPUMask.GetFirstIndex(), GPUIndex, bPullData, bLockStep);
+				}
+				else
+				{
+					CrossGPUTransferBuffers.Emplace(Texture, GPUMask.GetFirstIndex(), GPUIndex, bPullData, bLockStep);
+				}
 			}
 		}
 	}
@@ -2396,53 +2503,56 @@ void FNiagaraGpuComputeDispatch::AddCrossGPUTransfer(FRHICommandList& RHICmdList
 		{
 			if (!GPUMask.Contains(GPUIndex))
 			{
-				CrossGPUTransferBuffers.Emplace(Buffer, GPUMask.GetFirstIndex(), GPUIndex, bPullData, bLockStep);
+				if (OptimizedCrossGPUTransferMask & (1u << GPUIndex))
+				{
+					// Add a sync point for this destination GPU when we're adding the first transfer
+					if (OptimizedCrossGPUTransferBuffers[GPUIndex].Num() == 0)
+					{
+						OptimizedCrossGPUFences[GPUIndex] = RHICreateCrossGPUTransferFence();
+					}
+					OptimizedCrossGPUTransferBuffers[GPUIndex].Emplace(Buffer, GPUMask.GetFirstIndex(), GPUIndex, bPullData, bLockStep);
+				}
+				else
+				{
+					CrossGPUTransferBuffers.Emplace(Buffer, GPUMask.GetFirstIndex(), GPUIndex, bPullData, bLockStep);
+				}
 			}
 		}
 	}
 }
 
-void FNiagaraGpuComputeDispatch::CalculateCrossGPUTransferLocation()
+void FNiagaraGpuComputeDispatch::TransferMultiGPUBuffers(FRHICommandList& RHICmdList)
 {
-	StageToTransferGPUBuffers = ENiagaraGpuComputeTickStage::Last;
-	while (StageToTransferGPUBuffers > ENiagaraGpuComputeTickStage::First && !DispatchListPerStage[static_cast<int32>(StageToTransferGPUBuffers)].HasWork())
+	// Transfer buffers for cross GPU rendering
+	if (CrossGPUTransferBuffers.Num())
 	{
-		StageToTransferGPUBuffers = static_cast<ENiagaraGpuComputeTickStage::Type>(static_cast<int32>(StageToTransferGPUBuffers) - 1);
+		RHICmdList.TransferResources(CrossGPUTransferBuffers);
+		CrossGPUTransferBuffers.Reset();
 	}
 
-	StageToWaitForGPUTransfers = ENiagaraGpuComputeTickStage::First;
-	// If we're going to write to the instance count buffer after PreInitViews then
-	// that needs to be the wait stage, regardless of whether or not we're ticking
-	// anything in that stage.
-	if (!GPUInstanceCounterManager.HasEntriesPendingFree())
+	// Optimized transfer of buffers for cross GPU rendering (deferred fence wait)
+	for (uint32 GPUIndex = 0; GPUIndex < GNumExplicitGPUsForRendering; GPUIndex++)
 	{
-		while (StageToWaitForGPUTransfers < StageToTransferGPUBuffers && !DispatchListPerStage[static_cast<int32>(StageToWaitForGPUTransfers)].HasWork())
+		// Each destination GPU has a separate list of transfers, so we can wait for them on the GPU when it needs them
+		if (OptimizedCrossGPUTransferBuffers[GPUIndex].Num())
 		{
-			StageToWaitForGPUTransfers = static_cast<ENiagaraGpuComputeTickStage::Type>(static_cast<int32>(StageToWaitForGPUTransfers) + 1);
-		}
-	}
-}
-
-void FNiagaraGpuComputeDispatch::TransferMultiGPUBufers(FRHICommandList& RHICmdList, ENiagaraGpuComputeTickStage::Type TickStage)
-{
-	if (StageToTransferGPUBuffers != TickStage)
-	{
-		return;
-	}
+			// Signal function transitions resources on destination GPU and signals when transition has finished
+			TArray<FCrossGPUTransferFence*> PreTransferSyncPoints;
+			RHIGenerateCrossGPUPreTransferFences(OptimizedCrossGPUTransferBuffers[GPUIndex], PreTransferSyncPoints);
+			RHICmdList.CrossGPUTransferSignal(OptimizedCrossGPUTransferBuffers[GPUIndex], PreTransferSyncPoints);
 
-	// Transfer buffers for cross GPU rendering
-	if (CrossGPUTransferBuffers.Num())
-	{
-		RHICmdList.TransferResources(CrossGPUTransferBuffers);
-		CrossGPUTransferBuffers.Reset();
+			RHICmdList.CrossGPUTransfer(OptimizedCrossGPUTransferBuffers[GPUIndex], PreTransferSyncPoints, TArrayView<FCrossGPUTransferFence*>(&OptimizedCrossGPUFences[GPUIndex], 1));
+			OptimizedCrossGPUTransferBuffers[GPUIndex].Reset();
+		}
 	}
 }
 
-void FNiagaraGpuComputeDispatch::WaitForMultiGPUBuffers(FRHICommandList& RHICmdList, ENiagaraGpuComputeTickStage::Type TickStage)
+void FNiagaraGpuComputeDispatch::WaitForMultiGPUBuffers(FRHICommandList& RHICmdList, uint32 GPUIndex)
 {
-	if (StageToWaitForGPUTransfers == TickStage)
+	if (OptimizedCrossGPUFences[GPUIndex])
 	{
-		// TODO:  implement delayed fence wait here for cross GPU transfers issued above
+		RHICmdList.CrossGPUTransferWait(TArrayView<FCrossGPUTransferFence*>(&OptimizedCrossGPUFences[GPUIndex], 1));
+		OptimizedCrossGPUFences[GPUIndex] = nullptr;
 	}
 }
 #endif // WITH_MGPU
diff --git a/Engine/Plugins/FX/Niagara/Source/Niagara/Private/NiagaraGpuComputeDispatch.h b/Engine/Plugins/FX/Niagara/Source/Niagara/Private/NiagaraGpuComputeDispatch.h
--- a/Engine/Plugins/FX/Niagara/Source/Niagara/Private/NiagaraGpuComputeDispatch.h
+++ b/Engine/Plugins/FX/Niagara/Source/Niagara/Private/NiagaraGpuComputeDispatch.h
@@ -109,7 +109,7 @@ public:
 	virtual void AddVectorField(UVectorFieldComponent* VectorFieldComponent) override {}
 	virtual void RemoveVectorField(UVectorFieldComponent* VectorFieldComponent) override {}
 	virtual void UpdateVectorField(UVectorFieldComponent* VectorFieldComponent) override {}
-	virtual void PreInitViews(FRDGBuilder& GraphBuilder, bool bAllowGPUParticleUpdate) override;
+	virtual void PreInitViews(FRDGBuilder& GraphBuilder, bool bAllowGPUParticleUpdate, const TArrayView<const FSceneViewFamily*> &ViewFamilies, const FSceneViewFamily* CurrentFamily) override;
 	virtual void PostInitViews(FRDGBuilder& GraphBuilder, TConstStridedView<FSceneView> Views, bool bAllowGPUParticleUpdate) override;
 	virtual bool UsesGlobalDistanceField() const override;
 	virtual bool UsesDepthBuffer() const override;
@@ -246,18 +246,22 @@ private:
 	};
 	TArray<FDebugReadbackInfo> GpuDebugReadbackInfos;
 
-#if WITH_MGPU
-	ENiagaraGpuComputeTickStage::Type StageToTransferGPUBuffers = ENiagaraGpuComputeTickStage::First;
-	ENiagaraGpuComputeTickStage::Type StageToWaitForGPUTransfers = ENiagaraGpuComputeTickStage::First;
+	// List of items that need MultiViewPreviousDataToRender to be cleaned up on the last view family
+	TArray<FNiagaraComputeExecutionContext*> NeedsMultiViewPreviousDataClear;
 
+#if WITH_MGPU
+	// Cross GPU transfer arrays need to be mutable, so they can be written from const MultiGPUResourceModified utility functions
 	bool bCrossGPUTransferEnabled = false;
-	TArray<FTransferResourceParams> CrossGPUTransferBuffers;
+	mutable TArray<FTransferResourceParams> CrossGPUTransferBuffers;
+
+	uint32 OptimizedCrossGPUTransferMask = 0;
+	mutable TArray<FTransferResourceParams> OptimizedCrossGPUTransferBuffers[MAX_NUM_GPUS];
+	mutable FCrossGPUTransferFence* OptimizedCrossGPUFences[MAX_NUM_GPUS] = { 0 };
 
 	void AddCrossGPUTransfer(FRHICommandList& RHICmdList, FRHIBuffer* Buffer);
 
-	void CalculateCrossGPUTransferLocation();
-	void TransferMultiGPUBufers(FRHICommandList& RHICmdList, ENiagaraGpuComputeTickStage::Type TickStage);
-	void WaitForMultiGPUBuffers(FRHICommandList& RHICmdList, ENiagaraGpuComputeTickStage::Type TickStage);
+	void TransferMultiGPUBuffers(FRHICommandList& RHICmdList);
+	void WaitForMultiGPUBuffers(FRHICommandList& RHICmdList, uint32 GPUIndex);
 #endif // WITH_MGPU
 
 	// Cached information to build a dummy view info if necessary
diff --git a/Engine/Plugins/FX/Niagara/Source/Niagara/Private/NiagaraSystemGpuComputeProxy.cpp b/Engine/Plugins/FX/Niagara/Source/Niagara/Private/NiagaraSystemGpuComputeProxy.cpp
--- a/Engine/Plugins/FX/Niagara/Source/Niagara/Private/NiagaraSystemGpuComputeProxy.cpp
+++ b/Engine/Plugins/FX/Niagara/Source/Niagara/Private/NiagaraSystemGpuComputeProxy.cpp
@@ -107,7 +107,7 @@ void FNiagaraSystemGpuComputeProxy::RemoveFromRenderThread(FNiagaraGpuComputeDis
 		[this, ComputeDispatchInterface, bDeleteProxy](FRHICommandListImmediate& RHICmdList)
 		{
 			ComputeDispatchInterface->RemoveGpuComputeProxy(this);
-			ReleaseTicks(ComputeDispatchInterface->GetGPUInstanceCounterManager(), TNumericLimits<int32>::Max());
+			ReleaseTicks(ComputeDispatchInterface->GetGPUInstanceCounterManager(), TNumericLimits<int32>::Max(), true);
 
 			for (FNiagaraComputeExecutionContext* ComputeContext : ComputeContexts)
 			{
@@ -139,7 +139,7 @@ void FNiagaraSystemGpuComputeProxy::ClearTicksFromRenderThread(FNiagaraGpuComput
 	ENQUEUE_RENDER_COMMAND(ClearTicksFromProxy)(
 		[this, ComputeDispatchInterface](FRHICommandListImmediate& RHICmdList)
 		{
-			ReleaseTicks(ComputeDispatchInterface->GetGPUInstanceCounterManager(), TNumericLimits<int32>::Max());
+			ReleaseTicks(ComputeDispatchInterface->GetGPUInstanceCounterManager(), TNumericLimits<int32>::Max(), true);
 		}
 	);
 }
@@ -172,7 +172,7 @@ void FNiagaraSystemGpuComputeProxy::QueueTick(const FNiagaraGPUSystemTick& Tick)
 	}
 }
 
-void FNiagaraSystemGpuComputeProxy::ReleaseTicks(FNiagaraGPUInstanceCountManager& GPUInstanceCountManager, int32 NumTicksToRelease)
+void FNiagaraSystemGpuComputeProxy::ReleaseTicks(FNiagaraGPUInstanceCountManager& GPUInstanceCountManager, int32 NumTicksToRelease, bool bLastViewFamily)
 {
 	check(IsInRenderingThread());
 
@@ -191,10 +191,14 @@ void FNiagaraSystemGpuComputeProxy::ReleaseTicks(FNiagaraGPUInstanceCountManager
 		ComputeContext->bHasTickedThisFrame_RT = false;
 		ComputeContext->CurrentMaxInstances_RT = 0;
 
-		// Clear counter offsets
-		for (int i = 0; i < UE_ARRAY_COUNT(ComputeContext->DataBuffers_RT); ++i)
+		// This logic is deferred to the Niagara::MultiViewPreviousDataClear Pass for multi-view rendering
+		if (bLastViewFamily)
 		{
-			ComputeContext->DataBuffers_RT[i]->ClearGPUInstanceCount();
+			// Clear counter offsets
+			for (int i = 0; i < UE_ARRAY_COUNT(ComputeContext->DataBuffers_RT); ++i)
+			{
+				ComputeContext->DataBuffers_RT[i]->ClearGPUInstanceCount();
+			}
 		}
 	}
 }
diff --git a/Engine/Plugins/FX/Niagara/Source/Niagara/Public/NiagaraGpuComputeDispatchInterface.h b/Engine/Plugins/FX/Niagara/Source/Niagara/Public/NiagaraGpuComputeDispatchInterface.h
--- a/Engine/Plugins/FX/Niagara/Source/Niagara/Public/NiagaraGpuComputeDispatchInterface.h
+++ b/Engine/Plugins/FX/Niagara/Source/Niagara/Public/NiagaraGpuComputeDispatchInterface.h
@@ -145,6 +145,9 @@ public:
 
 	FORCEINLINE bool IsOutsideSceneRenderer() const { return bIsOutsideSceneRenderer; }
 
+	FORCEINLINE bool IsFirstViewFamily() const { return bIsFirstViewFamily; }
+	FORCEINLINE bool IsLastViewFamily() const { return bIsLastViewFamily; }
+
 #if WITH_MGPU
 	/**
 	Notify that a GPU resource was modified that will impact MultiGPU rendering.
@@ -185,4 +188,6 @@ protected:
 	TConstStridedView<FSceneView>			SimulationSceneViews;
 
 	bool									bIsOutsideSceneRenderer = false;
+	bool									bIsFirstViewFamily = true;
+	bool									bIsLastViewFamily = true;
 };
diff --git a/Engine/Plugins/FX/Niagara/Source/Niagara/Public/NiagaraSystemGpuComputeProxy.h b/Engine/Plugins/FX/Niagara/Source/Niagara/Public/NiagaraSystemGpuComputeProxy.h
--- a/Engine/Plugins/FX/Niagara/Source/Niagara/Public/NiagaraSystemGpuComputeProxy.h
+++ b/Engine/Plugins/FX/Niagara/Source/Niagara/Public/NiagaraSystemGpuComputeProxy.h
@@ -26,7 +26,7 @@ public:
 	FNiagaraSystemInstanceID GetSystemInstanceID() const { return SystemInstanceID; }
 	ENiagaraGpuComputeTickStage::Type GetComputeTickStage() const { return ComputeTickStage; }
 	void QueueTick(const FNiagaraGPUSystemTick& Tick);
-	void ReleaseTicks(FNiagaraGPUInstanceCountManager& GPUInstanceCountManager, int32 NumTicksToRelease);
+	void ReleaseTicks(FNiagaraGPUInstanceCountManager& GPUInstanceCountManager, int32 NumTicksToRelease, bool bLastViewFamily);
 
 	bool RequiresDistanceFieldData() const { return bRequiresDistanceFieldData; }
 	bool RequiresDepthBuffer() const { return bRequiresDepthBuffer; }
diff --git a/Engine/Source/Runtime/D3D12RHI/Private/D3D12CommandContext.cpp b/Engine/Source/Runtime/D3D12RHI/Private/D3D12CommandContext.cpp
--- a/Engine/Source/Runtime/D3D12RHI/Private/D3D12CommandContext.cpp
+++ b/Engine/Source/Runtime/D3D12RHI/Private/D3D12CommandContext.cpp
@@ -1210,3 +1210,158 @@ void FD3D12CommandContextRedirector::RHITransferResourceWait(const TArrayView<FT
 
 #endif // WITH_MGPU
 }
+
+void FD3D12CommandContextRedirector::RHICrossGPUTransfer(const TArrayView<const FTransferResourceParams> Params, const TArrayView<FCrossGPUTransferFence* const> PreTransfer, const TArrayView<FCrossGPUTransferFence* const> PostTransfer)
+{
+#if WITH_MGPU
+	if (Params.Num() == 0)
+		return;
+
+	for (const FTransferResourceParams& Param : Params)
+	{
+		FD3D12CommandContext* SrcContext = PhysicalContexts[Param.SrcGPUIndex];
+		check(SrcContext);
+
+		FD3D12Resource* SrcResource;
+
+		if (Param.Texture)
+		{
+			check(Param.Buffer == nullptr);
+
+			SrcResource = FD3D12CommandContext::RetrieveTexture(Param.Texture, Param.SrcGPUIndex )->GetResource();
+		}
+		else
+		{
+			check(Param.Buffer != nullptr);
+
+			SrcResource = FD3D12DynamicRHI::ResourceCast(Param.Buffer.GetReference(), Param.SrcGPUIndex )->GetResource();
+		}
+
+		// Destination GPU resources are transitioned in RHICrossGPUTransferSignal, potentially allowing them to be transitioned earlier
+		// in the timeline, reducing the likelihood of the source GPU needing to wait to start the transfer.
+		SrcContext->TransitionResource(SrcResource, D3D12_RESOURCE_STATE_TBD, D3D12_RESOURCE_STATE_COPY_SOURCE, 0);
+	}
+
+	// Wait on any pre-transfer fences first
+	for (FCrossGPUTransferFence* PreTransferSyncPoint : PreTransfer)
+	{
+		FD3D12SyncPoint* SyncPoint = static_cast<FD3D12SyncPoint*>(PreTransferSyncPoint->SyncPoint);
+
+		PhysicalContexts[PreTransferSyncPoint->WaitGPUIndex]->WaitSyncPoint(SyncPoint);
+
+		SyncPoint->Release();
+
+		delete PreTransferSyncPoint;
+	}
+	
+	// Enqueue the copy work
+	for (const FTransferResourceParams& Param : Params)
+	{
+		FD3D12CommandContext* SrcContext = PhysicalContexts[Param.SrcGPUIndex];
+
+		if (Param.Texture)
+		{
+			FD3D12Texture* SrcTexture = FD3D12CommandContext::RetrieveTexture(Param.Texture, Param.SrcGPUIndex);
+			FD3D12Texture* DstTexture = FD3D12CommandContext::RetrieveTexture(Param.Texture, Param.DestGPUIndex);
+
+			// If the texture size is zero (Max.Z == 0, set in the constructor), copy the whole resource
+			if (Param.Max.Z == 0)
+			{
+				SrcContext->GraphicsCommandList()->CopyResource(DstTexture->GetResource()->GetResource(), SrcTexture->GetResource()->GetResource());
+			}
+			else
+			{
+				// Must be a 2D texture for this code path
+				check(Param.Texture->GetTexture2D() != nullptr);
+
+				ensureMsgf(
+					Param.Min.X >= 0 && Param.Min.Y >= 0 && Param.Min.Z >= 0 &&
+					Param.Max.X >= 0 && Param.Max.Y >= 0 && Param.Max.Z >= 0,
+					TEXT("Invalid rect for texture transfer: %i, %i, %i, %i"), Param.Min.X, Param.Min.Y, Param.Min.Z, Param.Max.X, Param.Max.Y, Param.Max.Z);
+
+				D3D12_BOX Box = { (UINT)Param.Min.X, (UINT)Param.Min.Y, (UINT)Param.Min.Z, (UINT)Param.Max.X, (UINT)Param.Max.Y, (UINT)Param.Max.Z };
+
+				CD3DX12_TEXTURE_COPY_LOCATION SrcLocation(SrcTexture->GetResource()->GetResource(), 0);
+				CD3DX12_TEXTURE_COPY_LOCATION DstLocation(DstTexture->GetResource()->GetResource(), 0);
+
+				SrcContext->GraphicsCommandList()->CopyTextureRegion(&DstLocation, Box.left, Box.top, Box.front, &SrcLocation, &Box);
+			}
+		}
+		else
+		{
+			FD3D12Resource* SrcResource = FD3D12DynamicRHI::ResourceCast(Param.Buffer.GetReference(), Param.SrcGPUIndex)->GetResource();
+			FD3D12Resource* DstResource = FD3D12DynamicRHI::ResourceCast(Param.Buffer.GetReference(), Param.DestGPUIndex)->GetResource();
+
+			SrcContext->GraphicsCommandList()->CopyResource(DstResource->GetResource(), SrcResource->GetResource());
+		}
+	}
+
+	// Post-copy synchronization
+	FD3D12SyncPointRef SyncPoint = FD3D12SyncPoint::Create(ED3D12SyncPointType::GPUOnly);
+	PhysicalContexts[Params[0].SrcGPUIndex]->SignalSyncPoint(SyncPoint);
+
+	for (FCrossGPUTransferFence* PostTransferSyncPoint : PostTransfer)
+	{
+		// Copy the sync points into the delayed fence struct. These will be awaited later in RHITransferResourceWait().
+		SyncPoint->AddRef();
+		PostTransferSyncPoint->SyncPoint = SyncPoint.GetReference();
+	}
+#endif // WITH_MGPU
+}
+
+void FD3D12CommandContextRedirector::RHICrossGPUTransferSignal(const TArrayView<const FTransferResourceParams> Params, const TArrayView<FCrossGPUTransferFence* const> PreTransfer)
+{
+#if WITH_MGPU
+	for (const FTransferResourceParams& Param : Params)
+	{
+		FD3D12CommandContext* DstContext = PhysicalContexts[Param.DestGPUIndex];
+		check(DstContext);
+
+		FD3D12Resource* DstResource;
+
+		if (Param.Texture)
+		{
+			check(Param.Buffer == nullptr);
+
+			DstResource = FD3D12CommandContext::RetrieveTexture(Param.Texture, Param.DestGPUIndex)->GetResource();
+		}
+		else
+		{
+			check(Param.Buffer != nullptr);
+
+			DstResource = FD3D12DynamicRHI::ResourceCast(Param.Buffer.GetReference(), Param.DestGPUIndex)->GetResource();
+		}
+
+		DstContext->TransitionResource(DstResource, D3D12_RESOURCE_STATE_TBD, D3D12_RESOURCE_STATE_COPY_DEST, 0);
+	}
+
+	for (FCrossGPUTransferFence* TransferSyncPoint : PreTransfer)
+	{
+		FD3D12SyncPointRef SyncPoint = FD3D12SyncPoint::Create(ED3D12SyncPointType::GPUOnly);
+		SyncPoint->AddRef();
+
+		PhysicalContexts[TransferSyncPoint->SignalGPUIndex]->SignalSyncPoint(SyncPoint);
+
+		TransferSyncPoint->SyncPoint = SyncPoint;
+	}
+#endif
+}
+
+void FD3D12CommandContextRedirector::RHICrossGPUTransferWait(const TArrayView<FCrossGPUTransferFence* const> PostTransfer)
+{
+#if WITH_MGPU
+	for (FCrossGPUTransferFence* TransferSyncPoint : PostTransfer)
+	{
+		if (TransferSyncPoint->SyncPoint)
+		{
+			FD3D12SyncPoint* SyncPoint = static_cast<FD3D12SyncPoint*>(TransferSyncPoint->SyncPoint);
+			PhysicalContexts[TransferSyncPoint->WaitGPUIndex]->WaitSyncPoint(SyncPoint);
+
+			SyncPoint->Release();
+		}
+
+		delete TransferSyncPoint;
+	}
+
+#endif // WITH_MGPU
+}
diff --git a/Engine/Source/Runtime/D3D12RHI/Private/D3D12CommandContext.h b/Engine/Source/Runtime/D3D12RHI/Private/D3D12CommandContext.h
--- a/Engine/Source/Runtime/D3D12RHI/Private/D3D12CommandContext.h
+++ b/Engine/Source/Runtime/D3D12RHI/Private/D3D12CommandContext.h
@@ -784,6 +784,11 @@ public:
 	virtual void RHITransferResourceSignal(const TArrayView<FTransferResourceFenceData* const> FenceDatas, FRHIGPUMask SrcGPUMask) final override;
 	virtual void RHITransferResourceWait(const TArrayView<FTransferResourceFenceData* const> FenceDatas) final override;
 
+	// New and improved cross GPU transfer API
+	virtual void RHICrossGPUTransfer(const TArrayView<const FTransferResourceParams> Params, const TArrayView<FCrossGPUTransferFence* const> PreTransfer, const TArrayView<FCrossGPUTransferFence* const> PostTransfer) final override;
+	virtual void RHICrossGPUTransferSignal(const TArrayView<const FTransferResourceParams> Params, const TArrayView<FCrossGPUTransferFence* const> PreTransfer) final override;
+	virtual void RHICrossGPUTransferWait(const TArrayView<FCrossGPUTransferFence* const> PostTransfer) final override;
+
 	FORCEINLINE virtual void RHICopyToStagingBuffer(FRHIBuffer* SourceBuffer, FRHIStagingBuffer* DestinationStagingBuffer, uint32 Offset, uint32 NumBytes) final override
 	{
 		ContextRedirect(RHICopyToStagingBuffer(SourceBuffer, DestinationStagingBuffer, Offset, NumBytes));
diff --git a/Engine/Source/Runtime/Engine/Private/Particles/FXSystem.cpp b/Engine/Source/Runtime/Engine/Private/Particles/FXSystem.cpp
--- a/Engine/Source/Runtime/Engine/Private/Particles/FXSystem.cpp
+++ b/Engine/Source/Runtime/Engine/Private/Particles/FXSystem.cpp
@@ -423,7 +423,7 @@ void FFXSystem::DrawDebug( FCanvas* Canvas )
 	}
 }
 
-void FFXSystem::PreInitViews(FRDGBuilder& GraphBuilder, bool bAllowGPUParticleUpdate)
+void FFXSystem::PreInitViews(FRDGBuilder& GraphBuilder, bool bAllowGPUParticleUpdate, const TArrayView<const FSceneViewFamily*>& ViewFamilies, const FSceneViewFamily* CurrentFamily)
 {
 	if (RHISupportsGPUParticles())
 	{
diff --git a/Engine/Source/Runtime/Engine/Private/Particles/FXSystemPrivate.h b/Engine/Source/Runtime/Engine/Private/Particles/FXSystemPrivate.h
--- a/Engine/Source/Runtime/Engine/Private/Particles/FXSystemPrivate.h
+++ b/Engine/Source/Runtime/Engine/Private/Particles/FXSystemPrivate.h
@@ -202,7 +202,7 @@ public:
 	virtual void RemoveVectorField(UVectorFieldComponent* VectorFieldComponent) override;
 	virtual void UpdateVectorField(UVectorFieldComponent* VectorFieldComponent) override;
 	FParticleEmitterInstance* CreateGPUSpriteEmitterInstance(FGPUSpriteEmitterInfo& EmitterInfo);
-	virtual void PreInitViews(class FRDGBuilder& GraphBuilder, bool bAllowGPUParticleUpdate) override;
+	virtual void PreInitViews(class FRDGBuilder& GraphBuilder, bool bAllowGPUParticleUpdate, const TArrayView<const FSceneViewFamily*>& ViewFamilies, const FSceneViewFamily* CurrentFamily) override;
 	virtual void PostInitViews(FRDGBuilder& GraphBuilder, TConstStridedView<FSceneView> Views, bool bAllowGPUParticleUpdate) override;
 	virtual bool UsesGlobalDistanceField() const override;
 	virtual bool UsesDepthBuffer() const override;
diff --git a/Engine/Source/Runtime/Engine/Private/Particles/FXSystemSet.cpp b/Engine/Source/Runtime/Engine/Private/Particles/FXSystemSet.cpp
--- a/Engine/Source/Runtime/Engine/Private/Particles/FXSystemSet.cpp
+++ b/Engine/Source/Runtime/Engine/Private/Particles/FXSystemSet.cpp
@@ -126,12 +126,12 @@ void FFXSystemSet::UpdateVectorField(UVectorFieldComponent* VectorFieldComponent
 	}
 }
 
-void FFXSystemSet::PreInitViews(FRDGBuilder& GraphBuilder, bool bAllowGPUParticleUpdate)
+void FFXSystemSet::PreInitViews(FRDGBuilder& GraphBuilder, bool bAllowGPUParticleUpdate, const TArrayView<const FSceneViewFamily*>& ViewFamilies, const FSceneViewFamily* CurrentFamily)
 {
 	for (FFXSystemInterface* FXSystem : FXSystems)
 	{
 		check(FXSystem);
-		FXSystem->PreInitViews(GraphBuilder, bAllowGPUParticleUpdate);
+		FXSystem->PreInitViews(GraphBuilder, bAllowGPUParticleUpdate, ViewFamilies, CurrentFamily);
 	}
 }
 
diff --git a/Engine/Source/Runtime/Engine/Private/Particles/FXSystemSet.h b/Engine/Source/Runtime/Engine/Private/Particles/FXSystemSet.h
--- a/Engine/Source/Runtime/Engine/Private/Particles/FXSystemSet.h
+++ b/Engine/Source/Runtime/Engine/Private/Particles/FXSystemSet.h
@@ -39,7 +39,7 @@ public:
 	virtual void AddVectorField(UVectorFieldComponent* VectorFieldComponent) override;
 	virtual void RemoveVectorField(UVectorFieldComponent* VectorFieldComponent) override;
 	virtual void UpdateVectorField(UVectorFieldComponent* VectorFieldComponent) override;
-	virtual void PreInitViews(FRDGBuilder& GraphBuilder, bool bAllowGPUParticleUpdate) override;
+	virtual void PreInitViews(FRDGBuilder& GraphBuilder, bool bAllowGPUParticleUpdate, const TArrayView<const FSceneViewFamily*>& ViewFamilies, const FSceneViewFamily* CurrentFamily) override;
 	virtual void PostInitViews(FRDGBuilder& GraphBuilder, TConstStridedView<FSceneView> Views, bool bAllowGPUParticleUpdate) override;
 	virtual bool UsesGlobalDistanceField() const override;
 	virtual bool UsesDepthBuffer() const override;
diff --git a/Engine/Source/Runtime/Engine/Private/Particles/ParticleGpuSimulation.cpp b/Engine/Source/Runtime/Engine/Private/Particles/ParticleGpuSimulation.cpp
--- a/Engine/Source/Runtime/Engine/Private/Particles/ParticleGpuSimulation.cpp
+++ b/Engine/Source/Runtime/Engine/Private/Particles/ParticleGpuSimulation.cpp
@@ -4469,7 +4469,8 @@ void FFXSystem::SimulateGPUParticles(
 		// On some platforms, the textures are filled with garbage after creation, so we need to clear them to black the first time we use them
 		if ( !CurrentStateTextures.bTexturesCleared )
 		{
-			
+			SCOPED_GPU_MASK(RHICmdList, FRHIGPUMask::All());
+
 			FRHIRenderPassInfo RenderPassInfo(2, CurrentStateRenderTargets, ERenderTargetActions::Clear_Store);
 			RHICmdList.BeginRenderPass(RenderPassInfo, TEXT("GPUParticlesClearStateTextures"));
 			RHICmdList.EndRenderPass();
@@ -4479,6 +4480,8 @@ void FFXSystem::SimulateGPUParticles(
 		
 		if ( !PrevStateTextures.bTexturesCleared )
 		{
+			SCOPED_GPU_MASK(RHICmdList, FRHIGPUMask::All());
+
 			RHICmdList.Transition({
 				FRHITransitionInfo(PreviousStateRenderTargets[0], ERHIAccess::SRVMask, ERHIAccess::RTV),
 				FRHITransitionInfo(PreviousStateRenderTargets[1], ERHIAccess::SRVMask, ERHIAccess::RTV)});
@@ -4792,20 +4795,20 @@ void FFXSystem::SimulateGPUParticles(
 			FRHITransitionInfo(VisualizeStateTextures.PositionTextureRHI, ERHIAccess::RTV, ERHIAccess::SRVMask),
 			FRHITransitionInfo(VisualizeStateTextures.VelocityTextureRHI, ERHIAccess::RTV, ERHIAccess::SRVMask)
 		});
-	}
 
 #if WITH_MGPU
-	// TODO:  This transfer was done here due to AFR considerations, but AFR support has been removed.
-	// Should investigate whether it still needs to be done here, or would be more optimal to move
-	// somewhere else?
-	if (Phase == PhaseToBroadcastResourceTransfer)
-	{
-		if (CrossGPUTransferResources.Num() > 0)
+		// TODO:  This transfer was done here due to AFR considerations, but AFR support has been removed.
+		// Should investigate whether it still needs to be done here, or would be more optimal to move
+		// somewhere else?
+		if (Phase == PhaseToBroadcastResourceTransfer)
 		{
-			RHICmdList.TransferResources(CrossGPUTransferResources);
+			if (CrossGPUTransferResources.Num() > 0)
+			{
+				RHICmdList.TransferResources(CrossGPUTransferResources);
+			}
 		}
-	}
 #endif
+	}
 
 	// Stats.
 	if (Phase == GetLastParticleSimulationPhase(GetShaderPlatform()))
diff --git a/Engine/Source/Runtime/Engine/Public/FXSystem.h b/Engine/Source/Runtime/Engine/Public/FXSystem.h
--- a/Engine/Source/Runtime/Engine/Public/FXSystem.h
+++ b/Engine/Source/Runtime/Engine/Public/FXSystem.h
@@ -242,7 +242,7 @@ public:
 	 * Notification from the renderer that it is about to perform visibility
 	 * checks on FX belonging to this system.
 	 */
-	virtual void PreInitViews(FRDGBuilder& GraphBuilder, bool bAllowGPUParticleUpdate) = 0;
+	virtual void PreInitViews(FRDGBuilder& GraphBuilder, bool bAllowGPUParticleUpdate, const TArrayView<const FSceneViewFamily*>& ViewFamilies, const FSceneViewFamily* CurrentFamily) = 0;
 
 	virtual void PostInitViews(FRDGBuilder& GraphBuilder, TConstStridedView<FSceneView> Views, bool bAllowGPUParticleUpdate) = 0;
 
diff --git a/Engine/Source/Runtime/RHI/Private/RHIContext.cpp b/Engine/Source/Runtime/RHI/Private/RHIContext.cpp
--- a/Engine/Source/Runtime/RHI/Private/RHIContext.cpp
+++ b/Engine/Source/Runtime/RHI/Private/RHIContext.cpp
@@ -10,6 +10,48 @@ void IRHIComputeContext::StatsSetCategory(FRHIDrawStats* InStats, uint32 InCateg
 	Stats = &InStats->GetGPU(InGPUIndex).GetCategory(InCategoryID);
 }
 
+void RHIGenerateCrossGPUPreTransferFences(const TArrayView<const FTransferResourceParams> Params, TArray<FCrossGPUTransferFence*>& OutPreTransfer)
+{
+	// Generate destination GPU masks by source GPU
+	uint32 DestGPUMasks[MAX_NUM_GPUS];
+	for (uint32 GPUIndex = 0; GPUIndex < GNumExplicitGPUsForRendering; GPUIndex++)
+	{
+		DestGPUMasks[GPUIndex] = 0;
+	}
+
+	for (const FTransferResourceParams& Param : Params)
+	{
+		check(Param.SrcGPUIndex != Param.DestGPUIndex && Param.SrcGPUIndex < GNumExplicitGPUsForRendering&& Param.DestGPUIndex < GNumExplicitGPUsForRendering);
+		DestGPUMasks[Param.SrcGPUIndex] |= 1u << Param.DestGPUIndex;
+	}
+
+	// Count total number of bits in all the masks, and allocate that number of fences
+	uint32 NumFences = 0;
+	for (uint32 GPUIndex = 0; GPUIndex < GNumExplicitGPUsForRendering; GPUIndex++)
+	{
+		NumFences += FMath::CountBits(DestGPUMasks[GPUIndex]);
+	}
+	OutPreTransfer.SetNumUninitialized(NumFences);
+
+	// Allocate and initialize fence data
+	uint32 FenceIndex = 0;
+	for (uint32 SrcGPUIndex = 0; SrcGPUIndex < GNumExplicitGPUsForRendering; SrcGPUIndex++)
+	{
+		uint32 DestGPUMask = DestGPUMasks[SrcGPUIndex];
+		for (uint32 DestGPUIndex = FMath::CountTrailingZeros(DestGPUMask); DestGPUMask; DestGPUIndex = FMath::CountTrailingZeros(DestGPUMask))
+		{
+			FCrossGPUTransferFence* TransferSyncPoint = RHICreateCrossGPUTransferFence();
+			TransferSyncPoint->SignalGPUIndex = DestGPUIndex;
+			TransferSyncPoint->WaitGPUIndex = SrcGPUIndex;
+
+			OutPreTransfer[FenceIndex] = TransferSyncPoint;
+			FenceIndex++;
+
+			DestGPUMask &= ~(1u << DestGPUIndex);
+		}
+	}
+}
+
 void IRHICommandContextPSOFallback::SetGraphicsPipelineStateFromInitializer(const FGraphicsPipelineStateInitializer& PsoInit, uint32 StencilRef, bool bApplyAdditionalState)
 {
 	RHISetBoundShaderState(
diff --git a/Engine/Source/Runtime/RHI/Public/RHICommandList.h b/Engine/Source/Runtime/RHI/Public/RHICommandList.h
--- a/Engine/Source/Runtime/RHI/Public/RHICommandList.h
+++ b/Engine/Source/Runtime/RHI/Public/RHICommandList.h
@@ -1119,6 +1119,44 @@ FRHICOMMAND_MACRO(FRHICommandTransferResourceWait)
 	RHI_API void Execute(FRHICommandListBase & CmdList);
 };
 
+FRHICOMMAND_MACRO(FRHICommandCrossGPUTransfer)
+{
+	TArray<FTransferResourceParams, TInlineAllocator<4>> Params;
+	TArray<FCrossGPUTransferFence*, TInlineAllocator<1>> PreTransfer;
+	TArray<FCrossGPUTransferFence*, TInlineAllocator<1>> PostTransfer;
+
+	FORCEINLINE_DEBUGGABLE FRHICommandCrossGPUTransfer(TArrayView<const FTransferResourceParams> InParams, const TArrayView<FCrossGPUTransferFence* const> InPreTransfer, const TArrayView<FCrossGPUTransferFence* const> InPostTransfer)
+		: Params(InParams), PreTransfer(InPreTransfer), PostTransfer(InPostTransfer)
+	{
+	}
+
+	RHI_API void Execute(FRHICommandListBase & CmdList);
+};
+
+FRHICOMMAND_MACRO(FRHICommandCrossGPUTransferSignal)
+{
+	TArray<FTransferResourceParams, TInlineAllocator<4>> Params;
+	TArray<FCrossGPUTransferFence*, TInlineAllocator<1>> PreTransfer;
+
+	FORCEINLINE_DEBUGGABLE FRHICommandCrossGPUTransferSignal(TArrayView<const FTransferResourceParams> InParams, TArrayView<FCrossGPUTransferFence* const> InPreTransfer)
+		: Params(InParams), PreTransfer(InPreTransfer)
+	{
+	}
+
+	RHI_API void Execute(FRHICommandListBase & CmdList);
+};
+
+FRHICOMMAND_MACRO(FRHICommandCrossGPUTransferWait)
+{
+	TArray<FCrossGPUTransferFence*, TInlineAllocator<4>> SyncPoints;
+
+	FORCEINLINE_DEBUGGABLE FRHICommandCrossGPUTransferWait(TArrayView<FCrossGPUTransferFence* const> InSyncPoints)
+		: SyncPoints(InSyncPoints)
+	{
+	}
+
+	RHI_API void Execute(FRHICommandListBase & CmdList);
+};
 #endif // WITH_MGPU
 
 FRHICOMMAND_MACRO(FRHICommandSetStencilRef)
@@ -2595,6 +2633,48 @@ PRAGMA_ENABLE_DEPRECATION_WARNINGS
 #endif // WITH_MGPU
 	}
 
+	FORCEINLINE_DEBUGGABLE void CrossGPUTransfer(const TArrayView<const FTransferResourceParams> Params, const TArrayView<FCrossGPUTransferFence* const> PreTransfer, const TArrayView<FCrossGPUTransferFence* const> PostTransfer)
+	{
+#if WITH_MGPU
+		if (Bypass())
+		{
+			GetComputeContext().RHICrossGPUTransfer(Params, PreTransfer, PostTransfer);
+		}
+		else
+		{
+			ALLOC_COMMAND(FRHICommandCrossGPUTransfer)(Params, PreTransfer, PostTransfer);
+		}
+#endif // WITH_MGPU
+	}
+
+	FORCEINLINE_DEBUGGABLE void CrossGPUTransferSignal(const TArrayView<const FTransferResourceParams> Params, const TArrayView<FCrossGPUTransferFence* const> PreTransfer)
+	{
+#if WITH_MGPU
+		if (Bypass())
+		{
+			GetComputeContext().RHICrossGPUTransferSignal(Params, PreTransfer);
+		}
+		else
+		{
+			ALLOC_COMMAND(FRHICommandCrossGPUTransferSignal)(Params, PreTransfer);
+		}
+#endif // WITH_MGPU
+	}
+
+	FORCEINLINE_DEBUGGABLE void CrossGPUTransferWait(const TArrayView<FCrossGPUTransferFence* const> SyncPoints)
+	{
+#if WITH_MGPU
+		if (Bypass())
+		{
+			GetComputeContext().RHICrossGPUTransferWait(SyncPoints);
+		}
+		else
+		{
+			ALLOC_COMMAND(FRHICommandCrossGPUTransferWait)(SyncPoints);
+		}
+#endif // WITH_MGPU
+	}
+
 #if WITH_MGPU
 	UE_DEPRECATED(5.2, "AFR support has been removed in 5.2, function is a nop and calls should be removed.")
 	FORCEINLINE void WaitForTemporalEffect(const FName& EffectName) {}
diff --git a/Engine/Source/Runtime/RHI/Public/RHICommandListCommandExecutes.inl b/Engine/Source/Runtime/RHI/Public/RHICommandListCommandExecutes.inl
--- a/Engine/Source/Runtime/RHI/Public/RHICommandListCommandExecutes.inl
+++ b/Engine/Source/Runtime/RHI/Public/RHICommandListCommandExecutes.inl
@@ -95,6 +95,24 @@ void FRHICommandTransferResourceWait::Execute(FRHICommandListBase& CmdList)
 	INTERNAL_DECORATOR_COMPUTE(RHITransferResourceWait)(FenceDatas);
 }
 
+void FRHICommandCrossGPUTransfer::Execute(FRHICommandListBase& CmdList)
+{
+	RHISTAT(CrossGPUTransfer);
+	INTERNAL_DECORATOR_COMPUTE(RHICrossGPUTransfer)(Params, PreTransfer, PostTransfer);
+}
+
+void FRHICommandCrossGPUTransferSignal::Execute(FRHICommandListBase& CmdList)
+{
+	RHISTAT(CrossGPUTransferSignal);
+	INTERNAL_DECORATOR_COMPUTE(RHICrossGPUTransferSignal)(Params, PreTransfer);
+}
+
+void FRHICommandCrossGPUTransferWait::Execute(FRHICommandListBase& CmdList)
+{
+	RHISTAT(CrossGPUTransferWait);
+	INTERNAL_DECORATOR_COMPUTE(RHICrossGPUTransferWait)(SyncPoints);
+}
+
 #endif // WITH_MGPU
 
 void FRHICommandSetStencilRef::Execute(FRHICommandListBase& CmdList)
diff --git a/Engine/Source/Runtime/RHI/Public/RHIContext.h b/Engine/Source/Runtime/RHI/Public/RHIContext.h
--- a/Engine/Source/Runtime/RHI/Public/RHIContext.h
+++ b/Engine/Source/Runtime/RHI/Public/RHIContext.h
@@ -133,6 +133,15 @@ struct FTransferResourceFenceData
 	{}
 };
 
+struct FCrossGPUTransferFence
+{
+	uint32 SignalGPUIndex = 0;
+	uint32 WaitGPUIndex = 0;
+	void* SyncPoint = nullptr;
+
+	FCrossGPUTransferFence() = default;
+};
+
 FORCEINLINE FTransferResourceFenceData* RHICreateTransferResourceFenceData()
 {
 #if WITH_MGPU
@@ -142,6 +151,15 @@ FORCEINLINE FTransferResourceFenceData* RHICreateTransferResourceFenceData()
 #endif
 }
 
+FORCEINLINE FCrossGPUTransferFence* RHICreateCrossGPUTransferFence()
+{
+#if WITH_MGPU
+	return new FCrossGPUTransferFence;
+#else
+	return nullptr;
+#endif
+}
+
 /** Parameters for RHITransferResources, used to copy memory between GPUs */
 struct FTransferResourceParams
 {
@@ -408,6 +426,39 @@ public:
 #endif
 	}
 
+	/**
+	 * Synchronizes the content of a resource between two or more GPUs using a copy operation -- variation of above that includes separate arrays of fences.
+	 * @param Params - the parameters for each resource or texture region copied between GPUs.
+	 * @param PreTransfer - Fences to wait on before copying the relevant data (initialized with RHITransferResourceSignal before this function)
+	 * @param PostTransfer - Fences that can be waited on after copy (waited on by RHITransferResourceWait after this function)
+	 */
+	virtual void RHICrossGPUTransfer(const TArrayView<const FTransferResourceParams> Params, const TArrayView<FCrossGPUTransferFence* const> PreTransfer, const TArrayView<FCrossGPUTransferFence* const> PostTransfer)
+	{
+		/** empty default implementation. */
+	}
+
+	virtual void RHICrossGPUTransferSignal(const TArrayView<const FTransferResourceParams> Params, const TArrayView<FCrossGPUTransferFence* const> PreTransfer)
+	{
+		/* default noop implementation */
+#if WITH_MGPU
+		for (FCrossGPUTransferFence* SyncPoint : PreTransfer)
+		{
+			delete SyncPoint;
+		}
+#endif
+	}
+
+	virtual void RHICrossGPUTransferWait(const TArrayView<FCrossGPUTransferFence* const> SyncPoints)
+	{
+		/* default noop implementation */
+#if WITH_MGPU
+		for (FCrossGPUTransferFence* SyncPoint : SyncPoints)
+		{
+			delete SyncPoint;
+		}
+#endif
+	}
+
 	virtual void RHIBuildAccelerationStructures(const TArrayView<const FRayTracingGeometryBuildParams> Params, const FRHIBufferRange& ScratchBufferRange)
 	{
 		checkNoEntry();
@@ -479,6 +530,9 @@ public:
 	}
 };
 
+// Utility function to generate pre-transfer sync points to pass to CrossGPUTransferSignal and CrossGPUTransfer
+RHI_API void RHIGenerateCrossGPUPreTransferFences(const TArrayView<const FTransferResourceParams> Params, TArray<FCrossGPUTransferFence*>& OutPreTransfer);
+
 enum class EAccelerationStructureBuildMode
 {
 	// Perform a full acceleration structure build.
diff --git a/Engine/Source/Runtime/Renderer/Private/DeferredShadingRenderer.cpp b/Engine/Source/Runtime/Renderer/Private/DeferredShadingRenderer.cpp
--- a/Engine/Source/Runtime/Renderer/Private/DeferredShadingRenderer.cpp
+++ b/Engine/Source/Runtime/Renderer/Private/DeferredShadingRenderer.cpp
@@ -2838,7 +2838,7 @@ void FDeferredShadingSceneRenderer::Render(FRDGBuilder& GraphBuilder)
 	{
 		SCOPE_CYCLE_COUNTER(STAT_FDeferredShadingSceneRenderer_FXSystem_PreRender);
 		GraphBuilder.SetCommandListStat(GET_STATID(STAT_CLM_FXPreRender));
-		FXSystem->PreRender(GraphBuilder, GetSceneViews(), true /*bAllowGPUParticleUpdate*/);
+		FXSystem->PreRender(GraphBuilder, GetSceneViews(), bIsFirstSceneRenderer /*bAllowGPUParticleUpdate*/);
 		if (FGPUSortManager* GPUSortManager = FXSystem->GetGPUSortManager())
 		{
 			GPUSortManager->OnPreRender(GraphBuilder);
diff --git a/Engine/Source/Runtime/Renderer/Private/SceneRendering.cpp b/Engine/Source/Runtime/Renderer/Private/SceneRendering.cpp
--- a/Engine/Source/Runtime/Renderer/Private/SceneRendering.cpp
+++ b/Engine/Source/Runtime/Renderer/Private/SceneRendering.cpp
@@ -3146,7 +3146,7 @@ void FSceneRenderer::PreallocateCrossGPUFences(const TArray<FSceneRenderer*>& Sc
 	if (SceneRenderers.Num() > 1)
 	{
 		// Allocated fences to wait on are placed in the last scene renderer
-		TArray<FTransferResourceFenceData*>& LastRendererFencesWait = SceneRenderers.Last()->CrossGPUTransferFencesWait;
+		TArray<FCrossGPUTransferFence*>& LastRendererFencesWait = SceneRenderers.Last()->CrossGPUTransferFencesWait;
 
 		check(LastRendererFencesWait.IsEmpty());
 
@@ -3161,17 +3161,21 @@ void FSceneRenderer::PreallocateCrossGPUFences(const TArray<FSceneRenderer*>& Sc
 
 			if (SceneRenderer->ViewFamily.bMultiGPUForkAndJoin)
 			{
-				TArray<FCrossGPUTransfer> Transfers;
-				GetCrossGPUTransfers(SceneRenderer, Transfers);
+				// Check if we can do optimized transfers, which requires a single index
+				if (SceneRenderer->AllViewsGPUMask.HasSingleIndex())
+				{
+					TArray<FCrossGPUTransfer> Transfers;
+					GetCrossGPUTransfers(SceneRenderer, Transfers);
 
-				SceneRenderer->CrossGPUTransferFencesDefer.SetNumUninitialized(Transfers.Num());
+					SceneRenderer->CrossGPUTransferFencesDefer.SetNumUninitialized(Transfers.Num());
 
-				for (int32 TransferIndex = 0; TransferIndex < Transfers.Num(); TransferIndex++)
-				{
-					FTransferResourceFenceData* FenceData = RHICreateTransferResourceFenceData();
+					for (int32 TransferIndex = 0; TransferIndex < Transfers.Num(); TransferIndex++)
+					{
+						FCrossGPUTransferFence* FenceData = RHICreateCrossGPUTransferFence();
 
-					SceneRenderer->CrossGPUTransferFencesDefer[TransferIndex] = FenceData;
-					LastRendererFencesWait.Add(FenceData);
+						SceneRenderer->CrossGPUTransferFencesDefer[TransferIndex] = FenceData;
+						LastRendererFencesWait.Add(FenceData);
+					}
 				}
 			}
 		}
@@ -3195,30 +3199,43 @@ void FSceneRenderer::DoCrossGPUTransfers(FRDGBuilder& GraphBuilder, FRDGTextureR
 
 		if (Transfers.Num() > 0)
 		{
-			// Optionally delay cross GPU transfer fence wait
+			// Check if we can go through the optimized code path, with delay for the cross GPU transfer fence wait
 			if (CrossGPUTransferFencesDefer.Num())
 			{
-				check(CrossGPUTransferFencesDefer.Num() == Transfers.Num());
-
-				for (int32 TransferIndex = 0; TransferIndex < Transfers.Num(); TransferIndex++)
+				// A readback pass is the closest analog to what this is doing. There isn't a way to express cross-GPU transfers via the RHI barrier API.
+				AddReadbackTexturePass(GraphBuilder, RDG_EVENT_NAME("CrossGPUTransfers"), ViewFamilyTexture,
+					[this, ViewFamilyTexture, LocalTransfers = MoveTemp(Transfers), PostTransferFences = MoveTemp(CrossGPUTransferFencesDefer)](FRHICommandListImmediate& RHICmdList)
 				{
-					Transfers[TransferIndex].DelayedFence = CrossGPUTransferFencesDefer[TransferIndex];
-				}
-			}
+					TArray<FTransferResourceParams> TransferParams;
+					for (const FCrossGPUTransfer& Transfer : LocalTransfers)
+					{
+						TransferParams.Add(FTransferResourceParams(ViewFamilyTexture->GetRHI(), Transfer.SrcGPUIndex, Transfer.DestGPUIndex, false, false));
+					}
 
-			// A readback pass is the closest analog to what this is doing. There isn't a way to express cross-GPU transfers via the RHI barrier API.
-			AddReadbackTexturePass(GraphBuilder, RDG_EVENT_NAME("CrossGPUTransfers"), ViewFamilyTexture,
-				[this, ViewFamilyTexture, LocalTransfers = MoveTemp(Transfers)](FRHICommandListImmediate& RHICmdList)
+					// Transition resources on destination GPU and signal when transition has finished
+					TArray<FCrossGPUTransferFence*> PreTransferFences;
+					RHIGenerateCrossGPUPreTransferFences(TransferParams, PreTransferFences);
+					RHICmdList.CrossGPUTransferSignal(TransferParams, PreTransferFences);
+
+					// Then do the actual transfer
+					RHICmdList.CrossGPUTransfer(TransferParams, PreTransferFences, PostTransferFences);
+				});
+			}
+			else
+			{
+				// A readback pass is the closest analog to what this is doing. There isn't a way to express cross-GPU transfers via the RHI barrier API.
+				AddReadbackTexturePass(GraphBuilder, RDG_EVENT_NAME("CrossGPUTransfers"), ViewFamilyTexture,
+					[this, ViewFamilyTexture, LocalTransfers = MoveTemp(Transfers)](FRHICommandListImmediate& RHICmdList)
 				{
 					TArray<FTransferResourceParams> TransferParams;
 					for (const FCrossGPUTransfer& Transfer : LocalTransfers)
 					{
 						TransferParams.Add(FTransferResourceParams(ViewFamilyTexture->GetRHI(), Transfer.SrcGPUIndex, Transfer.DestGPUIndex, false, false));
-						TransferParams.Last().DelayedFence = Transfer.DelayedFence;
 					}
 
 					RHICmdList.TransferResources(TransferParams);
 				});
+			}
 		}
 	}
 #endif // WITH_MGPU
@@ -3234,7 +3251,7 @@ void FSceneRenderer::FlushCrossGPUFences(FRDGBuilder& GraphBuilder)
 		AddPass(GraphBuilder, RDG_EVENT_NAME("CrossGPUTransferSync"),
 			[LocalFenceDatas = MoveTemp(CrossGPUTransferFencesWait)](FRHICommandListImmediate& RHICmdList)
 		{
-			RHICmdList.TransferResourceWait(LocalFenceDatas);
+			RHICmdList.CrossGPUTransferWait(LocalFenceDatas);
 		});
 	}
 #endif // WITH_MGPU
@@ -4220,17 +4237,19 @@ void FSceneRenderer::RenderThreadBegin(FRHICommandListImmediate& RHICmdList, con
 {
 	CleanUp(RHICmdList);
 
-	// Initialize "AllFamilyViews" array in each scene renderer.  Initialize array in first SceneRenderer, then copy to remaining.
+	// Initialize "AllFamilyViews" and "AllFamilies" array in each scene renderer.  Initialize array in first SceneRenderer, then copy to remaining.
 	for (FSceneRenderer* SceneRenderer : SceneRenderers)
 	{
 		for (const FViewInfo& ViewInfo : SceneRenderer->Views)
 		{
 			SceneRenderers[0]->AllFamilyViews.Add(&ViewInfo);
 		}
+		SceneRenderers[0]->AllFamilies.Add(&SceneRenderer->ViewFamily);
 	}
 	for (int32 RendererIndex = 1; RendererIndex < SceneRenderers.Num(); RendererIndex++)
 	{
 		SceneRenderers[RendererIndex]->AllFamilyViews = SceneRenderers[0]->AllFamilyViews;
+		SceneRenderers[RendererIndex]->AllFamilies = SceneRenderers[0]->AllFamilies;
 	}
 
 	FScene* Scene = SceneRenderers[0]->Scene;
diff --git a/Engine/Source/Runtime/Renderer/Private/SceneRendering.h b/Engine/Source/Runtime/Renderer/Private/SceneRendering.h
--- a/Engine/Source/Runtime/Renderer/Private/SceneRendering.h
+++ b/Engine/Source/Runtime/Renderer/Private/SceneRendering.h
@@ -2080,6 +2080,7 @@ public:
 
 	/** Views across all view families (may contain additional views if rendering multiple families together). */
 	TArray<const FSceneView*> AllFamilyViews;
+	TArray<const FSceneViewFamily*> AllFamilies;
 
 	/** All the dynamic scaling informations */
 	DynamicRenderScaling::TMap<float> DynamicResolutionFractions;
@@ -2329,11 +2330,12 @@ protected:
 	/**
 	 * Fences for cross GPU render target transfers.  We defer the wait on cross GPU fences until the last scene renderer,
 	 * to avoid needless stalls in the middle of the frame, improving performance.  The "Defer" array holds fences issued
-	 * by each prior scene renderer, while the "Wait" array holds fences to be waited on in the last scene renderer.
-	 * The function "PreallocateCrossGPUFences" initializes these arrays.
+	 * by each prior scene renderer, while the "Wait" array holds fences to be waited on in the last scene renderer
+	 * (a collection of all the fences from prior scene renderers).  The function "PreallocateCrossGPUFences" initializes
+	 * these arrays.
 	 */
-	TArray<FTransferResourceFenceData*> CrossGPUTransferFencesDefer;
-	TArray<FTransferResourceFenceData*> CrossGPUTransferFencesWait;
+	TArray<FCrossGPUTransferFence*> CrossGPUTransferFencesDefer;
+	TArray<FCrossGPUTransferFence*> CrossGPUTransferFencesWait;
 
 	FRHIGPUMask AllViewsGPUMask;
 	bool IsShadowCached(FProjectedShadowInfo* ProjectedShadowInfo) const;
diff --git a/Engine/Source/Runtime/Renderer/Private/SceneVisibility.cpp b/Engine/Source/Runtime/Renderer/Private/SceneVisibility.cpp
--- a/Engine/Source/Runtime/Renderer/Private/SceneVisibility.cpp
+++ b/Engine/Source/Runtime/Renderer/Private/SceneVisibility.cpp
@@ -4008,7 +4008,7 @@ void FSceneRenderer::PreVisibilityFrameSetup(FRDGBuilder& GraphBuilder)
 
 	if (FXSystem && Views.IsValidIndex(0))
 	{
-		FXSystem->PreInitViews(GraphBuilder, Views[0].AllowGPUParticleUpdate() && !ViewFamily.EngineShowFlags.HitProxies);
+		FXSystem->PreInitViews(GraphBuilder, Views[0].AllowGPUParticleUpdate() && !ViewFamily.EngineShowFlags.HitProxies, AllFamilies, &ViewFamily);
 	}
 
 #if WITH_EDITOR
