Better UX for exposure options of input node
* improved display names and tooltips
* better default values for things like data interfaces
* hide when usage of parent node is not compatible

diff --git a/Engine/Plugins/FX/Niagara/Source/NiagaraEditor/Private/NiagaraEditorUtilities.cpp b/Engine/Plugins/FX/Niagara/Source/NiagaraEditor/Private/NiagaraEditorUtilities.cpp
--- a/Engine/Plugins/FX/Niagara/Source/NiagaraEditor/Private/NiagaraEditorUtilities.cpp
+++ b/Engine/Plugins/FX/Niagara/Source/NiagaraEditor/Private/NiagaraEditorUtilities.cpp
@@ -132,6 +132,12 @@ void FNiagaraEditorUtilities::InitializeParameterInputNode(UNiagaraNodeInput& In
 	InputName = UNiagaraNodeInput::GenerateUniqueName(InGraph, InputName, ENiagaraInputNodeUsage::Parameter);
 	InputNode.Input.SetName(InputName);
 	InputNode.Input.SetType(Type);
+	if (Type != FNiagaraTypeDefinition::GetParameterMapDef())
+	{
+		// when adding things like data interfaces as module inputs, making them required will just generate weird compile errors
+		// because "bRequired" prevents using the default value, but their pins are not exposed to the stack (as they are not part of the parameter map)
+		InputNode.ExposureOptions.bRequired = false;
+	}
 	if (InGraph) // Only compute sort priority if a graph was passed in, similar to the way that GenrateUniqueName works above.
 	{
 		InputNode.CallSortPriority = UNiagaraNodeInput::GenerateNewSortPriority(InGraph, InputName, ENiagaraInputNodeUsage::Parameter);
diff --git a/Engine/Plugins/FX/Niagara/Source/NiagaraEditor/Private/Widgets/SNiagaraGraphNodeInput.cpp b/Engine/Plugins/FX/Niagara/Source/NiagaraEditor/Private/Widgets/SNiagaraGraphNodeInput.cpp
--- a/Engine/Plugins/FX/Niagara/Source/NiagaraEditor/Private/Widgets/SNiagaraGraphNodeInput.cpp
+++ b/Engine/Plugins/FX/Niagara/Source/NiagaraEditor/Private/Widgets/SNiagaraGraphNodeInput.cpp
@@ -30,10 +30,10 @@ public:
 
 	virtual void RegisterCommands() override
 	{
-		UI_COMMAND(ToggleExposed, "Exposed", "Toggles whether or not this parameter is exposed.", EUserInterfaceActionType::ToggleButton, FInputChord());
-		UI_COMMAND(ToggleRequired, "Required", "Toggles whether or not this parameter is required.", EUserInterfaceActionType::ToggleButton, FInputChord());
+		UI_COMMAND(ToggleExposed, "Expose Input", "If true then this is exposed as an input to it's caller - turning this off makes sense when the input is defined in a function script and should not be visible to the caller.", EUserInterfaceActionType::ToggleButton, FInputChord());
+		UI_COMMAND(ToggleRequired, "Suppress Default Value", "If true then this input is required to be set by the caller. For optional values (e.g. values behind an edit condition) this can be set to false, so the translator uses the default value instead.", EUserInterfaceActionType::ToggleButton, FInputChord());
 		UI_COMMAND(ToggleCanAutoBind, "CanAutoBind", "Toggles whether or not this parameter can auto-bind to system parameters and emitter attributes. Will never auto bind to custom parameters.", EUserInterfaceActionType::ToggleButton, FInputChord());
-		UI_COMMAND(ToggleHidden, "Hidden", "Toggles whether or not this parameter is hidden.", EUserInterfaceActionType::ToggleButton, FInputChord());
+		UI_COMMAND(ToggleHidden, "Advanced Display", "If this input should be hidden in the advanced pin section of it's caller.", EUserInterfaceActionType::ToggleButton, FInputChord());
 	}
 
 public:
@@ -284,7 +284,6 @@ TSharedRef<SWidget> SNiagaraGraphNodeInput::GenerateExposureOptionsMenu() const
 		{
 			OptionsMenuBuilder.AddMenuEntry(FNiagaraGraphNodeInputCommands::Get().ToggleExposed);
 			OptionsMenuBuilder.AddMenuEntry(FNiagaraGraphNodeInputCommands::Get().ToggleRequired);
-			OptionsMenuBuilder.AddMenuEntry(FNiagaraGraphNodeInputCommands::Get().ToggleCanAutoBind);
 			OptionsMenuBuilder.AddMenuEntry(FNiagaraGraphNodeInputCommands::Get().ToggleHidden);
 		}
 
@@ -305,7 +304,8 @@ void SNiagaraGraphNodeInput::ExposureOptionsMenuOpenChanged(bool bOpened)
 
 EVisibility SNiagaraGraphNodeInput::GetExposureOptionsVisibility() const
 {
-	return EVisibility::Visible;
+	UNiagaraNodeInput* InputNode = CastChecked<UNiagaraNodeInput>(GraphNode);
+	return InputNode->Usage == ENiagaraInputNodeUsage::Parameter ? EVisibility::Visible : EVisibility::Collapsed;
 }
 
 FReply SNiagaraGraphNodeInput::HandleExposureOptionsMenuButtonClicked()
diff --git a/Engine/Plugins/FX/Niagara/Source/NiagaraEditor/Public/NiagaraNodeInput.h b/Engine/Plugins/FX/Niagara/Source/NiagaraEditor/Public/NiagaraNodeInput.h
--- a/Engine/Plugins/FX/Niagara/Source/NiagaraEditor/Public/NiagaraNodeInput.h
+++ b/Engine/Plugins/FX/Niagara/Source/NiagaraEditor/Public/NiagaraNodeInput.h
@@ -25,20 +25,20 @@ struct FNiagaraInputExposureOptions
 		, bHidden(0)
 	{}
 
-	/** If this input is exposed to it's caller. */
-	UPROPERTY(EditAnywhere, Category = Exposure)
+	/** If true then this is exposed as an input to it's caller - turning this off makes sense when the input is defined in a function script and should not be visible to the caller. */
+	UPROPERTY(EditAnywhere, DisplayName="Expose Input", Category = Exposure)
 	uint32 bExposed : 1;
 
-	/** If this input is required to be bound. */
-	UPROPERTY(EditAnywhere, Category = Exposure, meta = (editcondition = "bExposed"))
+	/** If true then this input is required to be set by the caller. For optional values (e.g. values behind an edit condition) this can be set to false, so the translator uses the default value instead. */
+	UPROPERTY(EditAnywhere, DisplayName="Suppress Default Value", Category = Exposure, meta = (editcondition = "bExposed"))
 	uint32 bRequired : 1;
 
 	/** If this input can auto-bind to system parameters and emitter attributes. Will never auto bind to custom parameters. */
-	UPROPERTY(EditAnywhere, Category = Exposure, meta = (editcondition = "bExposed"))
+	UPROPERTY()
 	uint32 bCanAutoBind : 1;
 
-	/** If this input should be hidden in the advanced pin section of it's caller. */
-	UPROPERTY(EditAnywhere, Category = Exposure, meta = (editcondition = "bExposed"))
+	/** If true then this input will be shown in the advanced pin section of the caller node. Only works in function scripts, does nothing in dynamic inputs or module scripts. */
+	UPROPERTY(EditAnywhere, DisplayName="Advanced Display", Category = Exposure, meta = (editcondition = "bExposed"))
 	uint32 bHidden : 1;
 };
 
@@ -60,7 +60,7 @@ public:
 	int32 CallSortPriority;
 
 	/** Controls this inputs exposure to callers. */
-	UPROPERTY(EditAnywhere, Category = Input)
+	UPROPERTY(EditAnywhere, Category = Input, meta=(EditCondition="Usage == ENiagaraInputNodeUsage::Parameter", EditConditionHides))
 	FNiagaraInputExposureOptions ExposureOptions;
 
 	//~ Begin UObject Interface
@@ -85,10 +85,10 @@ public:
 
 	virtual void Compile(class FHlslNiagaraTranslator* Translator, TArray<int32>& Outputs)override;
 	
-	bool IsExposed()const { return ExposureOptions.bExposed; }
-	bool IsRequired()const { return ExposureOptions.bExposed && ExposureOptions.bRequired; }
-	bool IsHidden()const { return ExposureOptions.bExposed && ExposureOptions.bHidden; }
-	bool CanAutoBind()const { return ExposureOptions.bExposed && ExposureOptions.bCanAutoBind; }
+	bool IsExposed()const { return ExposureOptions.bExposed && Usage == ENiagaraInputNodeUsage::Parameter; }
+	bool IsRequired()const { return IsExposed() && ExposureOptions.bRequired; }
+	bool IsHidden()const { return IsExposed() && ExposureOptions.bHidden; }
+	bool CanAutoBind()const { return IsExposed() && ExposureOptions.bCanAutoBind; }
 
 	bool ReferencesSameInput(UNiagaraNodeInput* Other) const;
 
