- Add support for setting objects in the stack

diff --git a/Engine/Plugins/FX/Niagara/Source/Niagara/Classes/NiagaraScript.h b/Engine/Plugins/FX/Niagara/Source/Niagara/Classes/NiagaraScript.h
--- a/Engine/Plugins/FX/Niagara/Source/Niagara/Classes/NiagaraScript.h
+++ b/Engine/Plugins/FX/Niagara/Source/Niagara/Classes/NiagaraScript.h
@@ -482,6 +482,10 @@ public:
 	TArray<FNiagaraFunctionSignature> AdditionalExternalFunctions;
 #endif
 
+	/** Information about all the UObjects used by this script. */
+	UPROPERTY()
+	TArray<FNiagaraScriptUObjectCompileInfo> UObjectInfos;
+
 	/** Information about all data interfaces used by this script. */
 	UPROPERTY()
 	TArray<FNiagaraScriptDataInterfaceCompileInfo> DataInterfaceInfo;
@@ -1217,10 +1221,16 @@ public:
 	void SetResolvedDataInterfaces(const TArray<FNiagaraScriptResolvedDataInterfaceInfo>& InResolvedDataInterfaces) { ResolvedDataInterfaces = InResolvedDataInterfaces; }
 	void SetResolvedUserDataInterfaceBindings(const TArray<FNiagaraResolvedUserDataInterfaceBinding>& InResolvedUserDataInterfaceBindings) { ResolvedUserDataInterfaceBindings = InResolvedUserDataInterfaceBindings; }
 #endif
-
 	TConstArrayView<FNiagaraScriptResolvedDataInterfaceInfo> GetResolvedDataInterfaces() const { return MakeArrayView(ResolvedDataInterfaces); }
 	TConstArrayView<FNiagaraResolvedUserDataInterfaceBinding> GetResolvedUserDataInterfaceBindings() const { return MakeArrayView(ResolvedUserDataInterfaceBindings); }
 
+#if WITH_EDITORONLY_DATA
+	TArrayView<FNiagaraScriptUObjectCompileInfo> GetCachedDefaultUObjects() { return CachedDefaultUObjects; }
+	TConstArrayView<FNiagaraScriptUObjectCompileInfo> GetCachedDefaultUObjects() const { return CachedDefaultUObjects; }
+	void SetResolvedUObjects(TConstArrayView<FNiagaraResolvedUObjectInfo> InResolvedUObjectInfos) { ResolvedUObjectInfos = InResolvedUObjectInfos; }
+#endif
+	TConstArrayView<FNiagaraResolvedUObjectInfo> GetResolvedUObjects() const { return MakeArrayView(ResolvedUObjectInfos); }
+
 #if STATS
 	TArrayView<const TStatId> GetStatScopeIDs() const { return MakeArrayView(StatScopesIDs); }
 #elif ENABLE_STATNAMEDEVENTS
@@ -1355,6 +1365,16 @@ private:
 	UPROPERTY()
 	TArray<FNiagaraResolvedUserDataInterfaceBinding> ResolvedUserDataInterfaceBindings;
 
+#if WITH_EDITORONLY_DATA
+	// Populated from compilation process, used to generate resolved list
+	UPROPERTY()
+	TArray<FNiagaraScriptUObjectCompileInfo> CachedDefaultUObjects;
+#endif
+
+	// Resolved objects are used to populate the script parameter stores, they are 1:1 mapping from variable -> object
+	UPROPERTY()
+	TArray<FNiagaraResolvedUObjectInfo> ResolvedUObjectInfos;
+
 	static UNiagaraDataInterface* CopyDataInterface(UNiagaraDataInterface* Src, UObject* Owner);
 
 	/** Flag used to guarantee that the RT isn't accessing the FNiagaraScriptResource before cleanup. */
diff --git a/Engine/Plugins/FX/Niagara/Source/Niagara/Classes/NiagaraSystem.h b/Engine/Plugins/FX/Niagara/Source/Niagara/Classes/NiagaraSystem.h
--- a/Engine/Plugins/FX/Niagara/Source/Niagara/Classes/NiagaraSystem.h
+++ b/Engine/Plugins/FX/Niagara/Source/Niagara/Classes/NiagaraSystem.h
@@ -758,6 +758,9 @@ public:
 	/** Updates any post compile data based upon data interfaces. */
 	void OnCompiledDataInterfaceChanged();
 
+	/** Updates the system post UObject change. */
+	void OnCompiledUObjectChanged();
+
 	/** Updates the rapid iteration parameters for all scripts referenced by the system. */
 	void PrepareRapidIterationParametersForCompilation();
 #endif
@@ -789,7 +792,7 @@ private:
 	/** Helper for filling in attribute datasets per emitter. */
 	void InitEmitterDataSetCompiledData(FNiagaraDataSetCompiledData& DataSetToInit, const FNiagaraEmitterHandle& InAssociatedEmitterHandle);
 
-	void ResolveDIBindings();
+	void ResolveParameterStoreBindings();
 #endif
 
 	void ResolveScalabilitySettings();
diff --git a/Engine/Plugins/FX/Niagara/Source/Niagara/Private/NiagaraEmitter.cpp b/Engine/Plugins/FX/Niagara/Source/Niagara/Private/NiagaraEmitter.cpp
--- a/Engine/Plugins/FX/Niagara/Source/Niagara/Private/NiagaraEmitter.cpp
+++ b/Engine/Plugins/FX/Niagara/Source/Niagara/Private/NiagaraEmitter.cpp
@@ -1804,32 +1804,60 @@ void FVersionedNiagaraEmitterData::RebuildRendererBindings(const UNiagaraEmitter
 			TargetScripts.Add(SystemOwner->GetSystemSpawnScript());
 			TargetScripts.Add(SystemOwner->GetSystemUpdateScript());
 		}
+		const int32 EmitterScriptIndex = TargetScripts.Num();
 		GetScripts(TargetScripts, false, true);
 
 		RendererBindings.Empty();
 		TArrayView<const FNiagaraVariableWithOffset> Vars = TempStore.ReadParameterVariables();
 		for (const FNiagaraVariableWithOffset& Var : Vars)
 		{
-			bool bFound = false;
-			for (UNiagaraScript* Script : TargetScripts)
+			for (int iScript=0; iScript < TargetScripts.Num(); ++iScript)
 			{
-				if (Script)
+				UNiagaraScript* Script = TargetScripts[iScript];
+				if (Script == nullptr)
+				{
+					continue;
+				}
+				bool bVariableFound = false;
+
+				// Find static variables
+				if (Var.GetType().IsStatic())
 				{
 					for (const FNiagaraVariable& StaticVar : Script->GetVMExecutableData().StaticVariablesWritten)
 					{
 						if (StaticVar.GetType().IsSameBaseDefinition(Var.GetType()) && StaticVar.GetName() == Var.GetName())
 						{
-							RendererBindings.AddParameter(StaticVar);
-							bFound = true;
+							RendererBindings.AddParameter(StaticVar, true, false);
+							bVariableFound = true;
 							break;
 						}
 					}
 				}
+				// Find Resolved UObjects
+				// Note: Most parameters are pushed in as part of the DataSet -> Parameters process
+				else if (Var.GetType().IsUObject())
+				{
+					const ENiagaraSimTarget ScriptSimTarget = iScript < EmitterScriptIndex ? ENiagaraSimTarget::CPUSim : SimTarget;
+					if (const FNiagaraScriptExecutionParameterStore* ScriptParameterStore = Script->GetExecutionReadyParameterStore(ScriptSimTarget))
+					{
+						if (UObject* FoundObject = ScriptParameterStore->GetUObject(Var))
+						{
+							int32 ParameterOffset = INDEX_NONE;
+							RendererBindings.AddParameter(Var, true, false, &ParameterOffset);
+							RendererBindings.SetUObject(FoundObject, ParameterOffset);
+							bVariableFound = true;
+						}
+					}
+				}
 
-				if (bFound)
+				if (bVariableFound)
+				{
 					break;
+				}
 			}
 		}
+
+		RendererBindings.TriggerOnLayoutChanged();
 	}
 #endif
 }
diff --git a/Engine/Plugins/FX/Niagara/Source/Niagara/Private/NiagaraModule.cpp b/Engine/Plugins/FX/Niagara/Source/Niagara/Private/NiagaraModule.cpp
--- a/Engine/Plugins/FX/Niagara/Source/Niagara/Private/NiagaraModule.cpp
+++ b/Engine/Plugins/FX/Niagara/Source/Niagara/Private/NiagaraModule.cpp
@@ -1070,9 +1070,9 @@ void FNiagaraTypeDefinition::RecreateUserDefinedTypeRegistry()
 
 	FNiagaraTypeRegistry::ClearUserDefinedRegistry();
 
-	ENiagaraTypeRegistryFlags VarFlags = ENiagaraTypeRegistryFlags::AllowAnyVariable;
-	ENiagaraTypeRegistryFlags ParamFlags = VarFlags | ENiagaraTypeRegistryFlags::AllowParameter;
-	ENiagaraTypeRegistryFlags PayloadFlags = VarFlags | ENiagaraTypeRegistryFlags::AllowPayload;
+	const ENiagaraTypeRegistryFlags VarFlags = ENiagaraTypeRegistryFlags::AllowAnyVariable;
+	const ENiagaraTypeRegistryFlags ParamFlags = VarFlags | ENiagaraTypeRegistryFlags::AllowParameter;
+	const ENiagaraTypeRegistryFlags PayloadFlags = VarFlags | ENiagaraTypeRegistryFlags::AllowPayload;
 
 	FNiagaraTypeRegistry::Register(CollisionEventDef, PayloadFlags);
 
@@ -1108,12 +1108,12 @@ void FNiagaraTypeDefinition::RecreateUserDefinedTypeRegistry()
 	UScriptStruct* SpawnInfoStruct = FindObjectChecked<UScriptStruct>(NiagaraPkg, TEXT("NiagaraSpawnInfo"));
 	FNiagaraTypeRegistry::Register(FNiagaraTypeDefinition(SpawnInfoStruct), ParamFlags);
 
-	FNiagaraTypeRegistry::Register(UObjectDef, VarFlags);
-	FNiagaraTypeRegistry::Register(UMaterialDef, VarFlags);
-	FNiagaraTypeRegistry::Register(UTextureDef, VarFlags);
-	FNiagaraTypeRegistry::Register(UTextureRenderTargetDef, VarFlags);
-	FNiagaraTypeRegistry::Register(UStaticMeshDef, VarFlags);
-	FNiagaraTypeRegistry::Register(USimCacheClassDef, VarFlags);
+	FNiagaraTypeRegistry::Register(UObjectDef, ParamFlags);
+	FNiagaraTypeRegistry::Register(UMaterialDef, ParamFlags);
+	FNiagaraTypeRegistry::Register(UTextureDef, ParamFlags);
+	FNiagaraTypeRegistry::Register(UTextureRenderTargetDef, ParamFlags);
+	FNiagaraTypeRegistry::Register(UStaticMeshDef, ParamFlags);
+	FNiagaraTypeRegistry::Register(USimCacheClassDef, ParamFlags);
 	FNiagaraTypeRegistry::Register(StaticEnum<ENiagaraLegacyTrailWidthMode>(), ParamFlags);
 
 
diff --git a/Engine/Plugins/FX/Niagara/Source/Niagara/Private/NiagaraParameterCollection.cpp b/Engine/Plugins/FX/Niagara/Source/Niagara/Private/NiagaraParameterCollection.cpp
--- a/Engine/Plugins/FX/Niagara/Source/Niagara/Private/NiagaraParameterCollection.cpp
+++ b/Engine/Plugins/FX/Niagara/Source/Niagara/Private/NiagaraParameterCollection.cpp
@@ -213,6 +213,10 @@ void UNiagaraParameterCollectionInstance::SyncWithCollection()
 			{
 				ParameterStorage.SetDataInterface(OldStore.GetDataInterface(Offset), Param);
 			}
+			else if (Param.IsUObject())
+			{
+				ParameterStorage.SetUObject(OldStore.GetUObject(Offset), Param);
+			}
 			else
 			{
 				ParameterStorage.SetParameterData(OldStore.GetParameterData(Offset), ParameterStorageOffset, Param.GetSizeInBytes());
@@ -231,6 +235,10 @@ void UNiagaraParameterCollectionInstance::SyncWithCollection()
 			{
 				ParameterStorage.SetDataInterface(CastChecked<UNiagaraDataInterface>(StaticDuplicateObject(DefaultStore.GetDataInterface(Offset), this)), Param);
 			}
+			else if (Param.IsUObject())
+			{
+				ParameterStorage.SetUObject(DefaultStore.GetUObject(Offset), Param);
+			}
 			else
 			{
 				ParameterStorage.SetParameterData(DefaultStore.GetParameterData(Offset), ParameterStorageOffset, Param.GetSizeInBytes());
diff --git a/Engine/Plugins/FX/Niagara/Source/Niagara/Private/NiagaraScript.cpp b/Engine/Plugins/FX/Niagara/Source/Niagara/Private/NiagaraScript.cpp
--- a/Engine/Plugins/FX/Niagara/Source/Niagara/Private/NiagaraScript.cpp
+++ b/Engine/Plugins/FX/Niagara/Source/Niagara/Private/NiagaraScript.cpp
@@ -1784,6 +1784,23 @@ void UNiagaraScript::PreSave(FObjectPreSaveContext ObjectSaveContext)
 		return;
 	}
 
+	if (CachedScriptVM.UObjectInfos.Num() != CachedDefaultUObjects.Num())
+	{
+		UE_LOG(LogNiagara, Warning, TEXT("UObject count mismatch during script presave. Invaliding compile results (see full log for details).  Script: %s"), *GetPathName());
+		UE_LOG(LogNiagara, Log, TEXT("Compiled UObjects:"));
+		for (const FNiagaraScriptUObjectCompileInfo& Info : CachedScriptVM.UObjectInfos)
+		{
+			UE_LOG(LogNiagara, Log, TEXT("Name:%s, Type: %s"), *Info.Variable.GetName().ToString(), *Info.Variable.GetType().GetName());
+		}
+		UE_LOG(LogNiagara, Log, TEXT("Cached UObjects:"));
+		for (const FNiagaraScriptUObjectCompileInfo& Info : CachedDefaultUObjects)
+		{
+			UE_LOG(LogNiagara, Log, TEXT("Name:%s, Type: %s"), *Info.Variable.GetName().ToString(), *Info.Variable.GetType().GetName());
+		}
+		InvalidateCompileResults(TEXT("UObject count mismatch during script presave."));
+		return;
+	}
+
 	const ITargetPlatform* TargetPlatform = ObjectSaveContext.GetTargetPlatform();
 	if (TargetPlatform && TargetPlatform->RequiresCookedData())
 	{
@@ -2147,6 +2164,14 @@ void UNiagaraScript::PostLoad()
 		}
 	}
 
+	for (FNiagaraResolvedUObjectInfo& ResolvedUObject : ResolvedUObjectInfos)
+	{
+		if (ResolvedUObject.Object != nullptr)
+		{
+			ResolvedUObject.Object->ConditionalPostLoad();
+		}
+	}
+
 #if WITH_EDITORONLY_DATA
 	// Because we might be using these cached data interfaces, we need to make sure that they are properly postloaded.
 	for (FNiagaraScriptDataInterfaceInfo& Info : CachedDefaultDataInterfaces)
@@ -2157,6 +2182,14 @@ void UNiagaraScript::PostLoad()
 		}
 	}
 
+	for (FNiagaraScriptUObjectCompileInfo& DefaultUObject : CachedDefaultUObjects)
+	{
+		if (DefaultUObject.Object)
+		{
+			DefaultUObject.Object->ConditionalPostLoad();
+		}
+	}
+
 	FVersionedNiagaraScriptData* ScriptData = GetLatestScriptData();
 	ensure(ScriptData);
 	if (NiagaraVer < FNiagaraCustomVersion::AddSimulationStageUsageEnum)
@@ -2873,6 +2906,8 @@ void UNiagaraScript::SetVMCompilationResults(const FNiagaraVMExecutableDataId& I
 		check(CachedDefaultDataInterfaces[Idx].DataInterface != nullptr);
 	}
 
+	CachedDefaultUObjects = CachedScriptVM.UObjectInfos;
+
 	if (bApplyRapidIterationParameters)
 	{
 		const bool bClearBindings = false;
@@ -3654,6 +3689,30 @@ void UNiagaraScript::SyncAliases(const FNiagaraAliasContext& ResolveAliasesConte
 		UpdateName(DataInterfaceInfo.RegisteredParameterMapRead);
 		UpdateName(DataInterfaceInfo.RegisteredParameterMapWrite);
 	}
+
+	// and UObjects
+	for (FNiagaraScriptUObjectCompileInfo& Info : CachedDefaultUObjects)
+	{
+		const FNiagaraVariableBase TempNamedVariable(FNiagaraTypeDefinition::GetFloatDef(), Info.Variable.GetName());
+		const FNiagaraVariableBase ResolvedVariable = FNiagaraUtilities::ResolveAliases(TempNamedVariable, ResolveAliasesContext);
+		Info.Variable.SetName(ResolvedVariable.GetName());
+
+		auto UpdateName =
+			[&](FName& NameVariable)
+			{
+				if (NameVariable == TempNamedVariable.GetName())
+				{
+					NameVariable = ResolvedVariable.GetName();
+				}
+			};
+
+		// also update the MapRead/MapWrite member variables, they seem to typically match the Name parameter, but may be None
+		UpdateName(Info.RegisteredParameterMapRead);
+		for (FName& VariableName : Info.RegisteredParameterMapWrites)
+		{
+			UpdateName(VariableName);
+		}
+	}
 }
 
 bool UNiagaraScript::SynchronizeExecutablesWithCompilation(const UNiagaraScript* Script, const TMap<FString, FString>& RenameMap)
@@ -3679,6 +3738,7 @@ bool UNiagaraScript::SynchronizeExecutablesWithCompilation(const UNiagaraScript*
 		CachedScriptVMId = Script->CachedScriptVMId;
 		CachedScriptVM.OptimizationTask.State = nullptr;
 		CachedParameterCollectionReferences = Script->CachedParameterCollectionReferences;
+		CachedDefaultUObjects = Script->CachedDefaultUObjects;
 		CachedDefaultDataInterfaces.Empty();
 		ResolvedDataInterfaces.Empty();
 		for (const FNiagaraScriptDataInterfaceInfo& Info : Script->CachedDefaultDataInterfaces)
@@ -3717,6 +3777,7 @@ void UNiagaraScript::InvalidateCompileResults(const FString& Reason)
 	{
 		GetLastGeneratedVMId().Invalidate();
 	}
+	CachedDefaultUObjects.Reset();
 	CachedDefaultDataInterfaces.Reset();
 	ResolvedDataInterfaces.Reset();
 	CachedScriptVM.OptimizationTask.State = nullptr;
diff --git a/Engine/Plugins/FX/Niagara/Source/Niagara/Private/NiagaraScriptExecutionParameterStore.cpp b/Engine/Plugins/FX/Niagara/Source/Niagara/Private/NiagaraScriptExecutionParameterStore.cpp
--- a/Engine/Plugins/FX/Niagara/Source/Niagara/Private/NiagaraScriptExecutionParameterStore.cpp
+++ b/Engine/Plugins/FX/Niagara/Source/Niagara/Private/NiagaraScriptExecutionParameterStore.cpp
@@ -318,6 +318,14 @@ void FNiagaraScriptExecutionParameterStore::AddScriptParams(UNiagaraScript* Scri
 		SetDataInterface(ResolvedDataInterface.ResolvedDataInterface, VarOffset);
 	}
 
+	check(Script->GetVMExecutableData().UObjectInfos.Num() == Script->GetCachedDefaultUObjects().Num());
+	for (const FNiagaraResolvedUObjectInfo& Info : Script->GetResolvedUObjects())
+	{
+		int32 VarOffset = INDEX_NONE;
+		bAdded |= AddParameter(Info.ResolvedVariable, false, false, &VarOffset);
+		SetUObject(Info.Object, VarOffset);
+	}
+
 	if (bAdded && bTriggerRebind)
 	{
 		OnLayoutChange();
diff --git a/Engine/Plugins/FX/Niagara/Source/Niagara/Private/NiagaraSystem.cpp b/Engine/Plugins/FX/Niagara/Source/Niagara/Private/NiagaraSystem.cpp
--- a/Engine/Plugins/FX/Niagara/Source/Niagara/Private/NiagaraSystem.cpp
+++ b/Engine/Plugins/FX/Niagara/Source/Niagara/Private/NiagaraSystem.cpp
@@ -64,14 +64,6 @@ namespace NiagaraScriptCookStats
 int32 GEnableNiagaraRuntimeCycleCounts = 0;
 static FAutoConsoleVariableRef CVarEnableNiagaraRuntimeCycleCounts(TEXT("fx.EnableNiagaraRuntimeCycleCounts"), GEnableNiagaraRuntimeCycleCounts, TEXT("Toggle for runtime cylce counts tracking Niagara's frame time. \n"), ECVF_ReadOnly);
 
-static int GNiagaraLogDDCStatusForSystems = 0;
-static FAutoConsoleVariableRef CVarLogDDCStatusForSystems(
-	TEXT("fx.NiagaraLogDDCStatusForSystems"),
-	GNiagaraLogDDCStatusForSystems,
-	TEXT("When enabled UNiagaraSystems will log out when their subscripts are pulled from the DDC or not."),
-	ECVF_Default
-);
-
 static float GNiagaraScalabiltiyMinumumMaxDistance = 1.0f;
 static FAutoConsoleVariableRef CVarNiagaraScalabiltiyMinumumMaxDistance(
 	TEXT("fx.Niagara.Scalability.MinMaxDistance"),
@@ -1990,7 +1982,7 @@ FBox UNiagaraSystem::GetFixedBounds() const
 
 const FGuid UNiagaraSystem::ResolveDIsMessageId(0xB6ACDD97, 0xA2C04B02, 0xAB9FA4E4, 0xDBC73E51);
 
-void UNiagaraSystem::ResolveDIBindings()
+void UNiagaraSystem::ResolveParameterStoreBindings()
 {
 	TArray<FText> ErrorMessages;
 	FNiagaraResolveDIHelpers::ResolveDIs(this, ErrorMessages);
@@ -2018,6 +2010,60 @@ void UNiagaraSystem::ResolveDIBindings()
 	{
 		MessageStore.RemoveMessage(ResolveDIsMessageId);
 	}
+
+
+	// Resolve Object bindings
+	{
+		TMap<FNiagaraVariableBase, UObject*> WroteUObjects;
+		ForEachScript(
+			[&](UNiagaraScript* Script)
+			{
+				for (const FNiagaraScriptUObjectCompileInfo& DefaultInfo : Script->GetCachedDefaultUObjects())
+				{
+					for (FName WriteName : DefaultInfo.RegisteredParameterMapWrites)
+					{
+						const FNiagaraVariableBase WriteVariable(DefaultInfo.Variable.GetType(), WriteName);
+						check(WroteUObjects.Contains(WriteVariable) == false);
+						WroteUObjects.Emplace(WriteVariable, DefaultInfo.Object);
+					}
+				}
+			}
+		);
+
+		// Create resolved data per script
+		ForEachScript(
+			[&](UNiagaraScript* Script)
+			{
+				if (Script->GetCachedDefaultUObjects().Num() == 0)
+				{
+					return;
+				}
+
+				TArray<FNiagaraResolvedUObjectInfo> ResolvedUObjects;
+				ResolvedUObjects.Reserve(Script->GetCachedDefaultUObjects().Num());
+
+				for (const FNiagaraScriptUObjectCompileInfo& DefaultInfo : Script->GetCachedDefaultUObjects())
+				{
+					const FNiagaraVariableBase ReadVariable(DefaultInfo.Variable.GetType(), DefaultInfo.RegisteredParameterMapRead);
+
+					//-TODO:UObject will need to handle User / NPC bindings here.  I.e. Variable=User.Variable in the stack set
+					//const bool bIsUserVariable = ReadVariable.IsInNameSpace(FNiagaraConstants::UserNamespaceString);
+
+					for (FName WriteName : DefaultInfo.RegisteredParameterMapWrites)
+					{
+						FNiagaraResolvedUObjectInfo& ResolvedInfo = ResolvedUObjects.AddDefaulted_GetRef();
+						const FNiagaraVariableBase WriteVariable(DefaultInfo.Variable.GetType(), WriteName);
+						ResolvedInfo.ReadVariableName = ReadVariable.GetName();
+						ResolvedInfo.ResolvedVariable = WriteVariable;
+						ResolvedInfo.Object = WroteUObjects.FindRef(WriteVariable);
+					}
+				}
+
+				ResolvedUObjects.Shrink();
+				Script->SetResolvedUObjects(ResolvedUObjects);
+			}
+		);
+	}
 }
 
 #endif // WITH_EDITORONLY_DATA
@@ -2123,6 +2169,11 @@ void UNiagaraSystem::OnCompiledDataInterfaceChanged()
 	UpdateDITickFlags();
 }
 
+void UNiagaraSystem::OnCompiledUObjectChanged()
+{
+	ForEachScript([&](UNiagaraScript* Script) { Script->InvalidateExecutionReadyParameterStores(); });
+}
+
 #endif
 
 bool UNiagaraSystem::IsValidInternal() const
@@ -2735,7 +2786,7 @@ bool UNiagaraSystem::QueryCompileComplete(bool bWait)
 	InitEmitterCompiledData();
 	InitSystemCompiledData();
 
-	ResolveDIBindings();
+	ResolveParameterStoreBindings();
 	UpdatePostCompileDIInfo();
 	ComputeEmittersExecutionOrder();
 	ComputeRenderersDrawOrder();
diff --git a/Engine/Plugins/FX/Niagara/Source/Niagara/Public/NiagaraCommon.h b/Engine/Plugins/FX/Niagara/Source/Niagara/Public/NiagaraCommon.h
--- a/Engine/Plugins/FX/Niagara/Source/Niagara/Public/NiagaraCommon.h
+++ b/Engine/Plugins/FX/Niagara/Source/Niagara/Public/NiagaraCommon.h
@@ -669,6 +669,42 @@ struct NIAGARA_API FNiagaraFunctionSignature
 	void GetVariadicOutputs(TArray<FNiagaraVariableBase>& OutVariadicOutputs, bool bStripNonExecution = true)const;
 };
 
+USTRUCT()
+struct NIAGARA_API FNiagaraScriptUObjectCompileInfo
+{
+	GENERATED_BODY()
+
+	UPROPERTY()
+	FNiagaraVariableBase Variable;
+
+	UPROPERTY()
+	TObjectPtr<class UObject> Object;
+
+	UPROPERTY()
+	FName RegisteredParameterMapRead;
+
+	UPROPERTY()
+	TArray<FName> RegisteredParameterMapWrites;
+};
+
+USTRUCT()
+struct NIAGARA_API FNiagaraResolvedUObjectInfo
+{
+	GENERATED_BODY()
+
+	UPROPERTY()
+	FName ReadVariableName;
+
+	UPROPERTY()
+	FNiagaraVariableBase ResolvedVariable;
+
+	UPROPERTY()
+	TObjectPtr<class UObject> Object;
+
+	//UPROPERTY()
+	//int32 UserPtrIdx = INDEX_NONE;
+};
+
 USTRUCT()
 struct NIAGARA_API FNiagaraScriptDataInterfaceInfo
 {
diff --git a/Engine/Plugins/FX/Niagara/Source/NiagaraEditor/Private/Customizations/NiagaraParameterCollectionCustomNodeBuilder.h b/Engine/Plugins/FX/Niagara/Source/NiagaraEditor/Private/Customizations/NiagaraParameterCollectionCustomNodeBuilder.h
--- a/Engine/Plugins/FX/Niagara/Source/NiagaraEditor/Private/Customizations/NiagaraParameterCollectionCustomNodeBuilder.h
+++ b/Engine/Plugins/FX/Niagara/Source/NiagaraEditor/Private/Customizations/NiagaraParameterCollectionCustomNodeBuilder.h
@@ -5,6 +5,7 @@
 #include "DetailCategoryBuilder.h"
 #include "IDetailPropertyRow.h"
 #include "DetailWidgetRow.h"
+#include "PropertyCustomizationHelpers.h"
 #include "Widgets/Text/SInlineEditableTextBlock.h"
 #include "ViewModels/NiagaraParameterCollectionViewModel.h"
 #include "ViewModels/NiagaraScriptInputCollectionViewModel.h"
@@ -85,32 +86,49 @@ public:
 			TSharedPtr<SWidget> CustomValueWidget;
 			bool bCustomize = true;
 
-			if (Parameter->GetDefaultValueType() == INiagaraParameterViewModel::EDefaultValueType::Struct)
+			switch (Parameter->GetDefaultValueType())
 			{
-				Row = ChildrenBuilder.AddExternalStructureProperty(Parameter->GetDefaultValueStruct(), NAME_None, FAddPropertyParams().UniqueId(Parameter->GetName()));
+				case INiagaraParameterViewModel::EDefaultValueType::Struct:
+					Row = ChildrenBuilder.AddExternalStructureProperty(Parameter->GetDefaultValueStruct(), NAME_None, FAddPropertyParams().UniqueId(Parameter->GetName()));
+					break;
 
-			}
-			else if (Parameter->GetDefaultValueType() == INiagaraParameterViewModel::EDefaultValueType::Object)
-			{
-				UObject* DefaultValueObject = Parameter->GetDefaultValueObject();
-
-				if (DefaultValueObject != nullptr)
+				case INiagaraParameterViewModel::EDefaultValueType::DataInterface:
 				{
-					TArray<UObject*> Objects;
-					Objects.Add(DefaultValueObject);
-
-					FAddPropertyParams Params = FAddPropertyParams()
-						.UniqueId(Parameter->GetName())
-						.AllowChildren(true)
-						.CreateCategoryNodes(false);
-
-					Row = ChildrenBuilder.AddExternalObjectProperty(Objects, NAME_None, Params);
-					CustomValueWidget =
-						SNew(STextBlock)
-						.TextStyle(FNiagaraEditorStyle::Get(), "NiagaraEditor.ParameterText")
-						.Text(FText::FromString(FName::NameToDisplayString(DefaultValueObject->GetClass()->GetName(), false)));
+					UNiagaraDataInterface* DefaultValueObject = Parameter->GetDefaultValueDataInterface();
+
+					if (DefaultValueObject != nullptr)
+					{
+						TArray<UObject*> Objects;
+						Objects.Add(DefaultValueObject);
+
+						FAddPropertyParams Params = FAddPropertyParams()
+							.UniqueId(Parameter->GetName())
+							.AllowChildren(true)
+							.CreateCategoryNodes(false);
+
+						Row = ChildrenBuilder.AddExternalObjectProperty(Objects, NAME_None, Params);
+						CustomValueWidget =
+							SNew(STextBlock)
+							.TextStyle(FNiagaraEditorStyle::Get(), "NiagaraEditor.ParameterText")
+							.Text(FText::FromString(FName::NameToDisplayString(DefaultValueObject->GetClass()->GetName(), false)));
+					}
+					else
+					{
+						ChildrenBuilder.AddCustomRow(FText())
+							.NameContent()
+							[
+								NameWidget.ToSharedRef()
+							]
+							.ValueContent()
+							[
+								SNew(STextBlock)
+								.Text(NSLOCTEXT("NiagaraParameterCollectionCustomNodeBuilder", "NullObjectValue", "(null)"))
+							];
+					}
+					break;
 				}
-				else
+
+				case INiagaraParameterViewModel::EDefaultValueType::ObjectAsset:
 				{
 					ChildrenBuilder.AddCustomRow(FText())
 						.NameContent()
@@ -119,10 +137,32 @@ public:
 						]
 						.ValueContent()
 						[
-							SNew(STextBlock)
-							.Text(NSLOCTEXT("NiagaraParameterCollectionCustomNodeBuilder", "NullObjectValue", "(null)"))
+							SNew(SObjectPropertyEntryBox)
+								.AllowedClass(Parameter->GetType()->GetClass())
+								.ObjectPath_Lambda(
+									[Parameter]() -> FString
+									{
+										UObject* DefaultValueObject = Parameter->GetDefaultValueObjectAsset();
+										return DefaultValueObject ? DefaultValueObject->GetPathName() : FString();
+									}
+								)
+								.DisplayBrowse(true)
+								.DisplayUseSelected(true)
+								.DisplayThumbnail(true)
+								.EnableContentPicker(true)
+								.OnObjectChanged_Lambda(
+									[Parameter](const FAssetData& AssetData)
+									{
+										Parameter->SetDefaultValueObjectAsset(AssetData.GetAsset());
+									}
+								)
 						];
+					break;
 				}
+
+				default:
+					checkNoEntry();
+					break;
 			}
 
 			if (Row)
@@ -159,12 +199,8 @@ public:
 					}
 
 				}
-
 			}
-
 		}
-
-
 	}
 
 private:
diff --git a/Engine/Plugins/FX/Niagara/Source/NiagaraEditor/Private/EdGraphSchema_Niagara.cpp b/Engine/Plugins/FX/Niagara/Source/NiagaraEditor/Private/EdGraphSchema_Niagara.cpp
--- a/Engine/Plugins/FX/Niagara/Source/NiagaraEditor/Private/EdGraphSchema_Niagara.cpp
+++ b/Engine/Plugins/FX/Niagara/Source/NiagaraEditor/Private/EdGraphSchema_Niagara.cpp
@@ -365,6 +365,10 @@ void AddParametersForGraph(TArray<TSharedPtr<FNiagaraAction_NewNode>>& NewAction
 				{
 					InputNodeTemplate->SetDataInterface(Cast<UNiagaraDataInterface>(StaticDuplicateObject(InputNode->GetDataInterface(), InputNodeTemplate, NAME_None, ~RF_Transient)));
 				}
+				if (InputNode->GetObjectAsset())
+				{
+					InputNodeTemplate->SetObjectAsset(InputNode->GetObjectAsset());
+				}
 			}
 		}
 	}
diff --git a/Engine/Plugins/FX/Niagara/Source/NiagaraEditor/Private/NiagaraClipboard.cpp b/Engine/Plugins/FX/Niagara/Source/NiagaraEditor/Private/NiagaraClipboard.cpp
--- a/Engine/Plugins/FX/Niagara/Source/NiagaraEditor/Private/NiagaraClipboard.cpp
+++ b/Engine/Plugins/FX/Niagara/Source/NiagaraEditor/Private/NiagaraClipboard.cpp
@@ -77,6 +77,13 @@ const UNiagaraClipboardFunctionInput* UNiagaraClipboardFunctionInput::CreateData
 	return NewInput;
 }
 
+const UNiagaraClipboardFunctionInput* UNiagaraClipboardFunctionInput::CreateObjectAssetValue(UObject* InOuter, FName InInputName, FNiagaraTypeDefinition InInputType, TOptional<bool> bInEditConditionValue, UObject* InObject)
+{
+	UNiagaraClipboardFunctionInput* NewInput = MakeNewInput(InOuter, InInputName, InInputType, bInEditConditionValue, ENiagaraClipboardFunctionInputValueMode::ObjectAsset);
+	NewInput->ObjectAsset = InObject;
+	return NewInput;
+}
+
 const UNiagaraClipboardFunctionInput* UNiagaraClipboardFunctionInput::CreateExpressionValue(UObject* InOuter, FName InInputName, FNiagaraTypeDefinition InInputType, TOptional<bool> bInEditConditionValue, const FString& InExpressionValue)
 {
 	UNiagaraClipboardFunctionInput* NewInput = MakeNewInput(InOuter, InInputName, InInputType, bInEditConditionValue, ENiagaraClipboardFunctionInputValueMode::Expression);
diff --git a/Engine/Plugins/FX/Niagara/Source/NiagaraEditor/Private/NiagaraCompiler.cpp b/Engine/Plugins/FX/Niagara/Source/NiagaraEditor/Private/NiagaraCompiler.cpp
--- a/Engine/Plugins/FX/Niagara/Source/NiagaraEditor/Private/NiagaraCompiler.cpp
+++ b/Engine/Plugins/FX/Niagara/Source/NiagaraEditor/Private/NiagaraCompiler.cpp
@@ -399,7 +399,7 @@ void FNiagaraCompileRequestData::GatherPreCompiledVariables(const FString& InNam
 			if (FNiagaraParameterMapHistory::IsInNamespace(EncounteredVariable, InNamespaceFilter))
 			{
 				FNiagaraVariable NewVar = EncounteredVariable;
-				if (NewVar.IsDataAllocated() == false && !NewVar.IsDataInterface())
+				if (NewVar.IsDataAllocated() == false && !NewVar.IsDataInterface() && !NewVar.IsUObject())
 				{
 					FNiagaraEditorUtilities::ResetVariableToDefaultValue(NewVar);
 				}
@@ -1890,6 +1890,7 @@ int32 FHlslNiagaraCompiler::CompileScript(const FStringView GroupName, const FNi
 		CompileResults.Data->Parameters.Empty();
 		CompileResults.Data->InternalParameters.Empty();
 		CompileResults.Data->DataInterfaceInfo.Empty();
+		CompileResults.Data->UObjectInfos.Empty();
 
 	}
 	CompilationJob->CompileResults = CompileResults;
@@ -2024,6 +2025,7 @@ TOptional<FNiagaraCompileResults> FHlslNiagaraCompiler::GetCompileResult(int32 J
 		Results.Data->Parameters.Empty();
 		Results.Data->InternalParameters.Empty();
 		Results.Data->DataInterfaceInfo.Empty();
+		Results.Data->UObjectInfos.Empty();
 		//Eventually Niagara will have all the shader plumbing and do things like materials.
 	}
 	else
diff --git a/Engine/Plugins/FX/Niagara/Source/NiagaraEditor/Private/NiagaraEditorUtilities.cpp b/Engine/Plugins/FX/Niagara/Source/NiagaraEditor/Private/NiagaraEditorUtilities.cpp
--- a/Engine/Plugins/FX/Niagara/Source/NiagaraEditor/Private/NiagaraEditorUtilities.cpp
+++ b/Engine/Plugins/FX/Niagara/Source/NiagaraEditor/Private/NiagaraEditorUtilities.cpp
@@ -149,12 +149,21 @@ void FNiagaraEditorUtilities::InitializeParameterInputNode(UNiagaraNodeInput& In
 		{
 			InputNode.SetDataInterface(nullptr);
 		}
+		if (InputNode.GetObjectAsset() != nullptr)
+		{
+			InputNode.SetObjectAsset(nullptr);
+		}
 	}
 	else if(Type.IsDataInterface())
 	{
 		InputNode.Input.AllocateData(); // Frees previously used memory if we're switching from a struct to a class type.
 		InputNode.SetDataInterface(NewObject<UNiagaraDataInterface>(&InputNode, Type.GetClass(), NAME_None, RF_Transactional | RF_Public));
 	}
+	else if (Type.IsUObject())
+	{
+		InputNode.Input.AllocateData(); // Frees previously used memory if we're switching from a struct to a class type.
+		InputNode.SetObjectAsset(nullptr);
+	}
 }
 
 void FNiagaraEditorUtilities::GetParameterVariablesFromSystem(UNiagaraSystem& System, TArray<FNiagaraVariable>& ParameterVariables,
@@ -517,16 +526,26 @@ bool FNiagaraEditorUtilities::NestedPropertiesAppendCompileHash(const void* Cont
 		{
 			FObjectProperty* CastProp = CastFieldChecked<FObjectProperty>(Property);
 			UObject* Obj = CastProp->GetObjectPropertyValue_InContainer(Container);
-			if (Obj != nullptr)
+
+			static FName HashObjectPathMeta = TEXT("CompileHashObjectPath");
+			if (Property->HasMetaData(HashObjectPathMeta))
 			{
-				// We just do name here as sometimes things will be in a transient package or something tricky.
-				// Because we do nested id's for each called graph, it should work out in the end to have a different
-				// value in the compile array if the scripts are the same name but different locations.
-				InVisitor->UpdateString(*PropertyName, Obj->GetName());
+				FSoftObjectPath SoftPath(Obj);
+				InVisitor->UpdateString(*PropertyName, SoftPath.GetAssetPath().ToString());
 			}
 			else
 			{
-				InVisitor->UpdateString(*PropertyName, TEXT("nullptr"));
+				if (Obj != nullptr)
+				{
+					// We just do name here as sometimes things will be in a transient package or something tricky.
+					// Because we do nested id's for each called graph, it should work out in the end to have a different
+					// value in the compile array if the scripts are the same name but different locations.
+					InVisitor->UpdateString(*PropertyName, Obj->GetName());
+				}
+				else
+				{
+					InVisitor->UpdateString(*PropertyName, TEXT("nullptr"));
+				}
 			}
 			continue;
 		}
diff --git a/Engine/Plugins/FX/Niagara/Source/NiagaraEditor/Private/NiagaraGraph.cpp b/Engine/Plugins/FX/Niagara/Source/NiagaraEditor/Private/NiagaraGraph.cpp
--- a/Engine/Plugins/FX/Niagara/Source/NiagaraEditor/Private/NiagaraGraph.cpp
+++ b/Engine/Plugins/FX/Niagara/Source/NiagaraEditor/Private/NiagaraGraph.cpp
@@ -352,7 +352,7 @@ void UNiagaraGraph::PostLoad_LWCFixup(int32 NiagaraVersion)
 				TArray<UEdGraphPin*> Pins = FindParameterMapDefaultValuePins(Variable.GetName());
 				for (UEdGraphPin* Pin : Pins)
 				{
-					check(ChangedScriptVariable->DefaultMode == ENiagaraDefaultMode::Value && !Variable.GetType().IsDataInterface());
+					check(ChangedScriptVariable->DefaultMode == ENiagaraDefaultMode::Value && !Variable.GetType().IsDataInterface() && !Variable.GetType().IsUObject());
 
 					Variable.AllocateData();
 					FString NewDefaultValue = TypeUtilityValue->GetPinDefaultStringFromValue(Variable);
@@ -2885,7 +2885,7 @@ void UNiagaraGraph::ScriptVariableChanged(FNiagaraVariable Variable)
 			MapGetNodes.Add(MapGetNode);
 		}
 		
-		if (ScriptVariable->DefaultMode == ENiagaraDefaultMode::Value && !Variable.GetType().IsDataInterface())
+		if (ScriptVariable->DefaultMode == ENiagaraDefaultMode::Value && !Variable.GetType().IsDataInterface() && !Variable.GetType().IsUObject())
 		{
 			FNiagaraEditorModule& EditorModule = FNiagaraEditorModule::Get();
 			auto TypeUtilityValue = EditorModule.GetTypeUtilities(Variable.GetType());
diff --git a/Engine/Plugins/FX/Niagara/Source/NiagaraEditor/Private/NiagaraHlslTranslator.cpp b/Engine/Plugins/FX/Niagara/Source/NiagaraEditor/Private/NiagaraHlslTranslator.cpp
--- a/Engine/Plugins/FX/Niagara/Source/NiagaraEditor/Private/NiagaraHlslTranslator.cpp
+++ b/Engine/Plugins/FX/Niagara/Source/NiagaraEditor/Private/NiagaraHlslTranslator.cpp
@@ -4595,7 +4595,7 @@ int32 FHlslNiagaraTranslator::GetRapidIterationParameter(const FNiagaraVariable&
 	bool bIsCandidateForRapidIteration = false;
 	if (ActiveHistoryForFunctionCalls.InTopLevelFunctionCall(CompileOptions.TargetUsage))
 	{
-		if (Parameter.GetType() != FNiagaraTypeDefinition::GetBoolDef() && !Parameter.GetType().IsEnum() && !Parameter.GetType().IsDataInterface())
+		if (Parameter.GetType() != FNiagaraTypeDefinition::GetBoolDef() && !Parameter.GetType().IsEnum() && !Parameter.IsDataInterface() && !Parameter.IsUObject())
 		{
 			bIsCandidateForRapidIteration = true;
 		}
@@ -4778,7 +4778,7 @@ int32 FHlslNiagaraTranslator::GetParameter(const FNiagaraVariable& Parameter)
 
 int32 FHlslNiagaraTranslator::GetConstant(const FNiagaraVariable& Constant, FString* DebugOutputValue)
 {
-	if (Constant.IsDataInterface())
+	if (Constant.IsDataInterface() || Constant.IsUObject())
 	{
 		return INDEX_NONE;
 	}
@@ -5462,9 +5462,6 @@ void FHlslNiagaraTranslator::ParameterMapSet(UNiagaraNodeParameterMapSet* SetNod
 				}
 			}
 
-
-			
-
 			if (Var.IsDataInterface())
 			{
 				if (CompileOptions.TargetUsage == ENiagaraScriptUsage::ParticleSpawnScriptInterpolated && TranslationStages[ActiveStageIdx].ScriptUsage == ENiagaraScriptUsage::ParticleUpdateScript)
@@ -5517,6 +5514,53 @@ void FHlslNiagaraTranslator::ParameterMapSet(UNiagaraNodeParameterMapSet* SetNod
 					}
 				}
 			}
+			else if ( Var.IsUObject() )
+			{
+				if (CompileOptions.TargetUsage == ENiagaraScriptUsage::ParticleSpawnScriptInterpolated && TranslationStages[ActiveStageIdx].ScriptUsage == ENiagaraScriptUsage::ParticleUpdateScript)
+				{
+					// We don't want to add writes for particle update data interface parameters in both interpolated spawn and update, so skip them when processing the update stage of the 
+					// interpolated spawn script.  We don't skip the writes when compiling the particle update script because it's not recompiled when the interpolated spawn flag is changed
+					// and this would result in missing data interfaces if interpolated spawn was turned off.
+					continue;
+				}
+
+				const bool bAllowDataInterfaces = true;	// Also allows UObjects
+				if (ParamMapHistoryIdx < ParamMapHistories.Num() && ParamMapHistories[ParamMapHistoryIdx].IsPrimaryDataSetOutput(Var, CompileOptions.TargetUsage, bAllowDataInterfaces))
+				{
+					if (Input < 0 || Input >= CompilationOutput.ScriptData.UObjectInfos.Num())
+					{
+						Error(FText::Format(LOCTEXT("ParameterMapUObjectNotFoundErrorFormat", "UObject could not be found for parameter map set.  Paramter: {0}"), FText::FromName(Var.GetName())), SetNode, Inputs[i].Pin);
+						continue;
+					}
+
+					FName ResolvedName;
+					if (FNiagaraParameterMapHistory::IsAliasedEmitterParameter(Var))
+					{
+						ResolvedName = ActiveHistoryForFunctionCalls.ResolveAliases(Var).GetName();
+					}
+					else
+					{
+						ResolvedName = Var.GetName();
+					}
+
+					FNiagaraScriptUObjectCompileInfo& Info = CompilationOutput.ScriptData.UObjectInfos[Input];
+					Info.RegisteredParameterMapWrites.AddUnique(ResolvedName);
+
+					//-TODO: We don't support Var=User.Var or Var=NPC.Var currently this requires some runtime remapping of the parameter stores
+					if (Info.Variable.IsInNameSpace(FNiagaraConstants::UserNamespaceString) || Info.Variable.IsInNameSpace(FNiagaraConstants::ParameterCollectionNamespaceString))
+					{
+						Error(
+							FText::Format(
+								LOCTEXT("UObjectParameterInvalidNamespaceFormat", "Variable Write '{0}' can not read User / NPC variable '{1}'. This is not supported currently."),
+								FText::FromName(ResolvedName),
+								FText::FromName(Info.Variable.GetName())
+							),
+							nullptr, nullptr
+						);
+						continue;
+					}
+				}
+			}
 			else
 			{
 				if (Var == FNiagaraVariable(FNiagaraTypeDefinition::GetBoolDef(), TEXT("DataInstance.Alive")))
@@ -5693,14 +5737,18 @@ bool FHlslNiagaraTranslator::ParameterMapRegisterExternalConstantNamespaceVariab
 			}
 		}
 
-		bool bIsDataInterface = InVariable.GetType().IsDataInterface();
+		const bool bIsDataInterface = InVariable.IsDataInterface();
+		const bool bIsUObject = InVariable.IsUObject();
 		const FString* EmitterAlias = ActiveHistoryForFunctionCalls.GetEmitterAlias();
-		bool bIsPerInstanceBulkSystemParam = IsBulkSystemScript() && !bIsDataInterface && (FNiagaraParameterMapHistory::IsUserParameter(InVariable) || FNiagaraParameterMapHistory::IsPerInstanceEngineParameter(InVariable, EmitterAlias != nullptr ? *EmitterAlias : TEXT("Emitter")));
+		
+		bool bIsPerInstanceBulkSystemParam = IsBulkSystemScript() && !bIsDataInterface && !bIsUObject && (FNiagaraParameterMapHistory::IsUserParameter(InVariable) || FNiagaraParameterMapHistory::IsPerInstanceEngineParameter(InVariable, EmitterAlias != nullptr ? *EmitterAlias : TEXT("Emitter")));
 		const bool bIsExternalConstantParameter = FNiagaraParameterMapHistory::IsRapidIterationParameter(InVariable) && !InVariable.GetType().IsStatic();
 
 		if (InVariable.GetType().IsStatic())
+		{
 			return false;
-		
+		}
+
 		if (!bIsPerInstanceBulkSystemParam)
 		{
 			int32 UniformChunk = 0;
@@ -5709,7 +5757,7 @@ bool FHlslNiagaraTranslator::ParameterMapRegisterExternalConstantNamespaceVariab
 			{
 				FString SymbolNameDefined = FlattenedName;
 
-				if (InVariable.GetType().IsDataInterface())
+				if (bIsDataInterface)
 				{
 					UNiagaraDataInterface* DataInterface = nullptr;
 					if (Collection)
@@ -5737,6 +5785,16 @@ bool FHlslNiagaraTranslator::ParameterMapRegisterExternalConstantNamespaceVariab
 						return true;
 					}
 				}
+				if (bIsUObject)
+				{
+					UObject* Object = nullptr;
+					if (Collection)
+					{
+						Object = Collection->GetDefaultInstance()->GetParameterStore().GetUObject(InVariable);
+					}
+					Output = RegisterUObject(InVariable, Object, true);
+					return true;
+				}
 				if (!InVariable.IsDataAllocated() && !InDefaultPin)
 				{
 					FNiagaraEditorUtilities::ResetVariableToDefaultValue(InVariable);
@@ -6265,9 +6323,10 @@ void FHlslNiagaraTranslator::ParameterMapGet(UNiagaraNodeParameterMapGet* GetNod
 		else // These are the pins that we are getting off the parameter map.
 		{
 			FNiagaraTypeDefinition OutputTypeDefinition = Schema->PinToTypeDefinition(OutputPins[i]);
-			bool bNeedsValue =
+			const bool bNeedsValue =
 				OutputTypeDefinition != FNiagaraTypeDefinition::GetParameterMapDef() &&
-				OutputTypeDefinition.IsDataInterface() == false;
+				OutputTypeDefinition.IsDataInterface() == false &&
+				OutputTypeDefinition.IsUObject() == false;
 			FNiagaraVariable Var = Schema->PinToNiagaraVariable(OutputPins[i], bNeedsValue, ENiagaraStructConversion::Simulation);
 
 			FNiagaraScriptVariableBinding DefaultBinding;
@@ -6419,7 +6478,7 @@ void FHlslNiagaraTranslator::HandleParameterRead(int32 ParamMapHistoryIdx, const
 			// Skip the case where the below condition is met, but it's overridden by a binding.
 			bIsCandidateForRapidIteration = false;
 		}
-		else if (InputPin != nullptr && InputPin->LinkedTo.Num() == 0 && Var.GetType() != FNiagaraTypeDefinition::GetBoolDef() && !Var.GetType().IsEnum() && !Var.GetType().IsDataInterface())
+		else if (InputPin != nullptr && InputPin->LinkedTo.Num() == 0 && Var.GetType() != FNiagaraTypeDefinition::GetBoolDef() && !Var.GetType().IsEnum() && !Var.IsDataInterface() && !Var.IsUObject())
 		{
 			bIsCandidateForRapidIteration = true;
 		}
@@ -6736,7 +6795,7 @@ void FHlslNiagaraTranslator::HandleParameterRead(int32 ParamMapHistoryIdx, const
 				LastSetChunkIdx = GetConstant(Var);
 			}
 
-			if (!Var.IsDataInterface() && LastSetChunkIdx != INDEX_NONE)
+			if (!Var.IsDataInterface() && !Var.IsUObject() && LastSetChunkIdx != INDEX_NONE)
 			{
 				if (bTreatAsUnknownParameterMap == false)
 				{
@@ -6766,12 +6825,15 @@ void FHlslNiagaraTranslator::HandleParameterRead(int32 ParamMapHistoryIdx, const
 				}
 			}
 
-			if (LastSetChunkIdx == INDEX_NONE && VarIdx != INDEX_NONE && Var.IsDataInterface())
+			if (LastSetChunkIdx == INDEX_NONE && VarIdx != INDEX_NONE)
 			{
-				// If this variable is a data interface and it's in the parameter map, but hasn't been set yet, then is is an external data interface so try to register it.
-				if (ParameterMapRegisterExternalConstantNamespaceVariable(Var, ErrorNode, ParamMapHistoryIdx, OutputChunkId, DefaultPin))
+				if (Var.IsDataInterface() || Var.IsUObject())
 				{
-					return;
+					// If this variable is a data interface and it's in the parameter map, but hasn't been set yet, then is is an external data interface so try to register it.
+					if (ParameterMapRegisterExternalConstantNamespaceVariable(Var, ErrorNode, ParamMapHistoryIdx, OutputChunkId, DefaultPin))
+					{
+						return;
+					}
 				}
 			}
 		}
@@ -6795,6 +6857,17 @@ void FHlslNiagaraTranslator::HandleParameterRead(int32 ParamMapHistoryIdx, const
 
 		OutputChunkId = LastSetChunkIdx;
 	}
+	else if (Var.IsUObject())
+	{
+		// In order for a module to compile successfully, we potentially need to generate default values
+		// for variables encountered without ever being set. We do this by creating an instance of the CDO.
+		if (UNiagaraScript::IsStandaloneScript(CompileOptions.TargetUsage) && LastSetChunkIdx == INDEX_NONE)
+		{
+			LastSetChunkIdx = RegisterUObject(Var, nullptr, true);
+		}
+
+		OutputChunkId = LastSetChunkIdx;
+	}
 	else
 	{
 		OutputChunkId = AddSourceChunk(ParameterMapInstanceName + TEXT(".") + GetSanitizedSymbolName(Var.GetName().ToString()), Var.GetType());
@@ -6936,6 +7009,51 @@ void FHlslNiagaraTranslator::WriteDataSet(const FNiagaraDataSetID DataSet, const
 
 }
 
+int32 FHlslNiagaraTranslator::RegisterUObject(FNiagaraVariable Variable, UObject* Object, bool bAddParameterMapRead)
+{
+	int32 ObjectIndex = INDEX_NONE;
+	if (bAddParameterMapRead)
+	{
+		// If we are registering a read then look for an existing one by name
+		ObjectIndex = CompilationOutput.ScriptData.UObjectInfos.IndexOfByPredicate([&](const FNiagaraScriptUObjectCompileInfo& ExistingInfo) { return ExistingInfo.Variable == Variable; });
+	}
+	// Assume all writes are unique
+	//else if ( Object != nullptr )
+	//{
+	//	// If we are registering a write then look for an existing one by value
+	//	ObjectIndex = CompilationOutput.ScriptData.UObjectInfos.IndexOfByPredicate([&](const FNiagaraScriptUObjectCompileInfo& ExistingInfo) { return ExistingInfo.Object == Object; });
+	//}
+
+	FNiagaraScriptUObjectCompileInfo* CompileInfo = nullptr;
+	if (ObjectIndex == INDEX_NONE)
+	{
+		ObjectIndex = CompilationOutput.ScriptData.UObjectInfos.AddDefaulted();
+		CompileInfo = &CompilationOutput.ScriptData.UObjectInfos[ObjectIndex];
+		CompileInfo->Variable = Variable;
+		CompileInfo->Object = Object;
+	}
+	else
+	{
+		CompileInfo = &CompilationOutput.ScriptData.UObjectInfos[ObjectIndex];
+	}
+
+	if (bAddParameterMapRead)
+	{
+		FName ResolvedName;
+		if (FNiagaraParameterMapHistory::IsAliasedEmitterParameter(Variable.GetName().ToString()))
+		{
+			ResolvedName = ActiveHistoryForFunctionCalls.ResolveAliases(Variable).GetName();
+		}
+		else
+		{
+			ResolvedName = Variable.GetName();
+		}
+		CompileInfo->RegisteredParameterMapRead = ResolvedName;
+	}
+
+	return ObjectIndex;
+}
+
 int32 FHlslNiagaraTranslator::RegisterDataInterface(FNiagaraVariable& Var, UNiagaraDataInterface* DataInterface, bool bPlaceholder, bool bAddParameterMapRead)
 {
 	FString Id = DataInterface ? *DataInterface->GetMergeId().ToString() : TEXT("??");
diff --git a/Engine/Plugins/FX/Niagara/Source/NiagaraEditor/Private/NiagaraHlslTranslator.h b/Engine/Plugins/FX/Niagara/Source/NiagaraEditor/Private/NiagaraHlslTranslator.h
--- a/Engine/Plugins/FX/Niagara/Source/NiagaraEditor/Private/NiagaraHlslTranslator.h
+++ b/Engine/Plugins/FX/Niagara/Source/NiagaraEditor/Private/NiagaraHlslTranslator.h
@@ -567,6 +567,7 @@ public:
 
 	virtual int32 CompilePin(const UEdGraphPin* Pin);
 
+	virtual int32 RegisterUObject(FNiagaraVariable Variable, UObject* Object, bool bAddParameterMapRead);
 	virtual int32 RegisterDataInterface(FNiagaraVariable& Var, UNiagaraDataInterface* DataInterface, bool bPlaceholder, bool bAddParameterMapRead);
 
 	virtual void Operation(class UNiagaraNodeOp* Operation, TArray<int32>& Inputs, TArray<int32>& Outputs);
diff --git a/Engine/Plugins/FX/Niagara/Source/NiagaraEditor/Private/NiagaraNodeInput.cpp b/Engine/Plugins/FX/Niagara/Source/NiagaraEditor/Private/NiagaraNodeInput.cpp
--- a/Engine/Plugins/FX/Niagara/Source/NiagaraEditor/Private/NiagaraNodeInput.cpp
+++ b/Engine/Plugins/FX/Niagara/Source/NiagaraEditor/Private/NiagaraNodeInput.cpp
@@ -26,6 +26,7 @@ UNiagaraNodeInput::UNiagaraNodeInput(const FObjectInitializer& ObjectInitializer
 	, Usage(ENiagaraInputNodeUsage::Parameter)
 	, CallSortPriority(0)
 	, DataInterface(nullptr)
+	, ObjectAsset(nullptr)
 {
 	bCanRenameNode = true;
 }
@@ -34,30 +35,32 @@ void UNiagaraNodeInput::PostEditChangeProperty(struct FPropertyChangedEvent& Pro
 {
 	if (PropertyChangedEvent.Property != nullptr)
 	{
-		if (UClass* Class = const_cast<UClass*>(Input.GetType().GetClass()))
+		UClass* Class = Input.GetType().GetClass();
+
+		if (DataInterface)
 		{
-			check(Class->IsChildOf(UNiagaraDataInterface::StaticClass()));
-			if (DataInterface)
+			if (!IsDataInterface() || DataInterface->GetClass() != Class)
+			{
+				//Class has changed so clear this out and allocate pins will create a new instance of the correct type.
+				//Should we preserve old objects somewhere so settings aren't lost when switching around types?
+				DataInterface = nullptr;
+			}
+			else
 			{
-				if (DataInterface->GetClass() != Class)
+				//Keep it with the same name as the input 
+				if (PropertyChangedEvent.Property->GetName() == GET_MEMBER_NAME_CHECKED(UNiagaraNodeInput, Input))
 				{
-					//Class has changed so clear this out and allocate pins will create a new instance of the correct type.
-					//Should we preserve old objects somewhere so settings aren't lost when switching around types?
-					DataInterface = nullptr;
-				}
-				else
-				{
-					//Keep it with the same name as the input 
-					if (PropertyChangedEvent.Property->GetName() == TEXT("Input"))
-					{
-						DataInterface->Rename(*Input.GetName().ToString());
-					}
+					DataInterface->Rename(*Input.GetName().ToString());
 				}
 			}
 		}
-		else
+
+		if (ObjectAsset)
 		{
-			DataInterface = nullptr;
+			if (!IsObjectAsset() || ObjectAsset->GetClass() != Class)
+			{
+				ObjectAsset = nullptr;
+			}
 		}
 
 		ReallocatePins();
@@ -69,11 +72,16 @@ void UNiagaraNodeInput::PostEditChangeProperty(struct FPropertyChangedEvent& Pro
 void UNiagaraNodeInput::PostLoad()
 {
 	Super::PostLoad();
+
 	if (DataInterface != nullptr)
 	{
 		DataInterface->ConditionalPostLoad();
 		DataInterface->OnChanged().AddUObject(this, &UNiagaraNodeInput::DataInterfaceChanged);
 	}
+	if (ObjectAsset != nullptr)
+	{
+		ObjectAsset->ConditionalPostLoad();
+	}
 
 	ValidateDataInterface();
 }
@@ -133,18 +141,12 @@ void UNiagaraNodeInput::AppendFunctionAliasForContext(const FNiagaraGraphFunctio
 void UNiagaraNodeInput::AllocateDefaultPins()
 {
 	const FNiagaraTypeDefinition& InputType = Input.GetType();
-	if (InputType.IsUObject() && InputType.IsDataInterface() == false)
-	{
-		this->ErrorMsg = TEXT("Invalid Input Type");
-		return;
-	}
 
-	if (UClass* Class = const_cast<UClass*>(Input.GetType().GetClass()))
+	if (InputType.IsDataInterface())
 	{
-		check(Class->IsChildOf(UNiagaraDataInterface::StaticClass()));
-		if (!DataInterface)
+		if (DataInterface == nullptr)
 		{
-			DataInterface = NewObject<UNiagaraDataInterface>(this, Class, NAME_None, RF_Transactional | RF_Public);
+			DataInterface = NewObject<UNiagaraDataInterface>(this, InputType.GetClass(), NAME_None, RF_Transactional | RF_Public);
 		}
 	}
 
@@ -154,7 +156,7 @@ void UNiagaraNodeInput::AllocateDefaultPins()
 	UNiagaraScript* OwnerScript = GetTypedOuter<UNiagaraScript>();
 	if (OwnerScript)
 	{
-		if ((!IsRequired() && IsExposed()) && DataInterface == nullptr && Usage == ENiagaraInputNodeUsage::Parameter && (OwnerScript->IsStandaloneScript()))
+		if ((!IsRequired() && IsExposed()) && !IsDataInterface() && !IsObjectAsset() && Usage == ENiagaraInputNodeUsage::Parameter && (OwnerScript->IsStandaloneScript()))
 		{
 			UEdGraphPin* NewPin = CreatePin(EGPD_Input, Schema->TypeDefinitionToPinType(Input.GetType()), TEXT("Default"));
 			NewPin->bDefaultValueIsIgnored = true;
@@ -246,9 +248,9 @@ void UNiagaraNodeInput::OnRenameNode(const FString& NewName)
 	{
 		Node->Modify();
 		Node->Input.SetName(FName(*NewName));
-		if (DataInterface != nullptr)
+		if ( DataInterface )
 		{
-			Node->GetDataInterface()->Rename(*NewName);
+			DataInterface->Rename(*NewName);
 		}
 		Node->ReallocatePins(false);
 		Node->MarkNodeRequiresSynchronization("Input renamed", true);
@@ -337,6 +339,7 @@ void UNiagaraNodeInput::AutowireNewNode(UEdGraphPin* FromPin)
 						check(InputNode->ExposureOptions.bHidden == this->ExposureOptions.bHidden);
 						check(InputNode->ExposureOptions.bRequired == this->ExposureOptions.bRequired);
 						check(InputNode->DataInterface == this->DataInterface);
+						check(InputNode->ObjectAsset == this->ObjectAsset);
 					}
 
 					if (InputNode->CallSortPriority > HighestSortPriority)
@@ -408,16 +411,23 @@ void UNiagaraNodeInput::Compile(class FHlslNiagaraTranslator* Translator, TArray
 			if (Default == INDEX_NONE)
 			{
 				//We failed to compile the default pin so just use the value of the input.
-				if (Usage == ENiagaraInputNodeUsage::Parameter && DataInterface != nullptr)
+				if (Usage == ENiagaraInputNodeUsage::Parameter)
 				{
-					check(Input.GetType().GetClass());
-					Outputs.Add(Translator->RegisterDataInterface(Input, DataInterface, false, false));
-					return;
-				}
-				else
-				{
-					Default = Translator->GetConstant(Input);
+					if ( IsDataInterface() )
+					{
+						if (ensure(DataInterface != nullptr))
+						{
+							Outputs.Add(Translator->RegisterDataInterface(Input, DataInterface, false, false));
+							return;
+						}
+					}
+					else if ( IsObjectAsset() )
+					{
+						Outputs.Add(Translator->RegisterUObject(Input, ObjectAsset, false));
+						return;
+					}
 				}
+				Default = Translator->GetConstant(Input);
 			}
 			Outputs.Add(Default);
 			return;
@@ -427,14 +437,20 @@ void UNiagaraNodeInput::Compile(class FHlslNiagaraTranslator* Translator, TArray
 	switch (Usage)
 	{
 	case ENiagaraInputNodeUsage::Parameter:
-		if (DataInterface)
+		if (IsDataInterface())
+		{
+			Outputs.Add(Translator->RegisterDataInterface(Input, DataInterface, false, false));
+			break;
+		}
+		else if ( IsObjectAsset() )
 		{
-			check(Input.GetType().GetClass());
-			Outputs.Add(Translator->RegisterDataInterface(Input, DataInterface, false, false)); break;
+			Outputs.Add(Translator->RegisterUObject(Input, ObjectAsset, false));
+			break;
 		}
 		else
 		{
-			Outputs.Add(Translator->GetParameter(Input)); break;
+			Outputs.Add(Translator->GetParameter(Input));
+			break;
 		}
 	case ENiagaraInputNodeUsage::SystemConstant:
 		Outputs.Add(Translator->GetParameter(Input)); break;
@@ -470,27 +486,48 @@ void UNiagaraNodeInput::SortNodes(TArray<UNiagaraNodeInput*>& InOutNodes)
 	InOutNodes.Sort(SortVars);
 }
 
+bool UNiagaraNodeInput::IsDataInterface() const
+{
+	return Input.IsDataInterface();
+}
+
+bool UNiagaraNodeInput::IsObjectAsset() const
+{
+	return Input.IsUObject();
+}
+
 UNiagaraDataInterface* UNiagaraNodeInput::GetDataInterface() const
 {
-	return DataInterface;
+	return Cast<UNiagaraDataInterface>(DataInterface);
 }
 
 void UNiagaraNodeInput::SetDataInterface(UNiagaraDataInterface* InDataInterface)
 {
-	if (DataInterface != nullptr)
+	UNiagaraDataInterface* OldDataInterface = GetDataInterface();
+	if (OldDataInterface != nullptr)
 	{
-		DataInterface->OnChanged().RemoveAll(this);
+		OldDataInterface->OnChanged().RemoveAll(this);
 	}
 	DataInterface = InDataInterface;
-	if (DataInterface != nullptr)
+	if (InDataInterface != nullptr)
 	{
-		DataInterface->OnChanged().AddUObject(this, &UNiagaraNodeInput::DataInterfaceChanged);
+		InDataInterface->OnChanged().AddUObject(this, &UNiagaraNodeInput::DataInterfaceChanged);
 	}
 
 	ValidateDataInterface();
 	DataInterfaceChanged();
 }
 
+UObject* UNiagaraNodeInput::GetObjectAsset() const
+{
+	return ObjectAsset;
+}
+
+void UNiagaraNodeInput::SetObjectAsset(UObject* InObject)
+{
+	ObjectAsset = InObject;
+}
+
 void UNiagaraNodeInput::DataInterfaceChanged()
 {
 	// Don't use GetNiagaraGraph() here since this may be called on a temporary node which isn't
diff --git a/Engine/Plugins/FX/Niagara/Source/NiagaraEditor/Private/NiagaraNodeParameterMapGet.cpp b/Engine/Plugins/FX/Niagara/Source/NiagaraEditor/Private/NiagaraNodeParameterMapGet.cpp
--- a/Engine/Plugins/FX/Niagara/Source/NiagaraEditor/Private/NiagaraNodeParameterMapGet.cpp
+++ b/Engine/Plugins/FX/Niagara/Source/NiagaraEditor/Private/NiagaraNodeParameterMapGet.cpp
@@ -112,9 +112,9 @@ UEdGraphPin* UNiagaraNodeParameterMapGet::CreateDefaultPin(UEdGraphPin* OutputPi
 	DefaultPin->bDefaultValueIsReadOnly = true;
 
 	const UEdGraphSchema_Niagara* Schema = GetDefault<UEdGraphSchema_Niagara>();
-	FNiagaraTypeDefinition NiagaraType = Schema->PinToTypeDefinition(OutputPin);
-	bool bNeedsValue = NiagaraType.IsDataInterface() == false;
-	FNiagaraVariable Var = Schema->PinToNiagaraVariable(OutputPin, bNeedsValue);
+	const FNiagaraTypeDefinition NiagaraType = Schema->PinToTypeDefinition(OutputPin);
+	const bool bNeedsValue = NiagaraType.IsDataInterface() == false && NiagaraType.IsUObject() == false;
+	const FNiagaraVariable Var = Schema->PinToNiagaraVariable(OutputPin, bNeedsValue);
 
 	FString PinDefaultValue;
 	if (Schema->TryGetPinDefaultValueFromNiagaraVariable(Var, PinDefaultValue))
diff --git a/Engine/Plugins/FX/Niagara/Source/NiagaraEditor/Private/NiagaraScriptMergeManager.cpp b/Engine/Plugins/FX/Niagara/Source/NiagaraEditor/Private/NiagaraScriptMergeManager.cpp
--- a/Engine/Plugins/FX/Niagara/Source/NiagaraEditor/Private/NiagaraScriptMergeManager.cpp
+++ b/Engine/Plugins/FX/Niagara/Source/NiagaraEditor/Private/NiagaraScriptMergeManager.cpp
@@ -64,7 +64,6 @@ FNiagaraStackFunctionInputOverrideMergeAdapter::FNiagaraStackFunctionInputOverri
 	: OwningScript(&InOwningScript)
 	, OwningFunctionCallNode(&InOwningFunctionCallNode)
 	, OverridePin(&InOverridePin)
-	, DataValueObject(nullptr)
 {
 	
 	InputName = FNiagaraParameterHandle(OverridePin->PinName).GetName().ToString();
@@ -107,8 +106,16 @@ FNiagaraStackFunctionInputOverrideMergeAdapter::FNiagaraStackFunctionInputOverri
 		else if (OverridePin->LinkedTo[0]->GetOwningNode()->IsA<UNiagaraNodeInput>())
 		{
 			UNiagaraNodeInput* DataInputNode = CastChecked<UNiagaraNodeInput>(OverridePin->LinkedTo[0]->GetOwningNode());
-			DataValueInputName = DataInputNode->Input.GetName();
-			DataValueObject = DataInputNode->GetDataInterface();
+			if ( DataInputNode->IsDataInterface() )
+			{
+				DataInterfaceValueInputName = DataInputNode->Input.GetName();
+				DataInterfaceValue = DataInputNode->GetDataInterface();
+			}
+			else
+			{
+				ObjectAssetInputVariable = DataInputNode->Input;
+				ObjectAssetValue = DataInputNode->GetObjectAsset();
+			}
 		}
 		else if (OverridePin->LinkedTo[0]->GetOwningNode()->IsA<UNiagaraNodeFunctionCall>())
 		{
@@ -135,18 +142,13 @@ FNiagaraStackFunctionInputOverrideMergeAdapter::FNiagaraStackFunctionInputOverri
 	, OwningFunctionCallNode(&InOwningFunctionCallNode)
 	, InputName(InInputName)
 	, Type(InRapidIterationParameter.GetType())
-	, OverridePin(nullptr)
 	, LocalValueRapidIterationParameter(InRapidIterationParameter)
-	, DataValueObject(nullptr)
 {
 }
 
 FNiagaraStackFunctionInputOverrideMergeAdapter::FNiagaraStackFunctionInputOverrideMergeAdapter(UEdGraphPin* InStaticSwitchPin)
-	: OwningScript(nullptr)
-	, OwningFunctionCallNode(CastChecked<UNiagaraNodeFunctionCall>(InStaticSwitchPin->GetOwningNode()))
+	: OwningFunctionCallNode(CastChecked<UNiagaraNodeFunctionCall>(InStaticSwitchPin->GetOwningNode()))
 	, InputName(InStaticSwitchPin->PinName.ToString())
-	, OverridePin(nullptr)
-	, DataValueObject(nullptr)
 	, StaticSwitchValue(InStaticSwitchPin->DefaultValue)
 {
 	const UEdGraphSchema_Niagara* NiagaraSchema = GetDefault<UEdGraphSchema_Niagara>();
@@ -203,14 +205,24 @@ TOptional<FNiagaraStackLinkedValueData> FNiagaraStackFunctionInputOverrideMergeA
 	return LinkedValueData;
 }
 
-TOptional<FName> FNiagaraStackFunctionInputOverrideMergeAdapter::GetDataValueInputName() const
+TOptional<FName> FNiagaraStackFunctionInputOverrideMergeAdapter::GetDataInterfaceValueInputName() const
+{
+	return DataInterfaceValueInputName;
+}
+
+TOptional<FNiagaraVariableBase> FNiagaraStackFunctionInputOverrideMergeAdapter::GetObjectAssetInputVariable() const
+{
+	return ObjectAssetInputVariable;
+}
+
+UObject* FNiagaraStackFunctionInputOverrideMergeAdapter::GetObjectAssetValue() const
 {
-	return DataValueInputName;
+	return ObjectAssetValue;
 }
 
-UNiagaraDataInterface* FNiagaraStackFunctionInputOverrideMergeAdapter::GetDataValueObject() const
+UNiagaraDataInterface* FNiagaraStackFunctionInputOverrideMergeAdapter::GetDataInterfaceValue() const
 {
-	return DataValueObject;
+	return DataInterfaceValue;
 }
 
 TSharedPtr<FNiagaraStackFunctionMergeAdapter> FNiagaraStackFunctionInputOverrideMergeAdapter::GetDynamicValueFunction() const
@@ -2728,20 +2740,33 @@ TOptional<bool> FNiagaraScriptMergeManager::DoFunctionInputOverridesMatch(TShare
 	}
 
 	// Data value
-	if ((BaseFunctionInputAdapter->GetDataValueInputName().IsSet() && OtherFunctionInputAdapter->GetDataValueInputName().IsSet() == false) ||
-		(BaseFunctionInputAdapter->GetDataValueInputName().IsSet() == false && OtherFunctionInputAdapter->GetDataValueInputName().IsSet()) ||
-		(BaseFunctionInputAdapter->GetDataValueObject() != nullptr && OtherFunctionInputAdapter->GetDataValueObject() == nullptr) ||
-		(BaseFunctionInputAdapter->GetDataValueObject() == nullptr && OtherFunctionInputAdapter->GetDataValueObject() != nullptr))
+	if ((BaseFunctionInputAdapter->GetDataInterfaceValueInputName().IsSet() && OtherFunctionInputAdapter->GetDataInterfaceValueInputName().IsSet() == false) ||
+		(BaseFunctionInputAdapter->GetDataInterfaceValueInputName().IsSet() == false && OtherFunctionInputAdapter->GetDataInterfaceValueInputName().IsSet()) ||
+		(BaseFunctionInputAdapter->GetDataInterfaceValue() != nullptr && OtherFunctionInputAdapter->GetDataInterfaceValue() == nullptr) ||
+		(BaseFunctionInputAdapter->GetDataInterfaceValue() == nullptr && OtherFunctionInputAdapter->GetDataInterfaceValue() != nullptr))
 	{
 		return false;
 	}
 
-	if (BaseFunctionInputAdapter->GetDataValueInputName().IsSet() && OtherFunctionInputAdapter->GetDataValueInputName().IsSet() &&
-		BaseFunctionInputAdapter->GetDataValueObject() != nullptr && OtherFunctionInputAdapter->GetDataValueObject() != nullptr)
+	if (BaseFunctionInputAdapter->GetDataInterfaceValueInputName().IsSet() && OtherFunctionInputAdapter->GetDataInterfaceValueInputName().IsSet() &&
+		BaseFunctionInputAdapter->GetDataInterfaceValue() != nullptr && OtherFunctionInputAdapter->GetDataInterfaceValue() != nullptr)
 	{
 		return 
-			BaseFunctionInputAdapter->GetDataValueInputName().GetValue() == OtherFunctionInputAdapter->GetDataValueInputName().GetValue() &&
-			BaseFunctionInputAdapter->GetDataValueObject()->Equals(OtherFunctionInputAdapter->GetDataValueObject());
+			BaseFunctionInputAdapter->GetDataInterfaceValueInputName().GetValue() == OtherFunctionInputAdapter->GetDataInterfaceValueInputName().GetValue() &&
+			BaseFunctionInputAdapter->GetDataInterfaceValue()->Equals(OtherFunctionInputAdapter->GetDataInterfaceValue());
+	}
+
+	// Object Reference
+	if (BaseFunctionInputAdapter->GetObjectAssetInputVariable().IsSet() != OtherFunctionInputAdapter->GetObjectAssetInputVariable().IsSet())
+	{
+		return false;
+	}
+
+	if (BaseFunctionInputAdapter->GetObjectAssetInputVariable().IsSet() && OtherFunctionInputAdapter->GetObjectAssetInputVariable().IsSet())
+	{
+		return
+			BaseFunctionInputAdapter->GetObjectAssetInputVariable().GetValue() == OtherFunctionInputAdapter->GetObjectAssetInputVariable().GetValue() &&
+			BaseFunctionInputAdapter->GetObjectAssetValue() == OtherFunctionInputAdapter->GetObjectAssetValue();
 	}
 
 	// Dynamic value
@@ -3054,15 +3079,22 @@ FNiagaraScriptMergeManager::FApplyDiffResults FNiagaraScriptMergeManager::AddInp
 				}
 				Results.bSucceeded = true;
 			}
-			else if (OverrideToAdd->GetDataValueInputName().IsSet() && OverrideToAdd->GetDataValueObject() != nullptr)
+			else if (OverrideToAdd->GetDataInterfaceValueInputName().IsSet() && OverrideToAdd->GetDataInterfaceValue() != nullptr)
 			{
-				FName OverrideValueInputName = OverrideToAdd->GetDataValueInputName().GetValue();
-				UNiagaraDataInterface* OverrideValueObject = OverrideToAdd->GetDataValueObject();
+				FName OverrideValueInputName = OverrideToAdd->GetDataInterfaceValueInputName().GetValue();
+				UNiagaraDataInterface* OverrideValueObject = OverrideToAdd->GetDataInterfaceValue();
 				UNiagaraDataInterface* NewOverrideValueObject;
-				FNiagaraStackGraphUtilities::SetDataValueObjectForFunctionInput(InputOverridePin, OverrideToAdd->GetDataValueObject()->GetClass(), OverrideValueInputName.ToString(), NewOverrideValueObject, OverrideToAdd->GetOverrideNodeId());
+				FNiagaraStackGraphUtilities::SetDataInterfaceValueForFunctionInput(InputOverridePin, OverrideToAdd->GetDataInterfaceValue()->GetClass(), OverrideValueInputName.ToString(), NewOverrideValueObject, OverrideToAdd->GetOverrideNodeId());
 				OverrideValueObject->CopyTo(NewOverrideValueObject);
 				Results.bSucceeded = true;
 			}
+			else if (OverrideToAdd->GetObjectAssetInputVariable().IsSet())
+			{
+				FNiagaraVariableBase OverrideVariable = OverrideToAdd->GetObjectAssetInputVariable().GetValue();
+				UObject* OverrideObjectValue = OverrideToAdd->GetObjectAssetValue();
+				FNiagaraStackGraphUtilities::SetObjectAssetValueForFunctionInput(InputOverridePin, OverrideVariable.GetType().GetClass(), OverrideVariable.GetName().ToString(), OverrideObjectValue);
+				Results.bSucceeded = true;
+			}
 			else if (OverrideToAdd->GetDynamicValueFunction().IsValid())
 			{
 				UNiagaraNodeCustomHlsl* CustomHlslFunction = Cast<UNiagaraNodeCustomHlsl>(OverrideToAdd->GetDynamicValueFunction()->GetFunctionCallNode());
diff --git a/Engine/Plugins/FX/Niagara/Source/NiagaraEditor/Private/NiagaraScriptMergeManager.h b/Engine/Plugins/FX/Niagara/Source/NiagaraEditor/Private/NiagaraScriptMergeManager.h
--- a/Engine/Plugins/FX/Niagara/Source/NiagaraEditor/Private/NiagaraScriptMergeManager.h
+++ b/Engine/Plugins/FX/Niagara/Source/NiagaraEditor/Private/NiagaraScriptMergeManager.h
@@ -71,8 +71,10 @@ public:
 	TOptional<FString> GetLocalValueString() const;
 	TOptional<FNiagaraVariable> GetLocalValueRapidIterationParameter() const;
 	TOptional<FNiagaraStackLinkedValueData> GetLinkedValueData() const;
-	TOptional<FName> GetDataValueInputName() const;
-	UNiagaraDataInterface* GetDataValueObject() const;
+	TOptional<FName> GetDataInterfaceValueInputName() const;
+	UNiagaraDataInterface* GetDataInterfaceValue() const;
+	TOptional<FNiagaraVariableBase> GetObjectAssetInputVariable() const;
+	UObject* GetObjectAssetValue() const;
 	TSharedPtr<FNiagaraStackFunctionMergeAdapter> GetDynamicValueFunction() const;
 	TOptional<FString> GetStaticSwitchValue() const;
 
@@ -83,13 +85,15 @@ private:
 	FNiagaraTypeDefinition Type;
 
 	TWeakObjectPtr<UNiagaraNodeParameterMapSet> OverrideNode;
-	UEdGraphPin* OverridePin;
+	UEdGraphPin* OverridePin = nullptr;
 
 	TOptional<FString> LocalValueString;
 	TOptional<FNiagaraVariable> LocalValueRapidIterationParameter;
 	TOptional<FNiagaraStackLinkedValueData> LinkedValueData;
-	TOptional<FName> DataValueInputName;
-	UNiagaraDataInterface* DataValueObject;
+	TOptional<FName> DataInterfaceValueInputName;
+	UNiagaraDataInterface* DataInterfaceValue = nullptr;
+	TOptional<FNiagaraVariableBase> ObjectAssetInputVariable;
+	UObject* ObjectAssetValue = nullptr;
 	TSharedPtr<FNiagaraStackFunctionMergeAdapter> DynamicValueFunction;
 	TOptional<FString> StaticSwitchValue;
 
diff --git a/Engine/Plugins/FX/Niagara/Source/NiagaraEditor/Private/ViewModels/NiagaraCollectionParameterViewModel.cpp b/Engine/Plugins/FX/Niagara/Source/NiagaraEditor/Private/ViewModels/NiagaraCollectionParameterViewModel.cpp
--- a/Engine/Plugins/FX/Niagara/Source/NiagaraEditor/Private/ViewModels/NiagaraCollectionParameterViewModel.cpp
+++ b/Engine/Plugins/FX/Niagara/Source/NiagaraEditor/Private/ViewModels/NiagaraCollectionParameterViewModel.cpp
@@ -15,7 +15,19 @@ FNiagaraCollectionParameterViewModel::FNiagaraCollectionParameterViewModel(FNiag
 	, CollectionInst(CollectionInstance)
 	, Parameter(Variable)
 {
-	DefaultValueType = Variable.IsDataInterface() ? INiagaraParameterViewModel::EDefaultValueType::Object : INiagaraParameterViewModel::EDefaultValueType::Struct;
+	if (Variable.IsDataInterface())
+	{
+		DefaultValueType = INiagaraParameterViewModel::EDefaultValueType::DataInterface;
+	}
+	else if (Variable.IsUObject())
+	{
+		DefaultValueType = INiagaraParameterViewModel::EDefaultValueType::ObjectAsset;
+	}
+	else
+	{
+		DefaultValueType = INiagaraParameterViewModel::EDefaultValueType::Struct;
+	}
+
 	RefreshParameterValue();
 }
 
@@ -166,27 +178,45 @@ TSharedRef<FStructOnScope> FNiagaraCollectionParameterViewModel::GetDefaultValue
 	return ParameterValue.ToSharedRef();
 }
 
-UObject* FNiagaraCollectionParameterViewModel::GetDefaultValueObject()
+UNiagaraDataInterface* FNiagaraCollectionParameterViewModel::GetDefaultValueDataInterface()
 {
 	return CollectionInst->GetParameterStore().GetDataInterface(Parameter);
 }
 
+UObject* FNiagaraCollectionParameterViewModel::GetDefaultValueObjectAsset()
+{
+	return CollectionInst->GetParameterStore().GetUObject(Parameter);
+}
+
+void FNiagaraCollectionParameterViewModel::SetDefaultValueObjectAsset(UObject* Object)
+{
+	CollectionInst->GetParameterStore().SetUObject(Object, Parameter);
+}
+
 void FNiagaraCollectionParameterViewModel::NotifyDefaultValuePropertyChanged(const FPropertyChangedEvent& PropertyChangedEvent)
 {
 	if (!IsOptional() || IsProvided() == ECheckBoxState::Checked)
 	{
 		FScopedTransaction ScopedTransaction(LOCTEXT("EditParameterValueProperty", "Edit parameter value"));
 		CollectionInst->Modify();
-		if (DefaultValueType == INiagaraParameterViewModel::EDefaultValueType::Struct)
-		{
-			CollectionInst->GetParameterStore().SetParameterData(ParameterValue->GetStructMemory(), Parameter);
-		}
-		else
+		switch (DefaultValueType)
 		{
-			if (UNiagaraDataInterface* Interface = CollectionInst->GetParameterStore().GetDataInterface(Parameter))
-			{
-		 		Interface->Modify();
-			}
+
+			case INiagaraParameterViewModel::EDefaultValueType::Struct:
+				CollectionInst->GetParameterStore().SetParameterData(ParameterValue->GetStructMemory(), Parameter);
+				break;
+
+			case INiagaraParameterViewModel::EDefaultValueType::DataInterface:
+				if (UNiagaraDataInterface* Interface = CollectionInst->GetParameterStore().GetDataInterface(Parameter))
+				{
+					Interface->Modify();
+				}
+				break;
+
+			case INiagaraParameterViewModel::EDefaultValueType::ObjectAsset:
+			default:
+				checkNoEntry();
+				break;
 		}
 		OnDefaultValueChangedDelegate.Broadcast();
 	}
@@ -225,7 +255,7 @@ void FNiagaraCollectionParameterViewModel::NotifyDefaultValueChanged()
 {
 	if (!IsOptional() || IsProvided() == ECheckBoxState::Checked)
 	{
-		if (!Parameter.GetType().IsDataInterface())
+		if (!Parameter.GetType().IsDataInterface() && !Parameter.GetType().IsUObject())
 		{
 			const uint8* ParamData = CollectionInst->GetParameterStore().GetParameterData(Parameter);
 			check(ParamData);
@@ -248,7 +278,8 @@ FNiagaraCollectionParameterViewModel::FOnNameChanged& FNiagaraCollectionParamete
 void FNiagaraCollectionParameterViewModel::RefreshParameterValue()
 {
 	FNiagaraTypeDefinition Type = Parameter.GetType();
-	if (!Type.IsDataInterface())
+
+	if (!Type.IsDataInterface() && !Type.IsUObject())
 	{
 		ParameterValue = MakeShareable(new FStructOnScope(Parameter.GetType().GetStruct()));
 
diff --git a/Engine/Plugins/FX/Niagara/Source/NiagaraEditor/Private/ViewModels/NiagaraCollectionParameterViewModel.h b/Engine/Plugins/FX/Niagara/Source/NiagaraEditor/Private/ViewModels/NiagaraCollectionParameterViewModel.h
--- a/Engine/Plugins/FX/Niagara/Source/NiagaraEditor/Private/ViewModels/NiagaraCollectionParameterViewModel.h
+++ b/Engine/Plugins/FX/Niagara/Source/NiagaraEditor/Private/ViewModels/NiagaraCollectionParameterViewModel.h
@@ -26,7 +26,9 @@ public:
 	virtual void SelectedTypeChanged(TSharedPtr<FNiagaraTypeDefinition> Item, ESelectInfo::Type SelectionType) override;
 	virtual EDefaultValueType GetDefaultValueType() override;
 	virtual TSharedRef<FStructOnScope> GetDefaultValueStruct() override;
-	virtual UObject* GetDefaultValueObject() override;
+	virtual UNiagaraDataInterface* GetDefaultValueDataInterface() override;
+	virtual UObject* GetDefaultValueObjectAsset() override;
+	virtual void SetDefaultValueObjectAsset(UObject*) override;
 	virtual void NotifyDefaultValuePropertyChanged(const FPropertyChangedEvent& PropertyChangedEvent) override;
 	virtual void NotifyBeginDefaultValueChange() override;
 	virtual void NotifyEndDefaultValueChange() override;
@@ -62,7 +64,6 @@ private:
 	FOnNameChanged OnNameChangedDelegate;
 
 	UNiagaraParameterCollectionInstance* CollectionInst;
-	//FNiagaraParameterDirectBinding ParameterBinding;
 	FNiagaraVariable Parameter;
 };
 
diff --git a/Engine/Plugins/FX/Niagara/Source/NiagaraEditor/Private/ViewModels/NiagaraParameterCollectionAssetViewModel.cpp b/Engine/Plugins/FX/Niagara/Source/NiagaraEditor/Private/ViewModels/NiagaraParameterCollectionAssetViewModel.cpp
--- a/Engine/Plugins/FX/Niagara/Source/NiagaraEditor/Private/ViewModels/NiagaraParameterCollectionAssetViewModel.cpp
+++ b/Engine/Plugins/FX/Niagara/Source/NiagaraEditor/Private/ViewModels/NiagaraParameterCollectionAssetViewModel.cpp
@@ -127,7 +127,7 @@ void FNiagaraParameterCollectionAssetViewModel::AddParameter(TSharedPtr<FNiagara
 	int32 ParamIdx = Collection->AddParameter(NewName, Type);
 
 	//TODO: It'd be nice to be able to get a default value for types in runtime code and do this inside the parameter store itself.
-	if (!Type.IsDataInterface())
+	if (!Type.IsDataInterface() && !Type.IsUObject())
 	{
 		TArray<uint8> DefaultData;
 		FNiagaraEditorUtilities::GetTypeDefaultValue(Type, DefaultData);
@@ -391,7 +391,7 @@ void FNiagaraParameterCollectionAssetViewModel::OnParameterTypeChanged(FNiagaraV
 	Collection->GetDefaultInstance()->AddParameter(Collection->GetParameters()[Index]);
 	
 	//TODO: It'd be nice to be able to get a default value for types in runtime code and do this inside the parameter store itself.
-	if (!Type.IsDataInterface())
+	if (!Type.IsDataInterface() && !Type.IsUObject())
 	{
 		TArray<uint8> DefaultData;
 		FNiagaraEditorUtilities::GetTypeDefaultValue(Type, DefaultData);
diff --git a/Engine/Plugins/FX/Niagara/Source/NiagaraEditor/Private/ViewModels/NiagaraParameterPanelViewModel.cpp b/Engine/Plugins/FX/Niagara/Source/NiagaraEditor/Private/ViewModels/NiagaraParameterPanelViewModel.cpp
--- a/Engine/Plugins/FX/Niagara/Source/NiagaraEditor/Private/ViewModels/NiagaraParameterPanelViewModel.cpp
+++ b/Engine/Plugins/FX/Niagara/Source/NiagaraEditor/Private/ViewModels/NiagaraParameterPanelViewModel.cpp
@@ -267,13 +267,10 @@ void INiagaraParameterPanelViewModel::RefreshFull(bool bDoCategoryExpansion) con
 
 bool INiagaraParameterPanelViewModel::CanAddType(const FNiagaraTypeDefinition& InType, TOptional<FNiagaraParameterPanelCategory> Category)
 {
-	bool bCanAddObjects = Category.IsSet() && Category.GetValue().NamespaceMetaData.GetGuid() == FNiagaraEditorGuids::UserNamespaceMetaDataGuid;
-	
 	return !(
 		InType == FNiagaraTypeDefinition::GetParameterMapDef()
 		|| InType == FNiagaraTypeDefinition::GetGenericNumericDef()
 		|| InType.IsInternalType()
-		|| ((!bCanAddObjects && InType.IsUObject()) && !InType.IsDataInterface())
 		);
 }
 
diff --git a/Engine/Plugins/FX/Niagara/Source/NiagaraEditor/Private/ViewModels/NiagaraParameterViewModel.h b/Engine/Plugins/FX/Niagara/Source/NiagaraEditor/Private/ViewModels/NiagaraParameterViewModel.h
--- a/Engine/Plugins/FX/Niagara/Source/NiagaraEditor/Private/ViewModels/NiagaraParameterViewModel.h
+++ b/Engine/Plugins/FX/Niagara/Source/NiagaraEditor/Private/ViewModels/NiagaraParameterViewModel.h
@@ -11,6 +11,7 @@
 
 struct FNiagaraTypeDefinition;
 struct FNiagaraVariable;
+class UNiagaraDataInterface;
 
 /** Defines the view model for a parameter in the parameter collection editor. */
 class INiagaraParameterViewModel
@@ -26,7 +27,8 @@ public:
 	{
 		None,
 		Struct,
-		Object
+		DataInterface,
+		ObjectAsset,
 	};
 
 public:
@@ -65,8 +67,13 @@ public:
 	/** Gets the struct representing the default value for the parameter. */
 	virtual TSharedRef<FStructOnScope> GetDefaultValueStruct() = 0;
 
-	/** Gets the object representing the default value for the parameter. */
-	virtual UObject* GetDefaultValueObject() = 0;
+	/** Gets the data interface representing the default value for the parameter. */
+	virtual UNiagaraDataInterface* GetDefaultValueDataInterface() = 0;
+
+	/** Gets the object reference representing the default value for the parameter. */
+	virtual UObject* GetDefaultValueObjectAsset() = 0;
+	/** Sets the object reference representing the default value for the parameter. */
+	virtual void SetDefaultValueObjectAsset(UObject*) = 0;
 
 	/** Called to notify the parameter view model that a property on the default value has been changed by the UI. */
 	virtual void NotifyDefaultValuePropertyChanged(const FPropertyChangedEvent& PropertyChangedEvent) = 0;
diff --git a/Engine/Plugins/FX/Niagara/Source/NiagaraEditor/Private/ViewModels/NiagaraPlaceholderDataInterfaceManager.cpp b/Engine/Plugins/FX/Niagara/Source/NiagaraEditor/Private/ViewModels/NiagaraPlaceholderDataInterfaceManager.cpp
--- a/Engine/Plugins/FX/Niagara/Source/NiagaraEditor/Private/ViewModels/NiagaraPlaceholderDataInterfaceManager.cpp
+++ b/Engine/Plugins/FX/Niagara/Source/NiagaraEditor/Private/ViewModels/NiagaraPlaceholderDataInterfaceManager.cpp
@@ -145,7 +145,7 @@ void FNiagaraPlaceholderDataInterfaceManager::PlaceholderDataInterfaceChanged(TW
 				{
 					FNiagaraStackGraphUtilities::RemoveNodesForStackFunctionInputOverridePin(OverridePin);
 				}
-				FNiagaraStackGraphUtilities::SetDataValueObjectForFunctionInput(OverridePin, PlaceholderDataInterface->GetClass(), AliasedInputParameterHandle.GetParameterHandleString().ToString(), OverrideDataInterface);
+				FNiagaraStackGraphUtilities::SetDataInterfaceValueForFunctionInput(OverridePin, PlaceholderDataInterface->GetClass(), AliasedInputParameterHandle.GetParameterHandleString().ToString(), OverrideDataInterface);
 				FNiagaraStackGraphUtilities::RelayoutGraph(*PlaceholderDataInterfaceInfo->OwningFunctionCall->GetNiagaraGraph());
 			}
 
diff --git a/Engine/Plugins/FX/Niagara/Source/NiagaraEditor/Private/ViewModels/NiagaraScriptInputCollectionViewModel.cpp b/Engine/Plugins/FX/Niagara/Source/NiagaraEditor/Private/ViewModels/NiagaraScriptInputCollectionViewModel.cpp
--- a/Engine/Plugins/FX/Niagara/Source/NiagaraEditor/Private/ViewModels/NiagaraScriptInputCollectionViewModel.cpp
+++ b/Engine/Plugins/FX/Niagara/Source/NiagaraEditor/Private/ViewModels/NiagaraScriptInputCollectionViewModel.cpp
@@ -275,7 +275,7 @@ void FNiagaraScriptInputCollectionViewModel::RefreshParameterViewModels()
 				ParameterViewModel = MakeShareable(new FNiagaraScriptParameterViewModel(GraphVariable, *InputNode, EmitterVariable, Script, ParameterEditMode));
 				
 			}
-			else
+			else if (InputNode->Input.IsDataInterface())
 			{
 				UNiagaraDataInterface* EmitterDataInterface = InputNode->GetDataInterface();
 				for (FVersionedNiagaraScriptWeakPtr& ScriptWeakPtr : Scripts)
@@ -296,6 +296,11 @@ void FNiagaraScriptInputCollectionViewModel::RefreshParameterViewModels()
 				}
 				ParameterViewModel = MakeShareable(new FNiagaraScriptParameterViewModel(GraphVariable, *InputNode, EmitterDataInterface, ParameterEditMode));
 			}
+			else if (InputNode->Input.IsUObject())
+			{
+				UObject* ObjectAsset = InputNode->GetObjectAsset();
+				ParameterViewModel = MakeShareable(new FNiagaraScriptParameterViewModel(GraphVariable, *InputNode, ObjectAsset, ParameterEditMode));
+			}
 
 			ParameterViewModel->OnNameChanged().AddSP(this, &FNiagaraScriptInputCollectionViewModel::OnParameterNameChanged, TWeakObjectPtr<UNiagaraNodeInput>(InputNode));
 			ParameterViewModel->OnTypeChanged().AddSP(this, &FNiagaraScriptInputCollectionViewModel::OnParameterTypeChanged, &GraphVariable);
@@ -310,11 +315,11 @@ void FNiagaraScriptInputCollectionViewModel::RefreshParameterViewModels()
 
 bool FNiagaraScriptInputCollectionViewModel::SupportsType(const FNiagaraTypeDefinition& Type) const
 {
-	if (Type.IsUObject() && Type.IsDataInterface() == false)
-	{
-		// Don't allow generic objects as script inputs.
-		return false;
-	}
+	//if (Type.IsUObject() && Type.IsDataInterface() == false)
+	//{
+	//	// Don't allow generic objects as script inputs.
+	//	return false;
+	//}
 
 	if (Type.IsStatic())
 	{
@@ -331,7 +336,7 @@ bool FNiagaraScriptInputCollectionViewModel::SupportsType(const FNiagaraTypeDefi
 		// We only support parameter map inputs for dynamic inputs and modules, with the ability to create data interfaces as needed for defaults.
 		if (Scripts[0].Script->GetUsage() == ENiagaraScriptUsage::DynamicInput || Scripts[0].Script->GetUsage() == ENiagaraScriptUsage::Module)
 		{
-			if (Type != FNiagaraTypeDefinition::GetParameterMapDef() && !Type.IsDataInterface())
+			if (Type != FNiagaraTypeDefinition::GetParameterMapDef() && !Type.IsDataInterface() && !Type.IsUObject())
 			{
 				return false;
 			}
@@ -419,6 +424,7 @@ void FNiagaraScriptInputCollectionViewModel::OnParameterTypeChanged(FNiagaraVari
 			InputNodeToUpdate->Modify();
 			InputNodeToUpdate->Input = FirstNodeToUpdate->Input;
 			InputNodeToUpdate->SetDataInterface(FirstNodeToUpdate->GetDataInterface());
+			InputNodeToUpdate->SetObjectAsset(FirstNodeToUpdate->GetObjectAsset());
 			InputNodeToUpdate->NotifyInputTypeChanged();
 		}
 	}
@@ -447,7 +453,8 @@ void FNiagaraScriptInputCollectionViewModel::OnParameterValueChangedInternal(TSh
 	// Since we potentially have multiple input nodes that point to the exact same underlying input variable, we need to make sure that 
 	// all of them are synchronized in their values or else we might end up with confusion when we finally compile. 
 	if (ChangedParameter->GetDefaultValueType() == INiagaraParameterViewModel::EDefaultValueType::Struct || 
-		ChangedParameter->GetDefaultValueType() == INiagaraParameterViewModel::EDefaultValueType::Object)
+		ChangedParameter->GetDefaultValueType() == INiagaraParameterViewModel::EDefaultValueType::DataInterface ||
+		ChangedParameter->GetDefaultValueType() == INiagaraParameterViewModel::EDefaultValueType::ObjectAsset)
 	{
 		TArray<UNiagaraNodeInput*> InputNodes;
 		TArray<UNiagaraNodeInput*> InputNodesToUpdate;
@@ -497,9 +504,9 @@ void FNiagaraScriptInputCollectionViewModel::OnParameterValueChangedInternal(TSh
 				}
 			}
 		}
-		else if (ChangedParameter->GetDefaultValueType() == INiagaraParameterViewModel::EDefaultValueType::Object)
+		else if (ChangedParameter->GetDefaultValueType() == INiagaraParameterViewModel::EDefaultValueType::DataInterface)
 		{
-			UNiagaraDataInterface* DataInterface = Cast<UNiagaraDataInterface>(ChangedParameter->GetDefaultValueObject());
+			UNiagaraDataInterface* DataInterface = ChangedParameter->GetDefaultValueDataInterface();
 			if (DataInterface != nullptr)
 			{
 				TSet<UNiagaraDataInterface*> DataInterfacesToUpdate;
@@ -535,6 +542,20 @@ void FNiagaraScriptInputCollectionViewModel::OnParameterValueChangedInternal(TSh
 				}
 			}
 		}
+		else if (ChangedParameter->GetDefaultValueType() == INiagaraParameterViewModel::EDefaultValueType::ObjectAsset)
+		{
+			UObject* NewObject = ChangedParameter->GetDefaultValueObjectAsset();
+
+			for (UNiagaraNodeInput* InputNodeToUpdate : InputNodesToUpdate)
+			{
+				InputNodeToUpdate->Modify();
+				InputNodeToUpdate->SetObjectAsset(NewObject);
+			}
+		}
+		else
+		{
+			checkNoEntry();
+		}
 	}
 
 	OnParameterValueChanged().Broadcast(ChangedParameter->GetName());
diff --git a/Engine/Plugins/FX/Niagara/Source/NiagaraEditor/Private/ViewModels/NiagaraScriptParameterViewModel.cpp b/Engine/Plugins/FX/Niagara/Source/NiagaraEditor/Private/ViewModels/NiagaraScriptParameterViewModel.cpp
--- a/Engine/Plugins/FX/Niagara/Source/NiagaraEditor/Private/ViewModels/NiagaraScriptParameterViewModel.cpp
+++ b/Engine/Plugins/FX/Niagara/Source/NiagaraEditor/Private/ViewModels/NiagaraScriptParameterViewModel.cpp
@@ -15,7 +15,6 @@ FNiagaraScriptParameterViewModel::FNiagaraScriptParameterViewModel(FNiagaraVaria
 	, GraphVariable(&InGraphVariable)
 	, CompiledVariable(InCompiledVariable)
 	, CompiledVariableOwner(InCompiledVariableOwner)
-	, ValueObject(nullptr)
 {
 	GraphVariableOwner = &InGraphVariableOwner;
 	bool bUsingCompiledVariable = InCompiledVariable != nullptr;
@@ -29,15 +28,20 @@ FNiagaraScriptParameterViewModel::FNiagaraScriptParameterViewModel(FNiagaraVaria
 FNiagaraScriptParameterViewModel::FNiagaraScriptParameterViewModel(FNiagaraVariable& InGraphVariable, UObject& InGraphVariableOwner, UObject* InValueObject, ENiagaraParameterEditMode ParameterEditMode)
 	: FNiagaraParameterViewModel(ParameterEditMode)
 	, GraphVariable(&InGraphVariable)
-	, CompiledVariable(nullptr)
-	, CompiledVariableOwner(nullptr)
-	, ValueVariable(nullptr)
-	, ValueVariableOwner(nullptr)
-	, ValueObject(InValueObject)
 {
 	GraphVariableOwner = &InGraphVariableOwner;
-	DefaultValueType = INiagaraParameterViewModel::EDefaultValueType::Object;
 	DebugName = GraphVariable->GetName().ToString();
+	if (InGraphVariable.IsDataInterface())
+	{
+		ValueDataInterface = Cast<UNiagaraDataInterface>(InValueObject);
+		DefaultValueType = INiagaraParameterViewModel::EDefaultValueType::DataInterface;
+	}
+	else
+	{
+		ValueObjectAsset = InValueObject;
+		DefaultValueType = INiagaraParameterViewModel::EDefaultValueType::ObjectAsset;
+	}
+
 	//UE_LOG(LogNiagaraEditor, Log, TEXT("Create %p Var %s"), this, *DebugName);
 }
 
@@ -54,7 +58,8 @@ void FNiagaraScriptParameterViewModel::Reset()
 	CompiledVariableOwner = nullptr;
 	ValueVariable = nullptr;
 	ValueVariableOwner = nullptr;
-	ValueObject = nullptr;
+	ValueDataInterface = nullptr;
+	ValueObjectAsset = nullptr;
 }
 
 
@@ -150,9 +155,20 @@ TSharedRef<FStructOnScope> FNiagaraScriptParameterViewModel::GetDefaultValueStru
 	return ParameterValue.ToSharedRef();
 }
 
-UObject* FNiagaraScriptParameterViewModel::GetDefaultValueObject()
+UNiagaraDataInterface* FNiagaraScriptParameterViewModel::GetDefaultValueDataInterface()
 {
-	return ValueObject;
+	return ValueDataInterface;
+}
+
+UObject* FNiagaraScriptParameterViewModel::GetDefaultValueObjectAsset()
+{
+	return ValueObjectAsset;
+}
+
+void FNiagaraScriptParameterViewModel::SetDefaultValueObjectAsset(UObject* Object)
+{
+	ValueObjectAsset = Object;
+	OnDefaultValueChangedDelegate.Broadcast();
 }
 
 void FNiagaraScriptParameterViewModel::NotifyDefaultValuePropertyChanged(const FPropertyChangedEvent& PropertyChangedEvent)
diff --git a/Engine/Plugins/FX/Niagara/Source/NiagaraEditor/Private/ViewModels/NiagaraScriptParameterViewModel.h b/Engine/Plugins/FX/Niagara/Source/NiagaraEditor/Private/ViewModels/NiagaraScriptParameterViewModel.h
--- a/Engine/Plugins/FX/Niagara/Source/NiagaraEditor/Private/ViewModels/NiagaraScriptParameterViewModel.h
+++ b/Engine/Plugins/FX/Niagara/Source/NiagaraEditor/Private/ViewModels/NiagaraScriptParameterViewModel.h
@@ -44,7 +44,9 @@ public:
 	virtual void SelectedTypeChanged(TSharedPtr<FNiagaraTypeDefinition> Item, ESelectInfo::Type SelectionType) override;
 	virtual EDefaultValueType GetDefaultValueType() override;
 	virtual TSharedRef<FStructOnScope> GetDefaultValueStruct() override;
-	virtual UObject* GetDefaultValueObject() override;
+	virtual UNiagaraDataInterface* GetDefaultValueDataInterface() override;
+	virtual UObject* GetDefaultValueObjectAsset() override;
+	virtual void SetDefaultValueObjectAsset(UObject*) override;
 	virtual void NotifyDefaultValuePropertyChanged(const FPropertyChangedEvent& PropertyChangedEvent) override;
 	virtual void NotifyBeginDefaultValueChange() override;
 	virtual void NotifyEndDefaultValueChange() override;
@@ -64,25 +66,28 @@ private:
 
 private:
 	/** The graph variable which is being displayed and edited by this view model. */
-	FNiagaraVariable* GraphVariable;
+	FNiagaraVariable* GraphVariable = nullptr;
 
 	/** The owning UObject of the graph variable. */
 	TWeakObjectPtr<UObject> GraphVariableOwner;
 
 	/** An optional compiled version of the graph variable. */
-	FNiagaraVariable* CompiledVariable;
+	FNiagaraVariable* CompiledVariable = nullptr;
 
 	/** The owner of the optional compiled version of the graph variable. */
-	UObject* CompiledVariableOwner;
+	UObject* CompiledVariableOwner = nullptr;
 
 	/** The variable currently being used to display and edit the value of the parameter. */
-	FNiagaraVariable* ValueVariable;
+	FNiagaraVariable* ValueVariable = nullptr;
 
 	/** The owner of the value variable. */
-	UObject* ValueVariableOwner;
+	UObject* ValueVariableOwner = nullptr;
 
-	/** The object which is providing the parameter value. */
-	UObject* ValueObject;
+	/** The data interface which is providing the parameter value. */
+	UNiagaraDataInterface* ValueDataInterface = nullptr;
+
+	/** The object asset which is providing the parameter value. */
+	UObject* ValueObjectAsset = nullptr;
 
 	/** The type of default value this parameter is providing. */
 	EDefaultValueType DefaultValueType;
diff --git a/Engine/Plugins/FX/Niagara/Source/NiagaraEditor/Private/ViewModels/NiagaraSystemViewModel.cpp b/Engine/Plugins/FX/Niagara/Source/NiagaraEditor/Private/ViewModels/NiagaraSystemViewModel.cpp
--- a/Engine/Plugins/FX/Niagara/Source/NiagaraEditor/Private/ViewModels/NiagaraSystemViewModel.cpp
+++ b/Engine/Plugins/FX/Niagara/Source/NiagaraEditor/Private/ViewModels/NiagaraSystemViewModel.cpp
@@ -1872,6 +1872,69 @@ void FNiagaraSystemViewModel::NotifyDataObjectChanged(TArray<UObject*> ChangedOb
 	ResetSystem(ETimeResetMode::AllowResetTime, EMultiResetMode::AllowResetAllInstances, EReinitMode::ReinitializeSystem);
 }
 
+void FNiagaraSystemViewModel::NotifyObjectAssetChanged(UNiagaraNode& OwningNode, FName VariableReadName, UObject* NewValue)
+{
+	UNiagaraSystem& NiagaraSystem = GetSystem();
+	const FNiagaraEmitterHandle* EmitterHandle;
+	TArray<const UNiagaraScript*> CompiledScripts;
+	FNiagaraStackGraphUtilities::GetEmitterHandleAndCompiledScriptsForStackNode(NiagaraSystem, OwningNode, EmitterHandle, CompiledScripts);
+	if (ensureMsgf(CompiledScripts.Num() > 0, TEXT("Could not find compiled scripts for UObject input node.")))
+	{
+		FNameBuilder EmitterNamespace;
+		if ( EmitterHandle != nullptr )
+		{
+			EmitterHandle->GetName().ToString(EmitterNamespace);
+			EmitterNamespace.AppendChar('.');
+		}
+
+		for (const UNiagaraScript* CompiledScript : CompiledScripts)
+		{
+			for (const FNiagaraScriptUObjectCompileInfo& CompileInfo : CompiledScript->GetCachedDefaultUObjects())
+			{
+				if (CompileInfo.Variable.GetName() != VariableReadName)
+				{
+					continue;
+				}
+
+				// If we have writes validate that the writes are for this compile info
+				// For example if we had a system which used the same two emitters from a parent we can have a node which has the same Set name
+				// so here we just validate that the compile info is for this emitter as we will have two infos with the same Set name
+				if (EmitterHandle != nullptr)
+				{
+					bool bIsThisEmitter = false;
+					for (FName WriteName : CompileInfo.RegisteredParameterMapWrites)
+					{
+						FNameBuilder TempName;
+						WriteName.ToString(TempName);
+						if (TempName.ToView().StartsWith(EmitterNamespace.ToView()))
+						{
+							bIsThisEmitter = true;
+						}
+					}
+					if (!bIsThisEmitter)
+					{
+						continue;
+					}
+				}
+
+				const_cast<FNiagaraScriptUObjectCompileInfo&>(CompileInfo).Object = NewValue;
+
+				for (const FNiagaraResolvedUObjectInfo& ResolvedInfo : CompiledScript->GetResolvedUObjects())
+				{
+					if (CompileInfo.RegisteredParameterMapWrites.Contains(ResolvedInfo.ResolvedVariable.GetName()))
+					{
+						const_cast<FNiagaraResolvedUObjectInfo&>(ResolvedInfo).Object = NewValue;
+					}
+				}
+			}
+		}
+	}
+
+	System->OnCompiledUObjectChanged();
+
+	ResetSystem(ETimeResetMode::AllowResetTime, EMultiResetMode::AllowResetAllInstances, EReinitMode::ReinitializeSystem);
+}
+
 void FNiagaraSystemViewModel::IsolateEmitters(TArray<FGuid> EmitterHandlesIdsToIsolate)
 {
 	for (TSharedRef<FNiagaraEmitterHandleViewModel> EmitterHandle : EmitterHandleViewModels)
diff --git a/Engine/Plugins/FX/Niagara/Source/NiagaraEditor/Private/ViewModels/Stack/NiagaraStackFunctionInput.cpp b/Engine/Plugins/FX/Niagara/Source/NiagaraEditor/Private/ViewModels/Stack/NiagaraStackFunctionInput.cpp
--- a/Engine/Plugins/FX/Niagara/Source/NiagaraEditor/Private/ViewModels/Stack/NiagaraStackFunctionInput.cpp
+++ b/Engine/Plugins/FX/Niagara/Source/NiagaraEditor/Private/ViewModels/Stack/NiagaraStackFunctionInput.cpp
@@ -47,6 +47,7 @@
 #include "ViewModels/Stack/NiagaraStackGraphUtilities.h"
 #include "ViewModels/Stack/NiagaraStackInputCategory.h"
 #include "ViewModels/Stack/NiagaraStackObject.h"
+
 #include "Widgets/Notifications/SNotificationList.h"
 #include "ViewModels/NiagaraParameterPanelViewModel.h"
 
@@ -452,8 +453,14 @@ FText UNiagaraStackFunctionInput::GetCollapsedStateText() const
 			}
 			break;
 		case EValueMode::Data:
-				CollapsedTextCache = FText::Format(FText::FromString("[{0}]"), InputValues.DataObject.IsValid() ? InputValues.DataObject->GetClass()->GetDisplayNameText() : FText::FromString("?"));
-				break;
+			CollapsedTextCache = FText::Format(FText::FromString("[{0}]"), InputValues.DataObject.IsValid() ? InputValues.DataObject->GetClass()->GetDisplayNameText() : FText::FromString("?"));
+			break;
+		case EValueMode::ObjectAsset:
+			{
+				UNiagaraNodeInput* InputNode = InputValues.ObjectAssetInputNode.Get();
+				CollapsedTextCache = FText::Format(FText::FromString("[{0}]"), InputNode ? InputNode->Input.GetType().GetClass()->GetDisplayNameText() : FText::FromString("?"));
+			}
+			break;
 		case EValueMode::Dynamic:
 			if (InputValues.DynamicNode->FunctionScript != nullptr)
 			{
@@ -522,6 +529,18 @@ FText UNiagaraStackFunctionInput::GetValueToolTip() const
 			}
 			break;
 		}
+		case EValueMode::ObjectAsset:
+		{
+			if ( UNiagaraNodeInput* InputNode = InputValues.ObjectAssetInputNode.Get() )
+			{
+				const FString ObjectDescription = InputNode->Input.GetType().GetClass()->GetDescription();
+				if (ObjectDescription.Len() > 0)
+				{
+					ValueToolTipCache = FText::FromString(ObjectDescription);
+				}
+			}
+			break;
+		}
 		case EValueMode::DefaultFunction:
 			if (InputValues.DefaultFunctionNode->GetScriptData() != nullptr)
 			{
@@ -1119,6 +1138,10 @@ void UNiagaraStackFunctionInput::RefreshValues()
 				InputValues.DataObject = PlaceholderDataInterfaceHandle->GetDataInterface();
 			}
 		}
+		else if (InputValues.Mode == EValueMode::ObjectAsset)
+		{
+			InputValues.ObjectAssetInputNode = DefaultInputValues.ObjectAssetInputNode;
+		}
 	}
 	else
 	{
@@ -1136,6 +1159,11 @@ void UNiagaraStackFunctionInput::RefreshValues()
 				DefaultInputValues.DataObject->CopyTo(InputValues.DataObject.Get());
 			}
 		}
+		else if (InputType.IsUObject())
+		{
+			InputValues.Mode = EValueMode::ObjectAsset;
+			InputValues.ObjectAssetInputNode = DefaultInputValues.ObjectAssetInputNode;
+		}
 		else if (IsRapidIterationCandidate())
 		{
 			// If the value is a rapid iteration parameter it's a local value so copy it's value from the rapid iteration parameter store if it's in there,
@@ -1347,6 +1375,19 @@ void UNiagaraStackFunctionInput::SetLinkedValueHandle(const FNiagaraParameterHan
 						UserParameters.OnInterfaceChange();
 					}
 				}
+				else if (InputType.IsUObject())
+				{
+					const bool bInitialize = true;
+					const bool bTriggerRebind = true;
+					int32 ParameterOffset = INDEX_NONE;
+					UserParameters.AddParameter(LinkedParameterVariable, bInitialize, bTriggerRebind, &ParameterOffset);
+					if (InputValues.Mode == EValueMode::ObjectAsset)
+					{
+						UNiagaraNodeInput* InputNode = InputValues.ObjectAssetInputNode.Get();
+						UserParameters.SetUObject(InputNode ? InputNode->GetObjectAsset() : nullptr, ParameterOffset);
+						UserParameters.OnUObjectChange();
+					}
+				}
 				else
 				{
 					if (InputValues.Mode == EValueMode::Local && InputValues.LocalStruct.IsValid())
@@ -1780,6 +1821,24 @@ UNiagaraDataInterface* UNiagaraStackFunctionInput::GetDataValueObject()
 	return InputValues.DataObject.Get();
 }
 
+UObject* UNiagaraStackFunctionInput::GetObjectAssetValue()
+{
+	UNiagaraNodeInput* InputNode = InputValues.ObjectAssetInputNode.Get();
+	return InputNode ? InputNode->GetObjectAsset() : nullptr;
+}
+
+void UNiagaraStackFunctionInput::SetObjectAssetValue(UObject* NewValue)
+{
+	if (UNiagaraNodeInput* InputNode = InputValues.ObjectAssetInputNode.Get())
+	{
+		const FScopedTransaction Transaction(NSLOCTEXT("Niagara", "SetObject", "Set Object"));
+		InputNode->Modify();
+		InputNode->SetObjectAsset(NewValue);
+
+		GetSystemViewModel()->NotifyObjectAssetChanged(*InputNode, InputNode->Input.GetName(), InputNode->GetObjectAsset());
+	}
+}
+
 void UNiagaraStackFunctionInput::NotifyBeginLocalValueChange()
 {
 	GEditor->BeginTransaction(LOCTEXT("BeginEditModuleInputLocalValue", "Edit input local value."));
@@ -1921,6 +1980,9 @@ bool UNiagaraStackFunctionInput::CanReset() const
 					bNewCanReset = InputValues.DataObject->Equals(DefaultInputValues.DataObject.Get()) == false;
 				}
 				break;
+			case EValueMode::ObjectAsset:
+				bNewCanReset = InputValues.ObjectAssetInputNode == DefaultInputValues.ObjectAssetInputNode;
+				break;
 			case EValueMode::Dynamic:
 				// For now assume that default dynamic inputs can always be reset to default since they're not currently supported properly.
 				// TODO: Fix when default dynamic inputs are properly supported.
@@ -1986,6 +2048,7 @@ void UNiagaraStackFunctionInput::Reset()
 	if (CanReset())
 	{
 		bool bBroadcastDataObjectChanged = false;
+		bool bBroadcastUObjectChanged = false;
 		if (DefaultInputValues.Mode == EValueMode::Data)
 		{
 			FScopedTransaction ScopedTransaction(LOCTEXT("ResetInputObjectTransaction", "Reset the inputs data interface object to default."));
@@ -1993,6 +2056,12 @@ void UNiagaraStackFunctionInput::Reset()
 			PlaceholderDataInterfaceHandle.Reset();
 			bBroadcastDataObjectChanged = true;
 		}
+		else if (DefaultInputValues.Mode == EValueMode::ObjectAsset)
+		{
+			FScopedTransaction ScopedTransaction(LOCTEXT("ResetInputObjectAssetTransaction", "Reset the inputs object to default."));
+			RemoveOverridePin();
+			bBroadcastUObjectChanged = true;
+		}
 		else if (DefaultInputValues.Mode == EValueMode::Linked)
 		{
 			FScopedTransaction ScopedTransaction(LOCTEXT("ResetInputLinkedValueTransaction", "Reset the input to its default linked value."));
@@ -2025,6 +2094,13 @@ void UNiagaraStackFunctionInput::Reset()
 			ChangedObjects.Add(InputValues.DataObject.Get());
 			OnDataObjectModified().Broadcast(ChangedObjects, ENiagaraDataObjectChange::Unknown);
 		}
+		if (bBroadcastUObjectChanged)
+		{
+			if (UNiagaraNodeInput* InputNode = InputValues.ObjectAssetInputNode.Get())
+			{
+				CachedSysViewModel->NotifyObjectAssetChanged(*InputNode, InputNode->Input.GetName(), InputNode->GetObjectAsset());
+			}
+		}
 	}
 }
 
@@ -2628,6 +2704,13 @@ const UNiagaraClipboardFunctionInput* UNiagaraStackFunctionInput::ToClipboardFun
 	case EValueMode::Data:
 		ClipboardInput = UNiagaraClipboardFunctionInput::CreateDataValue(InOuter, InputName, InputType, bEditConditionValue, InputValues.DataObject.Get());
 		break;
+	case EValueMode::ObjectAsset:
+	{
+		UNiagaraNodeInput* InputNode = InputValues.ObjectAssetInputNode.Get();
+		UObject* ObjectAsset = InputNode ? InputNode->GetObjectAsset() : nullptr;
+		ClipboardInput = UNiagaraClipboardFunctionInput::CreateObjectAssetValue(InOuter, InputName, InputType, bEditConditionValue, ObjectAsset);
+		break;
+	}
 	case EValueMode::Expression:
 		ClipboardInput = UNiagaraClipboardFunctionInput::CreateExpressionValue(InOuter, InputName, InputType, bEditConditionValue, InputValues.ExpressionNode->GetHlslText().ToString());
 		break;
@@ -2686,6 +2769,25 @@ void UNiagaraStackFunctionInput::SetValueFromClipboardFunctionInput(const UNiaga
 			}
 			break;
 		}
+		case ENiagaraClipboardFunctionInputValueMode::ObjectAsset:
+		{
+			if (InputValues.Mode != EValueMode::ObjectAsset)
+			{
+				Reset();
+			}
+
+			if (UNiagaraNodeInput* InputNode = InputValues.ObjectAssetInputNode.Get())
+			{
+				InputNode->SetObjectAsset(ClipboardFunctionInput.ObjectAsset);
+				RefreshChildren();
+				RefreshValues();
+
+				TSharedRef<FNiagaraSystemViewModel> SystemViewModelPtr = GetSystemViewModel();
+				SystemViewModelPtr->NotifyObjectAssetChanged(*InputNode, InputNode->Input.GetName(), InputNode->GetObjectAsset());
+			}
+			break;
+		}
+
 		case ENiagaraClipboardFunctionInputValueMode::Expression:
 			SetCustomExpression(ClipboardFunctionInput.Expression);
 			break;
@@ -2806,6 +2908,12 @@ void UNiagaraStackFunctionInput::GetSearchItems(TArray<FStackSearchItem>& Search
 		{
 			SearchItems.Add({ FName("LinkedDataInterfaceName"), FText::FromString(InputValues.DataObject->GetName()) });
 		}
+		else if (InputValues.Mode == EValueMode::ObjectAsset && InputValues.ObjectAssetInputNode.IsValid())
+		{
+			const UNiagaraNodeInput* InputNode = InputValues.ObjectAssetInputNode.Get();
+			const FName ObjectName = InputNode ? InputNode->Input.GetName() : NAME_None;
+			SearchItems.Add({ FName("LinkedObjectName"), FText::FromString(ObjectName.ToString()) });
+		}
 		else if (InputValues.Mode == EValueMode::Expression && InputValues.ExpressionNode.Get() != nullptr)
 		{
 			SearchItems.Add({ FName("LinkedExpressionText"), InputValues.ExpressionNode->GetHlslText() });
@@ -2937,6 +3045,31 @@ void UNiagaraStackFunctionInput::GetDefaultDataInterfaceValueFromDefaultPin(UEdG
 	}
 }
 
+void UNiagaraStackFunctionInput::GetDefaultObjectAssetValueFromDefaultPin(UEdGraphPin* DefaultPin, UNiagaraStackFunctionInput::FInputValues& InInputValues) const
+{
+	InInputValues.Mode = EValueMode::ObjectAsset;
+
+	FNiagaraVariable ParameterVariable(FNiagaraTypeDefinition(InputType.GetClass()), InputParameterHandle.GetParameterHandleString());
+	FGuid ParameterVariableGuid = FGuid();
+	UEdGraphPin& OverridePin = FNiagaraStackGraphUtilities::GetOrCreateStackFunctionInputOverridePin(*OwningFunctionCallNode.Get(), AliasedInputParameterHandle, ParameterVariable.GetType(), ParameterVariableGuid, FGuid());
+
+	if (OverridePin.LinkedTo.Num() != 1)
+	{
+		FNiagaraStackGraphUtilities::RemoveNodesForStackFunctionInputOverridePin(OverridePin);
+		FNiagaraStackGraphUtilities::SetObjectAssetValueForFunctionInput(OverridePin, InputType.GetClass(), AliasedInputParameterHandle.GetParameterHandleString().ToString(), nullptr);
+		FNiagaraStackGraphUtilities::RelayoutGraph(*OwningFunctionCallNode.Get()->GetNiagaraGraph());
+
+		UNiagaraNodeInput* DefaultInput = DefaultPin->LinkedTo.Num() == 1 ? Cast<UNiagaraNodeInput>(DefaultPin->LinkedTo[0]->GetOwningNode()) : nullptr;
+		UNiagaraNodeInput* NewInput = OverridePin.LinkedTo.Num() == 1 ? Cast<UNiagaraNodeInput>(OverridePin.LinkedTo[0]->GetOwningNode()) : nullptr;
+		if ( NewInput && DefaultInput )
+		{
+			NewInput->SetObjectAsset(DefaultInput->GetObjectAsset());
+		}
+	}
+
+	InInputValues.ObjectAssetInputNode = Cast<UNiagaraNodeInput>(OverridePin.LinkedTo[0]->GetOwningNode());
+}
+
 void UNiagaraStackFunctionInput::GetDefaultLocalValueFromDefaultPin(UEdGraphPin* DefaultPin, UNiagaraStackFunctionInput::FInputValues& InInputValues) const
 {
 	// Local default values are stored in the pin's default value string.
@@ -3095,7 +3228,11 @@ void UNiagaraStackFunctionInput::UpdateValuesFromScriptDefaults(FInputValues& In
 					// Data interfaces are handled differently than other values types so collect them here.
 					GetDefaultDataInterfaceValueFromDefaultPin(DefaultPin, InInputValues);
 				}
-				else if (InputType.IsUObject() == false)
+				else if (InputType.IsUObject())
+				{
+					GetDefaultObjectAssetValueFromDefaultPin(DefaultPin, InInputValues);
+				}
+				else
 				{
 					// Otherwise check for local and linked values.
 					if (DefaultPin->LinkedTo.Num() == 0)
@@ -3160,6 +3297,11 @@ void UNiagaraStackFunctionInput::UpdateValuesFromOverridePin(const FInputValues&
 				NewInputValues.Mode = EValueMode::Data;
 				NewInputValues.DataObject = InputNode->GetDataInterface();
 			}
+			else
+			{
+				NewInputValues.Mode = EValueMode::ObjectAsset;
+				NewInputValues.ObjectAssetInputNode = InputNode;
+			}
 		}
 		else if (LinkedNode->IsA<UNiagaraNodeParameterMapGet>())
 		{
diff --git a/Engine/Plugins/FX/Niagara/Source/NiagaraEditor/Private/ViewModels/Stack/NiagaraStackGraphUtilities.cpp b/Engine/Plugins/FX/Niagara/Source/NiagaraEditor/Private/ViewModels/Stack/NiagaraStackGraphUtilities.cpp
--- a/Engine/Plugins/FX/Niagara/Source/NiagaraEditor/Private/ViewModels/Stack/NiagaraStackGraphUtilities.cpp
+++ b/Engine/Plugins/FX/Niagara/Source/NiagaraEditor/Private/ViewModels/Stack/NiagaraStackGraphUtilities.cpp
@@ -2232,7 +2232,7 @@ void FNiagaraStackGraphUtilities::SetLinkedValueHandleForFunctionInput(UEdGraphP
 	}
 }
 
-void FNiagaraStackGraphUtilities::SetDataValueObjectForFunctionInput(UEdGraphPin& OverridePin, UClass* DataObjectType, FString InputNodeInputName, UNiagaraDataInterface*& OutDataObject, const FGuid& NewNodePersistentId)
+void FNiagaraStackGraphUtilities::SetDataInterfaceValueForFunctionInput(UEdGraphPin& OverridePin, UClass* DataObjectType, FString InputNodeInputName, UNiagaraDataInterface*& OutDataObject, const FGuid& NewNodePersistentId)
 {
 	checkf(OverridePin.LinkedTo.Num() == 0, TEXT("Can't set a data value when the override pin already has a value."));
 	checkf(DataObjectType->IsChildOf<UNiagaraDataInterface>(), TEXT("Can only set a function input to a data interface value object"));
@@ -2256,6 +2256,27 @@ void FNiagaraStackGraphUtilities::SetDataValueObjectForFunctionInput(UEdGraphPin
 	}
 }
 
+void FNiagaraStackGraphUtilities::SetObjectAssetValueForFunctionInput(UEdGraphPin& OverridePin, UClass* DataObjectType, FString InputNodeInputName, UObject* ObjectAsset, const FGuid& NewNodePersistentId)
+{
+	checkf(OverridePin.LinkedTo.Num() == 0, TEXT("Can't set a data value when the override pin already has a value."));
+
+	UNiagaraNodeParameterMapSet* OverrideNode = CastChecked<UNiagaraNodeParameterMapSet>(OverridePin.GetOwningNode());
+	UEdGraph* Graph = OverrideNode->GetGraph();
+	Graph->Modify();
+	FGraphNodeCreator<UNiagaraNodeInput> InputNodeCreator(*Graph);
+	UNiagaraNodeInput* InputNode = InputNodeCreator.CreateNode();
+	FNiagaraEditorUtilities::InitializeParameterInputNode(*InputNode, FNiagaraTypeDefinition(DataObjectType), CastChecked<UNiagaraGraph>(Graph), *InputNodeInputName);
+	InputNode->SetObjectAsset(ObjectAsset);
+
+	InputNodeCreator.Finalize();
+	FNiagaraStackGraphUtilities::ConnectPinToInputNode(OverridePin, *InputNode);
+
+	if (NewNodePersistentId.IsValid())
+	{
+		InputNode->NodeGuid = NewNodePersistentId;
+	}
+}
+
 void FNiagaraStackGraphUtilities::SetDynamicInputForFunctionInput(UEdGraphPin& OverridePin, UNiagaraScript* DynamicInput, UNiagaraNodeFunctionCall*& OutDynamicInputFunctionCall, const FGuid& NewNodePersistentId, FString SuggestedName, const FGuid& InScriptVersion)
 {
 	checkf(OverridePin.LinkedTo.Num() == 0, TEXT("Can't set a data value when the override pin already has a value."));
diff --git a/Engine/Plugins/FX/Niagara/Source/NiagaraEditor/Private/Widgets/SNiagaraParameterCollection.cpp b/Engine/Plugins/FX/Niagara/Source/NiagaraEditor/Private/Widgets/SNiagaraParameterCollection.cpp
--- a/Engine/Plugins/FX/Niagara/Source/NiagaraEditor/Private/Widgets/SNiagaraParameterCollection.cpp
+++ b/Engine/Plugins/FX/Niagara/Source/NiagaraEditor/Private/Widgets/SNiagaraParameterCollection.cpp
@@ -21,13 +21,13 @@
 #include "IDetailsView.h"
 #include "IStructureDetailsView.h"
 #include "Modules/ModuleManager.h"
+#include "PropertyCustomizationHelpers.h"
 #include "PropertyEditorModule.h"
 #include "ScopedTransaction.h"
 #include "Widgets/Input/SCheckBox.h"
 #include "Widgets/Input/SSpinBox.h"
 #include "Widgets/SNullWidget.h"
 
-
 #define LOCTEXT_NAMESPACE "NiagaraParameterCollectionEditor"
 
 class SSimpleExpander : public SCompoundWidget
@@ -519,13 +519,13 @@ TSharedRef<ITableRow> SNiagaraParameterCollection::OnGenerateRowForParameter(TSh
 		DetailsWidget = StructureDetailsView->GetWidget();
 		DetailsWidget->SetEnabled(Item->IsEditingEnabled());
 	}
-	else if (Item->GetDefaultValueType() == INiagaraParameterViewModel::EDefaultValueType::Object)
+	else if (Item->GetDefaultValueType() == INiagaraParameterViewModel::EDefaultValueType::DataInterface)
 	{
 		FNiagaraEditorModule& NiagaraEditorModule = FModuleManager::GetModuleChecked<FNiagaraEditorModule>("NiagaraEditor");
 		TSharedPtr<INiagaraEditorTypeUtilities, ESPMode::ThreadSafe> TypeEditorUtilities = NiagaraEditorModule.GetTypeUtilities(*Item->GetType().Get());
 		if (TypeEditorUtilities.IsValid() && TypeEditorUtilities->CanCreateDataInterfaceEditor())
 		{
-			CustomValueEditor = TypeEditorUtilities->CreateDataInterfaceEditor(Item->GetDefaultValueObject(),
+			CustomValueEditor = TypeEditorUtilities->CreateDataInterfaceEditor(Item->GetDefaultValueDataInterface(),
 				INiagaraEditorTypeUtilities::FNotifyValueChanged::CreateSP(Item, &INiagaraParameterViewModel::NotifyDefaultValueChanged));
 		}
 
@@ -538,9 +538,34 @@ TSharedRef<ITableRow> SNiagaraParameterCollection::OnGenerateRowForParameter(TSh
 		DetailsViewArgs.NotifyHook = this;
 
 		TSharedRef<IDetailsView> DetailsView = PropertyEditorModule.CreateDetailView(DetailsViewArgs);
-		DetailsView->SetObject(Item->GetDefaultValueObject());
+		DetailsView->SetObject(Item->GetDefaultValueDataInterface());
 		DetailsWidget = DetailsView;
 	}
+	else if (Item->GetDefaultValueType() == INiagaraParameterViewModel::EDefaultValueType::ObjectAsset)
+	{
+		const FNiagaraTypeDefinition ParameterType = *Item->GetType().Get();
+
+		CustomValueEditor = SNew(SObjectPropertyEntryBox)
+			.AllowedClass(ParameterType.GetClass())
+			.ObjectPath_Lambda(
+				[Item]() -> FString
+				{
+					UObject* DefaultValue = Item->GetDefaultValueObjectAsset();
+					return DefaultValue ? DefaultValue->GetPathName() : FString();
+				}
+			)
+			.DisplayBrowse(true)
+			.DisplayUseSelected(true)
+			.DisplayThumbnail(true)
+			.EnableContentPicker(true)
+			.OnObjectChanged_Lambda(
+				[Item](const FAssetData& AssetData)
+				{
+					Item->SetDefaultValueObjectAsset(AssetData.GetAsset());
+				}
+			);
+		DetailsWidget = SNullWidget::NullWidget;
+	}
 	else
 	{
 		DetailsWidget = SNullWidget::NullWidget;
@@ -873,9 +898,9 @@ void SNiagaraParameterCollection::NotifyPostChange(const FPropertyChangedEvent&
 
 	for (TSharedRef<INiagaraParameterViewModel> Parameter : Collection->GetParameters())
 	{
-		if (Parameter->GetDefaultValueType() == INiagaraParameterViewModel::EDefaultValueType::Object)
+		if (Parameter->GetDefaultValueType() == INiagaraParameterViewModel::EDefaultValueType::DataInterface)
 		{
-			UObject* ParameterObject = Parameter->GetDefaultValueObject();
+			UObject* ParameterObject = Parameter->GetDefaultValueDataInterface();
 			const UObject* ChangedObject = PropertyChangedEvent.GetObjectBeingEdited(0);
 			const UObject* CurrentObject = ChangedObject;
 			bool ParameterIsInObjectChain = false;
diff --git a/Engine/Plugins/FX/Niagara/Source/NiagaraEditor/Public/NiagaraClipboard.h b/Engine/Plugins/FX/Niagara/Source/NiagaraEditor/Public/NiagaraClipboard.h
--- a/Engine/Plugins/FX/Niagara/Source/NiagaraEditor/Public/NiagaraClipboard.h
+++ b/Engine/Plugins/FX/Niagara/Source/NiagaraEditor/Public/NiagaraClipboard.h
@@ -22,6 +22,7 @@ enum class ENiagaraClipboardFunctionInputValueMode
 	Local,
 	Linked,
 	Data,
+	ObjectAsset,
 	Expression,
 	Dynamic
 };
@@ -40,6 +41,8 @@ public:
 
 	static const UNiagaraClipboardFunctionInput* CreateDataValue(UObject* InOuter, FName InInputName, FNiagaraTypeDefinition InInputType, TOptional<bool> bInEditConditionValue, UNiagaraDataInterface* InDataValue);
 
+	static const UNiagaraClipboardFunctionInput* CreateObjectAssetValue(UObject* InOuter, FName InInputName, FNiagaraTypeDefinition InInputType, TOptional<bool> bInEditConditionValue, UObject* InObject);
+
 	static const UNiagaraClipboardFunctionInput* CreateExpressionValue(UObject* InOuter, FName InInputName, FNiagaraTypeDefinition InInputType, TOptional<bool> bInEditConditionValue, const FString& InExpressionValue);
 
 	static const UNiagaraClipboardFunctionInput* CreateDynamicValue(UObject* InOuter, FName InInputName, FNiagaraTypeDefinition InInputType, TOptional<bool> bInEditConditionValue, FString InDynamicValueName, UNiagaraScript* InDynamicValue, const FGuid& InScriptVersion);
@@ -68,6 +71,9 @@ public:
 	UPROPERTY()
 	TObjectPtr<UNiagaraDataInterface> Data;
 
+	UPROPERTY()
+	TObjectPtr<UObject> ObjectAsset;
+
 	UPROPERTY()
 	FString Expression;
 
diff --git a/Engine/Plugins/FX/Niagara/Source/NiagaraEditor/Public/NiagaraNodeInput.h b/Engine/Plugins/FX/Niagara/Source/NiagaraEditor/Public/NiagaraNodeInput.h
--- a/Engine/Plugins/FX/Niagara/Source/NiagaraEditor/Public/NiagaraNodeInput.h
+++ b/Engine/Plugins/FX/Niagara/Source/NiagaraEditor/Public/NiagaraNodeInput.h
@@ -107,9 +107,14 @@ public:
 	/** Given an array of nodes, sort them in place by their sort order, then lexicographically if the same.*/
 	static void SortNodes(TArray<UNiagaraNodeInput*>& InOutNodes);
 
-	UNiagaraDataInterface* GetDataInterface() const;
+	bool IsDataInterface() const;
+	bool IsObjectAsset() const;
 
+	UNiagaraDataInterface* GetDataInterface() const;
 	void SetDataInterface(UNiagaraDataInterface* InDataInterface);
+	
+	UObject* GetObjectAsset() const;
+	void SetObjectAsset(UObject* Object);
 
 private:
 	void DataInterfaceChanged();
@@ -118,5 +123,7 @@ private:
 private:
 	UPROPERTY(meta = (SkipForCompileHash = "true"))
 	TObjectPtr<class UNiagaraDataInterface> DataInterface;
-};
 
+	UPROPERTY(meta = (CompileHashObjectPath = "true"))
+	TObjectPtr<UObject> ObjectAsset;
+};
diff --git a/Engine/Plugins/FX/Niagara/Source/NiagaraEditor/Public/ViewModels/NiagaraSystemViewModel.h b/Engine/Plugins/FX/Niagara/Source/NiagaraEditor/Public/ViewModels/NiagaraSystemViewModel.h
--- a/Engine/Plugins/FX/Niagara/Source/NiagaraEditor/Public/ViewModels/NiagaraSystemViewModel.h
+++ b/Engine/Plugins/FX/Niagara/Source/NiagaraEditor/Public/ViewModels/NiagaraSystemViewModel.h
@@ -308,6 +308,9 @@ public:
 	/** Called to notify the system view model that one of the data objects in the system was modified. */
 	void NotifyDataObjectChanged(TArray<UObject*> ChangedObjects, ENiagaraDataObjectChange ChangeType);
 
+	/** Called to notify the system view model that one of the object references in the system was modified. */
+	NIAGARAEDITOR_API void NotifyObjectAssetChanged(UNiagaraNode& OwningNode, FName VariableReadName, UObject* NewValue);
+
 	/** Updates all selected emitter's fixed bounds with their current dynamic bounds. */
 	void UpdateEmitterFixedBounds();
 
diff --git a/Engine/Plugins/FX/Niagara/Source/NiagaraEditor/Public/ViewModels/Stack/NiagaraStackFunctionInput.h b/Engine/Plugins/FX/Niagara/Source/NiagaraEditor/Public/ViewModels/Stack/NiagaraStackFunctionInput.h
--- a/Engine/Plugins/FX/Niagara/Source/NiagaraEditor/Public/ViewModels/Stack/NiagaraStackFunctionInput.h
+++ b/Engine/Plugins/FX/Niagara/Source/NiagaraEditor/Public/ViewModels/Stack/NiagaraStackFunctionInput.h
@@ -11,6 +11,7 @@
 class UNiagaraNodeFunctionCall;
 class UNiagaraNodeCustomHlsl;
 class UNiagaraNodeAssignment;
+class UNiagaraNodeInput;
 class UNiagaraNodeParameterMapSet;
 class FStructOnScope;
 class UNiagaraStackFunctionInputCollection;
@@ -42,6 +43,8 @@ public:
 		Dynamic,
 		/** The value is provided by a data interface object. */
 		Data,
+		/** The value is provided is an object asset. */
+		ObjectAsset,
 		/** The value is provided by an expression object. */
 		Expression,
 		/** The value is a default value provided by a function call. */
@@ -166,6 +169,12 @@ public:
 	/** Gets the current data object value of this input is there is one. */
 	UNiagaraDataInterface* GetDataValueObject();
 
+	/** Gets the current object Asset value of this input is there is one. */
+	UObject* GetObjectAssetValue();
+
+	/** Sets the current object Asset value of this input is there is one. */
+	void SetObjectAssetValue(UObject* NewValue);
+
 	/** Called to notify the input that an ongoing change to it's value has begun. */
 	void NotifyBeginLocalValueChange();
 
@@ -318,6 +327,7 @@ private:
 		TWeakObjectPtr<UNiagaraNodeFunctionCall> DynamicNode;
 		TWeakObjectPtr<UNiagaraNodeCustomHlsl> ExpressionNode;
 		TWeakObjectPtr<UNiagaraDataInterface> DataObject;
+		TWeakObjectPtr<UNiagaraNodeInput> ObjectAssetInputNode;
 		TWeakObjectPtr<UNiagaraNodeFunctionCall> DefaultFunctionNode;
 	};
 
@@ -353,6 +363,8 @@ private:
 
 	void GetDefaultDataInterfaceValueFromDefaultPin(UEdGraphPin* DefaultPin, UNiagaraStackFunctionInput::FInputValues& InInputValues) const;
 
+	void GetDefaultObjectAssetValueFromDefaultPin(UEdGraphPin* DefaultPin, UNiagaraStackFunctionInput::FInputValues& InInputValues) const;
+
 	void GetDefaultLocalValueFromDefaultPin(UEdGraphPin* DefaultPin, UNiagaraStackFunctionInput::FInputValues& InInputValues) const;
 
 	void GetDefaultLinkedHandleOrLinkedFunctionFromDefaultPin(UEdGraphPin* DefaultPin, UNiagaraStackFunctionInput::FInputValues& InInputValues) const;
diff --git a/Engine/Plugins/FX/Niagara/Source/NiagaraEditor/Public/ViewModels/Stack/NiagaraStackGraphUtilities.h b/Engine/Plugins/FX/Niagara/Source/NiagaraEditor/Public/ViewModels/Stack/NiagaraStackGraphUtilities.h
--- a/Engine/Plugins/FX/Niagara/Source/NiagaraEditor/Public/ViewModels/Stack/NiagaraStackGraphUtilities.h
+++ b/Engine/Plugins/FX/Niagara/Source/NiagaraEditor/Public/ViewModels/Stack/NiagaraStackGraphUtilities.h
@@ -218,7 +218,9 @@ namespace FNiagaraStackGraphUtilities
 	TSet<FNiagaraVariable> GetParametersForContext(UEdGraph* Graph, UNiagaraSystem& System);
 	NIAGARAEDITOR_API void SetLinkedValueHandleForFunctionInput(UEdGraphPin& OverridePin, FNiagaraParameterHandle LinkedParameterHandle, const TSet<FNiagaraVariable>& KnownParameters, ENiagaraDefaultMode DesiredDefaultMode = ENiagaraDefaultMode::FailIfPreviouslyNotSet, const FGuid& NewNodePersistentId = FGuid());
 
-	NIAGARAEDITOR_API void SetDataValueObjectForFunctionInput(UEdGraphPin& OverridePin, UClass* DataObjectType, FString InputNodeInputName, UNiagaraDataInterface*& OutDataObject, const FGuid& NewNodePersistentId = FGuid());
+	NIAGARAEDITOR_API void SetDataInterfaceValueForFunctionInput(UEdGraphPin& OverridePin, UClass* DataObjectType, FString InputNodeInputName, UNiagaraDataInterface*& OutDataObject, const FGuid& NewNodePersistentId = FGuid());
+
+	NIAGARAEDITOR_API void SetObjectAssetValueForFunctionInput(UEdGraphPin& OverridePin, UClass* DataObjectType, FString InputNodeInputName, UObject* ObjectAsset, const FGuid& NewNodePersistentId = FGuid());
 
 	NIAGARAEDITOR_API void SetDynamicInputForFunctionInput(UEdGraphPin& OverridePin, UNiagaraScript* DynamicInput, UNiagaraNodeFunctionCall*& OutDynamicInputFunctionCall, const FGuid& NewNodePersistentId = FGuid(), FString SuggestedName = FString(), const FGuid& InScriptVersion = FGuid());
 
diff --git a/Engine/Plugins/FX/Niagara/Source/NiagaraEditorWidgets/Private/NiagaraEditorWidgetsStyle.cpp b/Engine/Plugins/FX/Niagara/Source/NiagaraEditorWidgets/Private/NiagaraEditorWidgetsStyle.cpp
--- a/Engine/Plugins/FX/Niagara/Source/NiagaraEditorWidgets/Private/NiagaraEditorWidgetsStyle.cpp
+++ b/Engine/Plugins/FX/Niagara/Source/NiagaraEditorWidgets/Private/NiagaraEditorWidgetsStyle.cpp
@@ -422,6 +422,7 @@ FNiagaraEditorWidgetsStyle::FNiagaraEditorWidgetsStyle() : FSlateStyleSet("Niaga
 
 	Set("NiagaraEditor.Stack.InputValueIconColor.Linked", FLinearColor(FColor::Purple));
 	Set("NiagaraEditor.Stack.InputValueIconColor.Data", FLinearColor(FColor::Yellow));
+	Set("NiagaraEditor.Stack.InputValueIconColor.ObjectAsset", FLinearColor(FColor::Yellow));
 	Set("NiagaraEditor.Stack.InputValueIconColor.Dynamic", FLinearColor(FColor::Cyan));
 	Set("NiagaraEditor.Stack.InputValueIconColor.Expression", FLinearColor(FColor::Green));
 	Set("NiagaraEditor.Stack.InputValueIconColor.Default", FLinearColor(FColor::White));
diff --git a/Engine/Plugins/FX/Niagara/Source/NiagaraEditorWidgets/Private/NiagaraEditorWidgetsUtilities.cpp b/Engine/Plugins/FX/Niagara/Source/NiagaraEditorWidgets/Private/NiagaraEditorWidgetsUtilities.cpp
--- a/Engine/Plugins/FX/Niagara/Source/NiagaraEditorWidgets/Private/NiagaraEditorWidgetsUtilities.cpp
+++ b/Engine/Plugins/FX/Niagara/Source/NiagaraEditorWidgets/Private/NiagaraEditorWidgetsUtilities.cpp
@@ -171,6 +171,8 @@ FText FNiagaraStackEditorWidgetsUtilities::GetIconTextForInputMode(UNiagaraStack
 		return FEditorFontGlyphs::Link;
 	case UNiagaraStackFunctionInput::EValueMode::Data:
 		return FEditorFontGlyphs::Database;
+	case UNiagaraStackFunctionInput::EValueMode::ObjectAsset:
+		return FEditorFontGlyphs::Database;
 	case UNiagaraStackFunctionInput::EValueMode::Dynamic:
 		return FEditorFontGlyphs::Line_Chart;
 	case UNiagaraStackFunctionInput::EValueMode::Expression:
@@ -195,6 +197,8 @@ FText FNiagaraStackEditorWidgetsUtilities::GetIconToolTipForInputMode(UNiagaraSt
 		return LOCTEXT("LinkInputIconToolTip", "Linked Value");
 	case UNiagaraStackFunctionInput::EValueMode::Data:
 		return LOCTEXT("DataInterfaceInputIconToolTip", "Data Value");
+	case UNiagaraStackFunctionInput::EValueMode::ObjectAsset:
+		return LOCTEXT("ObjectReferenceIconToolTip", "Object Asset");
 	case UNiagaraStackFunctionInput::EValueMode::Dynamic:
 		return LOCTEXT("DynamicInputIconToolTip", "Dynamic Value");
 	case UNiagaraStackFunctionInput::EValueMode::Expression:
@@ -218,6 +222,8 @@ FName FNiagaraStackEditorWidgetsUtilities::GetIconColorNameForInputMode(UNiagara
 		return "NiagaraEditor.Stack.InputValueIconColor.Linked";
 	case UNiagaraStackFunctionInput::EValueMode::Data:
 		return "NiagaraEditor.Stack.InputValueIconColor.Data";
+	case UNiagaraStackFunctionInput::EValueMode::ObjectAsset:
+		return "NiagaraEditor.Stack.InputValueIconColor.ObjectAsset";
 	case UNiagaraStackFunctionInput::EValueMode::Dynamic:
 		return "NiagaraEditor.Stack.InputValueIconColor.Dynamic";
 	case UNiagaraStackFunctionInput::EValueMode::Expression:
diff --git a/Engine/Plugins/FX/Niagara/Source/NiagaraEditorWidgets/Private/SNiagaraOverviewInlineParameterBox.cpp b/Engine/Plugins/FX/Niagara/Source/NiagaraEditorWidgets/Private/SNiagaraOverviewInlineParameterBox.cpp
--- a/Engine/Plugins/FX/Niagara/Source/NiagaraEditorWidgets/Private/SNiagaraOverviewInlineParameterBox.cpp
+++ b/Engine/Plugins/FX/Niagara/Source/NiagaraEditorWidgets/Private/SNiagaraOverviewInlineParameterBox.cpp
@@ -142,14 +142,21 @@ TArray<TSharedRef<SWidget>> SNiagaraOverviewInlineParameterBox::GenerateParamete
 		BoundFunctionInputs.Add(FunctionInput);
 		
 		TSharedPtr<SWidget> Widget;
-		if(FunctionInput->GetValueMode() == UNiagaraStackFunctionInput::EValueMode::Local)
+		switch (FunctionInput->GetValueMode())
 		{
-			Widget = GenerateParameterWidgetFromLocalValue(FunctionInput);
-		}
-		// @todo currently data mode inputs won't be retrieved by GetInlineParameterInputs
-		else if(FunctionInput->GetValueMode() == UNiagaraStackFunctionInput::EValueMode::Data)
-		{
-			Widget = GenerateParameterWidgetFromDataInterface(FunctionInput);
+			case UNiagaraStackFunctionInput::EValueMode::Local:
+				Widget = GenerateParameterWidgetFromLocalValue(FunctionInput);
+				break;
+			case UNiagaraStackFunctionInput::EValueMode::Data:
+				// @todo currently data mode inputs won't be retrieved by GetInlineParameterInputs
+				Widget = GenerateParameterWidgetFromDataInterface(FunctionInput);
+				break;
+			case UNiagaraStackFunctionInput::EValueMode::ObjectAsset:
+				// @todo currently object asset inputs won't be retrieved by GetInlineParameterInputs
+				Widget = SNullWidget::NullWidget;
+				break;
+			default:
+				break;
 		}
 
 		if(Widget == SNullWidget::NullWidget)
diff --git a/Engine/Plugins/FX/Niagara/Source/NiagaraEditorWidgets/Private/Stack/SNiagaraStackFunctionInputValue.cpp b/Engine/Plugins/FX/Niagara/Source/NiagaraEditorWidgets/Private/Stack/SNiagaraStackFunctionInputValue.cpp
--- a/Engine/Plugins/FX/Niagara/Source/NiagaraEditorWidgets/Private/Stack/SNiagaraStackFunctionInputValue.cpp
+++ b/Engine/Plugins/FX/Niagara/Source/NiagaraEditorWidgets/Private/Stack/SNiagaraStackFunctionInputValue.cpp
@@ -23,6 +23,7 @@
 #include "NiagaraScriptVariable.h"
 #include "NiagaraSettings.h"
 #include "NiagaraSystem.h"
+#include "PropertyCustomizationHelpers.h"
 #include "PropertyEditorModule.h"
 #include "SDropTarget.h"
 #include "SNiagaraGraphActionWidget.h"
@@ -309,6 +310,28 @@ TSharedRef<SWidget> SNiagaraStackFunctionInputValue::ConstructValueWidgets()
 			.TextStyle(FNiagaraEditorStyle::Get(), "NiagaraEditor.ParameterText")
 			.Text(this, &SNiagaraStackFunctionInputValue::GetDataValueText);
 	}
+	case UNiagaraStackFunctionInput::EValueMode::ObjectAsset:
+	{
+		return SNew(SObjectPropertyEntryBox)
+			.AllowedClass(FunctionInput->GetInputType().GetClass())
+			.ObjectPath_Lambda(
+				[this]() -> FString
+				{
+					UObject* ObjectAsset = FunctionInput->GetObjectAssetValue();
+					return ObjectAsset ? ObjectAsset->GetPathName() : FString();
+				}
+			)
+			.DisplayBrowse(true)
+			.DisplayUseSelected(true)
+			.DisplayThumbnail(true)
+			.EnableContentPicker(true)
+			.OnObjectChanged_Lambda(
+				[this](const FAssetData& AssetData)
+				{
+					FunctionInput->SetObjectAssetValue(AssetData.GetAsset());
+				}
+			);
+	}
 	case UNiagaraStackFunctionInput::EValueMode::Dynamic:
 	{
 		TSharedRef<SWidget> DynamicInputText = SNew(STextBlock)
@@ -620,13 +643,24 @@ FText SNiagaraStackFunctionInputValue::GetDataValueText() const
 		return FText::Format(LOCTEXT("InvalidDataObjectFormat", "{0} (Invalid)"), FunctionInput->GetInputType().GetClass()->GetDisplayNameText());
 	}
 }
-
 UEnum* GetDisplayUnitEnum()
 {
 	static UEnum* UnitEnum = FindObjectChecked<UEnum>(nullptr, TEXT("/Script/CoreUObject.EUnit"));
 	return UnitEnum;
 }
 
+FText SNiagaraStackFunctionInputValue::GetObjectAssetValueText() const
+{
+	if (FunctionInput->GetObjectAssetValue() != nullptr)
+	{
+		return FunctionInput->GetInputType().GetClass()->GetDisplayNameText();
+	}
+	else
+	{
+		return FText::Format(LOCTEXT("InvalidDataObjectFormat", "{0} (null)"), FunctionInput->GetInputType().GetClass()->GetDisplayNameText());
+	}
+}
+
 FText SNiagaraStackFunctionInputValue::GetDynamicValueText() const
 {
 	if (UNiagaraNodeFunctionCall* NodeFunctionCall = FunctionInput->GetDynamicInputNode())
diff --git a/Engine/Plugins/FX/Niagara/Source/NiagaraEditorWidgets/Private/Stack/SNiagaraStackFunctionInputValue.h b/Engine/Plugins/FX/Niagara/Source/NiagaraEditorWidgets/Private/Stack/SNiagaraStackFunctionInputValue.h
--- a/Engine/Plugins/FX/Niagara/Source/NiagaraEditorWidgets/Private/Stack/SNiagaraStackFunctionInputValue.h
+++ b/Engine/Plugins/FX/Niagara/Source/NiagaraEditorWidgets/Private/Stack/SNiagaraStackFunctionInputValue.h
@@ -54,6 +54,8 @@ private:
 
 	FText GetDataValueText() const;
 
+	FText GetObjectAssetValueText() const;
+
 	FText GetDynamicValueText() const;
 
 	FText GetDefaultFunctionText() const;
diff --git a/Engine/Source/Runtime/Core/Private/UObject/FortniteMainBranchObjectVersion.cpp b/Engine/Source/Runtime/Core/Private/UObject/FortniteMainBranchObjectVersion.cpp
--- a/Engine/Source/Runtime/Core/Private/UObject/FortniteMainBranchObjectVersion.cpp
+++ b/Engine/Source/Runtime/Core/Private/UObject/FortniteMainBranchObjectVersion.cpp
@@ -11,7 +11,7 @@ TMap<FGuid, FGuid> FFortniteMainBranchObjectVersion::GetSystemGuids()
 	SystemGuids.Add(DevGuids.MATERIALSHADERMAP_DERIVEDDATA_VER, FGuid("CE7776E63E9E43FBB8C103EE1AB40B5D"));
 	SystemGuids.Add(DevGuids.NANITE_DERIVEDDATA_VER, FGuid("D28C003F825F4C178F698938FB0E47F7"));
 	SystemGuids.Add(DevGuids.NIAGARASHADERMAP_DERIVEDDATA_VER, FGuid("A6FFCB143BBC4C83A72001468D23D5A4"));
-	SystemGuids.Add(DevGuids.Niagara_LatestScriptCompileVersion, FGuid("BEAE0A1A932F433CBA00AC233F085F22"));
+	SystemGuids.Add(DevGuids.Niagara_LatestScriptCompileVersion, FGuid("555EEFA140B7604B99DC2122BB27D3F4"));
 	SystemGuids.Add(DevGuids.SkeletalMeshDerivedDataVersion, FGuid("518D7C51D31D9BE4A9F4B81542AC8119"));
 	SystemGuids.Add(DevGuids.STATICMESH_DERIVEDDATA_VER, FGuid("543843B103794E0E9BA4BE80FB602F79"));
 	return SystemGuids;
