- Fix drag & drop of object parameter onto stack not working

#rnx

diff --git a/Engine/Plugins/FX/Niagara/Source/NiagaraEditor/Private/ViewModels/Stack/NiagaraStackScriptItemGroup.cpp b/Engine/Plugins/FX/Niagara/Source/NiagaraEditor/Private/ViewModels/Stack/NiagaraStackScriptItemGroup.cpp
--- a/Engine/Plugins/FX/Niagara/Source/NiagaraEditor/Private/ViewModels/Stack/NiagaraStackScriptItemGroup.cpp
+++ b/Engine/Plugins/FX/Niagara/Source/NiagaraEditor/Private/ViewModels/Stack/NiagaraStackScriptItemGroup.cpp
@@ -348,13 +348,6 @@ private:
 		{
 			Vars.Add(ParameterVariable);
 			DefaultVals.Add(FNiagaraConstants::GetAttributeDefaultValue(ParameterVariable));
-
-			FNiagaraTypeDefinition Type = ParameterVariable.GetType();
-			if (Type.IsUObject() && Type.IsDataInterface() == false)
-			{
-				// we currently don't support setting uobjects in the stack (except for data interfaces)
-				return nullptr;
-			}
 		}
 
 		UNiagaraNodeAssignment* NewAssignmentModule = FNiagaraStackGraphUtilities::AddParameterModuleToStack(Vars, *OutputNode, TargetIndex,DefaultVals );
