- Reduce size of FNiagaraRendererVariableInfo from 14 bytes -> 8 bytes

diff --git a/Engine/Plugins/FX/Niagara/Source/Niagara/Private/NiagaraRenderer.cpp b/Engine/Plugins/FX/Niagara/Source/Niagara/Private/NiagaraRenderer.cpp
--- a/Engine/Plugins/FX/Niagara/Source/Niagara/Private/NiagaraRenderer.cpp
+++ b/Engine/Plugins/FX/Niagara/Source/Niagara/Private/NiagaraRenderer.cpp
@@ -257,24 +257,23 @@ FParticleRenderData FNiagaraRenderer::TransferDataToGPU(FGlobalDynamicReadBuffer
 
 	for (const FNiagaraRendererVariableInfo& VarInfo : RendererLayout->GetVFVariables_RenderThread())
 	{
-		int32 GpuOffset = VarInfo.GetGPUOffset();
-		if (GpuOffset != INDEX_NONE && VarInfo.bUpload)
+		const int32 GpuOffset = VarInfo.GetRawGPUOffset();
+		if (GpuOffset != INDEX_NONE && VarInfo.ShouldUpload())
 		{
-			if (VarInfo.bHalfType)
+			if (VarInfo.IsHalfType())
 			{
-				GpuOffset &= ~(1 << 31);
-				for (int32 CompIdx = 0; CompIdx < VarInfo.NumComponents; ++CompIdx)
+				for (int32 CompIdx = 0; CompIdx < VarInfo.GetNumComponents(); ++CompIdx)
 				{
-					const uint8* SrcComponent = SrcData->GetComponentPtrHalf(VarInfo.DatasetOffset + CompIdx);
+					const uint8* SrcComponent = SrcData->GetComponentPtrHalf(VarInfo.GetRawDatasetOffset() + CompIdx);
 					void* Dest = Allocation.HalfData.Buffer + Allocation.HalfStride * (GpuOffset + CompIdx);
 					FMemory::Memcpy(Dest, SrcComponent, Allocation.HalfStride);
 				}
 			}
 			else
 			{
-				for (int32 CompIdx = 0; CompIdx < VarInfo.NumComponents; ++CompIdx)
+				for (int32 CompIdx = 0; CompIdx < VarInfo.GetNumComponents(); ++CompIdx)
 				{
-					const uint8* SrcComponent = SrcData->GetComponentPtrFloat(VarInfo.DatasetOffset + CompIdx);
+					const uint8* SrcComponent = SrcData->GetComponentPtrFloat(VarInfo.GetRawDatasetOffset() + CompIdx);
 					void* Dest = Allocation.FloatData.Buffer + Allocation.FloatStride * (GpuOffset + CompIdx);
 					FMemory::Memcpy(Dest, SrcComponent, Allocation.FloatStride);
 				}
@@ -651,7 +650,7 @@ void FNiagaraRenderer::SortIndices(const FNiagaraGPUSortInfo& SortInfo, const FN
 	check(SortInfo.SortAttributeOffset != INDEX_NONE);
 
 	const bool bUseRadixSort = GNiagaraRadixSortThreshold != -1 && (int32)NumInstances  > GNiagaraRadixSortThreshold;
-	const bool bSortVarIsHalf = SortVariable.bHalfType;
+	const bool bSortVarIsHalf = SortVariable.IsHalfType();
 
 	int32* RESTRICT IndexBuffer = (int32*)(OutIndices.Buffer);
 
@@ -662,7 +661,7 @@ void FNiagaraRenderer::SortIndices(const FNiagaraGPUSortInfo& SortInfo, const FN
 	{
 		if (bSortVarIsHalf)
 		{
-			const int32 BaseCompOffset = SortVariable.DatasetOffset;
+			const int32 BaseCompOffset = SortVariable.GetRawDatasetOffset();
 			FFloat16* RESTRICT PositionX = (FFloat16*)Buffer.GetComponentPtrHalf(BaseCompOffset);
 			FFloat16* RESTRICT PositionY = (FFloat16*)Buffer.GetComponentPtrHalf(BaseCompOffset + 1);
 			FFloat16* RESTRICT PositionZ = (FFloat16*)Buffer.GetComponentPtrHalf(BaseCompOffset + 2);
@@ -688,7 +687,7 @@ void FNiagaraRenderer::SortIndices(const FNiagaraGPUSortInfo& SortInfo, const FN
 		}
 		else
 		{
-			const int32 BaseCompOffset = SortVariable.DatasetOffset;
+			const int32 BaseCompOffset = SortVariable.GetRawDatasetOffset();
 			float* RESTRICT PositionX = (float*)Buffer.GetComponentPtrFloat(BaseCompOffset);
 			float* RESTRICT PositionY = (float*)Buffer.GetComponentPtrFloat(BaseCompOffset + 1);
 			float* RESTRICT PositionZ = (float*)Buffer.GetComponentPtrFloat(BaseCompOffset + 2);
@@ -717,7 +716,7 @@ void FNiagaraRenderer::SortIndices(const FNiagaraGPUSortInfo& SortInfo, const FN
 	{
 		if (bSortVarIsHalf)
 		{
-			FFloat16* RESTRICT CustomSorting = (FFloat16*)Buffer.GetComponentPtrHalf(SortVariable.DatasetOffset);
+			FFloat16* RESTRICT CustomSorting = (FFloat16*)Buffer.GetComponentPtrHalf(SortVariable.GetRawDatasetOffset());
 			if (SortInfo.SortMode == ENiagaraSortMode::CustomAscending)
 			{
 				for (uint32 i = 0; i < NumInstances; ++i)
@@ -735,7 +734,7 @@ void FNiagaraRenderer::SortIndices(const FNiagaraGPUSortInfo& SortInfo, const FN
 		}
 		else
 		{
-			float* RESTRICT CustomSorting = (float*)Buffer.GetComponentPtrFloat(SortVariable.DatasetOffset);
+			float* RESTRICT CustomSorting = (float*)Buffer.GetComponentPtrFloat(SortVariable.GetRawDatasetOffset());
 			if (SortInfo.SortMode == ENiagaraSortMode::CustomAscending)
 			{
 				for (uint32 i = 0; i < NumInstances; ++i)
diff --git a/Engine/Plugins/FX/Niagara/Source/Niagara/Private/NiagaraRendererMeshes.cpp b/Engine/Plugins/FX/Niagara/Source/Niagara/Private/NiagaraRendererMeshes.cpp
--- a/Engine/Plugins/FX/Niagara/Source/Niagara/Private/NiagaraRendererMeshes.cpp
+++ b/Engine/Plugins/FX/Niagara/Source/Niagara/Private/NiagaraRendererMeshes.cpp
@@ -1605,9 +1605,9 @@ void FNiagaraRendererMeshes::GetDynamicRayTracingInstances(FRayTracingMaterialGa
 				const FNiagaraRendererVariableInfo& VarScaleInfo = VFVariables[ENiagaraMeshVFLayout::Scale];
 				const FNiagaraRendererVariableInfo& VarTransformInfo = VFVariables[ENiagaraMeshVFLayout::Rotation];
 
-				const int32 PositionBaseCompOffset = VarPositionInfo.DatasetOffset;
-				const int32 ScaleBaseCompOffset = VarScaleInfo.DatasetOffset;
-				const int32 TransformBaseCompOffset = VarTransformInfo.DatasetOffset;
+				const int32 PositionBaseCompOffset = VarPositionInfo.GetRawDatasetOffset();
+				const int32 ScaleBaseCompOffset = VarScaleInfo.GetRawDatasetOffset();
+				const int32 TransformBaseCompOffset = VarTransformInfo.GetRawDatasetOffset();
 
 				const float* RESTRICT PositionX = reinterpret_cast<const float*>(ParticleMeshRenderData.SourceParticleData->GetComponentPtrFloat(PositionBaseCompOffset));
 				const float* RESTRICT PositionY = reinterpret_cast<const float*>(ParticleMeshRenderData.SourceParticleData->GetComponentPtrFloat(PositionBaseCompOffset + 1));
diff --git a/Engine/Plugins/FX/Niagara/Source/Niagara/Private/NiagaraRendererProperties.cpp b/Engine/Plugins/FX/Niagara/Source/Niagara/Private/NiagaraRendererProperties.cpp
--- a/Engine/Plugins/FX/Niagara/Source/Niagara/Private/NiagaraRendererProperties.cpp
+++ b/Engine/Plugins/FX/Niagara/Source/Niagara/Private/NiagaraRendererProperties.cpp
@@ -127,7 +127,7 @@ bool FNiagaraRendererLayout::SetVariable(const FNiagaraDataSetCompiledData* Comp
 	const int32 VarSize = bHalfVariable ? sizeof(FFloat16) : sizeof(float);
 	const int32 NumComponents = DataSetVariable.GetSizeInBytes() / VarSize;
 	const int32 Offset = bHalfVariable ? DataSetVariableLayout.HalfComponentStart : DataSetVariableLayout.FloatComponentStart;
-	int32& TotalVFComponents = bHalfVariable ? TotalHalfComponents_GT : TotalFloatComponents_GT;
+	uint16& TotalVFComponents = bHalfVariable ? TotalHalfComponents_GT : TotalFloatComponents_GT;
 
 	int32 GPULocation = INDEX_NONE;
 	bool bUpload = true;
@@ -144,6 +144,7 @@ bool FNiagaraRendererLayout::SetVariable(const FNiagaraDataSetCompiledData* Comp
 		{
 			//For CPU Sims we pack just the required data tightly in a GPU buffer we upload. For GPU sims the data is there already so we just provide the real data location.
 			GPULocation = CompiledData->SimTarget == ENiagaraSimTarget::CPUSim ? TotalVFComponents : Offset;
+			check(int32(TotalVFComponents ) + NumComponents <= TNumericLimits<uint16>::Max());
 			TotalVFComponents += NumComponents;
 		}
 	}
diff --git a/Engine/Plugins/FX/Niagara/Source/Niagara/Public/NiagaraRendererProperties.h b/Engine/Plugins/FX/Niagara/Source/Niagara/Public/NiagaraRendererProperties.h
--- a/Engine/Plugins/FX/Niagara/Source/Niagara/Public/NiagaraRendererProperties.h
+++ b/Engine/Plugins/FX/Niagara/Source/Niagara/Public/NiagaraRendererProperties.h
@@ -101,6 +101,10 @@ private:
 /** Mapping between a variable in the source dataset and the location we place it in the GPU buffer passed to the VF. */
 struct FNiagaraRendererVariableInfo
 {
+	friend struct FNiagaraRendererLayout;
+
+	static constexpr uint16 kInvalidOffset = 0xffff;
+
 	FNiagaraRendererVariableInfo() {}
 	FNiagaraRendererVariableInfo(int32 InDataOffset, int32 InGPUBufferOffset, int32 InNumComponents, bool bInUpload, bool bInHalfType)
 		: DatasetOffset(InDataOffset)
@@ -109,26 +113,26 @@ struct FNiagaraRendererVariableInfo
 		, bUpload(bInUpload)
 		, bHalfType(bInHalfType)
 	{
+		check(int32(InDataOffset) < TNumericLimits<uint16>::Max());
+		check(int32(InGPUBufferOffset) < TNumericLimits<uint16>::Max());
+		check(int32(InNumComponents) <= TNumericLimits<uint16>::Max());
 	}
 
-	FORCEINLINE int32 GetGPUOffset() const
-	{
-		int32 Offset = GPUBufferOffset;
-		if (bHalfType)
-		{
-			Offset |= 1 << 31;
-		}
-		return Offset;
-	}
+	FORCEINLINE int32 GetNumComponents() const { return NumComponents; }
 
-	FORCEINLINE int32 GetEncodedDatasetOffset() const
-	{
-		return DatasetOffset | (bHalfType ? (1<<31) : 0);
-	}
+	FORCEINLINE int32 GetRawGPUOffset() const { return GPUBufferOffset == kInvalidOffset ? INDEX_NONE : GPUBufferOffset; }
+	FORCEINLINE int32 GetGPUOffset() const { return GetRawGPUOffset() | (bHalfType ? (1 << 31) : 0); }
+
+	FORCEINLINE int32 GetRawDatasetOffset() const { return DatasetOffset == kInvalidOffset ? INDEX_NONE : DatasetOffset; }
+	FORCEINLINE int32 GetEncodedDatasetOffset() const { return GetRawDatasetOffset() | (bHalfType ? (1 << 31) : 0); }
 
-	int32 DatasetOffset = INDEX_NONE;
-	int32 GPUBufferOffset = INDEX_NONE;
-	int32 NumComponents = 0;
+	FORCEINLINE bool ShouldUpload() const { return bUpload; }
+	FORCEINLINE bool IsHalfType() const { return bHalfType; }
+
+protected:
+	uint16 DatasetOffset = kInvalidOffset;
+	uint16 GPUBufferOffset = kInvalidOffset;
+	uint16 NumComponents = 0;
 	bool bUpload = false;
 	bool bHalfType = false;
 };
@@ -150,12 +154,13 @@ struct NIAGARA_API FNiagaraRendererLayout
 
 private:
 	TArray<FNiagaraRendererVariableInfo> VFVariables_GT;
-	int32 TotalFloatComponents_GT;
-	int32 TotalHalfComponents_GT;
-
 	TArray<FNiagaraRendererVariableInfo> VFVariables_RT;
-	int32 TotalFloatComponents_RT;
-	int32 TotalHalfComponents_RT;
+
+	uint16 TotalFloatComponents_GT = 0;
+	uint16 TotalHalfComponents_GT = 0;
+
+	uint16 TotalFloatComponents_RT = 0;
+	uint16 TotalHalfComponents_RT = 0;
 };
 
 UENUM()
