- Fix for incorrect assumption that when looking over bindings they would match the static def

diff --git a/Engine/Plugins/FX/Niagara/Source/Niagara/Private/NiagaraEmitter.cpp b/Engine/Plugins/FX/Niagara/Source/Niagara/Private/NiagaraEmitter.cpp
--- a/Engine/Plugins/FX/Niagara/Source/Niagara/Private/NiagaraEmitter.cpp
+++ b/Engine/Plugins/FX/Niagara/Source/Niagara/Private/NiagaraEmitter.cpp
@@ -1820,22 +1820,9 @@ void FVersionedNiagaraEmitterData::RebuildRendererBindings(const UNiagaraEmitter
 				}
 				bool bVariableFound = false;
 
-				// Find static variables
-				if (Var.GetType().IsStatic())
-				{
-					for (const FNiagaraVariable& StaticVar : Script->GetVMExecutableData().StaticVariablesWritten)
-					{
-						if (StaticVar.GetType().IsSameBaseDefinition(Var.GetType()) && StaticVar.GetName() == Var.GetName())
-						{
-							RendererBindings.AddParameter(StaticVar, true, false);
-							bVariableFound = true;
-							break;
-						}
-					}
-				}
 				// Find Resolved UObjects
 				// Note: Most parameters are pushed in as part of the DataSet -> Parameters process
-				else if (Var.GetType().IsUObject())
+				if (Var.GetType().IsUObject())
 				{
 					const ENiagaraSimTarget ScriptSimTarget = iScript < EmitterScriptIndex ? ENiagaraSimTarget::CPUSim : SimTarget;
 					if (const FNiagaraScriptExecutionParameterStore* ScriptParameterStore = Script->GetExecutionReadyParameterStore(ScriptSimTarget))
@@ -1849,6 +1836,19 @@ void FVersionedNiagaraEmitterData::RebuildRendererBindings(const UNiagaraEmitter
 						}
 					}
 				}
+				// Find static variables
+				else
+				{
+					for (const FNiagaraVariable& StaticVar : Script->GetVMExecutableData().StaticVariablesWritten)
+					{
+						if (StaticVar.GetType().IsSameBaseDefinition(Var.GetType()) && StaticVar.GetName() == Var.GetName())
+						{
+							RendererBindings.AddParameter(StaticVar, true, false);
+							bVariableFound = true;
+							break;
+						}
+					}
+				}
 
 				if (bVariableFound)
 				{
