- Fix duplicate loctext key

#rnx

diff --git a/Engine/Plugins/FX/Niagara/Source/NiagaraEditorWidgets/Private/Stack/SNiagaraStackFunctionInputValue.cpp b/Engine/Plugins/FX/Niagara/Source/NiagaraEditorWidgets/Private/Stack/SNiagaraStackFunctionInputValue.cpp
--- a/Engine/Plugins/FX/Niagara/Source/NiagaraEditorWidgets/Private/Stack/SNiagaraStackFunctionInputValue.cpp
+++ b/Engine/Plugins/FX/Niagara/Source/NiagaraEditorWidgets/Private/Stack/SNiagaraStackFunctionInputValue.cpp
@@ -657,7 +657,7 @@ FText SNiagaraStackFunctionInputValue::GetObjectAssetValueText() const
 	}
 	else
 	{
-		return FText::Format(LOCTEXT("InvalidDataObjectFormat", "{0} (null)"), FunctionInput->GetInputType().GetClass()->GetDisplayNameText());
+		return FText::Format(LOCTEXT("InvalidObjectAssetFormat", "{0} (null)"), FunctionInput->GetInputType().GetClass()->GetDisplayNameText());
 	}
 }
 
