Added a 'is rename delegate bound' check to 'CanRenameItem' in SNiagaraParameterPanel, ensuring we can only rename items that are set up correctly
This fixes an issue where when the user hits F2 in the User Parameters panel which will trigger a request to rename in the parameters tab.
Ideally the SNiagaraParametersPanel would have its own set of commands, and not use the toolkit commands, which would make it so that F2 can only work if the SNiagaraParameterPanel is actively focused.

diff --git a/Engine/Plugins/FX/Niagara/Source/NiagaraEditor/Private/Widgets/SNiagaraParameterPanel.cpp b/Engine/Plugins/FX/Niagara/Source/NiagaraEditor/Private/Widgets/SNiagaraParameterPanel.cpp
--- a/Engine/Plugins/FX/Niagara/Source/NiagaraEditor/Private/Widgets/SNiagaraParameterPanel.cpp
+++ b/Engine/Plugins/FX/Niagara/Source/NiagaraEditor/Private/Widgets/SNiagaraParameterPanel.cpp
@@ -611,7 +611,8 @@ void SNiagaraParameterPanel::RequestRenameSelectedItem() const
 bool SNiagaraParameterPanel::CanRequestRenameSelectedItem() const
 {
 	const TArray<FNiagaraParameterPanelItem>& SelectedItems = ItemSelector->GetSelectedItems();
-	if (SelectedItems.Num() == 1)
+	// it's possible the request rename delegate has not been bound yet as the widget for it might not have been created at this point
+	if (SelectedItems.Num() == 1 && SelectedItems[0].GetOnRequestRename().IsBound())
 	{
 		FText Unused;
 		bool bCheckEmptyNameText = false;
