Adding a templated layer in the hlsl translator to encapsulate logic surrounding graph implementation to allow us to share code between the existing UNiagaraGraph and the upcoming graph digest.  Generated HLSL should remain identical with no significant change in performance; this change is all about setting up a new graph backend.
-renamed FHlslNiagaraTranslator to FNiagaraHlslTranslator
-fixes private include header being included in public NiagaraNode.h
-made translator's Compile() const

diff --git a/Engine/Plugins/FX/Niagara/Source/Niagara/Public/NiagaraCompilationTypes.h b/Engine/Plugins/FX/Niagara/Source/Niagara/Public/NiagaraCompilationTypes.h
--- a/Engine/Plugins/FX/Niagara/Source/Niagara/Public/NiagaraCompilationTypes.h
+++ b/Engine/Plugins/FX/Niagara/Source/Niagara/Public/NiagaraCompilationTypes.h
@@ -4,8 +4,51 @@
 
 #include "NiagaraCommon.h"
 
+#include "NiagaraScript.h"
+#include "NiagaraTypes.h"
+
+#include "NiagaraCompilationTypes.generated.h"
+
 class UNiagaraSystem;
 
+using FNiagaraCompilationTaskHandle = int32;
+
+USTRUCT()
+struct FNiagaraScriptAsyncCompileData
+{
+	GENERATED_USTRUCT_BODY()
+
+	FNiagaraVMExecutableDataId CompileId;
+	TSharedPtr<struct FNiagaraVMExecutableData> ExeData;
+	FString UniqueEmitterName;
+	bool bFromDerivedDataCache = false;
+
+	UPROPERTY()
+	TMap<FName, TObjectPtr<UNiagaraDataInterface>> NamedDataInterfaces;
+};
+
+USTRUCT()
+struct FNiagaraSystemAsyncCompileResults
+{
+	GENERATED_USTRUCT_BODY()
+
+	UPROPERTY()
+	TArray<TObjectPtr<UObject>> RootObjects;
+
+	FNiagaraCompilationTaskHandle CompilationTask;
+
+	using FCompileResultMap = TMap<UNiagaraScript*, FNiagaraScriptAsyncCompileData>;
+	FCompileResultMap CompileResultMap;
+
+	UPROPERTY()
+	TArray<FNiagaraVariable> ExposedVariables;
+
+	bool bForced = false;
+
+	float CombinedCompileTime = 0.0f;
+	float StartTime = 0.0f;
+};
+
 struct FNiagaraCompilationOptions
 {
 	UNiagaraSystem* System = nullptr;
@@ -52,4 +95,3 @@ protected:
 private:
 	bool bShouldApply = true;
 };
-
diff --git a/Engine/Plugins/FX/Niagara/Source/NiagaraEditor/Private/EdGraphSchema_Niagara.cpp b/Engine/Plugins/FX/Niagara/Source/NiagaraEditor/Private/EdGraphSchema_Niagara.cpp
--- a/Engine/Plugins/FX/Niagara/Source/NiagaraEditor/Private/EdGraphSchema_Niagara.cpp
+++ b/Engine/Plugins/FX/Niagara/Source/NiagaraEditor/Private/EdGraphSchema_Niagara.cpp
@@ -803,7 +803,7 @@ PRAGMA_ENABLE_DEPRECATION_WARNINGS
 			FNiagaraTypeDefinition PinType = PinToTypeDefinition(FromPin);
 
 			//Add all swizzles for this type if it's a vector.
-			if (FHlslNiagaraTranslator::IsHlslBuiltinVector(PinType))
+			if (FNiagaraHlslTranslator::IsHlslBuiltinVector(PinType))
 			{
 				TArray<FString> Components;
 				for (TFieldIterator<FProperty> PropertyIt(PinType.GetStruct(), EFieldIteratorFlags::IncludeSuper); PropertyIt; ++PropertyIt)
@@ -2190,7 +2190,7 @@ TArray<TTuple<FString, FString, FText>> UEdGraphSchema_Niagara::GetDataInterface
 				{
 					DataInterfaceName.RightChopInline(LastPinNameDot + 1);
 				}
-				DataInterfaceName = FHlslNiagaraTranslator::GetSanitizedSymbolName(DataInterfaceName);
+				DataInterfaceName = FNiagaraHlslTranslator::GetSanitizedSymbolName(DataInterfaceName);
 			}
 
 			// Generate all prototypes
@@ -2202,7 +2202,7 @@ TArray<TTuple<FString, FString, FText>> UEdGraphSchema_Niagara::GetDataInterface
 					continue;
 				}
 
-				FString Prototype = FHlslNiagaraTranslator::GenerateFunctionHlslPrototype(DataInterfaceName, FunctionSignature);
+				FString Prototype = FNiagaraHlslTranslator::GenerateFunctionHlslPrototype(DataInterfaceName, FunctionSignature);
 				if ( Prototype.Len() > 0 )
 				{
 					Prototype.Append(TEXT("\r\n"));
diff --git a/Engine/Plugins/FX/Niagara/Source/NiagaraEditor/Private/NiagaraAttributeTrimmer.cpp b/Engine/Plugins/FX/Niagara/Source/NiagaraEditor/Private/NiagaraAttributeTrimmer.cpp
--- a/Engine/Plugins/FX/Niagara/Source/NiagaraEditor/Private/NiagaraAttributeTrimmer.cpp
+++ b/Engine/Plugins/FX/Niagara/Source/NiagaraEditor/Private/NiagaraAttributeTrimmer.cpp
@@ -1,10 +1,11 @@
 // Copyright Epic Games, Inc. All Rights Reserved.
 
-#include "NiagaraHlslTranslator.h"
+#include "NiagaraAttributeTrimmer.h"
 
 #include "Algo/ForEach.h"
 #include "Algo/RemoveIf.h"
 #include "EdGraphSchema_Niagara.h"
+#include "NiagaraCompilationGraphBridge.h"
 #include "NiagaraConstants.h"
 #include "NiagaraModule.h"
 #include "NiagaraNodeCustomHlsl.h"
@@ -15,591 +16,552 @@
 #include "NiagaraNodeParameterMapSet.h"
 #include "NiagaraScriptVariable.h"
 
-namespace NiagaraAttributeTrimming
+template<typename GraphBridge>
+class FNiagaraAttributeTrimmerHelper<GraphBridge>::FFunctionInputResolver
 {
-	using FDependencySet = TSet<FModuleScopedPin>;
-	struct FCustomHlslNodeInfo
+public:
+	FFunctionInputResolver(const FGraph* Graph, ENiagaraScriptUsage Usage)
 	{
-		bool bHasDataInterfaceInputs = false;
-		bool bHasImpureFunctionText = false;
-	};
-
-	using FCustomHlslNodeMap = TMap<const UNiagaraNodeCustomHlsl*, FCustomHlslNodeInfo>;
-
-	struct FDependencyChain
-	{
-		FDependencySet Pins;
-		FCustomHlslNodeMap CustomNodes;
-	};
+		TMap<FName, FScopedPin> UnresolvedInputs;
 
-	using FDependencyMap = TMap<FModuleScopedPin, FDependencyChain>;
+		Traverse(TEXT(""), Graph, Usage, UnresolvedInputs);
+	}
 
-	class FFunctionInputResolver
+	FScopedPin ResolveInput(const FScopedPin& Input) const
 	{
-	public:
-		FFunctionInputResolver(UNiagaraGraph* Graph, ENiagaraScriptUsage Usage)
+		if (const FScopedPin* FunctionInput = FunctionInputMap.Find(Input))
 		{
-			TMap<FName, FModuleScopedPin> UnresolvedInputs;
-
-			Traverse(TEXT(""), Graph, Usage, UnresolvedInputs);
+			return *FunctionInput;
 		}
 
-		FModuleScopedPin ResolveInput(const FModuleScopedPin& Input) const
-		{
-			if (const FModuleScopedPin* FunctionInput = FunctionInputMap.Find(Input))
-			{
-				return *FunctionInput;
-			}
+		return FScopedPin();
+	}
 
-			return FModuleScopedPin();
-		}
+private:
+	void Traverse(FString Namespace, const FGraph* Graph, ENiagaraScriptUsage Usage, TMap<FName, FScopedPin>& UnresolvedInputs)
+	{
+		const FString NamespacePrefix = (Namespace.Len() > 0) ? (Namespace + TEXT(".")) : TEXT("");
 
-	private:
-		void Traverse(FString Namespace, UNiagaraGraph* Graph, ENiagaraScriptUsage Usage, TMap<FName, FModuleScopedPin>& UnresolvedInputs)
-		{
-			const FString NamespacePrefix = (Namespace.Len() > 0) ? (Namespace + TEXT(".")) : TEXT("");
+		TArray<const FOutputNode*> OutputNodes;
+		GraphBridge::FindOutputNodes(Graph, Usage, OutputNodes);
 
-			TArray<UNiagaraNodeOutput*> OutputNodes;
-			Graph->FindOutputNodes(Usage, OutputNodes);
+		for (const FOutputNode* OutputNode : OutputNodes)
+		{
+			TArray<const FNode*> TraversedNodes;
+			GraphBridge::BuildTraversal(Graph, OutputNode, TraversedNodes);
 
-			for (UNiagaraNodeOutput* OutputNode : OutputNodes)
+			for (const FNode* Node : TraversedNodes)
 			{
-				TArray<UNiagaraNode*> TraversedNodes;
-				Graph->BuildTraversal(TraversedNodes, OutputNode, true);
-
-				for (UNiagaraNode* Node : TraversedNodes)
+				if (const FFunctionCallNode* FunctionNode = GraphBridge::template AsNodeType<FFunctionCallNode>(Node))
 				{
-					if (UNiagaraNodeFunctionCall* FunctionNode = Cast<UNiagaraNodeFunctionCall>(Node))
+					if (const FGraph* FunctionGraph = GraphBridge::GetFunctionNodeGraph(FunctionNode))
 					{
-						if (UNiagaraGraph* FunctionGraph = FunctionNode->GetCalledGraph())
-						{
-							TMap<FName, FModuleScopedPin> FunctionInputs;
-							Traverse(NamespacePrefix + FunctionNode->GetFunctionName(), FunctionGraph, FunctionNode->GetCalledUsage(), FunctionInputs);
+						TMap<FName, FScopedPin> FunctionInputs;
+						Traverse(NamespacePrefix + GraphBridge::GetFunctionName(FunctionNode), FunctionGraph, GraphBridge::GetFunctionUsage(FunctionNode), FunctionInputs);
 
-							// resolve any of the UnresolvedInputs that we've accumulated
-							for (UEdGraphPin* NodePin : FunctionNode->GetAllPins())
+						// resolve any of the UnresolvedInputs that we've accumulated
+						for (const FInputPin* NodePin : GraphBridge::GetInputPins(FunctionNode))
+						{
+							if (FScopedPin* InnerPin = FunctionInputs.Find(NodePin->PinName))
 							{
-								if (NodePin->Direction == EEdGraphPinDirection::EGPD_Input)
-								{
-									if (FModuleScopedPin* InnerPin = FunctionInputs.Find(NodePin->PinName))
-									{
-										FunctionInputMap.Add(*InnerPin, FModuleScopedPin(NodePin, *Namespace));
-										FunctionInputs.Remove(NodePin->PinName);
-									}
-								}
+								FunctionInputMap.Add(*InnerPin, FScopedPin(NodePin, *Namespace));
+								FunctionInputs.Remove(NodePin->PinName);
 							}
 						}
 					}
-					if (UNiagaraNodeInput* InputNode = Cast<UNiagaraNodeInput>(Node))
+				}
+				if (const FInputNode* InputNode = GraphBridge::template AsNodeType<FInputNode>(Node))
+				{
+					int32 OutputPinCount = 0;
+					for (const FOutputPin* InputNodePin : GraphBridge::GetOutputPins(InputNode))
 					{
-						int32 OutputPinCount = 0;
-						for (UEdGraphPin* InputNodePin : InputNode->GetAllPins())
-						{
-							if (InputNodePin->Direction == EEdGraphPinDirection::EGPD_Output)
-							{
-								check(OutputPinCount == 0);
-								++OutputPinCount;
+						check(OutputPinCount == 0);
+						++OutputPinCount;
 
-								UnresolvedInputs.Add(InputNode->Input.GetName(), FModuleScopedPin(InputNodePin, *Namespace));
-							}
-						}
+						UnresolvedInputs.Add(GraphBridge::GetInputVariable(InputNode).GetName(), FScopedPin(InputNodePin, *Namespace));
 					}
 				}
 			}
 		}
+	}
 
-		TMap<FModuleScopedPin, FModuleScopedPin> FunctionInputMap;
-	};
+	TMap<FScopedPin, FScopedPin> FunctionInputMap;
+};
 
-	class FImpureFunctionParser
+template<typename GraphBridge>
+class FNiagaraAttributeTrimmerHelper<GraphBridge>::FImpureFunctionParser
+{
+public:
+	FImpureFunctionParser(const FGraph* Graph, ENiagaraScriptUsage Usage)
 	{
-	public:
-		FImpureFunctionParser(UNiagaraGraph* Graph, ENiagaraScriptUsage Usage)
-		{
-			Traverse(TEXT(""), Graph, Usage);
-		}
+		Traverse(TEXT(""), Graph, Usage);
+	}
 
-		const TArray<FModuleScopedPin>& ReadFunctionInputs() const
-		{
-			return ImpureFunctionInputs;
-		}
+	const TArray<FScopedPin>& ReadFunctionInputs() const
+	{
+		return ImpureFunctionInputs;
+	}
 
-	private:
-		void Traverse(FString Namespace, UNiagaraGraph* Graph, ENiagaraScriptUsage Usage)
-		{
-			const FString NamespacePrefix = (Namespace.Len() > 0) ? (Namespace + TEXT(".")) : TEXT("");
+private:
+	void Traverse(FString Namespace, const FGraph* Graph, ENiagaraScriptUsage Usage)
+	{
+		const FString NamespacePrefix = (Namespace.Len() > 0) ? (Namespace + TEXT(".")) : TEXT("");
 
-			TArray<UNiagaraNodeOutput*> OutputNodes;
-			Graph->FindOutputNodes(Usage, OutputNodes);
+		TArray<const FOutputNode*> OutputNodes;
+		GraphBridge::FindOutputNodes(Graph, Usage, OutputNodes);
 
-			for (UNiagaraNodeOutput* OutputNode : OutputNodes)
-			{
-				TArray<UNiagaraNode*> TraversedNodes;
-				Graph->BuildTraversal(TraversedNodes, OutputNode, true);
+		for (const FOutputNode* OutputNode : OutputNodes)
+		{
+			TArray<const FNode*> TraversedNodes;
+			GraphBridge::BuildTraversal(Graph, OutputNode, TraversedNodes);
 
-				for (UNiagaraNode* Node : TraversedNodes)
+			for (const FNode* Node : TraversedNodes)
+			{
+				if (const FFunctionCallNode* FunctionNode = GraphBridge::template AsNodeType<FFunctionCallNode>(Node))
 				{
-					if (UNiagaraNodeFunctionCall* FunctionNode = Cast<UNiagaraNodeFunctionCall>(Node))
+					if (const FGraph* FunctionGraph = GraphBridge::GetFunctionNodeGraph(FunctionNode))
 					{
-						if (UNiagaraGraph* FunctionGraph = FunctionNode->GetCalledGraph())
-						{
-							Traverse(NamespacePrefix + FunctionNode->GetFunctionName(), FunctionGraph, FunctionNode->GetCalledUsage());
-						}
+						Traverse(NamespacePrefix + GraphBridge::GetFunctionName(FunctionNode), FunctionGraph, GraphBridge::GetFunctionUsage(FunctionNode));
+					}
 
-						if (FunctionNode->Signature.bRequiresExecPin)
+					if (FunctionNode->Signature.bRequiresExecPin)
+					{
+						for (const FInputPin* Pin : GraphBridge::GetInputPins(FunctionNode))
 						{
-							for (const UEdGraphPin* Pin : FunctionNode->GetAllPins())
+							// setting the following to false will allow the trimmer to remove additional attributes.  A value of false
+							// will limit our dependency search to the function inputs rather than the execution path.
+							constexpr bool bBackCompat_IncludeExecPinsForImpureFunctions = true;
+
+							if (!bBackCompat_IncludeExecPinsForImpureFunctions)
 							{
-								if (Pin->Direction == EEdGraphPinDirection::EGPD_Input)
+								if (GraphBridge::GetPinType(Pin, ENiagaraStructConversion::Simulation) == FNiagaraTypeDefinition::GetParameterMapDef())
 								{
-									// setting the following to false will allow the trimmer to remove additional attributes.  A value of false
-									// will limit our dependency search to the function inputs rather than the execution path.
-									constexpr bool bBackCompat_IncludeExecPinsForImpureFunctions = true;
-
-									if (!bBackCompat_IncludeExecPinsForImpureFunctions)
-									{
-										if (UEdGraphSchema_Niagara::PinToTypeDefinition(Pin) == FNiagaraTypeDefinition::GetParameterMapDef())
-										{
-											continue;
-										}
-									}
-
-									ImpureFunctionInputs.Emplace(Pin, *Namespace);
+									continue;
 								}
 							}
+
+							ImpureFunctionInputs.Emplace(Pin, *Namespace);
 						}
 					}
 				}
 			}
 		}
+	}
 
-		TArray<FModuleScopedPin> ImpureFunctionInputs;
-	};
+	TArray<FScopedPin> ImpureFunctionInputs;
+};
 
-	class FExpressionBuilder
-	{
-	public:
+template<typename GraphBridge>
+class FNiagaraAttributeTrimmerHelper<GraphBridge>::FExpressionBuilder
+{
+public:
 
-		FExpressionBuilder(const FNiagaraParameterMapHistory& InParamMap, const FFunctionInputResolver& InInputResolver, const UEdGraphSchema_Niagara* InNiagaraSchema)
+	FExpressionBuilder(const FParamMapHistory& InParamMap, const FFunctionInputResolver& InInputResolver)
 		: ParamMap(InParamMap)
 		, InputResolver(InInputResolver)
-		, NiagaraSchema(InNiagaraSchema)
-		{}
-
-		// generates a set of dependencies for a specific pin
-		//	-generates the set of input pins to nodes that contribute to the 'expression' of the supplied output pin (EndPin)
-		//		An expression is the set of nodes that define what is being set in a MapSet.  It does not traverse through
-		//		Exec pins as those nodes would be part of another expression
-		//	-generates the list of custom HLSL nodes that are encountered in the traversal
-		void FindDependencies(const FModuleScopedPin& EndPin, FDependencyChain& Dependencies)
-		{
-			FindDependenciesInternal(EndPin, Dependencies, nullptr);
-		}
+	{}
+
+	// generates a set of dependencies for a specific pin
+	//	-generates the set of input pins to nodes that contribute to the 'expression' of the supplied output pin (EndPin)
+	//		An expression is the set of nodes that define what is being set in a MapSet.  It does not traverse through
+	//		Exec pins as those nodes would be part of another expression
+	//	-generates the list of custom HLSL nodes that are encountered in the traversal
+	void FindDependencies(const FScopedPin& EndPin, FDependencyChain& Dependencies)
+	{
+		FindDependenciesInternal(EndPin, Dependencies, nullptr);
+	}
 
-	private:
-		// generates a set of dependencies for a specific pin
-		//	-generates the set of input pins to nodes that contribute to the 'expression' of the supplied output pin (EndPin)
-		//		An expression is the set of nodes that define what is being set in a MapSet.  It does not traverse through
-		//		Exec pins as those nodes would be part of another expression
-		//	-generates the list of custom HLSL nodes that are encountered in the traversal
-		void FindDependenciesInternal(const FModuleScopedPin& EndPin, FDependencyChain& Dependencies, TArray<UNiagaraNodeInput*>* InputNodes)
-		{
-			TStringBuilder<1024> NamespaceBuilder;
+private:
+	// generates a set of dependencies for a specific pin
+	//	-generates the set of input pins to nodes that contribute to the 'expression' of the supplied output pin (EndPin)
+	//		An expression is the set of nodes that define what is being set in a MapSet.  It does not traverse through
+	//		Exec pins as those nodes would be part of another expression
+	//	-generates the list of custom HLSL nodes that are encountered in the traversal
+	void FindDependenciesInternal(const FScopedPin& EndPin, FDependencyChain& Dependencies, TArray<const FInputNode*>* InputNodes)
+	{
+		TStringBuilder<1024> NamespaceBuilder;
 
-			TArray<FModuleScopedPin> PinsToEvaluate;
-			TSet<UNiagaraNode*> EvaluatedNodes;
+		TArray<FScopedPin> PinsToEvaluate;
+		TSet<const FNode*> EvaluatedNodes;
 
-			FModuleScopedPin CurrentPin = EndPin;
-			while (CurrentPin.Pin)
+		FScopedPin CurrentPin = EndPin;
+		while (CurrentPin.Pin)
+		{
+			if (CurrentPin.Pin->Direction == EEdGraphPinDirection::EGPD_Output)
 			{
-				for (UEdGraphPin* LinkedPin : CurrentPin.Pin->LinkedTo)
-				{
-					if (LinkedPin->Direction == EEdGraphPinDirection::EGPD_Input)
-					{
-						continue;
-					}
+				continue;
+			}
 
-					if (UEdGraphPin* TracedPin = EvaluateStaticSwitches ? UNiagaraNode::TraceOutputPin(LinkedPin) : LinkedPin)
-					{
-						UNiagaraNode* CurrentNode = CastChecked<UNiagaraNode>(TracedPin->GetOwningNode());
+			const FInputPin* CurrentInputPin = static_cast<const FInputPin*>(CurrentPin.Pin);
+			if (const FOutputPin* TracedPin = GraphBridge::GetLinkedOutputPin(CurrentInputPin))
+			{
+				const FNode* CurrentNode = GraphBridge::GetOwningNode(TracedPin);
+				bool AlreadyEvaluated = false;
 
-						bool AlreadyEvaluated = false;
+				// record any output pins to a parameter get node
+				if (const FParamMapGetNode* ParamGetNode = GraphBridge::template AsNodeType<FParamMapGetNode>(CurrentNode))
+				{
+					Dependencies.Pins.Emplace(FScopedPin(TracedPin, CurrentPin.ModuleName));
+				}
+				else if (const FFunctionCallNode* FunctionCallNode = GraphBridge::template AsNodeType<FFunctionCallNode>(CurrentNode))
+				{
+					if (const FGraph* NodeGraph = GraphBridge::GetFunctionNodeGraph(FunctionCallNode))
+					{
+						TArray<const FOutputNode*> InnerOutputNodes;
+						GraphBridge::FindOutputNodes(NodeGraph, InnerOutputNodes);
 
-						// record any output pins to a parameter get node
-						if (UNiagaraNodeParameterMapGet* ParamGetNode = Cast<UNiagaraNodeParameterMapGet>(CurrentNode))
+						NamespaceBuilder.Reset();
+						if (CurrentPin.ModuleName != NAME_None)
 						{
-							Dependencies.Pins.Emplace(FModuleScopedPin(TracedPin, CurrentPin.ModuleName));
+							CurrentPin.ModuleName.AppendString(NamespaceBuilder);
+							NamespaceBuilder << TEXT(".");
 						}
-						else if (UNiagaraNodeFunctionCall* FunctionCallNode = Cast<UNiagaraNodeFunctionCall>(CurrentNode))
-						{
-							TArray<UNiagaraNodeOutput*> InnerOutputNodes;
+						NamespaceBuilder << GraphBridge::GetFunctionName(FunctionCallNode);
 
-							if (const UNiagaraGraph* NodeGraph = FunctionCallNode->GetCalledGraph())
-							{
-								NodeGraph->FindOutputNodes(InnerOutputNodes);
+						const FName AggregateModuleName = *NamespaceBuilder;
 
-								NamespaceBuilder.Reset();
-								if (CurrentPin.ModuleName != NAME_None)
+						for (const FOutputNode* InnerOutputNode : InnerOutputNodes)
+						{
+							for (const FInputPin* InnerOutputNodePin : GraphBridge::GetInputPins(InnerOutputNode))
+							{
+								if (InnerOutputNodePin->PinName == TracedPin->PinName)
 								{
-									CurrentPin.ModuleName.AppendString(NamespaceBuilder);
-									NamespaceBuilder << TEXT(".");
-								}
-								NamespaceBuilder << FunctionCallNode->GetFunctionName();
+									TArray<const FInputNode*> InnerInputNodes;
+									FindDependenciesInternal(FScopedPin(InnerOutputNodePin, AggregateModuleName), Dependencies, &InnerInputNodes);
 
-								const FName AggregateModuleName = *NamespaceBuilder;
-
-								for (const UNiagaraNodeOutput* InnerOutputNode : InnerOutputNodes)
-								{
-									for (const UEdGraphPin* InnerOutputNodePin : InnerOutputNode->GetAllPins())
+									// map the InputNodes we found to the input pins on the FunctionCallNode
+									for (const FInputNode* InnerInputNode : InnerInputNodes)
 									{
-										if (InnerOutputNodePin->Direction == EEdGraphPinDirection::EGPD_Input
-											&& InnerOutputNodePin->PinName == TracedPin->PinName)
+										for (const FInputPin* NodePin : GraphBridge::GetInputPins(FunctionCallNode))
 										{
-											TArray<UNiagaraNodeInput*> InnerInputNodes;
-											FindDependenciesInternal(FModuleScopedPin(InnerOutputNodePin, AggregateModuleName), Dependencies, &InnerInputNodes);
-
-											// map the InputNodes we found to the input pins on the FunctionCallNode
-											for (const UNiagaraNodeInput* InnerInputNode : InnerInputNodes)
+											if (NodePin->PinName == GraphBridge::GetInputVariable(InnerInputNode).GetName())
 											{
-												for (UEdGraphPin* NodePin : FunctionCallNode->GetAllPins())
-												{
-													if (NodePin->Direction == EEdGraphPinDirection::EGPD_Input
-														&& NodePin->PinName == InnerInputNode->Input.GetName())
-													{
-														PinsToEvaluate.Emplace(NodePin, CurrentPin.ModuleName);
-														break;
-													}
-												}
+												PinsToEvaluate.Emplace(NodePin, CurrentPin.ModuleName);
+												break;
 											}
 										}
 									}
 								}
-
-								// skip adding the inputs to the list of pins we need to evaluate because we've already search through the internal graph
-								// to figure out the minimal set of items we care about
-								AlreadyEvaluated = true;
-							}
-						}
-						else if (UNiagaraNodeInput* InputNode = Cast<UNiagaraNodeInput>(CurrentNode))
-						{
-							if (InputNodes)
-							{
-								InputNodes->AddUnique(InputNode);
-							}
-							else
-							{
-								const FModuleScopedPin ResolvedPin = InputResolver.ResolveInput(FModuleScopedPin(TracedPin, CurrentPin.ModuleName));
-								if (ResolvedPin.Pin)
-								{
-									PinsToEvaluate.Add(ResolvedPin);
-									EvaluatedNodes.Add(InputNode);
-									AlreadyEvaluated = true;
-								}
 							}
 						}
 
-						if (!AlreadyEvaluated)
+						// skip adding the inputs to the list of pins we need to evaluate because we've already search through the internal graph
+						// to figure out the minimal set of items we care about
+						AlreadyEvaluated = true;
+					}
+				}
+				else if (const FInputNode* InputNode = GraphBridge::template AsNodeType<FInputNode>(CurrentNode))
+				{
+					if (InputNodes)
+					{
+						InputNodes->AddUnique(InputNode);
+					}
+					else
+					{
+						const FScopedPin ResolvedPin = InputResolver.ResolveInput(FScopedPin(TracedPin, CurrentPin.ModuleName));
+						if (ResolvedPin.Pin)
 						{
-							EvaluatedNodes.Add(CurrentNode, &AlreadyEvaluated);
+							PinsToEvaluate.Add(ResolvedPin);
+							EvaluatedNodes.Add(InputNode);
+							AlreadyEvaluated = true;
 						}
+					}
+				}
 
-						if (!AlreadyEvaluated)
-						{
-							for (UEdGraphPin* NodePin : CurrentNode->GetAllPins())
-							{
-								if (NodePin->Direction == EEdGraphPinDirection::EGPD_Output)
-								{
-									continue;
-								}
+				if (!AlreadyEvaluated)
+				{
+					EvaluatedNodes.Add(CurrentNode, &AlreadyEvaluated);
+				}
 
-								// we want to exclude execution & parameter map pins because they terminate the subexpression
-								const bool IsParamMapPin = UEdGraphSchema_Niagara::PinToTypeDefinition(NodePin) == FNiagaraTypeDefinition::GetParameterMapDef();
-								const bool IsExecPin = NodePin->PinName == UNiagaraNodeParameterMapBase::SourcePinName;
-								if (IsParamMapPin || IsExecPin)
-								{
-									continue;
-								}
+				if (!AlreadyEvaluated)
+				{
+					for (const FInputPin* NodePin : GraphBridge::GetInputPins(CurrentNode))
+					{
+						// we want to exclude execution pins
+						const bool IsParamMapPin = GraphBridge::GetPinType(NodePin, ENiagaraStructConversion::Simulation) == FNiagaraTypeDefinition::GetParameterMapDef();
+						const bool IsExecPin = NodePin->PinName == UNiagaraNodeParameterMapBase::SourcePinName;
+						if (IsParamMapPin || IsExecPin)
+						{
+							continue;
+						}
 
-								PinsToEvaluate.Emplace(NodePin, CurrentPin.ModuleName);
-							}
+						PinsToEvaluate.Emplace(NodePin, CurrentPin.ModuleName);
+					}
 
-							if (const UNiagaraNodeCustomHlsl* CustomHlslNode = Cast<const UNiagaraNodeCustomHlsl>(CurrentNode))
-							{
-								if (!Dependencies.CustomNodes.Contains(CustomHlslNode))
-								{
-									FCustomHlslNodeInfo& NodeInfo = Dependencies.CustomNodes.Add(CustomHlslNode);
-									BuildCustomHlslNodeInfo(CustomHlslNode, NodeInfo);
-								}
-							}
+					if (const FCustomHlslNode* CustomHlslNode = GraphBridge::template AsNodeType<FCustomHlslNode>(CurrentNode))
+					{
+						if (!Dependencies.CustomNodes.Contains(CustomHlslNode))
+						{
+							FCustomHlslNodeInfo& NodeInfo = Dependencies.CustomNodes.Add(CustomHlslNode);
+							BuildCustomHlslNodeInfo(CustomHlslNode, NodeInfo);
 						}
 					}
 				}
+			}
 
-				if (PinsToEvaluate.Num())
-				{
-					CurrentPin = PinsToEvaluate.Pop(false);
-				}
-				else
-				{
-					break;
-				}
+			if (PinsToEvaluate.Num())
+			{
+				CurrentPin = PinsToEvaluate.Pop(false);
+			}
+			else
+			{
+				break;
 			}
 		}
+	}
 
-		void BuildCustomHlslNodeInfo(const UNiagaraNodeCustomHlsl* CustomNode, FCustomHlslNodeInfo& NodeInfo)
-		{
-			check(CustomNode);
+	void BuildCustomHlslNodeInfo(const FCustomHlslNode* CustomNode, FCustomHlslNodeInfo& NodeInfo)
+	{
+		check(CustomNode);
 
-			NodeInfo.bHasDataInterfaceInputs = false;
-			NodeInfo.bHasImpureFunctionText = false;
+		NodeInfo.bHasDataInterfaceInputs = false;
+		NodeInfo.bHasImpureFunctionText = false;
 
-			TArray<FString> ImpureFunctionNames;
+		TArray<FString> ImpureFunctionNames;
 
-			FPinCollectorArray InputPins;
-			CustomNode->GetInputPins(InputPins);
-			for (const UEdGraphPin* InputPin : InputPins)
+		for (const FInputPin* InputPin : GraphBridge::GetInputPins(CustomNode))
+		{
+			FNiagaraTypeDefinition NiagaraType = GraphBridge::GetPinType(InputPin, ENiagaraStructConversion::Simulation);
+			if (NiagaraType.IsDataInterface())
 			{
-				FNiagaraTypeDefinition NiagaraType = NiagaraSchema->PinToTypeDefinition(InputPin);
-				if (NiagaraType.IsDataInterface())
+				NodeInfo.bHasDataInterfaceInputs = true;
+
+				if (UNiagaraDataInterface* DataInterfaceClass = CastChecked<UNiagaraDataInterface>(NiagaraType.GetClass()->ClassDefaultObject))
 				{
-					NodeInfo.bHasDataInterfaceInputs = true;
+					TArray<FNiagaraFunctionSignature> FunctionSignatures;
+					DataInterfaceClass->GetFunctions(FunctionSignatures);
 
-					if (UNiagaraDataInterface* DataInterfaceClass = CastChecked<UNiagaraDataInterface>(NiagaraType.GetClass()->ClassDefaultObject))
+					for (const FNiagaraFunctionSignature& FunctionSignature : FunctionSignatures)
 					{
-						TArray<FNiagaraFunctionSignature> FunctionSignatures;
-						DataInterfaceClass->GetFunctions(FunctionSignatures);
-
-						for (const FNiagaraFunctionSignature& FunctionSignature : FunctionSignatures)
+						if (FunctionSignature.bRequiresExecPin)
 						{
-							if (FunctionSignature.bRequiresExecPin)
-							{
-								TStringBuilder<256> Builder;
-								InputPin->GetFName().AppendString(Builder);
-								Builder.AppendChar(TCHAR('.'));
-								FunctionSignature.Name.AppendString(Builder);
+							TStringBuilder<256> Builder;
+							InputPin->PinName.AppendString(Builder);
+							Builder.AppendChar(TCHAR('.'));
+							FunctionSignature.Name.AppendString(Builder);
 
-								ImpureFunctionNames.AddUnique(Builder.ToString());
-							}
+							ImpureFunctionNames.AddUnique(Builder.ToString());
 						}
 					}
 				}
 			}
+		}
 
-			if (!ImpureFunctionNames.IsEmpty())
-			{
-				TArray<FString> StringTokens;
-				UNiagaraNodeCustomHlsl::GetTokensFromString(CustomNode->GetCustomHlsl(), StringTokens, false, false);
-
-				for (const FString& FunctionCall : ImpureFunctionNames)
+		if (!ImpureFunctionNames.IsEmpty())
+		{
+			TArray<FStringView> ImpureFunctionNameViews;
+			ImpureFunctionNameViews.Reserve(ImpureFunctionNames.Num());
+			Algo::Transform(ImpureFunctionNames, ImpureFunctionNameViews, [](const FString& FunctionName) -> FStringView
 				{
-					for (const FString& Token : StringTokens)
-					{
-						if (Token.Equals(FunctionCall))
-						{
-							NodeInfo.bHasImpureFunctionText = true;
-							return;
-						}
-					}
-				}
+					return FunctionName;
+				});
+
+			if (GraphBridge::CustomHlslReferencesTokens(CustomNode, ImpureFunctionNameViews))
+			{
+				return;
 			}
 		}
+	}
 
-		const FNiagaraParameterMapHistory& ParamMap;
-		const FFunctionInputResolver& InputResolver;
-		const UEdGraphSchema_Niagara* NiagaraSchema = nullptr;
+	const FParamMapHistory& ParamMap;
+	const FFunctionInputResolver& InputResolver;
 
-		const bool EvaluateStaticSwitches = false;
-	};
+	const bool EvaluateStaticSwitches = false;
+};
 
-	// For a specific read of a variable finds the corresponding PreviousWritePin if one exists (only considers actual writes
-	// rather than default pins on a MapGet)
-	static FModuleScopedPin FindWriteForRead(const FNiagaraParameterMapHistory& ParamMap, const FModuleScopedPin& ReadPin, int32* OutVariableIndex)
+// For a specific read of a variable finds the corresponding PreviousWritePin if one exists (only considers actual writes
+// rather than default pins on a MapGet)
+template<typename GraphBridge>
+typename FNiagaraAttributeTrimmerHelper<GraphBridge>::FScopedPin FNiagaraAttributeTrimmerHelper<GraphBridge>::FindWriteForRead(const FParamMapHistory& ParamMap, const FScopedPin& ReadPin, int32* OutVariableIndex)
+{
+	if (ReadPin.Pin == nullptr)
 	{
-		if (ReadPin.Pin == nullptr)
-		{
-			return FModuleScopedPin();
-		}
+		return FScopedPin();
+	}
 
-		int32 VariableCount = ParamMap.Variables.Num();
+	int32 VariableCount = ParamMap.Variables.Num();
 
-		// without more context we need to do a full linear search through the ReadHistory
-		for (int32 VariableIt = 0; VariableIt < VariableCount; ++VariableIt)
-		{
-			const TArray<FNiagaraParameterMapHistory::FReadHistory>& ReadHistoryEntries = ParamMap.PerVariableReadHistory[VariableIt];
+	// without more context we need to do a full linear search through the ReadHistory
+	for (int32 VariableIt = 0; VariableIt < VariableCount; ++VariableIt)
+	{
+		const TArray<typename FParamMapHistory::FReadHistory>& ReadHistoryEntries = ParamMap.PerVariableReadHistory[VariableIt];
 
-			for (const FNiagaraParameterMapHistory::FReadHistory& ReadHistoryEntry : ReadHistoryEntries)
+		for (const typename FParamMapHistory::FReadHistory& ReadHistoryEntry : ReadHistoryEntries)
+		{
+			if (ReadHistoryEntry.ReadPin == ReadPin)
 			{
-				if (ReadHistoryEntry.ReadPin == ReadPin)
+				if (OutVariableIndex)
 				{
-					if (OutVariableIndex)
-					{
-						*OutVariableIndex = VariableIt;
-					}
+					*OutVariableIndex = VariableIt;
+				}
 
-					if (ReadHistoryEntry.PreviousWritePin.Pin)
+				if (ReadHistoryEntry.PreviousWritePin.Pin)
+				{
+					// we only care about writes (MapSets).  This will exclude default variables that are set as the input to the MapGet
+					if (const FParamMapSetNode* WriteNode = GraphBridge::template AsNodeType<FParamMapSetNode>(GraphBridge::GetOwningNode(ReadHistoryEntry.PreviousWritePin.Pin)))
 					{
-						// we only care about writes (MapSets).  This will exclude default variables that are set as the input to the MapGet
-						if (const UNiagaraNodeParameterMapSet* WriteNode = Cast<const UNiagaraNodeParameterMapSet>(ReadHistoryEntry.PreviousWritePin.Pin->GetOwningNode()))
-						{
-							return ReadHistoryEntry.PreviousWritePin;
-						}
+						return ReadHistoryEntry.PreviousWritePin;
 					}
-
-					return FModuleScopedPin();
 				}
+
+				return FScopedPin();
 			}
 		}
-
-		return FModuleScopedPin();
 	}
 
-	static int32 FindDefaultBinding(const FNiagaraParameterMapHistory& ParamMap, int32 VariableIndex, const FModuleScopedPin& ReadPin)
+	return FScopedPin();
+}
+
+template<typename GraphBridge>
+int32 FNiagaraAttributeTrimmerHelper<GraphBridge>::FindDefaultBinding(const FParamMapHistory& ParamMap, int32 VariableIndex, const FScopedPin& ReadPin)
+{
+	if (const FGraph* OwnerGraph = GraphBridge::GetOwningGraph(GraphBridge::GetOwningNode(ReadPin.Pin)))
 	{
-		if (const UNiagaraGraph* OwnerGraph = Cast<const UNiagaraGraph>(ReadPin.Pin->GetOwningNode()->GetGraph()))
-		{
-			FNiagaraVariable ParamVariable = ParamMap.Variables[VariableIndex];
-			FString ParamName = ParamVariable.GetName().ToString();
+		FNiagaraVariable ParamVariable = ParamMap.Variables[VariableIndex];
+		FString ParamName = ParamVariable.GetName().ToString();
 
-			// now we need to replace the module name with 'module'
-			ParamName.ReplaceInline(*ReadPin.ModuleName.ToString(), *FNiagaraConstants::ModuleNamespace.ToString());
+		// now we need to replace the module name with 'module'
+		ParamName.ReplaceInline(*ReadPin.ModuleName.ToString(), *FNiagaraConstants::ModuleNamespace.ToString());
 
-			ParamVariable.SetName(*ParamName);
+		ParamVariable.SetName(*ParamName);
 
-			FNiagaraScriptVariableBinding DefaultBinding;
-			TOptional<ENiagaraDefaultMode> DefaultMode = OwnerGraph->GetDefaultMode(ParamVariable, &DefaultBinding);
-			if (DefaultMode.IsSet() && *DefaultMode == ENiagaraDefaultMode::Binding && DefaultBinding.IsValid())
-			{
-				return ParamMap.FindVariableByName(DefaultBinding.GetName());
-			}
+		FNiagaraScriptVariableBinding DefaultBinding;
+		TOptional<ENiagaraDefaultMode> DefaultMode = GraphBridge::GetGraphDefaultMode(OwnerGraph, ParamVariable, DefaultBinding);
+		if (DefaultMode.IsSet() && *DefaultMode == ENiagaraDefaultMode::Binding && DefaultBinding.IsValid())
+		{
+			return ParamMap.FindVariableByName(DefaultBinding.GetName());
 		}
-
-		return INDEX_NONE;
 	}
 
-	// for a specific read of a variable finds the actual name of the attribute being read
-	static FName FindAttributeForRead(const FNiagaraParameterMapHistory& ParamMap, const FModuleScopedPin& ReadPin)
+	return INDEX_NONE;
+}
+
+// for a specific read of a variable finds the actual name of the attribute being read
+template<typename GraphBridge>
+FName FNiagaraAttributeTrimmerHelper<GraphBridge>::FindAttributeForRead(const FParamMapHistory& ParamMap, const FScopedPin& ReadPin)
+{
+	const int32 VariableCount = ParamMap.Variables.Num();
+
+	// wihtout more context we need to do a full linear search through the ReadHistory
+	for (int32 VariableIt = 0; VariableIt < VariableCount; ++VariableIt)
 	{
-		const int32 VariableCount = ParamMap.Variables.Num();
+		const TArray<typename FParamMapHistory::FReadHistory>& ReadHistoryEntries = ParamMap.PerVariableReadHistory[VariableIt];
 
-		// wihtout more context we need to do a full linear search through the ReadHistory
-		for (int32 VariableIt = 0; VariableIt < VariableCount; ++VariableIt)
+		for (const typename FParamMapHistory::FReadHistory& ReadHistoryEntry : ReadHistoryEntries)
 		{
-			const TArray<FNiagaraParameterMapHistory::FReadHistory>& ReadHistoryEntries = ParamMap.PerVariableReadHistory[VariableIt];
-
-			for (const FNiagaraParameterMapHistory::FReadHistory& ReadHistoryEntry : ReadHistoryEntries)
+			if (ReadHistoryEntry.ReadPin == ReadPin)
 			{
-				if (ReadHistoryEntry.ReadPin == ReadPin)
-				{
-					return ParamMap.Variables[VariableIt].GetName();
-				}
+				return ParamMap.Variables[VariableIt].GetName();
 			}
 		}
-
-		return NAME_None;
 	}
 
-	// given the set of expressions (as defined in FindDependencies above) we resolve the named attribute aggregating the dependent reads and custom nodes
-	static void ResolveDependencyChain(const FNiagaraParameterMapHistory& ParamMap, const FDependencyMap& DependencyData, const FName& AttributeName, FDependencyChain& ResolvedDependencies)
+	return NAME_None;
+}
+
+// given the set of expressions (as defined in FindDependencies above) we resolve the named attribute aggregating the dependent reads and custom nodes
+template<typename GraphBridge>
+void FNiagaraAttributeTrimmerHelper<GraphBridge>::ResolveDependencyChain(const FParamMapHistory& ParamMap, const FDependencyMap& DependencyData, const FName& AttributeName, FDependencyChain& ResolvedDependencies)
+{
+	const int32 VariableIndex = ParamMap.FindVariableByName(AttributeName);
+	if (VariableIndex != INDEX_NONE)
 	{
-		const int32 VariableIndex = ParamMap.FindVariableByName(AttributeName);
-		if (VariableIndex != INDEX_NONE)
+		if (ParamMap.PerVariableWriteHistory[VariableIndex].Num())
 		{
-			if (ParamMap.PerVariableWriteHistory[VariableIndex].Num())
-			{
-				const FModuleScopedPin& LastWritePin = ParamMap.PerVariableWriteHistory[VariableIndex].Last();
+			const FScopedPin& LastWritePin = ParamMap.PerVariableWriteHistory[VariableIndex].Last();
 
-				TArray<FModuleScopedPin> PinsToResolve;
-				PinsToResolve.Add(LastWritePin);
+			TArray<FScopedPin> PinsToResolve;
+			PinsToResolve.Add(LastWritePin);
 
-				while (PinsToResolve.Num() > 0)
-				{
-					const FModuleScopedPin PinToResolve = PinsToResolve.Pop(false);
+			while (PinsToResolve.Num() > 0)
+			{
+				const FScopedPin PinToResolve = PinsToResolve.Pop(false);
 
-					if (const FDependencyChain* Dependencies = DependencyData.Find(PinToResolve))
+				if (const FDependencyChain* Dependencies = DependencyData.Find(PinToResolve))
+				{
+					for (const FScopedPin& ReadPin : Dependencies->Pins)
 					{
-						for (const FModuleScopedPin& ReadPin : Dependencies->Pins)
+						bool AlreadyAdded = false;
+						ResolvedDependencies.Pins.Add(ReadPin, &AlreadyAdded);
+
+						if (!AlreadyAdded)
 						{
-							bool AlreadyAdded = false;
-							ResolvedDependencies.Pins.Add(ReadPin, &AlreadyAdded);
+							int32 SourceVariableIndex = INDEX_NONE;
+							const FScopedPin WritePin = FindWriteForRead(ParamMap, ReadPin, &SourceVariableIndex);
 
-							if (!AlreadyAdded)
+							if (WritePin.Pin)
 							{
-								int32 SourceVariableIndex = INDEX_NONE;
-								const FModuleScopedPin WritePin = FindWriteForRead(ParamMap, ReadPin, &SourceVariableIndex);
+								PinsToResolve.Add(WritePin);
+							}
+							else if (SourceVariableIndex != INDEX_NONE)
+							{
+								const int32 DefaultBoundVariable = FindDefaultBinding(ParamMap, SourceVariableIndex, ReadPin);
 
-								if (WritePin.Pin)
+								// if we are bound to a variable we need to grab the last relevant write to that variable and continue our resolution
+								if (DefaultBoundVariable != INDEX_NONE)
 								{
-									PinsToResolve.Add(WritePin);
-								}
-								else if (SourceVariableIndex != INDEX_NONE)
-								{
-									const int32 DefaultBoundVariable = FindDefaultBinding(ParamMap, SourceVariableIndex, ReadPin);
+									bool bValidDefaultWriteFound = false;
 
-									// if we are bound to a variable we need to grab the last relevant write to that variable and continue our resolution
-									if (DefaultBoundVariable != INDEX_NONE)
+									const FNode* ReadPinOwningNode = GraphBridge::GetOwningNode(ReadPin.Pin);
+									const int32 ReadNodeVisitationIndex = ParamMap.MapNodeVisitations.IndexOfByKey(ReadPinOwningNode);
+									if (ReadNodeVisitationIndex != INDEX_NONE)
 									{
-										bool bValidDefaultWriteFound = false;
-
-										const UNiagaraNode* ReadPinOwningNode = Cast<const UNiagaraNode>(ReadPin.Pin->GetOwningNode());
-										const int32 ReadNodeVisitationIndex = ParamMap.MapNodeVisitations.IndexOfByKey(ReadPinOwningNode);
-										if (ReadNodeVisitationIndex != INDEX_NONE)
+										const int32 VariableWriteCount = ParamMap.PerVariableWriteHistory[DefaultBoundVariable].Num();
+										for (int32 VariableWriteIt = VariableWriteCount - 1; VariableWriteIt >= 0; --VariableWriteIt)
 										{
-											const int32 VariableWriteCount = ParamMap.PerVariableWriteHistory[DefaultBoundVariable].Num();
-											for (int32 VariableWriteIt = VariableWriteCount - 1; VariableWriteIt >= 0; --VariableWriteIt)
+											const FScopedPin& DefaultWritePin = ParamMap.PerVariableWriteHistory[DefaultBoundVariable][VariableWriteIt];
+											const FNode* DefaultWritePinOwningNode = GraphBridge::GetOwningNode(DefaultWritePin.Pin);
+											const int32 DefaultWriteNodeVisitationIndex = ParamMap.MapNodeVisitations.IndexOfByKey(DefaultWritePinOwningNode);
+											if (DefaultWriteNodeVisitationIndex != INDEX_NONE && DefaultWriteNodeVisitationIndex < ReadNodeVisitationIndex)
 											{
-												const FModuleScopedPin& DefaultWritePin = ParamMap.PerVariableWriteHistory[DefaultBoundVariable][VariableWriteIt];
-												const UNiagaraNode* DefaultWritePinOwningNode = Cast<const UNiagaraNode>(DefaultWritePin.Pin->GetOwningNode());
-												const int32 DefaultWriteNodeVisitationIndex = ParamMap.MapNodeVisitations.IndexOfByKey(DefaultWritePinOwningNode);
-												if (DefaultWriteNodeVisitationIndex != INDEX_NONE && DefaultWriteNodeVisitationIndex < ReadNodeVisitationIndex)
-												{
-													PinsToResolve.Add(DefaultWritePin);
-													bValidDefaultWriteFound = true;
-													break;
-												}
+												PinsToResolve.Add(DefaultWritePin);
+												bValidDefaultWriteFound = true;
+												break;
 											}
 										}
+									}
 
-										// if no valid pin was found that writes into the default bound variable within this parameter map, then we assume
-										// that it will be present in a previous parameter map, and so we just record the earliest read from this parameter
-										// map as a dependency
+									// if no valid pin was found that writes into the default bound variable within this parameter map, then we assume
+									// that it will be present in a previous parameter map, and so we just record the earliest read from this parameter
+									// map as a dependency
 
-										if (!ParamMap.PerVariableReadHistory[DefaultBoundVariable].IsEmpty())
+									if (!ParamMap.PerVariableReadHistory[DefaultBoundVariable].IsEmpty())
+									{
+										// setting the following to false will allow the trimmer to remove additional attributes.  In the cases
+										// where an attribute is written to and used only in a single stage the value could be considered local and not
+										// need to be an attribute.  If bBackCompat_AlwaysIncludeFirstReadPin is true then any time we read from that
+										// variable we'll treat it as a dependent.  If it's false then the writes that contribute to that variable will 
+										// be resolved as part of the dependency resolution.
+										constexpr bool bBackCompat_AlwaysIncludeFirstReadPin = true;
+
+										if (!bValidDefaultWriteFound || bBackCompat_AlwaysIncludeFirstReadPin)
 										{
-											// setting the following to false will allow the trimmer to remove additional attributes.  In the cases
-											// where an attribute is written to and used only in a single stage the value could be considered local and not
-											// need to be an attribute.  If bBackCompat_AlwaysIncludeFirstReadPin is true then any time we read from that
-											// variable we'll treat it as a dependent.  If it's false then the writes that contribute to that variable will 
-											// be resolved as part of the dependency resolution.
-											constexpr bool bBackCompat_AlwaysIncludeFirstReadPin = true;
-
-											if (!bValidDefaultWriteFound || bBackCompat_AlwaysIncludeFirstReadPin)
-											{
-												ResolvedDependencies.Pins.Add(ParamMap.PerVariableReadHistory[DefaultBoundVariable][0].ReadPin);
-											}
+											ResolvedDependencies.Pins.Add(ParamMap.PerVariableReadHistory[DefaultBoundVariable][0].ReadPin);
 										}
 									}
 								}
 							}
 						}
-
-						ResolvedDependencies.CustomNodes.Append(Dependencies->CustomNodes);
 					}
+
+					ResolvedDependencies.CustomNodes.Append(Dependencies->CustomNodes);
 				}
 			}
 		}
 	}
 }
 
-static void TrimAttributes_Safe(TConstArrayView<const FNiagaraParameterMapHistory*> LocalParamHistories, TSet<FName>& AttributesToPreserve, TArray<FNiagaraVariable>& Attributes)
+template<typename GraphBridge>
+void FNiagaraAttributeTrimmerHelper<GraphBridge>::TrimAttributes_Safe(TConstArrayView<const FParamMapHistory*> LocalParamHistories, TSet<FName>& AttributesToPreserve, TArray<FNiagaraVariable>& Attributes)
 {
 	// go through the ParamMapHistories and collect any CustomHlsl nodes so that we can give a best effort to
 	// find references to our variable.  If a reference is found we'll assume that we can't trim the attribute
-	TArray<const UNiagaraNodeCustomHlsl*, TInlineAllocator<8>> CustomHlslNodes;
-	for (const FNiagaraParameterMapHistory* ParamMap : LocalParamHistories)
+	TArray<const FCustomHlslNode*, TInlineAllocator<8>> CustomHlslNodes;
+	for (const FParamMapHistory* ParamMap : LocalParamHistories)
 	{
-		for (const UEdGraphPin* Pin : ParamMap->MapPinHistory)
+		for (const FPin* Pin : ParamMap->MapPinHistory)
 		{
-			if (const UNiagaraNodeCustomHlsl* CustomHlslNode = Cast<const UNiagaraNodeCustomHlsl>(Pin->GetOwningNode()))
+			if (const FCustomHlslNode* CustomHlslNode = GraphBridge::template AsNodeType<FCustomHlslNode>(GraphBridge::GetOwningNode(Pin)))
 			{
 				CustomHlslNodes.AddUnique(CustomHlslNode);
 			}
@@ -609,7 +571,7 @@ static void TrimAttributes_Safe(TConstArrayView<const FNiagaraParameterMapHistor
 	Attributes.SetNum(Algo::StableRemoveIf(Attributes, [=](const FNiagaraVariable& Var)
 	{
 		// preserve attributes which have a record of being read
-		for (const FNiagaraParameterMapHistory* ParamMap : LocalParamHistories)
+		for (const FParamMapHistory* ParamMap : LocalParamHistories)
 		{
 			const int32 VarIdx = ParamMap->FindVariable(Var.GetName(), Var.GetType());
 			if (VarIdx != INDEX_NONE)
@@ -627,9 +589,11 @@ static void TrimAttributes_Safe(TConstArrayView<const FNiagaraParameterMapHistor
 			return false;
 		}
 
-		for (const UNiagaraNodeCustomHlsl* CustomHlslNode : CustomHlslNodes)
+		FNameBuilder VariableName(Var.GetName());
+
+		for (const FCustomHlslNode* CustomHlslNode : CustomHlslNodes)
 		{
-			if (CustomHlslNode->ReferencesVariable(Var))
+			if (GraphBridge::CustomHlslReferencesTokens(CustomHlslNode, { VariableName.ToView() }))
 			{
 				return false;
 			}
@@ -639,13 +603,14 @@ static void TrimAttributes_Safe(TConstArrayView<const FNiagaraParameterMapHistor
 	}));
 }
 
-static void TrimAttributes_Aggressive(const FNiagaraCompileRequestDuplicateData* CompileDuplicateData, TConstArrayView<const FNiagaraParameterMapHistory*> LocalParamHistories, TSet<FName>& AttributesToPreserve, TArray<FNiagaraVariable>& Attributes)
+template<typename GraphBridge>
+void FNiagaraAttributeTrimmerHelper<GraphBridge>::TrimAttributes_Aggressive(const FCompilationCopy* CompileDuplicateData, TConstArrayView<const FParamMapHistory*> LocalParamHistories, TSet<FName>& AttributesToPreserve, TArray<FNiagaraVariable>& Attributes)
 {
 	// variable references hidden in custom hlsl nodes may not be present in a specific stages parameter map
 	// so we consolidate all the variables into one list to use when going through custom hlsl nodes
 	TSet<FNiagaraVariableBase> UnifiedVariables;
 
-	for (const FNiagaraParameterMapHistory* ParamMap : LocalParamHistories)
+	for (const FParamMapHistory* ParamMap : LocalParamHistories)
 	{
 		for (const FNiagaraVariableBase& Variable : ParamMap->Variables)
 		{
@@ -653,15 +618,11 @@ static void TrimAttributes_Aggressive(const FNiagaraCompileRequestDuplicateData*
 		}
 	}
 
-	using namespace NiagaraAttributeTrimming;
-
 	FDependencyMap PerVariableDependencySets;
 
-	const UEdGraphSchema_Niagara* NiagaraSchema = GetDefault<UEdGraphSchema_Niagara>();
-
-	for (const FNiagaraParameterMapHistory* ParamMap : LocalParamHistories)
+	for (const FParamMapHistory* ParamMap : LocalParamHistories)
 	{
-		FFunctionInputResolver InputResolver(CompileDuplicateData->NodeGraphDeepCopy.Get(), ParamMap->OriginatingScriptUsage);
+		FFunctionInputResolver InputResolver(GraphBridge::GetGraph(CompileDuplicateData), ParamMap->OriginatingScriptUsage);
 
 		const int32 VariableCount = ParamMap->Variables.Num();
 
@@ -669,14 +630,14 @@ static void TrimAttributes_Aggressive(const FNiagaraCompileRequestDuplicateData*
 		{
 			const FNiagaraVariableBase& Var = ParamMap->Variables[VariableIt];
 
-			for (const FModuleScopedPin& WritePin : ParamMap->PerVariableWriteHistory[VariableIt])
+			for (const FScopedPin& WritePin : ParamMap->PerVariableWriteHistory[VariableIt])
 			{
 				FDependencyChain& Dependencies = PerVariableDependencySets.Add(WritePin);
 
-				FExpressionBuilder Builder(*ParamMap, InputResolver, NiagaraSchema);
+				FExpressionBuilder Builder(*ParamMap, InputResolver);
 				Builder.FindDependencies(WritePin, Dependencies);
 
-				for (FCustomHlslNodeMap::TConstIterator CustomNodeIt = Dependencies.CustomNodes.CreateConstIterator(); CustomNodeIt; ++CustomNodeIt)
+				for (typename FCustomHlslNodeMap::TConstIterator CustomNodeIt = Dependencies.CustomNodes.CreateConstIterator(); CustomNodeIt; ++CustomNodeIt)
 				{
 					if (CustomNodeIt.Value().bHasImpureFunctionText)
 					{
@@ -686,21 +647,21 @@ static void TrimAttributes_Aggressive(const FNiagaraCompileRequestDuplicateData*
 			}
 		}
 
-		FImpureFunctionParser ImpureFunctionParser(CompileDuplicateData->NodeGraphDeepCopy.Get(), ParamMap->OriginatingScriptUsage);
-		for (const FModuleScopedPin& RequiredFunctionInput : ImpureFunctionParser.ReadFunctionInputs())
+		FImpureFunctionParser ImpureFunctionParser(GraphBridge::GetGraph(CompileDuplicateData), ParamMap->OriginatingScriptUsage);
+		for (const FScopedPin& RequiredFunctionInput : ImpureFunctionParser.ReadFunctionInputs())
 		{
 			FDependencyChain Dependencies;
 
-			FExpressionBuilder Builder(*ParamMap, InputResolver, NiagaraSchema);
+			FExpressionBuilder Builder(*ParamMap, InputResolver);
 			Builder.FindDependencies(RequiredFunctionInput, Dependencies);
 
-			for (const FModuleScopedPin& DependentPin : Dependencies.Pins)
+			for (const FScopedPin& DependentPin : Dependencies.Pins)
 			{
 				// see if this variable corresponds to something in our parameter map's list of variables
 				// if it does then we need to mark it as something that needs to be preserved
 				// Note that we need to resolve the Pin name's module namespace before we actually search for it
 				// as a variable
-				FString VariableNameString = DependentPin.Pin->GetName();
+				FString VariableNameString = DependentPin.Pin->PinName.ToString();
 				VariableNameString.ReplaceInline(*DependentPin.ModuleName.ToString(), *FNiagaraConstants::ModuleNamespace.ToString());
 
 				const FName VariableName = *VariableNameString;
@@ -714,9 +675,9 @@ static void TrimAttributes_Aggressive(const FNiagaraCompileRequestDuplicateData*
 					// we also need to check to see if the attribute has resolved namespaces (i.e. Particles.Module.MyAttribute or StackContext.MyAttribute)
 					// for which we'll just search through the parameter maps variables directly
 					const int32 AliasedIndex = ParamMap->VariablesWithOriginalAliasesIntact.IndexOfByPredicate([&](const FNiagaraVariable& Variable)
-					{
-						return Variable.GetName() == VariableName;
-					});
+						{
+							return Variable.GetName() == VariableName;
+						});
 
 					if (ParamMap->Variables.IsValidIndex(AliasedIndex))
 					{
@@ -743,12 +704,12 @@ static void TrimAttributes_Aggressive(const FNiagaraCompileRequestDuplicateData*
 	{
 		const FName AttributeName = SearchList.Pop(false);
 
-		for (const FNiagaraParameterMapHistory* ParamMap : LocalParamHistories)
+		for (const FParamMapHistory* ParamMap : LocalParamHistories)
 		{
 			FDependencyChain ResolvedChain;
 			ResolveDependencyChain(*ParamMap, PerVariableDependencySets, AttributeName, ResolvedChain);
 
-			for (const FModuleScopedPin ParameterRead : ResolvedChain.Pins)
+			for (const FScopedPin& ParameterRead : ResolvedChain.Pins)
 			{
 				const FName DependentAttributeName = FindAttributeForRead(*ParamMap, ParameterRead);
 
@@ -763,15 +724,15 @@ static void TrimAttributes_Aggressive(const FNiagaraCompileRequestDuplicateData*
 					}
 
 					// if we're dependent on a Previous value, then make sure we add the attribute to the list of things to be preserved as well
-					if (FNiagaraParameterMapHistory::IsPreviousName(DependentAttributeName))
+					if (FParamMapHistory::IsPreviousName(DependentAttributeName))
 					{
-						ConditionalAddAttribute(FNiagaraParameterMapHistory::GetSourceForPreviousValue(DependentAttributeName));
+						ConditionalAddAttribute(FParamMapHistory::GetSourceForPreviousValue(DependentAttributeName));
 					}
 
 					// likewise for Initial value
-					if (FNiagaraParameterMapHistory::IsInitialName(DependentAttributeName))
+					if (FParamMapHistory::IsInitialName(DependentAttributeName))
 					{
-						ConditionalAddAttribute(FNiagaraParameterMapHistory::GetSourceForInitialValue(DependentAttributeName));
+						ConditionalAddAttribute(FParamMapHistory::GetSourceForInitialValue(DependentAttributeName));
 					}
 				}
 			}
@@ -787,7 +748,8 @@ static void TrimAttributes_Aggressive(const FNiagaraCompileRequestDuplicateData*
 						continue;
 					}
 
-					if (CustomHlslInfo.Key->ReferencesVariable(Variable))
+					FNameBuilder VariableName(Variable.GetName());
+					if (GraphBridge::CustomHlslReferencesTokens(CustomHlslInfo.Key, { VariableName.ToView() }))
 					{
 						ConditionalAddAttribute(Variable.GetName());
 					}
@@ -819,100 +781,5 @@ static void TrimAttributes_Aggressive(const FNiagaraCompileRequestDuplicateData*
 	}
 }
 
-void FHlslNiagaraTranslator::TrimAttributes(const FNiagaraCompileOptions& InCompileOptions, TArray<FNiagaraVariable>& Attributes)
-{
-	if (!UNiagaraScript::IsParticleScript(InCompileOptions.TargetUsage))
-	{
-		return;
-	}
-
-	const bool SafeTrimAttributesEnabled = InCompileOptions.AdditionalDefines.Contains(TEXT("TrimAttributesSafe"));
-	const bool AggressiveTrimAttributesEnabled = InCompileOptions.AdditionalDefines.Contains(TEXT("TrimAttributes"));
-
-	if (SafeTrimAttributesEnabled || AggressiveTrimAttributesEnabled)
-	{
-		// validate that the attributes have unique sanitized names
-		{
-			bool HasOverlappingNames = false;
-
-			TMap<FString, FNiagaraVariableBase> SanitizedNames;
-			for (const FNiagaraVariableBase& Attribute : Attributes)
-			{
-				const FString SanitizedName = GetSanitizedSymbolName(Attribute.GetName().ToString());
-				if (const FNiagaraVariableBase* ExistingVariable = SanitizedNames.Find(SanitizedName))
-				{
-					HasOverlappingNames = true;
-				}
-				else
-				{
-					SanitizedNames.Add(SanitizedName, Attribute);
-				}
-			}
+template class FNiagaraAttributeTrimmerHelper<FNiagaraCompilationGraphBridge>;
 
-			// the trimming algorithm doesn't work when names are overlapping, so just early out of the function
-			if (HasOverlappingNames)
-			{
-				return;
-			}
-		}
-
-		const bool bRequiresPersistentIDs = InCompileOptions.AdditionalDefines.Contains(TEXT("RequiresPersistentIDs"));
-
-		// we want to use the ParamMapHistories of both the particle update and spawn scripts because they need to
-		// agree to define a unified attribute set
-		TArray<const FNiagaraParameterMapHistory*, TInlineAllocator<2>> LocalParamHistories;
-		for (const FNiagaraParameterMapHistory& History : OtherOutputParamMapHistories)
-		{
-			if (UNiagaraScript::IsParticleScript(History.OriginatingScriptUsage))
-			{
-				// for now we'll be disabling attribute trimming if a family of particle scripts contain generation of
-				// additional dataset writes (events) as we don't have access to the connectivity of it's variables as
-				// we do for the rest of the script
-				if (History.AdditionalDataSetWrites.Num())
-				{
-					return;
-				}
-
-				LocalParamHistories.Add(&History);
-			}
-		}
-
-		// check through the AdditionalDefines to see if any variables have been explicitly preserved
-		TSet<FName> AttributesToPreserve;
-
-		for (const FString& AdditionalDefine : InCompileOptions.AdditionalDefines)
-		{
-			const FString PreserveTag = TEXT("PreserveAttribute=");
-			if (AdditionalDefine.StartsWith(PreserveTag))
-			{
-				AttributesToPreserve.Add(*AdditionalDefine.RightChop(PreserveTag.Len()));
-			}
-		}
-
-		AttributesToPreserve.Add(SYS_PARAM_INSTANCE_ALIVE.GetName());
-		AttributesToPreserve.Add(SYS_PARAM_PARTICLES_UNIQUE_ID.GetName());
-		if (bRequiresPersistentIDs)
-		{
-			AttributesToPreserve.Add(SYS_PARAM_PARTICLES_ID.GetName());
-		}
-
-		const TArray<FNiagaraVariable> PreTrimmedAttributes = Attributes;
-
-		if (SafeTrimAttributesEnabled)
-		{
-			TrimAttributes_Safe(LocalParamHistories, AttributesToPreserve, Attributes);
-		}
-		else if (AggressiveTrimAttributesEnabled)
-		{
-			TrimAttributes_Aggressive(CompileDuplicateData, LocalParamHistories, AttributesToPreserve, Attributes);
-		}
-
-		for (const FNiagaraVariable& Attribute : PreTrimmedAttributes)
-		{
-			if (!Attributes.Contains(Attribute))
-			{
-				TranslateResults.CompileTags.Emplace(Attribute, TEXT("Trimmed"));
-			}
-		}
-	}
-}
\ No newline at end of file
diff --git a/Engine/Plugins/FX/Niagara/Source/NiagaraEditor/Private/NiagaraCompilationGraphBridge.h b/Engine/Plugins/FX/Niagara/Source/NiagaraEditor/Private/NiagaraCompilationGraphBridge.h
new file mode 100644
--- /dev/null
+++ b/Engine/Plugins/FX/Niagara/Source/NiagaraEditor/Private/NiagaraCompilationGraphBridge.h
@@ -0,0 +1,463 @@
+// Copyright Epic Games, Inc. All Rights Reserved.
+
+#pragma once
+
+#include "EdGraphSchema_Niagara.h"
+#include "NiagaraCompilationPrivate.h"
+#include "NiagaraConstants.h"
+#include "NiagaraModule.h"
+#include "NiagaraNodeConvert.h"
+#include "NiagaraNodeCustomHlsl.h"
+#include "NiagaraNodeEmitter.h"
+#include "NiagaraNodeFunctionCall.h"
+#include "NiagaraNodeIf.h"
+#include "NiagaraNodeInput.h"
+#include "NiagaraNodeOp.h"
+#include "NiagaraNodeOutput.h"
+#include "NiagaraNodeParameterMapFor.h"
+#include "NiagaraNodeParameterMapGet.h"
+#include "NiagaraNodeParameterMapSet.h"
+#include "NiagaraNodeSelect.h"
+#include "NiagaraScriptVariable.h"
+
+// Defines an interface for UNiagaraGraph, their nodes & pins and all related support structures.
+// Intended as a temporary stopgap measure to allow us to reuse code like the hlsl translator and
+// parameter map history builders that can traverse different ways of representing the niagara graphs
+struct FNiagaraCompilationGraphBridge
+{
+	using FGraph = UNiagaraGraph;
+	using FPin = UEdGraphPin;
+	using FScopedPin = FModuleScopedPin;
+	using FCustomHlslNode = UNiagaraNodeCustomHlsl;
+	using FNode = UNiagaraNode;
+	using FOutputNode = UNiagaraNodeOutput;
+	using FInputNode = UNiagaraNodeInput;
+	using FOpNode = UNiagaraNodeOp;
+	using FEmitterNode = UNiagaraNodeEmitter;
+	using FIfNode = UNiagaraNodeIf;
+	using FConvertNode = UNiagaraNodeConvert;
+	using FSelectNode = UNiagaraNodeSelect;
+	using FFunctionCallNode = UNiagaraNodeFunctionCall;
+	using FParamMapGetNode = UNiagaraNodeParameterMapGet;
+	using FParamMapSetNode = UNiagaraNodeParameterMapSet;
+	using FParamMapForNode = UNiagaraNodeParameterMapFor;
+	using FStaticSwitchNode = UNiagaraNodeStaticSwitch;
+	using FInputPin = UEdGraphPin;
+	using FOutputPin = UEdGraphPin;
+	using FParamMapHistory = FNiagaraParameterMapHistory;
+	using FParamMapHistoryBuilder = FNiagaraParameterMapHistoryBuilder;
+	using FPrecompileData = FNiagaraCompileRequestData;
+	using FCompilationCopy = FNiagaraCompileRequestDuplicateData;
+	using FModuleScopedPin = FModuleScopedPin;
+	using FGraphTraversalHandle = FGraphTraversalHandle;
+	using FGraphFunctionAliasContext = FNiagaraGraphFunctionAliasContext;
+	using FConvertConnection = FNiagaraConvertConnection;
+
+	static const FGraph* GetGraph(const FCompilationCopy* CompilationCopy)
+	{
+		return CompilationCopy->NodeGraphDeepCopy.Get();
+	}
+
+	static const FNode* GetOwningNode(const FPin* Pin)
+	{
+		return Cast<const FNode>(Pin->GetOwningNode());
+	}
+
+	static FNode* GetMutableOwningNode(const FPin* Pin)
+	{
+		return Cast<FNode>(Pin->GetOwningNode());
+	}
+
+	static const FGraph* GetOwningGraph(const FNode* Node)
+	{
+		if (const UNiagaraNode* NiagaraNode = Cast<const UNiagaraNode>(Node))
+		{
+			return NiagaraNode->GetNiagaraGraph();
+		}
+		return nullptr;
+	}
+
+	static bool CustomHlslReferencesTokens(const FCustomHlslNode* CustomNode, TConstArrayView<FStringView> TokenStrings)
+	{
+		TArray<FString> StringTokens;
+		UNiagaraNodeCustomHlsl::GetTokensFromString(CustomNode->GetCustomHlsl(), StringTokens, false, false);
+
+		for (FStringView Token : TokenStrings)
+		{
+			const bool TokenMatched = StringTokens.ContainsByPredicate([&Token](const FString& HlslToken) -> bool
+			{
+				return HlslToken.Contains(Token);
+			});
+
+			if (TokenMatched)
+			{
+				return true;
+			}
+		}
+
+		return false;
+	}
+
+	static ENiagaraScriptUsage GetCustomHlslUsage(const FCustomHlslNode* CustomNode)
+	{
+		return CustomNode->ScriptUsage;
+	}
+
+	static FString GetCustomHlslString(const FCustomHlslNode* CustomNode)
+	{
+		return CustomNode->GetCustomHlsl();
+	}
+
+	static TArray<FNiagaraCustomHlslInclude> GetCustomHlslIncludePaths(const FCustomHlslNode* CustomNode)
+	{
+		TArray<FNiagaraCustomHlslInclude> Includes;
+		CustomNode->GetIncludeFilePaths(Includes);
+		return Includes;
+	}
+
+	static const TArray<FConvertConnection>& GetConvertConnections(const FConvertNode* ConvertNode)
+	{
+		return ConvertNode->GetConnections();
+	}
+
+	template<typename NodeType>
+	static const NodeType* AsNodeType(const FNode* Node)
+	{
+		return Cast<const NodeType>(Node);
+	}
+
+	static bool GraphHasParametersOfType(const FGraph* Graph, const FNiagaraTypeDefinition& TypeDef)
+	{
+		TArray<FNiagaraVariable> Inputs;
+		TArray<FNiagaraVariable> Outputs;
+		Graph->GetParameters(Inputs, Outputs);
+
+		auto ContainsVariableOfType = [TypeDef](const FNiagaraVariable& Variable) -> bool
+		{
+			return Variable.GetType() == TypeDef;
+		};
+
+		return Inputs.ContainsByPredicate(ContainsVariableOfType)
+			|| Outputs.ContainsByPredicate(ContainsVariableOfType);
+	}
+
+	static TArray<FNiagaraVariableBase> GraphGetStaticSwitchInputs(const FGraph* Graph)
+	{
+		TArray<FNiagaraVariableBase> StaticSwitchInputs;
+		const TArray<FNiagaraVariable>& Variables = Graph->FindStaticSwitchInputs();
+
+		StaticSwitchInputs.Reserve(Variables.Num());
+		Algo::Transform(Variables, StaticSwitchInputs, [](const FNiagaraVariable& InVar) -> FNiagaraVariableBase
+		{
+			return InVar;
+		});
+
+		return StaticSwitchInputs;
+	}
+
+	static void FindOutputNodes(const FGraph* Graph, ENiagaraScriptUsage ScriptUsage, TArray<const FOutputNode*>& OutputNodes)
+	{
+		TArray<FOutputNode*> MutableOutputNodes;
+		Graph->FindOutputNodes(ScriptUsage, MutableOutputNodes);
+		OutputNodes.Reserve(MutableOutputNodes.Num());
+		Algo::Transform(MutableOutputNodes, OutputNodes, [](FOutputNode* OutputNode) -> const FOutputNode*
+		{
+			return OutputNode;
+		});
+	}
+
+	static void FindOutputNodes(const FGraph* Graph, TArray<const FOutputNode*>& OutputNodes)
+	{
+		TArray<FOutputNode*> MutableOutputNodes;
+		Graph->FindOutputNodes(MutableOutputNodes);
+		OutputNodes.Reserve(MutableOutputNodes.Num());
+		Algo::Transform(MutableOutputNodes, OutputNodes, [](FOutputNode* OutputNode) -> const FOutputNode*
+		{
+			return OutputNode;
+		});
+	}
+
+	static void BuildTraversal(const FGraph* Graph, const FNode* OutputNode, TArray<const FNode*>& TraversedNodes)
+	{
+		TArray<UNiagaraNode*> MutableTraversedNodes;
+		Graph->BuildTraversal(MutableTraversedNodes, const_cast<FNode*>(OutputNode), true);
+		TraversedNodes.Reserve(MutableTraversedNodes.Num());
+		Algo::Transform(MutableTraversedNodes, TraversedNodes, [](UNiagaraNode* Node) -> const FNode*
+		{
+			return Node;
+		});
+	}
+
+	static const FGraph* GetEmitterGraph(const FEmitterNode* EmitterNode)
+	{
+		return EmitterNode->GetCalledGraph();
+	}
+
+	static FString GetEmitterUniqueName(const FEmitterNode* EmitterNode)
+	{
+		return EmitterNode->GetEmitterUniqueName();
+	}
+
+	static ENiagaraScriptUsage GetEmitterUsage(const FEmitterNode* EmitterNode)
+	{
+		return EmitterNode->GetUsage();
+	}
+
+	static FString GetEmitterName(const FEmitterNode* EmitterNode)
+	{
+		return EmitterNode->GetName();
+	}
+
+	static FString GetEmitterPathName(const FEmitterNode* EmitterNode)
+	{
+		return EmitterNode->GetFullName();
+	}
+
+	static FString GetEmitterHandleIdString(const FEmitterNode* EmitterNode)
+	{
+		return EmitterNode->GetEmitterHandleId().ToString(EGuidFormats::Digits);
+	}
+
+	static const FGraph* GetFunctionNodeGraph(const FFunctionCallNode* FunctionCall)
+	{
+		return FunctionCall->GetCalledGraph();
+	}
+
+	static FString GetFunctionFullName(const FFunctionCallNode* FunctionCall)
+	{
+		return FunctionCall->FunctionScript ? FunctionCall->FunctionScript->GetFullName() : FString();
+	}
+
+	static FString GetFunctionScriptName(const FFunctionCallNode* FunctionCall)
+	{
+		return FunctionCall->FunctionScript ? FunctionCall->FunctionScript->GetName() : FString();
+	}
+
+	static FString GetFunctionName(const FFunctionCallNode* FunctionCall)
+	{
+		return FunctionCall ? FunctionCall->GetFunctionName() : FString();
+	}
+
+	static ENiagaraScriptUsage GetFunctionUsage(const FFunctionCallNode* FunctionCall)
+	{
+		return FunctionCall->GetCalledUsage();
+	}
+
+	static TOptional<ENiagaraDefaultMode> GetGraphDefaultMode(const FGraph* Graph, const FNiagaraVariable& Variable, FNiagaraScriptVariableBinding& Binding)
+	{
+		return Graph->GetDefaultMode(Variable, &Binding);
+	}
+
+	static const FInputPin* GetDefaultPin(const FParamMapGetNode* GetNode, const FOutputPin* OutputPin)
+	{
+		return GetNode->GetDefaultPin(const_cast<FOutputPin*>(OutputPin));
+	}
+
+	template<typename ArrayType>
+	static void AppendCompilationPins(const FNode* Node, EEdGraphPinDirection PinDirection, ArrayType& OutPins)
+	{
+		const UNiagaraNodeWithDynamicPins* DynNode = Cast<const UNiagaraNodeWithDynamicPins>(Node);
+		for (FPin* Pin : Node->Pins)
+		{
+			if (Pin->Direction == PinDirection)
+			{
+				if (Pin->bOrphanedPin)
+				{
+					continue;
+				}
+
+				if (DynNode != nullptr && DynNode->IsAddPin(Pin))
+				{
+					continue;
+				}
+
+				OutPins.Add(Pin);
+			}
+		}
+	}
+
+	// retrieves all input pins (excluding any add pins that may be present)
+	static TArray<const FInputPin*> GetInputPins(const FNode* Node)
+	{
+		TArray<const FInputPin*> InputPins;
+		AppendCompilationPins(Node, EEdGraphPinDirection::EGPD_Input, InputPins);
+		return InputPins;
+	}
+
+	// retrieves all output pins (excluding both orphaned pins and add pins)
+	static TArray<const FOutputPin*> GetOutputPins(const FNode* Node)
+	{
+		TArray<const FInputPin*> OutputPins;
+		AppendCompilationPins(Node, EEdGraphPinDirection::EGPD_Output, OutputPins);
+		return OutputPins;
+	}
+
+	// gets all pins assoicated with the node
+	static TArray<const FPin*> GetPins(const FNode* Node)
+	{
+		TArray<const FPin*> Pins;
+		Pins.Reserve(Node->Pins.Num());
+		Algo::Transform(Node->Pins, Pins, [](FPin* InPin) -> const FPin*
+		{
+			return InPin;
+		});
+		return Pins;
+	}
+
+	static FNiagaraTypeDefinition GetPinType(const FPin* Pin, ENiagaraStructConversion Conversion)
+	{
+		return UEdGraphSchema_Niagara::PinToTypeDefinition(Pin, Conversion);
+	}
+
+	static FText GetPinFriendlyName(const FPin* Pin)
+	{
+		return Pin->PinFriendlyName;
+	}
+
+	static FText GetPinDisplayName(const FPin* Pin)
+	{
+		return Pin->GetDisplayName();
+	}
+
+	static FNiagaraVariable GetPinVariable(const FPin* Pin, bool bNeedsValue, ENiagaraStructConversion Conversion)
+	{
+		return UEdGraphSchema_Niagara::PinToNiagaraVariable(Pin, bNeedsValue, Conversion);
+	}
+
+	static const FInputPin* GetPinAsInput(const FPin* Pin)
+	{
+		if (Pin && Pin->Direction == EGPD_Input)
+		{
+			return Pin;
+		}
+		return nullptr;
+	}
+
+	static FNiagaraVariable GetInputVariable(const FInputNode* InputNode)
+	{
+		return InputNode->Input;
+	}
+
+	static const TArray<FNiagaraVariable>& GetOutputVariables(const FOutputNode* OutputNode)
+	{
+		return OutputNode->GetOutputs();
+	}
+
+	static TArray<FNiagaraVariable> GetGraphOutputNodeVariables(const FGraph* Graph, ENiagaraScriptUsage Usage)
+	{
+		TArray<FNiagaraVariable> OutputNodeVariables;
+		Graph->GetOutputNodeVariables(Usage, OutputNodeVariables);
+		return OutputNodeVariables;
+	}
+
+	static TArray<const FInputNode*> GetGraphInputNodes(const FGraph* Graph, const UNiagaraGraph::FFindInputNodeOptions& Options)
+	{
+		TArray<FInputNode*> MutableInputNodes;
+		Graph->FindInputNodes(MutableInputNodes, Options);
+
+		TArray<const FInputNode*> InputNodes;
+		InputNodes.Reserve(MutableInputNodes.Num());
+
+		Algo::Transform(MutableInputNodes, InputNodes, [](FInputNode* InputNode) -> const FInputNode*
+		{
+			return InputNode;
+		});
+
+		return InputNodes;
+	}
+
+	static const FOutputPin* GetLinkedOutputPin(const FInputPin* InputPin)
+	{
+		if (InputPin && !InputPin->LinkedTo.IsEmpty())
+		{
+			return InputPin->LinkedTo[0];
+		}
+		return nullptr;
+	}
+
+	static FPinConnectionResponse CanCreateConnection(const FOutputPin* OutputPin, const FInputPin* InputPin)
+	{
+		return GetDefault<UEdGraphSchema_Niagara>()->CanCreateConnection(OutputPin, InputPin);
+	}
+
+	static FGuid GetOutputNodeUsageId(const FOutputNode* OutputNode)
+	{
+		return OutputNode->GetUsageId();
+	}
+
+	static ENiagaraScriptUsage GetOutputNodeScriptType(const FOutputNode* OutputNode)
+	{
+		return OutputNode->ScriptType;
+	}
+
+	static FGuid GetOutputNodeScriptTypeId(const FOutputNode* OutputNode)
+	{
+		return OutputNode->ScriptTypeId;
+	}
+
+	static bool IsGraphEmpty(const FGraph* Graph)
+	{
+		return Graph->IsEmpty();
+	}
+
+	static void AddCollectionPaths(const FParamMapHistory& History, TArray<FString>& Paths)
+	{
+		for (UNiagaraParameterCollection* Collection : History.ParameterCollections)
+		{
+			Paths.AddUnique(FSoftObjectPath(Collection).ToString());
+		}
+	}
+
+	static bool NodeIsEnabled(const FNode* Node)
+	{
+		return Node->IsNodeEnabled();
+	}
+
+	static TOptional<ENiagaraDefaultMode> GraphGetDefaultMode(const FGraph* Graph, const FNiagaraVariableBase& Variable, FNiagaraScriptVariableBinding& Binding)
+	{
+		return Graph->GetDefaultMode(Variable, &Binding);
+	}
+
+	static const FOutputPin* GetSelectOutputPin(const FSelectNode* SelectNode, const FNiagaraVariableBase& Variable)
+	{
+		return SelectNode->GetOutputPin(Variable);
+	}
+
+	static FString GetNodeName(const FNode* Node)
+	{
+		return Node->GetName();
+	}
+
+	static FString GetNodeTitle(const FNode* Node)
+	{
+		return Node->GetNodeTitle(ENodeTitleType::FullTitle).ToString();
+	}
+
+	static const FInputPin* GetInputPin(const FNode* Node, int32 PinIndex)
+	{
+		return Node->GetInputPin(PinIndex);
+	}
+
+	static int32 GetPinIndexById(TConstArrayView<const FPin*> Pins, const FGuid& PinId)
+	{
+		const int32 PinCount = Pins.Num();
+		for (int32 PinIt = 0; PinIt < PinCount; ++PinIt)
+		{
+			if (Pins[PinIt]->PinId == PinId)
+			{
+				return PinIt;
+			}
+		}
+		return INDEX_NONE;
+	}
+
+	static void GetCompilationOutputPins(const UNiagaraNode* Node, FPinCollectorArray& Pins)
+	{
+		AppendCompilationPins(Node, EEdGraphPinDirection::EGPD_Output, Pins);
+	}
+
+	static void GetCompilationInputPins(const UNiagaraNode* Node, FPinCollectorArray& Pins)
+	{
+		AppendCompilationPins(Node, EEdGraphPinDirection::EGPD_Input, Pins);
+	}
+};
diff --git a/Engine/Plugins/FX/Niagara/Source/NiagaraEditor/Private/NiagaraCompiler.cpp b/Engine/Plugins/FX/Niagara/Source/NiagaraEditor/Private/NiagaraCompiler.cpp
--- a/Engine/Plugins/FX/Niagara/Source/NiagaraEditor/Private/NiagaraCompiler.cpp
+++ b/Engine/Plugins/FX/Niagara/Source/NiagaraEditor/Private/NiagaraCompiler.cpp
@@ -7,6 +7,7 @@
 #include "Misc/FileHelper.h"
 #include "Misc/PathViews.h"
 #include "Modules/ModuleManager.h"
+#include "NiagaraCompilationPrivate.h"
 #include "NiagaraComponent.h"
 #include "NiagaraDataInterface.h"
 #include "NiagaraEditorModule.h"
@@ -172,7 +173,7 @@ void FNiagaraCompileRequestDuplicateData::DuplicateReferencedGraphsRecursive(UNi
 				{
 					UNiagaraDataInterface* DataInterface = InputNode->GetDataInterface();
 					bool bIsParameterMapDataInterface = false;
-					FName DIName = FHlslNiagaraTranslator::GetDataInterfaceName(InputNode->Input.GetName(), EmitterUniqueName, bIsParameterMapDataInterface);
+					FName DIName = FNiagaraHlslTranslator::GetDataInterfaceName(InputNode->Input.GetName(), EmitterUniqueName, bIsParameterMapDataInterface);
 					UNiagaraDataInterface* Dupe = PrecompileDuplicateObject<UNiagaraDataInterface>(DataInterface, GetTransientPackage());
 					SharedNameToDuplicatedDataInterfaceMap->Add(DIName, Dupe);
 				}
@@ -1430,7 +1431,7 @@ int32 FNiagaraEditorModule::CompileScript(const FNiagaraCompileRequestDataBase*
 	FNiagaraCompileResults Results;
 	TSharedPtr<FHlslNiagaraCompiler> Compiler = MakeShared<FHlslNiagaraCompiler>();
 	FNiagaraTranslateResults TranslateResults;
-	FHlslNiagaraTranslator Translator;
+	TUniquePtr<INiagaraHlslTranslator> Translator = INiagaraHlslTranslator::CreateTranslator(CompileRequest, CompileRequestDuplicate);
 
 	FHlslNiagaraTranslatorOptions TranslateOptions;
 
@@ -1451,17 +1452,17 @@ int32 FNiagaraEditorModule::CompileScript(const FNiagaraCompileRequestDataBase*
 	if (GbForceNiagaraTranslatorSingleThreaded > 0)
 	{
 		FScopeLock Lock(&TranslationCritSec);
-		TranslateResults = Translator.Translate(CompileRequest, CompileRequestDuplicate, InCompileOptions, TranslateOptions);
+		TranslateResults = Translator->Translate(InCompileOptions, TranslateOptions);
 	}
 	else
 	{
-		TranslateResults = Translator.Translate(CompileRequest, CompileRequestDuplicate, InCompileOptions, TranslateOptions);
+		TranslateResults = Translator->Translate(InCompileOptions, TranslateOptions);
 	}
 	UE_LOG(LogNiagaraEditor, Verbose, TEXT("Translating System %s took %f sec."), *InCompileOptions.FullName, (float)(FPlatformTime::Seconds() - TranslationStartTime));
 
 	if (GbForceNiagaraTranslatorDump != 0)
 	{
-		DumpHLSLText(Translator.GetTranslatedHLSL(), InCompileOptions.FullName);
+		DumpHLSLText(Translator->GetTranslatedHLSL(), InCompileOptions.FullName);
 		if (GbForceNiagaraVMBinaryDump != 0 && Results.Data.IsValid())
 		{
 			DumpHLSLText(Results.Data->LastAssemblyTranslation, InCompileOptions.FullName);
@@ -1471,10 +1472,9 @@ int32 FNiagaraEditorModule::CompileScript(const FNiagaraCompileRequestDataBase*
 	NiagaraCompileRequestHelper::FDebugGroupNameBuilder DebugGroupName;
 	NiagaraCompileRequestHelper::BuildScriptDebugGroupName(CompileRequest, InCompileOptions, DebugGroupName);
 
-	int32 JobID = Compiler->CompileScript(DebugGroupName, InCompileOptions, TranslateResults, &Translator.GetTranslateOutput(), Translator.GetTranslatedHLSL());
+	int32 JobID = Compiler->CompileScript(DebugGroupName, InCompileOptions, TranslateResults, Translator->GetTranslateOutput(), Translator->GetTranslatedHLSL());
 	ActiveCompilations.Add(JobID, Compiler);
 	return JobID;
-	
 }
 
 TSharedPtr<FNiagaraGraphCachedDataBase, ESPMode::ThreadSafe> FNiagaraEditorModule::CacheGraphTraversal(const UObject* Obj, FGuid Version)
@@ -1764,10 +1764,10 @@ int32 FHlslNiagaraCompiler::CompileScript(const FNiagaraCompileRequestData* InCo
 	NiagaraCompileRequestHelper::FDebugGroupNameBuilder DebugGroupName;
 	NiagaraCompileRequestHelper::BuildScriptDebugGroupName(InCompileRequest, InOptions, DebugGroupName);
 
-	return CompileScript(DebugGroupName, InOptions, InTranslateResults, TranslatorOutput, TranslatedHLSL);
+	return CompileScript(DebugGroupName, InOptions, InTranslateResults, TranslatorOutput ? *TranslatorOutput : FNiagaraTranslatorOutput(), TranslatedHLSL);
 }
 
-int32 FHlslNiagaraCompiler::CompileScript(const FStringView GroupName, const FNiagaraCompileOptions& InOptions, const FNiagaraTranslateResults& InTranslateResults, FNiagaraTranslatorOutput *TranslatorOutput, FString &TranslatedHLSL)
+int32 FHlslNiagaraCompiler::CompileScript(const FStringView GroupName, const FNiagaraCompileOptions& InOptions, const FNiagaraTranslateResults& InTranslateResults, const FNiagaraTranslatorOutput& TranslatorOutput, const FString& TranslatedHLSL)
 {
 	SCOPE_CYCLE_COUNTER(STAT_NiagaraEditor_HlslCompiler_CompileScript);
 
@@ -1807,7 +1807,7 @@ int32 FHlslNiagaraCompiler::CompileScript(const FStringView GroupName, const FNi
 
 	uint32 JobID = FShaderCommonCompileJob::GetNextJobId();
 	CompilationJob = MakeUnique<FNiagaraCompilerJob>();
-	CompilationJob->TranslatorOutput = TranslatorOutput ? *TranslatorOutput : FNiagaraTranslatorOutput();
+	CompilationJob->TranslatorOutput = TranslatorOutput;
 
 	CompileResults.bVMSucceeded = (CompilationJob->TranslatorOutput.Errors.Len() == 0) && (TranslatedHLSL.Len() > 0) && !InTranslateResults.NumErrors;
 
@@ -2083,7 +2083,7 @@ TOptional<FNiagaraCompileResults> FHlslNiagaraCompiler::GetCompileResult(int32 J
 			{
 				Sig = NDIInfo.RegisteredFunctions.FindByPredicate([&](const FNiagaraFunctionSignature& CheckSig)
 				{
-					FString SigSymbol = FHlslNiagaraTranslator::GetFunctionSignatureSymbol(CheckSig);
+					FString SigSymbol = FNiagaraHlslTranslator::GetFunctionSignatureSymbol(CheckSig);
 					return SigSymbol == FuncInfo.Name;
 				});
 				if (Sig)
@@ -2098,7 +2098,7 @@ TOptional<FNiagaraCompileResults> FHlslNiagaraCompiler::GetCompileResult(int32 J
 				Sig = UNiagaraFunctionLibrary::GetVectorVMFastPathOps(true).FindByPredicate(
 					[&](const FNiagaraFunctionSignature& CheckSig)
 				{
-					FString SigSymbol = FHlslNiagaraTranslator::GetFunctionSignatureSymbol(CheckSig);
+					FString SigSymbol = FNiagaraHlslTranslator::GetFunctionSignatureSymbol(CheckSig);
 					return SigSymbol == FuncInfo.Name;
 				}
 				);
diff --git a/Engine/Plugins/FX/Niagara/Source/NiagaraEditor/Private/NiagaraCompiler.h b/Engine/Plugins/FX/Niagara/Source/NiagaraEditor/Private/NiagaraCompiler.h
--- a/Engine/Plugins/FX/Niagara/Source/NiagaraEditor/Private/NiagaraCompiler.h
+++ b/Engine/Plugins/FX/Niagara/Source/NiagaraEditor/Private/NiagaraCompiler.h
@@ -44,7 +44,7 @@ public:
 	UE_DEPRECATED(5.4, "Please update to supply the GroupName directly as this will be removed in a future version.")
 	virtual int32 CompileScript(const class FNiagaraCompileRequestData* InCompileRequest, const FNiagaraCompileOptions& InOptions, const FNiagaraTranslateResults& InTranslateResults, FNiagaraTranslatorOutput* TranslatorOutput, FString& TranslatedHLSL) override;
 
-	virtual int32 CompileScript(const FStringView GroupName, const FNiagaraCompileOptions& InOptions, const FNiagaraTranslateResults& InTranslateResults, FNiagaraTranslatorOutput* TranslatorOutput, FString& TranslatedHLSL) override;
+	virtual int32 CompileScript(const FStringView GroupName, const FNiagaraCompileOptions& InOptions, const FNiagaraTranslateResults& InTranslateResults, const FNiagaraTranslatorOutput& TranslatorOutput, const FString& TranslatedHLSL) override;
 	virtual TOptional<FNiagaraCompileResults> GetCompileResult(int32 JobID, bool bWait = false) override;
 
 	virtual void Error(FText ErrorText) override;
diff --git a/Engine/Plugins/FX/Niagara/Source/NiagaraEditor/Private/NiagaraEditorModule.cpp b/Engine/Plugins/FX/Niagara/Source/NiagaraEditor/Private/NiagaraEditorModule.cpp
--- a/Engine/Plugins/FX/Niagara/Source/NiagaraEditor/Private/NiagaraEditorModule.cpp
+++ b/Engine/Plugins/FX/Niagara/Source/NiagaraEditor/Private/NiagaraEditorModule.cpp
@@ -964,7 +964,6 @@ void FNiagaraEditorModule::StartupModule()
 {
 	bThumbnailRenderersRegistered = false;
 
-	FHlslNiagaraTranslator::Init();
 	MenuExtensibilityManager = MakeShareable(new FExtensibilityManager);
 	ToolBarExtensibilityManager = MakeShareable(new FExtensibilityManager);
 
diff --git a/Engine/Plugins/FX/Niagara/Source/NiagaraEditor/Private/NiagaraEditorUtilities.cpp b/Engine/Plugins/FX/Niagara/Source/NiagaraEditor/Private/NiagaraEditorUtilities.cpp
--- a/Engine/Plugins/FX/Niagara/Source/NiagaraEditor/Private/NiagaraEditorUtilities.cpp
+++ b/Engine/Plugins/FX/Niagara/Source/NiagaraEditor/Private/NiagaraEditorUtilities.cpp
@@ -19,6 +19,7 @@
 #include "NiagaraEditorStyle.h"
 #include "NiagaraGraph.h"
 #include "NiagaraEditorModule.h"
+#include "NiagaraHlslTranslator.h"
 #include "NiagaraNodeFunctionCall.h"
 #include "NiagaraNodeInput.h"
 #include "NiagaraNodeOutput.h"
@@ -4568,7 +4569,7 @@ UNiagaraDataInterface* FNiagaraEditorUtilities::GetResolvedRuntimeInstanceForEdi
 			{
 				FString EmitterName = EmitterHandle != nullptr ? EmitterHandle->GetUniqueInstanceName() : FString();
 				bool bIsParameterMapDataInterface = false;
-				FName DataInterfaceName = FHlslNiagaraTranslator::GetDataInterfaceName(OuterInputNode->Input.GetName(), EmitterName, bIsParameterMapDataInterface);
+				FName DataInterfaceName = FNiagaraHlslTranslator::GetDataInterfaceName(OuterInputNode->Input.GetName(), EmitterName, bIsParameterMapDataInterface);
 				for (const FNiagaraScriptResolvedDataInterfaceInfo& ResolvedDI : CompiledScript->GetResolvedDataInterfaces())
 				{
 					if (ResolvedDI.Name == DataInterfaceName)
diff --git a/Engine/Plugins/FX/Niagara/Source/NiagaraEditor/Private/NiagaraEmitterEditorData.cpp b/Engine/Plugins/FX/Niagara/Source/NiagaraEditor/Private/NiagaraEmitterEditorData.cpp
--- a/Engine/Plugins/FX/Niagara/Source/NiagaraEditor/Private/NiagaraEmitterEditorData.cpp
+++ b/Engine/Plugins/FX/Niagara/Source/NiagaraEditor/Private/NiagaraEmitterEditorData.cpp
@@ -3,6 +3,7 @@
 #include "NiagaraEmitterEditorData.h"
 
 #include "NiagaraNodeFunctionCall.h"
+#include "NiagaraScriptSource.h"
 #include "NiagaraScriptVariable.h"
 #include "EdGraph/EdGraph.h"
 #include "NiagaraStackEditorData.h"
diff --git a/Engine/Plugins/FX/Niagara/Source/NiagaraEditor/Private/NiagaraGraph.cpp b/Engine/Plugins/FX/Niagara/Source/NiagaraEditor/Private/NiagaraGraph.cpp
--- a/Engine/Plugins/FX/Niagara/Source/NiagaraEditor/Private/NiagaraGraph.cpp
+++ b/Engine/Plugins/FX/Niagara/Source/NiagaraEditor/Private/NiagaraGraph.cpp
@@ -3670,13 +3670,13 @@ void UNiagaraGraph::InvalidateNumericCache()
 	CachedNumericConversions.Empty();
 }
 
-FString UNiagaraGraph::GetFunctionAliasByContext(const FNiagaraGraphFunctionAliasContext& FunctionAliasContext)
+FString UNiagaraGraph::GetFunctionAliasByContext(const FNiagaraGraphFunctionAliasContext& FunctionAliasContext) const
 {
 	FString FunctionAlias;
 	TSet<UClass*> SkipNodeTypes;
-	for (UEdGraphNode* Node : Nodes)
+	for (const UEdGraphNode* Node : Nodes)
 	{
-		UNiagaraNode* NiagaraNode = Cast<UNiagaraNode>(Node);
+		const UNiagaraNode* NiagaraNode = Cast<const UNiagaraNode>(Node);
 		if (NiagaraNode != nullptr)
 		{
 			if (SkipNodeTypes.Contains(NiagaraNode->GetClass()))
@@ -3692,10 +3692,10 @@ FString UNiagaraGraph::GetFunctionAliasByContext(const FNiagaraGraphFunctionAlia
 		}
 	}
 
-	for (UEdGraphPin* Pin : FunctionAliasContext.StaticSwitchValues)
+	for (const UEdGraphPin* Pin : FunctionAliasContext.StaticSwitchValues)
 	{
-		FunctionAlias += TEXT("_") + FHlslNiagaraTranslator::GetSanitizedFunctionNameSuffix(Pin->GetName()) 
-			+ TEXT("_") + FHlslNiagaraTranslator::GetSanitizedFunctionNameSuffix(Pin->DefaultValue);
+		FunctionAlias += TEXT("_") + FNiagaraHlslTranslator::GetSanitizedFunctionNameSuffix(Pin->GetName())
+			+ TEXT("_") + FNiagaraHlslTranslator::GetSanitizedFunctionNameSuffix(Pin->DefaultValue);
 	}
 	return FunctionAlias;
 }
diff --git a/Engine/Plugins/FX/Niagara/Source/NiagaraEditor/Private/NiagaraGraphDataCache.cpp b/Engine/Plugins/FX/Niagara/Source/NiagaraEditor/Private/NiagaraGraphDataCache.cpp
--- a/Engine/Plugins/FX/Niagara/Source/NiagaraEditor/Private/NiagaraGraphDataCache.cpp
+++ b/Engine/Plugins/FX/Niagara/Source/NiagaraEditor/Private/NiagaraGraphDataCache.cpp
@@ -40,7 +40,7 @@ void FNiagaraGraphDataCache::GetStackFunctionInputsInternal(
 	// we don't currently support hashing the internal details of the translator within FStackFunctionInputPinKey so in the case
 	// we've been supplied a translator we'll just fall back to not using the cache
 	TArray<const UEdGraphPin*> InputPins;
-	if (const FHlslNiagaraTranslator* Translator = ConstantResolver.GetTranslator())
+	if (const FNiagaraHlslTranslator* Translator = ConstantResolver.GetTranslator())
 	{
 		FNiagaraStackGraphUtilities::GetStackFunctionInputPinsWithoutCache(
 			FunctionCallNode,
diff --git a/Engine/Plugins/FX/Niagara/Source/NiagaraEditor/Private/NiagaraGraphHlslTranslator.h b/Engine/Plugins/FX/Niagara/Source/NiagaraEditor/Private/NiagaraGraphHlslTranslator.h
new file mode 100644
--- /dev/null
+++ b/Engine/Plugins/FX/Niagara/Source/NiagaraEditor/Private/NiagaraGraphHlslTranslator.h
@@ -0,0 +1,7 @@
+// Copyright Epic Games, Inc. All Rights Reserved.
+
+#pragma once
+
+#include "CoreMinimal.h"
+#include "NiagaraHlslTranslator.h"
+#include "NiagaraCompilationGraphBridge.h"
