Adding a templated layer in the hlsl translator to encapsulate logic surrounding graph implementation to allow us to share code between the existing UNiagaraGraph and the upcoming graph digest.  Gener

diff --git a/Engine/Plugins/FX/Niagara/Source/NiagaraEditor/Private/NiagaraHlslTranslator.cpp b/Engine/Plugins/FX/Niagara/Source/NiagaraEditor/Private/NiagaraHlslTranslator.cpp
--- a/Engine/Plugins/FX/Niagara/Source/NiagaraEditor/Private/NiagaraHlslTranslator.cpp
+++ b/Engine/Plugins/FX/Niagara/Source/NiagaraEditor/Private/NiagaraHlslTranslator.cpp
@@ -5,6 +5,7 @@
 #include "EdGraphSchema_Niagara.h"
 #include "INiagaraEditorTypeUtilities.h"
 #include "Modules/ModuleManager.h"
+#include "NiagaraAttributeTrimmer.h"
 #include "NiagaraCommon.h"
 #include "NiagaraConstants.h"
 #include "NiagaraDataInterface.h"
@@ -16,6 +17,7 @@
 #include "NiagaraFunctionLibrary.h"
 #include "NiagaraNodeEmitter.h"
 #include "NiagaraGraph.h"
+#include "NiagaraGraphHlslTranslator.h"
 #include "NiagaraNode.h"
 #include "NiagaraNodeConvert.h"
 #include "NiagaraNodeCustomHlsl.h"
@@ -212,7 +214,7 @@ struct TSimStagePermutationContext
 		, HasBranch(false)
 	{}
 
-	TSimStagePermutationContext(const FHlslNiagaraTranslator& Translator, const FHlslNiagaraTranslationStage& TranslationStage, FString& InHlslOutput)
+	TSimStagePermutationContext(const FNiagaraHlslTranslator& Translator, const FNiagaraHlslTranslationStage& TranslationStage, FString& InHlslOutput)
 		: HlslOutput(InHlslOutput)
 		, Enabled(false)
 		, HasBranch(false)
@@ -220,7 +222,7 @@ struct TSimStagePermutationContext
 		AddBranch(Translator, TranslationStage);
 	}
 
-	TSimStagePermutationContext(const FHlslNiagaraTranslator& Translator, TConstArrayView<FHlslNiagaraTranslationStage> TranslationStages, TConstArrayView<int32> StageIndices, FString& InHlslOutput)
+	TSimStagePermutationContext(const FNiagaraHlslTranslator& Translator, TConstArrayView<FNiagaraHlslTranslationStage> TranslationStages, TConstArrayView<int32> StageIndices, FString& InHlslOutput)
 		: HlslOutput(InHlslOutput)
 		, Enabled(false)
 		, HasBranch(false)
@@ -233,10 +235,10 @@ struct TSimStagePermutationContext
 		Release();
 	}
 
-	void AddBranch(const FHlslNiagaraTranslator& Translator, const FHlslNiagaraTranslationStage& TranslationStage)
+	void AddBranch(const FNiagaraHlslTranslator& Translator, const FNiagaraHlslTranslationStage& TranslationStage)
 	{
 		// vs2017 generates a link error; this works around it...some how
-		const FHlslNiagaraTranslationStage& LocalTranslation = TranslationStage;
+		const FNiagaraHlslTranslationStage& LocalTranslation = TranslationStage;
 
 		AddBranchInternal(Translator, MakeArrayView(&LocalTranslation, 1), MakeArrayView({ 0 }));
 	}
@@ -250,14 +252,14 @@ struct TSimStagePermutationContext
 		}
 	}
 
-	static bool SupportsBranching(const FHlslNiagaraTranslator& Translator)
+	static bool SupportsBranching(const FNiagaraHlslTranslator& Translator)
 	{
 		return Translator.GetSimulationTarget() == ENiagaraSimTarget::GPUComputeSim;
 	}
 
 
 private:
-	FString BuildConditionString(TConstArrayView<FHlslNiagaraTranslationStage> TranslationStages, TConstArrayView<int32> StageIndices)
+	FString BuildConditionString(TConstArrayView<FNiagaraHlslTranslationStage> TranslationStages, TConstArrayView<int32> StageIndices)
 	{
 		FString ConditionString;
 
@@ -275,7 +277,7 @@ private:
 		return ConditionString;
 	}
 
-	void AddBranchInternal(const FHlslNiagaraTranslator& Translator, TConstArrayView<FHlslNiagaraTranslationStage> TranslationStages, TConstArrayView<int32> StageIndices)
+	void AddBranchInternal(const FNiagaraHlslTranslator& Translator, TConstArrayView<FNiagaraHlslTranslationStage> TranslationStages, TConstArrayView<int32> StageIndices)
 	{
 		if (SupportsBranching(Translator) && StageIndices.Num())
 		{
@@ -332,13 +334,13 @@ static void ReplaceNamespaceInline(FString& FullName, FStringView Source, FStrin
 	}
 }
 
-FString FHlslNiagaraTranslator::GetCode(int32 ChunkIdx)
+FString FNiagaraHlslTranslator::GetCode(int32 ChunkIdx)
 {
 	FNiagaraCodeChunk& Chunk = CodeChunks[ChunkIdx];
 	return GetCode(Chunk);
 }
 
-FString FHlslNiagaraTranslator::GetCode(FNiagaraCodeChunk& Chunk)
+FString FNiagaraHlslTranslator::GetCode(FNiagaraCodeChunk& Chunk)
 {
 	TArray<FStringFormatArg> Args;
 	for (int32 i = 0; i < Chunk.SourceChunks.Num(); ++i)
@@ -365,7 +367,7 @@ FString FHlslNiagaraTranslator::GetCode(FNiagaraCodeChunk& Chunk)
 		{
 			if (!Chunk.bDecl)//Otherwise, we're doing nothing here.
 			{
-				Warning(LOCTEXT("MissingDeclForChunk", "Missing definition string."), nullptr, nullptr);
+				Warning(LOCTEXT("MissingDeclForChunk", "Missing definition string."));
 			}
 
 			FinalString += GetStructHlslTypeName(Chunk.Type) + TEXT(" ") + Chunk.SymbolName + TEXT(";\n");
@@ -385,7 +387,7 @@ FString FHlslNiagaraTranslator::GetCode(FNiagaraCodeChunk& Chunk)
 	return FinalString;
 }
 
-FString FHlslNiagaraTranslator::GetCodeAsSource(int32 ChunkIdx)
+FString FNiagaraHlslTranslator::GetCodeAsSource(int32 ChunkIdx)
 {
 	if (ChunkIdx >= 0 && ChunkIdx < CodeChunks.Num())
 	{
@@ -395,19 +397,21 @@ FString FHlslNiagaraTranslator::GetCodeAsSource(int32 ChunkIdx)
 	return "Undefined";
 }
 
-bool FHlslNiagaraTranslator::ValidateTypePins(const UNiagaraNode* NodeToValidate)
+// specialization because validation may be different between graph implementations
+template<>
+bool TNiagaraHlslTranslator<FNiagaraCompilationGraphBridge>::ValidateTypePins(const FNode* NodeToValidate)
 {
 	bool bPinsAreValid = true;
-	for (UEdGraphPin* Pin : NodeToValidate->GetAllPins())
+	for (const FPin* Pin : GetPins(NodeToValidate))
 	{
-		if (Pin->PinType.PinCategory == "")
+		if (Pin->PinType.PinCategory.IsNone())
 		{
 			Error(LOCTEXT("InvalidPinTypeError", "Node pin has an undefined type."), NodeToValidate, Pin);
 			bPinsAreValid = false;
 		}
 		else if (Pin->PinType.PinCategory == UEdGraphSchema_Niagara::PinCategoryType || Pin->PinType.PinCategory == UEdGraphSchema_Niagara::PinCategoryStaticType)
 		{
-			FNiagaraTypeDefinition Type = Schema->PinToTypeDefinition(Pin, ENiagaraStructConversion::Simulation);
+			FNiagaraTypeDefinition Type = UEdGraphSchema_Niagara::PinToTypeDefinition(Pin, ENiagaraStructConversion::Simulation);
 			if (Type.IsValid() == false)
 			{
 				Error(LOCTEXT("InvalidPinTypeError", "Node pin has an undefined type."), NodeToValidate, Pin);
@@ -428,14 +432,14 @@ bool FHlslNiagaraTranslator::ValidateTypePins(const UNiagaraNode* NodeToValidate
 	return bPinsAreValid;
 }
 
-
-void FHlslNiagaraTranslator::GenerateFunctionSignature(ENiagaraScriptUsage ScriptUsage, FString InName, const FString& InFullName, const FString& InFunctionNameSuffix, UNiagaraGraph* FuncGraph, TArray<int32>& Inputs,
-	bool bHasNumericInputs, bool bHasParameterMapParameters, TArray<UEdGraphPin*> StaticSwitchValues, FNiagaraFunctionSignature& OutSig)const
+template<typename GraphBridge>
+void TNiagaraHlslTranslator<GraphBridge>::GenerateFunctionSignature(ENiagaraScriptUsage ScriptUsage, FString InName, const FString& InFullName, const FString& InFunctionNameSuffix, const FGraph* FuncGraph, TArray<int32>& Inputs,
+	bool bHasNumericInputs, bool bHasParameterMapParameters, const TArray<const FInputPin*>& StaticSwitchValues, FNiagaraFunctionSignature& OutSig) const
 {
 	NIAGARA_SCOPE_CYCLE_COUNTER(STAT_NiagaraEditor_Module_NiagaraHLSLTranslator_GenerateFunctionSignature);
 
 	TArray<FNiagaraVariable> InputVars;
-	TArray<UNiagaraNodeInput*> InputsNodes;
+	TArray<const FInputNode*> InputsNodes;
 	bool bHasDIParameters = false;
 
 	{
@@ -450,12 +454,12 @@ void FHlslNiagaraTranslator::GenerateFunctionSignature(ENiagaraScriptUsage Scrip
 		// to compile a graph that would otherwise work correctly.
 		Options.bFilterByScriptUsage = ScriptUsage == ENiagaraScriptUsage::EmitterSpawnScript || ScriptUsage == ENiagaraScriptUsage::EmitterUpdateScript;
 		Options.TargetScriptUsage = ScriptUsage;
-		FuncGraph->FindInputNodes(InputsNodes, Options);
+		InputsNodes = GraphBridge::GetGraphInputNodes(FuncGraph, Options);
 
 		if (Inputs.Num() != InputsNodes.Num())
 		{
-			const_cast<FHlslNiagaraTranslator*>(this)->Error(FText::Format(LOCTEXT("GenerateFunctionSignatureFail", "Generating function signature for {0} failed.  The function call is providing a different number of inputs than the function graph supplies."),
-				FText::FromString(InFullName)), nullptr, nullptr);
+			const_cast<TNiagaraHlslTranslator<GraphBridge>*>(this)->Error(FText::Format(LOCTEXT("GenerateFunctionSignatureFail", "Generating function signature for {0} failed.  The function call is providing a different number of inputs than the function graph supplies."),
+				FText::FromString(InFullName)));
 			return;
 		}
 	}
@@ -471,7 +475,7 @@ void FHlslNiagaraTranslator::GenerateFunctionSignature(ENiagaraScriptUsage Scrip
 			//Only add to the signature if the caller has provided it, otherwise we use a local default.
 			if (Inputs[i] != INDEX_NONE)
 			{
-				FNiagaraVariable InputVar = InputsNodes[i]->Input;
+				FNiagaraVariable InputVar = GraphBridge::GetInputVariable(InputsNodes[i]);
 				if (GetLiteralConstantVariable(InputVar))
 				{
 					checkf(InputVar.GetType().IsSameBaseDefinition(FNiagaraTypeDefinition::GetBoolDef()), TEXT("Only boolean types are currently supported for literal constants."));
@@ -509,8 +513,7 @@ void FHlslNiagaraTranslator::GenerateFunctionSignature(ENiagaraScriptUsage Scrip
 	{
 		NIAGARA_SCOPE_CYCLE_COUNTER(STAT_NiagaraEditor_Module_NiagaraHLSLTranslator_GenerateFunctionSignature_Outputs);
 
-		OutputVars.Reserve(10);
-		FuncGraph->GetOutputNodeVariables(ScriptUsage, OutputVars);
+		OutputVars = GraphBridge::GetGraphOutputNodeVariables(FuncGraph, ScriptUsage);
 
 		for (int32 i = 0; i < OutputVars.Num(); ++i)
 		{
@@ -554,7 +557,7 @@ void FHlslNiagaraTranslator::GenerateFunctionSignature(ENiagaraScriptUsage Scrip
 	}
 	else
 	{
-		FNiagaraGraphFunctionAliasContext FunctionAliasContext;
+		FGraphFunctionAliasContext FunctionAliasContext;
 		FunctionAliasContext.CompileUsage = GetCurrentUsage();
 		FunctionAliasContext.ScriptUsage = TranslationStages[ActiveStageIdx].ScriptUsage;
 		FunctionAliasContext.StaticSwitchValues = StaticSwitchValues;
@@ -572,9 +575,8 @@ void FHlslNiagaraTranslator::GenerateFunctionSignature(ENiagaraScriptUsage Scrip
 
 //////////////////////////////////////////////////////////////////////////
 
-FHlslNiagaraTranslator::FHlslNiagaraTranslator()
-	: Schema(GetDefault<UEdGraphSchema_Niagara>())
-	, CurrentBodyChunkMode(ENiagaraCodeChunkMode::Body)
+FNiagaraHlslTranslator::FNiagaraHlslTranslator()
+	: CurrentBodyChunkMode(ENiagaraCodeChunkMode::Body)
 	, ActiveStageIdx(-1)
 	, bInitializedDefaults(false)
 {
@@ -583,7 +585,7 @@ FHlslNiagaraTranslator::FHlslNiagaraTranslator()
 }
 
 
-FString FHlslNiagaraTranslator::GetFunctionDefinitions()
+FString FNiagaraHlslTranslator::GetFunctionDefinitions()
 {
 	FString FwdDeclString;
 	FString DefinitionsString;
@@ -633,32 +635,35 @@ FString FHlslNiagaraTranslator::GetFunctionDefinitions()
 	return FwdDeclString + TEXT("\n") + DefinitionsString;
 }
 
-void FHlslNiagaraTranslator::BuildMissingDefaults()
+template<typename GraphBridge>
+void TNiagaraHlslTranslator<GraphBridge>::BuildMissingDefaults()
 {
 	AddBodyComment(TEXT("// Begin HandleMissingDefaultValues"));
-
 	if (TranslationStages[ActiveStageIdx].ShouldDoSpawnOnlyLogic() || TranslationStages[ActiveStageIdx].bShouldUpdateInitialAttributeValues)
 	{
 		// First go through all the variables that we did not write the defaults for yet. For spawn scripts, this usually
 		// means variables that reference other variables but are not themselves used within spawn.
 		for (FNiagaraVariable& Var : DeferredVariablesMissingDefault)
 		{
-			const UEdGraphPin* DefaultPin = UniqueVarToDefaultPin.FindChecked(Var);
+			const FInputPin* DefaultPin = UniqueVarToDefaultPin.FindChecked(Var);
 			bool bWriteToParamMapEntries = UniqueVarToWriteToParamMap.FindChecked(Var);
 			int32 OutputChunkId = INDEX_NONE;
 
 			TOptional<ENiagaraDefaultMode> DefaultMode;
 			FNiagaraScriptVariableBinding DefaultBinding;
 
+			const FNode* DefaultPinNode = nullptr;
+
 			if (DefaultPin) 
 			{
-				if (UNiagaraGraph* DefaultPinGraph = CastChecked<UNiagaraGraph>(DefaultPin->GetOwningNode()->GetGraph())) 
+				DefaultPinNode = GraphBridge::GetOwningNode(DefaultPin);
+				if (const FGraph* DefaultPinGraph = GraphBridge::GetOwningGraph(DefaultPinNode))
 				{
-					DefaultMode = DefaultPinGraph->GetDefaultMode(Var, &DefaultBinding);
+					DefaultMode = GraphBridge::GetGraphDefaultMode(DefaultPinGraph, Var, DefaultBinding);
 				}
 			}
 
-			HandleParameterRead(ActiveStageIdx, Var, DefaultPin, DefaultPin != nullptr ? Cast<UNiagaraNode>(DefaultPin->GetOwningNode()) : nullptr, OutputChunkId, DefaultMode, DefaultBinding, !bWriteToParamMapEntries, true);
+			HandleParameterRead(ActiveStageIdx, Var, DefaultPin, DefaultPinNode, OutputChunkId, DefaultMode, DefaultBinding, !bWriteToParamMapEntries, true);
 		}
 
 		DeferredVariablesMissingDefault.Empty();
@@ -668,9 +673,9 @@ void FHlslNiagaraTranslator::BuildMissingDefaults()
 			// Now go through and initialize any "Particles.Initial." variables
 			for (FNiagaraVariable& Var : InitialNamespaceVariablesMissingDefault)
 			{
-				if (FNiagaraParameterMapHistory::IsInitialValue(Var))
+				if (FParamMapHistory::IsInitialValue(Var))
 				{
-					FNiagaraVariable SourceForInitialValue = FNiagaraParameterMapHistory::GetSourceForInitialValue(Var);
+					FNiagaraVariable SourceForInitialValue = FParamMapHistory::GetSourceForInitialValue(Var);
 					FString ParameterMapInstanceName = GetParameterMapInstanceName(0);
 					FString Value = FString::Printf(TEXT("%s.%s = %s.%s;\n"), *ParameterMapInstanceName, *GetSanitizedSymbolName(Var.GetName().ToString()),
 						*ParameterMapInstanceName, *GetSanitizedSymbolName(SourceForInitialValue.GetName().ToString()));
@@ -682,20 +687,20 @@ void FHlslNiagaraTranslator::BuildMissingDefaults()
 		}
 
 	}
-
 	AddBodyComment(TEXT("// End HandleMissingDefaultValues\n\n"));
 }
 
-FString FHlslNiagaraTranslator::BuildParameterMapHlslDefinitions(TArray<FNiagaraVariable>& PrimaryDataSetOutputEntries)
+template<typename GraphBridge>
+FString TNiagaraHlslTranslator<GraphBridge>::BuildParameterMapHlslDefinitions(TArray<FNiagaraVariable>& PrimaryDataSetOutputEntries)
 {
 	NIAGARA_SCOPE_CYCLE_COUNTER(STAT_NiagaraEditor_HlslTranslator_BuildParameterMapHlslDefinitions);
 	FString HlslOutputString;
 
 	// Determine the unique parameter map structs...
-	TArray<const UEdGraphPin*> UniqueParamMapStartingPins;
+	TArray<const FPin*> UniqueParamMapStartingPins;
 	for (int32 ParamMapIdx = 0; ParamMapIdx < ParamMapHistories.Num(); ParamMapIdx++)
 	{
-		const UEdGraphPin* OriginalPin = ParamMapHistories[ParamMapIdx].GetOriginalPin();
+		const FPin* OriginalPin = ParamMapHistories[ParamMapIdx].GetOriginalPin();
 		UniqueParamMapStartingPins.AddUnique(OriginalPin);
 	}
 
@@ -712,7 +717,7 @@ FString FHlslNiagaraTranslator::BuildParameterMapHlslDefinitions(TArray<FNiagara
 		}
 
 		// ignore those that are rapid iteration parameters as those will be read in directly from the cbuffer
-		if (FNiagaraParameterMapHistory::IsRapidIterationParameter(Var))
+		if (FParamMapHistory::IsRapidIterationParameter(Var))
 		{
 			continue;
 		}
@@ -777,9 +782,9 @@ FString FHlslNiagaraTranslator::BuildParameterMapHlslDefinitions(TArray<FNiagara
 		for (const FNiagaraVariableBase& Var : CompileOptions.AdditionalVariables)
 		{
 			bool bFoundSource = false;
-			if (FNiagaraParameterMapHistory::IsPreviousValue(Var))
+			if (FParamMapHistory::IsPreviousValue(Var))
 			{
-				FNiagaraVariable Source = FNiagaraParameterMapHistory::GetSourceForPreviousValue(FNiagaraVariable(Var));
+				FNiagaraVariable Source = FParamMapHistory::GetSourceForPreviousValue(FNiagaraVariable(Var));
 				const FNiagaraTypeDefinition& SourceType = Source.GetType();
 
 				for (int32 ParamMapIdxTest = 0; !bFoundSource && ParamMapIdxTest < OtherOutputParamMapHistories.Num(); ParamMapIdxTest++)
@@ -840,7 +845,7 @@ FString FHlslNiagaraTranslator::BuildParameterMapHlslDefinitions(TArray<FNiagara
 	for (int32 UniqueVarIdx = 0; UniqueVarIdx < UniqueVariables.Num(); UniqueVarIdx++)
 	{
 		const FNiagaraVariable& NiagaraVariable = UniqueVariables[UniqueVarIdx];
-		if (FNiagaraParameterMapHistory::IsAttribute(NiagaraVariable))
+		if (FParamMapHistory::IsAttribute(NiagaraVariable))
 		{
 			FString VariableName = GetSanitizedSymbolName(NiagaraVariable.GetName().ToString());
 			ReplaceNamespaceInline(VariableName, PARAM_MAP_ATTRIBUTE_STR, PARAM_MAP_INDICES_STR);
@@ -939,7 +944,7 @@ FString FHlslNiagaraTranslator::BuildParameterMapHlslDefinitions(TArray<FNiagara
 }
 
 
-bool FHlslNiagaraTranslator::ShouldConsiderTargetParameterMap(ENiagaraScriptUsage InUsage) const
+bool FNiagaraHlslTranslator::ShouldConsiderTargetParameterMap(ENiagaraScriptUsage InUsage) const
 {
 	ENiagaraScriptUsage TargetUsage = GetTargetUsage();
 	if (TargetUsage >= ENiagaraScriptUsage::ParticleSpawnScript && TargetUsage <= ENiagaraScriptUsage::ParticleEventScript)
@@ -965,7 +970,7 @@ bool FHlslNiagaraTranslator::ShouldConsiderTargetParameterMap(ENiagaraScriptUsag
 	return false;
 }
 
-void FHlslNiagaraTranslator::HandleNamespacedExternalVariablesToDataSetRead(TArray<FNiagaraVariable>& InDataSetVars, FString InNamespaceStr)
+void FNiagaraHlslTranslator::HandleNamespacedExternalVariablesToDataSetRead(TArray<FNiagaraVariable>& InDataSetVars, FString InNamespaceStr)
 {
 	for (const FNiagaraVariable& Var : ExternalVariablesForBulkUsage)
 	{
@@ -976,7 +981,7 @@ void FHlslNiagaraTranslator::HandleNamespacedExternalVariablesToDataSetRead(TArr
 	}
 }
 
-bool FHlslNiagaraTranslator::IsVariableInUniformBuffer(const FNiagaraVariable& Variable) const
+bool FNiagaraHlslTranslator::IsVariableInUniformBuffer(const FNiagaraVariable& Variable) const
 {
 	static FNiagaraVariable GpuExcludeVariables[] =
 	{
@@ -1007,7 +1012,7 @@ bool FHlslNiagaraTranslator::IsVariableInUniformBuffer(const FNiagaraVariable& V
 }
 
 template<typename T>
-void FHlslNiagaraTranslator::BuildConstantBuffer(ENiagaraCodeChunkMode ChunkMode)
+void FNiagaraHlslTranslator::BuildConstantBuffer(ENiagaraCodeChunkMode ChunkMode)
 {
 	for (const FNiagaraVariable& Variable : T::GetVariables())
 	{
@@ -1016,7 +1021,8 @@ void FHlslNiagaraTranslator::BuildConstantBuffer(ENiagaraCodeChunkMode ChunkMode
 	}
 }
 
-void FHlslNiagaraTranslator::RecordParamMapDefinedAttributeToNamespaceVar(const FNiagaraVariable& VarToRecord, const UEdGraphPin* VarAssociatedDefaultPin)
+template<typename GraphBridge>
+void TNiagaraHlslTranslator<GraphBridge>::RecordParamMapDefinedAttributeToNamespaceVar(const FNiagaraVariable& VarToRecord, const FPin* VarAssociatedDefaultPin)
 {
 	bool bDefaultPinExplicit = true;
 	if (VarAssociatedDefaultPin == nullptr || VarAssociatedDefaultPin->bHidden)
@@ -1118,32 +1124,31 @@ static void ConvertFloatToHalf(const FNiagaraCompileOptions& InCompileOptions, T
 	}
 }
 
-const FNiagaraTranslateResults &FHlslNiagaraTranslator::Translate(const FNiagaraCompileRequestData* InCompileData, const FNiagaraCompileRequestDuplicateData* InCompileDuplicateData, const FNiagaraCompileOptions& InCompileOptions, FHlslNiagaraTranslatorOptions InTranslateOptions)
+template<typename GraphBridge>
+FNiagaraTranslateResults TNiagaraHlslTranslator<GraphBridge>::Translate(const FNiagaraCompileOptions& InCompileOptions, const FHlslNiagaraTranslatorOptions& InTranslateOptions)
 {
 	TRACE_CPUPROFILER_EVENT_SCOPE(NiagaraHlslTranslate);
 	TRACE_CPUPROFILER_EVENT_SCOPE_TEXT_ON_CHANNEL(*InCompileOptions.GetPathName(), NiagaraChannel);
 
 	NIAGARA_SCOPE_CYCLE_COUNTER(STAT_NiagaraEditor_HlslTranslator_Translate);
-	check(InCompileData);
-	check(InCompileDuplicateData);
+	check(CompileData);
+	check(CompileDuplicateData);
 
 	CompileOptions = InCompileOptions;
-	CompileData = InCompileData;
-	CompileDuplicateData = InCompileDuplicateData;
 	TranslationOptions = InTranslateOptions;
 	CompilationTarget = TranslationOptions.SimTarget;
 	TranslateResults.bHLSLGenSucceeded = false;
 	TranslateResults.OutputHLSL = "";
 
-	TWeakObjectPtr<UNiagaraGraph> SourceGraph = CompileDuplicateData->NodeGraphDeepCopy;
+	const FGraph* SourceGraph = GraphBridge::GetGraph(CompileDuplicateData);
 
-	if (!SourceGraph.IsValid())
+	if (!SourceGraph)
 	{
 		Error(LOCTEXT("GetGraphFail", "Cannot find graph node!"), nullptr, nullptr);
 		return TranslateResults;
 	}
 
-	if (SourceGraph->IsEmpty())
+	if (GraphBridge::IsGraphEmpty(SourceGraph))
 	{
 		if (UNiagaraScript::IsSystemScript(CompileOptions.TargetUsage))
 		{
@@ -1190,43 +1195,49 @@ const FNiagaraTranslateResults &FHlslNiagaraTranslator::Translate(const FNiagara
 	switch (CompileOptions.TargetUsage)
 	{
 	case ENiagaraScriptUsage::ParticleSpawnScriptInterpolated:
-		TranslationStages.Add(FHlslNiagaraTranslationStage(CompileOptions.TargetUsage, CompileOptions.TargetUsageId));
-		TranslationStages.Add(FHlslNiagaraTranslationStage(ENiagaraScriptUsage::ParticleUpdateScript, FGuid()));
-		TranslationStages[0].PassNamespace = TEXT("MapSpawn");
-		TranslationStages[1].PassNamespace = TEXT("MapUpdate");
-		TranslationStages[0].ChunkModeIndex = ENiagaraCodeChunkMode::SpawnBody;
-		TranslationStages[1].ChunkModeIndex = ENiagaraCodeChunkMode::UpdateBody;
-		TranslationStages[0].OutputNode = SourceGraph->FindEquivalentOutputNode(ENiagaraScriptUsage::ParticleSpawnScript, TranslationStages[0].UsageId);
-		TranslationStages[1].OutputNode = SourceGraph->FindEquivalentOutputNode(TranslationStages[1].ScriptUsage, TranslationStages[1].UsageId);
-		TranslationStages[1].bInterpolatePreviousParams = true;
-		TranslationStages[0].SimulationStageIndex = 0;
-		TranslationStages[0].NumIterations = 1;
-		TranslationStages[1].SimulationStageIndex = 0;
-		TranslationStages[1].NumIterations = 1;
-		TranslationStages[0].bWritesParticles = true;
-		TranslationStages[1].bWritesParticles = true;
-		TranslationStages[0].bShouldUpdateInitialAttributeValues = true;
+	{
+		FNiagaraHlslTranslationStage& SpawnStage = TranslationStages.Emplace_GetRef(CompileOptions.TargetUsage, CompileOptions.TargetUsageId);
+		SpawnStage.PassNamespace = TEXT("MapSpawn");
+		SpawnStage.ChunkModeIndex = ENiagaraCodeChunkMode::SpawnBody;
+		SpawnStage.SimulationStageIndex = 0;
+		SpawnStage.NumIterations = 1;
+		SpawnStage.bWritesParticles = true;
+		SpawnStage.bShouldUpdateInitialAttributeValues = true;
+		OutputNodes.Add(SourceGraph->FindEquivalentOutputNode(ENiagaraScriptUsage::ParticleSpawnScript, SpawnStage.UsageId));
+
+		FNiagaraHlslTranslationStage& UpdateStage = TranslationStages.Emplace_GetRef(ENiagaraScriptUsage::ParticleUpdateScript, FGuid());
+		UpdateStage.PassNamespace = TEXT("MapUpdate");
+		UpdateStage.ChunkModeIndex = ENiagaraCodeChunkMode::UpdateBody;
+		UpdateStage.bInterpolatePreviousParams = true;
+		UpdateStage.SimulationStageIndex = 0;
+		UpdateStage.NumIterations = 1;
+		UpdateStage.bWritesParticles = true;
+		OutputNodes.Add(SourceGraph->FindEquivalentOutputNode(UpdateStage.ScriptUsage, UpdateStage.UsageId));
+
 		ParamMapHistories.AddDefaulted(2);
 		ParamMapSetVariablesToChunks.AddDefaulted(2);
 		ParamMapHistoriesSourceInOtherHistories.AddDefaulted(2);
-		break;
+	}	break;
 	case ENiagaraScriptUsage::ParticleGPUComputeScript:
-		TranslationStages.Add(FHlslNiagaraTranslationStage((bHasInterpolatedSpawn ? ENiagaraScriptUsage::ParticleSpawnScriptInterpolated : ENiagaraScriptUsage::ParticleSpawnScript), FGuid()));
-		TranslationStages.Add(FHlslNiagaraTranslationStage(ENiagaraScriptUsage::ParticleUpdateScript, FGuid()));
-		TranslationStages[0].PassNamespace = TEXT("MapSpawn");
-		TranslationStages[1].PassNamespace = TEXT("MapUpdate");
-		TranslationStages[0].ChunkModeIndex = ENiagaraCodeChunkMode::SpawnBody;
-		TranslationStages[1].ChunkModeIndex = ENiagaraCodeChunkMode::UpdateBody;
-		TranslationStages[0].OutputNode = SourceGraph->FindEquivalentOutputNode(ENiagaraScriptUsage::ParticleSpawnScript, TranslationStages[0].UsageId);
-		TranslationStages[1].OutputNode = SourceGraph->FindEquivalentOutputNode(TranslationStages[1].ScriptUsage, TranslationStages[1].UsageId);
-		TranslationStages[1].bInterpolatePreviousParams = bHasInterpolatedSpawn;
-		TranslationStages[0].SimulationStageIndex = 0;
-		TranslationStages[0].NumIterations = 1;
-		TranslationStages[1].SimulationStageIndex = 0;
-		TranslationStages[1].NumIterations = 1;
-		TranslationStages[0].bWritesParticles = true;
-		TranslationStages[1].bWritesParticles = true;
-		TranslationStages[0].bShouldUpdateInitialAttributeValues = true;
+	{
+		FNiagaraHlslTranslationStage& SpawnStage = TranslationStages.Emplace_GetRef((bHasInterpolatedSpawn ? ENiagaraScriptUsage::ParticleSpawnScriptInterpolated : ENiagaraScriptUsage::ParticleSpawnScript), FGuid());
+		SpawnStage.PassNamespace = TEXT("MapSpawn");
+		SpawnStage.ChunkModeIndex = ENiagaraCodeChunkMode::SpawnBody;
+		SpawnStage.SimulationStageIndex = 0;
+		SpawnStage.NumIterations = 1;
+		SpawnStage.bWritesParticles = true;
+		SpawnStage.bShouldUpdateInitialAttributeValues = true;
+		OutputNodes.Add(SourceGraph->FindEquivalentOutputNode(ENiagaraScriptUsage::ParticleSpawnScript, TranslationStages[0].UsageId));
+
+		FNiagaraHlslTranslationStage& UpdateStage = TranslationStages.Emplace_GetRef(ENiagaraScriptUsage::ParticleUpdateScript, FGuid());
+		UpdateStage.PassNamespace = TEXT("MapUpdate");
+		UpdateStage.ChunkModeIndex = ENiagaraCodeChunkMode::UpdateBody;
+		UpdateStage.bInterpolatePreviousParams = bHasInterpolatedSpawn;
+		UpdateStage.SimulationStageIndex = 0;
+		UpdateStage.NumIterations = 1;
+		UpdateStage.bWritesParticles = true;
+		OutputNodes.Add(SourceGraph->FindEquivalentOutputNode(TranslationStages[1].ScriptUsage, TranslationStages[1].UsageId));
+
 		ParamMapHistories.AddDefaulted(2);
 		ParamMapHistoriesSourceInOtherHistories.AddDefaulted(2);
 		ParamMapSetVariablesToChunks.AddDefaulted(2);
@@ -1247,17 +1258,17 @@ const FNiagaraTranslateResults &FHlslNiagaraTranslator::Translate(const FNiagara
 			int32 SourceSimStageIndex = 0;
 			
 			// OutputNode order in traversal doesn't necessarily match the stack ordering. Use the GUID order to define the actual stages.
-			TArray<const UNiagaraNodeOutput*> FoundOutputNodes;
+			TArray<const FOutputNode*> FoundOutputNodes;
 			TArray<int32> FoundStageHistories;
 
-			for (const auto& CompileSimStageData : InCompileData->CompileSimStageData )
+			for (const auto& CompileSimStageData : CompileData->CompileSimStageData )
 			{
 				const FGuid& StageGuid = CompileSimStageData.StageGuid;
 				for (int32 FoundHistoryIdx = 0; FoundHistoryIdx < OtherOutputParamMapHistories.Num(); FoundHistoryIdx++)
 				{
-					FNiagaraParameterMapHistory& FoundHistory = OtherOutputParamMapHistories[FoundHistoryIdx];
-					const UNiagaraNodeOutput* HistoryOutputNode = FoundHistory.GetFinalOutputNode();
-					if (HistoryOutputNode && HistoryOutputNode->GetUsageId() == StageGuid)
+					FParamMapHistory& FoundHistory = OtherOutputParamMapHistories[FoundHistoryIdx];
+					const FOutputNode* HistoryOutputNode = FoundHistory.GetFinalOutputNode();
+					if (HistoryOutputNode && GraphBridge::GetOutputNodeUsageId(HistoryOutputNode) == StageGuid)
 					{
 						FoundOutputNodes.Add(HistoryOutputNode);
 						FoundStageHistories.Add(FoundHistoryIdx);
@@ -1269,70 +1280,74 @@ const FNiagaraTranslateResults &FHlslNiagaraTranslator::Translate(const FNiagara
 			// Now iterate the nodes in the order we found them.
 			for (int32 FoundIdx = 0; FoundIdx < FoundOutputNodes.Num(); FoundIdx++)
 			{
-				const UNiagaraNodeOutput* HistoryOutputNode = FoundOutputNodes[FoundIdx];
-				const FNiagaraParameterMapHistory& FoundHistory = OtherOutputParamMapHistories[FoundStageHistories[FoundIdx]];
+				const FOutputNode* HistoryOutputNode = FoundOutputNodes[FoundIdx];
+				const FParamMapHistory& FoundHistory = OtherOutputParamMapHistories[FoundStageHistories[FoundIdx]];
 
-				if (HistoryOutputNode && HistoryOutputNode->ScriptType == ENiagaraScriptUsage::ParticleSimulationStageScript)
+				if (HistoryOutputNode && GraphBridge::GetOutputNodeScriptType(HistoryOutputNode) == ENiagaraScriptUsage::ParticleSimulationStageScript)
 				{
-					const auto& CompileSimStageData = InCompileData->CompileSimStageData[SourceSimStageIndex];
+					const auto& CompileSimStageData = CompileData->CompileSimStageData[SourceSimStageIndex];
 
 					FString StageName = CompileSimStageData.StageName.ToString();
 					StageName = TEXT("_") + GetSanitizedFunctionNameSuffix(StageName);
-					const int32 Index = TranslationStages.Add(FHlslNiagaraTranslationStage(HistoryOutputNode->ScriptType, HistoryOutputNode->ScriptTypeId));
+
+					const int32 TranslationStageIndex = TranslationStages.Num();
+					FNiagaraHlslTranslationStage& SimStage = TranslationStages.Emplace_GetRef(GraphBridge::GetOutputNodeScriptType(HistoryOutputNode), GraphBridge::GetOutputNodeScriptTypeId(HistoryOutputNode));
 
 					const int32 DestSimStageIndex = CompilationOutput.ScriptData.SimulationStageMetaData.Num();
-					TranslationStages[Index].PassNamespace = FString::Printf(TEXT("MapSimStage%d%s"), DestSimStageIndex, *StageName);
-					TranslationStages[Index].ChunkModeIndex = (ENiagaraCodeChunkMode)(((int32)ENiagaraCodeChunkMode::SimulationStageBody) + (Index - 2));
-					if (TranslationStages[Index].ChunkModeIndex >= ENiagaraCodeChunkMode::SimulationStageBodyMax)
+					SimStage.PassNamespace = FString::Printf(TEXT("MapSimStage%d%s"), DestSimStageIndex, *StageName);
+					SimStage.ChunkModeIndex = (ENiagaraCodeChunkMode)(((int32)ENiagaraCodeChunkMode::SimulationStageBody) + (TranslationStageIndex - 2));
+					if (SimStage.ChunkModeIndex >= ENiagaraCodeChunkMode::SimulationStageBodyMax)
 					{
 						Error(FText::Format(LOCTEXT("TooManySimulationStages", "Cannot support more than %d simulation stages when adding %d!"),
 							FText::AsNumber((int32)ENiagaraCodeChunkMode::SimulationStageBodyMax - (int32)ENiagaraCodeChunkMode::SimulationStageBody),
-							FText::AsNumber((int32)TranslationStages[Index].ChunkModeIndex)), nullptr, nullptr);
+							FText::AsNumber((int32)SimStage.ChunkModeIndex)), nullptr, nullptr);
 					}
-					TranslationStages[Index].OutputNode = SourceGraph->FindEquivalentOutputNode(TranslationStages[Index].ScriptUsage, TranslationStages[Index].UsageId);
-					ensure(TranslationStages[Index].OutputNode == HistoryOutputNode);
-					TranslationStages[Index].bInterpolatePreviousParams = false;
-					TranslationStages[Index].bCopyPreviousParams = false;
-					TranslationStages[Index].SimulationStageIndex = DestSimStageIndex;					
-					TranslationStages[Index].EnabledBinding = CompileSimStageData.EnabledBinding;
-					TranslationStages[Index].ElementCount = CompileSimStageData.ElementCount;
-					TranslationStages[Index].ElementCountXBinding = CompileSimStageData.ElementCountXBinding;
-					TranslationStages[Index].ElementCountYBinding = CompileSimStageData.ElementCountYBinding;
-					TranslationStages[Index].ElementCountZBinding = CompileSimStageData.ElementCountZBinding;
-					TranslationStages[Index].NumIterations = CompileSimStageData.NumIterations;
-					TranslationStages[Index].ExecuteBehavior = CompileSimStageData.ExecuteBehavior;
-					TranslationStages[Index].bPartialParticleUpdate = CompileSimStageData.PartialParticleUpdate;
-					TranslationStages[Index].IterationSourceType = CompileSimStageData.IterationSourceType;
-					TranslationStages[Index].IterationDataInterface = CompileSimStageData.IterationDataInterface;
-					TranslationStages[Index].IterationDirectBinding = CompileSimStageData.IterationDirectBinding;
-					TranslationStages[Index].NumIterationsBinding = CompileSimStageData.NumIterationsBinding;
-					TranslationStages[Index].bParticleIterationStateEnabled = CompileSimStageData.bParticleIterationStateEnabled;
-					TranslationStages[Index].ParticleIterationStateBinding = CompileSimStageData.ParticleIterationStateBinding;
-					TranslationStages[Index].ParticleIterationStateRange = CompileSimStageData.ParticleIterationStateRange;
-					TranslationStages[Index].bGpuDispatchForceLinear = CompileSimStageData.bGpuDispatchForceLinear;
-					TranslationStages[Index].DirectDispatchType = CompileSimStageData.DirectDispatchType;
-					TranslationStages[Index].DirectDispatchElementType = CompileSimStageData.DirectDispatchElementType;
-					TranslationStages[Index].bOverrideGpuDispatchNumThreads = CompileSimStageData.bOverrideGpuDispatchNumThreads;
-					TranslationStages[Index].OverrideGpuDispatchNumThreads = CompileSimStageData.OverrideGpuDispatchNumThreads;
+					SimStage.bInterpolatePreviousParams = false;
+					SimStage.bCopyPreviousParams = false;
+					SimStage.SimulationStageIndex = DestSimStageIndex;					
+					SimStage.EnabledBinding = CompileSimStageData.EnabledBinding;
+					SimStage.ElementCount = CompileSimStageData.ElementCount;
+					SimStage.ElementCountXBinding = CompileSimStageData.ElementCountXBinding;
+					SimStage.ElementCountYBinding = CompileSimStageData.ElementCountYBinding;
+					SimStage.ElementCountZBinding = CompileSimStageData.ElementCountZBinding;
+					SimStage.NumIterations = CompileSimStageData.NumIterations;
+					SimStage.ExecuteBehavior = CompileSimStageData.ExecuteBehavior;
+					SimStage.bPartialParticleUpdate = CompileSimStageData.PartialParticleUpdate;
+					SimStage.IterationSourceType = CompileSimStageData.IterationSourceType;
+					SimStage.IterationDataInterface = CompileSimStageData.IterationDataInterface;
+					SimStage.IterationDirectBinding = CompileSimStageData.IterationDirectBinding;
+					SimStage.NumIterationsBinding = CompileSimStageData.NumIterationsBinding;
+					SimStage.bParticleIterationStateEnabled = CompileSimStageData.bParticleIterationStateEnabled;
+					SimStage.ParticleIterationStateBinding = CompileSimStageData.ParticleIterationStateBinding;
+					SimStage.ParticleIterationStateRange = CompileSimStageData.ParticleIterationStateRange;
+					SimStage.bGpuDispatchForceLinear = CompileSimStageData.bGpuDispatchForceLinear;
+					SimStage.DirectDispatchType = CompileSimStageData.DirectDispatchType;
+					SimStage.DirectDispatchElementType = CompileSimStageData.DirectDispatchElementType;
+					SimStage.bOverrideGpuDispatchNumThreads = CompileSimStageData.bOverrideGpuDispatchNumThreads;
+					SimStage.OverrideGpuDispatchNumThreads = CompileSimStageData.OverrideGpuDispatchNumThreads;
 
 					if (CompileSimStageData.bOverrideGpuDispatchNumThreads)
 					{
-						TranslationStages[Index].OverrideGpuDispatchNumThreads.X = CompileData->GetStaticVariableValue(CompileSimStageData.OverrideGpuDispatchNumThreadsXBinding, TranslationStages[Index].OverrideGpuDispatchNumThreads.X);
-						TranslationStages[Index].OverrideGpuDispatchNumThreads.Y = CompileData->GetStaticVariableValue(CompileSimStageData.OverrideGpuDispatchNumThreadsYBinding, TranslationStages[Index].OverrideGpuDispatchNumThreads.Y);
-						TranslationStages[Index].OverrideGpuDispatchNumThreads.Z = CompileData->GetStaticVariableValue(CompileSimStageData.OverrideGpuDispatchNumThreadsZBinding, TranslationStages[Index].OverrideGpuDispatchNumThreads.Z);
+						SimStage.OverrideGpuDispatchNumThreads.X = CompileData->GetStaticVariableValue(CompileSimStageData.OverrideGpuDispatchNumThreadsXBinding, SimStage.OverrideGpuDispatchNumThreads.X);
+						SimStage.OverrideGpuDispatchNumThreads.Y = CompileData->GetStaticVariableValue(CompileSimStageData.OverrideGpuDispatchNumThreadsYBinding, SimStage.OverrideGpuDispatchNumThreads.Y);
+						SimStage.OverrideGpuDispatchNumThreads.Z = CompileData->GetStaticVariableValue(CompileSimStageData.OverrideGpuDispatchNumThreadsZBinding, SimStage.OverrideGpuDispatchNumThreads.Z);
 					}
 
+					const FOutputNode* OutputNode = SourceGraph->FindEquivalentOutputNode(SimStage.ScriptUsage, SimStage.UsageId);
+					ensure(OutputNode == HistoryOutputNode);
+					OutputNodes.Add(OutputNode);
+
 					ParamMapHistories.AddDefaulted(1);
 					ParamMapHistoriesSourceInOtherHistories.AddDefaulted(1);
 
 					// If we allow partial writes we need to ensure that we are not reading from our own buffer, we ask our data interfaces if this is true or not
-					if (TranslationStages[Index].bPartialParticleUpdate)
+					if (SimStage.bPartialParticleUpdate)
 					{
-						for (const FNiagaraCompileRequestData::FCompileDataInterfaceData& DataInterfaceData : (*InCompileData->SharedCompileDataInterfaceData.Get()))
+						for (const typename FPrecompileData::FCompileDataInterfaceData& DataInterfaceData : (*CompileData->SharedCompileDataInterfaceData.Get()))
 						{
-							if(DataInterfaceData.ReadsEmitterParticleData.Contains(InCompileData->EmitterUniqueName))
+							if(DataInterfaceData.ReadsEmitterParticleData.Contains(CompileData->EmitterUniqueName))
 							{
-								TranslationStages[Index].bPartialParticleUpdate = false;
+								SimStage.bPartialParticleUpdate = false;
 								break;
 							}
 						}
@@ -1342,15 +1357,15 @@ const FNiagaraTranslateResults &FHlslNiagaraTranslator::Translate(const FNiagara
 					for (int32 iVar = 0; iVar < FoundHistory.VariableMetaData.Num(); ++iVar)
 					{
 						// Particle attribute?
-						if (!FNiagaraParameterMapHistory::IsAttribute(FoundHistory.Variables[iVar]))
+						if (!FParamMapHistory::IsAttribute(FoundHistory.Variables[iVar]))
 						{
 							continue;
 						}
 
 						// Is this an output?
-						const bool bIsOutput = FoundHistory.PerVariableWriteHistory[iVar].ContainsByPredicate([](const FModuleScopedPin& InPin) -> bool
+						const bool bIsOutput = FoundHistory.PerVariableWriteHistory[iVar].ContainsByPredicate([](const FScopedPin& InPin) -> bool
 						{
-							return Cast<UNiagaraNodeParameterMapSet>(InPin.Pin->GetOwningNode()) != nullptr;
+							return GraphBridge::template AsNodeType<FParamMapSetNode>(GraphBridge::GetOwningNode(InPin.Pin)) != nullptr;
 						});
 
 						if (!bIsOutput)
@@ -1368,32 +1383,32 @@ const FNiagaraTranslateResults &FHlslNiagaraTranslator::Translate(const FNiagara
 						}
 
 						// We write particle attributes at this stage, store list off so we can potentially selectivly write them later
-						TranslationStages[Index].bWritesParticles = true;
-						TranslationStages[Index].SetParticleAttributes.Add(FoundHistory.Variables[iVar]);
+						SimStage.bWritesParticles = true;
+						SimStage.SetParticleAttributes.Add(FoundHistory.Variables[iVar]);
 					}
 
 					// If we don't write particles then disable particle updates, it's meaningless and produces different HLSL since we would use a RW buffer not plain old Input
-					TranslationStages[Index].bPartialParticleUpdate &= TranslationStages[Index].bWritesParticles;
+					SimStage.bPartialParticleUpdate &= SimStage.bWritesParticles;
 					
 					// Set up the compile output for the shader stages so that we can properly execute at runtime.
 					FSimulationStageMetaData& SimulationStageMetaData = CompilationOutput.ScriptData.SimulationStageMetaData.AddDefaulted_GetRef();
-					SimulationStageMetaData.SimulationStageName = InCompileData->CompileSimStageData[SourceSimStageIndex].StageName;
-					SimulationStageMetaData.EnabledBinding = InCompileData->CompileSimStageData[SourceSimStageIndex].EnabledBinding;
-					SimulationStageMetaData.ElementCount = InCompileData->CompileSimStageData[SourceSimStageIndex].ElementCount;
-					SimulationStageMetaData.ElementCountXBinding = InCompileData->CompileSimStageData[SourceSimStageIndex].ElementCountXBinding;
-					SimulationStageMetaData.ElementCountYBinding = InCompileData->CompileSimStageData[SourceSimStageIndex].ElementCountYBinding;
-					SimulationStageMetaData.ElementCountZBinding = InCompileData->CompileSimStageData[SourceSimStageIndex].ElementCountZBinding;
-					SimulationStageMetaData.ExecuteBehavior = TranslationStages[Index].ExecuteBehavior;
-					SimulationStageMetaData.IterationSourceType = InCompileData->CompileSimStageData[SourceSimStageIndex].IterationSourceType;
-					SimulationStageMetaData.IterationDataInterface = InCompileData->CompileSimStageData[SourceSimStageIndex].IterationDataInterface;
-					SimulationStageMetaData.IterationDirectBinding = InCompileData->CompileSimStageData[SourceSimStageIndex].IterationDirectBinding;
-					SimulationStageMetaData.NumIterationsBinding = InCompileData->CompileSimStageData[SourceSimStageIndex].NumIterationsBinding;
-					SimulationStageMetaData.NumIterations = TranslationStages[Index].NumIterations;
-					SimulationStageMetaData.bWritesParticles = TranslationStages[Index].bWritesParticles;
-					SimulationStageMetaData.bPartialParticleUpdate = TranslationStages[Index].bPartialParticleUpdate;
-					SimulationStageMetaData.bParticleIterationStateEnabled = TranslationStages[Index].bParticleIterationStateEnabled;
-					SimulationStageMetaData.ParticleIterationStateBinding = TranslationStages[Index].ParticleIterationStateBinding;
-					SimulationStageMetaData.ParticleIterationStateRange = TranslationStages[Index].ParticleIterationStateRange;
+					SimulationStageMetaData.SimulationStageName = CompileData->CompileSimStageData[SourceSimStageIndex].StageName;
+					SimulationStageMetaData.EnabledBinding = CompileData->CompileSimStageData[SourceSimStageIndex].EnabledBinding;
+					SimulationStageMetaData.ElementCount = CompileData->CompileSimStageData[SourceSimStageIndex].ElementCount;
+					SimulationStageMetaData.ElementCountXBinding = CompileData->CompileSimStageData[SourceSimStageIndex].ElementCountXBinding;
+					SimulationStageMetaData.ElementCountYBinding = CompileData->CompileSimStageData[SourceSimStageIndex].ElementCountYBinding;
+					SimulationStageMetaData.ElementCountZBinding = CompileData->CompileSimStageData[SourceSimStageIndex].ElementCountZBinding;
+					SimulationStageMetaData.ExecuteBehavior = SimStage.ExecuteBehavior;
+					SimulationStageMetaData.IterationSourceType = CompileData->CompileSimStageData[SourceSimStageIndex].IterationSourceType;
+					SimulationStageMetaData.IterationDataInterface = CompileData->CompileSimStageData[SourceSimStageIndex].IterationDataInterface;
+					SimulationStageMetaData.IterationDirectBinding = CompileData->CompileSimStageData[SourceSimStageIndex].IterationDirectBinding;
+					SimulationStageMetaData.NumIterationsBinding = CompileData->CompileSimStageData[SourceSimStageIndex].NumIterationsBinding;
+					SimulationStageMetaData.NumIterations = SimStage.NumIterations;
+					SimulationStageMetaData.bWritesParticles = SimStage.bWritesParticles;
+					SimulationStageMetaData.bPartialParticleUpdate = SimStage.bPartialParticleUpdate;
+					SimulationStageMetaData.bParticleIterationStateEnabled = SimStage.bParticleIterationStateEnabled;
+					SimulationStageMetaData.ParticleIterationStateBinding = SimStage.ParticleIterationStateBinding;
+					SimulationStageMetaData.ParticleIterationStateRange = SimStage.ParticleIterationStateRange;
 
 					// Determine dispatch information from iteration source (if we have one)
 					SimulationStageMetaData.GpuDispatchType = ENiagaraGpuDispatchType::OneD;
@@ -1409,33 +1424,33 @@ const FNiagaraTranslateResults &FHlslNiagaraTranslator::Translate(const FNiagara
 
 						if (const FNiagaraVariable* IterationSourceVar = CompileData->EncounteredVariables.FindByPredicate([&](const FNiagaraVariable& VarInfo) { return VarInfo.GetName() == SimulationStageMetaData.IterationDataInterface; }))
 						{
-							if ( UNiagaraDataInterface* IteratinoSourceCDO = CompileDuplicateData->GetDuplicatedDataInterfaceCDOForClass(IterationSourceVar->GetType().GetClass()) )
+							if ( UNiagaraDataInterface* IterationSourceCDO = CompileDuplicateData->GetDuplicatedDataInterfaceCDOForClass(IterationSourceVar->GetType().GetClass()) )
 							{
-								SimulationStageMetaData.GpuDispatchType = TranslationStages[Index].bGpuDispatchForceLinear ? ENiagaraGpuDispatchType::OneD : IteratinoSourceCDO->GetGpuDispatchType();
-								SimulationStageMetaData.bGpuIndirectDispatch = IteratinoSourceCDO->GetGpuUseIndirectDispatch();
+								SimulationStageMetaData.GpuDispatchType = SimStage.bGpuDispatchForceLinear ? ENiagaraGpuDispatchType::OneD : IterationSourceCDO->GetGpuDispatchType();
+								SimulationStageMetaData.bGpuIndirectDispatch = IterationSourceCDO->GetGpuUseIndirectDispatch();
 
-								if (TranslationStages[Index].bOverrideGpuDispatchNumThreads)
+								if (SimStage.bOverrideGpuDispatchNumThreads)
 								{
-									SimulationStageMetaData.GpuDispatchNumThreads = TranslationStages[Index].OverrideGpuDispatchNumThreads;
+									SimulationStageMetaData.GpuDispatchNumThreads = SimStage.OverrideGpuDispatchNumThreads;
 								}
 								else
 								{
-									SimulationStageMetaData.GpuDispatchNumThreads = (SimulationStageMetaData.GpuDispatchType == ENiagaraGpuDispatchType::Custom) ? IteratinoSourceCDO->GetGpuDispatchNumThreads() : FNiagaraShader::GetDefaultThreadGroupSize(SimulationStageMetaData.GpuDispatchType);
+									SimulationStageMetaData.GpuDispatchNumThreads = (SimulationStageMetaData.GpuDispatchType == ENiagaraGpuDispatchType::Custom) ? IterationSourceCDO->GetGpuDispatchNumThreads() : FNiagaraShader::GetDefaultThreadGroupSize(SimulationStageMetaData.GpuDispatchType);
 								}
 							}
 						}
 					}
 					else if (SimulationStageMetaData.IterationSourceType == ENiagaraIterationSource::DirectSet)
 					{
-						SimulationStageMetaData.GpuDispatchType = TranslationStages[Index].DirectDispatchType;
-						SimulationStageMetaData.GpuDirectDispatchElementType = TranslationStages[Index].DirectDispatchElementType;
-						if (TranslationStages[Index].bOverrideGpuDispatchNumThreads)
+						SimulationStageMetaData.GpuDispatchType = SimStage.DirectDispatchType;
+						SimulationStageMetaData.GpuDirectDispatchElementType = SimStage.DirectDispatchElementType;
+						if (SimStage.bOverrideGpuDispatchNumThreads)
 						{
-							SimulationStageMetaData.GpuDispatchNumThreads = TranslationStages[Index].OverrideGpuDispatchNumThreads;
+							SimulationStageMetaData.GpuDispatchNumThreads = SimStage.OverrideGpuDispatchNumThreads;
 						}
 						else
 						{
-							SimulationStageMetaData.GpuDispatchNumThreads = FNiagaraShader::GetDefaultThreadGroupSize(TranslationStages[Index].DirectDispatchType);
+							SimulationStageMetaData.GpuDispatchNumThreads = FNiagaraShader::GetDefaultThreadGroupSize(SimStage.DirectDispatchType);
 						}
 					}
 
@@ -1447,52 +1462,58 @@ const FNiagaraTranslateResults &FHlslNiagaraTranslator::Translate(const FNiagara
 				}
 			}
 		}
-		break;
+	}	break;
 	default:
-		TranslationStages.Add(FHlslNiagaraTranslationStage(CompileOptions.TargetUsage, CompileOptions.TargetUsageId));
-		TranslationStages[0].PassNamespace = TEXT("Map");
-		TranslationStages[0].OutputNode = SourceGraph->FindEquivalentOutputNode(TranslationStages[0].ScriptUsage, TranslationStages[0].UsageId);
-		TranslationStages[0].ChunkModeIndex = ENiagaraCodeChunkMode::Body;
-		TranslationStages[0].SimulationStageIndex = 0;
-		TranslationStages[0].NumIterations = 1;
-		TranslationStages[0].bWritesParticles = true;
-		TranslationStages[0].bShouldUpdateInitialAttributeValues = TranslationStages[0].ShouldDoSpawnOnlyLogic() || (IsEventSpawnScript() && CompileOptions.AdditionalDefines.Contains(FNiagaraCompileOptions::EventSpawnInitialAttribWritesDefine));
+	{
+		FNiagaraHlslTranslationStage& DefaultStage = TranslationStages.Emplace_GetRef(CompileOptions.TargetUsage, CompileOptions.TargetUsageId);
+		DefaultStage.PassNamespace = TEXT("Map");
+		DefaultStage.ChunkModeIndex = ENiagaraCodeChunkMode::Body;
+		DefaultStage.SimulationStageIndex = 0;
+		DefaultStage.NumIterations = 1;
+		DefaultStage.bWritesParticles = true;
+		DefaultStage.bShouldUpdateInitialAttributeValues = TranslationStages[0].ShouldDoSpawnOnlyLogic() || (IsEventSpawnScript() && CompileOptions.AdditionalDefines.Contains(FNiagaraCompileOptions::EventSpawnInitialAttribWritesDefine));
+
+		OutputNodes.Add(SourceGraph->FindEquivalentOutputNode(TranslationStages[0].ScriptUsage, TranslationStages[0].UsageId));
 
 		if (CompileOptions.TargetUsage == ENiagaraScriptUsage::ParticleSimulationStageScript)
 		{
-			for (int32 StageIdx = 0; StageIdx < InCompileData->CompileSimStageData.Num(); StageIdx++)
+			for (int32 StageIdx = 0; StageIdx < CompileData->CompileSimStageData.Num(); StageIdx++)
 			{
-				const FGuid& StageGuid = InCompileData->CompileSimStageData[StageIdx].StageGuid;
-				if (StageGuid == CompileOptions.TargetUsageId && InCompileData->CompileSimStageData.IsValidIndex(StageIdx))
+				const FGuid& StageGuid = CompileData->CompileSimStageData[StageIdx].StageGuid;
+				if (StageGuid == CompileOptions.TargetUsageId && CompileData->CompileSimStageData.IsValidIndex(StageIdx))
 				{
-					TranslationStages[0].IterationSourceType = InCompileData->CompileSimStageData[StageIdx].IterationSourceType;
-					TranslationStages[0].IterationDataInterface = InCompileData->CompileSimStageData[StageIdx].IterationDataInterface;
-					TranslationStages[0].IterationDirectBinding = InCompileData->CompileSimStageData[StageIdx].IterationDirectBinding;
+					DefaultStage.IterationSourceType = CompileData->CompileSimStageData[StageIdx].IterationSourceType;
+					DefaultStage.IterationDataInterface = CompileData->CompileSimStageData[StageIdx].IterationDataInterface;
+					DefaultStage.IterationDirectBinding = CompileData->CompileSimStageData[StageIdx].IterationDirectBinding;
 				}
 			}
 		}
 		ParamMapHistories.AddDefaulted(1);
 		ParamMapHistoriesSourceInOtherHistories.AddDefaulted(1);
 		ParamMapSetVariablesToChunks.AddDefaulted(1);
-		break;
+	}	break;
+
 	}
 
 	for (int32 i = 0; i < TranslationStages.Num(); i++)
 	{
-		if (TranslationStages[i].OutputNode == nullptr)
+		const FOutputNode* OutputNode = OutputNodes[i];
+
+		if (OutputNode == nullptr)
 		{
 			Error(FText::Format(LOCTEXT("GetOutputNodeFail", "Cannot find output node of type {0}!"), FText::AsNumber((int32)TranslationStages[i].ScriptUsage)), nullptr, nullptr);
 			return TranslateResults;
 		}
 		
-		ValidateTypePins(TranslationStages[i].OutputNode);
+		ValidateTypePins(OutputNode);
 		{
 			bool bHasAnyConnections = false;
-			for (int32 PinIdx = 0; PinIdx < TranslationStages[i].OutputNode->Pins.Num(); PinIdx++)
+			for (const FInputPin* InputPin : GraphBridge::GetInputPins(OutputNode))
 			{
-				if (TranslationStages[i].OutputNode->Pins[PinIdx]->Direction == EEdGraphPinDirection::EGPD_Input && TranslationStages[i].OutputNode->Pins[PinIdx]->LinkedTo.Num() != 0)
+				if (GraphBridge::GetLinkedOutputPin(InputPin) != nullptr)
 				{
 					bHasAnyConnections = true;
+					break;
 				}
 			}
 			if (!bHasAnyConnections)
@@ -1538,10 +1559,10 @@ const FNiagaraTranslateResults &FHlslNiagaraTranslator::Translate(const FNiagara
 		
 		for (int32 HistoryIdx = 0; HistoryIdx < OtherOutputParamMapHistories.Num(); HistoryIdx++)
 		{
-			FNiagaraParameterMapHistory& FoundHistory = OtherOutputParamMapHistories[HistoryIdx];
+			FParamMapHistory& FoundHistory = OtherOutputParamMapHistories[HistoryIdx];
 
-			const UNiagaraNodeOutput* HistoryOutputNode = FoundHistory.GetFinalOutputNode();
-			if (HistoryOutputNode != nullptr && !ShouldConsiderTargetParameterMap(HistoryOutputNode->GetUsage()))
+			const FOutputNode* HistoryOutputNode = FoundHistory.GetFinalOutputNode();
+			if (HistoryOutputNode != nullptr && !ShouldConsiderTargetParameterMap(GraphBridge::GetOutputNodeScriptType(HistoryOutputNode)))
 			{
 				continue;
 			}
@@ -1549,7 +1570,7 @@ const FNiagaraTranslateResults &FHlslNiagaraTranslator::Translate(const FNiagara
 			// Now see if we want to use any of these specifically..
 			for (int32 ParamMapIdx = 0; ParamMapIdx < TranslationStages.Num(); ParamMapIdx++)
 			{
-				UNiagaraNodeOutput* TargetOutputNode = TranslationStages[ParamMapIdx].OutputNode;
+				const FOutputNode* TargetOutputNode = OutputNodes[ParamMapIdx];
 				if (FoundHistory.GetFinalOutputNode() == TargetOutputNode)
 				{
 					if (bRequiresPersistentIDs)
@@ -1589,12 +1610,9 @@ const FNiagaraTranslateResults &FHlslNiagaraTranslator::Translate(const FNiagara
 	}
 
 	CompilationOutput.ScriptData.ParameterCollectionPaths.Empty();
-	for (FNiagaraParameterMapHistory& History : ParamMapHistories)
+	for (FParamMapHistory& History : ParamMapHistories)
 	{
-		for (UNiagaraParameterCollection* Collection : History.ParameterCollections)
-		{
-			CompilationOutput.ScriptData.ParameterCollectionPaths.AddUnique(FSoftObjectPath(Collection).ToString());
-		}
+		GraphBridge::AddCollectionPaths(History, CompilationOutput.ScriptData.ParameterCollectionPaths);
 	}
 	ENiagaraScriptUsage Usage = CompileOptions.TargetUsage;
 	if (Usage != ENiagaraScriptUsage::SystemSpawnScript && Usage != ENiagaraScriptUsage::SystemUpdateScript && Usage != ENiagaraScriptUsage::Module && Usage != ENiagaraScriptUsage::DynamicInput)
@@ -1620,7 +1638,6 @@ const FNiagaraTranslateResults &FHlslNiagaraTranslator::Translate(const FNiagara
 	CompilationOutput.ScriptData.StatScopes.Empty();
 	EnterStatsScope(FNiagaraStatScope(*CompileOptions.GetName(), *CompileOptions.GetName()));
 
-	FHlslNiagaraTranslator* ThisTranslator = this;
 	TArray<int32> OutputChunks;
 
 	bool bInterpolateParams = false;
@@ -1641,7 +1658,7 @@ const FNiagaraTranslateResults &FHlslNiagaraTranslator::Translate(const FNiagara
 				CurrentParamMapIndices.Empty();
 				CurrentParamMapIndices.Add(0);
 				ActiveHistoryForFunctionCalls.BeginUsage(TranslationStages[i].ScriptUsage);
-				TranslationStages[i].OutputNode->Compile(ThisTranslator, OutputChunks);
+				OutputNodes[i]->Compile(this, OutputChunks);
 				ActiveHistoryForFunctionCalls.EndUsage();
 				InstanceWrite = FDataSetAccessInfo(); // Reset after building the output..
 				AddBodyComment(TEXT("//End Spawn Script!\n\n"));
@@ -1669,7 +1686,7 @@ const FNiagaraTranslateResults &FHlslNiagaraTranslator::Translate(const FNiagara
 				CurrentParamMapIndices.Empty();
 				CurrentParamMapIndices.Add(1);
 				ActiveHistoryForFunctionCalls.BeginUsage(TranslationStages[i].ScriptUsage);
-				TranslationStages[i].OutputNode->Compile(ThisTranslator, OutputChunks);
+				OutputNodes[i]->Compile(this, OutputChunks);
 				ActiveHistoryForFunctionCalls.EndUsage();
 				AddBodyComment(TEXT("//End Update Script!\n\n"));
 			}
@@ -1683,7 +1700,7 @@ const FNiagaraTranslateResults &FHlslNiagaraTranslator::Translate(const FNiagara
 				PinToCodeChunks.AddDefaulted(1);				
 				FName IterSource = TranslationStages[i].GetIterationDataInterface();
 				ActiveHistoryForFunctionCalls.BeginUsage(TranslationStages[i].ScriptUsage, IterSource);
-				TranslationStages[i].OutputNode->Compile(ThisTranslator, OutputChunks);
+				OutputNodes[i]->Compile(this, OutputChunks);
 				HandleSimStageSetupAndTeardown(i, StageSetupAndTeardownHLSL);
 				ActiveHistoryForFunctionCalls.EndUsage();
 				AddBodyComment(FString::Printf(TEXT("//End Simulation Stage Script: %s\n\n"), *TranslationStages[i].PassNamespace));
@@ -1703,7 +1720,7 @@ const FNiagaraTranslateResults &FHlslNiagaraTranslator::Translate(const FNiagara
 
 		FName IterSource = TranslationStages[0].GetIterationDataInterface();
 		ActiveHistoryForFunctionCalls.BeginUsage(TranslationStages[0].ScriptUsage, IterSource);
-		TranslationStages[0].OutputNode->Compile(ThisTranslator, OutputChunks);
+		OutputNodes[0]->Compile(this, OutputChunks);
 		ActiveHistoryForFunctionCalls.EndUsage();
 
 		bool bIsEventSpawn = IsEventSpawnScript();
@@ -2075,9 +2092,9 @@ const FNiagaraTranslateResults &FHlslNiagaraTranslator::Translate(const FNiagara
 		//Get full list of instance data accessed by the script as the VM binding assumes same for input and output.
 		for (FNiagaraVariable& Var : DataSetVariables[InstanceReadVarsIndex])
 		{
-			if (FNiagaraParameterMapHistory::IsAttribute(Var))
+			if (FParamMapHistory::IsAttribute(Var))
 			{
-				FNiagaraVariable BasicAttribVar = FNiagaraParameterMapHistory::ResolveAsBasicAttribute(Var, false);
+				FNiagaraVariable BasicAttribVar = FParamMapHistory::ResolveAsBasicAttribute(Var, false);
 				CompilationOutput.ScriptData.Attributes.AddUnique(BasicAttribVar);
 			}
 			else
@@ -2180,9 +2197,10 @@ const FNiagaraTranslateResults &FHlslNiagaraTranslator::Translate(const FNiagara
 	return TranslateResults;
 }
 
-void FHlslNiagaraTranslator::HandleSimStageSetupAndTeardown(int32 InWhichStage, FString& OutHlsl)
+template<typename GraphBridge>
+void TNiagaraHlslTranslator<GraphBridge>::HandleSimStageSetupAndTeardown(int32 InWhichStage, FString& OutHlsl)
 {
-	FHlslNiagaraTranslationStage& TranslationStage = TranslationStages[InWhichStage];
+	FNiagaraHlslTranslationStage& TranslationStage = TranslationStages[InWhichStage];
 	const FName IterationDataInterface = TranslationStage.GetIterationDataInterface();
 
 	// If we're particles then do nothing different..
@@ -2442,17 +2460,16 @@ void FHlslNiagaraTranslator::HandleSimStageSetupAndTeardown(int32 InWhichStage,
 		TeardownFunctionHLSL += TEXT("\n}\n");
 		OutHlsl += SetupFunctionHLSL + TEXT("\n\n") + TeardownFunctionHLSL + TEXT("\n\n");
 	}
-
 }
 
 
-void FHlslNiagaraTranslator::GatherVariableForDataSetAccess(const FNiagaraVariable& Var, FString Format, int32& IntCounter, int32 &FloatCounter, int32& HalfCounter, int32 DataSetIndex, FString InstanceIdxSymbol, FString &HlslOutputString, bool bWriteHLSL)
+void FNiagaraHlslTranslator::GatherVariableForDataSetAccess(const FNiagaraVariable& Var, FString Format, int32& IntCounter, int32 &FloatCounter, int32& HalfCounter, int32 DataSetIndex, FString InstanceIdxSymbol, FString &HlslOutputString, bool bWriteHLSL)
 {
 	TArray<FString> Components;
 	UScriptStruct* Struct = Var.GetType().GetScriptStruct();
 	if (!Struct)
 	{
-		Error(FText::Format(LOCTEXT("BadStructDef", "Variable {0} missing struct definition."), FText::FromName(Var.GetName())), nullptr, nullptr);
+		Error(FText::Format(LOCTEXT("BadStructDef", "Variable {0} missing struct definition."), FText::FromName(Var.GetName())));
 		return;
 	}
 
@@ -2512,7 +2529,7 @@ void FHlslNiagaraTranslator::GatherVariableForDataSetAccess(const FNiagaraVariab
 	}
 }
 
-void FHlslNiagaraTranslator::GatherComponentsForDataSetAccess(UScriptStruct* Struct, FString VariableSymbol, bool bMatrixRoot, TArray<FString>& Components, TArray<ENiagaraBaseTypes>& Types)
+void FNiagaraHlslTranslator::GatherComponentsForDataSetAccess(UScriptStruct* Struct, FString VariableSymbol, bool bMatrixRoot, TArray<FString>& Components, TArray<ENiagaraBaseTypes>& Types)
 {
 	bool bIsVector = IsHlslBuiltinVector(FNiagaraTypeDefinition(Struct));
 	bool bIsScalar = FNiagaraTypeDefinition::IsScalarDefinition(Struct);
@@ -2596,7 +2613,7 @@ void FHlslNiagaraTranslator::GatherComponentsForDataSetAccess(UScriptStruct* Str
 	}
 }
 
-void FHlslNiagaraTranslator::DefinePreviousParametersFunction(FString& HlslOutputString, TArray<TArray<FNiagaraVariable>>& DataSetVariables, TMap<FNiagaraDataSetID, int32>& DataSetReads, TMap<FNiagaraDataSetID, int32>& DataSetWrites)
+void FNiagaraHlslTranslator::DefinePreviousParametersFunction(FString& HlslOutputString, TArray<TArray<FNiagaraVariable>>& DataSetVariables, TMap<FNiagaraDataSetID, int32>& DataSetReads, TMap<FNiagaraDataSetID, int32>& DataSetWrites)
 {
 	HlslOutputString +=
 		TEXT("#if (SimulationStageIndex == 0) // MapSpawn\n")
@@ -2634,7 +2651,8 @@ void FHlslNiagaraTranslator::DefinePreviousParametersFunction(FString& HlslOutpu
 	HlslOutputString += TEXT("}\n#endif\n\n");
 }
 
-void FHlslNiagaraTranslator::DefineInterpolatedParametersFunction(FString &HlslOutputString)
+template<typename GraphBridge>
+void TNiagaraHlslTranslator<GraphBridge>::DefineInterpolatedParametersFunction(FString &HlslOutputString)
 {
 	for (int32 i = 0; i < TranslationStages.Num(); i++)
 	{
@@ -2724,7 +2742,7 @@ void FHlslNiagaraTranslator::DefineInterpolatedParametersFunction(FString &HlslO
 	}
 }
 
-void FHlslNiagaraTranslator::DefineDataSetReadFunction(FString &HlslOutputString, TArray<FNiagaraDataSetID> &ReadDataSets)
+void FNiagaraHlslTranslator::DefineDataSetReadFunction(FString &HlslOutputString, TArray<FNiagaraDataSetID> &ReadDataSets)
 {
 	if (UNiagaraScript::IsParticleEventScript(CompileOptions.TargetUsage) && CompilationTarget == ENiagaraSimTarget::GPUComputeSim)
 	{
@@ -2779,7 +2797,7 @@ void FHlslNiagaraTranslator::DefineDataSetReadFunction(FString &HlslOutputString
 }
 
 
-void FHlslNiagaraTranslator::DefineDataSetWriteFunction(FString &HlslOutputString, TArray<FNiagaraDataSetProperties> &WriteDataSets, TArray<int32>& WriteConditionVarIndices)
+void FNiagaraHlslTranslator::DefineDataSetWriteFunction(FString &HlslOutputString, TArray<FNiagaraDataSetProperties> &WriteDataSets, TArray<int32>& WriteConditionVarIndices)
 {
 	HlslOutputString += TEXT("void WriteDataSets(inout FSimulationContext Context)\n{\n");
 
@@ -2845,7 +2863,7 @@ void FHlslNiagaraTranslator::DefineDataSetWriteFunction(FString &HlslOutputStrin
 	HlslOutput += TEXT("}\n\n");
 }
 
-void FHlslNiagaraTranslator::ConvertCompileInfoToParamInfo(const FNiagaraScriptDataInterfaceCompileInfo& Info, FNiagaraDataInterfaceGPUParamInfo& DIInstanceInfo, TArray<FNiagaraFunctionSignature>& GeneratedFunctionSignatures)
+void FNiagaraHlslTranslator::ConvertCompileInfoToParamInfo(const FNiagaraScriptDataInterfaceCompileInfo& Info, FNiagaraDataInterfaceGPUParamInfo& DIInstanceInfo, TArray<FNiagaraFunctionSignature>& GeneratedFunctionSignatures)
 {
 	FString OwnerIDString = Info.Name.ToString();
 	FString SanitizedOwnerIDString = GetSanitizedSymbolName(OwnerIDString.Replace(TEXT("."), TEXT("_")));
@@ -2868,7 +2886,7 @@ void FHlslNiagaraTranslator::ConvertCompileInfoToParamInfo(const FNiagaraScriptD
 
 		if (!OriginalSig.bSupportsGPU)
 		{
-			Error(FText::Format(LOCTEXT("GPUDataInterfaceFunctionNotSupported", "DataInterface {0} function {1} cannot run on the GPU."), FText::FromName(Info.Type.GetFName()), FText::FromName(OriginalSig.Name)), nullptr, nullptr);
+			Error(FText::Format(LOCTEXT("GPUDataInterfaceFunctionNotSupported", "DataInterface {0} function {1} cannot run on the GPU."), FText::FromName(Info.Type.GetFName()), FText::FromName(OriginalSig.Name)));
 			continue;
 		}
 		if (OriginalSig.bWriteFunction)
@@ -2913,7 +2931,7 @@ void FHlslNiagaraTranslator::ConvertCompileInfoToParamInfo(const FNiagaraScriptD
 	}
 }
 
-void FHlslNiagaraTranslator::DefineDataInterfaceHLSL(FString& InHlslOutput)
+void FNiagaraHlslTranslator::DefineDataInterfaceHLSL(FString& InHlslOutput)
 {
 	FString InterfaceCommonHLSL;
 	FString InterfaceUniformHLSL;
@@ -2924,7 +2942,7 @@ void FHlslNiagaraTranslator::DefineDataInterfaceHLSL(FString& InHlslOutput)
 	{
 		FNiagaraScriptDataInterfaceCompileInfo& Info = CompilationOutput.ScriptData.DataInterfaceInfo[i];
 
-		UNiagaraDataInterface* CDO = CompileDuplicateData->GetDuplicatedDataInterfaceCDOForClass(Info.Type.GetClass());
+		UNiagaraDataInterface* CDO = GetDataInterfaceCDO(Info.Type.GetClass());
 		check(CDO != nullptr);
 		if (CDO && CDO->CanExecuteOnTarget(ENiagaraSimTarget::GPUComputeSim))
 		{
@@ -2940,10 +2958,10 @@ void FHlslNiagaraTranslator::DefineDataInterfaceHLSL(FString& InHlslOutput)
 
 			FNiagaraDataInterfaceHlslGenerationContext DIHlslGenContext(DIInstanceInfo, GeneratedFunctionSignatures);
 
-			DIHlslGenContext.GetFunctionSignatureSymbolDelegate.BindStatic(&FHlslNiagaraTranslator::GetFunctionSignatureSymbol);
-			DIHlslGenContext.GetStructHlslTypeNameDelegate.BindStatic(&FHlslNiagaraTranslator::GetStructHlslTypeName);
-			DIHlslGenContext.GetPropertyHlslTypeNameDelegate.BindStatic(&FHlslNiagaraTranslator::GetPropertyHlslTypeName);
-			DIHlslGenContext.GetSanitizedSymbolNameDelegate.BindStatic(&FHlslNiagaraTranslator::GetSanitizedSymbolName);
+			DIHlslGenContext.GetFunctionSignatureSymbolDelegate.BindStatic(&FNiagaraHlslTranslator::GetFunctionSignatureSymbol);
+			DIHlslGenContext.GetStructHlslTypeNameDelegate.BindStatic(&FNiagaraHlslTranslator::GetStructHlslTypeName);
+			DIHlslGenContext.GetPropertyHlslTypeNameDelegate.BindStatic(&FNiagaraHlslTranslator::GetPropertyHlslTypeName);
+			DIHlslGenContext.GetSanitizedSymbolNameDelegate.BindStatic(&FNiagaraHlslTranslator::GetSanitizedSymbolName);
 			CDO->GetParameterDefinitionHLSL(DIHlslGenContext, InterfaceUniformHLSL);
 
 			// Ask the DI to generate HLSL.
@@ -2957,7 +2975,7 @@ void FHlslNiagaraTranslator::DefineDataInterfaceHLSL(FString& InHlslOutput)
 				const bool HlslOK = CDO->GetFunctionHLSL(DIHlslGenContext, InterfaceFunctionHLSL);
 				if (!HlslOK)
 				{
-					Error(FText::Format(LOCTEXT("GPUDataInterfaceFunctionNotImplemented", "DataInterface {0} function {1} is not implemented for GPU."), FText::FromName(Info.Type.GetFName()), FText::FromName(DIFunc.DefinitionName)), nullptr, nullptr);
+					Error(FText::Format(LOCTEXT("GPUDataInterfaceFunctionNotImplemented", "DataInterface {0} function {1} is not implemented for GPU."), FText::FromName(Info.Type.GetFName()), FText::FromName(DIFunc.DefinitionName)));
 				}
 				else
 				{
@@ -2967,24 +2985,25 @@ void FHlslNiagaraTranslator::DefineDataInterfaceHLSL(FString& InHlslOutput)
 		}
 		else
 		{
-			Error(FText::Format(LOCTEXT("NonGPUDataInterfaceError", "DataInterface {0} ({1}) cannot run on the GPU."), FText::FromName(Info.Name), FText::FromString(CDO ? CDO->GetClass()->GetName() : TEXT(""))), nullptr, nullptr);
+			Error(FText::Format(LOCTEXT("NonGPUDataInterfaceError", "DataInterface {0} ({1}) cannot run on the GPU."), FText::FromName(Info.Name), FText::FromString(CDO ? CDO->GetClass()->GetName() : TEXT(""))));
 		}
 	}
 	InHlslOutput += InterfaceCommonHLSL + InterfaceUniformHLSL + InterfaceFunctionHLSL;
 }
 
-void FHlslNiagaraTranslator::DefineExternalFunctionsHLSL(FString &InHlslOutput)
+void FNiagaraHlslTranslator::DefineExternalFunctionsHLSL(FString &InHlslOutput)
 {
 	for (FNiagaraFunctionSignature& FunctionSig : CompilationOutput.ScriptData.AdditionalExternalFunctions )
 	{
 		if ( UNiagaraFunctionLibrary::DefineFunctionHLSL(FunctionSig, InHlslOutput) == false )
 		{
-			Error(FText::Format(LOCTEXT("ExternFunctionMissingHLSL", "ExternalFunction {0} does not have a HLSL implementation for the GPU."), FText::FromName(FunctionSig.Name)), nullptr, nullptr);
+			Error(FText::Format(LOCTEXT("ExternFunctionMissingHLSL", "ExternalFunction {0} does not have a HLSL implementation for the GPU."), FText::FromName(FunctionSig.Name)));
 		}
 	}
 }
 
-void FHlslNiagaraTranslator::DefineMainGPUFunctions(
+template<typename GraphBridge>
+void TNiagaraHlslTranslator<GraphBridge>::DefineMainGPUFunctions(
 	const TArray<TArray<FNiagaraVariable>>& DataSetVariables,
 	const TMap<FNiagaraDataSetID, int32>& DataSetReads,
 	const TMap<FNiagaraDataSetID, int32>& DataSetWrites)
@@ -3159,7 +3178,7 @@ void FHlslNiagaraTranslator::DefineMainGPUFunctions(
 					{
 						VarFmt = VarName + TEXT("{0} = {4};\n");
 					}
-					else if (DataSetIndex == 0 && FNiagaraParameterMapHistory::IsPreviousValue(Var) && TranslationStages[i].ScriptUsage == ENiagaraScriptUsage::ParticleUpdateScript)
+					else if (DataSetIndex == 0 && FParamMapHistory::IsPreviousValue(Var) && TranslationStages[i].ScriptUsage == ENiagaraScriptUsage::ParticleUpdateScript)
 					{
 						GatheredPreviousVariables.AddUnique(Var);
 						bWrite = false; // We need to bump the read indices forwards, but not actually add the read.
@@ -3168,7 +3187,7 @@ void FHlslNiagaraTranslator::DefineMainGPUFunctions(
 					{
 						VarFmt = VarName + TEXT("{0} = InputData{1}({2}, {3}, InstanceIdx);\n");
 
-						if (FNiagaraParameterMapHistory::IsAttribute(Var))
+						if (FParamMapHistory::IsAttribute(Var))
 						{
 							FString RegisterName = VarName;
 							ReplaceNamespaceInline(RegisterName, PARAM_MAP_ATTRIBUTE_STR, PARAM_MAP_INDICES_STR);
@@ -3184,7 +3203,7 @@ void FHlslNiagaraTranslator::DefineMainGPUFunctions(
 			// Put any gathered previous variables into the list here so that we can use them by recording the last value from the parent variable on load.
 			for (FNiagaraVariable VarPrevious : GatheredPreviousVariables)
 			{
-				FNiagaraVariable SrcVar = FNiagaraParameterMapHistory::GetSourceForPreviousValue(VarPrevious);
+				FNiagaraVariable SrcVar = FParamMapHistory::GetSourceForPreviousValue(VarPrevious);
 				const FString VarName = ContextName + GetSanitizedSymbolName(SrcVar.GetName().ToString());
 				const FString VarPrevName = ContextName + GetSanitizedSymbolName(VarPrevious.GetName().ToString());
 				HlslOutput += VarPrevName + TEXT(" = ") + VarName + TEXT(";\n");
@@ -3245,9 +3264,9 @@ void FHlslNiagaraTranslator::DefineMainGPUFunctions(
 					for (const FVarAndDefaultSource& VarAndDefaultSource : VarAndDefaultSourceArray)
 					{
 						const FNiagaraVariable& Var = VarAndDefaultSource.Variable;
-						if (FNiagaraParameterMapHistory::IsPreviousValue(Var))
+						if (FParamMapHistory::IsPreviousValue(Var))
 						{
-							FNiagaraVariable SrcVar = FNiagaraParameterMapHistory::GetSourceForPreviousValue(Var);
+							FNiagaraVariable SrcVar = FParamMapHistory::GetSourceForPreviousValue(Var);
 							const FString VarName = GetSanitizedSymbolName(SrcVar.GetName().ToString());
 							const FString VarPrevName = GetSanitizedSymbolName(Var.GetName().ToString());
 							HlslOutput += TEXT("\t\tContext.") + TranslationStages[i].PassNamespace + TEXT(".") + VarPrevName + TEXT(" = Context.") + TranslationStages[i-1].PassNamespace + TEXT(".") + VarName + TEXT(";\n");
@@ -3423,7 +3442,7 @@ void FHlslNiagaraTranslator::DefineMainGPUFunctions(
 	// Generate each translation stages body
 	for (int32 i=1; i < TranslationStages.Num(); i++)
 	{
-		const FHlslNiagaraTranslationStage& TranslationStage = TranslationStages[i];
+		const FNiagaraHlslTranslationStage& TranslationStage = TranslationStages[i];
 		const bool bInterpolatedSpawning = CompileOptions.AdditionalDefines.Contains(TEXT("InterpolatedSpawn")) || i != 1;
 		const bool bAlwaysRunUpdateScript = CompileOptions.AdditionalDefines.Contains(TEXT("GpuAlwaysRunParticleUpdateScript"));
 		const bool bParticleSpawnStage = i == 1;
@@ -3620,7 +3639,7 @@ void FHlslNiagaraTranslator::DefineMainGPUFunctions(
 	HlslOutput.Append(TEXT("}\n"));
 }
 
-void FHlslNiagaraTranslator::DefineMain(FString &OutHlslOutput,
+void FNiagaraHlslTranslator::DefineMain(FString &OutHlslOutput,
 	const TArray<TArray<FNiagaraVariable>>& DataSetVariables,
 	const TMap<FNiagaraDataSetID, int32>& DataSetReads,
 	const TMap<FNiagaraDataSetID, int32>& DataSetWrites)
@@ -3669,10 +3688,10 @@ void FHlslNiagaraTranslator::DefineMain(FString &OutHlslOutput,
 		// Manually write to Particles.UniqueID on spawn, and deliberately place it at the top of SimulateMain to make sure it's initialized in the right order
 
 		// NOTE(mv): These relies on Particles.UniqueID and Engine.Emitter.TotalSpawnedParticles both being explicitly added to the parameter histories in 
-		//           FHlslNiagaraTranslator::Translate.
+		//           FNiagaraHlslTranslator::Translate.
 
 		// NOTE(mv): This relies on Particles.UniqueID being excluded from being default initialized. 
-		//           This happens in FNiagaraParameterMapHistory::ShouldIgnoreVariableDefault
+		//           This happens in FParamMapHistory::ShouldIgnoreVariableDefault
 		if (UNiagaraScript::IsParticleSpawnScript(CompileOptions.TargetUsage))
 		{
 			FString MapName = UNiagaraScript::IsInterpolatedParticleSpawnScript(CompileOptions.TargetUsage) ? TEXT("Context.MapSpawn") : TEXT("Context.Map");
@@ -3765,7 +3784,8 @@ void FHlslNiagaraTranslator::DefineMain(FString &OutHlslOutput,
 	OutHlslOutput += TEXT("}\n");
 }
 
-void FHlslNiagaraTranslator::DefineDataSetVariableWrites(FString &OutHlslOutput, const FNiagaraDataSetID& Id, int32 DataSetIndex, const TArray<FNiagaraVariable>& WriteVars)
+template<typename GraphBridge>
+void TNiagaraHlslTranslator<GraphBridge>::DefineDataSetVariableWrites(FString &OutHlslOutput, const FNiagaraDataSetID& Id, int32 DataSetIndex, const TArray<FNiagaraVariable>& WriteVars)
 {
 	check(CompilationTarget != ENiagaraSimTarget::GPUComputeSim);
 
@@ -3780,7 +3800,7 @@ void FHlslNiagaraTranslator::DefineDataSetVariableWrites(FString &OutHlslOutput,
 		bool bHasPerParticleAliveEvent = false;
 		for (int32 i = 0; i < ParamMapHistories.Num(); i++)
 		{
-			const UNiagaraNodeOutput* OutputNode = ParamMapHistories[i].GetFinalOutputNode();
+			const FOutputNode* OutputNode = ParamMapHistories[i].GetFinalOutputNode();
 			if (!OutputNode)
 			{
 				continue;
@@ -3791,7 +3811,7 @@ void FHlslNiagaraTranslator::DefineDataSetVariableWrites(FString &OutHlslOutput,
 				continue;
 			}
 
-			switch (OutputNode->GetUsage())
+			switch (GraphBridge::GetOutputNodeScriptType(OutputNode))
 			{
 			case ENiagaraScriptUsage::ParticleSpawnScript:
 			case ENiagaraScriptUsage::ParticleSpawnScriptInterpolated:
@@ -3862,7 +3882,7 @@ void FHlslNiagaraTranslator::DefineDataSetVariableWrites(FString &OutHlslOutput,
 	OutHlslOutput += "\t}\n";
 }
 
-void FHlslNiagaraTranslator::DefineDataSetVariableReads(FString &OutHlslOutput, const FNiagaraDataSetID& Id, int32 DataSetIndex, const TArray<FNiagaraVariable>& ReadVars)
+void FNiagaraHlslTranslator::DefineDataSetVariableReads(FString &OutHlslOutput, const FNiagaraDataSetID& Id, int32 DataSetIndex, const TArray<FNiagaraVariable>& ReadVars)
 {
 	check(CompilationTarget != ENiagaraSimTarget::GPUComputeSim);
 
@@ -3974,7 +3994,7 @@ void FHlslNiagaraTranslator::DefineDataSetVariableReads(FString &OutHlslOutput,
 	}
 }
 
-void FHlslNiagaraTranslator::WriteDataSetContextVars(TMap<FNiagaraDataSetID, TMap<int32, FDataSetAccessInfo>>& DataSetAccessInfo, bool bRead, FString& OutHlslOutput)
+void FNiagaraHlslTranslator::WriteDataSetContextVars(TMap<FNiagaraDataSetID, TMap<int32, FDataSetAccessInfo>>& DataSetAccessInfo, bool bRead, FString& OutHlslOutput)
 {
 	//Now the intermediate storage for the data set reads and writes.
 	uint32 DataSetIndex = 0;
@@ -3992,7 +4012,7 @@ void FHlslNiagaraTranslator::WriteDataSetContextVars(TMap<FNiagaraDataSetID, TMa
 };
 
 
-void FHlslNiagaraTranslator::WriteDataSetStructDeclarations(TMap<FNiagaraDataSetID, TMap<int32, FDataSetAccessInfo>>& DataSetAccessInfo, bool bRead, FString& OutHlslOutput)
+void FNiagaraHlslTranslator::WriteDataSetStructDeclarations(TMap<FNiagaraDataSetID, TMap<int32, FDataSetAccessInfo>>& DataSetAccessInfo, bool bRead, FString& OutHlslOutput)
 {
 	uint32 DataSetIndex = 1;
 	for (TPair<FNiagaraDataSetID, TMap<int32, FDataSetAccessInfo>>& DataSetInfoPair : DataSetAccessInfo)
@@ -4040,7 +4060,7 @@ void FHlslNiagaraTranslator::WriteDataSetStructDeclarations(TMap<FNiagaraDataSet
 
 
 //Decomposes each variable into its constituent register accesses.
-void FHlslNiagaraTranslator::DecomposeVariableAccess(UStruct* Struct, bool bRead, FString IndexSymbol, FString HLSLString)
+void FNiagaraHlslTranslator::DecomposeVariableAccess(UStruct* Struct, bool bRead, FString IndexSymbol, FString HLSLString)
 {
 	FString AccessStr;
 
@@ -4082,11 +4102,7 @@ void FHlslNiagaraTranslator::DecomposeVariableAccess(UStruct* Struct, bool bRead
 	}
 };
 
-void FHlslNiagaraTranslator::Init()
-{
-}
-
-FString FHlslNiagaraTranslator::GetSanitizedSymbolName(FStringView SymbolName, bool bCollapsNamespaces)
+FString FNiagaraHlslTranslator::GetSanitizedSymbolName(FStringView SymbolName, bool bCollapsNamespaces)
 {
 	if (SymbolName.Len() == 0)
 	{
@@ -4146,7 +4162,7 @@ FString FHlslNiagaraTranslator::GetSanitizedSymbolName(FStringView SymbolName, b
 	return Ret;
 }
 
-FString FHlslNiagaraTranslator::GetSanitizedDIFunctionName(const FString& FunctionName)
+FString FNiagaraHlslTranslator::GetSanitizedDIFunctionName(const FString& FunctionName)
 {
 	bool bWordStart = true;
 	FString Sanitized;
@@ -4172,7 +4188,7 @@ FString FHlslNiagaraTranslator::GetSanitizedDIFunctionName(const FString& Functi
 	return Sanitized;
 }
 
-FString FHlslNiagaraTranslator::GetSanitizedFunctionNameSuffix(FString Name)
+FString FNiagaraHlslTranslator::GetSanitizedFunctionNameSuffix(FString Name)
 {
 	if (Name.Len() == 0)
 	{
@@ -4195,7 +4211,7 @@ FString FHlslNiagaraTranslator::GetSanitizedFunctionNameSuffix(FString Name)
 	return ConvertToAsciiString(Ret);
 }
 
-FString FHlslNiagaraTranslator::ConvertToAsciiString(FString Str)
+FString FNiagaraHlslTranslator::ConvertToAsciiString(FString Str)
 {
 	FString AsciiString;
 	AsciiString.Reserve(Str.Len() * 6); // Assign room for every current char to be 'ASCXXX'
@@ -4219,7 +4235,7 @@ FString FHlslNiagaraTranslator::ConvertToAsciiString(FString Str)
 	return AsciiString;
 }
 
-FString FHlslNiagaraTranslator::GetUniqueSymbolName(FName BaseName)
+FString FNiagaraHlslTranslator::GetUniqueSymbolName(FName BaseName)
 {
 	FString RetString = GetSanitizedSymbolName(BaseName.ToString());
 	FName RetName = *RetString;
@@ -4238,7 +4254,8 @@ FString FHlslNiagaraTranslator::GetUniqueSymbolName(FName BaseName)
 	return RetString;
 }
 
-void FHlslNiagaraTranslator::EnterFunction(const FString& Name, FNiagaraFunctionSignature& Signature, TArrayView<const int32> Inputs, const FGuid& InGuid)
+template<typename GraphBridge>
+void TNiagaraHlslTranslator<GraphBridge>::EnterFunction(const FString& Name, FNiagaraFunctionSignature& Signature, TArrayView<const int32> Inputs, const FGuid& InGuid)
 {
 	FunctionContextStack.Emplace(Name, Signature, Inputs, InGuid);
 	TArray<FName> Entries;
@@ -4250,7 +4267,8 @@ void FHlslNiagaraTranslator::EnterFunction(const FString& Name, FNiagaraFunction
 	PinToCodeChunks.AddDefaulted(1);
 }
 
-void FHlslNiagaraTranslator::ExitFunction()
+template<typename GraphBridge>
+void TNiagaraHlslTranslator<GraphBridge>::ExitFunction()
 {
 	FunctionContextStack.Pop();
 	//May need some more heavy and scoped symbol tracking?
@@ -4279,7 +4297,7 @@ void FHlslNiagaraTranslator::ExitFunction()
 	}
 }
 
-FString FHlslNiagaraTranslator::GeneratedConstantString(float Constant)
+FString FNiagaraHlslTranslator::GeneratedConstantString(float Constant)
 {
 	return LexToString(Constant);
 }
@@ -4292,7 +4310,7 @@ static FAutoConsoleVariableRef CVarNiagaraScriptStatTracking(
 	ECVF_Default
 );
 
-void FHlslNiagaraTranslator::EnterStatsScope(FNiagaraStatScope StatScope)
+void FNiagaraHlslTranslator::EnterStatsScope(FNiagaraStatScope StatScope)
 {
 	if (GbNiagaraScriptStatTracking)
 	{
@@ -4302,7 +4320,7 @@ void FHlslNiagaraTranslator::EnterStatsScope(FNiagaraStatScope StatScope)
 	}
 }
 
-void FHlslNiagaraTranslator::ExitStatsScope()
+void FNiagaraHlslTranslator::ExitStatsScope()
 {
 	if (GbNiagaraScriptStatTracking)
 	{
@@ -4311,7 +4329,7 @@ void FHlslNiagaraTranslator::ExitStatsScope()
 	}
 }
 
-void FHlslNiagaraTranslator::EnterStatsScope(FNiagaraStatScope StatScope, FString& OutHlsl)
+void FNiagaraHlslTranslator::EnterStatsScope(FNiagaraStatScope StatScope, FString& OutHlsl)
 {
 	if (GbNiagaraScriptStatTracking)
 	{
@@ -4321,7 +4339,7 @@ void FHlslNiagaraTranslator::EnterStatsScope(FNiagaraStatScope StatScope, FStrin
 	}
 }
 
-void FHlslNiagaraTranslator::ExitStatsScope(FString& OutHlsl)
+void FNiagaraHlslTranslator::ExitStatsScope(FString& OutHlsl)
 {
 	if (GbNiagaraScriptStatTracking)
 	{
@@ -4330,7 +4348,7 @@ void FHlslNiagaraTranslator::ExitStatsScope(FString& OutHlsl)
 	}
 }
 
-FString FHlslNiagaraTranslator::GetCallstack()
+FString FNiagaraHlslTranslator::GetCallstack()
 {
 	FString Callstack = CompileOptions.GetName();
 
@@ -4342,7 +4360,7 @@ FString FHlslNiagaraTranslator::GetCallstack()
 	return Callstack;
 }
 
-TArray<FGuid> FHlslNiagaraTranslator::GetCallstackGuids()
+TArray<FGuid> FNiagaraHlslTranslator::GetCallstackGuids()
 {
 	TArray<FGuid> Callstack;
 	for (FFunctionContext& Ctx : FunctionContextStack)
@@ -4353,7 +4371,7 @@ TArray<FGuid> FHlslNiagaraTranslator::GetCallstackGuids()
 	return Callstack;
 }
 
-FString FHlslNiagaraTranslator::GeneratedConstantString(FVector4 Constant)
+FString FNiagaraHlslTranslator::GeneratedConstantString(FVector4 Constant)
 {
 	TArray<FStringFormatArg> Args;
 	Args.Add(LexToString(Constant.X));
@@ -4363,7 +4381,7 @@ FString FHlslNiagaraTranslator::GeneratedConstantString(FVector4 Constant)
 	return FString::Format(TEXT("float4({0}, {1}, {2}, {3})"), Args);
 }
 
-int32 FHlslNiagaraTranslator::AddUniformChunk(FString SymbolName, const FNiagaraVariable& InVariable, ENiagaraCodeChunkMode ChunkMode, bool AddPadding)
+int32 FNiagaraHlslTranslator::AddUniformChunk(FString SymbolName, const FNiagaraVariable& InVariable, ENiagaraCodeChunkMode ChunkMode, bool AddPadding)
 {
 	const FNiagaraTypeDefinition& Type = InVariable.GetType();
 
@@ -4409,7 +4427,7 @@ int32 FHlslNiagaraTranslator::AddUniformChunk(FString SymbolName, const FNiagara
 	return Ret;
 }
 
-int32 FHlslNiagaraTranslator::AddSourceChunk(FString SymbolName, const FNiagaraTypeDefinition& Type, bool bSanitize)
+int32 FNiagaraHlslTranslator::AddSourceChunk(FString SymbolName, const FNiagaraTypeDefinition& Type, bool bSanitize)
 {
 	int32 Ret = CodeChunks.IndexOfByPredicate(
 		[&](const FNiagaraCodeChunk& Chunk)
@@ -4433,17 +4451,17 @@ int32 FHlslNiagaraTranslator::AddSourceChunk(FString SymbolName, const FNiagaraT
 }
 
 
-int32 FHlslNiagaraTranslator::AddBodyComment(const FString& Comment)
+int32 FNiagaraHlslTranslator::AddBodyComment(const FString& Comment)
 {
 	return AddBodyChunk(TEXT(""), Comment, FNiagaraTypeDefinition::GetIntDef(), false, false);
 }
 
-int32 FHlslNiagaraTranslator::AddBodyChunk(const FString& Value)
+int32 FNiagaraHlslTranslator::AddBodyChunk(const FString& Value)
 {
 	return AddBodyChunk(TEXT(""), Value, FNiagaraTypeDefinition::GetIntDef(), INDEX_NONE, false, false);
 }
 
-int32 FHlslNiagaraTranslator::AddBodyChunk(FString SymbolName, FString Definition, const FNiagaraTypeDefinition& Type, TArray<int32>& SourceChunks, bool bDecl, bool bIsTerminated)
+int32 FNiagaraHlslTranslator::AddBodyChunk(FString SymbolName, FString Definition, const FNiagaraTypeDefinition& Type, TArray<int32>& SourceChunks, bool bDecl, bool bIsTerminated)
 {
 	check(CurrentBodyChunkMode == ENiagaraCodeChunkMode::Body || CurrentBodyChunkMode == ENiagaraCodeChunkMode::SpawnBody || CurrentBodyChunkMode == ENiagaraCodeChunkMode::UpdateBody ||
 		(CurrentBodyChunkMode >= ENiagaraCodeChunkMode::SimulationStageBody && CurrentBodyChunkMode < ENiagaraCodeChunkMode::SimulationStageBodyMax));
@@ -4464,7 +4482,7 @@ int32 FHlslNiagaraTranslator::AddBodyChunk(FString SymbolName, FString Definitio
 
 
 
-int32 FHlslNiagaraTranslator::AddBodyChunk(FString SymbolName, FString Definition, const FNiagaraTypeDefinition& Type, int32 SourceChunk, bool bDecl, bool bIsTerminated)
+int32 FNiagaraHlslTranslator::AddBodyChunk(FString SymbolName, FString Definition, const FNiagaraTypeDefinition& Type, int32 SourceChunk, bool bDecl, bool bIsTerminated)
 {
 	check(CurrentBodyChunkMode == ENiagaraCodeChunkMode::Body || CurrentBodyChunkMode == ENiagaraCodeChunkMode::SpawnBody || CurrentBodyChunkMode == ENiagaraCodeChunkMode::UpdateBody ||
 		(CurrentBodyChunkMode >= ENiagaraCodeChunkMode::SimulationStageBody && CurrentBodyChunkMode < ENiagaraCodeChunkMode::SimulationStageBodyMax));
@@ -4483,7 +4501,7 @@ int32 FHlslNiagaraTranslator::AddBodyChunk(FString SymbolName, FString Definitio
 	return Ret;
 }
 
-int32 FHlslNiagaraTranslator::AddBodyChunk(FString SymbolName, FString Definition, const FNiagaraTypeDefinition& Type, bool bDecl, bool bIsTerminated)
+int32 FNiagaraHlslTranslator::AddBodyChunk(FString SymbolName, FString Definition, const FNiagaraTypeDefinition& Type, bool bDecl, bool bIsTerminated)
 {
 	check(CurrentBodyChunkMode == ENiagaraCodeChunkMode::Body || CurrentBodyChunkMode == ENiagaraCodeChunkMode::SpawnBody || CurrentBodyChunkMode == ENiagaraCodeChunkMode::UpdateBody ||
 			(CurrentBodyChunkMode >= ENiagaraCodeChunkMode::SimulationStageBody && CurrentBodyChunkMode < ENiagaraCodeChunkMode::SimulationStageBodyMax));
@@ -4501,7 +4519,7 @@ int32 FHlslNiagaraTranslator::AddBodyChunk(FString SymbolName, FString Definitio
 	return Ret;
 }
 
-bool FHlslNiagaraTranslator::ShouldInterpolateParameter(const FNiagaraVariable& Parameter)
+bool FNiagaraHlslTranslator::ShouldInterpolateParameter(const FNiagaraVariable& Parameter)
 {
 	//TODO: Some data driven method of deciding what parameters to interpolate and how to do it.
 	//Possibly allow definition of a dynamic input for the interpolation?
@@ -4548,7 +4566,9 @@ bool FHlslNiagaraTranslator::ShouldInterpolateParameter(const FNiagaraVariable&
 	return true;
 }
 
-void FHlslNiagaraTranslator::UpdateStaticSwitchConstants(UEdGraphNode* Node)
+// specialization handling the case where the translation process treats the graph as mutable
+template<>
+void TNiagaraHlslTranslator<FNiagaraCompilationGraphBridge>::UpdateStaticSwitchConstants(FNode* Node)
 {
 	if (UNiagaraNodeStaticSwitch* SwitchNode = Cast<UNiagaraNodeStaticSwitch>(Node))
 	{
@@ -4561,12 +4581,12 @@ void FHlslNiagaraTranslator::UpdateStaticSwitchConstants(UEdGraphNode* Node)
 		for (int i = 0; i < NodesToUpdate.Num(); i++)
 		{
 			SwitchNode->UpdateCompilerConstantValue(this);
-			
+
 			// also check direct upstream static switches, because they are otherwise skipped during the compilation and
 			// might be evaluated without their values set correctly.
 			InPins.Reset();
 			SwitchNode->GetInputPins(InPins);
-			for (UEdGraphPin* Pin : InPins)
+			for (FPin* Pin : InPins)
 			{
 				if (UNiagaraNodeStaticSwitch* ConnectedNode = Cast<UNiagaraNodeStaticSwitch>(Pin->GetOwningNode()))
 				{
@@ -4577,7 +4597,8 @@ void FHlslNiagaraTranslator::UpdateStaticSwitchConstants(UEdGraphNode* Node)
 	}
 }
 
-int32 FHlslNiagaraTranslator::GetRapidIterationParameter(const FNiagaraVariable& Parameter)
+template<typename GraphBridge>
+int32 TNiagaraHlslTranslator<GraphBridge>::GetRapidIterationParameter(const FNiagaraVariable& Parameter)
 {
 	if (!AddStructToDefinitionSet(Parameter.GetType()))
 	{
@@ -4636,13 +4657,12 @@ int32 FHlslNiagaraTranslator::GetRapidIterationParameter(const FNiagaraVariable&
 			return OutputChunkId;
 		}
 	}
-
 	return INDEX_NONE;
 }
 
-int32 FHlslNiagaraTranslator::GetParameter(const FNiagaraVariable& Parameter)
+template<typename GraphBridge>
+int32 TNiagaraHlslTranslator<GraphBridge>::GetParameter(const FNiagaraVariable& Parameter)
 {
-
 	NIAGARA_SCOPE_CYCLE_COUNTER(STAT_NiagaraEditor_HlslTranslator_GetParameter);
 	if (!AddStructToDefinitionSet(Parameter.GetType()))
 	{
@@ -4719,12 +4739,12 @@ int32 FHlslNiagaraTranslator::GetParameter(const FNiagaraVariable& Parameter)
 
 			for (int32 PinIdx = 0; PinIdx < ParamMapHistories[i].MapPinHistory.Num(); PinIdx++)
 			{
-				const UEdGraphPin* Pin = ParamMapHistories[i].MapPinHistory[PinIdx];
+				const FPin* Pin = ParamMapHistories[i].MapPinHistory[PinIdx];
 
 				if (Pin != nullptr)
 				{
-					UNiagaraNodeInput* InputNode = Cast<UNiagaraNodeInput>(Pin->GetOwningNode());
-					if (InputNode != nullptr && InputNode->Input == Parameter)
+					const FInputNode* InputNode = GraphBridge::template AsNodeType<FInputNode>(GraphBridge::GetOwningNode(Pin));
+					if (InputNode != nullptr && GraphBridge::GetInputVariable(InputNode) == Parameter)
 					{
 						if (CurrentDefaultPinTraversal.Num() == 0)
 						{
@@ -4745,7 +4765,7 @@ int32 FHlslNiagaraTranslator::GetParameter(const FNiagaraVariable& Parameter)
 	//Not a in a function or not a valid function parameter so grab from the main uniforms.
 	int32 OutputChunkIdx = INDEX_NONE;
 	FNiagaraVariable OutputVariable = Parameter;
-	if (FNiagaraParameterMapHistory::IsInNamespace(OutputVariable, PARAM_MAP_ATTRIBUTE_STR) || FNiagaraParameterMapHistory::IsExternalConstantNamespace(OutputVariable, CompileOptions.TargetUsage, CompileOptions.GetTargetUsageBitmask()))
+	if (FParamMapHistory::IsInNamespace(OutputVariable, PARAM_MAP_ATTRIBUTE_STR) || FParamMapHistory::IsExternalConstantNamespace(OutputVariable, CompileOptions.TargetUsage, CompileOptions.GetTargetUsageBitmask()))
 	{
 		if (!ParameterMapRegisterExternalConstantNamespaceVariable(OutputVariable, nullptr, 0, OutputChunkIdx, nullptr))
 		{
@@ -4754,7 +4774,7 @@ int32 FHlslNiagaraTranslator::GetParameter(const FNiagaraVariable& Parameter)
 	}
 	else
 	{
-		OutputVariable = FNiagaraParameterMapHistory::MoveToExternalConstantNamespaceVariable(OutputVariable, CompileOptions.TargetUsage);
+		OutputVariable = FParamMapHistory::MoveToExternalConstantNamespaceVariable(OutputVariable, CompileOptions.TargetUsage);
 		if (!ParameterMapRegisterExternalConstantNamespaceVariable(OutputVariable, nullptr, 0, OutputChunkIdx, nullptr))
 		{
 			OutputChunkIdx = INDEX_NONE;
@@ -4776,7 +4796,7 @@ int32 FHlslNiagaraTranslator::GetParameter(const FNiagaraVariable& Parameter)
 	return OutputChunkIdx;
 }
 
-int32 FHlslNiagaraTranslator::GetConstant(const FNiagaraVariable& Constant, FString* DebugOutputValue)
+int32 FNiagaraHlslTranslator::GetConstant(const FNiagaraVariable& Constant, FString* DebugOutputValue)
 {
 	if (Constant.IsDataInterface() || Constant.IsUObject())
 	{
@@ -4812,7 +4832,7 @@ int32 FHlslNiagaraTranslator::GetConstant(const FNiagaraVariable& Constant, FStr
 	return BodyChunk;
 }
 
-int32 FHlslNiagaraTranslator::GetConstantDirect(float InConstantValue)
+int32 FNiagaraHlslTranslator::GetConstantDirect(float InConstantValue)
 {
 	FNiagaraVariable Constant(FNiagaraTypeDefinition::GetFloatDef(), TEXT("Constant"));
 	Constant.SetValue(InConstantValue);
@@ -4820,7 +4840,7 @@ int32 FHlslNiagaraTranslator::GetConstantDirect(float InConstantValue)
 	return GetConstant(Constant);
 }
 
-int32 FHlslNiagaraTranslator::GetConstantDirect(bool InConstantValue)
+int32 FNiagaraHlslTranslator::GetConstantDirect(bool InConstantValue)
 {
 	FNiagaraVariable Constant(FNiagaraTypeDefinition::GetBoolDef(), TEXT("Constant"));
 	Constant.SetValue(InConstantValue);
@@ -4828,7 +4848,7 @@ int32 FHlslNiagaraTranslator::GetConstantDirect(bool InConstantValue)
 	return GetConstant(Constant);
 }
 
-int32 FHlslNiagaraTranslator::GetConstantDirect(int InConstantValue)
+int32 FNiagaraHlslTranslator::GetConstantDirect(int InConstantValue)
 {
 	FNiagaraVariable Constant(FNiagaraTypeDefinition::GetIntDef(), TEXT("Constant"));
 	Constant.SetValue(InConstantValue);
@@ -4836,7 +4856,7 @@ int32 FHlslNiagaraTranslator::GetConstantDirect(int InConstantValue)
 	return GetConstant(Constant);
 }
 
-bool FHlslNiagaraTranslator::GenerateStructInitializer(TStringBuilder<128>& InitializerString, UStruct* UserDefinedStruct, const void* StructData, int32 ByteOffset)
+bool FNiagaraHlslTranslator::GenerateStructInitializer(TStringBuilder<128>& InitializerString, UStruct* UserDefinedStruct, const void* StructData, int32 ByteOffset)
 {
 	//-TODO: Alignment Issues
 	//-TODO: Also adjust automation test NiagaraTest_CustomLWCStructHlsl when this is fixed 
@@ -4897,7 +4917,7 @@ bool FHlslNiagaraTranslator::GenerateStructInitializer(TStringBuilder<128>& Init
 		}
 		else
 		{
-			Error(FText::Format(LOCTEXT("GenerateConstantStructInitializeTypeError", "Unknown type '{0}' member '{1}' in structure '{2}' when generating initializer struct."), FText::FromString(ChildProperty->GetClass()->GetName()), FText::FromString(ChildProperty->GetName()), FText::FromString(UserDefinedStruct->GetName())), nullptr, nullptr);
+			Error(FText::Format(LOCTEXT("GenerateConstantStructInitializeTypeError", "Unknown type '{0}' member '{1}' in structure '{2}' when generating initializer struct."), FText::FromString(ChildProperty->GetClass()->GetName()), FText::FromString(ChildProperty->GetName()), FText::FromString(UserDefinedStruct->GetName())));
 			return false;
 		}
 	}
@@ -4905,12 +4925,12 @@ bool FHlslNiagaraTranslator::GenerateStructInitializer(TStringBuilder<128>& Init
 	return true;
 }
 
-FString FHlslNiagaraTranslator::GenerateConstantString(const FNiagaraVariable& Constant)
+FString FNiagaraHlslTranslator::GenerateConstantString(const FNiagaraVariable& Constant)
 {
 	FNiagaraTypeDefinition Type = Constant.GetType();
 	if (!AddStructToDefinitionSet(Type))
 	{
-		Error(FText::Format(LOCTEXT("GetConstantFail", "Cannot handle type {0}! Variable: {1}"), Type.GetNameText(), FText::FromName(Constant.GetName())), nullptr, nullptr);
+		Error(FText::Format(LOCTEXT("GetConstantFail", "Cannot handle type {0}! Variable: {1}"), Type.GetNameText(), FText::FromName(Constant.GetName())));
 	}
 	FString ConstantStr = GetHlslDefaultForType(Type);
 
@@ -4966,7 +4986,7 @@ FString FHlslNiagaraTranslator::GenerateConstantString(const FNiagaraVariable& C
 		{
 			if (Constant.GetValue<FNiagaraBool>().IsValid() == false)
 			{
-				Error(FText::Format(LOCTEXT("StructContantsBoolInvalidError", "Boolean constant {0} is not set to explicit True or False. Defaulting to False."), FText::FromName(Constant.GetName())), nullptr, nullptr);
+				Error(FText::Format(LOCTEXT("StructContantsBoolInvalidError", "Boolean constant {0} is not set to explicit True or False. Defaulting to False."), FText::FromName(Constant.GetName())));
 				ConstantStr = TEXT("false");
 			}
 			else
@@ -4980,21 +5000,21 @@ FString FHlslNiagaraTranslator::GenerateConstantString(const FNiagaraVariable& C
 			TStringBuilder<128> InitializerString;
 			if ( !GenerateStructInitializer(InitializerString, UserDefinedStruct, Constant.GetData()) )
 			{
-				Error(FText::Format(LOCTEXT("FailedToGenerateConstantInitialiezrError", "Type '{0}' constant '{1}' failed to create structure initializer. Defaulting to 0."), FText::FromString(Type.GetName()), FText::FromName(Constant.GetName())), nullptr, nullptr);
+				Error(FText::Format(LOCTEXT("FailedToGenerateConstantInitialiezrError", "Type '{0}' constant '{1}' failed to create structure initializer. Defaulting to 0."), FText::FromString(Type.GetName()), FText::FromName(Constant.GetName())));
 				return ConstantStr;
 			}
 			return InitializerString.ToString();
 		}
 		else
 		{
-			Warning(FText::Format(LOCTEXT("GenerateConstantUnknownTypeError", "Type '{0}' constant '{1}' is unknown.  Defaulting to 0."), FText::FromString(Type.GetName()), FText::FromName(Constant.GetName())), nullptr, nullptr);
+			Warning(FText::Format(LOCTEXT("GenerateConstantUnknownTypeError", "Type '{0}' constant '{1}' is unknown.  Defaulting to 0."), FText::FromString(Type.GetName()), FText::FromName(Constant.GetName())));
 			return ConstantStr;
 		}
 	}
 	return ConstantStr;
 }
 
-bool FHlslNiagaraTranslationStage::ShouldDoSpawnOnlyLogic() const
+bool FNiagaraHlslTranslationStage::ShouldDoSpawnOnlyLogic() const
 {
 	if (UNiagaraScript::IsSpawnScript(ScriptUsage))
 	{
@@ -5008,7 +5028,7 @@ bool FHlslNiagaraTranslationStage::ShouldDoSpawnOnlyLogic() const
 	return false;
 }
 
-bool FHlslNiagaraTranslationStage::IsExternalConstantNamespace(const FNiagaraVariable& InVar, ENiagaraScriptUsage InTargetUsage, uint32 InTargetBitmask)
+bool FNiagaraHlslTranslationStage::IsExternalConstantNamespace(const FNiagaraVariable& InVar, ENiagaraScriptUsage InTargetUsage, uint32 InTargetBitmask)
 {
 	if (FNiagaraParameterMapHistory::IsExternalConstantNamespace(InVar, InTargetUsage, InTargetBitmask))
 	{
@@ -5020,7 +5040,8 @@ bool FHlslNiagaraTranslationStage::IsExternalConstantNamespace(const FNiagaraVar
 	return false;
 }
 
-bool FHlslNiagaraTranslationStage::IsRelevantToSpawnForStage(const FNiagaraParameterMapHistory& InHistory, const FNiagaraVariable& InAliasedVar, const FNiagaraVariable& InVar) const
+template<typename GraphBridge>
+bool FNiagaraHlslTranslationStage::IsRelevantToSpawnForStage(const typename GraphBridge::FParamMapHistory& InHistory, const FNiagaraVariable& InAliasedVar, const FNiagaraVariable& InVar) const
 {
 	if (InHistory.IsPrimaryDataSetOutput(InAliasedVar, ScriptUsage) && (UNiagaraScript::IsSpawnScript(ScriptUsage) || bShouldUpdateInitialAttributeValues))
 	{
@@ -5040,7 +5061,8 @@ bool FHlslNiagaraTranslationStage::IsRelevantToSpawnForStage(const FNiagaraParam
 	return false;
 }
 
-void FHlslNiagaraTranslator::InitializeParameterMapDefaults(int32 ParamMapHistoryIdx)
+template<typename GraphBridge>
+void TNiagaraHlslTranslator<GraphBridge>::InitializeParameterMapDefaults(int32 ParamMapHistoryIdx)
 {
 	bInitializedDefaults = true;
 	AddBodyComment(TEXT("//Begin Initialize Parameter Map Defaults"));
@@ -5051,10 +5073,10 @@ void FHlslNiagaraTranslator::InitializeParameterMapDefaults(int32 ParamMapHistor
 	UniqueVarToWriteToParamMap.Empty();
 	UniqueVarToChunk.Empty();
 
-	FHlslNiagaraTranslationStage& ActiveStage = TranslationStages[ActiveStageIdx];
+	FNiagaraHlslTranslationStage& ActiveStage = TranslationStages[ActiveStageIdx];
 	// First pass just use the current parameter map.
 	{
-		const FNiagaraParameterMapHistory& History = ParamMapHistories[ParamMapHistoryIdx];
+		const typename GraphBridge::FParamMapHistory& History = ParamMapHistories[ParamMapHistoryIdx];
 		for (int32 i = 0; i < History.Variables.Num(); i++)
 		{
 			const FNiagaraVariable& Var = History.Variables[i];
@@ -5062,11 +5084,11 @@ void FHlslNiagaraTranslator::InitializeParameterMapDefaults(int32 ParamMapHistor
 			// Only add primary data set outputs at the top of the script if in a spawn script, otherwise they should be left alone.
 			if (TranslationStages[ActiveStageIdx].ShouldDoSpawnOnlyLogic() || TranslationStages[ActiveStageIdx].bShouldUpdateInitialAttributeValues)
 			{
-				if (TranslationStages[ActiveStageIdx].IsRelevantToSpawnForStage(History, AliasedVar, Var) &&
+				if (TranslationStages[ActiveStageIdx].template IsRelevantToSpawnForStage<GraphBridge>(History, AliasedVar, Var) &&
 					!UniqueVars.Contains(Var))
 				{
 					UniqueVars.Add(Var);
-					const UEdGraphPin* DefaultPin = History.GetDefaultValuePin(i);
+					const FInputPin* DefaultPin = History.GetDefaultValuePin(i);
 					UniqueVarToDefaultPin.Add(Var, DefaultPin);
 					UniqueVarToWriteToParamMap.Add(Var, true);
 				}
@@ -5082,16 +5104,16 @@ void FHlslNiagaraTranslator::InitializeParameterMapDefaults(int32 ParamMapHistor
 		// in the primary data set output namespaces.
 		for (int32 ParamMapIdx = 0; ParamMapIdx < OtherOutputParamMapHistories.Num(); ParamMapIdx++)
 		{
-			const FNiagaraParameterMapHistory& History = OtherOutputParamMapHistories[ParamMapIdx];
+			const typename GraphBridge::FParamMapHistory& History = OtherOutputParamMapHistories[ParamMapIdx];
 			for (int32 i = 0; i < History.Variables.Num(); i++)
 			{
 				const FNiagaraVariable& Var = History.Variables[i];
 				const FNiagaraVariable& AliasedVar = History.VariablesWithOriginalAliasesIntact[i];
-				if (TranslationStages[ActiveStageIdx].IsRelevantToSpawnForStage(History, AliasedVar, Var) &&
+				if (TranslationStages[ActiveStageIdx].template IsRelevantToSpawnForStage<GraphBridge>(History, AliasedVar, Var) &&
 					!UniqueVars.Contains(Var))
 				{
 					UniqueVars.Add(Var);
-					const UEdGraphPin* DefaultPin = History.GetDefaultValuePin(i);
+					const FInputPin* DefaultPin = History.GetDefaultValuePin(i);
 					UniqueVarToDefaultPin.Add(Var, DefaultPin);
 					UniqueVarToWriteToParamMap.Add(Var, false);
 				}
@@ -5102,29 +5124,34 @@ void FHlslNiagaraTranslator::InitializeParameterMapDefaults(int32 ParamMapHistor
 		// or defined by linkage or other script (defer to end if not originating from spawn, otherwise insert before first use)
 		for (FNiagaraVariable& Var : UniqueVars)
 		{
-			const UEdGraphPin* DefaultPin = UniqueVarToDefaultPin.FindChecked(Var);
+			const FInputPin* DefaultPin = UniqueVarToDefaultPin.FindChecked(Var);
 			bool bWriteToParamMapEntries = UniqueVarToWriteToParamMap.FindChecked(Var);
 			int32 OutputChunkId = INDEX_NONE;
 			
 			TOptional<ENiagaraDefaultMode> DefaultMode;
 			FNiagaraScriptVariableBinding DefaultBinding;
+			const FNode* DefaultPinNode = nullptr;
+
 			if (DefaultPin) 
 			{
-				if (UNiagaraGraph* DefaultPinGraph = CastChecked<UNiagaraGraph>(DefaultPin->GetOwningNode()->GetGraph())) 
+				DefaultPinNode = GraphBridge::GetOwningNode(DefaultPin);
+				if (const FGraph* DefaultPinGraph = GraphBridge::GetOwningGraph(DefaultPinNode))
 				{
-					DefaultMode = DefaultPinGraph->GetDefaultMode(Var, &DefaultBinding);
+					DefaultMode = GraphBridge::GetGraphDefaultMode(DefaultPinGraph, Var, DefaultBinding);
 				}
 			}
 
 			// During the initial pass, only support constants for the default pin and non-bound variables
-			if (!FNiagaraParameterMapHistory::IsInitialValue(Var) && (DefaultPin == nullptr || DefaultPin->LinkedTo.Num() == 0) && !(DefaultMode.IsSet() && (*DefaultMode == ENiagaraDefaultMode::Binding || *DefaultMode == ENiagaraDefaultMode::FailIfPreviouslyNotSet)))
+			const bool bDefaultPinConnected = GraphBridge::GetLinkedOutputPin(DefaultPin) != nullptr;
+
+			if (!GraphBridge::FParamMapHistory::IsInitialValue(Var) && !bDefaultPinConnected && !(DefaultMode.IsSet() && (*DefaultMode == ENiagaraDefaultMode::Binding || *DefaultMode == ENiagaraDefaultMode::FailIfPreviouslyNotSet)))
 			{
-				HandleParameterRead(ParamMapHistoryIdx, Var, DefaultPin, DefaultPin != nullptr ? Cast<UNiagaraNode>(DefaultPin->GetOwningNode()) : nullptr, OutputChunkId, TOptional<ENiagaraDefaultMode>(), TOptional<FNiagaraScriptVariableBinding>(), !bWriteToParamMapEntries);
+				HandleParameterRead(ParamMapHistoryIdx, Var, DefaultPin, DefaultPinNode, OutputChunkId, TOptional<ENiagaraDefaultMode>(), TOptional<FNiagaraScriptVariableBinding>(), !bWriteToParamMapEntries);
 				UniqueVarToChunk.Add(Var, OutputChunkId);
 			}
-			else if (FNiagaraParameterMapHistory::IsInitialValue(Var))
+			else if (GraphBridge::FParamMapHistory::IsInitialValue(Var))
 			{
-				FNiagaraVariable SourceForInitialValue = FNiagaraParameterMapHistory::GetSourceForInitialValue(Var);
+				FNiagaraVariable SourceForInitialValue = GraphBridge::FParamMapHistory::GetSourceForInitialValue(Var);
 				if (!UniqueVars.Contains(SourceForInitialValue))
 				{
 					//@todo(ng) disabled pending investigation UE-150159
@@ -5142,7 +5169,8 @@ void FHlslNiagaraTranslator::InitializeParameterMapDefaults(int32 ParamMapHistor
 	AddBodyComment(TEXT("//End Initialize Parameter Map Defaults"));
 }
 
-void FHlslNiagaraTranslator::Output(UNiagaraNodeOutput* OutputNode, const TArray<int32>& ComputedInputs)
+template<typename GraphBridge>
+void TNiagaraHlslTranslator<GraphBridge>::Output(const FOutputNode* OutputNode, const TArray<int32>& ComputedInputs)
 {
 	NIAGARA_SCOPE_CYCLE_COUNTER(STAT_NiagaraEditor_HlslTranslator_Output);
 
@@ -5150,7 +5178,7 @@ void FHlslNiagaraTranslator::Output(UNiagaraNodeOutput* OutputNode, const TArray
 	TArray<int32> Inputs;
 
 	// Build up the attribute list. We don't auto-expand parameter maps here.
-	TArray<FNiagaraVariable> Outputs = OutputNode->GetOutputs();
+	const TArray<FNiagaraVariable>& Outputs = GraphBridge::GetOutputVariables(OutputNode);
 	int32 NumberOfValidComputedInputs = 0;
 	for (int32 ComputedInput : ComputedInputs)
 	{
@@ -5172,7 +5200,7 @@ void FHlslNiagaraTranslator::Output(UNiagaraNodeOutput* OutputNode, const TArray
 		{
 			if (!AddStructToDefinitionSet(Attributes[i].GetType()))
 			{
-				Error(FText::Format(LOCTEXT("GetConstantFail", "Cannot handle type {0}! Variable: {1}"), Attributes[i].GetType().GetNameText(), FText::FromName(Attributes[i].GetName())), nullptr, nullptr);
+				Error(FText::Format(LOCTEXT("GetConstantFail", "Cannot handle type {0}! Variable: {1}"), Attributes[i].GetType().GetNameText(), FText::FromName(Attributes[i].GetName())));
 			}
 
 			if (Attributes[i].GetType() != FNiagaraTypeDefinition::GetParameterMapDef())
@@ -5199,7 +5227,7 @@ void FHlslNiagaraTranslator::Output(UNiagaraNodeOutput* OutputNode, const TArray
 
 				if (!AddStructToDefinitionSet(Var.GetType()))
 				{
-					Error(FText::Format(LOCTEXT("GetConstantFail", "Cannot handle type {0}! Variable: {1}"), Var.GetType().GetNameText(), FText::FromName(Var.GetName())), nullptr, nullptr);
+					Error(FText::Format(LOCTEXT("GetConstantFail", "Cannot handle type {0}! Variable: {1}"), Var.GetType().GetNameText(), FText::FromName(Var.GetName())));
 				}
 
 				//DATASET TODO: add and treat input 0 as the 'valid' input for conditional write
@@ -5208,7 +5236,7 @@ void FHlslNiagaraTranslator::Output(UNiagaraNodeOutput* OutputNode, const TArray
 
 				if (Var.GetType() != FNiagaraTypeDefinition::GetParameterMapDef())
 				{
-					FNiagaraVariable VarNamespaced = FNiagaraParameterMapHistory::BasicAttributeToNamespacedAttribute(Var);
+					FNiagaraVariable VarNamespaced = FParamMapHistory::BasicAttributeToNamespacedAttribute(Var);
 					FString ParameterMapInstanceName = GetParameterMapInstanceName(0);
 					int32 ChunkIdx = AddBodyChunk(ParameterMapInstanceName + TEXT(".") + GetSanitizedSymbolName(VarNamespaced.GetName().ToString()), TEXT("{0}"), VarNamespaced.GetType(), Input, false);
 
@@ -5234,11 +5262,11 @@ void FHlslNiagaraTranslator::Output(UNiagaraNodeOutput* OutputNode, const TArray
 	}
 }
 
-int32 FHlslNiagaraTranslator::GetAttribute(const FNiagaraVariable& Attribute)
+int32 FNiagaraHlslTranslator::GetAttribute(const FNiagaraVariable& Attribute)
 {
 	if (!AddStructToDefinitionSet(Attribute.GetType()))
 	{
-		Error(FText::Format(LOCTEXT("GetConstantFail", "Cannot handle type {0}! Variable: {1}"), Attribute.GetType().GetNameText(), FText::FromName(Attribute.GetName())), nullptr, nullptr);
+		Error(FText::Format(LOCTEXT("GetConstantFail", "Cannot handle type {0}! Variable: {1}"), Attribute.GetType().GetNameText(), FText::FromName(Attribute.GetName())));
 	}
 
 	if (TranslationStages.Num() > 1 && UNiagaraScript::IsParticleSpawnScript(TranslationStages[0].ScriptUsage) && (Attribute.GetName() != TEXT("Particles.UniqueID")))
@@ -5255,7 +5283,7 @@ int32 FHlslNiagaraTranslator::GetAttribute(const FNiagaraVariable& Attribute)
 		}
 		else
 		{
-			Error(LOCTEXT("AttrReadInSpawnError", "Cannot read attribute in a spawn script as it's value is not yet initialized."), nullptr, nullptr);
+			Error(LOCTEXT("AttrReadInSpawnError", "Cannot read attribute in a spawn script as it's value is not yet initialized."));
 			return INDEX_NONE;
 		}
 	}
@@ -5267,16 +5295,16 @@ int32 FHlslNiagaraTranslator::GetAttribute(const FNiagaraVariable& Attribute)
 		CompilationOutput.ScriptData.DataUsage.bReadsAttributeData |= (Attribute.GetName() != TEXT("Particles.UniqueID"));
 
 		int32 Chunk = INDEX_NONE;
-		if (!ParameterMapRegisterNamespaceAttributeVariable(Attribute, nullptr, 0, Chunk))
+		if (!ParameterMapRegisterNamespaceAttributeVariable(Attribute, 0, Chunk))
 		{
-			Error(FText::Format(LOCTEXT("AttrReadError", "Cannot read attribute {0} {1}."), Attribute.GetType().GetNameText(), FText::FromString(*Attribute.GetName().ToString())), nullptr, nullptr);
+			Error(FText::Format(LOCTEXT("AttrReadError", "Cannot read attribute {0} {1}."), Attribute.GetType().GetNameText(), FText::FromString(*Attribute.GetName().ToString())));
 			return INDEX_NONE;
 		}
 		return Chunk;
 	}
 }
 
-FString FHlslNiagaraTranslator::GetDataSetAccessSymbol(FNiagaraDataSetID DataSet, int32 IndexChunk, bool bRead)
+FString FNiagaraHlslTranslator::GetDataSetAccessSymbol(FNiagaraDataSetID DataSet, int32 IndexChunk, bool bRead)
 {
 	FString Ret = TEXT("\tContext.") + DataSet.Name.ToString() + (bRead ? TEXT("Read") : TEXT("Write"));
 	/*
@@ -5286,7 +5314,8 @@ FString FHlslNiagaraTranslator::GetDataSetAccessSymbol(FNiagaraDataSetID DataSet
 	return Ret;
 }
 
-void FHlslNiagaraTranslator::ParameterMapForBegin(UNiagaraNodeParameterMapFor* ForNode, int32 IterationCount)
+template<typename GraphBridge>
+void TNiagaraHlslTranslator<GraphBridge>::ParameterMapForBegin(const FParamMapForNode* ForNode, int32 IterationCount)
 {
 	NIAGARA_SCOPE_CYCLE_COUNTER(STAT_NiagaraEditor_HlslTranslator_MapForBegin);
 
@@ -5300,14 +5329,16 @@ void FHlslNiagaraTranslator::ParameterMapForBegin(UNiagaraNodeParameterMapFor* F
 	AddBodyChunk(TEXT(""), TEXT("for({0} = 0; {0} < {1}; ++{0})\n\t{"), FNiagaraTypeDefinition::GetIntDef(), SourceChunks, false, false);
 }
 
-void FHlslNiagaraTranslator::ParameterMapForContinue(UNiagaraNodeParameterMapFor* ForNode, int32 IterationEnabled)
+template<typename GraphBridge>
+void TNiagaraHlslTranslator<GraphBridge>::ParameterMapForContinue(const FParamMapForNode* ForNode, int32 IterationEnabled)
 {
 	NIAGARA_SCOPE_CYCLE_COUNTER(STAT_NiagaraEditor_HlslTranslator_MapForBegin);
 
 	AddBodyChunk(TEXT(""), TEXT("if (!{0}) continue;"), FNiagaraTypeDefinition::GetBoolDef(), IterationEnabled, false, false);
 }
 
-void FHlslNiagaraTranslator::ParameterMapForEnd(UNiagaraNodeParameterMapFor* ForNode)
+template<typename GraphBridge>
+void TNiagaraHlslTranslator<GraphBridge>::ParameterMapForEnd(const FParamMapForNode* ForNode)
 {
 	NIAGARA_SCOPE_CYCLE_COUNTER(STAT_NiagaraEditor_HlslTranslator_MapForEnd);
 
@@ -5316,7 +5347,7 @@ void FHlslNiagaraTranslator::ParameterMapForEnd(UNiagaraNodeParameterMapFor* For
 	ParameterMapForIndexStack.Pop();
 }
 
-int32 FHlslNiagaraTranslator::ParameterMapForInnerIndex() const
+int32 FNiagaraHlslTranslator::ParameterMapForInnerIndex() const
 {
 	if (ParameterMapForIndexStack.Num())
 	{
@@ -5326,7 +5357,8 @@ int32 FHlslNiagaraTranslator::ParameterMapForInnerIndex() const
 	return INDEX_NONE;
 }
 
-void FHlslNiagaraTranslator::ParameterMapSet(UNiagaraNodeParameterMapSet* SetNode, TArrayView<const FCompiledPin> Inputs, TArray<int32>& Outputs)
+template<typename GraphBridge>
+void TNiagaraHlslTranslator<GraphBridge>::ParameterMapSet(const FParamMapSetNode* SetNode, TArrayView<const FCompiledPin> Inputs, TArray<int32>& Outputs)
 {
 	NIAGARA_SCOPE_CYCLE_COUNTER(STAT_NiagaraEditor_HlslTranslator_MapSet);
 
@@ -5336,7 +5368,7 @@ void FHlslNiagaraTranslator::ParameterMapSet(UNiagaraNodeParameterMapSet* SetNod
 
 	// There is only one output pin for a set node, the parameter map must 
 	// continue to route through it.
-	if (!SetNode->IsNodeEnabled())
+	if (!GraphBridge::NodeIsEnabled(SetNode))
 	{
 		if (Inputs.Num() >= 1)
 		{
@@ -5360,16 +5392,16 @@ void FHlslNiagaraTranslator::ParameterMapSet(UNiagaraNodeParameterMapSet* SetNod
 				Error(LOCTEXT("NoParamMapIdxForInput", "Cannot find parameter map for input!"), SetNode, nullptr);
 				Outputs[0] = INDEX_NONE;
 					return;
-				}
+			}
 			continue;
 		}
 		else // These are the pins that we are setting on the parameter map.
 		{
-			FNiagaraVariable Var = Schema->PinToNiagaraVariable(Inputs[i].Pin, false, ENiagaraStructConversion::Simulation);
+			FNiagaraVariable Var = GraphBridge::GetPinVariable(Inputs[i].Pin, false, ENiagaraStructConversion::Simulation);
 
 			if (!AddStructToDefinitionSet(Var.GetType()))
 			{
-				Error(FText::Format(LOCTEXT("ParameterMapSetTypeError", "Cannot handle type {0}! Variable: {1}"), Var.GetType().GetNameText(), FText::FromName(Var.GetName())), nullptr, nullptr);
+				Error(FText::Format(LOCTEXT("ParameterMapSetTypeError", "Cannot handle type {0}! Variable: {1}"), Var.GetType().GetNameText(), FText::FromName(Var.GetName())));
 			}
 
 			if (TranslationStages[ActiveStageIdx].IsExternalConstantNamespace(Var, CompileOptions.TargetUsage, CompileOptions.GetTargetUsageBitmask()))
@@ -5384,12 +5416,12 @@ void FHlslNiagaraTranslator::ParameterMapSet(UNiagaraNodeParameterMapSet* SetNod
 			if (ConstantInfo.ConstantVar != nullptr && ConstantInfo.ConstantVar->GetType() != Var.GetType() && ConstantInfo.ConstantType != ENiagaraKnownConstantType::Attribute)
 			{
 				Error(FText::Format(LOCTEXT("MismatchedConstantTypes", "Variable {0} is a system constant, but its type is different! {1} != {2}"), FText::FromName(Var.GetName()),
-					ConstantInfo.ConstantVar->GetType().GetNameText(), Var.GetType().GetNameText()), nullptr, nullptr);
+					ConstantInfo.ConstantVar->GetType().GetNameText(), Var.GetType().GetNameText()));
 			}
 
 			if (FNiagaraConstants::IsEngineManagedAttribute(Var))
 			{
-				Error(FText::Format(LOCTEXT("SettingSystemAttr", "Variable {0} is an engine managed particle attribute and cannot be set directly."), FText::FromName(Var.GetName())), nullptr, nullptr);
+				Error(FText::Format(LOCTEXT("SettingSystemAttr", "Variable {0} is an engine managed particle attribute and cannot be set directly."), FText::FromName(Var.GetName())));
 				continue;
 			}
 
@@ -5427,7 +5459,7 @@ void FHlslNiagaraTranslator::ParameterMapSet(UNiagaraNodeParameterMapSet* SetNod
 							{
 								bAllSame = false;
 							}
-							else if (FoundOverrideIdx == INDEX_NONE && FNiagaraParameterMapHistory::IsAliasedEmitterParameter(Var))
+							else if (FoundOverrideIdx == INDEX_NONE && FParamMapHistory::IsAliasedEmitterParameter(Var))
 							{
 								FNiagaraAliasContext ResolveAliasesContext(FNiagaraAliasContext::ERapidIterationParameterMode::EmitterOrParticleScript);
 								ResolveAliasesContext.ChangeEmitterToEmitterName(CompileData->GetUniqueEmitterName());
@@ -5482,7 +5514,7 @@ void FHlslNiagaraTranslator::ParameterMapSet(UNiagaraNodeParameterMapSet* SetNod
 					}
 
 					FName UsageName;
-					if (FNiagaraParameterMapHistory::IsAliasedEmitterParameter(Var))
+					if (FParamMapHistory::IsAliasedEmitterParameter(Var))
 					{
 						FNiagaraVariable AliasedVar = ActiveHistoryForFunctionCalls.ResolveAliases(Var);
 						UsageName = AliasedVar.GetName();
@@ -5580,14 +5612,13 @@ void FHlslNiagaraTranslator::ParameterMapSet(UNiagaraNodeParameterMapSet* SetNod
 			}
 		}
 	}
-
 }
 
-FString FHlslNiagaraTranslator::GetUniqueEmitterName() const
+FString FNiagaraHlslTranslator::GetUniqueEmitterName() const
 {
 	if (CompileOptions.TargetUsage == ENiagaraScriptUsage::SystemSpawnScript || CompileOptions.TargetUsage == ENiagaraScriptUsage::SystemUpdateScript)
 	{
-		return CompileData->GetUniqueEmitterName();
+		return GetEmitterUniqueName();
 	}
 	else
 	{
@@ -5595,12 +5626,12 @@ FString FHlslNiagaraTranslator::GetUniqueEmitterName() const
 	}
 }
 
-bool FHlslNiagaraTranslator::IsBulkSystemScript() const
+bool FNiagaraHlslTranslator::IsBulkSystemScript() const
 {
 	return (CompileOptions.TargetUsage == ENiagaraScriptUsage::SystemSpawnScript || CompileOptions.TargetUsage == ENiagaraScriptUsage::SystemUpdateScript);
 }
 
-bool FHlslNiagaraTranslator::IsSpawnScript() const
+bool FNiagaraHlslTranslator::IsSpawnScript() const
 {
 	for (int32 i = 0; i < TranslationStages.Num(); i++)
 	{
@@ -5612,12 +5643,12 @@ bool FHlslNiagaraTranslator::IsSpawnScript() const
 	return false;
 }
 
-bool FHlslNiagaraTranslator::IsEventSpawnScript()const 
+bool FNiagaraHlslTranslator::IsEventSpawnScript()const
 {
 	return UNiagaraScript::IsParticleEventScript(CompileOptions.TargetUsage) && CompileOptions.AdditionalDefines.Contains(FNiagaraCompileOptions::EventSpawnDefine);
 }
 
-bool FHlslNiagaraTranslator::RequiresInterpolation() const
+bool FNiagaraHlslTranslator::RequiresInterpolation() const
 {
 	for (int32 i = 0; i < TranslationStages.Num(); i++)
 	{
@@ -5629,9 +5660,10 @@ bool FHlslNiagaraTranslator::RequiresInterpolation() const
 	return false;
 }
 
-bool FHlslNiagaraTranslator::GetLiteralConstantVariable(FNiagaraVariable& OutVar) const
+template<typename GraphBridge>
+bool TNiagaraHlslTranslator<GraphBridge>::GetLiteralConstantVariable(FNiagaraVariable& OutVar) const
 {
-	if (FNiagaraParameterMapHistory::IsInNamespace(OutVar, PARAM_MAP_EMITTER_STR) || FNiagaraParameterMapHistory::IsInNamespace(OutVar, PARAM_MAP_SYSTEM_STR))
+	if (FParamMapHistory::IsInNamespace(OutVar, PARAM_MAP_EMITTER_STR) || FParamMapHistory::IsInNamespace(OutVar, PARAM_MAP_SYSTEM_STR))
 	{
 		FNiagaraVariable ResolvedVar = ActiveHistoryForFunctionCalls.ResolveAliases(OutVar);
 		if (OutVar == FNiagaraVariable(FNiagaraTypeDefinition::GetBoolDef(), TEXT("Emitter.Localspace")))
@@ -5691,11 +5723,11 @@ bool FHlslNiagaraTranslator::GetLiteralConstantVariable(FNiagaraVariable& OutVar
 			return true;
 		}
 	}
-
 	return false;
 }
 
-bool FHlslNiagaraTranslator::HandleBoundConstantVariableToDataSetRead(FNiagaraVariable InVariable, UNiagaraNode* InNode, int32 InParamMapHistoryIdx, int32& Output, const UEdGraphPin* InDefaultPin)
+template<typename GraphBridge>
+bool TNiagaraHlslTranslator<GraphBridge>::HandleBoundConstantVariableToDataSetRead(FNiagaraVariable InVariable, const FNode* InNode, int32 InParamMapHistoryIdx, int32& Output, const FPin* InDefaultPin)
 {
 	if (InVariable == SYS_PARAM_ENGINE_EMITTER_SIMULATION_POSITION)
 	{
@@ -5713,7 +5745,8 @@ bool FHlslNiagaraTranslator::HandleBoundConstantVariableToDataSetRead(FNiagaraVa
 	return false;
 }
 
-bool FHlslNiagaraTranslator::ParameterMapRegisterExternalConstantNamespaceVariable(FNiagaraVariable InVariable, UNiagaraNode* InNodeForErrorReporting, int32 InParamMapHistoryIdx, int32& Output, const UEdGraphPin* InDefaultPin)
+template<typename GraphBridge>
+bool TNiagaraHlslTranslator<GraphBridge>::ParameterMapRegisterExternalConstantNamespaceVariable(FNiagaraVariable InVariable, const FNode* InNodeForErrorReporting, int32 InParamMapHistoryIdx, int32& Output, const FPin* InDefaultPin)
 {
 	InVariable = ActiveHistoryForFunctionCalls.ResolveAliases(InVariable);
 	FString VarName = InVariable.GetName().ToString();
@@ -5740,9 +5773,9 @@ bool FHlslNiagaraTranslator::ParameterMapRegisterExternalConstantNamespaceVariab
 		const bool bIsDataInterface = InVariable.IsDataInterface();
 		const bool bIsUObject = InVariable.IsUObject();
 		const FString* EmitterAlias = ActiveHistoryForFunctionCalls.GetEmitterAlias();
-		
-		bool bIsPerInstanceBulkSystemParam = IsBulkSystemScript() && !bIsDataInterface && !bIsUObject && (FNiagaraParameterMapHistory::IsUserParameter(InVariable) || FNiagaraParameterMapHistory::IsPerInstanceEngineParameter(InVariable, EmitterAlias != nullptr ? *EmitterAlias : TEXT("Emitter")));
-		const bool bIsExternalConstantParameter = FNiagaraParameterMapHistory::IsRapidIterationParameter(InVariable) && !InVariable.GetType().IsStatic();
+
+		bool bIsPerInstanceBulkSystemParam = IsBulkSystemScript() && !bIsDataInterface && !bIsUObject && (FParamMapHistory::IsUserParameter(InVariable) || FParamMapHistory::IsPerInstanceEngineParameter(InVariable, EmitterAlias != nullptr ? *EmitterAlias : TEXT("Emitter")));
+		const bool bIsExternalConstantParameter = FParamMapHistory::IsRapidIterationParameter(InVariable) && !InVariable.GetType().IsStatic();
 
 		if (InVariable.GetType().IsStatic())
 		{
@@ -5884,9 +5917,10 @@ bool FHlslNiagaraTranslator::ParameterMapRegisterExternalConstantNamespaceVariab
 	return false;
 }
 
-void FHlslNiagaraTranslator::FillVariableWithDefaultValue(FNiagaraVariable& InVariable, const UEdGraphPin* InDefaultPin)
+template<typename GraphBridge>
+void TNiagaraHlslTranslator<GraphBridge>::FillVariableWithDefaultValue(FNiagaraVariable& InVariable, const FPin* InDefaultPin)
 {
-	FNiagaraVariable Var = Schema->PinToNiagaraVariable(InDefaultPin, true, ENiagaraStructConversion::Simulation);
+	FNiagaraVariable Var = GraphBridge::GetPinVariable(InDefaultPin, true, ENiagaraStructConversion::Simulation);
 	FNiagaraEditorUtilities::ResetVariableToDefaultValue(InVariable);
 	if (Var.IsDataAllocated() && Var.GetData() != nullptr)
 	{
@@ -5894,9 +5928,10 @@ void FHlslNiagaraTranslator::FillVariableWithDefaultValue(FNiagaraVariable& InVa
 	}
 }
 
-void FHlslNiagaraTranslator::FillVariableWithDefaultValue(int32& OutValue, const UEdGraphPin* InDefaultPin)
+template<typename GraphBridge>
+void TNiagaraHlslTranslator<GraphBridge>::FillVariableWithDefaultValue(int32& OutValue, const FPin* InDefaultPin)
 {
-	FNiagaraVariable Var = Schema->PinToNiagaraVariable(InDefaultPin, true, ENiagaraStructConversion::Simulation);
+	FNiagaraVariable Var = GraphBridge::GetPinVariable(InDefaultPin, true, ENiagaraStructConversion::Simulation);
 	FNiagaraVariable VarFinal = Var;
 	FNiagaraEditorUtilities::ResetVariableToDefaultValue(VarFinal); // Do this to handle non-zero defaults
 	if (Var.IsDataAllocated() && Var.GetData() != nullptr)
@@ -5915,14 +5950,14 @@ void FHlslNiagaraTranslator::FillVariableWithDefaultValue(int32& OutValue, const
 }
 
 
-void FHlslNiagaraTranslator::SetConstantByStaticVariable(int32& OutValue, const UEdGraphPin* InDefaultPin, FString* OutDebugString)
+template<typename GraphBridge>
+void TNiagaraHlslTranslator<GraphBridge>::SetConstantByStaticVariable(int32& OutValue, const FPin* InDefaultPin, FString* OutDebugString)
 {
 	if (InDefaultPin == nullptr)
 		return;
-
 	FNiagaraEditorModule& NiagaraEditorModule = FModuleManager::GetModuleChecked<FNiagaraEditorModule>("NiagaraEditor");
 	OutValue = 0;
-	FNiagaraVariable Var = Schema->PinToNiagaraVariable(InDefaultPin, true);
+	FNiagaraVariable Var = GraphBridge::GetPinVariable(InDefaultPin, true, ENiagaraStructConversion::UserFacing);
 	FNiagaraVariable VarDefault = Var;
 	FNiagaraEditorUtilities::ResetVariableToDefaultValue(VarDefault);// Do this to handle non-zero defaults
 	if (VarDefault.GetType().IsStatic())
@@ -5930,11 +5965,16 @@ void FHlslNiagaraTranslator::SetConstantByStaticVariable(int32& OutValue, const
 		FNiagaraVariable VarWithValue = FNiagaraVariable(Var.GetType(), Var.GetName());
 		FString Value;
 		int32 Input = INDEX_NONE;
-		const UEdGraphPin* PinToTest = InDefaultPin;
-		if (InDefaultPin && InDefaultPin->Direction == EEdGraphPinDirection::EGPD_Input && InDefaultPin->LinkedTo.Num() != 0)
+
+		const FPin* PinToTest = InDefaultPin;
+		if (const FInputPin* DefaultInputPin = GraphBridge::GetPinAsInput(InDefaultPin))
 		{
-			PinToTest = InDefaultPin->LinkedTo[0];
+			if (const FOutputPin* ConnectedPin = GraphBridge::GetLinkedOutputPin(DefaultInputPin))
+			{
+				PinToTest = ConnectedPin;
+			}
 		}
+
 		FGraphTraversalHandle PinHandle(ActiveHistoryForFunctionCalls.ActivePath);
 		PinHandle.PushPin(PinToTest);
 
@@ -5977,12 +6017,13 @@ void FHlslNiagaraTranslator::SetConstantByStaticVariable(int32& OutValue, const
 		else if (!UNiagaraScript::IsModuleScript(CompileOptions.TargetUsage) && !UNiagaraScript::IsFunctionScript(CompileOptions.TargetUsage) &&
 			!UNiagaraScript::IsDynamicInputScript(CompileOptions.TargetUsage)) // Can't always resolve these outside the context of an emitter or system
 		{
-			Error(LOCTEXT("CouldNotResolveStaticVarByPin", "Could not resolve static variable through pin."), Cast<UNiagaraNode>(InDefaultPin->GetOwningNode()), InDefaultPin);
+			Error(LOCTEXT("CouldNotResolveStaticVarByPin", "Could not resolve static variable through pin."), GraphBridge::GetOwningNode(InDefaultPin), InDefaultPin);
 		}
 	}
 }
 
-void FHlslNiagaraTranslator::SetConstantByStaticVariable(FNiagaraVariable& OutValue, const UEdGraphPin* InDefaultPin, FString* DebugString )
+template<typename GraphBridge>
+void TNiagaraHlslTranslator<GraphBridge>::SetConstantByStaticVariable(FNiagaraVariable& OutValue, const FPin* InDefaultPin, FString* DebugString)
 {
 	OutValue.AllocateData();
 	int32 Constant = 0;
@@ -5998,7 +6039,7 @@ void FHlslNiagaraTranslator::SetConstantByStaticVariable(FNiagaraVariable& OutVa
 	}
 }
 
-void FHlslNiagaraTranslator::SetConstantByStaticVariable(FNiagaraVariable& OutValue, const FNiagaraVariable& Var, FString* DebugString)
+void FNiagaraHlslTranslator::SetConstantByStaticVariable(FNiagaraVariable& OutValue, const FNiagaraVariable& Var, FString* DebugString)
 {
 	OutValue = Var;
 	OutValue.AllocateData();
@@ -6015,7 +6056,7 @@ void FHlslNiagaraTranslator::SetConstantByStaticVariable(FNiagaraVariable& OutVa
 	}
 }
 
-void FHlslNiagaraTranslator::SetConstantByStaticVariable(int32& OutValue, const FNiagaraVariable& Var, FString* DebugString)
+void FNiagaraHlslTranslator::SetConstantByStaticVariable(int32& OutValue, const FNiagaraVariable& Var, FString* DebugString)
 {
 	OutValue = 0;
 	FNiagaraVariable VarDefault = Var;
@@ -6035,29 +6076,30 @@ void FHlslNiagaraTranslator::SetConstantByStaticVariable(int32& OutValue, const
 		{
 			if (!VarWithValue.IsDataAllocated())
 			{
-				int32 FoundOverrideIdx = CompileData->StaticVariables.IndexOfByPredicate([&](const FNiagaraVariable& InObj) -> bool
-					{
-						return (InObj.GetName() == *Value);
-					});;
+				TConstArrayView<FNiagaraVariable> StaticVariables = GetStaticVariables();
+				int32 FoundOverrideIdx = StaticVariables.IndexOfByPredicate([&](const FNiagaraVariable& InObj) -> bool
+				{
+					return (InObj.GetName() == *Value);
+				});;
 
 				if (FoundOverrideIdx != INDEX_NONE)
 				{
-					VarWithValue.SetData(CompileData->StaticVariables[FoundOverrideIdx].GetData());
+					VarWithValue.SetData(StaticVariables[FoundOverrideIdx].GetData());
 				}
 				else if (FNiagaraParameterMapHistory::IsAliasedEmitterParameter(Var))
 				{
 					FNiagaraAliasContext ResolveAliasesContext(FNiagaraAliasContext::ERapidIterationParameterMode::EmitterOrParticleScript);
-					ResolveAliasesContext.ChangeEmitterToEmitterName(CompileData->GetUniqueEmitterName());
+					ResolveAliasesContext.ChangeEmitterToEmitterName(GetUniqueEmitterName());
 					FNiagaraVariable TestEmitterResolvedVar = FNiagaraUtilities::ResolveAliases(Var, ResolveAliasesContext);
 
-					FoundOverrideIdx = CompileData->StaticVariables.IndexOfByPredicate([&](const FNiagaraVariable& InObj) -> bool
-						{
-							return (InObj.GetName() == TestEmitterResolvedVar.GetName());
-						});;
+					FoundOverrideIdx = StaticVariables.IndexOfByPredicate([&](const FNiagaraVariable& InObj) -> bool
+					{
+						return (InObj.GetName() == TestEmitterResolvedVar.GetName());
+					});;
 
 					if (FoundOverrideIdx != INDEX_NONE)
 					{
-						VarWithValue.SetData(CompileData->StaticVariables[FoundOverrideIdx].GetData());
+						VarWithValue.SetData(StaticVariables[FoundOverrideIdx].GetData());
 					}
 				}
 			}
@@ -6078,15 +6120,16 @@ void FHlslNiagaraTranslator::SetConstantByStaticVariable(int32& OutValue, const
 		}
 		else
 		{
-			Error(FText::Format(LOCTEXT("CouldNotResolveStaticVar", "Could not resolve static variable \"{0}\". Default type value used instead."),  FText::FromName(Var.GetName())), nullptr, nullptr);
+			Error(FText::Format(LOCTEXT("CouldNotResolveStaticVar", "Could not resolve static variable \"{0}\". Default type value used instead."),  FText::FromName(Var.GetName())));
 		}
 	}
 }
 
 
-bool FHlslNiagaraTranslator::ParameterMapRegisterUniformAttributeVariable(const FNiagaraVariable& InVariable, UNiagaraNode* InNode, int32 InParamMapHistoryIdx, int32& Output)
+template<typename GraphBridge>
+bool TNiagaraHlslTranslator<GraphBridge>::ParameterMapRegisterUniformAttributeVariable(const FNiagaraVariable& InVariable, FNode* InNode, int32 InParamMapHistoryIdx, int32& Output)
 {
-	FNiagaraVariable NewVar = FNiagaraParameterMapHistory::BasicAttributeToNamespacedAttribute(InVariable);
+	FNiagaraVariable NewVar = FParamMapHistory::BasicAttributeToNamespacedAttribute(InVariable);
 	if (NewVar.IsValid())
 	{
 		return ParameterMapRegisterNamespaceAttributeVariable(NewVar, InNode, InParamMapHistoryIdx, Output);
@@ -6094,7 +6137,8 @@ bool FHlslNiagaraTranslator::ParameterMapRegisterUniformAttributeVariable(const
 	return false;
 }
 
-void FHlslNiagaraTranslator::ValidateParticleIDUsage()
+template<typename GraphBridge>
+void TNiagaraHlslTranslator<GraphBridge>::ValidateParticleIDUsage()
 {
 	if (CompileOptions.AdditionalDefines.Contains(TEXT("RequiresPersistentIDs")))
 	{
@@ -6102,7 +6146,7 @@ void FHlslNiagaraTranslator::ValidateParticleIDUsage()
 		return;
 	}
 	FName particleIDName(TEXT("Particles.ID"));
-	for (FNiagaraParameterMapHistory& History : ParamMapHistories)
+	for (FParamMapHistory& History : ParamMapHistories)
 	{
 		for (const FNiagaraVariable& Variable : History.Variables)
 		{
@@ -6114,7 +6158,8 @@ void FHlslNiagaraTranslator::ValidateParticleIDUsage()
 	}
 }
 
-bool FHlslNiagaraTranslator::ParameterMapRegisterNamespaceAttributeVariable(const FNiagaraVariable& InVariable, UNiagaraNode* InNode, int32 InParamMapHistoryIdx, int32& Output)
+template<typename GraphBridge>
+bool TNiagaraHlslTranslator<GraphBridge>::ParameterMapRegisterNamespaceAttributeVariable(const FNiagaraVariable& InVariable, const FNode* InNode, int32 InParamMapHistoryIdx, int32& Output)
 {
 	FString VarName = InVariable.GetName().ToString();
 	FString SymbolNameNamespaced = GetSanitizedSymbolName(VarName);
@@ -6122,7 +6167,7 @@ bool FHlslNiagaraTranslator::ParameterMapRegisterNamespaceAttributeVariable(cons
 	FNiagaraVariable NamespaceVar = InVariable;
 
 	Output = INDEX_NONE;
-	FNiagaraVariable BasicVar = FNiagaraParameterMapHistory::ResolveAsBasicAttribute(InVariable);
+	FNiagaraVariable BasicVar = FParamMapHistory::ResolveAsBasicAttribute(InVariable);
 	if (BasicVar.IsValid())
 	{
 		if (false == ParamMapDefinedAttributesToUniformChunks.Contains(BasicVar.GetName()))
@@ -6158,7 +6203,13 @@ bool FHlslNiagaraTranslator::ParameterMapRegisterNamespaceAttributeVariable(cons
 	return false;
 }
 
-FString FHlslNiagaraTranslator::GetParameterMapInstanceName(int32 ParamMapHistoryIdx)
+template<typename GraphBridge>
+bool TNiagaraHlslTranslator<GraphBridge>::ParameterMapRegisterNamespaceAttributeVariable(const FNiagaraVariable& InVariable, int32 InParamMapHistoryIdx, int32& Output)
+{
+	return ParameterMapRegisterNamespaceAttributeVariable(InVariable, nullptr, InParamMapHistoryIdx, Output);
+}
+
+FString FNiagaraHlslTranslator::GetParameterMapInstanceName(int32 ParamMapHistoryIdx)
 {
 	FString ParameterMapInstanceName;
 
@@ -6169,21 +6220,17 @@ FString FHlslNiagaraTranslator::GetParameterMapInstanceName(int32 ParamMapHistor
 	return ParameterMapInstanceName;
 }
 
-void FHlslNiagaraTranslator::Emitter(UNiagaraNodeEmitter* EmitterNode, TArray<int32>& Inputs, TArray<int32>& Outputs)
+template<typename GraphBridge>
+void TNiagaraHlslTranslator<GraphBridge>::Emitter(const FEmitterNode* EmitterNode, TArray<int32>& Inputs, TArray<int32>& Outputs)
 {
 	NIAGARA_SCOPE_CYCLE_COUNTER(STAT_NiagaraEditor_Module_NiagaraHLSLTranslator_Emitter);
 
 	// Just pass through the input parameter map pin if the node isn't enabled...
-	if (!EmitterNode->IsNodeEnabled())
+	if (!GraphBridge::NodeIsEnabled(EmitterNode))
 	{
-		FPinCollectorArray OutputPins;
-		EmitterNode->GetOutputPins(OutputPins);
+		TArray<const FOutputPin*> OutputPins = GraphBridge::GetOutputPins(EmitterNode);
+		Outputs.Init(INDEX_NONE, OutputPins.Num());
 
-		Outputs.SetNum(OutputPins.Num());
-		for (int32 i = 0; i < OutputPins.Num(); i++)
-		{
-			Outputs[i] = INDEX_NONE;
-		}
 		if (Inputs.Num() >= 1)
 		{
 			Outputs[0] = Inputs[0];
@@ -6192,43 +6239,39 @@ void FHlslNiagaraTranslator::Emitter(UNiagaraNodeEmitter* EmitterNode, TArray<in
 	}
 
 	FNiagaraFunctionSignature Signature;
-	UNiagaraScriptSource* Source = EmitterNode->GetScriptSource();
-	if (Source == nullptr)
+	const FGraph* EmitterGraph = GraphBridge::GetEmitterGraph(EmitterNode);
+	if (EmitterGraph == nullptr)
 	{
 		Error(LOCTEXT("FunctionCallNonexistantScriptSource", "Emitter call missing ScriptSource"), EmitterNode, nullptr);
 		return;
 	}
 
 	// We need the generated string to generate the proper signature for now.
-	FString EmitterUniqueName = EmitterNode->GetEmitterUniqueName();
+	FString EmitterUniqueName = GraphBridge::GetEmitterUniqueName(EmitterNode);
 
-	ENiagaraScriptUsage ScriptUsage = EmitterNode->GetUsage();
-	FString Name = EmitterNode->GetName();
-	FString FullName = EmitterNode->GetFullName();
+	ENiagaraScriptUsage ScriptUsage = GraphBridge::GetEmitterUsage(EmitterNode);
+	FString Name = GraphBridge::GetEmitterName(EmitterNode);
+	FString FullName = GraphBridge::GetEmitterPathName(EmitterNode);
 
 	FName StatName = *EmitterUniqueName;
 	EnterStatsScope(FNiagaraStatScope(StatName, StatName));
 
-	FPinCollectorArray CallOutputs;
-	FPinCollectorArray CallInputs;
-	EmitterNode->GetOutputPins(CallOutputs);
-	EmitterNode->GetInputPins(CallInputs);
-
+	TArray<const FOutputPin*> CallOutputs = GraphBridge::GetOutputPins(EmitterNode);
+	TArray<const FInputPin*> CallInputs = GraphBridge::GetInputPins(EmitterNode);
 
-	if (Inputs.Num() == 0 || Schema->PinToNiagaraVariable(CallInputs[0]).GetType() != FNiagaraTypeDefinition::GetParameterMapDef())
+	if (Inputs.Num() == 0 || GraphBridge::GetPinType(CallInputs[0], ENiagaraStructConversion::UserFacing) != FNiagaraTypeDefinition::GetParameterMapDef())
 	{
 		Error(LOCTEXT("EmitterMissingParamMap", "Emitter call missing ParameterMap input pin!"), EmitterNode, nullptr);
 		return;
 	}
 
-
 	int32 ParamMapHistoryIdx = Inputs[0];
 	if (ParamMapHistoryIdx == INDEX_NONE)
 	{
 		Error(LOCTEXT("EmitterMissingParamMapIndex", "Emitter call missing valid ParameterMap index!"), EmitterNode, nullptr);
 		return;
 	}
-	ActiveHistoryForFunctionCalls.EnterEmitter(EmitterUniqueName, EmitterNode->GetCalledGraph(), EmitterNode);
+	ActiveHistoryForFunctionCalls.EnterEmitter(EmitterUniqueName, GraphBridge::GetEmitterGraph(EmitterNode), EmitterNode);
 
 	// Clear out the parameter map writes to emitter module parameters as they should not be shared across emitters.
 	if (ParamMapHistoryIdx != -1 && ParamMapHistoryIdx < ParamMapHistories.Num())
@@ -6237,7 +6280,7 @@ void FHlslNiagaraTranslator::Emitter(UNiagaraNodeEmitter* EmitterNode, TArray<in
 		{
 			check(ParamMapHistories[ParamMapHistoryIdx].VariablesWithOriginalAliasesIntact.Num() > i);
 			FNiagaraVariable Var = ParamMapHistories[ParamMapHistoryIdx].VariablesWithOriginalAliasesIntact[i];
-			if (FNiagaraParameterMapHistory::IsAliasedModuleParameter(Var))
+			if (FParamMapHistory::IsAliasedModuleParameter(Var))
 			{
 				ParamMapSetVariablesToChunks[ParamMapHistoryIdx][i] = INDEX_NONE;
 			}
@@ -6245,7 +6288,7 @@ void FHlslNiagaraTranslator::Emitter(UNiagaraNodeEmitter* EmitterNode, TArray<in
 	}
 
 	// We act like a function call here as the semantics are identical.
-	RegisterFunctionCall(ScriptUsage, Name, FullName, EmitterNode->NodeGuid, EmitterNode->GetEmitterHandleId().ToString(EGuidFormats::Digits), Source, Signature, false, FString(), {}, Inputs, CallInputs, CallOutputs, Signature);
+	RegisterFunctionCall(ScriptUsage, Name, FullName, EmitterNode->NodeGuid, GraphBridge::GetEmitterHandleIdString(EmitterNode), EmitterGraph, Signature, false, FString(), {}, Inputs, CallInputs, CallOutputs, Signature);
 	GenerateFunctionCall(ScriptUsage, Signature, Inputs, Outputs);
 
 	// Clear out the parameter map writes to emitter module parameters as they should not be shared across emitters.
@@ -6255,7 +6298,7 @@ void FHlslNiagaraTranslator::Emitter(UNiagaraNodeEmitter* EmitterNode, TArray<in
 		{
 			check(ParamMapHistories[ParamMapHistoryIdx].VariablesWithOriginalAliasesIntact.Num() > i);
 			FNiagaraVariable Var = ParamMapHistories[ParamMapHistoryIdx].VariablesWithOriginalAliasesIntact[i];
-			if (ActiveHistoryForFunctionCalls.IsInEncounteredFunctionNamespace(Var) || FNiagaraParameterMapHistory::IsAliasedModuleParameter(Var) || FNiagaraParameterMapHistory::IsInNamespace(Var, PARAM_MAP_TRANSIENT_STR))
+			if (ActiveHistoryForFunctionCalls.IsInEncounteredFunctionNamespace(Var) || FParamMapHistory::IsAliasedModuleParameter(Var) || FParamMapHistory::IsInNamespace(Var, PARAM_MAP_TRANSIENT_STR))
 			{
 				ParamMapSetVariablesToChunks[ParamMapHistoryIdx][i] = INDEX_NONE;
 			}
@@ -6266,47 +6309,31 @@ void FHlslNiagaraTranslator::Emitter(UNiagaraNodeEmitter* EmitterNode, TArray<in
 	ExitStatsScope();
 }
 
-void FHlslNiagaraTranslator::ParameterMapGet(UNiagaraNodeParameterMapGet* GetNode, TArrayView<const int32> Inputs, TArray<int32>& Outputs)
+template<typename GraphBridge>
+void TNiagaraHlslTranslator<GraphBridge>::ParameterMapGet(const FParamMapGetNode* GetNode, TArrayView<const int32> Inputs, TArray<int32>& Outputs)
 {
 	NIAGARA_SCOPE_CYCLE_COUNTER(STAT_NiagaraEditor_Module_NiagaraHLSLTranslator_MapGet);
 
-	FPinCollectorArray OutputPins;
-	GetNode->GetOutputPins(OutputPins);
+	TArray<const typename GraphBridge::FOutputPin*> OutputPins = GraphBridge::GetOutputPins(GetNode);
+	Outputs.Init(INDEX_NONE, OutputPins.Num());
 
 	// Push out invalid values for all output pins if the node is disabled.
-	if (!GetNode->IsNodeEnabled())
+	if (!GraphBridge::NodeIsEnabled(GetNode))
 	{
-		Outputs.SetNum(OutputPins.Num());
-		for (int32 i = 0; i < OutputPins.Num(); i++)
-		{
-			Outputs[i] = INDEX_NONE;
-		}
 		return;
 	}
 
-	FPinCollectorArray InputPins;
-	GetNode->GetInputPins(InputPins);
-
+	TArray<const typename GraphBridge::FInputPin*> InputPins = GraphBridge::GetInputPins(GetNode);
 	int32 ParamMapHistoryIdx = Inputs[0];
 
-	Outputs.SetNum(OutputPins.Num());
-
 	if (ParamMapHistoryIdx == -1)
 	{
 		Error(LOCTEXT("NoParamMapIdxForInput", "Cannot find parameter map for input!"), GetNode, nullptr);
-		for (int32 i = 0; i < Outputs.Num(); i++)
-		{
-			Outputs[i] = INDEX_NONE;
-		}
 		return;
 	}
 	else if (ParamMapHistoryIdx >= ParamMapHistories.Num())
 	{
 		Error(FText::Format(LOCTEXT("InvalidParamMapIdxForInput", "Invalid parameter map index for input {0} of {1}!"), ParamMapHistoryIdx, ParamMapHistories.Num()), GetNode, nullptr);
-		for (int32 i = 0; i < Outputs.Num(); i++)
-		{
-			Outputs[i] = INDEX_NONE;
-		}
 		return;
 	}
 
@@ -6314,56 +6341,50 @@ void FHlslNiagaraTranslator::ParameterMapGet(UNiagaraNodeParameterMapGet* GetNod
 
 	for (int32 i = 0; i < OutputPins.Num(); i++)
 	{
-		if (GetNode->IsAddPin(OutputPins[i]))
-		{
-			// Not a real pin.
-			Outputs[i] = INDEX_NONE;
-			continue;
-		}
-		else // These are the pins that we are getting off the parameter map.
-		{
-			FNiagaraTypeDefinition OutputTypeDefinition = Schema->PinToTypeDefinition(OutputPins[i]);
-			const bool bNeedsValue =
-				OutputTypeDefinition != FNiagaraTypeDefinition::GetParameterMapDef() &&
-				OutputTypeDefinition.IsDataInterface() == false &&
-				OutputTypeDefinition.IsUObject() == false;
-			FNiagaraVariable Var = Schema->PinToNiagaraVariable(OutputPins[i], bNeedsValue, ENiagaraStructConversion::Simulation);
+		// These are the pins that we are getting off the parameter map.
+		FNiagaraTypeDefinition OutputTypeDefinition = GraphBridge::GetPinType(OutputPins[i], ENiagaraStructConversion::UserFacing);
+		const bool bNeedsValue =
+			OutputTypeDefinition != FNiagaraTypeDefinition::GetParameterMapDef() &&
+			OutputTypeDefinition.IsDataInterface() == false &&
+			OutputTypeDefinition.IsUObject() == false;
 
-			FNiagaraScriptVariableBinding DefaultBinding;
-			TOptional<ENiagaraDefaultMode> DefaultMode = GetNode->GetNiagaraGraph()->GetDefaultMode(Var, &DefaultBinding);
-			if (Var.GetType().IsStatic())
+		FNiagaraVariable Var = GraphBridge::GetPinVariable(OutputPins[i], bNeedsValue, ENiagaraStructConversion::Simulation);
+
+		FNiagaraScriptVariableBinding DefaultBinding;
+		TOptional<ENiagaraDefaultMode> DefaultMode = GraphBridge::GetGraphDefaultMode(GraphBridge::GetOwningGraph(GetNode), Var, DefaultBinding);
+		if (Var.GetType().IsStatic())
+		{
+			if (FParamMapHistory::IsExternalConstantNamespace(Var, CompileOptions.TargetUsage, CompileOptions.GetTargetUsageBitmask()))
 			{
-				if (FNiagaraParameterMapHistory::IsExternalConstantNamespace(Var, CompileOptions.TargetUsage, CompileOptions.GetTargetUsageBitmask()))
+				if (DefaultMode.IsSet() && *DefaultMode == ENiagaraDefaultMode::FailIfPreviouslyNotSet && !Var.IsInNameSpace(FNiagaraConstants::UserNamespaceString))
 				{
-					if (DefaultMode.IsSet() && *DefaultMode == ENiagaraDefaultMode::FailIfPreviouslyNotSet && !Var.IsInNameSpace(FNiagaraConstants::UserNamespaceString))
-					{
-						// Register an external dependency...
-						RegisterCompileDependency(Var, FText::Format(LOCTEXT("UsedBeforeSet", "Variable {0} was read before being set. It's default mode is \"Fail If Previously Not Set\", so this isn't allowed."), FText::FromName(Var.GetName())), GetNode, OutputPins[i], true, ParamMapHistoryIdx);
-					}
-				}
-				else if (DefaultMode.IsSet() && *DefaultMode == ENiagaraDefaultMode::FailIfPreviouslyNotSet && !Var.IsInNameSpace(FNiagaraConstants::UserNamespaceString) && !Var.IsInNameSpace(FNiagaraConstants::ModuleNamespaceString))
-				{
-					// Check for an internal dependency
-					bool bFailIfNotSet = false;
-					FNiagaraVariable TestVar = ActiveHistoryForFunctionCalls.ResolveAliases(Var);
-					ValidateFailIfPreviouslyNotSet(TestVar, bFailIfNotSet);
-					if (bFailIfNotSet)
-					{
-						RegisterCompileDependency(Var, FText::Format(LOCTEXT("UsedBeforeSet", "Variable {0} was read before being set. It's default mode is \"Fail If Previously Not Set\", so this isn't allowed."), FText::FromName(Var.GetName())), GetNode, nullptr, false, ParamMapHistoryIdx);
-					}
+					// Register an external dependency...
+					RegisterCompileDependency(Var, FText::Format(LOCTEXT("UsedBeforeSet", "Variable {0} was read before being set. It's default mode is \"Fail If Previously Not Set\", so this isn't allowed."), FText::FromName(Var.GetName())), GetNode, OutputPins[i], true, ParamMapHistoryIdx);
 				}
-
-				Outputs[i] = MakeStaticVariableDirect(OutputPins[i]);
 			}
-			else
+			else if (DefaultMode.IsSet() && *DefaultMode == ENiagaraDefaultMode::FailIfPreviouslyNotSet && !Var.IsInNameSpace(FNiagaraConstants::UserNamespaceString) && !Var.IsInNameSpace(FNiagaraConstants::ModuleNamespaceString))
 			{
-				HandleParameterRead(ParamMapHistoryIdx, Var, GetNode->GetDefaultPin(OutputPins[i]), GetNode, Outputs[i], DefaultMode, DefaultBinding);
+				// Check for an internal dependency
+				bool bFailIfNotSet = false;
+				FNiagaraVariable TestVar = ActiveHistoryForFunctionCalls.ResolveAliases(Var);
+				ValidateFailIfPreviouslyNotSet(TestVar, bFailIfNotSet);
+				if (bFailIfNotSet)
+				{
+					RegisterCompileDependency(Var, FText::Format(LOCTEXT("UsedBeforeSet", "Variable {0} was read before being set. It's default mode is \"Fail If Previously Not Set\", so this isn't allowed."), FText::FromName(Var.GetName())), GetNode, nullptr, false, ParamMapHistoryIdx);
+				}
 			}
+
+			Outputs[i] = MakeStaticVariableDirect(OutputPins[i]);
+		}
+		else
+		{
+			HandleParameterRead(ParamMapHistoryIdx, Var, GraphBridge::GetDefaultPin(GetNode, OutputPins[i]), GetNode, Outputs[i], DefaultMode, DefaultBinding);
 		}
 	}
 }
 
-int32 FHlslNiagaraTranslator::MakeStaticVariableDirect(const UEdGraphPin* InDefaultPin)
+template<typename GraphBridge>
+int32 TNiagaraHlslTranslator<GraphBridge>::MakeStaticVariableDirect(const FPin* InDefaultPin)
 {
 	int32 Constant = INDEX_NONE;
 	FString DebugStr;
@@ -6373,17 +6394,18 @@ int32 FHlslNiagaraTranslator::MakeStaticVariableDirect(const UEdGraphPin* InDefa
 	return GetConstantDirect(Constant);
 }
 
-void FHlslNiagaraTranslator::ValidateFailIfPreviouslyNotSet(const FNiagaraVariable& InVar, bool& bFailIfNotSet)
+template<typename GraphBridge>
+void TNiagaraHlslTranslator<GraphBridge>::ValidateFailIfPreviouslyNotSet(const FNiagaraVariable& InVar, bool& bFailIfNotSet)
 {
 	bFailIfNotSet = false;
 	FNiagaraVariable SearchVar = InVar;
-	if (FNiagaraParameterMapHistory::IsInitialValue(InVar))
+	if (FParamMapHistory::IsInitialValue(InVar))
 	{
-		SearchVar = FNiagaraParameterMapHistory::GetSourceForInitialValue(InVar);
+		SearchVar = FParamMapHistory::GetSourceForInitialValue(InVar);
 	}
-	else if (FNiagaraParameterMapHistory::IsPreviousValue(InVar))
+	else if (FParamMapHistory::IsPreviousValue(InVar))
 	{
-		SearchVar = FNiagaraParameterMapHistory::GetSourceForPreviousValue(InVar);
+		SearchVar = FParamMapHistory::GetSourceForPreviousValue(InVar);
 	}
 
 	FVarAndDefaultSource* ParamMapDefinedVarAndDefaultSource = ParamMapDefinedAttributesToNamespaceVars.Find(SearchVar.GetName());
@@ -6397,7 +6419,7 @@ void FHlslNiagaraTranslator::ValidateFailIfPreviouslyNotSet(const FNiagaraVariab
 		int32 FoundInParamIdx = OtherOutputParamMapHistories[OtherParamIdx].FindVariableByName(SearchVar.GetName());
 		if (INDEX_NONE != FoundInParamIdx)
 		{
-			for (const FModuleScopedPin& ScopedPin : OtherOutputParamMapHistories[OtherParamIdx].PerVariableWriteHistory[FoundInParamIdx])
+			for (const typename GraphBridge::FModuleScopedPin& ScopedPin : OtherOutputParamMapHistories[OtherParamIdx].PerVariableWriteHistory[FoundInParamIdx])
 			{
 				if (ScopedPin.Pin->Direction == EEdGraphPinDirection::EGPD_Input && ScopedPin.Pin->bHidden == false)
 				{
@@ -6412,7 +6434,8 @@ void FHlslNiagaraTranslator::ValidateFailIfPreviouslyNotSet(const FNiagaraVariab
 }
 
 
-void FHlslNiagaraTranslator::HandleParameterRead(int32 ParamMapHistoryIdx, const FNiagaraVariable& InVar, const UEdGraphPin* DefaultPin, UNiagaraNode* ErrorNode, int32& OutputChunkId, TOptional<ENiagaraDefaultMode> DefaultMode, TOptional<FNiagaraScriptVariableBinding> DefaultBinding, bool bTreatAsUnknownParameterMap, bool bIgnoreDefaultSetFirst)
+template<typename GraphBridge>
+void TNiagaraHlslTranslator<GraphBridge>::HandleParameterRead(int32 ParamMapHistoryIdx, const FNiagaraVariable& InVar, const FInputPin* DefaultPin, const FNode* ErrorNode, int32& OutputChunkId, TOptional<ENiagaraDefaultMode> DefaultMode, TOptional<FNiagaraScriptVariableBinding> DefaultBinding, bool bTreatAsUnknownParameterMap, bool bIgnoreDefaultSetFirst)
 {
 	FString ParameterMapInstanceName = GetParameterMapInstanceName(ParamMapHistoryIdx);
 	FNiagaraVariable Var = ConvertToSimulationVariable(InVar);
@@ -6427,9 +6450,9 @@ void FHlslNiagaraTranslator::HandleParameterRead(int32 ParamMapHistoryIdx, const
 
 	bool bIsPerInstanceAttribute = false;
 	bool bIsCandidateForRapidIteration = false;
-	const UEdGraphPin* InputPin = DefaultPin;
+	const FInputPin* InputPin = DefaultPin;
 
-	FString Namespace = FNiagaraParameterMapHistory::GetNamespace(Var);
+	FString Namespace = FParamMapHistory::GetNamespace(Var);
 	if (!ParamMapHistories[ParamMapHistoryIdx].IsValidNamespaceForReading(CompileOptions.TargetUsage, CompileOptions.TargetUsageBitmask, Namespace))
 	{
 		if (UNiagaraScript::IsStandaloneScript(CompileOptions.TargetUsage) && Namespace.StartsWith(PARAM_MAP_ATTRIBUTE_STR))
@@ -6455,7 +6478,7 @@ void FHlslNiagaraTranslator::HandleParameterRead(int32 ParamMapHistoryIdx, const
 	}
 
 	
-	if (FNiagaraParameterMapHistory::IsExternalConstantNamespace(Var, CompileOptions.TargetUsage, CompileOptions.GetTargetUsageBitmask()))
+	if (FParamMapHistory::IsExternalConstantNamespace(Var, CompileOptions.TargetUsage, CompileOptions.GetTargetUsageBitmask()))
 	{
 		if (DefaultMode.IsSet() && *DefaultMode == ENiagaraDefaultMode::FailIfPreviouslyNotSet && !bIgnoreDefaultSetFirst)
 		{
@@ -6471,21 +6494,21 @@ void FHlslNiagaraTranslator::HandleParameterRead(int32 ParamMapHistoryIdx, const
 			return;
 		}			
 	}
-	else if (FNiagaraParameterMapHistory::IsAliasedModuleParameter(Var) && ActiveHistoryForFunctionCalls.InTopLevelFunctionCall(CompileOptions.TargetUsage))
+	else if (FParamMapHistory::IsAliasedModuleParameter(Var) && ActiveHistoryForFunctionCalls.InTopLevelFunctionCall(CompileOptions.TargetUsage))
 	{
 		if (DefaultMode.IsSet() && *DefaultMode == ENiagaraDefaultMode::Binding && DefaultBinding.IsSet() && DefaultBinding->IsValid())
 		{
 			// Skip the case where the below condition is met, but it's overridden by a binding.
 			bIsCandidateForRapidIteration = false;
 		}
-		else if (InputPin != nullptr && InputPin->LinkedTo.Num() == 0 && Var.GetType() != FNiagaraTypeDefinition::GetBoolDef() && !Var.GetType().IsEnum() && !Var.IsDataInterface() && !Var.IsUObject())
+		else if (InputPin != nullptr && GraphBridge::GetLinkedOutputPin(InputPin) == nullptr && Var.GetType() != FNiagaraTypeDefinition::GetBoolDef() && !Var.GetType().IsEnum() && !Var.GetType().IsDataInterface() && !Var.IsUObject())
 		{
 			bIsCandidateForRapidIteration = true;
 		}
 	}
 
-	FNiagaraParameterMapHistory& History = ParamMapHistories[ParamMapHistoryIdx];
-	bool bWasEmitterAliased = FNiagaraParameterMapHistory::IsAliasedEmitterParameter(Var);
+	FParamMapHistory& History = ParamMapHistories[ParamMapHistoryIdx];
+	bool bWasEmitterAliased = FParamMapHistory::IsAliasedEmitterParameter(Var);
 	Var = ActiveHistoryForFunctionCalls.ResolveAliases(Var);
 
 	const FNiagaraKnownConstantInfo ConstantInfo = FNiagaraConstants::GetKnownConstantInfo(Var.GetName(), false);
@@ -6550,12 +6573,12 @@ void FHlslNiagaraTranslator::HandleParameterRead(int32 ParamMapHistoryIdx, const
 			bool bVarChanged = false;
 			if (!bWasEmitterAliased && ActiveHistoryForFunctionCalls.GetEmitterAlias() != nullptr)
 			{
-				Var = FNiagaraParameterMapHistory::ConvertVariableToRapidIterationConstantName(Var, *(*ActiveHistoryForFunctionCalls.GetEmitterAlias()), GetTargetUsage());
+				Var = FParamMapHistory::ConvertVariableToRapidIterationConstantName(Var, *(*ActiveHistoryForFunctionCalls.GetEmitterAlias()), GetTargetUsage());
 				bVarChanged = true;
 			}
 			else if (UNiagaraScript::IsSystemSpawnScript(CompileOptions.TargetUsage) || UNiagaraScript::IsSystemUpdateScript(CompileOptions.TargetUsage))
 			{
-				Var = FNiagaraParameterMapHistory::ConvertVariableToRapidIterationConstantName(Var, nullptr, GetTargetUsage());
+				Var = FParamMapHistory::ConvertVariableToRapidIterationConstantName(Var, nullptr, GetTargetUsage());
 				bVarChanged = true;
 			}
 
@@ -6592,7 +6615,7 @@ void FHlslNiagaraTranslator::HandleParameterRead(int32 ParamMapHistoryIdx, const
 				{
 					if (!bWasEmitterAliased && ActiveHistoryForFunctionCalls.GetEmitterAlias() != nullptr && CompileData != nullptr)
 					{
-						Var = FNiagaraParameterMapHistory::ConvertVariableToRapidIterationConstantName(OriginalVar, *CompileData->EmitterUniqueName, GetTargetUsage());
+						Var = FParamMapHistory::ConvertVariableToRapidIterationConstantName(OriginalVar, *CompileData->EmitterUniqueName, GetTargetUsage());
 						bVarChanged = true;
 						FoundIdx = TranslationOptions.OverrideModuleConstants.Find(Var);
 					}
@@ -6607,7 +6630,7 @@ void FHlslNiagaraTranslator::HandleParameterRead(int32 ParamMapHistoryIdx, const
 				}
 				else if (InputPin != nullptr && !InputPin->bDefaultValueIsIgnored) // Use the default from the input pin because this variable was previously never encountered.
 				{
-					FNiagaraVariable PinVar = Schema->PinToNiagaraVariable(InputPin, true, ENiagaraStructConversion::Simulation);
+					FNiagaraVariable PinVar = GraphBridge::GetPinVariable(InputPin, true, ENiagaraStructConversion::Simulation);
 					FString DebugConstantStr;
 					OutputChunkId = GetConstant(PinVar, &DebugConstantStr);
 					UE_LOG(LogNiagaraEditor, VeryVerbose, TEXT("Converted default value of parameter %s to constant %s for script %s. Likely added since this system was last compiled."), *Var.GetName().ToString(), *DebugConstantStr, *CompileOptions.FullName);
@@ -6663,9 +6686,9 @@ void FHlslNiagaraTranslator::HandleParameterRead(int32 ParamMapHistoryIdx, const
 
 		if (LastSetChunkIdx == INDEX_NONE && (UNiagaraScript::IsSpawnScript(TranslationStages[ActiveStageIdx].ScriptUsage)))
 		{
-			if (FNiagaraParameterMapHistory::IsInitialValue(Var))
+			if (FParamMapHistory::IsInitialValue(Var))
 			{
-				FNiagaraVariable SourceForInitialValue = FNiagaraParameterMapHistory::GetSourceForInitialValue(Var);
+				FNiagaraVariable SourceForInitialValue = FParamMapHistory::GetSourceForInitialValue(Var);
 				bool bFoundExistingSet = false;
 				for (int32 OtherParamIdx = 0; OtherParamIdx < OtherOutputParamMapHistories.Num(); OtherParamIdx++)
 				{
@@ -6730,19 +6753,21 @@ void FHlslNiagaraTranslator::HandleParameterRead(int32 ParamMapHistoryIdx, const
 			}
 			else if (InputPin != nullptr) // Default was found, trace back its inputs.
 			{
+				const FOutputPin* LinkedOutputPin = GraphBridge::GetLinkedOutputPin(InputPin);
+
 				// Check to see if there are any overrides passed in to the translator. This allows us to bake in rapid iteration variables for performance.
-				if (InputPin->LinkedTo.Num() == 0 && bIsCandidateForRapidIteration && !TranslationOptions.bParameterRapidIteration)
+				if (LinkedOutputPin == nullptr && bIsCandidateForRapidIteration && !TranslationOptions.bParameterRapidIteration)
 				{
 					bool bVarChanged = false;
 					FNiagaraVariable RapidIterationConstantVar;
 					if (!bWasEmitterAliased && ActiveHistoryForFunctionCalls.GetEmitterAlias() != nullptr)
 					{
-						RapidIterationConstantVar = FNiagaraParameterMapHistory::ConvertVariableToRapidIterationConstantName(Var, *(*ActiveHistoryForFunctionCalls.GetEmitterAlias()), GetTargetUsage());
+						RapidIterationConstantVar = FParamMapHistory::ConvertVariableToRapidIterationConstantName(Var, *(*ActiveHistoryForFunctionCalls.GetEmitterAlias()), GetTargetUsage());
 						bVarChanged = true;
 					}
 					else if (UNiagaraScript::IsSystemSpawnScript(CompileOptions.TargetUsage) || UNiagaraScript::IsSystemUpdateScript(CompileOptions.TargetUsage))
 					{
-						RapidIterationConstantVar = FNiagaraParameterMapHistory::ConvertVariableToRapidIterationConstantName(Var, nullptr, GetTargetUsage());
+						RapidIterationConstantVar = FParamMapHistory::ConvertVariableToRapidIterationConstantName(Var, nullptr, GetTargetUsage());
 						bVarChanged = true;
 					}
 
@@ -6754,9 +6779,9 @@ void FHlslNiagaraTranslator::HandleParameterRead(int32 ParamMapHistoryIdx, const
 						UE_LOG(LogNiagaraEditor, Display, TEXT("Converted parameter %s to constant %s for script %s"), *Var.GetName().ToString(), *DebugConstantStr, *CompileOptions.FullName);
 						return;
 					}
-					else if (InputPin != nullptr && !InputPin->bDefaultValueIsIgnored) // Use the default from the input pin because this variable was previously never encountered.
+					else if (!InputPin->bDefaultValueIsIgnored) // Use the default from the input pin because this variable was previously never encountered.
 					{
-						FNiagaraVariable PinVar = Schema->PinToNiagaraVariable(InputPin, true, ENiagaraStructConversion::Simulation);
+						FNiagaraVariable PinVar = GraphBridge::GetPinVariable(InputPin, true, ENiagaraStructConversion::Simulation);
 						FString DebugConstantStr;
 						OutputChunkId = GetConstant(PinVar, &DebugConstantStr);
 						UE_LOG(LogNiagaraEditor, Display, TEXT("Converted default value of parameter %s to constant %s for script %s. Likely added since this system was last compiled."), *Var.GetName().ToString(), *DebugConstantStr, *CompileOptions.FullName);
@@ -6767,16 +6792,17 @@ void FHlslNiagaraTranslator::HandleParameterRead(int32 ParamMapHistoryIdx, const
 				}
 
 				CurrentDefaultPinTraversal.Push(InputPin);
-				if (InputPin->LinkedTo.Num() != 0 && InputPin->LinkedTo[0] != nullptr)
+				if (LinkedOutputPin)
 				{
 					// Double-check to make sure that we are connected to a TRANSLATOR_PARAM_BEGIN_DEFAULTS input node rather than
 					// a normal parameter-based parameter map input node to ensure that we don't get into weird traversals.
-					TArray<UNiagaraNode*> Nodes;
-					UNiagaraGraph::BuildTraversal(Nodes, Cast<UNiagaraNode>(InputPin->LinkedTo[0]->GetOwningNode()));
-					for (UNiagaraNode* Node : Nodes)
+					TArray<const FNode*> Nodes;
+					const FNode* LinkedOutputPinNode = GraphBridge::GetOwningNode(LinkedOutputPin);
+					GraphBridge::BuildTraversal(GraphBridge::GetOwningGraph(LinkedOutputPinNode), LinkedOutputPinNode, Nodes);
+					for (const FNode* Node : Nodes)
 					{
-						UNiagaraNodeInput* InputNode = Cast<UNiagaraNodeInput>(Node);
-						if (InputNode && InputNode->Input.GetType() == FNiagaraTypeDefinition::GetParameterMapDef())
+						const FInputNode* InputNode = GraphBridge::template AsNodeType<FInputNode>(Node);
+						if (InputNode && GraphBridge::GetInputVariable(InputNode).GetType() == FNiagaraTypeDefinition::GetParameterMapDef())
 						{
 							if (InputNode->Usage != ENiagaraInputNodeUsage::TranslatorConstant)
 							{
@@ -6787,7 +6813,7 @@ void FHlslNiagaraTranslator::HandleParameterRead(int32 ParamMapHistoryIdx, const
 						}
 					}
 				}
-				LastSetChunkIdx = CompilePin(InputPin);
+				LastSetChunkIdx = CompileInputPin(InputPin);
 				CurrentDefaultPinTraversal.Pop();
 			}
 			else
@@ -6875,12 +6901,13 @@ void FHlslNiagaraTranslator::HandleParameterRead(int32 ParamMapHistoryIdx, const
 	}
 }
 
-bool FHlslNiagaraTranslator::IsCompileOptionDefined(const TCHAR* InDefineStr)
+bool FNiagaraHlslTranslator::IsCompileOptionDefined(const TCHAR* InDefineStr)
 {
 	return CompileOptions.AdditionalDefines.Contains(InDefineStr);
 }
 
-void FHlslNiagaraTranslator::ReadDataSet(const FNiagaraDataSetID DataSet, const TArray<FNiagaraVariable>& Variables, ENiagaraDataSetAccessMode AccessMode, int32 InputChunk, TArray<int32>& Outputs)
+template<typename GraphBridge>
+void TNiagaraHlslTranslator<GraphBridge>::ReadDataSet(const FNiagaraDataSetID DataSet, const TArray<FNiagaraVariable>& Variables, ENiagaraDataSetAccessMode AccessMode, int32 InputChunk, TArray<int32>& Outputs)
 {
 	//Eventually may allow events that take in a direct index or condition but for now we don't
 	int32 ParamMapHistoryIdx = InputChunk;
@@ -6942,7 +6969,8 @@ void FHlslNiagaraTranslator::ReadDataSet(const FNiagaraDataSetID DataSet, const
 	}
 }
 
-void FHlslNiagaraTranslator::WriteDataSet(const FNiagaraDataSetID DataSet, const TArray<FNiagaraVariable>& Variables, ENiagaraDataSetAccessMode AccessMode, const TArray<int32>& Inputs, TArray<int32>& Outputs)
+template<typename GraphBridge>
+void TNiagaraHlslTranslator<GraphBridge>::WriteDataSet(const FNiagaraDataSetID DataSet, const TArray<FNiagaraVariable>& Variables, ENiagaraDataSetAccessMode AccessMode, const TArray<int32>& Inputs, TArray<int32>& Outputs)
 {
 	int32 ParamMapHistoryIdx = Inputs[0];
 	int32 ConditionalChunk = Inputs[1];
@@ -7006,10 +7034,10 @@ void FHlslNiagaraTranslator::WriteDataSet(const FNiagaraDataSetID DataSet, const
 		int32 Input = Inputs[i + 2];//input 0 is the valid input (no entry in variables array), so we need of offset all other inputs by 1.
 		DataSetWriteForInput->CodeChunks.Add(AddBodyChunk(DataSetAccessName + TEXT(".") + GetSanitizedSymbolName(Var.GetName().ToString()), TEXT("{0}"), Var.GetType(), Input, false));
 	}
-
 }
 
-int32 FHlslNiagaraTranslator::RegisterUObject(FNiagaraVariable Variable, UObject* Object, bool bAddParameterMapRead)
+template<typename GraphBridge>
+int32 TNiagaraHlslTranslator<GraphBridge>::RegisterUObject(FNiagaraVariable Variable, UObject* Object, bool bAddParameterMapRead)
 {
 	int32 ObjectIndex = INDEX_NONE;
 	if (bAddParameterMapRead)
@@ -7054,7 +7082,8 @@ int32 FHlslNiagaraTranslator::RegisterUObject(FNiagaraVariable Variable, UObject
 	return ObjectIndex;
 }
 
-int32 FHlslNiagaraTranslator::RegisterDataInterface(FNiagaraVariable& Var, UNiagaraDataInterface* DataInterface, bool bPlaceholder, bool bAddParameterMapRead)
+template<typename GraphBridge>
+int32 TNiagaraHlslTranslator<GraphBridge>::RegisterDataInterface(const FNiagaraVariable& Var, const UNiagaraDataInterface* DataInterface, bool bPlaceholder, bool bAddParameterMapRead)
 {
 	FString Id = DataInterface ? *DataInterface->GetMergeId().ToString() : TEXT("??");
 	FString PathName = DataInterface ? *DataInterface->GetPathName() : TEXT("XX");
@@ -7083,7 +7112,7 @@ int32 FHlslNiagaraTranslator::RegisterDataInterface(FNiagaraVariable& Var, UNiag
 	FNiagaraVariable ReadVariable;
 	if (bAddParameterMapRead)
 	{
-		ReadVariable = FNiagaraParameterMapHistory::IsAliasedEmitterParameter(Var.GetName().ToString())
+		ReadVariable = FParamMapHistory::IsAliasedEmitterParameter(Var.GetName().ToString())
 			? ActiveHistoryForFunctionCalls.ResolveAliases(Var) 
 			: Var;
 	}
@@ -7130,30 +7159,23 @@ int32 FHlslNiagaraTranslator::RegisterDataInterface(FNiagaraVariable& Var, UNiag
 	return Idx;
 }
 
-void FHlslNiagaraTranslator::Operation(UNiagaraNodeOp* Operation, TArray<int32>& Inputs, TArray<int32>& Outputs)
+template<typename GraphBridge>
+void TNiagaraHlslTranslator<GraphBridge>::Operation(const FOpNode* Operation, TArray<int32>& Inputs, TArray<int32>& Outputs)
 {
 	NIAGARA_SCOPE_CYCLE_COUNTER(STAT_NiagaraEditor_HlslTranslator_Operation);
 
 	// Use the pins to determine the output type here since they may have been changed due to numeric pin fix up.
 	const FNiagaraOpInfo* OpInfo = FNiagaraOpInfo::GetOpInfo(Operation->OpName);
-	FPinCollectorArray OutputPins;
-	Operation->GetOutputPins(OutputPins);
-	OutputPins.RemoveAll([](UEdGraphPin* Pin)
-		{
-			return Pin->bOrphanedPin == true;
-		});
+	TArray<const FOutputPin*> OutputPins = GraphBridge::GetOutputPins(Operation);
 
-	TArray<UEdGraphPin*> InputPins;
+	TArray<const FInputPin*> InputPins = GraphBridge::GetInputPins(Operation);
 	TArray<FNiagaraTypeDefinition> InputTypes;
-	Operation->GetInputPins(InputPins);
 
 	bool bAllPinsStatic = true;
 	{
 		for (int32 InputIdx = 0; InputIdx < InputPins.Num(); InputIdx++)
 		{
-			if (Operation->IsAddPin(InputPins[InputIdx]))
-				continue;
-			FNiagaraTypeDefinition InputType = Schema->PinToTypeDefinition(InputPins[InputIdx]);
+			FNiagaraTypeDefinition InputType = GraphBridge::GetPinType(InputPins[InputIdx], ENiagaraStructConversion::UserFacing);
 			InputTypes.Add(InputType);
 			if (!InputType.IsStatic())
 				bAllPinsStatic = false;
@@ -7161,9 +7183,7 @@ void FHlslNiagaraTranslator::Operation(UNiagaraNodeOp* Operation, TArray<int32>&
 
 		for (int32 OutputIdx = 0; OutputIdx < OutputPins.Num(); OutputIdx++)
 		{
-			if (Operation->IsAddPin(OutputPins[OutputIdx]))
-				continue;
-			FNiagaraTypeDefinition OutputType = Schema->PinToTypeDefinition(OutputPins[OutputIdx]);
+			FNiagaraTypeDefinition OutputType = GraphBridge::GetPinType(OutputPins[OutputIdx], ENiagaraStructConversion::UserFacing);
 			if (!OutputType.IsStatic())
 				bAllPinsStatic = false;
 		}
@@ -7179,13 +7199,13 @@ void FHlslNiagaraTranslator::Operation(UNiagaraNodeOp* Operation, TArray<int32>&
 	{
 		if (OpInfo->Outputs.Num() != 1 || OutputPins.Num() != OpInfo->Outputs.Num())
 		{
-			FText PinNameText = OutputPins[0]->PinFriendlyName.IsEmpty() ? FText::FromName(OutputPins[0]->PinName) : OutputPins[0]->PinFriendlyName;
+			FText PinNameText = GraphBridge::GetPinFriendlyName(OutputPins[0]).IsEmpty() ? FText::FromName(OutputPins[0]->PinName) : GraphBridge::GetPinFriendlyName(OutputPins[0]);
 			Error(LOCTEXT("InvalidOutputPinCount", "Only one output pin is supported for static variables"), Operation, OutputPins[0]);
 			Outputs.Add(INDEX_NONE);
 			return;
 		}
 
-		FNiagaraTypeDefinition OutputType = Schema->PinToTypeDefinition(OutputPins[0]);
+		FNiagaraTypeDefinition OutputType = GraphBridge::GetPinType(OutputPins[0], ENiagaraStructConversion::UserFacing);
 		if(!OutputType.IsStatic())
 		{
 			Error(LOCTEXT("InvalidOutputPinType", "Only static types are supported for this operation!"), Operation, OutputPins[0]);
@@ -7197,10 +7217,7 @@ void FHlslNiagaraTranslator::Operation(UNiagaraNodeOp* Operation, TArray<int32>&
 		
 		for (int32 InputIdx = 0; InputIdx < InputPins.Num(); InputIdx++)
 		{
-			if (Operation->IsAddPin(InputPins[InputIdx]))
-				continue;			
-
-			FNiagaraTypeDefinition InputType = Schema->PinToTypeDefinition(InputPins[InputIdx]);
+			FNiagaraTypeDefinition InputType = GraphBridge::GetPinType(InputPins[InputIdx], ENiagaraStructConversion::UserFacing);
 			if (!InputType.IsStatic())
 			{
 				Error(LOCTEXT("InvalidInputPinType", "Only static types are supported for this operation!"), Operation, InputPins[InputIdx]);
@@ -7226,12 +7243,12 @@ void FHlslNiagaraTranslator::Operation(UNiagaraNodeOp* Operation, TArray<int32>&
 
 	for (int32 OutputIndex = 0; OutputIndex < OutputPins.Num(); OutputIndex++)
 	{
-		UEdGraphPin* OutputPin = OutputPins[OutputIndex];
-		FNiagaraTypeDefinition OutputType = Schema->PinToTypeDefinition(OutputPin, ENiagaraStructConversion::Simulation);
+		const FOutputPin* OutputPin = OutputPins[OutputIndex];
+		FNiagaraTypeDefinition OutputType = GraphBridge::GetPinType(OutputPin, ENiagaraStructConversion::Simulation);
 
 		if (!AddStructToDefinitionSet(OutputType))
 		{
-			FText PinNameText = OutputPin->PinFriendlyName.IsEmpty() ? FText::FromName(OutputPin->PinName) : OutputPin->PinFriendlyName;
+			FText PinNameText = GraphBridge::GetPinFriendlyName(OutputPin).IsEmpty() ? FText::FromName(OutputPin->PinName) : GraphBridge::GetPinFriendlyName(OutputPin);
 			Error(FText::Format(LOCTEXT("GetConstantFailTypePin", "Cannot handle type {0}! Output Pin: {1}"), OutputType.GetNameText(), PinNameText), Operation, OutputPin);
 		}
 		if (OpInfo != nullptr)
@@ -7242,12 +7259,13 @@ void FHlslNiagaraTranslator::Operation(UNiagaraNodeOp* Operation, TArray<int32>&
 			{
 				if (!OpInfo->CreateHlslForAddedInputs(Inputs.Num(), OutputHlsl))
 				{
-					FText PinNameText = OutputPin->PinFriendlyName.IsEmpty() ? FText::FromName(OutputPin->PinName) : OutputPin->PinFriendlyName;
+					FText PinNameText = GraphBridge::GetPinFriendlyName(OutputPin).IsEmpty() ? FText::FromName(OutputPin->PinName) : GraphBridge::GetPinFriendlyName(OutputPin);
 					Error(FText::Format(LOCTEXT("AggregateInputFailTypePin", "Cannot create hlsl output for type {0}! Output Pin: {1}"), OutputType.GetNameText(), PinNameText), Operation, OutputPin);
 					OutputHlsl = IOInfo.HlslSnippet;
 				}
 			}
-			else {
+			else
+			{
 				OutputHlsl = IOInfo.HlslSnippet;
 			}
 			check(!OutputHlsl.IsEmpty());
@@ -7256,20 +7274,47 @@ void FHlslNiagaraTranslator::Operation(UNiagaraNodeOp* Operation, TArray<int32>&
 	}
 }
 
-void FHlslNiagaraTranslator::FunctionCall(UNiagaraNodeFunctionCall* FunctionNode, TArray<int32>& Inputs, TArray<int32>& Outputs)
+// simple helper template allowing us to keep FunctionCall() generic even if the Builder interface may diverge
+template<typename GraphBridge>
+class FScopedBuilderFunctionCall
+{
+};
+
+// specialization of this helper function intended for the NiagaraGraphInterface 
+template<>
+class FScopedBuilderFunctionCall<FNiagaraCompilationGraphBridge>
+{
+public:
+	FScopedBuilderFunctionCall(TNiagaraHlslTranslator<FNiagaraCompilationGraphBridge>* InTranslator, FNiagaraCompilationGraphBridge::FParamMapHistoryBuilder& InBuilder, const FNiagaraCompilationGraphBridge::FFunctionCallNode* InFunctionCallNode)
+		: FunctionCallNode(InFunctionCallNode)
+		, Builder(InBuilder)
+	{
+		Builder.EnterFunction(FunctionCallNode->GetFunctionName(), FunctionCallNode->GetCalledGraph(), FunctionCallNode);
+	}
+
+	~FScopedBuilderFunctionCall()
+	{
+		Builder.ExitFunction(FunctionCallNode->GetFunctionName(), FunctionCallNode);
+	}
+
+private:
+	const FNiagaraCompilationGraphBridge::FFunctionCallNode* FunctionCallNode;
+	FNiagaraCompilationGraphBridge::FParamMapHistoryBuilder& Builder;
+};
+
+template<typename GraphBridge>
+void TNiagaraHlslTranslator<GraphBridge>::FunctionCall(const FFunctionCallNode* FunctionNode, TArray<int32>& Inputs, TArray<int32>& Outputs)
 {
 	NIAGARA_SCOPE_CYCLE_COUNTER(STAT_NiagaraEditor_Module_NiagaraHLSLTranslator_FunctionCall);
 
-	FPinCollectorArray CallOutputs;
-	FPinCollectorArray CallInputs;
-	FunctionNode->GetOutputPins(CallOutputs);
-	FunctionNode->GetInputPins(CallInputs);
+	TArray<const FOutputPin*> CallOutputs = GraphBridge::GetOutputPins(FunctionNode);
+	TArray<const FInputPin*> CallInputs = GraphBridge::GetInputPins(FunctionNode);
 
 	// Validate that there are no input pins with the same name and type
 	TMultiMap<FName, FEdGraphPinType> SeenPins;
-	for (UEdGraphPin* Pin : CallInputs)
+	for (const FPin* Pin : CallInputs)
 	{
-		FEdGraphPinType* SeenType = SeenPins.FindPair(Pin->GetFName(), Pin->PinType);
+		FEdGraphPinType* SeenType = SeenPins.FindPair(Pin->PinName, Pin->PinType);
 		if (SeenType)
 		{
 			Error(LOCTEXT("FunctionCallDuplicateInput", "Function call has duplicated inputs. Please make sure that each function parameter is unique."), FunctionNode, Pin);
@@ -7277,21 +7322,21 @@ void FHlslNiagaraTranslator::FunctionCall(UNiagaraNodeFunctionCall* FunctionNode
 		}
 		else
 		{
-			SeenPins.Add(Pin->GetFName(), Pin->PinType);
+			SeenPins.Add(Pin->PinName, Pin->PinType);
 		}
 	}
 
 	// If the function call is disabled, we 
 	// need to route the input parameter map pin to the output parameter map pin.
 	// Any other outputs become invalid.
-	if (!FunctionNode->IsNodeEnabled())
+	if (!GraphBridge::NodeIsEnabled(FunctionNode))
 	{
 		int32 InputPinIdx = INDEX_NONE;
 
 		for (int32 i = 0; i < CallInputs.Num(); i++)
 		{
-			const UEdGraphPin* Pin = CallInputs[i];
-			if (Schema->PinToTypeDefinition(Pin) == FNiagaraTypeDefinition::GetParameterMapDef())
+			const FPin* Pin = CallInputs[i];
+			if (GraphBridge::GetPinType(Pin, ENiagaraStructConversion::UserFacing) == FNiagaraTypeDefinition::GetParameterMapDef())
 			{
 				// Found the input pin
 				InputPinIdx = Inputs[i];
@@ -7303,8 +7348,8 @@ void FHlslNiagaraTranslator::FunctionCall(UNiagaraNodeFunctionCall* FunctionNode
 		for (int32 i = 0; i < CallOutputs.Num(); i++)
 		{
 			Outputs[i] = INDEX_NONE;
-			const UEdGraphPin* Pin = CallOutputs[i];
-			if (Schema->PinToTypeDefinition(Pin) == FNiagaraTypeDefinition::GetParameterMapDef())
+			const FPin* Pin = CallOutputs[i];
+			if (GraphBridge::GetPinType(Pin, ENiagaraStructConversion::UserFacing) == FNiagaraTypeDefinition::GetParameterMapDef())
 			{
 				// Mapping the input parameter map pin to the output.
 				Outputs[i] = InputPinIdx;
@@ -7314,7 +7359,8 @@ void FHlslNiagaraTranslator::FunctionCall(UNiagaraNodeFunctionCall* FunctionNode
 	}
 
 	FNiagaraFunctionSignature OutputSignature;
-	if (FunctionNode->FunctionScript == nullptr && !FunctionNode->Signature.IsValid())
+	const FGraph* FunctionGraph = GraphBridge::GetFunctionNodeGraph(FunctionNode);
+	if (FunctionGraph == nullptr && !FunctionNode->Signature.IsValid())
 	{
 		Error(LOCTEXT("FunctionCallNonexistantFunctionScript", "Function call missing FunctionScript and invalid signature"), FunctionNode, nullptr);
 		return;
@@ -7322,7 +7368,7 @@ void FHlslNiagaraTranslator::FunctionCall(UNiagaraNodeFunctionCall* FunctionNode
 
 	//UE_LOG(LogNiagaraEditor, Log, TEXT("Function Call: %s %d"), *FunctionNode->GetFunctionName(), ActiveStageIdx);
 	// We need the generated string to generate the proper signature for now.
-	ActiveHistoryForFunctionCalls.EnterFunction(FunctionNode->GetFunctionName(), FunctionNode->FunctionScript, FunctionNode->GetNiagaraGraph(), FunctionNode);
+	FScopedBuilderFunctionCall<GraphBridge> ScopedBuilderFunctionCall(this, ActiveHistoryForFunctionCalls, FunctionNode);
 
 	// Check if there are static switch parameters being set directly by a set node from the stack UI.
 	// This can happen if a module was changed and the original parameter was replaced by a static switch with the same name, but the emitter was not yet updated.
@@ -7347,54 +7393,27 @@ void FHlslNiagaraTranslator::FunctionCall(UNiagaraNodeFunctionCall* FunctionNode
 		}
 	}
 
-	// Remove input add pin if it exists
-	bool bHasOutputAdd = false;
-	for (int32 i = 0; i < CallOutputs.Num(); i++)
-	{
-		if (FunctionNode->IsAddPin(CallOutputs[i]))
-		{
-			bHasOutputAdd = true;
-			CallOutputs.RemoveAt(i);
-			break;
-		}
-	}
-
-	// Remove output add pin if it exists
-	bool bHasInputAdd = false;
-	for (int32 i = 0; i < CallInputs.Num(); i++)
-	{
-		if (FunctionNode->IsAddPin(CallInputs[i]))
-		{
-			bHasInputAdd = true;
-			CallInputs.RemoveAt(i);
-			break;
-		}
-	}
-
 	ENiagaraScriptUsage ScriptUsage = ENiagaraScriptUsage::Function;
 	FString Name;
 	FString FullName;
-	UNiagaraScriptSource* Source = nullptr;
 	bool bCustomHlsl = false;
 	FString CustomHlsl;
 	TArray<FNiagaraCustomHlslInclude> CustomHlslIncludeFilePaths;
 	FNiagaraFunctionSignature Signature = FunctionNode->Signature;
 
-	if (FunctionNode->FunctionScript)
+	if (FunctionGraph)
 	{
-		ScriptUsage = FunctionNode->FunctionScript->GetUsage();
-		Name = FunctionNode->FunctionScript->GetName();
-		FullName = FunctionNode->FunctionScript->GetFullName();
-		Source = FunctionNode->GetFunctionScriptSource();
-		check(Source->GetOutermost() == GetTransientPackage());
+		ScriptUsage = GraphBridge::GetFunctionUsage(FunctionNode);
+		Name = GraphBridge::GetFunctionScriptName(FunctionNode);
+		FullName = GraphBridge::GetFunctionFullName(FunctionNode);
 	}
 	else if (Signature.bRequiresExecPin)
 	{
-		if (CallInputs.Num() == 0 || Schema->PinToTypeDefinition(CallInputs[0]) != FNiagaraTypeDefinition::GetParameterMapDef())
+		if (CallInputs.Num() == 0 || GraphBridge::GetPinType(CallInputs[0], ENiagaraStructConversion::UserFacing) != FNiagaraTypeDefinition::GetParameterMapDef())
 		{
 			Error(LOCTEXT("FunctionCallInvalidSignatureExecIn", "The first input pin must be a parameter map pin because the signature RequiresExecPin!"), FunctionNode, nullptr);
 		}
-		if (CallOutputs.Num() == 0 || Schema->PinToTypeDefinition(CallOutputs[0]) != FNiagaraTypeDefinition::GetParameterMapDef())
+		if (CallOutputs.Num() == 0 || GraphBridge::GetPinType(CallOutputs[0], ENiagaraStructConversion::UserFacing) != FNiagaraTypeDefinition::GetParameterMapDef())
 		{
 			Error(LOCTEXT("FunctionCallInvalidSignatureExecOut", "The first output pin must be a parameter map pin because the signature RequiresExecPin!"), FunctionNode, nullptr);
 		}
@@ -7411,7 +7430,7 @@ void FHlslNiagaraTranslator::FunctionCall(UNiagaraNodeFunctionCall* FunctionNode
 			FNiagaraVariable ResolvedVariable;
 			for (int32 i = 0; i < CallInputs.Num(); i++)
 			{
-				FNiagaraTypeDefinition TypeDef = Schema->PinToTypeDefinition(CallInputs[i]);
+				FNiagaraTypeDefinition TypeDef = GraphBridge::GetPinType(CallInputs[i], ENiagaraStructConversion::UserFacing);
 				if (TypeDef.IsDataInterface())
 				{
 					UNiagaraDataInterface* CDO = CompileDuplicateData->GetDuplicatedDataInterfaceCDOForClass(TypeDef.GetClass());
@@ -7431,7 +7450,7 @@ void FHlslNiagaraTranslator::FunctionCall(UNiagaraNodeFunctionCall* FunctionNode
 
 					FNiagaraVariable Variable(TypeDef, *(TEXT("Module.") + Prefix));
 
-					if (FNiagaraParameterMapHistory::IsAliasedModuleParameter(Variable) && ActiveHistoryForFunctionCalls.InTopLevelFunctionCall(CompileOptions.TargetUsage))
+					if (FParamMapHistory::IsAliasedModuleParameter(Variable) && ActiveHistoryForFunctionCalls.InTopLevelFunctionCall(CompileOptions.TargetUsage))
 					{
 						ResolvedVariable = ActiveHistoryForFunctionCalls.ResolveAliases(Variable);
 					}
@@ -7449,7 +7468,7 @@ void FHlslNiagaraTranslator::FunctionCall(UNiagaraNodeFunctionCall* FunctionNode
 
 			for (int32 i = 0; i < CallInputs.Num(); i++)
 			{
-				FNiagaraTypeDefinition TypeDef = Schema->PinToTypeDefinition(CallInputs[i]);
+				FNiagaraTypeDefinition TypeDef = GraphBridge::GetPinType(CallInputs[i], ENiagaraStructConversion::UserFacing);
 				if (!(TypeDef.IsDataInterface() || TypeDef == FNiagaraTypeDefinition::GetParameterMapDef()))
 				{
 					WriteCompilerTag(Inputs[i], CallInputs[i], false, FNiagaraCompileEventSeverity::Display, ResolvedVariable.GetName().ToString());
@@ -7458,13 +7477,13 @@ void FHlslNiagaraTranslator::FunctionCall(UNiagaraNodeFunctionCall* FunctionNode
 		}
 	}
 
-	if (UNiagaraNodeCustomHlsl* CustomFunctionHlsl = Cast<UNiagaraNodeCustomHlsl>(FunctionNode))
+	if (const FCustomHlslNode* CustomFunctionHlsl = GraphBridge::template AsNodeType<FCustomHlslNode>(FunctionNode))
 	{
 		// All of the arguments here are resolved withing the HandleCustomHlsl function..
 		HandleCustomHlslNode(CustomFunctionHlsl, ScriptUsage, Name, FullName, bCustomHlsl, CustomHlsl, CustomHlslIncludeFilePaths, Signature, Inputs);
 	}
 
-	RegisterFunctionCall(ScriptUsage, Name, FullName, FunctionNode->NodeGuid, FString(), Source, Signature, bCustomHlsl, CustomHlsl, CustomHlslIncludeFilePaths, Inputs, CallInputs, CallOutputs, OutputSignature);
+	RegisterFunctionCall(ScriptUsage, Name, FullName, FunctionNode->NodeGuid, FString(), FunctionGraph, Signature, bCustomHlsl, CustomHlsl, CustomHlslIncludeFilePaths, Inputs, CallInputs, CallOutputs, OutputSignature);
 
 	if (OutputSignature.IsValid() == false)
 	{
@@ -7473,31 +7492,20 @@ void FHlslNiagaraTranslator::FunctionCall(UNiagaraNodeFunctionCall* FunctionNode
 	}
 
 	GenerateFunctionCall(ScriptUsage, OutputSignature, Inputs, Outputs);
-
-	// Re-add the add pins.
-	if (bHasInputAdd)
-	{
-		Inputs.Add(INDEX_NONE);
-	}
-	if(bHasOutputAdd)
-	{
-		Outputs.Add(INDEX_NONE);
-	}
-	ActiveHistoryForFunctionCalls.ExitFunction(FunctionNode->GetFunctionName(), FunctionNode->FunctionScript, FunctionNode);
 }
 
-void FHlslNiagaraTranslator::EnterFunctionCallNode(const TSet<FName>& UnusedInputs)
+void FNiagaraHlslTranslator::EnterFunctionCallNode(const TSet<FName>& UnusedInputs)
 {
 	FunctionNodeStack.AddDefaulted_GetRef().UnusedInputs = UnusedInputs;
 }
 
-void FHlslNiagaraTranslator::ExitFunctionCallNode()
+void FNiagaraHlslTranslator::ExitFunctionCallNode()
 {
 	ensure(FunctionNodeStack.Num() > 0);
 	FunctionNodeStack.Pop(false);
 }
 
-bool FHlslNiagaraTranslator::IsFunctionVariableCulledFromCompilation(const FName& InputName) const
+bool FNiagaraHlslTranslator::IsFunctionVariableCulledFromCompilation(const FName& InputName) const
 {
 	if (FunctionNodeStack.Num() == 0)
 	{
@@ -7523,23 +7531,31 @@ bool FHlslNiagaraTranslator::IsFunctionVariableCulledFromCompilation(const FName
 	return false;
 }
 
-void FHlslNiagaraTranslator::CullMapSetInputPin(UEdGraphPin* InputPin)
+template<typename GraphBridge>
+void TNiagaraHlslTranslator<GraphBridge>::CullMapSetInputPin(const FPin* InputPin)
 {
-	if (InputPin == nullptr || InputPin->LinkedTo.Num() != 1 || FunctionNodeStack.Num() == 0)
+	if (FunctionNodeStack.IsEmpty())
+	{
+		return;
+	}
+
+	const FOutputPin* LinkedPin = GraphBridge::GetLinkedOutputPin(GraphBridge::GetPinAsInput(InputPin));
+	if (!LinkedPin)
 	{
 		return;
 	}
 
 	// when a map set input is culled that is connected to a function call node (as is the case for dynamic inputs), we also need to cull any upstream pins that set inputs for the culled function call node
-	if (UNiagaraNodeFunctionCall* FunctionNode = Cast<UNiagaraNodeFunctionCall>(InputPin->LinkedTo[0]->GetOwningNode()))
+	if (const FFunctionCallNode* FunctionNode = GraphBridge::template AsNodeType<FFunctionCallNode>(GraphBridge::GetOwningNode(LinkedPin)))
 	{
-		FString FunctionScriptName = FunctionNode->GetFunctionName();
+		FString FunctionScriptName = GraphBridge::GetFunctionName(FunctionNode);
 		FunctionNodeStack.Last().CulledFunctionNames.Add(FunctionScriptName);
 	}
 }
 
 // From a valid list of namespaces, resolve any aliased tokens and promote namespaced variables without a main namespace to the input parameter map instance namespace
-void FHlslNiagaraTranslator::FinalResolveNamespacedTokens(const FString& ParameterMapInstanceNamespace, TArray<FString>& Tokens, TArray<FString>& ValidChildNamespaces, FNiagaraParameterMapHistoryBuilder& Builder, TArray<FNiagaraVariable>& UniqueParameterMapEntriesAliasesIntact, TArray<FNiagaraVariable>& UniqueParameterMapEntries, int32 ParamMapHistoryIdx, UNiagaraNode* InNodeForErrorReporting)
+template<typename GraphBridge>
+void TNiagaraHlslTranslator<GraphBridge>::FinalResolveNamespacedTokens(const FString& ParameterMapInstanceNamespace, TArray<FString>& Tokens, TArray<FString>& ValidChildNamespaces, FParamMapHistoryBuilder& Builder, TArray<FNiagaraVariable>& UniqueParameterMapEntriesAliasesIntact, TArray<FNiagaraVariable>& UniqueParameterMapEntries, int32 ParamMapHistoryIdx, const FNode* InNodeForErrorReporting)
 {
 	for (int32 i = 0; i < Tokens.Num(); i++)
 	{
@@ -7614,7 +7630,8 @@ static bool IsWhitespaceToken(const FString& Token)
 		;
 }
 
-bool FHlslNiagaraTranslator::ParseDIFunctionSpecifiers(UNiagaraNode* NodeForErrorReporting, FNiagaraFunctionSignature& Sig, TArray<FString>& Tokens, int32& TokenIdx)
+template<typename GraphBridge>
+bool TNiagaraHlslTranslator<GraphBridge>::ParseDIFunctionSpecifiers(const FNode* NodeForErrorReporting, FNiagaraFunctionSignature& Sig, TArray<FString>& Tokens, int32& TokenIdx)
 {
 	const int32 NumTokens = Tokens.Num();
 
@@ -7716,7 +7733,8 @@ bool FHlslNiagaraTranslator::ParseDIFunctionSpecifiers(UNiagaraNode* NodeForErro
 	return true;
 }
 
-void FHlslNiagaraTranslator::ProcessCustomHlsl(const FString& InCustomHlsl, ENiagaraScriptUsage InUsage, const FNiagaraFunctionSignature& InSignature, const TArray<int32>& Inputs, UNiagaraNode* InNodeForErrorReporting, FString& OutCustomHlsl,  FNiagaraFunctionSignature& OutSignature)
+template<typename GraphBridge>
+void TNiagaraHlslTranslator<GraphBridge>::ProcessCustomHlsl(const FString& InCustomHlsl, ENiagaraScriptUsage InUsage, const FNiagaraFunctionSignature& InSignature, const TArray<int32>& Inputs, const FNode* InNodeForErrorReporting, FString& OutCustomHlsl, FNiagaraFunctionSignature& OutSignature)
 {
 	// Split up the hlsl into constituent tokens
 	TArray<FString> Tokens;
@@ -8003,13 +8021,13 @@ void FHlslNiagaraTranslator::ProcessCustomHlsl(const FString& InCustomHlsl, ENia
 		// Clean up any namespaced variables in the token list if they are aliased or promote any tokens that are namespaced to the parent 
 		// parameter map.
 		TArray<FString> PossibleNamespaces;
-		FNiagaraParameterMapHistory::GetValidNamespacesForReading(CompileOptions.TargetUsage, 0, PossibleNamespaces);
+		FParamMapHistory::GetValidNamespacesForReading(CompileOptions.TargetUsage, 0, PossibleNamespaces);
 
-		for (FNiagaraParameterMapHistory& History : ParamMapHistories)
+		for (FParamMapHistory& History : ParamMapHistories)
 		{
 			for (FNiagaraVariable& Var : History.Variables)
 			{
-				FString Namespace = FNiagaraParameterMapHistory::GetNamespace(Var);
+				FString Namespace = FParamMapHistory::GetNamespace(Var);
 				PossibleNamespaces.AddUnique(Namespace);
 			}
 		}
@@ -8024,8 +8042,8 @@ void FHlslNiagaraTranslator::ProcessCustomHlsl(const FString& InCustomHlsl, ENia
 			FNiagaraVariable VarAliased = UniqueParamMapEntriesAliasesIntact[VarIdx];
 			FNiagaraVariable VarActual = UniqueParamMapEntries[VarIdx];
 
-			if (FNiagaraParameterMapHistory::IsExternalConstantNamespace(VarAliased, CompileOptions.TargetUsage, CompileOptions.GetTargetUsageBitmask()) ||
-				FNiagaraParameterMapHistory::IsExternalConstantNamespace(VarActual, CompileOptions.TargetUsage, CompileOptions.GetTargetUsageBitmask()))
+			if (FParamMapHistory::IsExternalConstantNamespace(VarAliased, CompileOptions.TargetUsage, CompileOptions.GetTargetUsageBitmask()) ||
+				FParamMapHistory::IsExternalConstantNamespace(VarActual, CompileOptions.TargetUsage, CompileOptions.GetTargetUsageBitmask()))
 			{
 				int32 TempOutput;
 				if (ParameterMapRegisterExternalConstantNamespaceVariable(VarActual, InNodeForErrorReporting, ParamMapHistoryIdx, TempOutput, nullptr))
@@ -8069,7 +8087,8 @@ void FHlslNiagaraTranslator::ProcessCustomHlsl(const FString& InCustomHlsl, ENia
 	OutCustomHlsl = TEXT("\n") + OutCustomHlsl + TEXT("\n");
 }
 
-void FHlslNiagaraTranslator::HandleCustomHlslNode(UNiagaraNodeCustomHlsl* CustomFunctionHlsl, ENiagaraScriptUsage& OutScriptUsage, FString& OutName, FString& OutFullName, bool& bOutCustomHlsl, FString& OutCustomHlsl, TArray<FNiagaraCustomHlslInclude>& OutCustomHlslIncludeFilePaths,
+template<typename GraphBridge>
+void TNiagaraHlslTranslator<GraphBridge>::HandleCustomHlslNode(const FCustomHlslNode* CustomFunctionHlsl, ENiagaraScriptUsage& OutScriptUsage, FString& OutName, FString& OutFullName, bool& bOutCustomHlsl, FString& OutCustomHlsl, TArray<FNiagaraCustomHlslInclude>& OutCustomHlslIncludeFilePaths,
 	FNiagaraFunctionSignature& OutSignature, TArray<int32>& Inputs)
 {
 	NIAGARA_SCOPE_CYCLE_COUNTER(STAT_NiagaraEditor_Module_NiagaraHLSLTranslator_CustomHLSL);
@@ -8079,24 +8098,25 @@ void FHlslNiagaraTranslator::HandleCustomHlslNode(UNiagaraNodeCustomHlsl* Custom
 	}
 
 	// Determine the important outputs
-	OutScriptUsage = CustomFunctionHlsl->ScriptUsage;
+	OutScriptUsage = GraphBridge::GetCustomHlslUsage(CustomFunctionHlsl);
 	OutName = GetSanitizedSymbolName(CustomFunctionHlsl->Signature.Name.ToString() + CustomFunctionHlsl->NodeGuid.ToString());
 	OutSignature = CustomFunctionHlsl->Signature;
-	OutFullName = CustomFunctionHlsl->GetFullName();
+	OutFullName = GraphBridge::GetFunctionFullName(CustomFunctionHlsl);
 	OutSignature.Name = *OutName; // Force the name to be set to include the node guid for safety...
 	bOutCustomHlsl = true;
-	OutCustomHlsl = CustomFunctionHlsl->GetCustomHlsl();
-	CustomFunctionHlsl->GetIncludeFilePaths(OutCustomHlslIncludeFilePaths);
+	OutCustomHlsl = GraphBridge::GetCustomHlslString(CustomFunctionHlsl);
+	OutCustomHlslIncludeFilePaths = GraphBridge::GetCustomHlslIncludePaths(CustomFunctionHlsl);
 
 	FNiagaraFunctionSignature InSignature = CustomFunctionHlsl->Signature;
-	ProcessCustomHlsl(CustomFunctionHlsl->GetCustomHlsl(), OutScriptUsage, InSignature, Inputs, CustomFunctionHlsl, OutCustomHlsl, OutSignature);
+	ProcessCustomHlsl(OutCustomHlsl, OutScriptUsage, InSignature, Inputs, CustomFunctionHlsl, OutCustomHlsl, OutSignature);
 }
 
-void FHlslNiagaraTranslator::HandleDataInterfaceCall(FNiagaraScriptDataInterfaceCompileInfo& Info, const FNiagaraFunctionSignature& InMatchingSignature)
+template<typename GraphBridge>
+void TNiagaraHlslTranslator<GraphBridge>::HandleDataInterfaceCall(FNiagaraScriptDataInterfaceCompileInfo& Info, const FNiagaraFunctionSignature& InMatchingSignature)
 {
 	const bool bCPUSim = CompileOptions.IsCpuScript();
 	const bool bGPUSim = CompileOptions.IsGpuScript();
-	const UNiagaraNode* CurNode = ActiveHistoryForFunctionCalls.GetCallingContext();
+	const FNode* CurNode = ActiveHistoryForFunctionCalls.GetCallingContext();
 	if (bCPUSim && !InMatchingSignature.bSupportsCPU)
 	{
 		Error(FText::Format(LOCTEXT("FunctionCallDataInterfaceCPUMissing", "Function call \"{0}\" does not work on CPU sims."), FText::FromName(InMatchingSignature.Name)), CurNode, nullptr);
@@ -8148,9 +8168,10 @@ void FHlslNiagaraTranslator::HandleDataInterfaceCall(FNiagaraScriptDataInterface
 	}
 }
 
-bool IsVariableWriteBeforeRead(const TArray<FNiagaraParameterMapHistory::FReadHistory>& ReadHistory)
+template<typename GraphBridge>
+bool IsVariableWriteBeforeRead(const TArray<typename GraphBridge::FParamMapHistory::FReadHistory>& ReadHistory)
 {
-	for (const FNiagaraParameterMapHistory::FReadHistory& History : ReadHistory)
+	for (const typename GraphBridge::FParamMapHistory::FReadHistory& History : ReadHistory)
 	{
 		if (History.PreviousWritePin.Pin == nullptr)
 		{
@@ -8160,35 +8181,27 @@ bool IsVariableWriteBeforeRead(const TArray<FNiagaraParameterMapHistory::FReadHi
 	return true;
 }
 
-void FHlslNiagaraTranslator::RegisterFunctionCall(ENiagaraScriptUsage ScriptUsage, const FString& InName, const FString& InFullName, const FGuid& CallNodeId, const FString& InFunctionNameSuffix, UNiagaraScriptSource* Source,
-                                                  FNiagaraFunctionSignature& InSignature, bool bIsCustomHlsl, const FString& InCustomHlsl, const TArray<FNiagaraCustomHlslInclude>& InCustomHlslIncludeFilePaths, TArray<int32>& Inputs, TArrayView<UEdGraphPin* const> CallInputs, TArrayView<UEdGraphPin* const> CallOutputs,
+template<typename GraphBridge>
+void TNiagaraHlslTranslator<GraphBridge>::RegisterFunctionCall(ENiagaraScriptUsage ScriptUsage, const FString& InName, const FString& InFullName, const FGuid& CallNodeId, const FString& InFunctionNameSuffix, const FGraph* SourceGraph,
+                                                  FNiagaraFunctionSignature& InSignature, bool bIsCustomHlsl, const FString& InCustomHlsl, const TArray<FNiagaraCustomHlslInclude>& InCustomHlslIncludeFilePaths, TArray<int32>& Inputs, TConstArrayView<const FInputPin*> CallInputs, TConstArrayView<const FOutputPin*> CallOutputs,
                                                   FNiagaraFunctionSignature& OutSignature)
 {
 	NIAGARA_SCOPE_CYCLE_COUNTER(STAT_NiagaraEditor_Module_NiagaraHLSLTranslator_RegisterFunctionCall);
 
 	//////////////////////////////////////////////////////////////////////////
-	if (Source)
+	if (SourceGraph)
 	{
 		NIAGARA_SCOPE_CYCLE_COUNTER(STAT_NiagaraEditor_Module_NiagaraHLSLTranslator_RegisterFunctionCall_Source);
-		UNiagaraGraph* SourceGraph = CastChecked<UNiagaraGraph>(Source->NodeGraph);
 
-		bool bHasNumericInputs = false;
-		if (SourceGraph->HasNumericParameters())
-		{
-			for (int32 i = 0; i < CallInputs.Num(); i++)
-			{
-				if (Schema->PinToTypeDefinition(CallInputs[i]) == FNiagaraTypeDefinition::GetGenericNumericDef())
-				{
-					bHasNumericInputs = true;
-				}
-			}
-		}
-		TArray<UEdGraphPin*> StaticSwitchValues;
-		for (FNiagaraVariable StaticSwitchInput : SourceGraph->FindStaticSwitchInputs())
+		const bool bHasNumericInputs = GraphBridge::GraphHasParametersOfType(SourceGraph, FNiagaraTypeDefinition::GetGenericNumericDef());
+		const bool bHasParameterMapParameters = GraphBridge::GraphHasParametersOfType(SourceGraph, FNiagaraTypeDefinition::GetParameterMapDef());
+
+		TArray<const FInputPin*> StaticSwitchValues;
+		for (const FNiagaraVariableBase& StaticSwitchInput : GraphBridge::GraphGetStaticSwitchInputs(SourceGraph))
 		{
-			for (UEdGraphPin* Pin : CallInputs)
+			for (const FInputPin* Pin : CallInputs)
 			{
-				if (StaticSwitchInput.GetName().IsEqual(Pin->GetFName()))
+				if (StaticSwitchInput.GetName().IsEqual(Pin->PinName))
 				{
 					StaticSwitchValues.Add(Pin);
 					break;
@@ -8196,8 +8209,6 @@ void FHlslNiagaraTranslator::RegisterFunctionCall(ENiagaraScriptUsage ScriptUsag
 			}
 		}
 
-		bool bHasParameterMapParameters = SourceGraph->HasParameterMapParameters();
-
 		GenerateFunctionSignature(ScriptUsage, InName, InFullName, InFunctionNameSuffix, SourceGraph, Inputs, bHasNumericInputs, bHasParameterMapParameters, StaticSwitchValues, OutSignature);
 
 		// 		//Sort the input and outputs to match the sorted parameters. They may be different.
@@ -8239,17 +8250,17 @@ void FHlslNiagaraTranslator::RegisterFunctionCall(ENiagaraScriptUsage ScriptUsag
 			//We've not compiled this function yet so compile it now.
 			EnterFunction(InName, OutSignature, Inputs, CallNodeId);
 
-			UNiagaraNodeOutput* FuncOutput = SourceGraph->FindOutputNode(ScriptUsage);
+			const FOutputNode* FuncOutput = SourceGraph->FindOutputNode(ScriptUsage, FGuid());
 			check(FuncOutput);
 
 			if (ActiveHistoryForFunctionCalls.GetModuleAlias() != nullptr)
 			{
 				bool bIsInTopLevelFunction = ActiveHistoryForFunctionCalls.InTopLevelFunctionCall(CompileOptions.TargetUsage);
 
-				UEdGraphPin* ParamMapPin = nullptr;
-				for (UEdGraphPin* Pin : CallInputs)
+				const FInputPin* ParamMapPin = nullptr;
+				for (const FInputPin* Pin : CallInputs)
 				{
-					if (Schema->PinToTypeDefinition(Pin) == FNiagaraTypeDefinition::GetParameterMapDef())
+					if (GraphBridge::GetPinType(Pin, ENiagaraStructConversion::UserFacing) == FNiagaraTypeDefinition::GetParameterMapDef())
 					{
 						ParamMapPin = Pin;
 						break;
@@ -8260,11 +8271,11 @@ void FHlslNiagaraTranslator::RegisterFunctionCall(ENiagaraScriptUsage ScriptUsag
 				{
 					bIsModuleFunction = (bIsInTopLevelFunction && ParamMapPin != nullptr && UNiagaraScript::IsGPUScript(CompileOptions.TargetUsage));
 
-					UNiagaraNode* ParamNode = Cast<UNiagaraNode>(ParamMapPin->GetOwningNode());
+					const FNode* ParamNode = GraphBridge::GetOwningNode(ParamMapPin);
 					if (ParamNode)
 					{
 						check(ParamMapHistories.Num() == TranslationStages.Num());
-						const FNiagaraParameterMapHistory& History = ParamMapHistories[ActiveStageIdx];
+						const FParamMapHistory& History = ParamMapHistories[ActiveStageIdx];
 						uint32 FoundIdx = History.MapNodeVisitations.Find(ParamNode);
 						if (FoundIdx != INDEX_NONE)
 						{
@@ -8274,7 +8285,7 @@ void FHlslNiagaraTranslator::RegisterFunctionCall(ENiagaraScriptUsage ScriptUsag
 
 							for (uint32 VarIdx = History.MapNodeVariableMetaData[FoundIdx].Key; VarIdx < History.MapNodeVariableMetaData[FoundIdx].Value; VarIdx++)
 							{
-								if (IsVariableWriteBeforeRead(History.PerVariableReadHistory[VarIdx]))
+								if (IsVariableWriteBeforeRead<GraphBridge>(History.PerVariableReadHistory[VarIdx]))
 								{
 									// We don't need to worry about defaults if the variable is written before being read or never read at all.
 									continue;
@@ -8293,9 +8304,9 @@ void FHlslNiagaraTranslator::RegisterFunctionCall(ENiagaraScriptUsage ScriptUsag
 									int32 LastSetChunkIdx = ParamMapSetVariablesToChunks[ActiveStageIdx][VarIdx];
 									if (LastSetChunkIdx == INDEX_NONE)
 									{
-										const UEdGraphPin* DefaultPin = History.GetDefaultValuePin(VarIdx);
+										const FInputPin* DefaultPin = History.GetDefaultValuePin(VarIdx);
 										FNiagaraScriptVariableBinding DefaultBinding;
-										TOptional<ENiagaraDefaultMode> DefaultMode = SourceGraph->GetDefaultMode(AliasedVar, &DefaultBinding);
+										TOptional<ENiagaraDefaultMode> DefaultMode = GraphBridge::GetGraphDefaultMode(SourceGraph, AliasedVar, DefaultBinding);
 
 										// Do not error on defaults for parameter reads here; we may be entering a SetVariable function call which is setting the first default for a parameter.
 										const bool bTreatAsUnknownParameterMap = false;
@@ -8325,14 +8336,13 @@ void FHlslNiagaraTranslator::RegisterFunctionCall(ENiagaraScriptUsage ScriptUsag
 				ChunkStartsByMode[i] = ChunksByMode[i].Num();
 			}
 
-			FHlslNiagaraTranslator* ThisTranslator = this;
 			TArray<int32> FuncOutputChunks;
 
 			ENiagaraCodeChunkMode OldMode = CurrentBodyChunkMode;
 			CurrentBodyChunkMode = ENiagaraCodeChunkMode::Body;
 			{
 				NIAGARA_SCOPE_CYCLE_COUNTER(STAT_NiagaraEditor_Module_NiagaraHLSLTranslator_RegisterFunctionCall_Compile);
-				FuncOutput->Compile(ThisTranslator, FuncOutputChunks);
+				FuncOutput->Compile(this, FuncOutputChunks);
 			}
 			CurrentBodyChunkMode = OldMode;
 
@@ -8440,7 +8450,6 @@ void FHlslNiagaraTranslator::RegisterFunctionCall(ENiagaraScriptUsage ScriptUsag
 		NIAGARA_SCOPE_CYCLE_COUNTER(STAT_NiagaraEditor_Module_NiagaraHLSLTranslator_RegisterFunctionCall_Signature);
 
 		check(InSignature.IsValid());
-		check(Inputs.Num() > 0);
 
 		OutSignature = InSignature;
 
@@ -8488,6 +8497,7 @@ void FHlslNiagaraTranslator::RegisterFunctionCall(ENiagaraScriptUsage ScriptUsag
 		}
 		else
 		{
+			check(!Inputs.IsEmpty());
 
 			// Usually the DataInterface is the zeroth entry in the signature inputs, unless we are using the exec pin, in which case it is at index 1.
 			int32 DataInterfaceOwnerIdx = Inputs[0]; 
@@ -8556,7 +8566,7 @@ void FHlslNiagaraTranslator::RegisterFunctionCall(ENiagaraScriptUsage ScriptUsag
 	}
 }
 
-void FHlslNiagaraTranslator::GenerateFunctionCall(ENiagaraScriptUsage ScriptUsage, FNiagaraFunctionSignature& FunctionSignature, TArrayView<const int32> Inputs, TArray<int32>& Outputs)
+void FNiagaraHlslTranslator::GenerateFunctionCall(ENiagaraScriptUsage ScriptUsage, FNiagaraFunctionSignature& FunctionSignature, TArrayView<const int32> Inputs, TArray<int32>& Outputs)
 {
 	NIAGARA_SCOPE_CYCLE_COUNTER(STAT_NiagaraEditor_Module_NiagaraHLSLTranslator_GenerateFunctionCall);
 
@@ -8577,7 +8587,7 @@ void FHlslNiagaraTranslator::GenerateFunctionCall(ENiagaraScriptUsage ScriptUsag
 		FNiagaraTypeDefinition Type = FunctionSignature.Inputs[i].GetType();
 		if (Type.UnderlyingType != 0 && Type.ClassStructOrEnum == nullptr)
 		{
-			Error(FText::Format(LOCTEXT("InvalidTypeDefError", "Invalid data in niagara type definition, might be due to broken serialization or missing DI implementation! Variable: {0}"), FText::FromName(FunctionSignature.Inputs[i].GetName())), nullptr, nullptr);
+			Error(FText::Format(LOCTEXT("InvalidTypeDefError", "Invalid data in niagara type definition, might be due to broken serialization or missing DI implementation! Variable: {0}"), FText::FromName(FunctionSignature.Inputs[i].GetName())));
 			continue;
 		}
 
@@ -8585,7 +8595,7 @@ void FHlslNiagaraTranslator::GenerateFunctionCall(ENiagaraScriptUsage ScriptUsag
 		{
 			Error(FText::Format(LOCTEXT("InvalidInputNum", "Functon Input of %d is out of bounds in function signature! Variable: {0}"), 
 			FText::AsNumber(i),
-			FText::FromName(FunctionSignature.Inputs[i].GetName())), nullptr, nullptr);
+			FText::FromName(FunctionSignature.Inputs[i].GetName())));
 			continue;
 		}
 
@@ -8594,7 +8604,7 @@ void FHlslNiagaraTranslator::GenerateFunctionCall(ENiagaraScriptUsage ScriptUsag
 		{
 			if (!AddStructToDefinitionSet(Type))
 			{
-				Error(FText::Format(LOCTEXT("GetConstantFailTypeVar2", "Cannot handle type {0}! Variable: {1}"), Type.GetNameText(), FText::FromName(FunctionSignature.Inputs[i].GetName())), nullptr, nullptr);
+				Error(FText::Format(LOCTEXT("GetConstantFailTypeVar2", "Cannot handle type {0}! Variable: {1}"), Type.GetNameText(), FText::FromName(FunctionSignature.Inputs[i].GetName())));
 			}
 
 			int32 Input = Inputs[i];
@@ -8637,7 +8647,7 @@ void FHlslNiagaraTranslator::GenerateFunctionCall(ENiagaraScriptUsage ScriptUsag
 		{
 			if (!AddStructToDefinitionSet(Type))
 			{
-				Error(FText::Format(LOCTEXT("GetConstantFailTypeVar3", "Cannot handle type {0}! Variable: {1}"), Type.GetNameText(), FText::FromName(FunctionSignature.Outputs[i].GetName())), nullptr, nullptr);
+				Error(FText::Format(LOCTEXT("GetConstantFailTypeVar3", "Cannot handle type {0}! Variable: {1}"), Type.GetNameText(), FText::FromName(FunctionSignature.Outputs[i].GetName())));
 			}
 
 			int32 Output = INDEX_NONE;
@@ -8707,7 +8717,7 @@ void FHlslNiagaraTranslator::GenerateFunctionCall(ENiagaraScriptUsage ScriptUsag
 		{
 			FText Fmt = LOCTEXT("ErrorCompilingParameterFmt", "Error compiling parameter {0} in function call {1}");
 			FText ErrorText = FText::Format(Fmt, FText::FromString(MissingParam), FText::FromString(GetFunctionSignatureSymbol(FunctionSignature)));
-			Error(ErrorText, nullptr, nullptr);
+			Error(ErrorText);
 		}
 		return;
 	}
@@ -8720,7 +8730,7 @@ void FHlslNiagaraTranslator::GenerateFunctionCall(ENiagaraScriptUsage ScriptUsag
 	}
 }
 
-FString FHlslNiagaraTranslator::GetFunctionSignatureSymbol(const FNiagaraFunctionSignature& Sig)
+FString FNiagaraHlslTranslator::GetFunctionSignatureSymbol(const FNiagaraFunctionSignature& Sig)
 {
 	FString SigStr = Sig.GetNameString();
 	if (!Sig.OwnerName.IsNone() && Sig.OwnerName.IsValid())
@@ -8759,14 +8769,14 @@ FString FHlslNiagaraTranslator::GetFunctionSignatureSymbol(const FNiagaraFunctio
 	return GetSanitizedSymbolName(SigStr);
 }
 
-FString FHlslNiagaraTranslator::GenerateFunctionHlslPrototype(FStringView InVariableName, const FNiagaraFunctionSignature& FunctionSignature)
+FString FNiagaraHlslTranslator::GenerateFunctionHlslPrototype(FStringView InVariableName, const FNiagaraFunctionSignature& FunctionSignature)
 {
 	TStringBuilder<512> StringBuilder;
 	if (FunctionSignature.bMemberFunction)
 	{
 		StringBuilder.Append(InVariableName);
 		StringBuilder.Append(TEXT("."));
-		StringBuilder.Append(FHlslNiagaraTranslator::GetSanitizedSymbolName(FunctionSignature.Name.ToString()));
+		StringBuilder.Append(FNiagaraHlslTranslator::GetSanitizedSymbolName(FunctionSignature.Name.ToString()));
 
 		// Build specifiers
 		if (FunctionSignature.FunctionSpecifiers.Num())
@@ -8806,9 +8816,9 @@ FString FHlslNiagaraTranslator::GenerateFunctionHlslPrototype(FStringView InVari
 				bNeedsComma = true;
 
 				StringBuilder.Append(TEXT("in "));
-				StringBuilder.Append(FHlslNiagaraTranslator::GetStructHlslTypeName(InputVar.GetType()));
+				StringBuilder.Append(FNiagaraHlslTranslator::GetStructHlslTypeName(InputVar.GetType()));
 				StringBuilder.Append(TEXT(" In_"));
-				StringBuilder.Append(FHlslNiagaraTranslator::GetSanitizedSymbolName(InputVar.GetName().ToString()));
+				StringBuilder.Append(FNiagaraHlslTranslator::GetSanitizedSymbolName(InputVar.GetName().ToString()));
 			}
 
 			// Outputs
@@ -8822,9 +8832,9 @@ FString FHlslNiagaraTranslator::GenerateFunctionHlslPrototype(FStringView InVari
 				bNeedsComma = true;
 
 				StringBuilder.Append(TEXT("out "));
-				StringBuilder.Append(FHlslNiagaraTranslator::GetStructHlslTypeName(OutputVar.GetType()));
+				StringBuilder.Append(FNiagaraHlslTranslator::GetStructHlslTypeName(OutputVar.GetType()));
 				StringBuilder.Append(TEXT(" Out_"));
-				StringBuilder.Append(FHlslNiagaraTranslator::GetSanitizedSymbolName(OutputVar.GetName().ToString()));
+				StringBuilder.Append(FNiagaraHlslTranslator::GetSanitizedSymbolName(OutputVar.GetName().ToString()));
 			}
 			StringBuilder.Append(TEXT(");"));
 		}
@@ -8837,7 +8847,7 @@ FString FHlslNiagaraTranslator::GenerateFunctionHlslPrototype(FStringView InVari
 	return FString(StringBuilder.ToString());
 }
 
-FName FHlslNiagaraTranslator::GetDataInterfaceName(FName BaseName, const FString& UniqueEmitterName, bool bIsParameterMapDataInterface)
+FName FNiagaraHlslTranslator::GetDataInterfaceName(FName BaseName, const FString& UniqueEmitterName, bool bIsParameterMapDataInterface)
 {
 	if (UniqueEmitterName.IsEmpty() == false)
 	{
@@ -8855,7 +8865,7 @@ FName FHlslNiagaraTranslator::GetDataInterfaceName(FName BaseName, const FString
 	return BaseName;
 }
 
-FString FHlslNiagaraTranslator::GetFunctionIncludeStatement(const FNiagaraCustomHlslInclude& Include) const
+FString FNiagaraHlslTranslator::GetFunctionIncludeStatement(const FNiagaraCustomHlslInclude& Include) const
 {
 	TStringBuilder<128> IncludeStatement;
 
@@ -8873,7 +8883,7 @@ FString FHlslNiagaraTranslator::GetFunctionIncludeStatement(const FNiagaraCustom
 	return IncludeStatement.ToString();
 }
 
-FString FHlslNiagaraTranslator::GetFunctionSignature(const FNiagaraFunctionSignature& Sig)
+FString FNiagaraHlslTranslator::GetFunctionSignature(const FNiagaraFunctionSignature& Sig)
 {
 	FString SigStr = TEXT("void ") + GetFunctionSignatureSymbol(Sig);
 
@@ -8897,7 +8907,7 @@ FString FHlslNiagaraTranslator::GetFunctionSignature(const FNiagaraFunctionSigna
 				}
 
 				FNiagaraVariable SimInput = ConvertToSimulationVariable(Input);
-				SigStr += FHlslNiagaraTranslator::GetStructHlslTypeName(SimInput.GetType()) + TEXT(" In_") + FHlslNiagaraTranslator::GetSanitizedSymbolName(Input.GetName().ToString(), true);
+				SigStr += FNiagaraHlslTranslator::GetStructHlslTypeName(SimInput.GetType()) + TEXT(" In_") + FNiagaraHlslTranslator::GetSanitizedSymbolName(Input.GetName().ToString(), true);
 				++ParamIdx;
 			}
 		}
@@ -8921,7 +8931,7 @@ FString FHlslNiagaraTranslator::GetFunctionSignature(const FNiagaraFunctionSigna
 				}
 
 				FNiagaraVariable SimOutput = ConvertToSimulationVariable(Output);
-				SigStr += TEXT("out ") + FHlslNiagaraTranslator::GetStructHlslTypeName(SimOutput.GetType()) + TEXT(" ") + FHlslNiagaraTranslator::GetSanitizedSymbolName(TEXT("Out_") + Output.GetName().ToString());
+				SigStr += TEXT("out ") + FNiagaraHlslTranslator::GetStructHlslTypeName(SimOutput.GetType()) + TEXT(" ") + FNiagaraHlslTranslator::GetSanitizedSymbolName(TEXT("Out_") + Output.GetName().ToString());
 				++ParamIdx;
 			}
 		}
@@ -8937,19 +8947,7 @@ FString FHlslNiagaraTranslator::GetFunctionSignature(const FNiagaraFunctionSigna
 	return SigStr + TEXT(")");
 }
 
-int32 GetPinIndexById(TArrayView<UEdGraphPin* const> Pins, FGuid PinId)
-{
-	for (int32 i = 0; i < Pins.Num(); ++i)
-	{
-		if (Pins[i]->PinId == PinId)
-		{
-			return i;
-		}
-	}
-	return INDEX_NONE;
-}
-
-FNiagaraTypeDefinition FHlslNiagaraTranslator::GetChildType(const FNiagaraTypeDefinition& BaseType, const FName& PropertyName)
+FNiagaraTypeDefinition FNiagaraHlslTranslator::GetChildType(const FNiagaraTypeDefinition& BaseType, const FName& PropertyName)
 {
 	const UScriptStruct* Struct = BaseType.GetScriptStruct();
 	if (Struct != nullptr)
@@ -8992,7 +8990,7 @@ FNiagaraTypeDefinition FHlslNiagaraTranslator::GetChildType(const FNiagaraTypeDe
 	return FNiagaraTypeDefinition();
 }
 
-FString FHlslNiagaraTranslator::ComputeMatrixColumnAccess(const FString& Name)
+FString FNiagaraHlslTranslator::ComputeMatrixColumnAccess(const FString& Name)
 {
 	FString Value;
 	int32 Column = -1;
@@ -9014,12 +9012,12 @@ FString FHlslNiagaraTranslator::ComputeMatrixColumnAccess(const FString& Name)
 	}
 	else
 	{
-		Error(FText::FromString("Failed to generate type for " + Name + " up to path " + Value), nullptr, nullptr);
+		Error(FText::FromString("Failed to generate type for " + Name + " up to path " + Value));
 	}
 	return Value;
 }
 
-FString FHlslNiagaraTranslator::ComputeMatrixRowAccess(const FString& Name)
+FString FNiagaraHlslTranslator::ComputeMatrixRowAccess(const FString& Name)
 {
 	FString Value;
 	int32 Row = -1;
@@ -9040,12 +9038,12 @@ FString FHlslNiagaraTranslator::ComputeMatrixRowAccess(const FString& Name)
 	}
 	else
 	{
-		Error(FText::FromString("Failed to generate type for " + Name + " up to path " + Value), nullptr, nullptr);
+		Error(FText::FromString("Failed to generate type for " + Name + " up to path " + Value));
 	}
 	return Value;
 }
 
-FString FHlslNiagaraTranslator::NamePathToString(const FString& Prefix, const FNiagaraTypeDefinition& RootType, const TArray<FName>& NamePath)
+FString FNiagaraHlslTranslator::NamePathToString(const FString& Prefix, const FNiagaraTypeDefinition& RootType, const TArray<FName>& NamePath)
 {
 	// We need to deal with matrix parameters differently than any other type by using array syntax.
 	// As we recurse down the tree, we stay away of when we're dealing with a matrix and adjust 
@@ -9080,13 +9078,13 @@ FString FHlslNiagaraTranslator::NamePathToString(const FString& Prefix, const FN
 		}
 		else
 		{
-			Error(FText::FromString("Failed to generate type for " + Name + " up to path " + Value), nullptr, nullptr);
+			Error(FText::FromString("Failed to generate type for " + Name + " up to path " + Value));
 		}
 	}
 	return Value;
 }
 
-FString FHlslNiagaraTranslator::GenerateAssignment(const FNiagaraTypeDefinition& SrcPinType, const TArray<FName>& ConditionedSourcePath, const FNiagaraTypeDefinition& DestPinType, const TArray<FName>& ConditionedDestinationPath)
+FString FNiagaraHlslTranslator::GenerateAssignment(const FNiagaraTypeDefinition& SrcPinType, const TArray<FName>& ConditionedSourcePath, const FNiagaraTypeDefinition& DestPinType, const TArray<FName>& ConditionedDestinationPath)
 {
 	FString SourceDefinition = NamePathToString("{1}", SrcPinType, ConditionedSourcePath);
 	FString DestinationDefinition = NamePathToString("{0}", DestPinType, ConditionedDestinationPath);
@@ -9094,48 +9092,46 @@ FString FHlslNiagaraTranslator::GenerateAssignment(const FNiagaraTypeDefinition&
 	return DestinationDefinition + " = " + SourceDefinition;
 }
 
-void FHlslNiagaraTranslator::Convert(UNiagaraNodeConvert* Convert, TArrayView<const int32> Inputs, TArray<int32>& Outputs)
+template<typename GraphBridge>
+void TNiagaraHlslTranslator<GraphBridge>::Convert(const FConvertNode* Convert, TArrayView<const int32> Inputs, TArray<int32>& Outputs)
 {
 	if (ValidateTypePins(Convert) == false)
 	{
 		return;
 	}
 
-	FPinCollectorArray InputPins;
-	Convert->GetInputPins(InputPins);
-
-	FPinCollectorArray OutputPins;
-	Convert->GetOutputPins(OutputPins);
+	TArray<const FInputPin*> InputPins = GraphBridge::GetInputPins(Convert);
+	TArray<const FOutputPin*> OutputPins = GraphBridge::GetOutputPins(Convert);
 
 	// Add input struct definitions if necessary.
-	for (UEdGraphPin* InputPin : InputPins)
+	for (const FInputPin* InputPin : InputPins)
 	{
 		if (InputPin->PinType.PinCategory == UEdGraphSchema_Niagara::PinCategoryType ||
 			InputPin->PinType.PinCategory == UEdGraphSchema_Niagara::PinCategoryStaticType || 
 			InputPin->PinType.PinCategory == UEdGraphSchema_Niagara::PinCategoryEnum || 
 			InputPin->PinType.PinCategory == UEdGraphSchema_Niagara::PinCategoryStaticEnum)
 		{
-			FNiagaraTypeDefinition Type = Schema->PinToTypeDefinition(InputPin, ENiagaraStructConversion::Simulation);
+			FNiagaraTypeDefinition Type = GraphBridge::GetPinType(InputPin, ENiagaraStructConversion::Simulation);
 			if (!AddStructToDefinitionSet(Type))
 			{
-				Error(FText::Format(LOCTEXT("ConvertTypeError_InvalidInput", "Cannot handle input pin type {0}! Pin: {1}"), Type.GetNameText(), InputPin->GetDisplayName()), nullptr, nullptr);
+				Error(FText::Format(LOCTEXT("ConvertTypeError_InvalidInput", "Cannot handle input pin type {0}! Pin: {1}"), Type.GetNameText(), GraphBridge::GetPinFriendlyName(InputPin)), Convert, InputPin);
 			}
 		}
 	}
 
 	// Generate outputs.
 	Outputs.Reserve(Outputs.Num() + OutputPins.Num() + 1);
-	for (UEdGraphPin* OutputPin : OutputPins)
+	for (const FOutputPin* OutputPin : OutputPins)
 	{
 		if (OutputPin->PinType.PinCategory == UEdGraphSchema_Niagara::PinCategoryType ||
 			OutputPin->PinType.PinCategory == UEdGraphSchema_Niagara::PinCategoryStaticType ||
 			OutputPin->PinType.PinCategory == UEdGraphSchema_Niagara::PinCategoryEnum ||
 			OutputPin->PinType.PinCategory == UEdGraphSchema_Niagara::PinCategoryStaticEnum)
 		{
-			FNiagaraTypeDefinition Type = Schema->PinToTypeDefinition(OutputPin, ENiagaraStructConversion::Simulation);
+			FNiagaraTypeDefinition Type = GraphBridge::GetPinType(OutputPin, ENiagaraStructConversion::Simulation);
 			if (!AddStructToDefinitionSet(Type))
 			{
-				Error(FText::Format(LOCTEXT("ConvertTypeError_InvalidOutput", "Cannot handle output pin type {0}! Pin: {1}"), Type.GetNameText(), OutputPin->GetDisplayName()), nullptr, nullptr);
+				Error(FText::Format(LOCTEXT("ConvertTypeError_InvalidOutput", "Cannot handle output pin type {0}! Pin: {1}"), Type.GetNameText(), GraphBridge::GetPinFriendlyName(OutputPin)), Convert, OutputPin);
 			}
 			
 			// The convert node should already have issued errors if the connections aren't complete.
@@ -9147,27 +9143,24 @@ void FHlslNiagaraTranslator::Convert(UNiagaraNodeConvert* Convert, TArrayView<co
 		}
 	}
 
-	// Add an additional invalid output for the add pin which doesn't get compiled.
-	Outputs.Add(INDEX_NONE);
-
 	// Set output values based on connections.
-	for (FNiagaraConvertConnection& Connection : Convert->GetConnections())
+	for (const FConvertConnection& Connection : GraphBridge::GetConvertConnections(Convert))
 	{
-		int32 SourceIndex = GetPinIndexById(InputPins, Connection.SourcePinId);
-		int32 DestinationIndex = GetPinIndexById(OutputPins, Connection.DestinationPinId);
-		if (SourceIndex != INDEX_NONE && SourceIndex < Inputs.Num() && DestinationIndex != INDEX_NONE && DestinationIndex < Outputs.Num())
+		int32 SourceIndex = GraphBridge::GetPinIndexById(InputPins, Connection.SourcePinId);
+		int32 DestinationIndex = GraphBridge::GetPinIndexById(OutputPins, Connection.DestinationPinId);
+		if (Inputs.IsValidIndex(SourceIndex) && Outputs.IsValidIndex(DestinationIndex))
 		{
-			FNiagaraTypeDefinition SrcPinType = Schema->PinToTypeDefinition(InputPins[SourceIndex], ENiagaraStructConversion::Simulation);
+			FNiagaraTypeDefinition SrcPinType = GraphBridge::GetPinType(InputPins[SourceIndex], ENiagaraStructConversion::Simulation);
 			if (!AddStructToDefinitionSet(SrcPinType))
 			{
-				Error(FText::Format(LOCTEXT("ConvertTypeError_InvalidSubpinInput", "Cannot handle input subpin type {0}! Subpin: {1}"), SrcPinType.GetNameText(), InputPins[SourceIndex]->GetDisplayName()), nullptr, nullptr);
+				Error(FText::Format(LOCTEXT("ConvertTypeError_InvalidSubpinInput", "Cannot handle input subpin type {0}! Subpin: {1}"), SrcPinType.GetNameText(), GraphBridge::GetPinDisplayName(InputPins[SourceIndex])));
 			}
 			TArray<FName> ConditionedSourcePath = ConditionPropertyPath(SrcPinType, Connection.SourcePath);
 
-			FNiagaraTypeDefinition DestPinType = Schema->PinToTypeDefinition(OutputPins[DestinationIndex], ENiagaraStructConversion::Simulation);
+			FNiagaraTypeDefinition DestPinType = GraphBridge::GetPinType(OutputPins[DestinationIndex], ENiagaraStructConversion::Simulation);
 			if (!AddStructToDefinitionSet(DestPinType))
 			{
-				Error(FText::Format(LOCTEXT("ConvertTypeError_InvalidSubpinOutput", "Cannot handle output subpin type type {0}! Subpin: {1}"), DestPinType.GetNameText(), OutputPins[SourceIndex]->GetDisplayName()), nullptr, nullptr);
+				Error(FText::Format(LOCTEXT("ConvertTypeError_InvalidSubpinOutput", "Cannot handle output subpin type type {0}! Subpin: {1}"), DestPinType.GetNameText(), GraphBridge::GetPinDisplayName(OutputPins[SourceIndex])));
 			}
 			TArray<FName> ConditionedDestinationPath = ConditionPropertyPath(DestPinType, Connection.DestinationPath);
 
@@ -9181,7 +9174,8 @@ void FHlslNiagaraTranslator::Convert(UNiagaraNodeConvert* Convert, TArrayView<co
 	}
 }
 
-void FHlslNiagaraTranslator::If(UNiagaraNodeIf* IfNode, TArray<FNiagaraVariable>& Vars, int32 Condition, TArray<int32>& PathA, TArray<int32>& PathB, TArray<int32>& Outputs)
+template<typename GraphBridge>
+void TNiagaraHlslTranslator<GraphBridge>::If(const FIfNode* IfNode, const TArray<FNiagaraVariable>& Vars, int32 Condition, TArray<int32>& PathA, TArray<int32>& PathB, TArray<int32>& Outputs)
 {
 	NIAGARA_SCOPE_CYCLE_COUNTER(STAT_NiagaraEditor_HlslTranslator_If);
 
@@ -9192,9 +9186,9 @@ void FHlslNiagaraTranslator::If(UNiagaraNodeIf* IfNode, TArray<FNiagaraVariable>
 	TArray<FString> OutSymbols;
 	OutSymbols.Reserve(Vars.Num());
 	int32 PinIdx = 1;
-	for (FNiagaraVariable& Var : Vars)
+	for (const FNiagaraVariable& Var : Vars)
 	{
-		FNiagaraTypeDefinition Type = Schema->PinToTypeDefinition(IfNode->GetInputPin(PinIdx++), ENiagaraStructConversion::Simulation);
+		FNiagaraTypeDefinition Type = GraphBridge::GetPinType(GraphBridge::GetInputPin(IfNode, PinIdx), ENiagaraStructConversion::Simulation);
 		if (!AddStructToDefinitionSet(Type))
 		{
 			FText OutErrorMessage = FText::Format(LOCTEXT("If_UnknownNumeric", "Variable in If node uses invalid type. Var: {0} Type: {1}"),
@@ -9204,6 +9198,8 @@ void FHlslNiagaraTranslator::If(UNiagaraNodeIf* IfNode, TArray<FNiagaraVariable>
 		}
 		OutSymbols.Add(GetUniqueSymbolName(*(Var.GetName().ToString() + TEXT("_IfResult"))));
 		Outputs.Add(AddBodyChunk(OutSymbols.Last(), TEXT(""), Type, true));
+
+		++PinIdx;
 	}
 	AddBodyChunk(TEXT(""), TEXT("if({0})\n\t{"), FNiagaraTypeDefinition::GetFloatDef(), Condition, false, false);
 	for (int32 i = 0; i < NumVars; ++i)
@@ -9220,12 +9216,10 @@ void FHlslNiagaraTranslator::If(UNiagaraNodeIf* IfNode, TArray<FNiagaraVariable>
 		BranchChunk.AddSourceChunk(PathB[i]);
 	}
 	AddBodyChunk(TEXT(""), TEXT("}"), FNiagaraTypeDefinition::GetFloatDef(), false, false);
-
-	// Add an additional invalid output for the add pin which doesn't get compiled.
-	Outputs.Add(INDEX_NONE);
 }
 
-void FHlslNiagaraTranslator::Select(UNiagaraNodeSelect* SelectNode, int32 Selector, const TArray<FNiagaraVariable>& OutputVariables, TMap<int32, TArray<int32>>& Options, TArray<int32>& Outputs)
+template<typename GraphBridge>
+void TNiagaraHlslTranslator<GraphBridge>::Select(const FSelectNode* SelectNode, int32 Selector, const TArray<FNiagaraVariable>& OutputVariables, TMap<int32, TArray<int32>>& Options, TArray<int32>& Outputs)
 {
 	SCOPE_CYCLE_COUNTER(STAT_NiagaraEditor_HlslTranslator_Select);
 	
@@ -9242,7 +9236,7 @@ void FHlslNiagaraTranslator::Select(UNiagaraNodeSelect* SelectNode, int32 Select
 			FText OutErrorMessage = FText::Format(LOCTEXT("Select_UnknownNumeric", "Output type in Select node uses invalid type. Type: {0}"),
 				Variable.GetType().GetNameText());
 
-			Error(OutErrorMessage, SelectNode, SelectNode->GetOutputPin(Variable));
+			Error(OutErrorMessage, SelectNode, GraphBridge::GetSelectOutputPin(SelectNode, Variable));
 		}
 	}
 
@@ -9304,12 +9298,10 @@ void FHlslNiagaraTranslator::Select(UNiagaraNodeSelect* SelectNode, int32 Select
 			AddBodyChunk(TEXT(""), TEXT("}"), FNiagaraTypeDefinition::GetFloatDef(), false, false);
 		}
 	}
-
-	// Add an additional invalid output for the add pin which doesn't get compiled.
-	Outputs.Add(INDEX_NONE);
 }
 
-void FHlslNiagaraTranslator::FindConstantValue(int32 InputCompileResult, const FNiagaraTypeDefinition& TypeDef, FString& Value, FNiagaraVariable& Variable)
+template<typename GraphBridge>
+void TNiagaraHlslTranslator<GraphBridge>::FindConstantValue(int32 InputCompileResult, const FNiagaraTypeDefinition& TypeDef, FString& Value, FNiagaraVariable& Variable)
 {
 	FNiagaraEditorModule& NiagaraEditorModule = FModuleManager::GetModuleChecked<FNiagaraEditorModule>("NiagaraEditor");
 	{
@@ -9464,11 +9456,12 @@ void FHlslNiagaraTranslator::FindConstantValue(int32 InputCompileResult, const F
 	}
 }
 
-void FHlslNiagaraTranslator::WriteCompilerTag(int32 InputCompileResult, const UEdGraphPin* Pin, bool bEmitMessageOnFailure, FNiagaraCompileEventSeverity FailureSeverity, const FString& Prefix)
+template<typename GraphBridge>
+void TNiagaraHlslTranslator<GraphBridge>::WriteCompilerTag(int32 InputCompileResult, const FPin* Pin, bool bEmitMessageOnFailure, FNiagaraCompileEventSeverity FailureSeverity, const FString& Prefix)
 {
 	FString Value;
-	FNiagaraTypeDefinition TypeDef = Schema->PinToTypeDefinition(Pin);
-	FNiagaraVariable Variable(TypeDef, Prefix.Len() ? *(Prefix + TEXT(".") + Pin->GetName()) : *Pin->GetName());
+	FNiagaraTypeDefinition TypeDef = GraphBridge::GetPinType(Pin, ENiagaraStructConversion::UserFacing);
+	FNiagaraVariable Variable(TypeDef, Prefix.Len() ? *(Prefix + TEXT(".") + Pin->PinName.ToString()) : Pin->PinName);
 
 
 	//If we're in an emitter script then the tag needs to be made per emitter with EmitterName.Tag
@@ -9488,7 +9481,7 @@ void FHlslNiagaraTranslator::WriteCompilerTag(int32 InputCompileResult, const UE
 
 	if (Value.Len() == 0 && Variable.IsDataAllocated() == false && bEmitMessageOnFailure)
 	{
-		Message(FailureSeverity, FText::FromString(TEXT("Output Compile Tag must be connected to a constant or a uniform variable to work! Ignoring the compile tag.")), Cast<UNiagaraNode>(Pin->GetOwningNode()), Pin);
+		Message(FailureSeverity, FText::FromString(TEXT("Output Compile Tag must be connected to a constant or a uniform variable to work! Ignoring the compile tag.")), GraphBridge::GetOwningNode(Pin), Pin);
 	}
 	else
 	{
@@ -9505,61 +9498,54 @@ void FHlslNiagaraTranslator::WriteCompilerTag(int32 InputCompileResult, const UE
 	}
 }
 
-int32 FHlslNiagaraTranslator::CompilePin(const UEdGraphPin* Pin)
+template<typename GraphBridge>
+int32 TNiagaraHlslTranslator<GraphBridge>::CompileInputPin(const FInputPin* Pin)
 {
 	NIAGARA_SCOPE_CYCLE_COUNTER(STAT_NiagaraEditor_HlslTranslator_CompilePin);
 
 	check(Pin);
-	int32 Ret = INDEX_NONE;
-	
-	if (Pin->Direction == EGPD_Input)
+	if (!ensure(Pin->Direction == EGPD_Input))
 	{
-		if (Pin->LinkedTo.Num() > 0)
-		{
-			if (Pin->LinkedTo[0])
-			{
-				FPinConnectionResponse ConnectionResponse = Schema->CanCreateConnection(Pin->LinkedTo[0], Pin);
-				if (ConnectionResponse.Response == CONNECT_RESPONSE_DISALLOW)
-				{
-					FText OutErrorMessage = FText::Format(LOCTEXT("InputConnectionDisallowed", "Input connection is not allowed! Reason: {0}"),
-						ConnectionResponse.Message);
+		return INDEX_NONE;
+	}
 
-					Error(OutErrorMessage, Cast<UNiagaraNode>(Pin->GetOwningNode()), Pin);
-				}
-			}
-			Ret = CompileOutputPin(Pin->LinkedTo[0]);
-		}
-		else if (!Pin->bDefaultValueIsIgnored && (Pin->PinType.PinCategory == UEdGraphSchema_Niagara::PinCategoryType || Pin->PinType.PinCategory == UEdGraphSchema_Niagara::PinCategoryStaticType))
+	if (const FOutputPin* LinkedPin = GraphBridge::GetLinkedOutputPin(Pin))
+	{
+		FPinConnectionResponse ConnectionResponse = GraphBridge::CanCreateConnection(LinkedPin, Pin);
+		if (ConnectionResponse.Response == CONNECT_RESPONSE_DISALLOW)
 		{
-			FNiagaraTypeDefinition TypeDef = Schema->PinToTypeDefinition(Pin);
-			if (TypeDef == FNiagaraTypeDefinition::GetParameterMapDef())
-			{
-				Error(FText::FromString(TEXT("Parameter Maps must be created via an Input Node, not the default value of a pin! Please connect to a valid input Parameter Map.")), Cast<UNiagaraNode>(Pin->GetOwningNode()), Pin);
-				return INDEX_NONE;
-			}
-			else
-			{
-				//No connections to this input so add the default as a const expression.			
-				FNiagaraVariable PinVar = Schema->PinToNiagaraVariable(Pin, true, ENiagaraStructConversion::Simulation);
-				return GetConstant(PinVar);
-			}
+			FText OutErrorMessage = FText::Format(LOCTEXT("InputConnectionDisallowed", "Input connection is not allowed! Reason: {0}"),
+				ConnectionResponse.Message);
+
+			Error(OutErrorMessage, GraphBridge::GetOwningNode(Pin), Pin);
 		}
-		else if (!Pin->bDefaultValueIsIgnored && (Pin->PinType.PinCategory == UEdGraphSchema_Niagara::PinCategoryEnum || Pin->PinType.PinCategory == UEdGraphSchema_Niagara::PinCategoryStaticEnum))
+		return CompileOutputPin(LinkedPin);
+	}
+	else if (!Pin->bDefaultValueIsIgnored && (Pin->PinType.PinCategory == UEdGraphSchema_Niagara::PinCategoryType || Pin->PinType.PinCategory == UEdGraphSchema_Niagara::PinCategoryStaticType))
+	{
+		FNiagaraTypeDefinition TypeDef = GraphBridge::GetPinType(Pin, ENiagaraStructConversion::UserFacing);
+		if (TypeDef == FNiagaraTypeDefinition::GetParameterMapDef())
 		{
-			//No connections to this input so add the default as a const expression.			
-			FNiagaraVariable PinVar = Schema->PinToNiagaraVariable(Pin, true, ENiagaraStructConversion::Simulation);
-			return GetConstant(PinVar);
+			Error(FText::FromString(TEXT("Parameter Maps must be created via an Input Node, not the default value of a pin! Please connect to a valid input Parameter Map.")), GraphBridge::GetOwningNode(Pin), Pin);
+			return INDEX_NONE;
 		}
+
+		//No connections to this input so add the default as a const expression.			
+		FNiagaraVariable PinVar = GraphBridge::GetPinVariable(Pin, true, ENiagaraStructConversion::Simulation);
+		return GetConstant(PinVar);
 	}
-	else
+	else if (!Pin->bDefaultValueIsIgnored && (Pin->PinType.PinCategory == UEdGraphSchema_Niagara::PinCategoryEnum || Pin->PinType.PinCategory == UEdGraphSchema_Niagara::PinCategoryStaticEnum))
 	{
-		Ret = CompileOutputPin(Pin);
+		//No connections to this input so add the default as a const expression.			
+		FNiagaraVariable PinVar = GraphBridge::GetPinVariable(Pin, true, ENiagaraStructConversion::Simulation);
+		return GetConstant(PinVar);
 	}
 
-	return Ret;
+	return INDEX_NONE;
 }
 
-int32 FHlslNiagaraTranslator::CompileOutputPin(const UEdGraphPin* InPin)
+template<typename GraphBridge>
+int32 TNiagaraHlslTranslator<GraphBridge>::CompileOutputPin(const FPin* InPin)
 {
 	NIAGARA_SCOPE_CYCLE_COUNTER(STAT_NiagaraEditor_HlslTranslator_CompileOutputPin);
 
@@ -9569,31 +9555,35 @@ int32 FHlslNiagaraTranslator::CompileOutputPin(const UEdGraphPin* InPin)
 		return INDEX_NONE;
 	}
 
-	UpdateStaticSwitchConstants(InPin->GetOwningNode());
+	UpdateStaticSwitchConstants(GraphBridge::GetMutableOwningNode(InPin));
 
 	// The incoming pin to compile may be pointing to a reroute node. If so, we just jump over it
 	// to where it really came from.
-	const UEdGraphPin* Pin = InPin;
+	const FPin* Pin = InPin;
 	if (Pin == nullptr || Pin->Direction != EGPD_Output)
 	{
-		Error(LOCTEXT("TraceOutputPinFailed", "Failed to trace pin to an output!"), Cast<UNiagaraNode>(InPin->GetOwningNode()), InPin);
+		Error(LOCTEXT("TraceOutputPinFailed", "Failed to trace pin to an output!"), GraphBridge::GetOwningNode(InPin), InPin);
 		return INDEX_NONE;
 	}
 
-	UNiagaraNode* Node = Cast<UNiagaraNode>(Pin->GetOwningNode());
+	const FNode* Node = GraphBridge::GetOwningNode(Pin);
 
-	if (!Node->IsNodeEnabled())
+	// if the node is disabled (which commonly happens when a module is disabled in the stack), we skip it and follow the execution pin to the next node
+	if (!GraphBridge::NodeIsEnabled(Node))
 	{
-		// if the node is disabled (which commonly happens when a module is disabled in the stack), we skip it and follow the execution pin to the next node
-		FPinCollectorArray InputPins;
-		Node->GetInputPins(InputPins);
-		UEdGraphPin* OutputPin = Node->GetOutputPin(0);
-		for (int32 PinIndex = 0; PinIndex < InputPins.Num(); PinIndex++)
+		const FPin* OutputPin = GraphBridge::GetOutputPins(Node)[0];
+		if ((OutputPin == Pin) && GraphBridge::GetPinType(Pin, ENiagaraStructConversion::UserFacing) == FNiagaraTypeDefinition::GetParameterMapDef())
 		{
-			UEdGraphPin* ExecPin = InputPins[PinIndex];
-			if (OutputPin == Pin && ExecPin && ExecPin->LinkedTo.Num() == 1 && Node->IsParameterMapPin(Pin))
+			TArray<const FInputPin*> InputPins = GraphBridge::GetInputPins(Node);
+			for (int32 PinIndex = 0; PinIndex < InputPins.Num(); PinIndex++)
 			{
-				return CompileOutputPin(ExecPin->LinkedTo[0]);
+				if (const FInputPin* ExecPin = InputPins[PinIndex])
+				{
+					if (const FOutputPin* ConnectedPin = GraphBridge::GetLinkedOutputPin(ExecPin))
+					{
+						return CompileOutputPin(ConnectedPin);
+					}
+				}
 			}
 		}
 		Error(LOCTEXT("TraceDisabledPinFailed", "Failed to trace output pin of disabled node to a valid input!"), Node, Pin);
@@ -9612,15 +9602,9 @@ int32 FHlslNiagaraTranslator::CompileOutputPin(const UEdGraphPin* InPin)
 		if (ValidateTypePins(Node))
 		{
 			TArray<int32> Outputs;
-			FPinCollectorArray OutputPins;
-			Node->GetOutputPins(OutputPins);
-			OutputPins.RemoveAll([](UEdGraphPin* Pin)
-			{
-				return Pin->bOrphanedPin == true;
-			});
+			TArray<const FOutputPin*> OutputPins = GraphBridge::GetOutputPins(Node);
 			
-			FHlslNiagaraTranslator* ThisTranslator = this;
-			Node->Compile(ThisTranslator, Outputs);
+			Node->Compile(this, Outputs);
 			// this requires the nodes to only compile their valid output pins - no orphaned pins
 			if (OutputPins.Num() == Outputs.Num())
 			{
@@ -9647,15 +9631,15 @@ int32 FHlslNiagaraTranslator::CompileOutputPin(const UEdGraphPin* InPin)
 	return Ret;
 }
 
-
-FString FHlslNiagaraTranslator::NodePinToMessage(FText MessageText, const UNiagaraNode* Node, const UEdGraphPin* Pin)
+template<typename GraphBridge>
+FString TNiagaraHlslTranslator<GraphBridge>::NodePinToMessage(FText MessageText, const FNode* Node, const FPin* Pin)
 {
 	FString NodePinStr = TEXT("");
 	FString NodePinPrefix = TEXT(" - ");
 	FString NodePinSuffix = TEXT("");
 	if (Node)
 	{
-		FString NodeTitle = Node->GetNodeTitle(ENodeTitleType::FullTitle).ToString();
+		FString NodeTitle = GraphBridge::GetNodeTitle(Node);
 		if (NodeTitle.Len() > 0)
 		{
 			NodePinStr += TEXT("Node: ") + NodeTitle;
@@ -9663,7 +9647,7 @@ FString FHlslNiagaraTranslator::NodePinToMessage(FText MessageText, const UNiaga
 		}
 		else
 		{
-			FString NodeName = Node->GetName();
+			FString NodeName = GraphBridge::GetNodeName(Node);
 			if (NodeName.Len() > 0)
 			{
 				NodePinStr += TEXT("Node: ") + NodeName;
@@ -9673,7 +9657,8 @@ FString FHlslNiagaraTranslator::NodePinToMessage(FText MessageText, const UNiaga
 	}
 	if (Pin)
 	{
-		NodePinStr += TEXT(" Pin: ") + (Pin->PinFriendlyName.ToString().Len() > 0 ? Pin->PinFriendlyName.ToString() : Pin->GetName());
+		const FText FriendlyPinName = GraphBridge::GetPinFriendlyName(Pin);
+		NodePinStr += TEXT(" Pin: ") + (FriendlyPinName.IsEmpty() ? Pin->PinName.ToString() : FriendlyPinName.ToString());
 		NodePinSuffix = TEXT(" - ");
 	}
 
@@ -9681,13 +9666,14 @@ FString FHlslNiagaraTranslator::NodePinToMessage(FText MessageText, const UNiaga
 	return MessageString;
 }
 
-void FHlslNiagaraTranslator::Message(FNiagaraCompileEventSeverity Severity, FText MessageText, const UNiagaraNode* InNode, const UEdGraphPin* Pin, FString ShortDescription)
+template<typename GraphBridge>
+void TNiagaraHlslTranslator<GraphBridge>::Message(FNiagaraCompileEventSeverity Severity, FText MessageText, const FNode* InNode, const FPin* Pin, FStringView ShortDescription)
 {
-	const UNiagaraNode* CurContextNode = ActiveHistoryForFunctionCalls.GetCallingContext();
-	const UNiagaraNode* TargetNode = InNode ? InNode : CurContextNode;
+	const FNode* CurContextNode = ActiveHistoryForFunctionCalls.GetCallingContext();
+	const FNode* TargetNode = InNode ? InNode : CurContextNode;
  
 	FString MessageString = NodePinToMessage(MessageText, TargetNode, Pin);
-	TranslateResults.CompileEvents.Add(FNiagaraCompileEvent(Severity, MessageString, ShortDescription, TargetNode ? TargetNode->NodeGuid : FGuid(), Pin ? Pin->PersistentGuid : FGuid(), GetCallstackGuids()));
+	TranslateResults.CompileEvents.Add(FNiagaraCompileEvent(Severity, MessageString, FString(ShortDescription), TargetNode ? TargetNode->NodeGuid : FGuid(), Pin ? Pin->PersistentGuid : FGuid(), GetCallstackGuids()));
  
 	if (Severity == FNiagaraCompileEventSeverity::Error)
 	{
@@ -9698,18 +9684,39 @@ void FHlslNiagaraTranslator::Message(FNiagaraCompileEventSeverity Severity, FTex
 		TranslateResults.NumWarnings++;
 	}
 }
- 
-void FHlslNiagaraTranslator::Error(FText ErrorText, const UNiagaraNode* InNode, const UEdGraphPin* Pin, FString ShortDescription)
+
+template<typename GraphBridge>
+void TNiagaraHlslTranslator<GraphBridge>::Message(FNiagaraCompileEventSeverity Severity, FText MessageText, FStringView ShortDescription)
+{
+	Message(Severity, MessageText, nullptr, nullptr, ShortDescription);
+}
+
+template<typename GraphBridge>
+void TNiagaraHlslTranslator<GraphBridge>::Error(FText ErrorText, const FNode* InNode, const FPin* Pin, FStringView ShortDescription)
 {
 	Message(FNiagaraCompileEventSeverity::Error, ErrorText, InNode, Pin, ShortDescription);
 }
- 
-void FHlslNiagaraTranslator::Warning(FText WarningText, const UNiagaraNode* InNode, const UEdGraphPin* Pin, FString ShortDescription)
+
+template<typename GraphBridge>
+void TNiagaraHlslTranslator<GraphBridge>::Error(FText ErrorText, FStringView ShortDescription)
+{
+	Error(ErrorText, nullptr, nullptr, ShortDescription);
+}
+
+template<typename GraphBridge>
+void TNiagaraHlslTranslator<GraphBridge>::Warning(FText WarningText, const FNode* InNode, const FPin* Pin, FStringView ShortDescription)
 {
 	Message(FNiagaraCompileEventSeverity::Warning, WarningText, InNode, Pin, ShortDescription);
 }
 
-void FHlslNiagaraTranslator::RegisterCompileDependency(const FNiagaraVariableBase& InVar, FText MessageText, const UNiagaraNode* Node, const UEdGraphPin* Pin, bool bEmitAsLinker, int32 ParamMapHistoryIdx)
+template<typename GraphBridge>
+void TNiagaraHlslTranslator<GraphBridge>::Warning(FText WarningText, FStringView ShortDescription)
+{
+	Warning(WarningText, nullptr, nullptr, ShortDescription);
+}
+
+template<typename GraphBridge>
+void TNiagaraHlslTranslator<GraphBridge>::RegisterCompileDependency(const FNiagaraVariableBase& InVar, FText MessageText, const FNode* Node, const FPin* Pin, bool bEmitAsLinker, int32 ParamMapHistoryIdx)
 {
 	if (FNiagaraCVarUtilities::GetShouldEmitMessagesForFailIfNotSet() == false)
 	{
@@ -9729,8 +9736,8 @@ void FHlslNiagaraTranslator::RegisterCompileDependency(const FNiagaraVariableBas
 	if (bEmitAsLinker)																					
 	{			
 		bool bVarFromCustomIterationNamespaceOverride = ParamMapHistories[ParamMapHistoryIdx].IsVariableFromCustomIterationNamespaceOverride(InVar);
-		const UNiagaraNode* CurContextNode = ActiveHistoryForFunctionCalls.GetCallingContext();
-		const UNiagaraNode* TargetNode = Node ? Node : CurContextNode;
+		const FNode* CurContextNode = ActiveHistoryForFunctionCalls.GetCallingContext();
+		const FNode* TargetNode = Node ? Node : CurContextNode;
 
 		FString MessageString = NodePinToMessage(MessageText, TargetNode, Pin);
 		TranslateResults.CompileDependencies.AddUnique(FNiagaraCompileDependency(InVar, MessageString, TargetNode ? TargetNode->NodeGuid : FGuid(), Pin ? Pin->PersistentGuid : FGuid(), GetCallstackGuids(), bVarFromCustomIterationNamespaceOverride));
@@ -9743,7 +9750,7 @@ void FHlslNiagaraTranslator::RegisterCompileDependency(const FNiagaraVariableBas
 	}
 }
 
-bool FHlslNiagaraTranslator::GetFunctionParameter(const FNiagaraVariable& Parameter, int32& OutParam)const
+bool FNiagaraHlslTranslator::GetFunctionParameter(const FNiagaraVariable& Parameter, int32& OutParam)const
 {
 	// Assume that it wasn't bound by default.
 	OutParam = INDEX_NONE;
@@ -9759,19 +9766,19 @@ bool FHlslNiagaraTranslator::GetFunctionParameter(const FNiagaraVariable& Parame
 	return false;
 }
 
-int32 FHlslNiagaraTranslator::GetUniqueCallerID()
+int32 FNiagaraHlslTranslator::GetUniqueCallerID()
 {
 	if (!TranslationStages[ActiveStageIdx].bCallIDInitialized)
 	{
 		// The Call ID is changed every time a compiled node requests it, but we want to randomize it a bit from the start.
 		// Otherwise compilation units all start from the same ID (resulting in the same chain of generated randoms).
-		TranslationStages[ActiveStageIdx].CurrentCallID = (int32)(GetTypeHash(CompileData->EmitterUniqueName) + (uint8)TranslationStages[ActiveStageIdx].ScriptUsage * 1024);
+		TranslationStages[ActiveStageIdx].CurrentCallID = (int32)(GetTypeHash(GetEmitterUniqueName()) + (uint8)TranslationStages[ActiveStageIdx].ScriptUsage * 1024);
 		TranslationStages[ActiveStageIdx].bCallIDInitialized = true;
 	}
 	return TranslationStages[ActiveStageIdx].CurrentCallID++;
 }
 
-bool FHlslNiagaraTranslator::CanReadAttributes()const
+bool FNiagaraHlslTranslator::CanReadAttributes()const
 {
 	if (UNiagaraScript::IsParticleUpdateScript(TranslationStages[ActiveStageIdx].ScriptUsage))
 	{
@@ -9780,7 +9787,8 @@ bool FHlslNiagaraTranslator::CanReadAttributes()const
 	return false;
 }
 
-ENiagaraScriptUsage FHlslNiagaraTranslator::GetCurrentUsage() const
+template<typename GraphBridge>
+ENiagaraScriptUsage TNiagaraHlslTranslator<GraphBridge>::GetCurrentUsage() const
 {
 	if (UNiagaraScript::IsParticleScript(CompileOptions.TargetUsage))
 	{
@@ -9810,7 +9818,7 @@ ENiagaraScriptUsage FHlslNiagaraTranslator::GetCurrentUsage() const
 	}
 }
 
-ENiagaraScriptUsage FHlslNiagaraTranslator::GetTargetUsage() const
+ENiagaraScriptUsage FNiagaraHlslTranslator::GetTargetUsage() const
 {
 	if (CompileOptions.TargetUsage == ENiagaraScriptUsage::ParticleGPUComputeScript) // Act as if building spawn script.
 	{
@@ -9823,13 +9831,141 @@ ENiagaraScriptUsage FHlslNiagaraTranslator::GetTargetUsage() const
 	return CompileOptions.TargetUsage;
 }
 
-FGuid FHlslNiagaraTranslator::GetTargetUsageId() const
+FGuid FNiagaraHlslTranslator::GetTargetUsageId() const
 {
 	return CompileOptions.TargetUsageId;
 }
 //////////////////////////////////////////////////////////////////////////
 
-FString FHlslNiagaraTranslator::GetHlslDefaultForType(const FNiagaraTypeDefinition& Type)
+template<typename GraphBridge>
+void TNiagaraHlslTranslator<GraphBridge>::TrimAttributes(const FNiagaraCompileOptions& InCompileOptions, TArray<FNiagaraVariable>& Attributes)
+{
+	using FNiagaraAttributeTrimming = FNiagaraAttributeTrimmerHelper<GraphBridge>;
+
+	if (!UNiagaraScript::IsParticleScript(InCompileOptions.TargetUsage))
+	{
+		return;
+	}
+
+	const bool SafeTrimAttributesEnabled = InCompileOptions.AdditionalDefines.Contains(TEXT("TrimAttributesSafe"));
+	const bool AggressiveTrimAttributesEnabled = InCompileOptions.AdditionalDefines.Contains(TEXT("TrimAttributes"));
+
+	if (SafeTrimAttributesEnabled || AggressiveTrimAttributesEnabled)
+	{
+		// validate that the attributes have unique sanitized names
+		{
+			bool HasOverlappingNames = false;
+
+			TMap<FString, FNiagaraVariableBase> SanitizedNames;
+			for (const FNiagaraVariableBase& Attribute : Attributes)
+			{
+				const FString SanitizedName = GetSanitizedSymbolName(Attribute.GetName().ToString());
+				if (const FNiagaraVariableBase* ExistingVariable = SanitizedNames.Find(SanitizedName))
+				{
+					HasOverlappingNames = true;
+				}
+				else
+				{
+					SanitizedNames.Add(SanitizedName, Attribute);
+				}
+			}
+
+			// the trimming algorithm doesn't work when names are overlapping, so just early out of the function
+			if (HasOverlappingNames)
+			{
+				return;
+			}
+		}
+
+		const bool bRequiresPersistentIDs = InCompileOptions.AdditionalDefines.Contains(TEXT("RequiresPersistentIDs"));
+
+		// we want to use the ParamMapHistories of both the particle update and spawn scripts because they need to
+		// agree to define a unified attribute set
+		TArray<const FParamMapHistory*, TInlineAllocator<2>> LocalParamHistories;
+		for (const FParamMapHistory& History : OtherOutputParamMapHistories)
+		{
+			if (UNiagaraScript::IsParticleScript(History.OriginatingScriptUsage))
+			{
+				// for now we'll be disabling attribute trimming if a family of particle scripts contain generation of
+				// additional dataset writes (events) as we don't have access to the connectivity of it's variables as
+				// we do for the rest of the script
+				if (History.AdditionalDataSetWrites.Num())
+				{
+					return;
+				}
+
+				LocalParamHistories.Add(&History);
+			}
+		}
+
+		// check through the AdditionalDefines to see if any variables have been explicitly preserved
+		TSet<FName> AttributesToPreserve;
+
+		for (const FString& AdditionalDefine : InCompileOptions.AdditionalDefines)
+		{
+			const FString PreserveTag = TEXT("PreserveAttribute=");
+			if (AdditionalDefine.StartsWith(PreserveTag))
+			{
+				AttributesToPreserve.Add(*AdditionalDefine.RightChop(PreserveTag.Len()));
+			}
+		}
+
+		AttributesToPreserve.Add(SYS_PARAM_INSTANCE_ALIVE.GetName());
+		AttributesToPreserve.Add(SYS_PARAM_PARTICLES_UNIQUE_ID.GetName());
+		if (bRequiresPersistentIDs)
+		{
+			AttributesToPreserve.Add(SYS_PARAM_PARTICLES_ID.GetName());
+		}
+
+		const TArray<FNiagaraVariable> PreTrimmedAttributes = Attributes;
+
+		if (SafeTrimAttributesEnabled)
+		{
+			FNiagaraAttributeTrimming::TrimAttributes_Safe(LocalParamHistories, AttributesToPreserve, Attributes);
+		}
+		else if (AggressiveTrimAttributesEnabled)
+		{
+			FNiagaraAttributeTrimming::TrimAttributes_Aggressive(CompileDuplicateData, LocalParamHistories, AttributesToPreserve, Attributes);
+		}
+
+		for (const FNiagaraVariable& Attribute : PreTrimmedAttributes)
+		{
+			if (!Attributes.Contains(Attribute))
+			{
+				TranslateResults.CompileTags.Emplace(Attribute, TEXT("Trimmed"));
+			}
+		}
+	}
+}
+
+TUniquePtr<INiagaraHlslTranslator> INiagaraHlslTranslator::CreateTranslator(const FNiagaraCompileRequestDataBase* InCompileData, const FNiagaraCompileRequestDuplicateDataBase* InDuplicateData)
+{
+	// translate everything
+	const FNiagaraCompileRequestData* CompileRequest = static_cast<const FNiagaraCompileRequestData*>(InCompileData);
+	const FNiagaraCompileRequestDuplicateData* CompileRequestDuplicate = static_cast<const FNiagaraCompileRequestDuplicateData*>(InDuplicateData);
+
+	return MakeUnique<TNiagaraHlslTranslator<FNiagaraCompilationGraphBridge>>(CompileRequest, CompileRequestDuplicate);
+}
+
+template<typename GraphBridge>
+const FString& TNiagaraHlslTranslator<GraphBridge>::GetEmitterUniqueName() const
+{
+	return CompileData->EmitterUniqueName;
+}
+
+template<typename GraphBridge>
+TConstArrayView<FNiagaraVariable> TNiagaraHlslTranslator<GraphBridge>::GetStaticVariables() const
+{
+	return CompileData->StaticVariables;
+}
+
+template<typename GraphBridge>
+UNiagaraDataInterface* TNiagaraHlslTranslator<GraphBridge>::GetDataInterfaceCDO(UClass* DIClass) const
+{
+	return CompileDuplicateData->GetDuplicatedDataInterfaceCDOForClass(DIClass);
+}
+
+FString INiagaraHlslTranslator::GetHlslDefaultForType(const FNiagaraTypeDefinition& Type)
 {
 	if (Type == FNiagaraTypeDefinition::GetFloatDef())
 	{
@@ -9869,7 +10005,7 @@ FString FHlslNiagaraTranslator::GetHlslDefaultForType(const FNiagaraTypeDefiniti
 	}
 }
 
-bool FHlslNiagaraTranslator::IsBuiltInHlslType(const FNiagaraTypeDefinition& Type)
+bool INiagaraHlslTranslator::IsBuiltInHlslType(const FNiagaraTypeDefinition& Type)
 {
 	return
 		Type == FNiagaraTypeDefinition::GetFloatDef() ||
@@ -9885,7 +10021,7 @@ bool FHlslNiagaraTranslator::IsBuiltInHlslType(const FNiagaraTypeDefinition& Typ
 		Type.IsSameBaseDefinition(FNiagaraTypeDefinition::GetBoolDef());
 }
 
-FString FHlslNiagaraTranslator::GetStructHlslTypeName(const FNiagaraTypeDefinition& Type)
+FString INiagaraHlslTranslator::GetStructHlslTypeName(const FNiagaraTypeDefinition& Type)
 {
 	check(FNiagaraTypeHelper::IsLWCType(Type) == false);
 
@@ -9931,7 +10067,7 @@ FString FHlslNiagaraTranslator::GetStructHlslTypeName(const FNiagaraTypeDefiniti
 	}
 }
 
-FString FHlslNiagaraTranslator::GetPropertyHlslTypeName(const FProperty* Property)
+FString INiagaraHlslTranslator::GetPropertyHlslTypeName(const FProperty* Property)
 {
 	if (Property->IsA(FFloatProperty::StaticClass()))
 	{
@@ -9968,7 +10104,7 @@ FString FHlslNiagaraTranslator::GetPropertyHlslTypeName(const FProperty* Propert
 	}
 }
 
-FString FHlslNiagaraTranslator::BuildHLSLStructDecl(const FNiagaraTypeDefinition& Type, FText& OutErrorMessage, bool bGpuScript)
+FString INiagaraHlslTranslator::BuildHLSLStructDecl(const FNiagaraTypeDefinition& Type, FText& OutErrorMessage, bool bGpuScript)
 {
 	if (!IsBuiltInHlslType(Type))
 	{
@@ -9993,7 +10129,7 @@ FString FHlslNiagaraTranslator::BuildHLSLStructDecl(const FNiagaraTypeDefinition
 				Property->IsA(FEnumProperty::StaticClass()) ||
 				Property->IsA(FByteProperty::StaticClass()) ||
 				Property->IsA(FBoolProperty::StaticClass())
-			)
+				)
 			{
 				PropertyTypeName = TEXT("int");
 				PropertyTypeSize = 4;
@@ -10050,7 +10186,7 @@ FString FHlslNiagaraTranslator::BuildHLSLStructDecl(const FNiagaraTypeDefinition
 	return TEXT("");
 }
 
-bool FHlslNiagaraTranslator::IsHlslBuiltinVector(const FNiagaraTypeDefinition& Type)
+bool INiagaraHlslTranslator::IsHlslBuiltinVector(const FNiagaraTypeDefinition& Type)
 {
 	if ((Type == FNiagaraTypeDefinition::GetVec2Def()) ||
 		(Type == FNiagaraTypeDefinition::GetVec3Def()) ||
@@ -10064,8 +10200,7 @@ bool FHlslNiagaraTranslator::IsHlslBuiltinVector(const FNiagaraTypeDefinition& T
 	return false;
 }
 
-
-bool FHlslNiagaraTranslator::AddStructToDefinitionSet(const FNiagaraTypeDefinition& TypeDef)
+bool FNiagaraHlslTranslator::AddStructToDefinitionSet(const FNiagaraTypeDefinition& TypeDef)
 {
 	// First make sure that this is a type that we do need to define...
 	if (IsBuiltInHlslType(TypeDef))
@@ -10119,7 +10254,7 @@ bool FHlslNiagaraTranslator::AddStructToDefinitionSet(const FNiagaraTypeDefiniti
 	return true;
 }
 
-TArray<FName> FHlslNiagaraTranslator::ConditionPropertyPath(const FNiagaraTypeDefinition& Type, const TArray<FName>& InPath)
+TArray<FName> FNiagaraHlslTranslator::ConditionPropertyPath(const FNiagaraTypeDefinition& Type, const TArray<FName>& InPath)
 {
 	// TODO: Build something more extensible and less hard coded for path conditioning.
 	UScriptStruct* Struct = Type.GetScriptStruct();
@@ -10181,9 +10316,8 @@ TArray<FName> FHlslNiagaraTranslator::ConditionPropertyPath(const FNiagaraTypeDe
 }
 //////////////////////////////////////////////////////////////////////////
 
-
-
-
+// explicit template instantiation of the translators available based on the graph interfaces
+template class TNiagaraHlslTranslator<FNiagaraCompilationGraphBridge>;
 
 #undef NIAGARA_SCOPE_CYCLE_COUNTER
 #undef LOCTEXT_NAMESPACE
