Adding a templated layer in the hlsl translator to encapsulate logic surrounding graph implementation to allow us to share code between the existing UNiagaraGraph and the upcoming graph digest.  Gener

diff --git a/Engine/Plugins/FX/Niagara/Source/NiagaraEditor/Private/NiagaraHlslTranslator.h b/Engine/Plugins/FX/Niagara/Source/NiagaraEditor/Private/NiagaraHlslTranslator.h
--- a/Engine/Plugins/FX/Niagara/Source/NiagaraEditor/Private/NiagaraHlslTranslator.h
+++ b/Engine/Plugins/FX/Niagara/Source/NiagaraEditor/Private/NiagaraHlslTranslator.h
@@ -3,25 +3,18 @@
 #pragma once
 
 #include "CoreMinimal.h"
+
+#include "INiagaraCompiler.h"
 #include "Kismet2/CompilerResultsLog.h"
+#include "NiagaraCommon.h"
 #include "NiagaraScript.h"
-#include "NiagaraParameterMapHistory.h"
-#include "NiagaraShaderCompilationManager.h"
-#include "NiagaraDataInterface.h"
-#include "NiagaraScriptSource.h"
-#include "NiagaraSimulationStageCompileData.h"
+#include "NiagaraScriptSourceBase.h"
+#include "NiagaraShared.h"
+#include "NiagaraTypes.h"
 
 #include "NiagaraHlslTranslator.generated.h"
 
-class Error;
-class UNiagaraGraph;
-class UNiagaraNode;
-class UNiagaraNodeFunctionCall;
-class UNiagaraNodeCustomHlsl;
-class UNiagaraNodeOutput;
-class UNiagaraScriptSource;
-class UNiagaraNodeEmitter;
-class UNiagaraScriptVariable;
+struct FNiagaraTranslatorOutput;
 
 UENUM()
 enum class ENiagaraDataSetAccessMode : uint8
@@ -39,174 +32,15 @@ struct FNiagaraCustomHlslInclude
 	bool bIsVirtual = false;
 	FString FilePath;
 
-	bool operator==(const FNiagaraCustomHlslInclude& Other) const 
-	{ 
-		return bIsVirtual == Other.bIsVirtual && FilePath == Other.FilePath; 
-	}
-};
-
-
-/** Defines information about the results of a Niagara script compile. */
-struct FNiagaraTranslateResults
-{
-	/** Whether or not HLSL generation was successful */
-	bool bHLSLGenSucceeded = false;;
-
-	/** A results log with messages, warnings, and errors which occurred during the compile. */
-	TArray<FNiagaraCompileEvent> CompileEvents;
-	uint32 NumErrors;
-	uint32 NumWarnings;
-
-	TArray<FNiagaraCompileDependency> CompileDependencies;
-
-	TArray<FNiagaraCompilerTag > CompileTags;
-
-	/** A string representation of the compilation output. */
-	FString OutputHLSL;
-
-	FNiagaraTranslateResults() : NumErrors(0), NumWarnings(0)
-	{
-	}
-
-	static ENiagaraScriptCompileStatus TranslateResultsToSummary(const FNiagaraTranslateResults* CompileResults);
-};
-
-class FNiagaraCompileRequestData : public FNiagaraCompileRequestDataBase
-{
-public:
-	FNiagaraCompileRequestData()
+	bool operator==(const FNiagaraCustomHlslInclude& Other) const
 	{
-
-	}
-
-	virtual void GatherPreCompiledVariables(const FString& InNamespaceFilter, TArray<FNiagaraVariable>& OutVars) const override;
-	virtual FName ResolveEmitterAlias(FName VariableName) const override;
-
-	const FString& GetUniqueEmitterName() const { return EmitterUniqueName; }
-	void FinishPrecompile(const TArray<FNiagaraVariable>& EncounterableVariables, const TArray<FNiagaraVariable>& InStaticVariables, FCompileConstantResolver ConstantResolver, const TArray<ENiagaraScriptUsage>& UsagesToProcess, const TArray<class UNiagaraSimulationStageBase*>* SimStages, const TArray<FString> EmitterNames);
-	virtual int32 GetDependentRequestCount() const override {
-		return EmitterData.Num();
-	};
-	virtual TSharedPtr<FNiagaraCompileRequestDataBase, ESPMode::ThreadSafe> GetDependentRequest(int32 Index) override {
-		return EmitterData[Index];
+		return bIsVirtual == Other.bIsVirtual && FilePath == Other.FilePath;
 	}
-	virtual const FNiagaraCompileRequestDataBase* GetDependentRequest(int32 Index) const override
-	{
-		return EmitterData[Index].Get();
-	}
-	void AddRapidIterationParameters(const FNiagaraParameterStore& InParamStore, FCompileConstantResolver InResolver);
-	virtual bool GetUseRapidIterationParams() const override { return bUseRapidIterationParams; }
-
-	virtual bool GetDisableDebugSwitches() const override { return bDisableDebugSwitches; }
-
-	// Simulation Stage Variables. Sim stage of 0 is always Spawn/Update
-	TArray<FNiagaraSimulationStageCompilationData> CompileSimStageData;
-
-	struct FCompileDataInterfaceData
-	{
-		FString EmitterName;
-		ENiagaraScriptUsage Usage;
-		FGuid UsageId;
-		FNiagaraVariable Variable;
-		TArray<FString> ReadsEmitterParticleData;
-	};
-	TSharedPtr<TArray<FCompileDataInterfaceData>> SharedCompileDataInterfaceData;
-
-	TArray<FNiagaraVariable> EncounteredVariables;
-	FString EmitterUniqueName;
-	TArray<TSharedPtr<FNiagaraCompileRequestData, ESPMode::ThreadSafe>> EmitterData;
-	TWeakObjectPtr<UNiagaraScriptSource> Source;
-	FString SourceName;
-	bool bUseRapidIterationParams = true;
-	bool bDisableDebugSwitches = false;
-
-	UEnum* ENiagaraScriptCompileStatusEnum = nullptr;
-	UEnum* ENiagaraScriptUsageEnum = nullptr;
-
-	TArray<FNiagaraVariable> RapidIterationParams;
-
-	TMap<FGraphTraversalHandle, FString> PinToConstantValues;
-	TArray<FNiagaraVariable> StaticVariables;
-	TArray<FNiagaraVariable> StaticVariablesWithMultipleWrites;
-
-	template<typename T> T GetStaticVariableValue(const FNiagaraVariableBase Variable, const T DefaultValue) const
-	{
-		const int32 Index = StaticVariables.Find(Variable);
-		return Index != INDEX_NONE && StaticVariables[Index].IsDataAllocated() ? StaticVariables[Index].GetValue<T>() : DefaultValue;
-	}
-
-	void CompareAgainst(FNiagaraGraphCachedBuiltHistory* InCachedDataBase);
-
-	static void SortOutputNodesByDependencies(TArray<class UNiagaraNodeOutput*>& NodesToSort, const TArray<class UNiagaraSimulationStageBase*>* SimStages);
-};
-
-class FNiagaraCompileRequestDuplicateData : public FNiagaraCompileRequestDuplicateDataBase
-{
-public:
-	FNiagaraCompileRequestDuplicateData()
-	{
-	}
-
-	virtual bool IsDuplicateDataFor(UNiagaraSystem* InSystem, UNiagaraEmitter* InEmitter, UNiagaraScript* InScript) const override;
-	virtual void GetDuplicatedObjects(TArray<UObject*>& Objects) override;
-	virtual const TMap<FName, UNiagaraDataInterface*>& GetObjectNameMap() override;
-	virtual const UNiagaraScriptSourceBase* GetScriptSource() const override;
-	UNiagaraDataInterface* GetDuplicatedDataInterfaceCDOForClass(UClass* Class) const;
-
-	TArray<FNiagaraParameterMapHistory>& GetPrecomputedHistories() { return PrecompiledHistories; }
-	const TArray<FNiagaraParameterMapHistory>& GetPrecomputedHistories() const { return PrecompiledHistories; }
-
-	void DuplicateReferencedGraphs(UNiagaraGraph* InSrcGraph, UNiagaraGraph* InDupeGraph, ENiagaraScriptUsage InUsage, FCompileConstantResolver ConstantResolver, TMap<UNiagaraNodeFunctionCall*, ENiagaraScriptUsage> FunctionsWithUsage = TMap<UNiagaraNodeFunctionCall*, ENiagaraScriptUsage>());
-	void DeepCopyGraphs(UNiagaraScriptSource* ScriptSource, ENiagaraScriptUsage InUsage, FCompileConstantResolver ConstantResolver);
-	void DeepCopyGraphs(const FVersionedNiagaraEmitter& Emitter);
-	
-	void FinishPrecompileDuplicate(const TArray<FNiagaraVariable>& EncounterableVariables, const TArray<FNiagaraVariable>& InStaticVariables, FCompileConstantResolver ConstantResolver, const TArray<class UNiagaraSimulationStageBase*>* SimStages, const TArray<FNiagaraVariable>& InParamStore);
-	void CreateDataInterfaceCDO(TArrayView<UClass*> VariableDataInterfaces);
-
-	virtual int32 GetDependentRequestCount() const override	{ return EmitterData.Num();	}
-	virtual TSharedPtr<FNiagaraCompileRequestDuplicateDataBase, ESPMode::ThreadSafe> GetDependentRequest(int32 Index) override { return EmitterData[Index];	}
-
-	virtual void ReleaseCompilationCopies() override;
-
-	TWeakObjectPtr<UNiagaraSystem> OwningSystem;
-	TWeakObjectPtr<UNiagaraEmitter> OwningEmitter;
-	TArray<ENiagaraScriptUsage> ValidUsages;
-
-	TWeakObjectPtr<UNiagaraScriptSource> SourceDeepCopy;
-	TWeakObjectPtr<UNiagaraGraph> NodeGraphDeepCopy;
-
-	TArray<FNiagaraParameterMapHistory> PrecompiledHistories;
-
-	TArray<FNiagaraVariable> ChangedFromNumericVars;
-	FString EmitterUniqueName;
-	TArray<TSharedPtr<FNiagaraCompileRequestDuplicateData, ESPMode::ThreadSafe>> EmitterData;
-
-	struct FDuplicatedGraphData
-	{
-		UNiagaraScript* ClonedScript;
-		UNiagaraGraph* ClonedGraph;
-		TArray<UEdGraphPin*> CallInputs;
-		TArray<UEdGraphPin*> CallOutputs;
-		ENiagaraScriptUsage Usage;
-		bool bHasNumericParameters;
-	};
-
-	TSharedPtr<TMap<const UNiagaraGraph*, TArray<FDuplicatedGraphData>>> SharedSourceGraphToDuplicatedGraphsMap;
-	TSharedPtr<TMap<FName, UNiagaraDataInterface*>> SharedNameToDuplicatedDataInterfaceMap;
-	TSharedPtr<TMap<UClass*, UNiagaraDataInterface*>> SharedDataInterfaceClassToDuplicatedCDOMap;
-
-protected:
-	void DuplicateReferencedGraphsRecursive(UNiagaraGraph* InGraph, const FCompileConstantResolver& ConstantResolver, TMap<UNiagaraNodeFunctionCall*, ENiagaraScriptUsage> FunctionsWithUsage);
-private:
-	TArray<TWeakObjectPtr<UNiagaraScriptSource>> TrackedScriptSourceCopies;
 };
 
-
 /** Data which is generated from the hlsl by the VectorVMBackend and fed back into the */
 struct FNiagaraTranslatorOutput
 {
-	FNiagaraTranslatorOutput() {}
-
 	FNiagaraVMExecutableData ScriptData;
 
 	/** Ordered table of functions actually called by the VM script. */
@@ -223,15 +57,6 @@ struct FNiagaraTranslatorOutput
 
 };
 
-struct FCompiledPin
-{
-	int32 CompilationIndex;
-	UEdGraphPin* Pin;
-
-	FCompiledPin(int32 CompilationIndex, UEdGraphPin* Pin) : CompilationIndex(CompilationIndex), Pin(Pin)
-	{}
-};
-
 enum class ENiagaraCodeChunkMode : uint8
 {
 	GlobalConstant,
@@ -322,46 +147,17 @@ struct FNiagaraCodeChunk
 	}
 };
 
-class NIAGARAEDITOR_API FHlslNiagaraTranslatorOptions
-{
-public:
-	FHlslNiagaraTranslatorOptions()
-		: SimTarget(ENiagaraSimTarget::CPUSim)
-		, bParameterRapidIteration(true)
-		, bDisableDebugSwitches(false)
-	{
-
-	}
-
-	ENiagaraSimTarget SimTarget;
-
-	/** Any parameters in these namespaces will be pulled from an "InstanceParameters" dataset rather than from the uniform table. */
-	TArray<FString> InstanceParameterNamespaces;
-
-	/** Whether or not to treat top-level module variables as external values for rapid iteration without need for compilation.*/
-	bool bParameterRapidIteration;
-
-	/** Should we disable debug switches during translation. */
-	bool bDisableDebugSwitches;
-
-	/** Whether or not to override top-level module variables with values from the constant override table. This is only used for variables that were candidates for rapid iteration.*/
-	TArray<FNiagaraVariable> OverrideModuleConstants;
-
-
-};
-
-class NIAGARAEDITOR_API FHlslNiagaraTranslationStage
+class FNiagaraHlslTranslationStage
 {
 public:
-	FHlslNiagaraTranslationStage(ENiagaraScriptUsage InScriptUsage, FGuid InUsageId)
+	FNiagaraHlslTranslationStage(ENiagaraScriptUsage InScriptUsage, FGuid InUsageId)
 		: ScriptUsage(InScriptUsage)
 		, UsageId(InUsageId)
 	{
 	}
-		
+
 	ENiagaraScriptUsage ScriptUsage;
 	FGuid UsageId;
-	UNiagaraNodeOutput* OutputNode = nullptr;
 	FString PassNamespace;
 	bool bInterpolatePreviousParams = false;
 	bool bCopyPreviousParams = true;
@@ -397,10 +193,10 @@ public:
 	int32 ParamMapHistoryIndex = INDEX_NONE;
 
 	bool ShouldDoSpawnOnlyLogic() const;
-	bool IsRelevantToSpawnForStage(const FNiagaraParameterMapHistory& InHistory, const FNiagaraVariable& InAliasedVar, const FNiagaraVariable& InVar) const;
+	template<typename GraphBridge>
+	bool IsRelevantToSpawnForStage(const typename GraphBridge::FParamMapHistory& InHistory, const FNiagaraVariable& InAliasedVar, const FNiagaraVariable& InVar) const;
 
 	bool IsExternalConstantNamespace(const FNiagaraVariable& InVar, ENiagaraScriptUsage InTargetUsage, uint32 InTargetBitmask);
-
 	FName GetIterationDataInterface() const { return IterationSourceType == ENiagaraIterationSource::DataInterface ? IterationDataInterface : NAME_None; }
 
 	int32 CurrentCallID = 0;
@@ -413,7 +209,7 @@ struct FunctionNodeStackEntry
 	TSet<FString> CulledFunctionNames;
 };
 
-class NIAGARAEDITOR_API FHlslNiagaraTranslator
+class NIAGARAEDITOR_API FNiagaraHlslTranslator : public INiagaraHlslTranslator
 {
 public:
 
@@ -432,26 +228,16 @@ public:
 	};
 
 protected:
-	const FNiagaraCompileRequestData* CompileData;
-	const FNiagaraCompileRequestDuplicateData* CompileDuplicateData;
 	FNiagaraCompileOptions CompileOptions;
 
 	FHlslNiagaraTranslatorOptions TranslationOptions;
 
-	const class UEdGraphSchema_Niagara* Schema;
-
 	/** The set of all generated code chunks for this script. */
 	TArray<FNiagaraCodeChunk> CodeChunks;
 
 	/** Array of code chunks of each different type. */
 	TArray<int32> ChunksByMode[(int32)ENiagaraCodeChunkMode::Num];
 
-	/**
-	Map of Pins to compiled code chunks. Allows easy reuse of previously compiled pins.
-	A stack so that we can track pin reuse within function calls but not have cached pins cross talk with subsequent calls to the same function.
-	*/
-	TArray<TMap<const UEdGraphPin*, int32>> PinToCodeChunks;
-
 	/** The combined output of the compilation of this script. This is temporary and will be reworked soon. */
 	FNiagaraTranslatorOutput CompilationOutput;
 
@@ -491,17 +277,12 @@ protected:
 	TMap<FNiagaraFunctionSignature, TArray<FName>> FunctionStageReadTargets;
 	TArray<TArray<FName>> ActiveStageReadTargets;
 
-	void RegisterFunctionCall(ENiagaraScriptUsage ScriptUsage, const FString& InName, const FString& InFullName, const FGuid& CallNodeId, const FString& InFunctionNameSuffix, UNiagaraScriptSource* Source, FNiagaraFunctionSignature& InSignature, bool bIsCustomHlsl, const FString& InCustomHlsl, const TArray<FNiagaraCustomHlslInclude>& InCustomHlslIncludeFilePaths, TArray<int32>& Inputs, TArrayView<UEdGraphPin* const> CallInputs, TArrayView<UEdGraphPin* const> CallOutputs,
-		FNiagaraFunctionSignature& OutSignature);
 	void GenerateFunctionCall(ENiagaraScriptUsage ScriptUsage, FNiagaraFunctionSignature& FunctionSignature, TArrayView<const int32> Inputs, TArray<int32>& Outputs);
 	FString GetFunctionIncludeStatement(const FNiagaraCustomHlslInclude& Include) const;
 	FString GetFunctionSignature(const FNiagaraFunctionSignature& Sig);
 
-	/** Compiles an output Pin on a graph node. Caches the result for any future inputs connected to it. */
-	int32 CompileOutputPin(const UEdGraphPin* Pin);
-
-	void WriteDataSetContextVars(TMap<FNiagaraDataSetID, TMap<int32, FDataSetAccessInfo>>& DataSetAccessInfo, bool bRead, FString &OutHLSLOutput);
-	void WriteDataSetStructDeclarations(TMap<FNiagaraDataSetID, TMap<int32, FDataSetAccessInfo>>& DataSetAccessInfo, bool bRead, FString &OutHLSLOutput);
+	void WriteDataSetContextVars(TMap<FNiagaraDataSetID, TMap<int32, FDataSetAccessInfo>>& DataSetAccessInfo, bool bRead, FString& OutHLSLOutput);
+	void WriteDataSetStructDeclarations(TMap<FNiagaraDataSetID, TMap<int32, FDataSetAccessInfo>>& DataSetAccessInfo, bool bRead, FString& OutHLSLOutput);
 	void DecomposeVariableAccess(UStruct* Struct, bool bRead, FString IndexSymbol, FString HLSLString);
 
 	FString GetUniqueSymbolName(FName BaseName);
@@ -522,8 +303,6 @@ protected:
 	};
 	TArray<FFunctionContext> FunctionContextStack;
 	const FFunctionContext* FunctionCtx()const { return FunctionContextStack.Num() > 0 ? &FunctionContextStack.Last() : nullptr; }
-	void EnterFunction(const FString& Name, FNiagaraFunctionSignature& Signature, TArrayView<const int32> Inputs, const FGuid& InGuid);
-	void ExitFunction();
 	FString GetCallstack();
 	TArray<FGuid> GetCallstackGuids();
 
@@ -548,111 +327,59 @@ protected:
 	int32 AddBodyComment(const FString& Comment);
 	int32 AddBodyChunk(const FString& Definition);
 
-
-	FString GetFunctionDefinitions();
-
 	FString GetUniqueEmitterName() const;
-
-	void HandleDataInterfaceCall(FNiagaraScriptDataInterfaceCompileInfo& Info, const FNiagaraFunctionSignature& InMatchingSignature);
 	void ConvertCompileInfoToParamInfo(const FNiagaraScriptDataInterfaceCompileInfo& InCompileInfo, FNiagaraDataInterfaceGPUParamInfo& OutGPUParamInfo, TArray<FNiagaraFunctionSignature>& GeneratedFunctionSignatures);
-public:
-
-	FHlslNiagaraTranslator();
-	virtual ~FHlslNiagaraTranslator() {}
 
-	static void Init();
-	
-	virtual const FNiagaraTranslateResults &Translate(const FNiagaraCompileRequestData* InCompileData, const FNiagaraCompileRequestDuplicateData* InCompileDuplicateData, const FNiagaraCompileOptions& InCompileOptions, FHlslNiagaraTranslatorOptions Options);
-	FNiagaraTranslatorOutput &GetTranslateOutput() { return CompilationOutput; }
-
-	virtual int32 CompilePin(const UEdGraphPin* Pin);
-
-	virtual int32 RegisterUObject(FNiagaraVariable Variable, UObject* Object, bool bAddParameterMapRead);
-	virtual int32 RegisterDataInterface(FNiagaraVariable& Var, UNiagaraDataInterface* DataInterface, bool bPlaceholder, bool bAddParameterMapRead);
+	FString GetFunctionDefinitions();
 
-	virtual void Operation(class UNiagaraNodeOp* Operation, TArray<int32>& Inputs, TArray<int32>& Outputs);
-	virtual void Output(UNiagaraNodeOutput* OutputNode, const TArray<int32>& ComputedInputs);
+public:
+	FNiagaraHlslTranslator();
+	virtual ~FNiagaraHlslTranslator() = default;
 
-	virtual int32 GetParameter(const FNiagaraVariable& Parameter);
-	virtual int32 GetRapidIterationParameter(const FNiagaraVariable& Parameter);
+	virtual const FNiagaraTranslatorOutput& GetTranslateOutput() const override { return CompilationOutput; }
+	virtual const FString& GetTranslatedHLSL() const override { return HlslOutput; }
 
 	bool DisableDebugSwitches() const { return TranslationOptions.bDisableDebugSwitches; }
 
 	bool IsCompileOptionDefined(const TCHAR* InDefineStr);
 
-	virtual int32 GetAttribute(const FNiagaraVariable& Attribute);
+	int32 GetAttribute(const FNiagaraVariable& Attribute);
+
+	int32 GetConstant(const FNiagaraVariable& Constant, FString* DebugOutputValue = nullptr);
+
+	int32 ParameterMapForInnerIndex() const;
 
-	virtual int32 GetConstant(const FNiagaraVariable& Constant, FString* DebugOutputValue = nullptr);
-	
-	virtual void ReadDataSet(const FNiagaraDataSetID DataSet, const TArray<FNiagaraVariable>& Variable, ENiagaraDataSetAccessMode AccessMode, int32 InputChunk, TArray<int32>& Outputs);
-	virtual void WriteDataSet(const FNiagaraDataSetID DataSet, const TArray<FNiagaraVariable>& Variable, ENiagaraDataSetAccessMode AccessMode, const TArray<int32>& Inputs, TArray<int32>& Outputs);
-	virtual void ParameterMapSet(class UNiagaraNodeParameterMapSet* SetNode, TArrayView<const FCompiledPin> Inputs, TArray<int32>& Outputs);
-	virtual void ParameterMapGet(class UNiagaraNodeParameterMapGet* GetNode, TArrayView<const int32> Inputs, TArray<int32>& Outputs);
-	virtual void Emitter(class UNiagaraNodeEmitter* GetNode, TArray<int32>& Inputs, TArray<int32>& Outputs);
-
-	virtual void ParameterMapForBegin(class UNiagaraNodeParameterMapFor* ForNode, int32 IterationCount);
-	virtual void ParameterMapForContinue(class UNiagaraNodeParameterMapFor* ForNode, int32 IterationEnabled);
-	virtual void ParameterMapForEnd(class UNiagaraNodeParameterMapFor* ForNode);
-	virtual int32 ParameterMapForInnerIndex() const;
-
-	void DefineInterpolatedParametersFunction(FString &HlslOutput);
 	void DefinePreviousParametersFunction(FString& HlslOutput, TArray<TArray<FNiagaraVariable>>& DataSetVariables, TMap<FNiagaraDataSetID, int32>& DataSetReads, TMap<FNiagaraDataSetID, int32>& DataSetWrites);
 
-	void DefineDataSetReadFunction(FString &HlslOutput, TArray<FNiagaraDataSetID> &ReadDataSets);
-	void DefineDataSetWriteFunction(FString &HlslOutput, TArray<FNiagaraDataSetProperties> &WriteDataSets, TArray<int32>& WriteConditionVarIndices);
-	void DefineMain(FString &HLSLOutput, const TArray<TArray<FNiagaraVariable>>& DataSetVariables, const TMap<FNiagaraDataSetID, int32>& DataSetReads, const TMap<FNiagaraDataSetID, int32>& DataSetWrites);
-	void DefineMainGPUFunctions(const TArray<TArray<FNiagaraVariable>>& DataSetVariables, const TMap<FNiagaraDataSetID, int32>& DataSetReads, const TMap<FNiagaraDataSetID, int32>& DataSetWrites);
+	void DefineDataSetReadFunction(FString& HlslOutput, TArray<FNiagaraDataSetID>& ReadDataSets);
+	void DefineDataSetWriteFunction(FString& HlslOutput, TArray<FNiagaraDataSetProperties>& WriteDataSets, TArray<int32>& WriteConditionVarIndices);
+	void DefineMain(FString& HLSLOutput, const TArray<TArray<FNiagaraVariable>>& DataSetVariables, const TMap<FNiagaraDataSetID, int32>& DataSetReads, const TMap<FNiagaraDataSetID, int32>& DataSetWrites);
 
-	void DefineDataSetVariableReads(FString &HLSLOutput, const FNiagaraDataSetID& Id, int32 DataSetIndex, const TArray<FNiagaraVariable>& ReadVars);
-	void DefineDataSetVariableWrites(FString &HlslOutput, const FNiagaraDataSetID& Id, int32 DataSetIndex, const TArray<FNiagaraVariable>& WriteVars);
-	void DefineDataInterfaceHLSL(FString &HlslOutput);
-	void DefineExternalFunctionsHLSL(FString &HlslOutput);
+	void DefineDataSetVariableReads(FString& HLSLOutput, const FNiagaraDataSetID& Id, int32 DataSetIndex, const TArray<FNiagaraVariable>& ReadVars);
+	void DefineDataInterfaceHLSL(FString& HlslOutput);
+	void DefineExternalFunctionsHLSL(FString& HlslOutput);
 
 	// Format string should have up to 5 entries, {{0} = Computed Variable Suffix, {1} = Float or Int, {2} = Data Set Index, {3} = Register Index, {4} Default value for that type.
-	void GatherVariableForDataSetAccess(const FNiagaraVariable& Variable, FString Format, int32& RegisterIdxInt, int32& RegisterIdxFloat, int32& RegisterIdxHalf, int32 DataSetIndex, FString InstanceIdxSymbol, FString &HlslOutput, bool bWriteHLSL = true);
+	void GatherVariableForDataSetAccess(const FNiagaraVariable& Variable, FString Format, int32& RegisterIdxInt, int32& RegisterIdxFloat, int32& RegisterIdxHalf, int32 DataSetIndex, FString InstanceIdxSymbol, FString& HlslOutput, bool bWriteHLSL = true);
 	void GatherComponentsForDataSetAccess(UScriptStruct* Struct, FString VariableSymbol, bool bMatrixRoot, TArray<FString>& Components, TArray<ENiagaraBaseTypes>& Types);
 
-	virtual void FunctionCall(UNiagaraNodeFunctionCall* FunctionNode, TArray<int32>& Inputs, TArray<int32>& Outputs);
 	void EnterFunctionCallNode(const TSet<FName>& UnusedInputs);
 	void ExitFunctionCallNode();
 	bool IsFunctionVariableCulledFromCompilation(const FName& InputName) const;
-	void CullMapSetInputPin(UEdGraphPin* InputPin);
 
-	virtual void Convert(class UNiagaraNodeConvert* Convert, TArrayView<const int32> Inputs, TArray<int32>& Outputs);
-	virtual void If(class UNiagaraNodeIf* IfNode, TArray<FNiagaraVariable>& Vars, int32 Condition, TArray<int32>& PathA, TArray<int32>& PathB, TArray<int32>& Outputs);
-	/** Options is a map from selector values to compiled pin code chunk indices */
-	virtual void Select(class UNiagaraNodeSelect* SelectNode, int32 Selector, const TArray<FNiagaraVariable>& OutputVariables, TMap<int32, TArray<int32>>& Options, TArray<int32>& Outputs);
-	
-	void WriteCompilerTag(int32 InputCompileResult, const UEdGraphPin* Pin, bool bEmitMessageOnFailure, FNiagaraCompileEventSeverity FailureSeverity, const FString& Prefix = FString());
-
-	void Message(FNiagaraCompileEventSeverity Severity, FText MessageText, const UNiagaraNode* Node, const UEdGraphPin* Pin, FString ShortDescription = FString());
-	virtual void Error(FText ErrorText, const UNiagaraNode* Node, const UEdGraphPin* Pin, FString ShortDescription = FString());
-	virtual void Warning(FText WarningText, const UNiagaraNode* Node, const UEdGraphPin* Pin, FString ShortDescription = FString());
-	void RegisterCompileDependency(const FNiagaraVariableBase& InVar, FText ErrorText, const UNiagaraNode* Node, const UEdGraphPin* Pin, bool bEmitAsLinker, int32 ParamMapHistoryIdx);
-	FString NodePinToMessage(FText MessageText, const UNiagaraNode* Node, const UEdGraphPin* Pin);
-
-	virtual bool GetFunctionParameter(const FNiagaraVariable& Parameter, int32& OutParam)const;
+	bool GetFunctionParameter(const FNiagaraVariable& Parameter, int32& OutParam)const;
 	int32 GetUniqueCallerID();
 
-	virtual bool CanReadAttributes()const;
-	virtual ENiagaraScriptUsage GetTargetUsage() const;
+	bool CanReadAttributes()const;
+	ENiagaraScriptUsage GetTargetUsage() const;
 	FGuid GetTargetUsageId() const;
-	virtual ENiagaraScriptUsage GetCurrentUsage() const;
-	virtual ENiagaraSimTarget GetSimulationTarget() const
+	ENiagaraSimTarget GetSimulationTarget() const
 	{
 		return CompilationTarget;
 	}
 
-	static bool IsBuiltInHlslType(const FNiagaraTypeDefinition& Type);
-	static FString GetStructHlslTypeName(const FNiagaraTypeDefinition& Type);
-	static FString GetPropertyHlslTypeName(const FProperty* Property);
-	static FString BuildHLSLStructDecl(const FNiagaraTypeDefinition& Type, FText& OutErrorMessage, bool bGpuScript);
-	static FString GetHlslDefaultForType(const FNiagaraTypeDefinition& Type);
-	static bool IsHlslBuiltinVector(const FNiagaraTypeDefinition& Type);
 	static TArray<FName> ConditionPropertyPath(const FNiagaraTypeDefinition& Type, const TArray<FName>& InPath);
-
-
-	static FString GetSanitizedSymbolName(FStringView SymbolName, bool bCollapseNamespaces=false);
+	static FString GetSanitizedSymbolName(FStringView SymbolName, bool bCollapseNamespaces = false);
 	static FString GetSanitizedDIFunctionName(const FString& FunctionName);
 	static FString GetSanitizedFunctionNameSuffix(FString Name);
 
@@ -661,58 +388,40 @@ public:
 
 	bool AddStructToDefinitionSet(const FNiagaraTypeDefinition& TypeDef);
 
-	FString &GetTranslatedHLSL()
-	{
-		return HlslOutput;
-	}
-
 	static FString GetFunctionSignatureSymbol(const FNiagaraFunctionSignature& Sig);
 
 	static FName GetDataInterfaceName(FName BaseName, const FString& UniqueEmitterName, bool bIsParameterMapDataInterface);
 
-	/** If OutVar can be replaced by a literal constant, it's data is initialized with the correct value and we return true. Returns false otherwise. */
-	bool GetLiteralConstantVariable(FNiagaraVariable& OutVar) const;
 
-	/** If Var can be replaced by a another constant variable, or is a constant itself, add the appropriate body chunk and return true. */
-	bool HandleBoundConstantVariableToDataSetRead(FNiagaraVariable InVariable, UNiagaraNode* InNode, int32 InParamMapHistoryIdx, int32& Output, const UEdGraphPin* InDefaultPin);
 
 	static FString GenerateFunctionHlslPrototype(FStringView InVariableName, const FNiagaraFunctionSignature& FunctionSignature);
-	void FillVariableWithDefaultValue(FNiagaraVariable& InVar, const UEdGraphPin* InDefaultPin);
-	void FillVariableWithDefaultValue(int32& OutValue, const UEdGraphPin* InDefaultPin);
 
-	void SetConstantByStaticVariable(int32& OutValue, const UEdGraphPin* InDefaultPin, FString* DebugString = nullptr);
 	void SetConstantByStaticVariable(int32& OutValue, const FNiagaraVariable& Var, FString* DebugString = nullptr);
 	void SetConstantByStaticVariable(FNiagaraVariable& OutValue, const FNiagaraVariable& Var, FString* DebugString = nullptr);
-	void SetConstantByStaticVariable(FNiagaraVariable& OutValue, const UEdGraphPin* InDefaultPin, FString* DebugString = nullptr);
-
-	int32 MakeStaticVariableDirect(const UEdGraphPin* InDefaultPin);
 
 	bool IsEventSpawnScript()const;
 
-private:
-	void InitializeParameterMapDefaults(int32 ParamMapHistoryIdx);
-	void HandleParameterRead(int32 ParamMapHistoryIdx, const FNiagaraVariable& Var, const UEdGraphPin* DefaultPin, UNiagaraNode* ErrorNode, int32& OutputChunkId, TOptional<ENiagaraDefaultMode> DefaultMode = TOptional<ENiagaraDefaultMode>(), TOptional<FNiagaraScriptVariableBinding> DefaultBinding = TOptional<FNiagaraScriptVariableBinding>(), bool bTreatAsUnknownParameterMap = false, bool bIgnoreDefaultSetFirst = false);
+	// collection of functions that rely on the graph implementation and so must be relegated to the TNiagaraHlslTranslator
+	virtual void DefineDataSetVariableWrites(FString& HlslOutput, const FNiagaraDataSetID& Id, int32 DataSetIndex, const TArray<FNiagaraVariable>& WriteVars) = 0;
+	virtual void Message(FNiagaraCompileEventSeverity Severity, FText MessageText, FStringView ShortDescription = FStringView()) = 0;
+	virtual void Error(FText ErrorText, FStringView ShortDescription = FStringView()) = 0;
+	virtual void Warning(FText WarningText, FStringView ShortDescription = FStringView()) = 0;
+
+	/** If OutVar can be replaced by a literal constant, it's data is initialized with the correct value and we return true. Returns false otherwise. */
+	virtual bool GetLiteralConstantVariable(FNiagaraVariable& OutVar) const = 0;
+
+
+protected:
 	bool ShouldConsiderTargetParameterMap(ENiagaraScriptUsage InUsage) const;
-	FString BuildParameterMapHlslDefinitions(TArray<FNiagaraVariable>& PrimaryDataSetOutputEntries);
-	void BuildMissingDefaults();
-	void FinalResolveNamespacedTokens(const FString& ParameterMapInstanceNamespace, TArray<FString>& Tokens, TArray<FString>& ValidChildNamespaces, FNiagaraParameterMapHistoryBuilder& Builder, TArray<FNiagaraVariable>& UniqueParameterMapEntriesAliasesIntact, TArray<FNiagaraVariable>& UniqueParameterMapEntries, int32 ParamMapHistoryIdx, UNiagaraNode* InNodeForErrorReporting);
 
 	void HandleNamespacedExternalVariablesToDataSetRead(TArray<FNiagaraVariable>& InDataSetVars, FString InNamespaceStr);
 
-	void FindConstantValue(int32 InputCompileResult, const FNiagaraTypeDefinition& TypeDef, FString& Value, FNiagaraVariable& Variable);
-
 	// For GPU simulations we have to special case some variables and pass them view shader parameters rather than the uniform buffer as they vary from CPU simulations
 	bool IsVariableInUniformBuffer(const FNiagaraVariable& Variable) const;
 
 	FString ComputeMatrixColumnAccess(const FString& Name);
 	FString ComputeMatrixRowAccess(const FString& Name);
 
-	bool ParseDIFunctionSpecifiers(UNiagaraNode* NodeForErrorReporting, FNiagaraFunctionSignature& Sig, TArray<FString>& Tokens, int32& TokenIdx);
-	void HandleCustomHlslNode(UNiagaraNodeCustomHlsl* CustomFunctionHlsl, ENiagaraScriptUsage& OutScriptUsage, FString& OutName, FString& OutFullName, bool& bOutCustomHlsl, FString& OutCustomHlsl, TArray<FNiagaraCustomHlslInclude>& OutCustomHlslIncludeFilePaths,
-		FNiagaraFunctionSignature& OutSignature, TArray<int32>& Inputs);
-	void ProcessCustomHlsl(const FString& InCustomHlsl, ENiagaraScriptUsage InUsage, const FNiagaraFunctionSignature& InSignature, const TArray<int32>& Inputs, UNiagaraNode* NodeForErrorReporting, FString& OutCustomHlsl, FNiagaraFunctionSignature& OutSignature);
-	void HandleSimStageSetupAndTeardown(int32 InWhichStage, FString& OutHlsl);
-	
 	// Add a raw float constant chunk
 	int32 GetConstantDirect(float InValue);
 	int32 GetConstantDirect(bool InValue);
@@ -736,27 +445,12 @@ private:
 
 	// Takes the current script state (interpolated or not) and determines the correct context variable.
 	FString GetParameterMapInstanceName(int32 ParamMapHistoryIdx);
-	
-	// Register a System/Engine/read-only variable in its namespaced form
-	bool ParameterMapRegisterExternalConstantNamespaceVariable(FNiagaraVariable InVariable, UNiagaraNode* InNode, int32 InParamMapHistoryIdx, int32& Output, const UEdGraphPin* InDefaultPin);
 
 	// Register an attribute in its non namespaced form
-	bool ParameterMapRegisterNamespaceAttributeVariable(const FNiagaraVariable& InVariable, UNiagaraNode* InNode, int32 InParamMapHistoryIdx, int32& Output);
-	// Register an attribute in its namespaced form
-	bool ParameterMapRegisterUniformAttributeVariable(const FNiagaraVariable& InVariable, UNiagaraNode* InNode, int32 InParamMapHistoryIdx, int32& Output);
-
-    // Checks that the Partices.ID parameter is only used if persistent IDs are active
-    void ValidateParticleIDUsage();
-	bool ValidateTypePins(const UNiagaraNode* NodeToValidate);
-	void GenerateFunctionSignature(ENiagaraScriptUsage ScriptUsage, FString InName, const FString& InFullName, const FString& InFunctionNameSuffix, UNiagaraGraph* FuncGraph, TArray<int32>& Inputs, 
-		bool bHadNumericInputs, bool bHasParameterMapParameters, TArray<UEdGraphPin*> StaticSwitchValues, FNiagaraFunctionSignature& OutSig) const;
-
-	void ValidateFailIfPreviouslyNotSet(const FNiagaraVariable& InVar, bool& bFailIfNotSet);
+	virtual bool ParameterMapRegisterNamespaceAttributeVariable(const FNiagaraVariable& InVariable, int32 InParamMapHistoryIdx, int32& Output) = 0;
 
 	bool ShouldInterpolateParameter(const FNiagaraVariable& Parameter);
 
-	void UpdateStaticSwitchConstants(UEdGraphNode* Node);
-
 	bool IsBulkSystemScript() const;
 	bool IsSpawnScript() const;
 	bool RequiresInterpolation() const;
@@ -764,33 +458,22 @@ private:
 	template<typename T>
 	void BuildConstantBuffer(ENiagaraCodeChunkMode ChunkMode);
 
-	void TrimAttributes(const FNiagaraCompileOptions& InCompileOptions, TArray<FNiagaraVariable>& Attributes);
+	virtual const FString& GetEmitterUniqueName() const = 0;
+	virtual TConstArrayView<FNiagaraVariable> GetStaticVariables() const = 0;
+	virtual UNiagaraDataInterface* GetDataInterfaceCDO(UClass* DIClass) const = 0;
 
-	// Specific method to reconcile whether the default value has been implicitly or explicitly set for a namespaced var added to the param map.
-	void RecordParamMapDefinedAttributeToNamespaceVar(const FNiagaraVariable& VarToRecord, const UEdGraphPin* VarAssociatedDefaultPin);
-	
 	/** Map of symbol names to count of times it's been used. Used for generating unique symbol names. */
 	TMap<FName, uint32> SymbolCounts;
 
 	//Set of non-builtin structs we have to define in hlsl.
 	TArray<FNiagaraTypeDefinition> StructsToDefine;
 
-	// Keep track of all the paths that the parameter maps can take through the graph.
-	TArray<FNiagaraParameterMapHistory> ParamMapHistories;
-
 	// Keep track of which parameter map history this came from.
 	TArray<int32> ParamMapHistoriesSourceInOtherHistories;
 
-	// Keep track of the other output nodes in the graph's histories so that we can make sure to 
-	// create any variables that are needed downstream.
-	TArray<FNiagaraParameterMapHistory> OtherOutputParamMapHistories;
-
 	// All of the variables arrays in the other histories converted to sanitized HLSL format. Used in parsing custom hlsl nodes.
 	TArray< TArray<FNiagaraVariable> > OtherOutputParamMapHistoriesSanitizedVariables;
 
-	// Make sure that the function call names match up on the second traversal.
-	FNiagaraParameterMapHistoryBuilder ActiveHistoryForFunctionCalls;
-
 	// Synced to the ParamMapHistories.
 	TArray<TArray<int32>> ParamMapSetVariablesToChunks;
 
@@ -829,9 +512,6 @@ private:
 
 	// List of variables for interpolated spawning
 	TArray<FNiagaraVariable> InterpSpawnVariables;
-
-	// Map of primary ouput variable description to its default value pin
-	TMap<FNiagaraVariable, const UEdGraphPin*> UniqueVarToDefaultPin;
 	
 	// Map of primary output variable description to whether or not it came from this script's parameter map
 	TMap<FNiagaraVariable, bool> UniqueVarToWriteToParamMap;
@@ -864,14 +544,206 @@ private:
 
 	ENiagaraCodeChunkMode CurrentBodyChunkMode;
 
-	TArray<FHlslNiagaraTranslationStage> TranslationStages;
+	TArray<FNiagaraHlslTranslationStage> TranslationStages;
 	int32 ActiveStageIdx;
 	bool bInitializedDefaults;
 	bool bEnforceStrictTypesValidations;
 
-	TArray<const UEdGraphPin*> CurrentDefaultPinTraversal;
 	// Variables that need to be initialized based on some other variable's value at the end of spawn.
 	TArray<FNiagaraVariable> InitialNamespaceVariablesMissingDefault;
 	// Variables that need to be initialized in the body or at the end of spawn.
 	TArray<FNiagaraVariable> DeferredVariablesMissingDefault;
 };
+
+// all data and function implementations that rely on specifics of the graph implementation
+// are shuffled into this class to allow us to share the code through the inherited GraphBridge
+template<typename GraphBridge>
+class TNiagaraHlslTranslator : public FNiagaraHlslTranslator, private GraphBridge
+{
+	using typename GraphBridge::FGraph;
+	using typename GraphBridge::FPin;
+	using typename GraphBridge::FScopedPin;
+	using typename GraphBridge::FNode;
+	using typename GraphBridge::FParamMapHistory;
+	using typename GraphBridge::FParamMapHistoryBuilder;
+	using typename GraphBridge::FPrecompileData;
+	using typename GraphBridge::FCompilationCopy;
+	using typename GraphBridge::FFunctionCallNode;
+	using typename GraphBridge::FOpNode;
+	using typename GraphBridge::FOutputNode;
+	using typename GraphBridge::FIfNode;
+	using typename GraphBridge::FInputNode;
+	using typename GraphBridge::FSelectNode;
+	using typename GraphBridge::FCustomHlslNode;
+	using typename GraphBridge::FConvertNode;
+	using typename GraphBridge::FParamMapForNode;
+	using typename GraphBridge::FParamMapSetNode;
+	using typename GraphBridge::FParamMapGetNode;
+	using typename GraphBridge::FEmitterNode;
+	using typename GraphBridge::FStaticSwitchNode;
+	using typename GraphBridge::FInputPin;
+	using typename GraphBridge::FOutputPin;
+	using typename GraphBridge::FModuleScopedPin;
+	using typename GraphBridge::FGraphTraversalHandle;
+	using typename GraphBridge::FGraphFunctionAliasContext;
+	using typename GraphBridge::FConvertConnection;
+
+public:
+	using FBridge = GraphBridge;
+
+	TNiagaraHlslTranslator() = delete;
+	TNiagaraHlslTranslator(const FPrecompileData* InPrecompileData, const FCompilationCopy* InCompilationCopy)
+	: CompileData(InPrecompileData)
+	, CompileDuplicateData(InCompilationCopy)
+	{}
+
+	virtual FNiagaraTranslateResults Translate(const FNiagaraCompileOptions& InCompileOptions, const FHlslNiagaraTranslatorOptions& Options) override;
+
+	void RegisterFunctionCall(ENiagaraScriptUsage ScriptUsage, const FString& InName, const FString& InFullName, const FGuid& CallNodeId, const FString& InFunctionNameSuffix, const FGraph* SourceGraph, FNiagaraFunctionSignature& InSignature, bool bIsCustomHlsl, const FString& InCustomHlsl, const TArray<FNiagaraCustomHlslInclude>& InCustomHlslIncludeFilePaths, TArray<int32>& Inputs, TConstArrayView<const FInputPin*> CallInputs, TConstArrayView<const FOutputPin*> CallOutputs,
+		FNiagaraFunctionSignature& OutSignature);
+	/** Compiles an output Pin on a graph node. Caches the result for any future inputs connected to it. */
+	int32 CompileOutputPin(const FPin* Pin);
+
+	void EnterFunction(const FString& Name, FNiagaraFunctionSignature& Signature, TArrayView<const int32> Inputs, const FGuid& InGuid);
+	void ExitFunction();
+	void HandleDataInterfaceCall(FNiagaraScriptDataInterfaceCompileInfo& Info, const FNiagaraFunctionSignature& InMatchingSignature);
+
+	struct FCompiledPin
+	{
+		int32 CompilationIndex;
+		const FPin* Pin;
+
+		FCompiledPin(int32 CompilationIndex, const FPin* InPin)
+		: CompilationIndex(CompilationIndex)
+		, Pin(InPin)
+		{}
+	};
+
+public:
+
+	int32 CompileInputPin(const FInputPin* Pin);
+
+	int32 RegisterUObject(FNiagaraVariable Variable, UObject* Object, bool bAddParameterMapRead);
+	int32 RegisterDataInterface(const FNiagaraVariable& Var, const UNiagaraDataInterface* DataInterface, bool bPlaceholder, bool bAddParameterMapRead);
+
+	void Operation(const FOpNode* Operation, TArray<int32>& Inputs, TArray<int32>& Outputs);
+	void Output(const FOutputNode* OutputNode, const TArray<int32>& ComputedInputs);
+
+	int32 GetParameter(const FNiagaraVariable& Parameter);
+	int32 GetRapidIterationParameter(const FNiagaraVariable& Parameter);
+	void ReadDataSet(const FNiagaraDataSetID DataSet, const TArray<FNiagaraVariable>& Variable, ENiagaraDataSetAccessMode AccessMode, int32 InputChunk, TArray<int32>& Outputs);
+	void WriteDataSet(const FNiagaraDataSetID DataSet, const TArray<FNiagaraVariable>& Variable, ENiagaraDataSetAccessMode AccessMode, const TArray<int32>& Inputs, TArray<int32>& Outputs);
+	void ParameterMapSet(const FParamMapSetNode* SetNode, TArrayView<const FCompiledPin> Inputs, TArray<int32>& Outputs);
+	void ParameterMapGet(const FParamMapGetNode* GetNode, TArrayView<const int32> Inputs, TArray<int32>& Outputs);
+	void ParameterMapForBegin(const FParamMapForNode* ForNode, int32 IterationCount);
+	void ParameterMapForContinue(const FParamMapForNode* ForNode, int32 IterationEnabled);
+	void ParameterMapForEnd(const FParamMapForNode* ForNode);
+	void Emitter(const FEmitterNode* GetNode, TArray<int32>& Inputs, TArray<int32>& Outputs);
+	void FunctionCall(const FFunctionCallNode* FunctionNode, TArray<int32>& Inputs, TArray<int32>& Outputs);
+	void Convert(const FConvertNode* Convert, TArrayView<const int32> Inputs, TArray<int32>& Outputs);
+	void If(const FIfNode* IfNode, const TArray<FNiagaraVariable>& Vars, int32 Condition, TArray<int32>& PathA, TArray<int32>& PathB, TArray<int32>& Outputs);
+	/** Options is a map from selector values to compiled pin code chunk indices */
+	void Select(const FSelectNode* SelectNode, int32 Selector, const TArray<FNiagaraVariable>& OutputVariables, TMap<int32, TArray<int32>>& Options, TArray<int32>& Outputs);
+
+	void DefineInterpolatedParametersFunction(FString& HlslOutput);
+	void DefineMainGPUFunctions(const TArray<TArray<FNiagaraVariable>>& DataSetVariables, const TMap<FNiagaraDataSetID, int32>& DataSetReads, const TMap<FNiagaraDataSetID, int32>& DataSetWrites);
+	virtual void DefineDataSetVariableWrites(FString& HlslOutput, const FNiagaraDataSetID& Id, int32 DataSetIndex, const TArray<FNiagaraVariable>& WriteVars) override;
+
+	void CullMapSetInputPin(const FPin* InputPin);
+
+	void WriteCompilerTag(int32 InputCompileResult, const FPin* Pin, bool bEmitMessageOnFailure, FNiagaraCompileEventSeverity FailureSeverity, const FString& Prefix = FString());
+
+	void Message(FNiagaraCompileEventSeverity Severity, FText MessageText, const FNode* Node, const FPin* Pin, FStringView ShortDescription = FStringView());
+	void Error(FText ErrorText, const FNode* Node, const FPin* Pin, FStringView ShortDescription = FStringView());
+	void Warning(FText WarningText, const FNode* Node, const FPin* Pin, FStringView ShortDescription = FStringView());
+	void RegisterCompileDependency(const FNiagaraVariableBase& InVar, FText ErrorText, const FNode* Node, const FPin* Pin, bool bEmitAsLinker, int32 ParamMapHistoryIdx);
+	ENiagaraScriptUsage GetCurrentUsage() const;
+
+	virtual void Message(FNiagaraCompileEventSeverity Severity, FText MessageText, FStringView ShortDescription = FStringView()) override;
+	virtual void Error(FText ErrorText, FStringView ShortDescription = FStringView()) override;
+	virtual void Warning(FText WarningText, FStringView ShortDescription = FStringView()) override;
+	FString NodePinToMessage(FText MessageText, const FNode* Node, const FPin* Pin);
+
+	/** If OutVar can be replaced by a literal constant, it's data is initialized with the correct value and we return true. Returns false otherwise. */
+	virtual bool GetLiteralConstantVariable(FNiagaraVariable& OutVar) const override;
+
+	/** If Var can be replaced by a another constant variable, or is a constant itself, add the appropriate body chunk and return true. */
+	bool HandleBoundConstantVariableToDataSetRead(FNiagaraVariable InVariable, const FNode* InNode, int32 InParamMapHistoryIdx, int32& Output, const FPin* InDefaultPin);
+
+	void FillVariableWithDefaultValue(FNiagaraVariable& InVar, const FPin* InDefaultPin);
+	void FillVariableWithDefaultValue(int32& OutValue, const FPin* InDefaultPin);
+
+	void SetConstantByStaticVariable(int32& OutValue, const FPin* InDefaultPin, FString* DebugString = nullptr);
+	void SetConstantByStaticVariable(FNiagaraVariable& OutValue, const FPin* InDefaultPin, FString* DebugString = nullptr);
+
+	int32 MakeStaticVariableDirect(const FPin* InDefaultPin);
+
+	void TrimAttributes(const FNiagaraCompileOptions& InCompileOptions, TArray<FNiagaraVariable>& Attributes);
+
+protected:
+	void InitializeParameterMapDefaults(int32 ParamMapHistoryIdx);
+	void HandleParameterRead(int32 ParamMapHistoryIdx, const FNiagaraVariable& Var, const FInputPin* DefaultPin, const FNode* ErrorNode, int32& OutputChunkId, TOptional<ENiagaraDefaultMode> DefaultMode = TOptional<ENiagaraDefaultMode>(), TOptional<FNiagaraScriptVariableBinding> DefaultBinding = TOptional<FNiagaraScriptVariableBinding>(), bool bTreatAsUnknownParameterMap = false, bool bIgnoreDefaultSetFirst = false);
+	FString BuildParameterMapHlslDefinitions(TArray<FNiagaraVariable>& PrimaryDataSetOutputEntries);
+	void BuildMissingDefaults();
+	void FinalResolveNamespacedTokens(const FString& ParameterMapInstanceNamespace, TArray<FString>& Tokens, TArray<FString>& ValidChildNamespaces, FParamMapHistoryBuilder& Builder, TArray<FNiagaraVariable>& UniqueParameterMapEntriesAliasesIntact, TArray<FNiagaraVariable>& UniqueParameterMapEntries, int32 ParamMapHistoryIdx, const FNode* InNodeForErrorReporting);
+	void FindConstantValue(int32 InputCompileResult, const FNiagaraTypeDefinition& TypeDef, FString& Value, FNiagaraVariable& Variable);
+
+	bool ParseDIFunctionSpecifiers(const FNode* NodeForErrorReporting, FNiagaraFunctionSignature& Sig, TArray<FString>& Tokens, int32& TokenIdx);
+	void HandleCustomHlslNode(const FCustomHlslNode* CustomFunctionHlsl, ENiagaraScriptUsage& OutScriptUsage, FString& OutName, FString& OutFullName, bool& bOutCustomHlsl, FString& OutCustomHlsl, TArray<FNiagaraCustomHlslInclude>& OutCustomHlslIncludeFilePaths,
+		FNiagaraFunctionSignature& OutSignature, TArray<int32>& Inputs);
+	void ProcessCustomHlsl(const FString& InCustomHlsl, ENiagaraScriptUsage InUsage, const FNiagaraFunctionSignature& InSignature, const TArray<int32>& Inputs, const FNode* NodeForErrorReporting, FString& OutCustomHlsl, FNiagaraFunctionSignature& OutSignature);
+	void HandleSimStageSetupAndTeardown(int32 InWhichStage, FString& OutHlsl);
+	// Register a System/Engine/read-only variable in its namespaced form
+	bool ParameterMapRegisterExternalConstantNamespaceVariable(FNiagaraVariable InVariable, const FNode* InNode, int32 InParamMapHistoryIdx, int32& Output, const FPin* InDefaultPin);
+
+	// Register an attribute in its non namespaced form
+	bool ParameterMapRegisterNamespaceAttributeVariable(const FNiagaraVariable& InVariable, const FNode* InNode, int32 InParamMapHistoryIdx, int32& Output);
+	virtual bool ParameterMapRegisterNamespaceAttributeVariable(const FNiagaraVariable& InVariable, int32 InParamMapHistoryIdx, int32& Output) override;
+
+	// Register an attribute in its namespaced form
+	bool ParameterMapRegisterUniformAttributeVariable(const FNiagaraVariable& InVariable, FNode* InNode, int32 InParamMapHistoryIdx, int32& Output);
+
+	// Checks that the Partices.ID parameter is only used if persistent IDs are active
+	void ValidateParticleIDUsage();
+	bool ValidateTypePins(const FNode* NodeToValidate);
+	void GenerateFunctionSignature(ENiagaraScriptUsage ScriptUsage, FString InName, const FString& InFullName, const FString& InFunctionNameSuffix, const FGraph* FuncGraph, TArray<int32>& Inputs,
+		bool bHadNumericInputs, bool bHasParameterMapParameters, const TArray<const FInputPin*>& StaticSwitchValues, FNiagaraFunctionSignature& OutSig) const;
+
+	void ValidateFailIfPreviouslyNotSet(const FNiagaraVariable& InVar, bool& bFailIfNotSet);
+
+	void UpdateStaticSwitchConstants(FNode* Node);
+
+	virtual const FString& GetEmitterUniqueName() const override;
+	virtual TConstArrayView<FNiagaraVariable> GetStaticVariables() const override;
+	virtual UNiagaraDataInterface* GetDataInterfaceCDO(UClass* DIClass) const override;
+
+	// Specific method to reconcile whether the default value has been implicitly or explicitly set for a namespaced var added to the param map.
+	void RecordParamMapDefinedAttributeToNamespaceVar(const FNiagaraVariable& VarToRecord, const FPin* VarAssociatedDefaultPin);
+
+private:
+	const FPrecompileData* CompileData;
+	const FCompilationCopy* CompileDuplicateData;
+
+	/**
+	Map of Pins to compiled code chunks. Allows easy reuse of previously compiled pins.
+	A stack so that we can track pin reuse within function calls but not have cached pins cross talk with subsequent calls to the same function.
+	*/
+	TArray<TMap<const FPin*, int32>> PinToCodeChunks;
+
+	// Keep track of all the paths that the parameter maps can take through the graph.
+	TArray<FParamMapHistory> ParamMapHistories;
+
+	// Keep track of the other output nodes in the graph's histories so that we can make sure to 
+	// create any variables that are needed downstream.
+	TArray<FParamMapHistory> OtherOutputParamMapHistories;
+
+	// Make sure that the function call names match up on the second traversal.
+	FParamMapHistoryBuilder ActiveHistoryForFunctionCalls;
+
+	// Map of primary ouput variable description to its default value pin
+	TMap<FNiagaraVariable, const FInputPin*> UniqueVarToDefaultPin;
+
+	TArray<const FPin*> CurrentDefaultPinTraversal;
+
+	TArray<const FOutputNode*> OutputNodes;
+};
diff --git a/Engine/Plugins/FX/Niagara/Source/NiagaraEditor/Private/NiagaraMessages.cpp b/Engine/Plugins/FX/Niagara/Source/NiagaraEditor/Private/NiagaraMessages.cpp
--- a/Engine/Plugins/FX/Niagara/Source/NiagaraEditor/Private/NiagaraMessages.cpp
+++ b/Engine/Plugins/FX/Niagara/Source/NiagaraEditor/Private/NiagaraMessages.cpp
@@ -5,7 +5,7 @@
 #include "NiagaraEditorUtilities.h"
 #include "NiagaraMessageManager.h"
 #include "NiagaraMessageUtilities.h"
-#include "NiagaraScriptSourceBase.h"
+#include "NiagaraScriptSource.h"
 #include "NiagaraNodeFunctionCall.h"
 #include "NiagaraNodeEmitter.h"
 
diff --git a/Engine/Plugins/FX/Niagara/Source/NiagaraEditor/Private/NiagaraNode.cpp b/Engine/Plugins/FX/Niagara/Source/NiagaraEditor/Private/NiagaraNode.cpp
--- a/Engine/Plugins/FX/Niagara/Source/NiagaraEditor/Private/NiagaraNode.cpp
+++ b/Engine/Plugins/FX/Niagara/Source/NiagaraEditor/Private/NiagaraNode.cpp
@@ -3,7 +3,7 @@
 #include "NiagaraNode.h"
 #include "NiagaraGraph.h"
 #include "EdGraphSchema_Niagara.h"
-#include "NiagaraHlslTranslator.h"
+#include "NiagaraGraphHlslTranslator.h"
 #include "GraphEditAction.h"
 #include "Widgets/SNiagaraGraphNode.h"
 #include "Misc/SecureHash.h"
@@ -355,9 +355,9 @@ bool UNiagaraNode::ReallocatePins(bool bMarkNeedsResynchronizeOnChange)
 	return bAllSame;
 }
 
-int32 UNiagaraNode::CompileInputPin(FHlslNiagaraTranslator *Translator, UEdGraphPin* Pin)
+int32 UNiagaraNode::CompileInputPin(FTranslator* Translator, UEdGraphPin* Pin) const
 {
-	return Translator->CompilePin(Pin);
+	return Translator->CompileInputPin(Pin);
 }
 
 bool UNiagaraNode::IsValidPinToCompile(UEdGraphPin* Pin) const 
@@ -365,7 +365,7 @@ bool UNiagaraNode::IsValidPinToCompile(UEdGraphPin* Pin) const
 	return Pin->bOrphanedPin == false;
 }
 
-bool UNiagaraNode::CompileInputPins(FHlslNiagaraTranslator *Translator, TArray<int32>& OutCompiledInputs)
+bool UNiagaraNode::CompileInputPins(FTranslator* Translator, TArray<int32>& OutCompiledInputs) const
 {
 	bool bError = false;
 	
@@ -681,7 +681,7 @@ UNiagaraScriptSource* UNiagaraNode::GetSource()const
 	return GetNiagaraGraph()->GetSource();
 }
 
-void UNiagaraNode::Compile(FHlslNiagaraTranslator *Translator, TArray<int32>& Outputs)
+void UNiagaraNode::Compile(FTranslator* Translator, TArray<int32>& Outputs) const
 {
 	Translator->Error(FText::FromString("Unimplemented Node!"), this, nullptr);
 }
@@ -901,7 +901,7 @@ UEdGraphPin* UNiagaraNode::GetTracedOutputPin(UEdGraphPin* LocallyOwnedOutputPin
 }
 
 
-bool UNiagaraNode::SubstituteCompiledPin(FHlslNiagaraTranslator* Translator, UEdGraphPin** LocallyOwnedPin)
+bool UNiagaraNode::SubstituteCompiledPin(FTranslator* Translator, UEdGraphPin** LocallyOwnedPin)
 {
 	return false;
 }
diff --git a/Engine/Plugins/FX/Niagara/Source/NiagaraEditor/Private/NiagaraNodeConvert.cpp b/Engine/Plugins/FX/Niagara/Source/NiagaraEditor/Private/NiagaraNodeConvert.cpp
--- a/Engine/Plugins/FX/Niagara/Source/NiagaraEditor/Private/NiagaraNodeConvert.cpp
+++ b/Engine/Plugins/FX/Niagara/Source/NiagaraEditor/Private/NiagaraNodeConvert.cpp
@@ -4,7 +4,7 @@
 #include "EdGraphSchema_Niagara.h"
 #include "NiagaraEditorUtilities.h"
 #include "Widgets/SNiagaraGraphNodeConvert.h"
-#include "NiagaraHlslTranslator.h"
+#include "NiagaraGraphHlslTranslator.h"
 #include "UObject/UnrealType.h"
 #include "Kismet2/StructureEditorUtils.h"
 
@@ -149,10 +149,10 @@ TSharedPtr<SGraphNode> UNiagaraNodeConvert::CreateVisualWidget()
 	return SNew(SNiagaraGraphNodeConvert, this);
 }
 
-void UNiagaraNodeConvert::Compile(class FHlslNiagaraTranslator* Translator, TArray<int32>& CompileOutputs)
+void UNiagaraNodeConvert::Compile(FTranslator* Translator, TArray<int32>& CompileOutputs) const
 {
 	FPinCollectorArray InputPins;
-	GetInputPins(InputPins);
+	FTranslator::FBridge::GetCompilationInputPins(this, InputPins);
 
 	TArray<int32, TInlineAllocator<16>> CompileInputs;
 	CompileInputs.Reserve(InputPins.Num());
@@ -163,7 +163,7 @@ void UNiagaraNodeConvert::Compile(class FHlslNiagaraTranslator* Translator, TArr
 			InputPin->PinType.PinCategory == UEdGraphSchema_Niagara::PinCategoryStaticType ||
 			InputPin->PinType.PinCategory == UEdGraphSchema_Niagara::PinCategoryStaticEnum)
 		{
-			int32 CompiledInput = Translator->CompilePin(InputPin);
+			int32 CompiledInput = Translator->CompileInputPin(InputPin);
 			if (CompiledInput == INDEX_NONE)
 			{
 				Translator->Error(LOCTEXT("InputError", "Error compiling input for convert node."), this, InputPin);
@@ -176,7 +176,7 @@ void UNiagaraNodeConvert::Compile(class FHlslNiagaraTranslator* Translator, TArr
 	check(Schema);
 
 	FPinCollectorArray OutputPins;
-	GetOutputPins(OutputPins);
+	FTranslator::FBridge::GetCompilationOutputPins(this, OutputPins);
 
 	// Go through all the output nodes and cross-reference them with the connections list. Output errors if any connections are incomplete.
 	{
diff --git a/Engine/Plugins/FX/Niagara/Source/NiagaraEditor/Private/NiagaraNodeConvert.h b/Engine/Plugins/FX/Niagara/Source/NiagaraEditor/Private/NiagaraNodeConvert.h
--- a/Engine/Plugins/FX/Niagara/Source/NiagaraEditor/Private/NiagaraNodeConvert.h
+++ b/Engine/Plugins/FX/Niagara/Source/NiagaraEditor/Private/NiagaraNodeConvert.h
@@ -108,7 +108,7 @@ public:
 	virtual FText GetNodeTitle(ENodeTitleType::Type TitleType) const override;
 
 	//~ UNiagaraNode interface
-	virtual void Compile(class FHlslNiagaraTranslator* Translator, TArray<int32>& Outputs) override;
+	virtual void Compile(FTranslator* Translator, TArray<int32>& Outputs) const override;
 
 	virtual bool RefreshFromExternalChanges() override;
 
diff --git a/Engine/Plugins/FX/Niagara/Source/NiagaraEditor/Private/NiagaraNodeCustomHlsl.cpp b/Engine/Plugins/FX/Niagara/Source/NiagaraEditor/Private/NiagaraNodeCustomHlsl.cpp
--- a/Engine/Plugins/FX/Niagara/Source/NiagaraEditor/Private/NiagaraNodeCustomHlsl.cpp
+++ b/Engine/Plugins/FX/Niagara/Source/NiagaraEditor/Private/NiagaraNodeCustomHlsl.cpp
@@ -1,11 +1,13 @@
 // Copyright Epic Games, Inc. All Rights Reserved.
 
 #include "NiagaraNodeCustomHlsl.h"
-#include "Widgets/SNiagaraGraphNodeCustomHlsl.h"
+
 #include "EdGraphSchema_Niagara.h"
-#include "ScopedTransaction.h"
-#include "NiagaraGraph.h"
 #include "Misc/FileHelper.h"
+#include "NiagaraGraph.h"
+#include "NiagaraHlslTranslator.h"
+#include "ScopedTransaction.h"
+#include "Widgets/SNiagaraGraphNodeCustomHlsl.h"
 
 #include UE_INLINE_GENERATED_CPP_BY_NAME(NiagaraNodeCustomHlsl)
 
@@ -307,7 +309,7 @@ bool UNiagaraNodeCustomHlsl::VerifyEditablePinName(const FText& InName, FText& O
 	// Check to see if the symbol has to be mangled to be valid hlsl. If it does, then prevent it from being 
 	// valid. This helps clear up any ambiguity downstream in the translator.
 	FString NewName = InName.ToString();
-	FString SanitizedNewName = FHlslNiagaraTranslator::GetSanitizedSymbolName(NewName);
+	FString SanitizedNewName = FNiagaraHlslTranslator::GetSanitizedSymbolName(NewName);
 
 	if (NewName != SanitizedNewName || NewName.Len() == 0)
 	{
diff --git a/Engine/Plugins/FX/Niagara/Source/NiagaraEditor/Private/NiagaraNodeDataSetBase.cpp b/Engine/Plugins/FX/Niagara/Source/NiagaraEditor/Private/NiagaraNodeDataSetBase.cpp
--- a/Engine/Plugins/FX/Niagara/Source/NiagaraEditor/Private/NiagaraNodeDataSetBase.cpp
+++ b/Engine/Plugins/FX/Niagara/Source/NiagaraEditor/Private/NiagaraNodeDataSetBase.cpp
@@ -103,7 +103,7 @@ void UNiagaraNodeDataSetBase::PostLoad()
 	}
 }
 
-bool UNiagaraNodeDataSetBase::IsSynchronizedWithStruct(bool bIgnoreConditionVariable, FString* Issues, bool bLogIssues)
+bool UNiagaraNodeDataSetBase::IsSynchronizedWithStruct(bool bIgnoreConditionVariable, FString* Issues, bool bLogIssues) const
 {
 	if (ExternalStructAsset != nullptr)
 	{
@@ -224,7 +224,7 @@ bool UNiagaraNodeDataSetBase::RefreshFromExternalChanges()
 	return false;
 }
 
-bool UNiagaraNodeDataSetBase::GetSupportedNiagaraTypeDef(const FProperty* Property, FNiagaraTypeDefinition& TypeDef)
+bool UNiagaraNodeDataSetBase::GetSupportedNiagaraTypeDef(const FProperty* Property, FNiagaraTypeDefinition& TypeDef) const
 {
 	const FStructProperty* StructProp = CastField<FStructProperty>(Property);
 	if (Property->IsA(FFloatProperty::StaticClass()))
diff --git a/Engine/Plugins/FX/Niagara/Source/NiagaraEditor/Private/NiagaraNodeEmitter.cpp b/Engine/Plugins/FX/Niagara/Source/NiagaraEditor/Private/NiagaraNodeEmitter.cpp
--- a/Engine/Plugins/FX/Niagara/Source/NiagaraEditor/Private/NiagaraNodeEmitter.cpp
+++ b/Engine/Plugins/FX/Niagara/Source/NiagaraEditor/Private/NiagaraNodeEmitter.cpp
@@ -10,7 +10,7 @@
 #include "NiagaraGraph.h"
 #include "NiagaraNodeInput.h"
 #include "NiagaraNodeOutput.h"
-#include "NiagaraHlslTranslator.h"
+#include "NiagaraGraphHlslTranslator.h"
 #include "Stats/Stats.h"
 #include "NiagaraEditorModule.h"
 
@@ -440,7 +440,7 @@ void UNiagaraNodeEmitter::BuildParameterMapHistory(FNiagaraParameterMapHistoryBu
 	}
 }
 
-void UNiagaraNodeEmitter::Compile(FHlslNiagaraTranslator *Translator, TArray<int32>& Outputs)
+void UNiagaraNodeEmitter::Compile(FTranslator* Translator, TArray<int32>& Outputs) const
 {
 	NIAGARA_SCOPE_CYCLE_COUNTER(STAT_NiagaraEditor_Module_NiagaraNodeEmitter_Compile);
 	FPinCollectorArray InputPins;
@@ -448,7 +448,7 @@ void UNiagaraNodeEmitter::Compile(FHlslNiagaraTranslator *Translator, TArray<int
 	InputPins.RemoveAll([](UEdGraphPin* InputPin) { return (InputPin->PinType.PinCategory != UEdGraphSchema_Niagara::PinCategoryType) && (InputPin->PinType.PinCategory != UEdGraphSchema_Niagara::PinCategoryEnum); });
 
 	FPinCollectorArray OutputPins;
-	GetOutputPins(OutputPins);
+	FTranslator::FBridge::GetCompilationOutputPins(this, OutputPins);
 
 	check(Outputs.Num() == 0);
 
@@ -474,7 +474,7 @@ void UNiagaraNodeEmitter::Compile(FHlslNiagaraTranslator *Translator, TArray<int
 		return;
 	}
 
-	int32 InputPinCompiled = Translator->CompilePin(InputPins[0]);
+	int32 InputPinCompiled = Translator->CompileInputPin(InputPins[0]);
 	if (!IsNodeEnabled())
 	{
 		// Do the minimal amount of work necessary if we are disabled.
diff --git a/Engine/Plugins/FX/Niagara/Source/NiagaraEditor/Private/NiagaraNodeEmitter.h b/Engine/Plugins/FX/Niagara/Source/NiagaraEditor/Private/NiagaraNodeEmitter.h
--- a/Engine/Plugins/FX/Niagara/Source/NiagaraEditor/Private/NiagaraNodeEmitter.h
+++ b/Engine/Plugins/FX/Niagara/Source/NiagaraEditor/Private/NiagaraNodeEmitter.h
@@ -59,7 +59,7 @@ public:
 	UNiagaraScriptSource* GetScriptSource() const;
 	UNiagaraGraph* GetCalledGraph() const;
 	
-	virtual void Compile(FHlslNiagaraTranslator *Translator, TArray<int32>& Outputs) override;
+	virtual void Compile(FTranslator* Translator, TArray<int32>& Outputs) const override;
 	virtual void GatherExternalDependencyData(ENiagaraScriptUsage InUsage, const FGuid& InUsageId, TArray<FNiagaraCompileHash>& InReferencedCompileHashes, TArray<FString>& InReferencedObjs) const override;
 
 	void SetCachedVariablesForCompilation(const FName& InUniqueName, UNiagaraGraph* InGraph, UNiagaraScriptSourceBase* InSource);
diff --git a/Engine/Plugins/FX/Niagara/Source/NiagaraEditor/Private/NiagaraNodeFunctionCall.cpp b/Engine/Plugins/FX/Niagara/Source/NiagaraEditor/Private/NiagaraNodeFunctionCall.cpp
--- a/Engine/Plugins/FX/Niagara/Source/NiagaraEditor/Private/NiagaraNodeFunctionCall.cpp
+++ b/Engine/Plugins/FX/Niagara/Source/NiagaraEditor/Private/NiagaraNodeFunctionCall.cpp
@@ -11,7 +11,7 @@
 #include "Subsystems/AssetEditorSubsystem.h"
 #include "NiagaraEditorUtilities.h"
 #include "NiagaraGraph.h"
-#include "NiagaraHlslTranslator.h"
+#include "NiagaraGraphHlslTranslator.h"
 #include "NiagaraMessages.h"
 #include "NiagaraNodeInput.h"
 #include "NiagaraNodeOutput.h"
@@ -694,14 +694,16 @@ static FText GetFormattedDeprecationMessage(const FVersionedNiagaraScriptData* S
 	return FText::Format(FormatString, Args);
 }
 
-void UNiagaraNodeFunctionCall::Compile(class FHlslNiagaraTranslator* Translator, TArray<int32>& Outputs)
+void UNiagaraNodeFunctionCall::Compile(FTranslator* Translator, TArray<int32>& Outputs) const
 {
+	// this is dirty, but we'll continue doing this for now because it's all going to go away with Digested graphs
+	UNiagaraNodeFunctionCall* MutableThis = const_cast<UNiagaraNodeFunctionCall*>(this);
+
 	TArray<int32> Inputs;
 
 	bool bError = false;
 
 	const UEdGraphSchema_Niagara* Schema = CastChecked<UEdGraphSchema_Niagara>(GetSchema());
-	UNiagaraGraph* CallerGraph = GetNiagaraGraph();
 	FVersionedNiagaraScriptData* ScriptData = GetScriptData();
 	if (ScriptData && HasValidScriptAndGraph())
 	{
@@ -775,6 +777,8 @@ void UNiagaraNodeFunctionCall::Compile(class FHlslNiagaraTranslator* Translator,
 					ENiagaraInputNodeUsage AutBoundUsage = ENiagaraInputNodeUsage::Undefined;
 					if (FindAutoBoundInput(FunctionInputNode, CallerPin, AutoBoundVar, AutBoundUsage))
 					{
+						UNiagaraGraph* CallerGraph = MutableThis->GetNiagaraGraph();
+
 						UNiagaraNodeInput* NewNode = NewObject<UNiagaraNodeInput>(CallerGraph);
 						NewNode->Input = PinVar;
 						NewNode->Usage = AutBoundUsage;
@@ -789,7 +793,7 @@ void UNiagaraNodeFunctionCall::Compile(class FHlslNiagaraTranslator* Translator,
 			if (CallerLinkedTo)
 			{
 				//Param is provided by the caller. Typical case.
-				Inputs.Add(Translator->CompilePin(CallerPin));
+				Inputs.Add(Translator->CompileInputPin(CallerPin));
 				continue;
 			}
 			else
@@ -808,7 +812,7 @@ void UNiagaraNodeFunctionCall::Compile(class FHlslNiagaraTranslator* Translator,
 					else
 					{
 						//We also compile the pin anyway if it is required as we'll be attempting to use it's inline default.
-						Inputs.Add(Translator->CompilePin(CallerPin));
+						Inputs.Add(Translator->CompileInputPin(CallerPin));
 					}
 				}
 				else
@@ -823,13 +827,13 @@ void UNiagaraNodeFunctionCall::Compile(class FHlslNiagaraTranslator* Translator,
 		FNiagaraEditorUtilities::SetStaticSwitchConstants(GetCalledGraph(), CallerInputPins, ConstantResolver);
 		Translator->ExitFunctionCallNode();
 	}
-	else if (Signature.IsValid())
+	else if (MutableThis->Signature.IsValid())
 	{
-		if (Signature.Inputs.Num() > 0)
+		if (MutableThis->Signature.Inputs.Num() > 0)
 		{
-			if (Signature.Inputs[0].GetType().IsDataInterface() && GetValidateDataInterfaces())
+			if (MutableThis->Signature.Inputs[0].GetType().IsDataInterface() && GetValidateDataInterfaces())
 			{
-				UClass* DIClass = Signature.Inputs[0].GetType().GetClass();
+				UClass* DIClass = MutableThis->Signature.Inputs[0].GetType().GetClass();
 				if (UNiagaraDataInterface* DataInterfaceCDO = Cast<UNiagaraDataInterface>(DIClass->GetDefaultObject()))
 				{
 					TArray<FText> ValidationErrors;
@@ -850,7 +854,7 @@ void UNiagaraNodeFunctionCall::Compile(class FHlslNiagaraTranslator* Translator,
 			}
 		}
 		Translator->EnterFunctionCallNode(TSet<FName>());
-		Signature.FunctionSpecifiers = FunctionSpecifiers;
+		MutableThis->Signature.FunctionSpecifiers = FunctionSpecifiers;
 		bError = CompileInputPins(Translator, Inputs);
 		Translator->ExitFunctionCallNode();
 	}		
@@ -1306,7 +1310,7 @@ void UNiagaraNodeFunctionCall::BuildParameterMapHistory(FNiagaraParameterMapHist
 			ParamMapIdx = OutHistory.TraceParameterMapOutputPin((ParamMapPin->LinkedTo[0]));
 		}
 
-		OutHistory.EnterFunction(GetFunctionName(), FunctionScript, FunctionGraph, this);
+		OutHistory.EnterFunction(GetFunctionName(), FunctionGraph, this);
 		if (ParamMapIdx != INDEX_NONE)
 		{
 			NodeIdx = OutHistory.BeginNodeVisitation(ParamMapIdx, this);
@@ -1387,7 +1391,7 @@ void UNiagaraNodeFunctionCall::BuildParameterMapHistory(FNiagaraParameterMapHist
 			OutHistory.EndNodeVisitation(ParamMapIdx, NodeIdx);
 		}
 
-		OutHistory.ExitFunction(GetFunctionName(), FunctionScript, this);
+		OutHistory.ExitFunction(GetFunctionName(), this);
 
 		if (DoDepthTraversal)
 		{
@@ -1834,7 +1838,7 @@ bool UNiagaraNodeFunctionCall::IsValidPropagatedVariable(const FNiagaraVariable&
 	return false;
 }
 
-bool UNiagaraNodeFunctionCall::FindAutoBoundInput(UNiagaraNodeInput* InputNode, UEdGraphPin* PinToAutoBind, FNiagaraVariable& OutFoundVar, ENiagaraInputNodeUsage& OutNodeUsage)
+bool UNiagaraNodeFunctionCall::FindAutoBoundInput(const UNiagaraNodeInput* InputNode, const UEdGraphPin* PinToAutoBind, FNiagaraVariable& OutFoundVar, ENiagaraInputNodeUsage& OutNodeUsage) const
 {
 	check(InputNode);
 	if (PinToAutoBind->LinkedTo.Num() > 0 || !InputNode->CanAutoBind())
@@ -1847,7 +1851,7 @@ bool UNiagaraNodeFunctionCall::FindAutoBoundInput(UNiagaraNodeInput* InputNode,
 	FNiagaraVariable PinVar = Schema->PinToNiagaraVariable(PinToAutoBind);
 
 	//See if we can auto bind this pin to something in the caller script.
-	UNiagaraGraph* CallerGraph = GetNiagaraGraph();
+	const UNiagaraGraph* CallerGraph = GetNiagaraGraph();
 	check(CallerGraph);
 	UNiagaraNodeOutput* CallerOutputNodeSpawn = CallerGraph->FindOutputNode(ENiagaraScriptUsage::ParticleSpawnScript);
 	UNiagaraNodeOutput* CallerOutputNodeUpdate = CallerGraph->FindOutputNode(ENiagaraScriptUsage::ParticleUpdateScript);
diff --git a/Engine/Plugins/FX/Niagara/Source/NiagaraEditor/Private/NiagaraNodeIf.cpp b/Engine/Plugins/FX/Niagara/Source/NiagaraEditor/Private/NiagaraNodeIf.cpp
--- a/Engine/Plugins/FX/Niagara/Source/NiagaraEditor/Private/NiagaraNodeIf.cpp
+++ b/Engine/Plugins/FX/Niagara/Source/NiagaraEditor/Private/NiagaraNodeIf.cpp
@@ -2,7 +2,7 @@
 
 #include "NiagaraNodeIf.h"
 #include "NiagaraEditorUtilities.h"
-#include "NiagaraHlslTranslator.h"
+#include "NiagaraGraphHlslTranslator.h"
 
 #include UE_INLINE_GENERATED_CPP_BY_NAME(NiagaraNodeIf)
 
@@ -121,12 +121,12 @@ void UNiagaraNodeIf::AllocateDefaultPins()
 	CreateAddPin(EGPD_Output);
 }
 
-void UNiagaraNodeIf::Compile(class FHlslNiagaraTranslator* Translator, TArray<int32>& Outputs)
+void UNiagaraNodeIf::Compile(FTranslator* Translator, TArray<int32>& Outputs) const
 {
 	checkSlow(PathAssociatedPinGuids.Num() == OutputVars.Num());
 	const UEdGraphSchema_Niagara* Schema = CastChecked<UEdGraphSchema_Niagara>(GetSchema());
 
-	int32 Condition = Translator->CompilePin(GetPinByGuid(ConditionPinGuid));
+	int32 Condition = Translator->CompileInputPin(GetPinByGuid(ConditionPinGuid));
 
 	TArray<int32> PathTrue;
 	PathTrue.Reserve(PathAssociatedPinGuids.Num());
@@ -137,7 +137,7 @@ void UNiagaraNodeIf::Compile(class FHlslNiagaraTranslator* Translator, TArray<in
 		{
 			Translator->Error(LOCTEXT("UnsupportedParamMapInIf", "Parameter maps are not supported in if nodes."), this, InputTruePin);
 		}
-		PathTrue.Add(Translator->CompilePin(InputTruePin));
+		PathTrue.Add(Translator->CompileInputPin(InputTruePin));
 	}
 	TArray<int32> PathFalse;
 	PathFalse.Reserve(PathAssociatedPinGuids.Num());
@@ -148,7 +148,7 @@ void UNiagaraNodeIf::Compile(class FHlslNiagaraTranslator* Translator, TArray<in
 		{
 			Translator->Error(LOCTEXT("UnsupportedParamMapInIf", "Parameter maps are not supported in if nodes."), this, InputFalsePin);
 		}
-		PathFalse.Add(Translator->CompilePin(InputFalsePin));
+		PathFalse.Add(Translator->CompileInputPin(InputFalsePin));
 	}
 
 	Translator->If(this, OutputVars, Condition, PathTrue, PathFalse, Outputs);
diff --git a/Engine/Plugins/FX/Niagara/Source/NiagaraEditor/Private/NiagaraNodeIf.h b/Engine/Plugins/FX/Niagara/Source/NiagaraEditor/Private/NiagaraNodeIf.h
--- a/Engine/Plugins/FX/Niagara/Source/NiagaraEditor/Private/NiagaraNodeIf.h
+++ b/Engine/Plugins/FX/Niagara/Source/NiagaraEditor/Private/NiagaraNodeIf.h
@@ -61,7 +61,7 @@ public:
 	//~ End EdGraphNode Interface
 
 	//~ Begin UNiagaraNode Interface
-	virtual void Compile(class FHlslNiagaraTranslator* Translator, TArray<int32>& Outputs) override;
+	virtual void Compile(FTranslator* Translator, TArray<int32>& Outputs) const override;
 	virtual bool RefreshFromExternalChanges() override;
 	virtual ENiagaraNumericOutputTypeSelectionMode GetNumericOutputTypeSelectionMode() const;
 	virtual void ResolveNumerics(const UEdGraphSchema_Niagara* Schema, bool bSetInline, TMap<TPair<FGuid, UEdGraphNode*>, FNiagaraTypeDefinition>* PinCache);
diff --git a/Engine/Plugins/FX/Niagara/Source/NiagaraEditor/Private/NiagaraNodeInput.cpp b/Engine/Plugins/FX/Niagara/Source/NiagaraEditor/Private/NiagaraNodeInput.cpp
--- a/Engine/Plugins/FX/Niagara/Source/NiagaraEditor/Private/NiagaraNodeInput.cpp
+++ b/Engine/Plugins/FX/Niagara/Source/NiagaraEditor/Private/NiagaraNodeInput.cpp
@@ -2,7 +2,7 @@
 
 #include "NiagaraNodeInput.h"
 #include "UObject/UnrealType.h"
-#include "NiagaraHlslTranslator.h"
+#include "NiagaraGraphHlslTranslator.h"
 #include "NiagaraGraph.h"
 
 #include "NiagaraNodeOutput.h"
@@ -126,7 +126,7 @@ void UNiagaraNodeInput::BuildParameterMapHistory(FNiagaraParameterMapHistoryBuil
 	}
 }
 
-void UNiagaraNodeInput::AppendFunctionAliasForContext(const FNiagaraGraphFunctionAliasContext& InFunctionAliasContext, FString& InOutFunctionAlias, bool& OutOnlyOncePerNodeType)
+void UNiagaraNodeInput::AppendFunctionAliasForContext(const FNiagaraGraphFunctionAliasContext& InFunctionAliasContext, FString& InOutFunctionAlias, bool& OutOnlyOncePerNodeType) const
 {
 	if (Usage == ENiagaraInputNodeUsage::TranslatorConstant && Input == TRANSLATOR_PARAM_CALL_ID)
 	{
@@ -383,7 +383,7 @@ void UNiagaraNodeInput::NotifyExposureOptionsChanged()
 	ReallocatePins();
 }
 
-void UNiagaraNodeInput::Compile(class FHlslNiagaraTranslator* Translator, TArray<int32>& Outputs)
+void UNiagaraNodeInput::Compile(FTranslator* Translator, TArray<int32>& Outputs) const
 {
 	if (!IsNodeEnabled())
 	{
@@ -391,7 +391,6 @@ void UNiagaraNodeInput::Compile(class FHlslNiagaraTranslator* Translator, TArray
 		return;
 	}
 
-	UNiagaraGraph* Graph = GetNiagaraGraph();
 	if (Input.GetType() == FNiagaraTypeDefinition::GetGenericNumericDef())
 	{
 		Outputs.Add(INDEX_NONE);
@@ -407,7 +406,7 @@ void UNiagaraNodeInput::Compile(class FHlslNiagaraTranslator* Translator, TArray
 		{
 			FPinCollectorArray InputPins;
 			GetInputPins(InputPins);
-			int32 Default = InputPins.Num() > 0 ? Translator->CompilePin(InputPins[0]) : INDEX_NONE;
+			int32 Default = InputPins.Num() > 0 ? Translator->CompileInputPin(InputPins[0]) : INDEX_NONE;
 			if (Default == INDEX_NONE)
 			{
 				//We failed to compile the default pin so just use the value of the input.
diff --git a/Engine/Plugins/FX/Niagara/Source/NiagaraEditor/Private/NiagaraNodeOp.cpp b/Engine/Plugins/FX/Niagara/Source/NiagaraEditor/Private/NiagaraNodeOp.cpp
--- a/Engine/Plugins/FX/Niagara/Source/NiagaraEditor/Private/NiagaraNodeOp.cpp
+++ b/Engine/Plugins/FX/Niagara/Source/NiagaraEditor/Private/NiagaraNodeOp.cpp
@@ -1,7 +1,7 @@
 // Copyright Epic Games, Inc. All Rights Reserved.
 
 #include "NiagaraNodeOp.h"
-#include "NiagaraHlslTranslator.h"
+#include "NiagaraGraphHlslTranslator.h"
 #include "GraphEditorSettings.h"
 
 #include "EdGraphSchema_Niagara.h"
@@ -161,7 +161,7 @@ FNiagaraTypeDefinition UNiagaraNodeOp::ResolveCustomNumericType(const TArray<FNi
 }
 
 
-void UNiagaraNodeOp::Compile(class FHlslNiagaraTranslator* Translator, TArray<int32>& Outputs)
+void UNiagaraNodeOp::Compile(FTranslator* Translator, TArray<int32>& Outputs) const
 {
 	const FNiagaraOpInfo* OpInfo = FNiagaraOpInfo::GetOpInfo(OpName);
 	if (!OpInfo)
@@ -185,7 +185,7 @@ void UNiagaraNodeOp::Compile(class FHlslNiagaraTranslator* Translator, TArray<in
 		{
 			continue;
 		}
-		int32 CompiledInput = Translator->CompilePin(Pin);
+		int32 CompiledInput = Translator->CompileInputPin(Pin);
 		if (CompiledInput == INDEX_NONE)
 		{
 			bError = true;
diff --git a/Engine/Plugins/FX/Niagara/Source/NiagaraEditor/Private/NiagaraNodeOutput.cpp b/Engine/Plugins/FX/Niagara/Source/NiagaraEditor/Private/NiagaraNodeOutput.cpp
--- a/Engine/Plugins/FX/Niagara/Source/NiagaraEditor/Private/NiagaraNodeOutput.cpp
+++ b/Engine/Plugins/FX/Niagara/Source/NiagaraEditor/Private/NiagaraNodeOutput.cpp
@@ -2,7 +2,7 @@
 
 #include "NiagaraNodeOutput.h"
 #include "UObject/UnrealType.h"
-#include "NiagaraHlslTranslator.h"
+#include "NiagaraGraphHlslTranslator.h"
 #include "NiagaraEmitter.h"
 #include "NiagaraScript.h"
 #include "NiagaraGraph.h"
@@ -267,7 +267,7 @@ void UNiagaraNodeOutput::NotifyOutputVariablesChanged()
 	ReallocatePins();
 }
 
-int32 UNiagaraNodeOutput::CompileInputPin(FHlslNiagaraTranslator *Translator, UEdGraphPin* Pin)
+int32 UNiagaraNodeOutput::CompileInputPin(FTranslator* Translator, UEdGraphPin* Pin) const
 {
 	// If we are an update script, automatically fill in any unwired values with the previous frame's value...
 	if (GetUsage() == ENiagaraScriptUsage::ParticleUpdateScript && Pin->LinkedTo.Num() == 0)
@@ -283,10 +283,10 @@ int32 UNiagaraNodeOutput::CompileInputPin(FHlslNiagaraTranslator *Translator, UE
 		}
 	}
 
-	return Translator->CompilePin(Pin);
+	return Translator->CompileInputPin(Pin);
 }
 
-void UNiagaraNodeOutput::Compile(FHlslNiagaraTranslator *Translator, TArray<int32>& OutputExpressions)
+void UNiagaraNodeOutput::Compile(FTranslator* Translator, TArray<int32>& OutputExpressions) const
 {
 	TArray<int32> Results;
 	bool bError = CompileInputPins(Translator, Results);
diff --git a/Engine/Plugins/FX/Niagara/Source/NiagaraEditor/Private/NiagaraNodeOutputTag.cpp b/Engine/Plugins/FX/Niagara/Source/NiagaraEditor/Private/NiagaraNodeOutputTag.cpp
--- a/Engine/Plugins/FX/Niagara/Source/NiagaraEditor/Private/NiagaraNodeOutputTag.cpp
+++ b/Engine/Plugins/FX/Niagara/Source/NiagaraEditor/Private/NiagaraNodeOutputTag.cpp
@@ -2,6 +2,7 @@
 
 #include "NiagaraNodeOutputTag.h"
 #include "NiagaraCustomVersion.h"
+#include "NiagaraGraphHlslTranslator.h"
 #include "Widgets/SNiagaraGraphNodeCustomHlsl.h"
 #include "EdGraphSchema_Niagara.h"
 #include "ScopedTransaction.h"
@@ -32,7 +33,7 @@ void UNiagaraNodeOutputTag::AllocateDefaultPins()
 	CreateAddPin(EGPD_Input);
 }
 
-void UNiagaraNodeOutputTag::Compile(class FHlslNiagaraTranslator* Translator, TArray<int32>& Outputs)
+void UNiagaraNodeOutputTag::Compile(FTranslator* Translator, TArray<int32>& Outputs) const
 {
 	FPinCollectorArray InputPins;
 	GetInputPins(InputPins);
@@ -49,7 +50,7 @@ void UNiagaraNodeOutputTag::Compile(class FHlslNiagaraTranslator* Translator, TA
 		{
 			continue;
 		}
-		int32 CompiledInput = Translator->CompilePin(InputPin);
+		int32 CompiledInput = Translator->CompileInputPin(InputPin);
 
 		if (Schema->PinToTypeDefinition(InputPin) == FNiagaraTypeDefinition::GetParameterMapDef())
 		{
diff --git a/Engine/Plugins/FX/Niagara/Source/NiagaraEditor/Private/NiagaraNodeParameterMapFor.cpp b/Engine/Plugins/FX/Niagara/Source/NiagaraEditor/Private/NiagaraNodeParameterMapFor.cpp
--- a/Engine/Plugins/FX/Niagara/Source/NiagaraEditor/Private/NiagaraNodeParameterMapFor.cpp
+++ b/Engine/Plugins/FX/Niagara/Source/NiagaraEditor/Private/NiagaraNodeParameterMapFor.cpp
@@ -3,7 +3,7 @@
 #include "NiagaraNodeParameterMapFor.h"
 #include "EdGraphSchema_Niagara.h"
 #include "NiagaraEditorUtilities.h"
-#include "NiagaraHlslTranslator.h"
+#include "NiagaraGraphHlslTranslator.h"
 #include "EdGraph/EdGraphNode.h"
 
 #include UE_INLINE_GENERATED_CPP_BY_NAME(NiagaraNodeParameterMapFor)
@@ -57,7 +57,7 @@ bool UNiagaraNodeParameterMapFor::CanModifyPin(const UEdGraphPin* Pin) const
 	return true;
 }
 
-void UNiagaraNodeParameterMapFor::Compile(FHlslNiagaraTranslator* Translator, TArray<int32>& Outputs)
+void UNiagaraNodeParameterMapFor::Compile(FTranslator* Translator, TArray<int32>& Outputs) const
 {
 	FPinCollectorArray InputPins;
 	GetInputPins(InputPins);
@@ -65,7 +65,7 @@ void UNiagaraNodeParameterMapFor::Compile(FHlslNiagaraTranslator* Translator, TA
 	{
 		if (Translator->GetSimulationTarget() == ENiagaraSimTarget::GPUComputeSim)
 		{
-			const int32 IterationCount = Translator->CompilePin(InputPins[1]);
+			const int32 IterationCount = Translator->CompileInputPin(InputPins[1]);
 			Translator->ParameterMapForBegin(this, IterationCount);
 
 			UNiagaraNodeParameterMapSet::Compile(Translator, Outputs);
@@ -97,7 +97,7 @@ UNiagaraNodeParameterMapForWithContinue::UNiagaraNodeParameterMapForWithContinue
 	UEdGraphNode::NodeUpgradeMessage = LOCTEXT("NodeExperimental", "This node is marked as experimental, use with care!");
 }
 
-void UNiagaraNodeParameterMapForWithContinue::Compile(FHlslNiagaraTranslator* Translator, TArray<int32>& Outputs)
+void UNiagaraNodeParameterMapForWithContinue::Compile(FTranslator* Translator, TArray<int32>& Outputs) const
 {
 	FPinCollectorArray InputPins;
 	GetInputPins(InputPins);
@@ -105,11 +105,11 @@ void UNiagaraNodeParameterMapForWithContinue::Compile(FHlslNiagaraTranslator* Tr
 	{
 		if (Translator->GetSimulationTarget() == ENiagaraSimTarget::GPUComputeSim)
 		{
-			const int32 IterationCount = Translator->CompilePin(InputPins[1]);
+			const int32 IterationCount = Translator->CompileInputPin(InputPins[1]);
 
 			Translator->ParameterMapForBegin(this, IterationCount);
 			
-			const int32 IterationEnabledChunk = Translator->CompilePin(InputPins[2]);
+			const int32 IterationEnabledChunk = Translator->CompileInputPin(InputPins[2]);
 			Translator->ParameterMapForContinue(this, IterationEnabledChunk);
 
 			UNiagaraNodeParameterMapSet::Compile(Translator, Outputs);
@@ -194,7 +194,7 @@ FText UNiagaraNodeParameterMapForIndex::GetNodeTitle(ENodeTitleType::Type TitleT
 	return LOCTEXT("UNiagaraNodeParameterMapForIndex", "Map For Index");
 }
 
-void UNiagaraNodeParameterMapForIndex::Compile(class FHlslNiagaraTranslator* Translator, TArray<int32>& Outputs)
+void UNiagaraNodeParameterMapForIndex::Compile(FTranslator* Translator, TArray<int32>& Outputs) const
 {
 	if (Translator)
 	{
diff --git a/Engine/Plugins/FX/Niagara/Source/NiagaraEditor/Private/NiagaraNodeParameterMapFor.h b/Engine/Plugins/FX/Niagara/Source/NiagaraEditor/Private/NiagaraNodeParameterMapFor.h
--- a/Engine/Plugins/FX/Niagara/Source/NiagaraEditor/Private/NiagaraNodeParameterMapFor.h
+++ b/Engine/Plugins/FX/Niagara/Source/NiagaraEditor/Private/NiagaraNodeParameterMapFor.h
@@ -17,7 +17,7 @@ public:
 
 	virtual void AllocateDefaultPins() override;
 
-	virtual void Compile(class FHlslNiagaraTranslator* Translator, TArray<int32>& Outputs) override;
+	virtual void Compile(FTranslator* Translator, TArray<int32>& Outputs) const override;
 
 	virtual FText GetNodeTitle(ENodeTitleType::Type TitleType) const override;
 	virtual bool IncludeParentNodeContextMenu() const override { return false; }
@@ -38,7 +38,7 @@ public:
 	UNiagaraNodeParameterMapForWithContinue();
 
 	virtual void AllocateDefaultPins() override;
-	virtual void Compile(class FHlslNiagaraTranslator* Translator, TArray<int32>& Outputs) override;
+	virtual void Compile(FTranslator* Translator, TArray<int32>& Outputs) const override;
 	virtual FText GetNodeTitle(ENodeTitleType::Type TitleType) const override;
 	virtual bool IncludeParentNodeContextMenu() const override { return false; }
 	virtual void GetPinHoverText(const UEdGraphPin& Pin, FString& HoverTextOut) const override;
@@ -58,5 +58,5 @@ public:
 
 	virtual void AllocateDefaultPins() override;
 	virtual FText GetNodeTitle(ENodeTitleType::Type TitleType) const override;
-	virtual void Compile(class FHlslNiagaraTranslator* Translator, TArray<int32>& Outputs) override;
+	virtual void Compile(FTranslator* Translator, TArray<int32>& Outputs) const override;
 };
diff --git a/Engine/Plugins/FX/Niagara/Source/NiagaraEditor/Private/NiagaraNodeParameterMapGet.cpp b/Engine/Plugins/FX/Niagara/Source/NiagaraEditor/Private/NiagaraNodeParameterMapGet.cpp
--- a/Engine/Plugins/FX/Niagara/Source/NiagaraEditor/Private/NiagaraNodeParameterMapGet.cpp
+++ b/Engine/Plugins/FX/Niagara/Source/NiagaraEditor/Private/NiagaraNodeParameterMapGet.cpp
@@ -4,7 +4,7 @@
 #include "EdGraphSchema_Niagara.h"
 #include "NiagaraEditorUtilities.h"
 #include "NiagaraGraph.h"
-#include "NiagaraHlslTranslator.h"
+#include "NiagaraGraphHlslTranslator.h"
 #include "ScopedTransaction.h"
 #include "Widgets/SNiagaraGraphParameterMapGetNode.h"
 #include "NiagaraEditorModule.h"
@@ -517,25 +517,18 @@ void UNiagaraNodeParameterMapGet::BuildParameterMapHistory(FNiagaraParameterMapH
 	}
 }
 
-void UNiagaraNodeParameterMapGet::Compile(class FHlslNiagaraTranslator* Translator, TArray<int32>& Outputs)
+void UNiagaraNodeParameterMapGet::Compile(FTranslator* Translator, TArray<int32>& Outputs) const
 {
 	const UEdGraphSchema_Niagara* Schema = CastChecked<UEdGraphSchema_Niagara>(GetSchema());
 
 	FPinCollectorArray InputPins;
-	GetInputPins(InputPins);
+	FTranslator::FBridge::GetCompilationInputPins(this, InputPins);
 	FPinCollectorArray OutputPins;
-	GetOutputPins(OutputPins);
+	FTranslator::FBridge::GetCompilationOutputPins(this, OutputPins);
 
 	// Initialize the outputs to invalid values.
 	check(Outputs.Num() == 0);
-	for (int32 i = 0; i < OutputPins.Num(); i++)
-	{
-		if (IsAddPin(OutputPins[i]))
-		{
-			continue;
-		}
-		Outputs.Add(INDEX_NONE);
-	}
+	Outputs.Init(INDEX_NONE, OutputPins.Num());
 
 	// First compile fully down the hierarchy for our predecessors..
 	TArray<int32, TInlineAllocator<16>> CompileInputs;
@@ -553,7 +546,7 @@ void UNiagaraNodeParameterMapGet::Compile(class FHlslNiagaraTranslator* Translat
 			int32 CompiledInput = INDEX_NONE;
 			if (i == 0) // Only the zeroth item is not an default value pin.
 			{
-				CompiledInput = Translator->CompilePin(InputPin);
+				CompiledInput = Translator->CompileInputPin(InputPin);
 				if (CompiledInput == INDEX_NONE)
 				{
 					Translator->Error(LOCTEXT("InputError", "Error compiling input for param map get node."), this, InputPin);
diff --git a/Engine/Plugins/FX/Niagara/Source/NiagaraEditor/Private/NiagaraNodeParameterMapSet.cpp b/Engine/Plugins/FX/Niagara/Source/NiagaraEditor/Private/NiagaraNodeParameterMapSet.cpp
--- a/Engine/Plugins/FX/Niagara/Source/NiagaraEditor/Private/NiagaraNodeParameterMapSet.cpp
+++ b/Engine/Plugins/FX/Niagara/Source/NiagaraEditor/Private/NiagaraNodeParameterMapSet.cpp
@@ -7,7 +7,7 @@
 #include "EdGraphSchema_Niagara.h"
 #include "NiagaraConstants.h"
 #include "NiagaraEditorUtilities.h"
-#include "NiagaraHlslTranslator.h"
+#include "NiagaraGraphHlslTranslator.h"
 #include "NiagaraGraph.h"
 #include "NiagaraScriptVariable.h"
 #include "ScopedTransaction.h"
@@ -192,30 +192,21 @@ bool UNiagaraNodeParameterMapSet::CommitEditablePinName(const FText& InName, UEd
 	return false;
 }
 
-void UNiagaraNodeParameterMapSet::Compile(class FHlslNiagaraTranslator* Translator, TArray<int32>& Outputs)
+void UNiagaraNodeParameterMapSet::Compile(FTranslator* Translator, TArray<int32>& Outputs) const
 {
 	FPinCollectorArray InputPins;
-	GetInputPins(InputPins);
+	FTranslator::FBridge::GetCompilationInputPins(this, InputPins);
 
 	FPinCollectorArray OutputPins;
-	GetOutputPins(OutputPins);
+	FTranslator::FBridge::GetCompilationOutputPins(this, OutputPins);
 
 	// Initialize the outputs to invalid values.
 	check(Outputs.Num() == 0);
-	Outputs.Reserve(OutputPins.Num());
-	for (int32 i = 0; i < OutputPins.Num(); i++)
-	{
-		Outputs.Add(INDEX_NONE);
-	}
+	Outputs.Init(INDEX_NONE, OutputPins.Num());
 
 	// update the translator with the culled function calls before compiling any further
 	for (UEdGraphPin* InputPin : InputPins)
 	{
-		if (IsAddPin(InputPin))
-		{
-			continue;
-		}
-
 		if (Translator->IsFunctionVariableCulledFromCompilation(InputPin->PinName))
 		{
 			Translator->CullMapSetInputPin(InputPin);
@@ -225,11 +216,11 @@ void UNiagaraNodeParameterMapSet::Compile(class FHlslNiagaraTranslator* Translat
 	const UEdGraphSchema_Niagara* Schema = CastChecked<UEdGraphSchema_Niagara>(GetSchema());
 
 	// First compile fully down the hierarchy for our predecessors..
-	TArray<FCompiledPin, TInlineAllocator<16>> CompileInputs;
+	TArray<FTranslator::FCompiledPin, TInlineAllocator<16>> CompileInputs;
 	CompileInputs.Reserve(InputPins.Num());
 	for (UEdGraphPin* InputPin : InputPins)
 	{
-		if (IsAddPin(InputPin) || SkipPinCompilation(InputPin) || Translator->IsFunctionVariableCulledFromCompilation(InputPin->PinName))
+		if (SkipPinCompilation(InputPin) || Translator->IsFunctionVariableCulledFromCompilation(InputPin->PinName))
 		{
 			continue;
 		}
@@ -239,12 +230,12 @@ void UNiagaraNodeParameterMapSet::Compile(class FHlslNiagaraTranslator* Translat
 			continue;
 		}
 
-		int32 CompiledInput = Translator->CompilePin(InputPin);
+		int32 CompiledInput = Translator->CompileInputPin(InputPin);
 		if (CompiledInput == INDEX_NONE)
 		{
 			Translator->Error(LOCTEXT("InputError", "Error compiling input for set node."), this, InputPin);
 		}
-		CompileInputs.Add(FCompiledPin(CompiledInput, InputPin));
+		CompileInputs.Emplace(CompiledInput, InputPin);
 	}
 
 	if (GetInputPin(0) != nullptr && GetInputPin(0)->LinkedTo.Num() > 0)
diff --git a/Engine/Plugins/FX/Niagara/Source/NiagaraEditor/Private/NiagaraNodeReadDataSet.cpp b/Engine/Plugins/FX/Niagara/Source/NiagaraEditor/Private/NiagaraNodeReadDataSet.cpp
--- a/Engine/Plugins/FX/Niagara/Source/NiagaraEditor/Private/NiagaraNodeReadDataSet.cpp
+++ b/Engine/Plugins/FX/Niagara/Source/NiagaraEditor/Private/NiagaraNodeReadDataSet.cpp
@@ -2,7 +2,7 @@
 
 #include "NiagaraNodeReadDataSet.h"
 #include "UObject/UnrealType.h"
-#include "NiagaraHlslTranslator.h"
+#include "NiagaraGraphHlslTranslator.h"
 #include "NiagaraEvents.h"
 #include "EdGraphSchema_Niagara.h"
 #include "NiagaraGraph.h"
@@ -98,7 +98,7 @@ bool UNiagaraNodeReadDataSet::CanAddToGraph(UNiagaraGraph* TargetGraph, FString&
 	return true;
 }
 
-void UNiagaraNodeReadDataSet::Compile(class FHlslNiagaraTranslator* Translator, TArray<int32>& Outputs)
+void UNiagaraNodeReadDataSet::Compile(FTranslator* Translator, TArray<int32>& Outputs) const
 {
 	TArray<int32> Inputs;
 	CompileInputPins(Translator, Inputs);
diff --git a/Engine/Plugins/FX/Niagara/Source/NiagaraEditor/Private/NiagaraNodeReroute.cpp b/Engine/Plugins/FX/Niagara/Source/NiagaraEditor/Private/NiagaraNodeReroute.cpp
--- a/Engine/Plugins/FX/Niagara/Source/NiagaraEditor/Private/NiagaraNodeReroute.cpp
+++ b/Engine/Plugins/FX/Niagara/Source/NiagaraEditor/Private/NiagaraNodeReroute.cpp
@@ -2,7 +2,7 @@
 
 #include "NiagaraNodeReroute.h"
 #include "NiagaraEditorUtilities.h"
-#include "NiagaraHlslTranslator.h"
+#include "NiagaraGraphHlslTranslator.h"
 #include "NiagaraConstants.h"
 
 #include UE_INLINE_GENERATED_CPP_BY_NAME(NiagaraNodeReroute)
@@ -106,13 +106,13 @@ bool UNiagaraNodeReroute::ShouldDrawNodeAsControlPointOnly(int32& OutInputPinInd
 	return true;
 }
 
-void UNiagaraNodeReroute::Compile(class FHlslNiagaraTranslator* Translator, TArray<int32>& Outputs)
+void UNiagaraNodeReroute::Compile(FTranslator* Translator, TArray<int32>& Outputs) const
 {
 	FPinCollectorArray InputPins;
-	GetInputPins(InputPins);
+	FTranslator::FBridge::GetCompilationInputPins(this, InputPins);
 
 	FPinCollectorArray OutputPins;
-	GetOutputPins(OutputPins);
+	FTranslator::FBridge::GetCompilationOutputPins(this, OutputPins);
 
 	// Initialize the outputs to invalid values.
 	check(Outputs.Num() == 0);
@@ -121,7 +121,7 @@ void UNiagaraNodeReroute::Compile(class FHlslNiagaraTranslator* Translator, TArr
 	{
 		if (InputPins.IsValidIndex(i))
 		{
-			int32 CompiledInput = Translator->CompilePin(InputPins[i]);
+			int32 CompiledInput = Translator->CompileInputPin(InputPins[i]);
 			Outputs.Add(CompiledInput);
 		}
 		else
diff --git a/Engine/Plugins/FX/Niagara/Source/NiagaraEditor/Private/NiagaraNodeSelect.cpp b/Engine/Plugins/FX/Niagara/Source/NiagaraEditor/Private/NiagaraNodeSelect.cpp
--- a/Engine/Plugins/FX/Niagara/Source/NiagaraEditor/Private/NiagaraNodeSelect.cpp
+++ b/Engine/Plugins/FX/Niagara/Source/NiagaraEditor/Private/NiagaraNodeSelect.cpp
@@ -6,7 +6,7 @@
 #include "FindInBlueprintManager.h"
 #include "NiagaraEditorStyle.h"
 #include "NiagaraEditorUtilities.h"
-#include "NiagaraHlslTranslator.h"
+#include "NiagaraGraphHlslTranslator.h"
 #include "EdGraph/EdGraphPin.h"
 #include "Widgets/SNiagaraPinTypeSelector.h"
 #include "ScopedTransaction.h"
@@ -306,7 +306,7 @@ void UNiagaraNodeSelect::GetNodeContextMenuActions(UToolMenu* Menu, UGraphNodeCo
 	}
 }
 
-void UNiagaraNodeSelect::Compile(FHlslNiagaraTranslator* Translator, TArray<int32>& Outputs)
+void UNiagaraNodeSelect::Compile(FTranslator* Translator, TArray<int32>& Outputs) const
 {	
 	const UEdGraphPin* SelectorPin = GetSelectorPin();
 
@@ -315,7 +315,7 @@ void UNiagaraNodeSelect::Compile(FHlslNiagaraTranslator* Translator, TArray<int3
 		Translator->Warning(LOCTEXT("SelectNodePinSelectorTypeInvalid", "Select node selector pin should have a valid type."), this, nullptr);
 	}
 	
-	for(FNiagaraVariable& Variable : OutputVars)
+	for(const FNiagaraVariable& Variable : OutputVars)
 	{
 		if(!Variable.GetType().IsValid())
 		{
@@ -323,7 +323,7 @@ void UNiagaraNodeSelect::Compile(FHlslNiagaraTranslator* Translator, TArray<int3
 		}
 	}
 	
-	int32 Selection = Translator->CompilePin(SelectorPin);
+	int32 Selection = Translator->CompileInputPin(SelectorPin);
 
 	// a map from selector value to compiled option pins (i.e.: for selector value 0 all pins that should be case "if 0" get their compiled index added under key 0)
 	TMap<int32, TArray<int32>> OptionValues;
@@ -359,7 +359,7 @@ void UNiagaraNodeSelect::Compile(FHlslNiagaraTranslator* Translator, TArray<int3
 		for (int32 OutputIndex = 0; OutputIndex < OutputVars.Num(); OutputIndex++)
 		{
 			TArray<UEdGraphPin*> OptionPins = GetOptionPins(OutputIndex);
-			int32 CodeChunkIndex = Translator->CompilePin(OptionPins[SelectorValueIndex]);
+			int32 CodeChunkIndex = Translator->CompileInputPin(OptionPins[SelectorValueIndex]);
 			OptionValues[SelectorValues[SelectorValueIndex]].Add(CodeChunkIndex);
 		}
 	}	
diff --git a/Engine/Plugins/FX/Niagara/Source/NiagaraEditor/Private/NiagaraNodeSelect.h b/Engine/Plugins/FX/Niagara/Source/NiagaraEditor/Private/NiagaraNodeSelect.h
--- a/Engine/Plugins/FX/Niagara/Source/NiagaraEditor/Private/NiagaraNodeSelect.h
+++ b/Engine/Plugins/FX/Niagara/Source/NiagaraEditor/Private/NiagaraNodeSelect.h
@@ -54,7 +54,7 @@ private:
 	virtual void GetNodeContextMenuActions(UToolMenu* Menu, UGraphNodeContextMenuContext* Context) const override;
 
 	/** UNiagaraNode interface */
-	virtual void Compile(FHlslNiagaraTranslator* Translator, TArray<int32>& Outputs) override;
+	virtual void Compile(FTranslator* Translator, TArray<int32>& Outputs) const override;
 	virtual bool AllowExternalPinTypeChanges(const UEdGraphPin* InGraphPin) const override;
 	virtual bool AllowNiagaraTypeForPinTypeChange(const FNiagaraTypeDefinition& InType, UEdGraphPin* Pin) const override;
 	virtual bool OnNewPinTypeRequested(UEdGraphPin* PinToChange, FNiagaraTypeDefinition NewType) override;
diff --git a/Engine/Plugins/FX/Niagara/Source/NiagaraEditor/Private/NiagaraNodeSimTargetSelector.cpp b/Engine/Plugins/FX/Niagara/Source/NiagaraEditor/Private/NiagaraNodeSimTargetSelector.cpp
--- a/Engine/Plugins/FX/Niagara/Source/NiagaraEditor/Private/NiagaraNodeSimTargetSelector.cpp
+++ b/Engine/Plugins/FX/Niagara/Source/NiagaraEditor/Private/NiagaraNodeSimTargetSelector.cpp
@@ -2,7 +2,7 @@
 
 #include "NiagaraNodeSimTargetSelector.h"
 #include "NiagaraEditorUtilities.h"
-#include "NiagaraHlslTranslator.h"
+#include "NiagaraGraphHlslTranslator.h"
 
 #include UE_INLINE_GENERATED_CPP_BY_NAME(NiagaraNodeSimTargetSelector)
 
@@ -47,12 +47,12 @@ FString UNiagaraNodeSimTargetSelector::GetInputCaseName(int32 Case) const
 	return Case == 0 ? TEXT("CPU VM") : TEXT("GPU Shader");
 }
 
-void UNiagaraNodeSimTargetSelector::Compile(class FHlslNiagaraTranslator* Translator, TArray<int32>& Outputs)
+void UNiagaraNodeSimTargetSelector::Compile(FTranslator* Translator, TArray<int32>& Outputs) const
 {
 	FPinCollectorArray InputPins;
-	GetInputPins(InputPins);
+	FTranslator::FBridge::GetCompilationInputPins(this, InputPins);
 	FPinCollectorArray OutputPins;
-	GetOutputPins(OutputPins);
+	FTranslator::FBridge::GetCompilationOutputPins(this, OutputPins);
 
 	//ENiagaraSimTarget SimulationTarget = Translator->GetSimulationTarget();
 	bool bCPUSim = Translator->IsCompileOptionDefined(*FNiagaraCompileOptions::CpuScriptDefine);
@@ -82,11 +82,9 @@ void UNiagaraNodeSimTargetSelector::Compile(class FHlslNiagaraTranslator* Transl
 	Outputs.SetNumUninitialized(OutputPins.Num());
 	for (int32 i = 0; i < OutputVars.Num(); i++)
 	{
-		int32 InputIdx = Translator->CompilePin(InputPins[VarIdx + i]);
+		int32 InputIdx = Translator->CompileInputPin(InputPins[VarIdx + i]);
 		Outputs[i] = InputIdx;
 	}
-	check(this->IsAddPin(OutputPins[OutputPins.Num() - 1]));
-	Outputs[OutputPins.Num() - 1] = INDEX_NONE;
 }
 
 UEdGraphPin* UNiagaraNodeSimTargetSelector::GetPassThroughPin(const UEdGraphPin* LocallyOwnedOutputPin, ENiagaraScriptUsage InUsage) const
diff --git a/Engine/Plugins/FX/Niagara/Source/NiagaraEditor/Private/NiagaraNodeStaticSwitch.cpp b/Engine/Plugins/FX/Niagara/Source/NiagaraEditor/Private/NiagaraNodeStaticSwitch.cpp
--- a/Engine/Plugins/FX/Niagara/Source/NiagaraEditor/Private/NiagaraNodeStaticSwitch.cpp
+++ b/Engine/Plugins/FX/Niagara/Source/NiagaraEditor/Private/NiagaraNodeStaticSwitch.cpp
@@ -2,7 +2,7 @@
 
 #include "NiagaraNodeStaticSwitch.h"
 #include "NiagaraEditorUtilities.h"
-#include "NiagaraHlslTranslator.h"
+#include "NiagaraGraphHlslTranslator.h"
 #include "NiagaraConstants.h"
 #include "NiagaraEditorStyle.h"
 #include "NiagaraScriptVariable.h"
@@ -316,7 +316,7 @@ bool UNiagaraNodeStaticSwitch::AllowNiagaraTypeForAddPin(const FNiagaraTypeDefin
 	return InType.GetScriptStruct() != nullptr && !InType.IsInternalType();
 }
 
-bool UNiagaraNodeStaticSwitch::GetVarIndex(FHlslNiagaraTranslator* Translator, int32 InputPinCount, int32& VarIndexOut) const
+bool UNiagaraNodeStaticSwitch::GetVarIndex(FTranslator* Translator, int32 InputPinCount, int32& VarIndexOut) const
 {	
 	return GetVarIndex(Translator, InputPinCount, SwitchValue, VarIndexOut);
 }
@@ -338,7 +338,7 @@ UEdGraphPin* UNiagaraNodeStaticSwitch::GetSelectorPin() const
 }
 
 
-void UNiagaraNodeStaticSwitch::UpdateCompilerConstantValue(FHlslNiagaraTranslator* Translator)
+void UNiagaraNodeStaticSwitch::UpdateCompilerConstantValue(FTranslator* Translator)
 {
 	if (IsSetByPin() && Translator)
 	{
@@ -409,7 +409,7 @@ void UNiagaraNodeStaticSwitch::UpdateCompilerConstantValue(FHlslNiagaraTranslato
 	}
 }
 
-bool UNiagaraNodeStaticSwitch::GetVarIndex(FHlslNiagaraTranslator* Translator, int32 InputPinCount, int32 Value, int32& VarIndexOut) const
+bool UNiagaraNodeStaticSwitch::GetVarIndex(FTranslator* Translator, int32 InputPinCount, int32 Value, int32& VarIndexOut) const
 {
 	bool Success = false;
 	if (SwitchTypeData.SwitchType == ENiagaraStaticSwitchType::Bool)
@@ -486,7 +486,7 @@ bool UNiagaraNodeStaticSwitch::GetVarIndex(FHlslNiagaraTranslator* Translator, i
 	return Success;
 }
 
-void UNiagaraNodeStaticSwitch::CheckForOutdatedEnum(FHlslNiagaraTranslator* Translator)
+void UNiagaraNodeStaticSwitch::CheckForOutdatedEnum(FTranslator* Translator)
 {
 	// we check during compilation if the node is using outdated enum information and log it
 	if(SwitchTypeData.SwitchType == ENiagaraStaticSwitchType::Enum)
@@ -552,7 +552,7 @@ void UNiagaraNodeStaticSwitch::CheckForOutdatedEnum(FHlslNiagaraTranslator* Tran
 	}
 }
 
-void UNiagaraNodeStaticSwitch::Compile(FHlslNiagaraTranslator* Translator, TArray<int32>& Outputs)
+void UNiagaraNodeStaticSwitch::Compile(FTranslator* Translator, TArray<int32>& Outputs) const
 {	
 	FPinCollectorArray InputPins;
 	GetInputPins(InputPins);
@@ -563,27 +563,19 @@ void UNiagaraNodeStaticSwitch::Compile(FHlslNiagaraTranslator* Translator, TArra
 		UEdGraphPin* SelectorPin = GetSelectorPin();	
 		if (SelectorPin)
 		{
-			Translator->CompilePin(SelectorPin);
+			Translator->CompileInputPin(SelectorPin);
 		}
 	}
 	FPinCollectorArray OutputPins;
-	GetOutputPins(OutputPins);
+	FTranslator::FBridge::GetCompilationOutputPins(this, OutputPins);
 
 	// Initialize the outputs to invalid values.
 	check(Outputs.Num() == 0);
-	Outputs.Reserve(OutputPins.Num());
-	for (int32 i = 0; i < OutputPins.Num(); i++)
-	{
-		Outputs.Add(INDEX_NONE);
-	}
+	Outputs.Init(INDEX_NONE, OutputPins.Num());
 
 	for (int i = 0; i < OutputPins.Num(); i++)
 	{
 		UEdGraphPin* OutPin = OutputPins[i];
-		if (IsAddPin(OutPin))
-		{
-			continue;
-		}
 		int32 VarIdx;
 		int NumChoices = IsSetByPin() ? InputPins.Num() - 1 : InputPins.Num();
 		if (GetVarIndex(nullptr, NumChoices, SelectorValue, VarIdx))
@@ -591,7 +583,7 @@ void UNiagaraNodeStaticSwitch::Compile(FHlslNiagaraTranslator* Translator, TArra
 			UEdGraphPin* InputPin = InputPins[VarIdx + i];
 			if (InputPin)
 			{
-				int32 CompiledInput = Translator->CompilePin(InputPin);
+				int32 CompiledInput = Translator->CompileInputPin(InputPin);
 				Outputs[i] = CompiledInput;
 			}
 		}		
@@ -644,7 +636,7 @@ void UNiagaraNodeStaticSwitch::ResolveNumerics(const UEdGraphSchema_Niagara* Sch
 		VarIdx++;
 	}
 }
-bool UNiagaraNodeStaticSwitch::SubstituteCompiledPin(FHlslNiagaraTranslator* Translator, UEdGraphPin** LocallyOwnedPin)
+bool UNiagaraNodeStaticSwitch::SubstituteCompiledPin(FTranslator* Translator, UEdGraphPin** LocallyOwnedPin)
 {
 	// if we compile the standalone module or function we don't have any valid input yet, so we just take the first option to satisfy the compiler
 	ENiagaraScriptUsage TargetUsage = Translator->GetTargetUsage();
diff --git a/Engine/Plugins/FX/Niagara/Source/NiagaraEditor/Private/NiagaraNodeStaticSwitch.h b/Engine/Plugins/FX/Niagara/Source/NiagaraEditor/Private/NiagaraNodeStaticSwitch.h
--- a/Engine/Plugins/FX/Niagara/Source/NiagaraEditor/Private/NiagaraNodeStaticSwitch.h
+++ b/Engine/Plugins/FX/Niagara/Source/NiagaraEditor/Private/NiagaraNodeStaticSwitch.h
@@ -74,9 +74,9 @@ public:
 	UEdGraphPin* GetSelectorPin() const;
 
 	/** This is a hack used in the translator to check for inconsistencies with old static switches before auto refresh was a thing */
-	void CheckForOutdatedEnum(FHlslNiagaraTranslator* Translator);
+	void CheckForOutdatedEnum(FTranslator* Translator);
 
-	void UpdateCompilerConstantValue(class FHlslNiagaraTranslator* Translator);
+	void UpdateCompilerConstantValue(FTranslator* Translator);
 
 	//~ Begin EdGraphNode Interface
 	virtual void AllocateDefaultPins() override;
@@ -88,8 +88,8 @@ public:
 	//~ End EdGraphNode Interface
 	
 	//~ Begin UNiagaraNode Interface
-	virtual void Compile(class FHlslNiagaraTranslator* Translator, TArray<int32>& Outputs) override;
-	virtual bool SubstituteCompiledPin(FHlslNiagaraTranslator* Translator, UEdGraphPin** LocallyOwnedPin) override;
+	virtual void Compile(FTranslator* Translator, TArray<int32>& Outputs) const override;
+	virtual bool SubstituteCompiledPin(FTranslator* Translator, UEdGraphPin** LocallyOwnedPin) override;
 	virtual UEdGraphPin* GetTracedOutputPin(UEdGraphPin* LocallyOwnedOutputPin, bool bFilterForCompilation, TArray<const UNiagaraNode*>* OutNodesVisitedDuringTrace = nullptr) const override;
 	virtual UEdGraphPin* GetTracedOutputPin(UEdGraphPin* LocallyOwnedOutputPin, bool bRecursive, bool bFilterForCompilation, TArray<const UNiagaraNode*>* OutNodesVisitedDuringTrace = nullptr) const;
 	virtual UEdGraphPin* GetPassThroughPin(const UEdGraphPin* LocallyOwnedOutputPin, ENiagaraScriptUsage InUsage) const override;
@@ -119,9 +119,9 @@ private:
 	virtual void PostChange(const UUserDefinedEnum* Changed, FEnumEditorUtils::EEnumEditorChangeInfo ChangedType) override;
 	
 	/** This finds the first valid input pin index for the current switch value, returns false if no value can be found */
-	bool GetVarIndex(class FHlslNiagaraTranslator* Translator, int32 InputPinCount, int32& VarIndexOut) const;
+	bool GetVarIndex(FTranslator* Translator, int32 InputPinCount, int32& VarIndexOut) const;
 
-	bool GetVarIndex(class FHlslNiagaraTranslator* Translator, int32 InputPinCount, int32 Value, int32& VarIndexOut) const;
+	bool GetVarIndex(FTranslator* Translator, int32 InputPinCount, int32 Value, int32& VarIndexOut) const;
 
 	void RemoveUnusedGraphParameter(const FNiagaraVariable& OldParameter);
 	
