Adding a templated layer in the hlsl translator to encapsulate logic surrounding graph implementation to allow us to share code between the existing UNiagaraGraph and the upcoming graph digest.  Gener

diff --git a/Engine/Plugins/FX/Niagara/Source/NiagaraEditor/Private/NiagaraNodeUsageSelector.cpp b/Engine/Plugins/FX/Niagara/Source/NiagaraEditor/Private/NiagaraNodeUsageSelector.cpp
--- a/Engine/Plugins/FX/Niagara/Source/NiagaraEditor/Private/NiagaraNodeUsageSelector.cpp
+++ b/Engine/Plugins/FX/Niagara/Source/NiagaraEditor/Private/NiagaraNodeUsageSelector.cpp
@@ -2,7 +2,8 @@
 
 #include "NiagaraNodeUsageSelector.h"
 #include "NiagaraEditorUtilities.h"
-#include "NiagaraHlslTranslator.h"
+#include "NiagaraGraphHlslTranslator.h"
+#include "NiagaraScriptSource.h"
 #include "Widgets/Layout/SSeparator.h"
 
 #include UE_INLINE_GENERATED_CPP_BY_NAME(NiagaraNodeUsageSelector)
@@ -294,12 +295,12 @@ void UNiagaraNodeUsageSelector::AddWidgetsToInputBox(TSharedPtr<SVerticalBox> In
 }
 
 
-void UNiagaraNodeUsageSelector::Compile(class FHlslNiagaraTranslator* Translator, TArray<int32>& Outputs)
+void UNiagaraNodeUsageSelector::Compile(FTranslator* Translator, TArray<int32>& Outputs) const
 {
 	FPinCollectorArray InputPins;
-	GetInputPins(InputPins);
+	FTranslator::FBridge::GetCompilationInputPins(this, InputPins);
 	FPinCollectorArray OutputPins;
-	GetOutputPins(OutputPins);
+	FTranslator::FBridge::GetCompilationOutputPins(this, OutputPins);
 
 	ENiagaraScriptUsage CurrentUsage = Translator->GetCurrentUsage();
 	ENiagaraScriptGroup UsageGroup = ENiagaraScriptGroup::Max;
@@ -319,11 +320,9 @@ void UNiagaraNodeUsageSelector::Compile(class FHlslNiagaraTranslator* Translator
 		Outputs.SetNumUninitialized(OutputPins.Num());
 		for (int32 i = 0; i < OutputVars.Num(); i++)
 		{
-			int32 InputIdx = Translator->CompilePin(InputPins[VarIdx + i]);
+			int32 InputIdx = Translator->CompileInputPin(InputPins[VarIdx + i]);
 			Outputs[i] = InputIdx;
 		}
-		check(this->IsAddPin(OutputPins[OutputPins.Num() - 1]));
-		Outputs[OutputPins.Num() - 1] = INDEX_NONE;
 	}
 	else
 	{
@@ -366,7 +365,7 @@ UEdGraphPin* UNiagaraNodeUsageSelector::GetPassThroughPin(const UEdGraphPin* Loc
 	return nullptr;
 }
 
-void UNiagaraNodeUsageSelector::AppendFunctionAliasForContext(const FNiagaraGraphFunctionAliasContext& InFunctionAliasContext, FString& InOutFunctionAlias, bool& OutOnlyOncePerNodeType)
+void UNiagaraNodeUsageSelector::AppendFunctionAliasForContext(const FNiagaraGraphFunctionAliasContext& InFunctionAliasContext, FString& InOutFunctionAlias, bool& OutOnlyOncePerNodeType) const
 {
 	OutOnlyOncePerNodeType = true;
 	
diff --git a/Engine/Plugins/FX/Niagara/Source/NiagaraEditor/Private/NiagaraNodeUsageSelector.h b/Engine/Plugins/FX/Niagara/Source/NiagaraEditor/Private/NiagaraNodeUsageSelector.h
--- a/Engine/Plugins/FX/Niagara/Source/NiagaraEditor/Private/NiagaraNodeUsageSelector.h
+++ b/Engine/Plugins/FX/Niagara/Source/NiagaraEditor/Private/NiagaraNodeUsageSelector.h
@@ -36,7 +36,7 @@ public:
 	//~ End EdGraphNode Interface
 
 	//~ Begin UNiagaraNode Interface
-	virtual void Compile(class FHlslNiagaraTranslator* Translator, TArray<int32>& Outputs) override;
+	virtual void Compile(FTranslator* Translator, TArray<int32>& Outputs) const override;
 	virtual bool RefreshFromExternalChanges() override;
 	virtual void AddWidgetsToInputBox(TSharedPtr<SVerticalBox> InputBox) override;
 	virtual void BuildParameterMapHistory(FNiagaraParameterMapHistoryBuilder& OutHistory, bool bRecursive = true, bool bFilterForCompilation = true) const override;
@@ -44,7 +44,7 @@ public:
 	virtual bool AllowNiagaraTypeForPinTypeChange(const FNiagaraTypeDefinition& InType, UEdGraphPin* Pin) const override;
 	virtual bool OnNewPinTypeRequested(UEdGraphPin* PinToChange, FNiagaraTypeDefinition NewType) override;
 	virtual UEdGraphPin* GetPassThroughPin(const UEdGraphPin* LocallyOwnedOutputPin, ENiagaraScriptUsage InUsage) const override;
-	virtual void AppendFunctionAliasForContext(const FNiagaraGraphFunctionAliasContext& InFunctionAliasContext, FString& InOutFunctionAlias, bool& OutOnlyOncePerNodeType) override;
+	virtual void AppendFunctionAliasForContext(const FNiagaraGraphFunctionAliasContext& InFunctionAliasContext, FString& InOutFunctionAlias, bool& OutOnlyOncePerNodeType) const override;
 	//~ End UNiagaraNode Interface
 
 	/** Helper function to create a variable to add to the OutputVars and FGuid to add to OutputVarGuids. */
diff --git a/Engine/Plugins/FX/Niagara/Source/NiagaraEditor/Private/NiagaraNodeWriteDataSet.cpp b/Engine/Plugins/FX/Niagara/Source/NiagaraEditor/Private/NiagaraNodeWriteDataSet.cpp
--- a/Engine/Plugins/FX/Niagara/Source/NiagaraEditor/Private/NiagaraNodeWriteDataSet.cpp
+++ b/Engine/Plugins/FX/Niagara/Source/NiagaraEditor/Private/NiagaraNodeWriteDataSet.cpp
@@ -2,7 +2,7 @@
 
 #include "NiagaraNodeWriteDataSet.h"
 #include "UObject/UnrealType.h"
-#include "NiagaraHlslTranslator.h"
+#include "NiagaraGraphHlslTranslator.h"
 #include "NiagaraEvents.h"
 #include "Widgets/SNiagaraGraphNodeWriteDataSet.h"
 #include "EdGraphSchema_Niagara.h"
@@ -75,8 +75,11 @@ FText UNiagaraNodeWriteDataSet::GetNodeTitle(ENodeTitleType::Type TitleType) con
 	return FText::Format(LOCTEXT("NiagaraDataSetWriteFormat", "{0} Write"), FText::FromName(DataSet.Name));
 }
 
-void UNiagaraNodeWriteDataSet::Compile(class FHlslNiagaraTranslator* Translator, TArray<int32>& Outputs)
+void UNiagaraNodeWriteDataSet::Compile(FTranslator* Translator, TArray<int32>& Outputs) const
 {
+	// this is dirty, but we'll continue doing this for now because it's all going to go away with Digested graphs
+	UNiagaraNodeWriteDataSet* MutableThis = const_cast<UNiagaraNodeWriteDataSet*>(this);
+
 	bool bError = false;
 
 	//TODO implement writing to data sets in hlsl compiler and vm.
@@ -105,8 +108,8 @@ void UNiagaraNodeWriteDataSet::Compile(class FHlslNiagaraTranslator* Translator,
 	
 	if (EventName.IsNone())
 	{
-		EventName = DataSet.Name;
-		VisualsChangedDelegate.Broadcast(this);
+		MutableThis->EventName = DataSet.Name;
+		VisualsChangedDelegate.Broadcast(MutableThis);
 	}
 
 	FNiagaraDataSetID AlteredDataSet = DataSet;
diff --git a/Engine/Plugins/FX/Niagara/Source/NiagaraEditor/Private/NiagaraParameterMapHistory.cpp b/Engine/Plugins/FX/Niagara/Source/NiagaraEditor/Private/NiagaraParameterMapHistory.cpp
--- a/Engine/Plugins/FX/Niagara/Source/NiagaraEditor/Private/NiagaraParameterMapHistory.cpp
+++ b/Engine/Plugins/FX/Niagara/Source/NiagaraEditor/Private/NiagaraParameterMapHistory.cpp
@@ -422,7 +422,7 @@ FName FNiagaraParameterMapHistory::ResolveEmitterAlias(const FName& InName, cons
 
 FString FNiagaraParameterMapHistory::MakeSafeNamespaceString(const FString& InStr)
 {
-	FString  Sanitized = FHlslNiagaraTranslator::GetSanitizedSymbolName(InStr);
+	FString  Sanitized = FNiagaraHlslTranslator::GetSanitizedSymbolName(InStr);
 	return Sanitized;
 }
 
@@ -1457,39 +1457,33 @@ bool FNiagaraParameterMapHistoryBuilder::InTopLevelFunctionCall(ENiagaraScriptUs
 }
 
 
-void FNiagaraParameterMapHistoryBuilder::EnterFunction(const FString& InNodeName, const UNiagaraScript* InScript, const UNiagaraGraph* InGraph, const UNiagaraNode* Node)
+void FNiagaraParameterMapHistoryBuilder::EnterFunction(const FString& InNodeName, const UNiagaraGraph* InGraph, const UNiagaraNode* Node)
 {
-	if (InScript != nullptr )
+	RegisterNodeVisitation(Node);
+	CallingContext.Push(Node);
+	ActivePath.PushNode(Node);
+	CallingGraphContext.Push(InGraph);
+	PinToParameterMapIndices.Emplace();
+	PinToConstantIndices.Emplace();
+	FunctionNameContextStack.Emplace(*InNodeName);
+	BuildCurrentAliases();
+	if (EncounteredFunctionNames.Num() != 0)
 	{
-		RegisterNodeVisitation(Node);
-		CallingContext.Push(Node);
-		ActivePath.PushNode(Node);
-		CallingGraphContext.Push(InGraph);
-		PinToParameterMapIndices.Emplace();
-		PinToConstantIndices.Emplace();
-		FunctionNameContextStack.Emplace(*InNodeName);
-		BuildCurrentAliases();
-		if (EncounteredFunctionNames.Num() != 0)
-		{
-			EncounteredFunctionNames.Last().AddUnique(InNodeName);
-		}
-		ContextuallyVisitedNodes.Emplace();
+		EncounteredFunctionNames.Last().AddUnique(InNodeName);
 	}
+	ContextuallyVisitedNodes.Emplace();
 }
 
-void FNiagaraParameterMapHistoryBuilder::ExitFunction(const FString& InNodeName, const UNiagaraScript* InScript, const UNiagaraNode* Node)
+void FNiagaraParameterMapHistoryBuilder::ExitFunction(const FString& InNodeName, const UNiagaraNode* Node)
 {
-	if (InScript != nullptr)
-	{
-		CallingContext.Pop();
-		CallingGraphContext.Pop();
-		ActivePath.PopNode();
-		PinToParameterMapIndices.Pop();
-		FunctionNameContextStack.Pop();
-		PinToConstantIndices.Pop();
-		BuildCurrentAliases();
-		ContextuallyVisitedNodes.Pop();
-	}
+	CallingContext.Pop();
+	CallingGraphContext.Pop();
+	ActivePath.PopNode();
+	PinToParameterMapIndices.Pop();
+	FunctionNameContextStack.Pop();
+	PinToConstantIndices.Pop();
+	BuildCurrentAliases();
+	ContextuallyVisitedNodes.Pop();
 }
 
 void FNiagaraParameterMapHistoryBuilder::EnterEmitter(const FString& InEmitterName, const UNiagaraGraph* InGraph, const UNiagaraNode* Node)
diff --git a/Engine/Plugins/FX/Niagara/Source/NiagaraEditor/Private/ViewModels/HierarchyEditor/NiagaraHierarchyViewModelBase.cpp b/Engine/Plugins/FX/Niagara/Source/NiagaraEditor/Private/ViewModels/HierarchyEditor/NiagaraHierarchyViewModelBase.cpp
--- a/Engine/Plugins/FX/Niagara/Source/NiagaraEditor/Private/ViewModels/HierarchyEditor/NiagaraHierarchyViewModelBase.cpp
+++ b/Engine/Plugins/FX/Niagara/Source/NiagaraEditor/Private/ViewModels/HierarchyEditor/NiagaraHierarchyViewModelBase.cpp
@@ -17,6 +17,7 @@
 #include "IPropertyRowGenerator.h"
 #include "ScopedTransaction.h"
 #include "Framework/Commands/GenericCommands.h"
+#include "Framework/Notifications/NotificationManager.h"
 #include "ToolMenus.h"
 #include "Widgets/Notifications/SNotificationList.h"
 
diff --git a/Engine/Plugins/FX/Niagara/Source/NiagaraEditor/Private/ViewModels/NiagaraScriptGraphViewModel.cpp b/Engine/Plugins/FX/Niagara/Source/NiagaraEditor/Private/ViewModels/NiagaraScriptGraphViewModel.cpp
--- a/Engine/Plugins/FX/Niagara/Source/NiagaraEditor/Private/ViewModels/NiagaraScriptGraphViewModel.cpp
+++ b/Engine/Plugins/FX/Niagara/Source/NiagaraEditor/Private/ViewModels/NiagaraScriptGraphViewModel.cpp
@@ -8,6 +8,7 @@
 #include "NiagaraNodeInput.h"
 
 #include "Framework/Commands/GenericCommands.h"
+#include "Framework/Notifications/NotificationManager.h"
 #include "ScopedTransaction.h"
 #include "EdGraphUtilities.h"
 #include "Editor.h"
diff --git a/Engine/Plugins/FX/Niagara/Source/NiagaraEditor/Private/ViewModels/Stack/NiagaraStackGraphUtilities.cpp b/Engine/Plugins/FX/Niagara/Source/NiagaraEditor/Private/ViewModels/Stack/NiagaraStackGraphUtilities.cpp
--- a/Engine/Plugins/FX/Niagara/Source/NiagaraEditor/Private/ViewModels/Stack/NiagaraStackGraphUtilities.cpp
+++ b/Engine/Plugins/FX/Niagara/Source/NiagaraEditor/Private/ViewModels/Stack/NiagaraStackGraphUtilities.cpp
@@ -1664,7 +1664,7 @@ TSet<FGuid> FNiagaraStackGraphUtilities::GetHiddenChildrenInputGuids(const UNiag
 	return Result;
 }
 
-TArray<UEdGraphPin*> FNiagaraStackGraphUtilities::GetUnusedFunctionInputPins(UNiagaraNodeFunctionCall& FunctionCallNode, FCompileConstantResolver ConstantResolver)
+TArray<UEdGraphPin*> FNiagaraStackGraphUtilities::GetUnusedFunctionInputPins(const UNiagaraNodeFunctionCall& FunctionCallNode, FCompileConstantResolver ConstantResolver)
 {
 	UNiagaraGraph* FunctionGraph = FunctionCallNode.GetCalledGraph();
 	if (!FunctionGraph || FunctionCallNode.FunctionScript->Usage != ENiagaraScriptUsage::Module)
diff --git a/Engine/Plugins/FX/Niagara/Source/NiagaraEditor/Private/Widgets/SNiagaraGraphPinAdd.cpp b/Engine/Plugins/FX/Niagara/Source/NiagaraEditor/Private/Widgets/SNiagaraGraphPinAdd.cpp
--- a/Engine/Plugins/FX/Niagara/Source/NiagaraEditor/Private/Widgets/SNiagaraGraphPinAdd.cpp
+++ b/Engine/Plugins/FX/Niagara/Source/NiagaraEditor/Private/Widgets/SNiagaraGraphPinAdd.cpp
@@ -8,6 +8,7 @@
 #include "NiagaraNodeParameterMapBase.h"
 #include "NiagaraNodeWithDynamicPins.h"
 #include "NiagaraScript.h"
+#include "NiagaraScriptSource.h"
 #include "ViewModels/NiagaraScriptViewModel.h"
 #include "ViewModels/TNiagaraViewModelManager.h"
 #include "Widgets/SNiagaraParameterMenu.h"
diff --git a/Engine/Plugins/FX/Niagara/Source/NiagaraEditor/Public/INiagaraCompiler.h b/Engine/Plugins/FX/Niagara/Source/NiagaraEditor/Public/INiagaraCompiler.h
--- a/Engine/Plugins/FX/Niagara/Source/NiagaraEditor/Public/INiagaraCompiler.h
+++ b/Engine/Plugins/FX/Niagara/Source/NiagaraEditor/Public/INiagaraCompiler.h
@@ -4,9 +4,11 @@
 
 #include "CoreMinimal.h"
 #include "Containers/StringView.h"
+#include "NiagaraCommon.h"
 #include "NiagaraEditorCommon.h"
 #include "NiagaraParameters.h"
 #include "NiagaraScript.h"
+#include "NiagaraShared.h"
 
 class Error;
 class UEdGraphPin;
@@ -19,6 +21,11 @@ struct FNiagaraTranslateResults;
 struct FNiagaraTranslatorOutput;
 struct FNiagaraVMExecutableData;
 class FNiagaraCompileOptions;
+class FNiagaraCompileRequestData;
+class FNiagaraCompileRequestDuplicateData;
+class FNiagaraPrecompileData;
+class FNiagaraCompilationCopyData;
+
 
 /** Defines information about the results of a Niagara script compile. */
 struct FNiagaraCompileResults
@@ -64,7 +71,7 @@ public:
 	virtual int32 CompileScript(const class FNiagaraCompileRequestData* InCompileRequest, const FNiagaraCompileOptions& InOptions, const FNiagaraTranslateResults& InTranslateResults, FNiagaraTranslatorOutput* TranslatorOutput, FString& TranslatedHLSL) = 0;
 
 	/** Starts the async compilation of a script and returns the job handle to retrieve the results */
-	virtual int32 CompileScript(const FStringView GroupName, const FNiagaraCompileOptions& InOptions, const FNiagaraTranslateResults& InTranslateResults, FNiagaraTranslatorOutput* TranslatorOutput, FString& TranslatedHLSL) = 0;
+	virtual int32 CompileScript(const FStringView GroupName, const FNiagaraCompileOptions& InOptions, const FNiagaraTranslateResults& InTranslateResults, const FNiagaraTranslatorOutput& TranslatorOutput, const FString& TranslatedHLSL) = 0;
 
 	/** Returns the compile result for a given job id once the job has finished compiling. */
 	virtual TOptional<FNiagaraCompileResults> GetCompileResult(int32 JobID, bool bWait = false) = 0;
@@ -75,3 +82,63 @@ public:
 	/** Adds a warning to be reported to the user. Warnings will not cause a compilation failure. */
 	virtual void Warning(FText WarningText) = 0;
 };
+
+/** Defines information about the results of a Niagara script compile. */
+struct FNiagaraTranslateResults
+{
+	/** Whether or not HLSL generation was successful */
+	bool bHLSLGenSucceeded = false;;
+
+	/** A results log with messages, warnings, and errors which occurred during the compile. */
+	TArray<FNiagaraCompileEvent> CompileEvents;
+	uint32 NumErrors = 0;
+	uint32 NumWarnings = 0;
+
+	TArray<FNiagaraCompileDependency> CompileDependencies;
+
+	TArray<FNiagaraCompilerTag > CompileTags;
+
+	/** A string representation of the compilation output. */
+	FString OutputHLSL;
+
+	static ENiagaraScriptCompileStatus TranslateResultsToSummary(const FNiagaraTranslateResults* CompileResults);
+};
+
+class NIAGARAEDITOR_API FHlslNiagaraTranslatorOptions
+{
+public:
+	ENiagaraSimTarget SimTarget = ENiagaraSimTarget::CPUSim;
+
+	/** Any parameters in these namespaces will be pulled from an "InstanceParameters" dataset rather than from the uniform table. */
+	TArray<FString> InstanceParameterNamespaces;
+
+	/** Whether or not to treat top-level module variables as external values for rapid iteration without need for compilation.*/
+	bool bParameterRapidIteration = true;
+
+	/** Should we disable debug switches during translation. */
+	bool bDisableDebugSwitches = false;
+
+	/** Whether or not to override top-level module variables with values from the constant override table. This is only used for variables that were candidates for rapid iteration.*/
+	TArray<FNiagaraVariable> OverrideModuleConstants;
+};
+
+class NIAGARAEDITOR_API INiagaraHlslTranslator
+{
+public:
+	virtual ~INiagaraHlslTranslator() = default;
+
+	virtual FNiagaraTranslateResults Translate(const FNiagaraCompileOptions& InCompileOptions, const FHlslNiagaraTranslatorOptions& Options) = 0;
+	virtual const FNiagaraTranslatorOutput& GetTranslateOutput() const = 0;
+	virtual const FString& GetTranslatedHLSL() const = 0;
+
+	static TUniquePtr<INiagaraHlslTranslator> CreateTranslator(const FNiagaraCompileRequestDataBase* InCompileData, const FNiagaraCompileRequestDuplicateDataBase* InDuplicateData);
+	static TUniquePtr<INiagaraHlslTranslator> CreateTranslator(const FNiagaraPrecompileData* InCompileData, const FNiagaraCompilationCopyData* InDuplicateData);
+
+	static FString BuildHLSLStructDecl(const FNiagaraTypeDefinition& Type, FText& OutErrorMessage, bool bGpuScript);
+	static bool IsBuiltInHlslType(const FNiagaraTypeDefinition& Type);
+	static FString GetStructHlslTypeName(const FNiagaraTypeDefinition& Type);
+	static FString GetPropertyHlslTypeName(const FProperty* Property);
+	static bool IsHlslBuiltinVector(const FNiagaraTypeDefinition& Type);
+	static TArray<FName> ConditionPropertyPath(const FNiagaraTypeDefinition& Type, const TArray<FName>& InPath);
+	static FString GetHlslDefaultForType(const FNiagaraTypeDefinition& Type);
+};
\ No newline at end of file
diff --git a/Engine/Plugins/FX/Niagara/Source/NiagaraEditor/Public/NiagaraGraph.h b/Engine/Plugins/FX/Niagara/Source/NiagaraEditor/Public/NiagaraGraph.h
--- a/Engine/Plugins/FX/Niagara/Source/NiagaraEditor/Public/NiagaraGraph.h
+++ b/Engine/Plugins/FX/Niagara/Source/NiagaraEditor/Public/NiagaraGraph.h
@@ -137,7 +137,7 @@ struct FNiagaraGraphFunctionAliasContext
 
 	// the usage as defined in the current translation stage
 	ENiagaraScriptUsage ScriptUsage;
-	TArray<UEdGraphPin*> StaticSwitchValues;
+	TArray<const UEdGraphPin*> StaticSwitchValues;
 };
 
 UCLASS(MinimalAPI)
@@ -369,7 +369,7 @@ class UNiagaraGraph : public UEdGraph
 	  * function calls to the same graph. For example, if the graph contains static switches and two functions call it with
 	  * different switch parameters, the final function names in the hlsl must be different.
 	  */
-	FString GetFunctionAliasByContext(const FNiagaraGraphFunctionAliasContext& FunctionAliasContext);
+	FString GetFunctionAliasByContext(const FNiagaraGraphFunctionAliasContext& FunctionAliasContext) const;
 
 	/** In order to support reducing the number of times we need to fully generate the CompileId this function is introduced
 	  * to work with LastBuiltScriptVersionId & bHasValidLastBuiltScriptVersionId.
diff --git a/Engine/Plugins/FX/Niagara/Source/NiagaraEditor/Public/NiagaraGraphDataCache.h b/Engine/Plugins/FX/Niagara/Source/NiagaraEditor/Public/NiagaraGraphDataCache.h
--- a/Engine/Plugins/FX/Niagara/Source/NiagaraEditor/Public/NiagaraGraphDataCache.h
+++ b/Engine/Plugins/FX/Niagara/Source/NiagaraEditor/Public/NiagaraGraphDataCache.h
@@ -9,7 +9,6 @@
 #include "UObject/WeakObjectPtrTemplates.h"
 #include "ViewModels/Stack/NiagaraStackGraphUtilities.h"
 
-class FHlslNiagaraTranslator;
 class UEdGraphPin;
 class UNiagaraGraph;
 class UNiagaraNodeFunctionCall;
diff --git a/Engine/Plugins/FX/Niagara/Source/NiagaraEditor/Public/NiagaraNode.h b/Engine/Plugins/FX/Niagara/Source/NiagaraEditor/Public/NiagaraNode.h
--- a/Engine/Plugins/FX/Niagara/Source/NiagaraEditor/Public/NiagaraNode.h
+++ b/Engine/Plugins/FX/Niagara/Source/NiagaraEditor/Public/NiagaraNode.h
@@ -3,7 +3,6 @@
 #pragma once
 
 #include "CoreMinimal.h"
-#include "../Private/NiagaraHlslTranslator.h"
 #include "UObject/ObjectMacros.h"
 #include "EdGraph/EdGraphNode.h"
 #include "NiagaraEditorCommon.h"
@@ -15,10 +14,13 @@
 #include "Widgets/SBoxPanel.h"
 #include "NiagaraNode.generated.h"
 
+struct FNiagaraCompilationGraphBridge;
 class UEdGraphPin;
+class UEdGraphSchema_Niagara;
 class INiagaraCompiler;
 struct FNiagaraGraphFunctionAliasContext;
 class FSHA1;
+template<typename T> class TNiagaraHlslTranslator;
 
 typedef TArray<UEdGraphPin*, TInlineAllocator<16>> FPinCollectorArray;
 
@@ -28,9 +30,11 @@ class NIAGARAEDITOR_API UNiagaraNode : public UEdGraphNode
 	GENERATED_UCLASS_BODY()
 protected:
 
+	using FTranslator = TNiagaraHlslTranslator<FNiagaraCompilationGraphBridge>;
+
 	bool ReallocatePins(bool bMarkNeedsResynchronizeOnChange = true);
 
-	bool CompileInputPins(class FHlslNiagaraTranslator *Translator, TArray<int32>& OutCompiledInputs);
+	bool CompileInputPins(FTranslator* Translator, TArray<int32>& OutCompiledInputs) const;
 
 	virtual void OnPostSynchronizationInReallocatePins() {}
 	
@@ -75,7 +79,7 @@ public:
 	/** Refreshes the node due to external changes, e.g. the underlying function changed for a function call node. Return true if the graph changed.*/
 	virtual bool RefreshFromExternalChanges() { return false; }
 
-	virtual void Compile(class FHlslNiagaraTranslator *Translator, TArray<int32>& Outputs);
+	virtual void Compile(FTranslator* Translator, TArray<int32>& Outputs) const;
 
 	UEdGraphPin* GetInputPin(int32 InputIndex) const;
 	UEdGraphPin* GetOutputPin(int32 OutputIndex) const;
@@ -92,7 +96,7 @@ public:
 	}
 
 	UEdGraphPin* GetPinByPersistentGuid(const FGuid& InGuid) const;
-	virtual void ResolveNumerics(const UEdGraphSchema_Niagara* Schema, bool bSetInline, TMap<TPair<FGuid, UEdGraphNode*>, FNiagaraTypeDefinition>* PinCache);
+	virtual void ResolveNumerics(const UEdGraphSchema_Niagara* Schema, bool bSetInline, TMap<TPair<FGuid, class UEdGraphNode*>, FNiagaraTypeDefinition>* PinCache);
 
 	/** Apply any node-specific logic to determine if it is safe to add this node to the graph. This is meant to be called only in the Editor before placing the node.*/
 	virtual bool CanAddToGraph(UNiagaraGraph* TargetGraph, FString& OutErrorMsg) const;
@@ -161,7 +165,7 @@ public:
 	static UEdGraphPin* TraceOutputPin(UEdGraphPin* LocallyOwnedOutputPin, bool bFilterForCompilation = true, TArray<const UNiagaraNode*>* OutNodesVisitedDuringTrace = nullptr);
 
 	/** Allows a node to replace a pin that is about to be compiled with another pin. This can be used for either optimizations or features such as the static switch. Returns true if the pin was successfully replaced, false otherwise. */
-	virtual bool SubstituteCompiledPin(FHlslNiagaraTranslator* Translator, UEdGraphPin** LocallyOwnedPin);
+	virtual bool SubstituteCompiledPin(FTranslator* Translator, UEdGraphPin** LocallyOwnedPin);
 
 	virtual UEdGraphPin* GetPassThroughPin(const UEdGraphPin* LocallyOwnedOutputPin) const override { return nullptr; }
 	virtual UEdGraphPin* GetPassThroughPin(const UEdGraphPin* LocallyOwnedOutputPin, ENiagaraScriptUsage InUsage) const { return nullptr; }
@@ -177,7 +181,7 @@ public:
 
 	FOnNodeVisualsChanged& OnVisualsChanged();
 
-	virtual void AppendFunctionAliasForContext(const FNiagaraGraphFunctionAliasContext& InFunctionAliasContext, FString& InOutFunctionAlias, bool& OutOnlyOncePerNodeType) { };
+	virtual void AppendFunctionAliasForContext(const FNiagaraGraphFunctionAliasContext& InFunctionAliasContext, FString& InOutFunctionAlias, bool& OutOnlyOncePerNodeType) const { };
 
 	/** Old style compile hash code. To be removed in the future.*/
 	virtual void UpdateCompileHashForNode(FSHA1& HashState) const;
@@ -206,7 +210,7 @@ protected:
 	/** For a simple Plain old data type UProperty, hash the data.*/
 	virtual bool PODPropertyAppendCompileHash(const void* Container, FProperty* Property, const FString& PropertyName, FNiagaraCompileHashVisitor* InVisitor) const;
 
-	virtual int32 CompileInputPin(class FHlslNiagaraTranslator *Translator, UEdGraphPin* Pin);
+	virtual int32 CompileInputPin(FTranslator* Translator, UEdGraphPin* Pin) const;
 	virtual bool IsValidPinToCompile(UEdGraphPin* Pin) const;
 
 	void NumericResolutionByPins(const UEdGraphSchema_Niagara* Schema, TArrayView<UEdGraphPin* const> InputPins, TArrayView<UEdGraphPin* const> OutputPins,
diff --git a/Engine/Plugins/FX/Niagara/Source/NiagaraEditor/Public/NiagaraNodeDataSetBase.h b/Engine/Plugins/FX/Niagara/Source/NiagaraEditor/Public/NiagaraNodeDataSetBase.h
--- a/Engine/Plugins/FX/Niagara/Source/NiagaraEditor/Public/NiagaraNodeDataSetBase.h
+++ b/Engine/Plugins/FX/Niagara/Source/NiagaraEditor/Public/NiagaraNodeDataSetBase.h
@@ -47,9 +47,9 @@ protected:
 	void AddParameterMapPins();
 	bool InitializeFromStructInternal(const UStruct* PayloadStruct);
 
-	bool IsSynchronizedWithStruct(bool bIgnoreConditionVar, FString* Issues, bool bLogIssues = true);
+	bool IsSynchronizedWithStruct(bool bIgnoreConditionVar, FString* Issues, bool bLogIssues = true) const;
 	virtual bool SynchronizeWithStruct();
-	bool GetSupportedNiagaraTypeDef(const FProperty* Property, FNiagaraTypeDefinition& TypeDef);
+	bool GetSupportedNiagaraTypeDef(const FProperty* Property, FNiagaraTypeDefinition& TypeDef) const;
 
 	static const FName ConditionVarName;
 	static const FName ParamMapInVarName;
diff --git a/Engine/Plugins/FX/Niagara/Source/NiagaraEditor/Public/NiagaraNodeFunctionCall.h b/Engine/Plugins/FX/Niagara/Source/NiagaraEditor/Public/NiagaraNodeFunctionCall.h
--- a/Engine/Plugins/FX/Niagara/Source/NiagaraEditor/Public/NiagaraNodeFunctionCall.h
+++ b/Engine/Plugins/FX/Niagara/Source/NiagaraEditor/Public/NiagaraNodeFunctionCall.h
@@ -106,7 +106,7 @@ public:
 	//End UObject interface
 
 	//~ Begin UNiagaraNode Interface
-	virtual void Compile(class FHlslNiagaraTranslator* Translator, TArray<int32>& Outputs) override;
+	virtual void Compile(FTranslator* Translator, TArray<int32>& Outputs) const override;
 	virtual UObject* GetReferencedAsset() const override;
 	virtual void OpenReferencedAsset() const override;
 	virtual bool RefreshFromExternalChanges() override;
@@ -137,7 +137,7 @@ public:
 	/** When overriding an input value, this updates which variable guid was bound to which input name, so it can be reassigned when the input is renamed.*/
 	void UpdateInputNameBinding(const FGuid& BoundVariableGuid, const FName& BoundName);
 
-	bool FindAutoBoundInput(UNiagaraNodeInput* InputNode, UEdGraphPin* PinToAutoBind, FNiagaraVariable& OutFoundVar, ENiagaraInputNodeUsage& OutNodeUsage);
+	bool FindAutoBoundInput(const UNiagaraNodeInput* InputNode, const UEdGraphPin* PinToAutoBind, FNiagaraVariable& OutFoundVar, ENiagaraInputNodeUsage& OutNodeUsage) const;
 
 	virtual void BuildParameterMapHistory(FNiagaraParameterMapHistoryBuilder& OutHistory, bool bRecursive = true, bool bFilterForCompilation = true) const override;
 
diff --git a/Engine/Plugins/FX/Niagara/Source/NiagaraEditor/Public/NiagaraNodeInput.h b/Engine/Plugins/FX/Niagara/Source/NiagaraEditor/Public/NiagaraNodeInput.h
--- a/Engine/Plugins/FX/Niagara/Source/NiagaraEditor/Public/NiagaraNodeInput.h
+++ b/Engine/Plugins/FX/Niagara/Source/NiagaraEditor/Public/NiagaraNodeInput.h
@@ -83,7 +83,7 @@ public:
 	/** Notifies the node that the exposure options have changed externally.*/
 	void NotifyExposureOptionsChanged();
 
-	virtual void Compile(class FHlslNiagaraTranslator* Translator, TArray<int32>& Outputs)override;
+	virtual void Compile(FTranslator* Translator, TArray<int32>& Outputs) const override;
 	
 	bool IsExposed()const { return ExposureOptions.bExposed && Usage == ENiagaraInputNodeUsage::Parameter; }
 	bool IsRequired()const { return IsExposed() && ExposureOptions.bRequired; }
@@ -93,7 +93,7 @@ public:
 	bool ReferencesSameInput(UNiagaraNodeInput* Other) const;
 
 	virtual void BuildParameterMapHistory(FNiagaraParameterMapHistoryBuilder& OutHistory, bool bRecursive = true, bool bFilterForCompilation = true) const override;
-	virtual void AppendFunctionAliasForContext(const FNiagaraGraphFunctionAliasContext& InFunctionAliasContext, FString& InOutFunctionAlias, bool& OutOnlyOncePerNodeType) override;
+	virtual void AppendFunctionAliasForContext(const FNiagaraGraphFunctionAliasContext& InFunctionAliasContext, FString& InOutFunctionAlias, bool& OutOnlyOncePerNodeType) const override;
 
 	static const FLinearColor TitleColor_Attribute;
 	static const FLinearColor TitleColor_Constant;
diff --git a/Engine/Plugins/FX/Niagara/Source/NiagaraEditor/Public/NiagaraNodeOutput.h b/Engine/Plugins/FX/Niagara/Source/NiagaraEditor/Public/NiagaraNodeOutput.h
--- a/Engine/Plugins/FX/Niagara/Source/NiagaraEditor/Public/NiagaraNodeOutput.h
+++ b/Engine/Plugins/FX/Niagara/Source/NiagaraEditor/Public/NiagaraNodeOutput.h
@@ -50,7 +50,7 @@ public:
 	/** Notifies the node that it's output variables have been modified externally. */
 	void NotifyOutputVariablesChanged();
 
-	virtual void Compile(class FHlslNiagaraTranslator *Translator, TArray<int32>& OutputExpressions) override;
+	virtual void Compile(FTranslator* Translator, TArray<int32>& OutputExpressions) const override;
 	const TArray<FNiagaraVariable>& GetOutputs() const {return Outputs;}
 	
 	/** Gets the usage of this graph root. */
@@ -67,7 +67,7 @@ public:
 	TArray<FName> GetAllStackContextOverrides() const;
 
 protected:
-	virtual int32 CompileInputPin(class FHlslNiagaraTranslator *Translator, UEdGraphPin* Pin) override;
+	virtual int32 CompileInputPin(FTranslator* Translator, UEdGraphPin* Pin) const override;
 
 	/** Removes a pin from this node with a transaction. */
 	virtual void RemoveOutputPin(UEdGraphPin* Pin);
diff --git a/Engine/Plugins/FX/Niagara/Source/NiagaraEditor/Public/NiagaraNodeReadDataSet.h b/Engine/Plugins/FX/Niagara/Source/NiagaraEditor/Public/NiagaraNodeReadDataSet.h
--- a/Engine/Plugins/FX/Niagara/Source/NiagaraEditor/Public/NiagaraNodeReadDataSet.h
+++ b/Engine/Plugins/FX/Niagara/Source/NiagaraEditor/Public/NiagaraNodeReadDataSet.h
@@ -25,7 +25,7 @@ public:
 	virtual bool CanAddToGraph(UNiagaraGraph* TargetGraph, FString& OutErrorMsg) const override;
 
 	//~End UNiagaraGraph Interface
-	virtual void Compile(class FHlslNiagaraTranslator* Translator, TArray<int32>& Outputs)override;
+	virtual void Compile(FTranslator* Translator, TArray<int32>& Outputs) const override;
 
 	virtual void BuildParameterMapHistory(FNiagaraParameterMapHistoryBuilder& OutHistory, bool bRecursive = true, bool bFilterForCompilation = true) const override;
 };
diff --git a/Engine/Plugins/FX/Niagara/Source/NiagaraEditor/Public/NiagaraNodeWriteDataSet.h b/Engine/Plugins/FX/Niagara/Source/NiagaraEditor/Public/NiagaraNodeWriteDataSet.h
--- a/Engine/Plugins/FX/Niagara/Source/NiagaraEditor/Public/NiagaraNodeWriteDataSet.h
+++ b/Engine/Plugins/FX/Niagara/Source/NiagaraEditor/Public/NiagaraNodeWriteDataSet.h
@@ -32,7 +32,7 @@ public:
 	virtual FText GetNodeTitle(ENodeTitleType::Type TitleType) const override;
 	//~ End EdGraphNode Interface
 
-	virtual void Compile(class FHlslNiagaraTranslator* Translator, TArray<int32>& Outputs)override;
+	virtual void Compile(FTranslator* Translator, TArray<int32>& Outputs) const override;
 
 	virtual void BuildParameterMapHistory(FNiagaraParameterMapHistoryBuilder& OutHistory, bool bRecursive = true, bool bFilterForCompilation = true) const override;
 	virtual TSharedPtr<SGraphNode> CreateVisualWidget() override;
diff --git a/Engine/Plugins/FX/Niagara/Source/NiagaraEditor/Public/NiagaraParameterMapHistory.h b/Engine/Plugins/FX/Niagara/Source/NiagaraEditor/Public/NiagaraParameterMapHistory.h
--- a/Engine/Plugins/FX/Niagara/Source/NiagaraEditor/Public/NiagaraParameterMapHistory.h
+++ b/Engine/Plugins/FX/Niagara/Source/NiagaraEditor/Public/NiagaraParameterMapHistory.h
@@ -14,7 +14,7 @@ class UEdGraphPin;
 class UEdGraphNode;
 class UNiagaraGraph;
 class UNiagaraParameterCollection;
-class FHlslNiagaraTranslator;
+class FNiagaraHlslTranslator;
 class UNiagaraGraph;
 
 class FCompileConstantResolver
@@ -23,11 +23,11 @@ public:
 	FCompileConstantResolver() : Emitter(FVersionedNiagaraEmitter()), System(nullptr), Translator(nullptr), Usage(ENiagaraScriptUsage::Function), DebugState(ENiagaraFunctionDebugState::NoDebug) {}
 	FCompileConstantResolver(const FVersionedNiagaraEmitter& Emitter, ENiagaraScriptUsage Usage, ENiagaraFunctionDebugState DebugState = ENiagaraFunctionDebugState::NoDebug) : Emitter(Emitter), System(nullptr), Translator(nullptr), Usage(Usage), DebugState(DebugState) {}
 	FCompileConstantResolver(const UNiagaraSystem* System, ENiagaraScriptUsage Usage, ENiagaraFunctionDebugState DebugState = ENiagaraFunctionDebugState::NoDebug) : Emitter(FVersionedNiagaraEmitter()), System(System), Translator(nullptr), Usage(Usage), DebugState(DebugState) {}
-	FCompileConstantResolver(const FHlslNiagaraTranslator* Translator, ENiagaraFunctionDebugState DebugState = ENiagaraFunctionDebugState::NoDebug) : Emitter(FVersionedNiagaraEmitter()), System(nullptr), Translator(Translator), Usage(ENiagaraScriptUsage::Function), DebugState(DebugState) {}
+	FCompileConstantResolver(const FNiagaraHlslTranslator* Translator, ENiagaraFunctionDebugState DebugState = ENiagaraFunctionDebugState::NoDebug) : Emitter(FVersionedNiagaraEmitter()), System(nullptr), Translator(Translator), Usage(ENiagaraScriptUsage::Function), DebugState(DebugState) {}
 
 	const UNiagaraEmitter* GetEmitter() const { return Emitter.Emitter; }
 	const UNiagaraSystem* GetSystem() const { return System; }
-	const FHlslNiagaraTranslator* GetTranslator() const { return Translator; }
+	const FNiagaraHlslTranslator* GetTranslator() const { return Translator; }
 	ENiagaraScriptUsage GetUsage() const { return Usage; }
 	ENiagaraFunctionDebugState GetDebugState() const { return DebugState; }
 
@@ -45,7 +45,7 @@ public:
 private:
 	FVersionedNiagaraEmitter Emitter;
 	const UNiagaraSystem* System;
-	const FHlslNiagaraTranslator* Translator;
+	const FNiagaraHlslTranslator* Translator;
 	ENiagaraScriptUsage Usage;
 	ENiagaraFunctionDebugState DebugState;
 };
@@ -430,12 +430,12 @@ public:
 	/**
 	* Record that we have entered a new function scope.
 	*/
-	void EnterFunction(const FString& InNodeName, const class UNiagaraScript* InScript, const class UNiagaraGraph* InGraph, const class UNiagaraNode* Node);
+	void EnterFunction(const FString& InNodeName, const class UNiagaraGraph* InGraph, const class UNiagaraNode* Node);
 
 	/**
 	* Record that we have exited a function scope.
 	*/
-	void ExitFunction(const FString& InNodeName, const class UNiagaraScript* InScript, const class UNiagaraNode* Node);
+	void ExitFunction(const FString& InNodeName, const class UNiagaraNode* Node);
 
 	/**
 	* Record that we have entered an emitter scope.
diff --git a/Engine/Plugins/FX/Niagara/Source/NiagaraEditor/Public/ViewModels/Stack/NiagaraStackGraphUtilities.h b/Engine/Plugins/FX/Niagara/Source/NiagaraEditor/Public/ViewModels/Stack/NiagaraStackGraphUtilities.h
--- a/Engine/Plugins/FX/Niagara/Source/NiagaraEditor/Public/ViewModels/Stack/NiagaraStackGraphUtilities.h
+++ b/Engine/Plugins/FX/Niagara/Source/NiagaraEditor/Public/ViewModels/Stack/NiagaraStackGraphUtilities.h
@@ -188,7 +188,7 @@ namespace FNiagaraStackGraphUtilities
 	TSet<FGuid> GetHiddenChildrenInputGuids(const UNiagaraNodeFunctionCall& FunctionCallNode, FName ParentInput, TSharedRef<FNiagaraEmitterViewModel> EmitterViewModel);
 	
 	/* Module script calls do not have direct inputs, but rely on the parameter map being initialized correctly. This utility function resolves which of the module's parameters are reachable during compilation and returns a list of pins on the parameter map node that do not have to be compiled. */
-	TArray<UEdGraphPin*> GetUnusedFunctionInputPins(UNiagaraNodeFunctionCall& FunctionCallNode, FCompileConstantResolver ConstantResolver);
+	TArray<UEdGraphPin*> GetUnusedFunctionInputPins(const UNiagaraNodeFunctionCall& FunctionCallNode, FCompileConstantResolver ConstantResolver);
 
 	void GetStackFunctionStaticSwitchPins(const UNiagaraNodeFunctionCall& FunctionCallNode, TArray<UEdGraphPin*>& OutInputPins, TSet<UEdGraphPin*>& OutHiddenPins, FCompileConstantResolver& ConstantResolver);
 
diff --git a/Engine/Plugins/FX/Niagara/Source/NiagaraEditorWidgets/Private/Stack/SNiagaraStackFunctionInputValue.cpp b/Engine/Plugins/FX/Niagara/Source/NiagaraEditorWidgets/Private/Stack/SNiagaraStackFunctionInputValue.cpp
--- a/Engine/Plugins/FX/Niagara/Source/NiagaraEditorWidgets/Private/Stack/SNiagaraStackFunctionInputValue.cpp
+++ b/Engine/Plugins/FX/Niagara/Source/NiagaraEditorWidgets/Private/Stack/SNiagaraStackFunctionInputValue.cpp
@@ -821,7 +821,7 @@ void SNiagaraStackFunctionInputValue::DynamicInputScriptSelected(UNiagaraScript*
 void SNiagaraStackFunctionInputValue::CustomExpressionSelected()
 {
 	FText CustomHLSLComment = LOCTEXT("NewCustomExpressionComment", "Custom HLSL!");
-	FunctionInput->SetCustomExpression(FHlslNiagaraTranslator::GetHlslDefaultForType(FunctionInput->GetInputType()) + TEXT(" /* ") + CustomHLSLComment.ToString() + TEXT(" */"));
+	FunctionInput->SetCustomExpression(INiagaraHlslTranslator::GetHlslDefaultForType(FunctionInput->GetInputType()) + TEXT(" /* ") + CustomHLSLComment.ToString() + TEXT(" */"));
 }
 
 void SNiagaraStackFunctionInputValue::CreateScratchSelected()
diff --git a/Engine/Source/Runtime/Core/Private/UObject/UE5MainStreamObjectVersion.cpp b/Engine/Source/Runtime/Core/Private/UObject/UE5MainStreamObjectVersion.cpp
--- a/Engine/Source/Runtime/Core/Private/UObject/UE5MainStreamObjectVersion.cpp
+++ b/Engine/Source/Runtime/Core/Private/UObject/UE5MainStreamObjectVersion.cpp
@@ -14,7 +14,7 @@ TMap<FGuid, FGuid> FUE5MainStreamObjectVersion::GetSystemGuids()
 	SystemGuids.Add(DevGuids.MATERIALSHADERMAP_DERIVEDDATA_VER, FGuid("65C01D817C9A4EEFAE9E988D41A1F3DD"));
 	SystemGuids.Add(DevGuids.NANITE_DERIVEDDATA_VER, FGuid("E107E819CFF5459CB3C2677C6C5E964D"));
 	SystemGuids.Add(DevGuids.NIAGARASHADERMAP_DERIVEDDATA_VER, FGuid("8A37C45D24F2423CBE5F8F371DE33575"));
-	SystemGuids.Add(DevGuids.Niagara_LatestScriptCompileVersion, FGuid("D30EFEEB3A9D4059AC7CD2B9748949DE"));
+	SystemGuids.Add(DevGuids.Niagara_LatestScriptCompileVersion, FGuid("9C8812406D824DE29BBC6F303964BB7E"));
 	SystemGuids.Add(DevGuids.SkeletalMeshDerivedDataVersion, FGuid("4F04FE8848F36B4F2AEB5D156092E08F"));
 	SystemGuids.Add(DevGuids.STATICMESH_DERIVEDDATA_VER, FGuid("98F7D79A5811013825E75BBCC41ED3E9"));
 	SystemGuids.Add(DevGuids.POSESEARCHDB_DERIVEDDATA_VER, FGuid("4E595C2AC5E947D6BA9ABC874353E5BC"));
