Added widget support for stack inputs

diff --git a/Engine/Plugins/FX/Niagara/Source/Niagara/Public/NiagaraTypes.h b/Engine/Plugins/FX/Niagara/Source/Niagara/Public/NiagaraTypes.h
--- a/Engine/Plugins/FX/Niagara/Source/Niagara/Public/NiagaraTypes.h
+++ b/Engine/Plugins/FX/Niagara/Source/Niagara/Public/NiagaraTypes.h
@@ -657,6 +657,75 @@ struct NIAGARA_API FNiagaraEnumParameterMetaData
 	FLinearColor ColorOverride = FLinearColor::White;
 };
 
+UENUM()
+enum class ENiagaraInputWidgetType : uint8
+{
+	// Default input widget
+	Default,
+
+	// slider widget, for float and int type
+	Slider,
+
+	// audio volume slider with mute control, for float input only
+	Volume,
+
+	// a numeric input, but also has a dropdown with named values
+	NumericDropdown
+};
+
+USTRUCT()
+struct NIAGARA_API FWidgetNamedInputValue
+{
+	GENERATED_BODY()
+
+	UPROPERTY(EditAnywhere, Category="Customization")
+	float Value = 0;
+	
+	UPROPERTY(EditAnywhere, Category="Customization")
+	FText DisplayName;
+
+	UPROPERTY(EditAnywhere, Category="Customization")
+	FText Tooltip;
+};
+
+USTRUCT()
+struct NIAGARA_API FNiagaraInputParameterCustomization
+{
+	GENERATED_BODY()
+
+	// Changes the widget implementation used for the input
+	UPROPERTY(EditAnywhere, Category="Customization")
+	ENiagaraInputWidgetType WidgetType = ENiagaraInputWidgetType::Default;
+	
+	UPROPERTY(EditAnywhere, Category="Customization", meta=(InlineEditConditionToggle))
+	bool bHasMinValue = false;
+	
+	/** min ui value (float and int types only) */
+	UPROPERTY(EditAnywhere, Category="Customization", meta=(EditCondition="bHasMinValue"))
+	float MinValue = 0;
+
+	UPROPERTY(EditAnywhere, Category="Customization", meta=(InlineEditConditionToggle))
+	bool bHasMaxValue = false;
+	
+	/** max ui value (float and int types only) */
+	UPROPERTY(EditAnywhere, Category="Customization", meta=(EditCondition="bHasMaxValue"))
+	float MaxValue = 1;
+	
+	UPROPERTY(EditAnywhere, Category="Customization", meta=(InlineEditConditionToggle))
+	bool bHasStepWidth = false;
+	
+	/** Step width used by the input when dragging */
+	UPROPERTY(EditAnywhere, Category="Customization", meta=(EditCondition="bHasStepWidth"))
+	float StepWidth = 1;
+
+	UPROPERTY(EditAnywhere, Category="Customization", meta=(EditCondition="WidgetType == ENiagaraInputWidgetType::NumericDropdown", EditConditionHides))
+	TArray<FWidgetNamedInputValue> InputDropdownValues;
+	
+	/** If true then the input is also displayed and editable as a 3d widget in the viewport (vector and transform types only). */
+	//UPROPERTY(EditAnywhere, Category="Customization")
+	//bool bCreateViewPortEditWidget = false;
+};
+
 UENUM()
 enum class ENiagaraBoolDisplayMode : uint8
 {
@@ -773,6 +842,10 @@ struct NIAGARA_API FNiagaraVariableMetaData
 	UPROPERTY(EditAnywhere, Category = "Variable", DisplayName = "Alternate Aliases For Variable", AdvancedDisplay, meta = (ToolTip = "List of alternate/previous names for this variable. Note that this is not normally needed if you rename through the UX. However, if you delete and then add a different variable, intending for it to match, you will likely want to add the prior name here.\n\nYou may need to restart and reload assets after making this change to have it take effect on already loaded assets."))
 	TArray<FName> AlternateAliases;
 
+	/** Changes how the input is displayed. */
+	UPROPERTY(EditAnywhere, Category = "Variable", meta = (SkipForCompileHash = "true"))
+	FNiagaraInputParameterCustomization WidgetCustomization;
+
 	bool GetIsStaticSwitch_DEPRECATED() const { return bIsStaticSwitch_DEPRECATED; };
 
 	int32 GetStaticSwitchDefaultValue_DEPRECATED() const { return StaticSwitchDefaultValue_DEPRECATED; };
@@ -802,7 +875,6 @@ private:
 	 */
 	UPROPERTY()
 	int32 StaticSwitchDefaultValue_DEPRECATED;  // TODO: This should be moved to the UNiagaraScriptVariable in the future
-
 };
 
 USTRUCT()
diff --git a/Engine/Plugins/FX/Niagara/Source/NiagaraEditor/Private/NiagaraEditorStyle.cpp b/Engine/Plugins/FX/Niagara/Source/NiagaraEditor/Private/NiagaraEditorStyle.cpp
--- a/Engine/Plugins/FX/Niagara/Source/NiagaraEditor/Private/NiagaraEditorStyle.cpp
+++ b/Engine/Plugins/FX/Niagara/Source/NiagaraEditor/Private/NiagaraEditorStyle.cpp
@@ -2,16 +2,12 @@
 
 #include "NiagaraEditorStyle.h"
 
-#include "ClassIconFinder.h"
 #include "Framework/Application/SlateApplication.h"
 #include "Styling/SlateStyleMacros.h"
 #include "Styling/AppStyle.h"
-#include "Slate/SlateGameResources.h"
 #include "Styling/SlateStyleRegistry.h"
 #include "Styling/SlateTypes.h"
 #include "Styling/CoreStyle.h"
-#include "Interfaces/IPluginManager.h"
-#include "Settings/EditorStyleSettings.h"
 #include "Styling/StarshipCoreStyle.h"
 #include "Styling/StyleColors.h"
 
@@ -490,6 +486,21 @@ void FNiagaraEditorStyle::InitStackIcons()
 	Set("NiagaraEditor.Stack.CPUIcon", new IMAGE_BRUSH("Icons/Simulate_CPU_x40", Icon16x16));
 }
 
+void FNiagaraEditorStyle::InitStackWidgets()
+{
+	FTextBlockStyle NormalText = FCoreStyle::Get().GetWidgetStyle<FTextBlockStyle>("NormalText");
+	FEditableTextBoxStyle DefaultTextBoxStyle = FAppStyle::GetWidgetStyle<FEditableTextBoxStyle>("NormalEditableTextBox");
+	const FEditableTextBoxStyle EditableTextBoxStyle = DefaultTextBoxStyle
+			.SetTextStyle(NormalText.SetFont(FAppStyle::Get().GetFontStyle("PropertyWindow.NormalFont")))
+			.SetPadding(0)
+			.SetBackgroundImageNormal( FSlateNoResource() )
+			.SetBackgroundImageHovered( FSlateNoResource() )
+			.SetBackgroundImageFocused( FSlateNoResource() )
+			.SetBackgroundImageReadOnly( FSlateNoResource() );
+	
+	Set("NiagaraEditor.Stack.NumericDropdownInput", EditableTextBoxStyle);
+}
+
 void FNiagaraEditorStyle::InitNiagaraSequence()
 {
 	// Niagara sequence
@@ -841,6 +852,7 @@ FNiagaraEditorStyle::FNiagaraEditorStyle() : FSlateStyleSet("NiagaraEditorStyle"
 	InitThumbnails();
 	InitClassIcon();
 	InitStackIcons();
+	InitStackWidgets();
 	InitNiagaraSequence();
 	InitPlatformSet();
 	InitDropTarget();
diff --git a/Engine/Plugins/FX/Niagara/Source/NiagaraEditor/Private/TypeEditorUtilities/NiagaraBoolTypeEditorUtilities.cpp b/Engine/Plugins/FX/Niagara/Source/NiagaraEditor/Private/TypeEditorUtilities/NiagaraBoolTypeEditorUtilities.cpp
--- a/Engine/Plugins/FX/Niagara/Source/NiagaraEditor/Private/TypeEditorUtilities/NiagaraBoolTypeEditorUtilities.cpp
+++ b/Engine/Plugins/FX/Niagara/Source/NiagaraEditor/Private/TypeEditorUtilities/NiagaraBoolTypeEditorUtilities.cpp
@@ -59,7 +59,7 @@ private:
 	bool bBoolValue;
 };
 
-TSharedPtr<SNiagaraParameterEditor> FNiagaraEditorBoolTypeUtilities::CreateParameterEditor(const FNiagaraTypeDefinition& ParameterType, EUnit DisplayUnit) const
+TSharedPtr<SNiagaraParameterEditor> FNiagaraEditorBoolTypeUtilities::CreateParameterEditor(const FNiagaraTypeDefinition& ParameterType, EUnit DisplayUnit, const FNiagaraInputParameterCustomization& WidgetCustomization) const
 {
 	return SNew(SNiagaraBoolParameterEditor);
 }
diff --git a/Engine/Plugins/FX/Niagara/Source/NiagaraEditor/Private/TypeEditorUtilities/NiagaraBoolTypeEditorUtilities.h b/Engine/Plugins/FX/Niagara/Source/NiagaraEditor/Private/TypeEditorUtilities/NiagaraBoolTypeEditorUtilities.h
--- a/Engine/Plugins/FX/Niagara/Source/NiagaraEditor/Private/TypeEditorUtilities/NiagaraBoolTypeEditorUtilities.h
+++ b/Engine/Plugins/FX/Niagara/Source/NiagaraEditor/Private/TypeEditorUtilities/NiagaraBoolTypeEditorUtilities.h
@@ -12,7 +12,7 @@ class FNiagaraEditorBoolTypeUtilities : public FNiagaraEditorTypeUtilities
 public:
 	//~ INiagaraEditorTypeUtilities interface.
 	virtual bool CanCreateParameterEditor() const override { return true; }
-	virtual TSharedPtr<SNiagaraParameterEditor> CreateParameterEditor(const FNiagaraTypeDefinition& ParameterType, EUnit DisplayUnit) const override;
+	virtual TSharedPtr<SNiagaraParameterEditor> CreateParameterEditor(const FNiagaraTypeDefinition& ParameterType, EUnit DisplayUnit, const FNiagaraInputParameterCustomization& WidgetCustomization) const override;
 	virtual bool CanHandlePinDefaults() const override;
 	virtual FString GetPinDefaultStringFromValue(const FNiagaraVariable& AllocatedVariable) const override;
 	virtual bool SetValueFromPinDefaultString(const FString& StringValue, FNiagaraVariable& Variable) const override;
diff --git a/Engine/Plugins/FX/Niagara/Source/NiagaraEditor/Private/TypeEditorUtilities/NiagaraColorTypeEditorUtilities.cpp b/Engine/Plugins/FX/Niagara/Source/NiagaraEditor/Private/TypeEditorUtilities/NiagaraColorTypeEditorUtilities.cpp
--- a/Engine/Plugins/FX/Niagara/Source/NiagaraEditor/Private/TypeEditorUtilities/NiagaraColorTypeEditorUtilities.cpp
+++ b/Engine/Plugins/FX/Niagara/Source/NiagaraEditor/Private/TypeEditorUtilities/NiagaraColorTypeEditorUtilities.cpp
@@ -202,7 +202,7 @@ void FNiagaraEditorColorTypeUtilities::UpdateVariableWithDefaultValue(FNiagaraVa
 	Variable.SetValue<FLinearColor>(FLinearColor(1, 1, 1, 1));
 }
 
-TSharedPtr<SNiagaraParameterEditor> FNiagaraEditorColorTypeUtilities::CreateParameterEditor(const FNiagaraTypeDefinition& ParameterType, EUnit DisplayUnit) const
+TSharedPtr<SNiagaraParameterEditor> FNiagaraEditorColorTypeUtilities::CreateParameterEditor(const FNiagaraTypeDefinition& ParameterType, EUnit DisplayUnit, const FNiagaraInputParameterCustomization& WidgetCustomization) const
 {
 	return SNew(SNiagaraColorParameterEditor);
 }
diff --git a/Engine/Plugins/FX/Niagara/Source/NiagaraEditor/Private/TypeEditorUtilities/NiagaraColorTypeEditorUtilities.h b/Engine/Plugins/FX/Niagara/Source/NiagaraEditor/Private/TypeEditorUtilities/NiagaraColorTypeEditorUtilities.h
--- a/Engine/Plugins/FX/Niagara/Source/NiagaraEditor/Private/TypeEditorUtilities/NiagaraColorTypeEditorUtilities.h
+++ b/Engine/Plugins/FX/Niagara/Source/NiagaraEditor/Private/TypeEditorUtilities/NiagaraColorTypeEditorUtilities.h
@@ -14,7 +14,7 @@ public:
 	virtual bool CanProvideDefaultValue() const override { return true; }
 	virtual void UpdateVariableWithDefaultValue(FNiagaraVariable& Variable) const override;
 	virtual bool CanCreateParameterEditor() const override { return true; }
-	virtual TSharedPtr<SNiagaraParameterEditor> CreateParameterEditor(const FNiagaraTypeDefinition& ParameterType, EUnit DisplayUnit) const override;
+	virtual TSharedPtr<SNiagaraParameterEditor> CreateParameterEditor(const FNiagaraTypeDefinition& ParameterType, EUnit DisplayUnit, const FNiagaraInputParameterCustomization& WidgetCustomization) const override;
 	virtual bool CanHandlePinDefaults() const override;
 	virtual FString GetPinDefaultStringFromValue(const FNiagaraVariable& AllocatedVariable) const override;
 	virtual bool SetValueFromPinDefaultString(const FString& StringValue, FNiagaraVariable& Variable) const override;
diff --git a/Engine/Plugins/FX/Niagara/Source/NiagaraEditor/Private/TypeEditorUtilities/NiagaraEnumTypeEditorUtilities.cpp b/Engine/Plugins/FX/Niagara/Source/NiagaraEditor/Private/TypeEditorUtilities/NiagaraEnumTypeEditorUtilities.cpp
--- a/Engine/Plugins/FX/Niagara/Source/NiagaraEditor/Private/TypeEditorUtilities/NiagaraEnumTypeEditorUtilities.cpp
+++ b/Engine/Plugins/FX/Niagara/Source/NiagaraEditor/Private/TypeEditorUtilities/NiagaraEnumTypeEditorUtilities.cpp
@@ -76,7 +76,7 @@ void FNiagaraEditorEnumTypeUtilities::UpdateVariableWithDefaultValue(FNiagaraVar
 	Variable.SetValue<FNiagaraInt32>(EnumIntValue);
 }
 
-TSharedPtr<SNiagaraParameterEditor> FNiagaraEditorEnumTypeUtilities::CreateParameterEditor(const FNiagaraTypeDefinition& ParameterType, EUnit DisplayUnit) const
+TSharedPtr<SNiagaraParameterEditor> FNiagaraEditorEnumTypeUtilities::CreateParameterEditor(const FNiagaraTypeDefinition& ParameterType, EUnit DisplayUnit, const FNiagaraInputParameterCustomization& WidgetCustomization) const
 {
 	return SNew(SNiagaraEnumParameterEditor, ParameterType.GetEnum());
 }
diff --git a/Engine/Plugins/FX/Niagara/Source/NiagaraEditor/Private/TypeEditorUtilities/NiagaraEnumTypeEditorUtilities.h b/Engine/Plugins/FX/Niagara/Source/NiagaraEditor/Private/TypeEditorUtilities/NiagaraEnumTypeEditorUtilities.h
--- a/Engine/Plugins/FX/Niagara/Source/NiagaraEditor/Private/TypeEditorUtilities/NiagaraEnumTypeEditorUtilities.h
+++ b/Engine/Plugins/FX/Niagara/Source/NiagaraEditor/Private/TypeEditorUtilities/NiagaraEnumTypeEditorUtilities.h
@@ -13,7 +13,7 @@ public:
 	virtual bool CanProvideDefaultValue() const override;
 	virtual void UpdateVariableWithDefaultValue(FNiagaraVariable& Variable) const override;
 	virtual bool CanCreateParameterEditor() const override { return true; }
-	virtual TSharedPtr<SNiagaraParameterEditor> CreateParameterEditor(const FNiagaraTypeDefinition& ParameterType, EUnit DisplayUnit) const override;
+	virtual TSharedPtr<SNiagaraParameterEditor> CreateParameterEditor(const FNiagaraTypeDefinition& ParameterType, EUnit DisplayUnit, const FNiagaraInputParameterCustomization& WidgetCustomization) const override;
 	virtual bool CanHandlePinDefaults() const override;
 	virtual FString GetPinDefaultStringFromValue(const FNiagaraVariable& AllocatedVariable) const override;
 	virtual bool SetValueFromPinDefaultString(const FString& StringValue, FNiagaraVariable& Variable) const override;
diff --git a/Engine/Plugins/FX/Niagara/Source/NiagaraEditor/Private/TypeEditorUtilities/NiagaraFloatTypeEditorUtilities.cpp b/Engine/Plugins/FX/Niagara/Source/NiagaraEditor/Private/TypeEditorUtilities/NiagaraFloatTypeEditorUtilities.cpp
--- a/Engine/Plugins/FX/Niagara/Source/NiagaraEditor/Private/TypeEditorUtilities/NiagaraFloatTypeEditorUtilities.cpp
+++ b/Engine/Plugins/FX/Niagara/Source/NiagaraEditor/Private/TypeEditorUtilities/NiagaraFloatTypeEditorUtilities.cpp
@@ -2,12 +2,16 @@
 
 #include "NiagaraFloatTypeEditorUtilities.h"
 
+#include "EdGraphSchema_Niagara.h"
 #include "GraphEditorSettings.h"
+#include "NiagaraEditorStyle.h"
 #include "SNiagaraParameterEditor.h"
 #include "NiagaraTypes.h"
-#include "NiagaraEditorStyle.h"
 #include "Math/UnitConversion.h"
+#include "Widgets/SNiagaraNumericDropDown.h"
 #include "Widgets/Input/SNumericEntryBox.h"
+#include "Widgets/Input/SSlider.h"
+#include "Widgets/Input/SVolumeControl.h"
 
 class SNiagaraFloatParameterEditor : public SNiagaraParameterEditor
 {
@@ -15,47 +19,170 @@ public:
 	SLATE_BEGIN_ARGS(SNiagaraFloatParameterEditor) { }
 	SLATE_END_ARGS();
 
-	void Construct(const FArguments& InArgs, EUnit DisplayUnit)
+	void Construct(const FArguments& InArgs, EUnit DisplayUnit, const FNiagaraInputParameterCustomization& WidgetCustomization)
 	{
 		SNiagaraParameterEditor::Construct(SNiagaraParameterEditor::FArguments()
 			.MinimumDesiredWidth(DefaultInputSize)
 			.MaximumDesiredWidth(DefaultInputSize));
 
 		const UGraphEditorSettings* Settings = GetDefault<UGraphEditorSettings>();
-		ChildSlot
-		[
-			SNew(SNumericEntryBox<float>)
-			.Font(FAppStyle::Get().GetFontStyle("PropertyWindow.NormalFont"))
-			.MinValue(TOptional<float>())
-			.MaxValue(TOptional<float>())
-			.MaxSliderValue(TOptional<float>())
-			.MinSliderValue(TOptional<float>())
-			.Value(this, &SNiagaraFloatParameterEditor::GetValue)
-			.OnValueChanged(this, &SNiagaraFloatParameterEditor::ValueChanged)
-			.OnValueCommitted(this, &SNiagaraFloatParameterEditor::ValueCommitted)
-			.OnBeginSliderMovement(this, &SNiagaraFloatParameterEditor::BeginSliderMovement)
-			.OnEndSliderMovement(this, &SNiagaraFloatParameterEditor::EndSliderMovement)
-			.TypeInterface(GetTypeInterface<float>(DisplayUnit))
-			.AllowSpin(true)
-			.LabelPadding(FMargin(3) )
-			.LabelLocation(SNumericEntryBox<float>::ELabelLocation::Inside)
-			.Label()
+
+		if (WidgetCustomization.WidgetType == ENiagaraInputWidgetType::Slider)
+		{
+			float MinValue = WidgetCustomization.bHasMinValue ? WidgetCustomization.MinValue : 0;
+			float MaxValue = WidgetCustomization.bHasMaxValue ? WidgetCustomization.MaxValue : 1;
+			ChildSlot
+			[
+				SNew(SHorizontalBox)
+				+ SHorizontalBox::Slot()
+				.VAlign(VAlign_Center)
+				.HAlign(HAlign_Center)
+				.AutoWidth()
+				[
+					SNew(SImage)
+					.ColorAndOpacity(UEdGraphSchema_Niagara::GetTypeColor(FNiagaraTypeDefinition::GetFloatDef()))
+					.Image(FNiagaraEditorStyle::Get().GetBrush("NiagaraEditor.Module.TypeIconPill"))
+				]
+				+ SHorizontalBox::Slot().AutoWidth()
+				[
+					SNew(SBox).WidthOverride(100.0f)
+					[
+						SNew(SSlider)
+						.MinValue(MinValue)
+						.MaxValue(MaxValue)
+						.Value(this, &SNiagaraFloatParameterEditor::GetSliderValue)
+						.OnValueChanged_Lambda([this, WidgetCustomization](float NewVal)
+						{
+							SliderValue = NewVal;
+							FloatValue = NewVal;
+							if (WidgetCustomization.bHasStepWidth && WidgetCustomization.StepWidth != 0)
+							{
+								FloatValue = FMath::RoundToFloat(NewVal / WidgetCustomization.StepWidth) * WidgetCustomization.StepWidth;
+							}
+							ExecuteOnValueChanged();
+						})
+						.OnMouseCaptureBegin(this, &SNiagaraFloatParameterEditor::ExecuteOnBeginValueChange)
+						.OnMouseCaptureEnd(this, &SNiagaraFloatParameterEditor::ExecuteOnEndValueChange)
+					]
+				]
+				+ SHorizontalBox::Slot().AutoWidth()
+				[
+					SNew(SBox).WidthOverride(75.0f)
+					[
+						SNew(SNumericEntryBox<float>)
+						.Font(FAppStyle::Get().GetFontStyle("PropertyWindow.NormalFont"))
+						.MinValue(MinValue)
+						.MaxValue(MaxValue)
+						.Value(this, &SNiagaraFloatParameterEditor::GetValue)
+						.OnValueChanged(this, &SNiagaraFloatParameterEditor::ValueChanged)
+						.OnValueCommitted(this, &SNiagaraFloatParameterEditor::ValueCommitted)
+						.TypeInterface(GetTypeInterface<float>(DisplayUnit))
+						.AllowSpin(false)
+						.Delta(WidgetCustomization.bHasStepWidth ? WidgetCustomization.StepWidth : 0)
+					]
+				]
+			];
+		}
+		else if (WidgetCustomization.WidgetType == ENiagaraInputWidgetType::Volume)
+		{
+			ChildSlot
+			[
+				SNew(SVolumeControl)
+				.Volume(this, &SNiagaraFloatParameterEditor::GetSliderValue)
+				.Muted(this, &SNiagaraFloatParameterEditor::IsMuted)
+				.OnVolumeChanged_Lambda([this](float NewVal)
+				{
+					SliderValue = NewVal;
+					FloatValue = NewVal;
+					ExecuteOnValueChanged();
+				})
+				.OnMuteChanged_Lambda([this](bool NewVal)
+				{
+					bMuted = NewVal;
+					FloatValue = bMuted ? 0 : SliderValue;
+					ExecuteOnValueChanged();
+				})
+			];
+		}
+		else if (WidgetCustomization.WidgetType == ENiagaraInputWidgetType::NumericDropdown && WidgetCustomization.InputDropdownValues.Num() > 0)
+		{
+			TArray<SNiagaraNumericDropDown<float>::FNamedValue> DropDownValues;
+			for (const FWidgetNamedInputValue& Value : WidgetCustomization.InputDropdownValues)
+			{
+				DropDownValues.Add(SNiagaraNumericDropDown<float>::FNamedValue(Value.Value, Value.DisplayName, Value.Tooltip));
+			}
+
+			ChildSlot
 			[
-				SNumericEntryBox<int32>::BuildNarrowColorLabel(Settings->FloatPinTypeColor)
-			]
-		];
+				SNew(SNiagaraNumericDropDown<float>)
+				.DropDownValues(DropDownValues)
+				.bShowNamedValue(true)
+				.MinDesiredValueWidth(75)
+				.PillType(FNiagaraTypeDefinition::GetFloatDef())
+				.Value_Lambda([this]() 
+				{ 
+					return FloatValue;
+				})
+				.OnValueChanged_Lambda([this](float NewVal)
+				{
+					FloatValue = NewVal;
+					ExecuteOnValueChanged();
+				})
+			];
+		}
+		else
+		{
+			TOptional<float> MinValue;
+			TOptional<float> MaxValue;
+			if (WidgetCustomization.bHasMinValue)
+			{
+				MinValue = WidgetCustomization.MinValue;
+			}
+			if (WidgetCustomization.bHasMaxValue)
+			{
+				MaxValue = WidgetCustomization.MaxValue;
+			}
+			
+			ChildSlot
+			[
+				SNew(SNumericEntryBox<float>)
+				.Font(FAppStyle::Get().GetFontStyle("PropertyWindow.NormalFont"))
+				.MinValue(MinValue)
+				.MaxValue(MaxValue)
+				.MinSliderValue(MinValue)
+				.MaxSliderValue(MaxValue)
+				.Value(this, &SNiagaraFloatParameterEditor::GetValue)
+				.OnValueChanged(this, &SNiagaraFloatParameterEditor::ValueChanged)
+				.OnValueCommitted(this, &SNiagaraFloatParameterEditor::ValueCommitted)
+				.OnBeginSliderMovement(this, &SNiagaraFloatParameterEditor::BeginSliderMovement)
+				.OnEndSliderMovement(this, &SNiagaraFloatParameterEditor::EndSliderMovement)
+				.TypeInterface(GetTypeInterface<float>(DisplayUnit))
+				.AllowSpin(true)
+				.LabelPadding(FMargin(3))
+				.Delta(WidgetCustomization.bHasStepWidth ? WidgetCustomization.StepWidth : 0)
+				.LabelLocation(SNumericEntryBox<float>::ELabelLocation::Inside)
+				.Label()
+				[
+					SNumericEntryBox<float>::BuildNarrowColorLabel(Settings->FloatPinTypeColor)
+				]
+			];
+		}
 	}
 
 	virtual void UpdateInternalValueFromStruct(TSharedRef<FStructOnScope> Struct) override
 	{
 		checkf(Struct->GetStruct() == FNiagaraTypeDefinition::GetFloatStruct(), TEXT("Struct type not supported."));
-		FloatValue = ((FNiagaraFloat*)Struct->GetStructMemory())->Value;
+		FloatValue = reinterpret_cast<FNiagaraFloat*>(Struct->GetStructMemory())->Value;
+		if (!bMuted)
+		{
+			SliderValue = FloatValue;
+		}
 	}
 
 	virtual void UpdateStructFromInternalValue(TSharedRef<FStructOnScope> Struct) override
 	{
 		checkf(Struct->GetStruct() == FNiagaraTypeDefinition::GetFloatStruct(), TEXT("Struct type not supported."));
-		((FNiagaraFloat*)Struct->GetStructMemory())->Value = FloatValue;
+		reinterpret_cast<FNiagaraFloat*>(Struct->GetStructMemory())->Value = FloatValue;
 	}
 
 	virtual bool CanChangeContinuously() const override { return true; }
@@ -76,6 +203,16 @@ private:
 		return FloatValue;
 	}
 
+	float GetSliderValue() const
+	{
+		return SliderValue;
+	}
+
+	bool IsMuted() const
+	{
+		return bMuted;
+	}
+	
 	void ValueChanged(float Value)
 	{
 		FloatValue = Value;
@@ -92,11 +229,13 @@ private:
 
 private:
 	float FloatValue = 0;
+	float SliderValue = 0;
+	bool bMuted = false;
 };
 
-TSharedPtr<SNiagaraParameterEditor> FNiagaraEditorFloatTypeUtilities::CreateParameterEditor(const FNiagaraTypeDefinition& ParameterType, EUnit DisplayUnit) const
+TSharedPtr<SNiagaraParameterEditor> FNiagaraEditorFloatTypeUtilities::CreateParameterEditor(const FNiagaraTypeDefinition& ParameterType, EUnit DisplayUnit, const FNiagaraInputParameterCustomization& WidgetCustomization) const
 {
-	return SNew(SNiagaraFloatParameterEditor, DisplayUnit);
+	return SNew(SNiagaraFloatParameterEditor, DisplayUnit, WidgetCustomization);
 }
 
 bool FNiagaraEditorFloatTypeUtilities::CanHandlePinDefaults() const
diff --git a/Engine/Plugins/FX/Niagara/Source/NiagaraEditor/Private/TypeEditorUtilities/NiagaraFloatTypeEditorUtilities.h b/Engine/Plugins/FX/Niagara/Source/NiagaraEditor/Private/TypeEditorUtilities/NiagaraFloatTypeEditorUtilities.h
--- a/Engine/Plugins/FX/Niagara/Source/NiagaraEditor/Private/TypeEditorUtilities/NiagaraFloatTypeEditorUtilities.h
+++ b/Engine/Plugins/FX/Niagara/Source/NiagaraEditor/Private/TypeEditorUtilities/NiagaraFloatTypeEditorUtilities.h
@@ -13,7 +13,7 @@ class FNiagaraEditorFloatTypeUtilities : public FNiagaraEditorTypeUtilities
 public:
 	//~ INiagaraEditorTypeUtilities interface.
 	virtual bool CanCreateParameterEditor() const override { return true; }
-	virtual TSharedPtr<SNiagaraParameterEditor> CreateParameterEditor(const FNiagaraTypeDefinition& ParameterType, EUnit DisplayUnit) const override;
+	virtual TSharedPtr<SNiagaraParameterEditor> CreateParameterEditor(const FNiagaraTypeDefinition& ParameterType, EUnit DisplayUnit, const FNiagaraInputParameterCustomization& WidgetCustomization) const override;
 	virtual bool CanHandlePinDefaults() const override;
 	virtual FString GetPinDefaultStringFromValue(const FNiagaraVariable& AllocatedVariable) const override;
 	virtual bool SetValueFromPinDefaultString(const FString& StringValue, FNiagaraVariable& Variable) const override;
diff --git a/Engine/Plugins/FX/Niagara/Source/NiagaraEditor/Private/TypeEditorUtilities/NiagaraIntegerTypeEditorUtilities.cpp b/Engine/Plugins/FX/Niagara/Source/NiagaraEditor/Private/TypeEditorUtilities/NiagaraIntegerTypeEditorUtilities.cpp
--- a/Engine/Plugins/FX/Niagara/Source/NiagaraEditor/Private/TypeEditorUtilities/NiagaraIntegerTypeEditorUtilities.cpp
+++ b/Engine/Plugins/FX/Niagara/Source/NiagaraEditor/Private/TypeEditorUtilities/NiagaraIntegerTypeEditorUtilities.cpp
@@ -2,12 +2,15 @@
 
 #include "NiagaraIntegerTypeEditorUtilities.h"
 
+#include "EdGraphSchema_Niagara.h"
 #include "GraphEditorSettings.h"
 #include "SNiagaraParameterEditor.h"
 #include "NiagaraTypes.h"
 #include "NiagaraEditorStyle.h"
+#include "Widgets/SNiagaraNumericDropDown.h"
 
 #include "Widgets/Input/SNumericEntryBox.h"
+#include "Widgets/Input/SSlider.h"
 
 class SNiagaraIntegerParameterEditor : public SNiagaraParameterEditor
 {
@@ -15,47 +18,144 @@ public:
 	SLATE_BEGIN_ARGS(SNiagaraIntegerParameterEditor) { }
 	SLATE_END_ARGS();
 
-	void Construct(const FArguments& InArgs, EUnit DisplayUnit)
+	void Construct(const FArguments& InArgs, EUnit DisplayUnit, const FNiagaraInputParameterCustomization& WidgetCustomization)
 	{
 		SNiagaraParameterEditor::Construct(SNiagaraParameterEditor::FArguments()
 			.MinimumDesiredWidth(DefaultInputSize)
 			.MaximumDesiredWidth(DefaultInputSize));
 
 		const UGraphEditorSettings* Settings = GetDefault<UGraphEditorSettings>();
-		ChildSlot
-		[
-			SNew(SNumericEntryBox<int32>)
-			.Font(FAppStyle::Get().GetFontStyle("PropertyWindow.NormalFont"))
-			.MinValue(TOptional<int32>())
-			.MaxValue(TOptional<int32>())
-			.MaxSliderValue(TOptional<int32>())
-			.MinSliderValue(TOptional<int32>())
-			.Value(this, &SNiagaraIntegerParameterEditor::GetValue)
-			.OnValueChanged(this, &SNiagaraIntegerParameterEditor::ValueChanged)
-			.OnValueCommitted(this, &SNiagaraIntegerParameterEditor::ValueCommitted)
-			.OnBeginSliderMovement(this, &SNiagaraIntegerParameterEditor::BeginSliderMovement)
-			.OnEndSliderMovement(this, &SNiagaraIntegerParameterEditor::EndSliderMovement)
-			.TypeInterface(GetTypeInterface<int32>(DisplayUnit))
-			.AllowSpin(true)
-			.LabelPadding(FMargin(3))
-			.LabelLocation(SNumericEntryBox<int32>::ELabelLocation::Inside)
-			.Label()
+
+		if (WidgetCustomization.WidgetType == ENiagaraInputWidgetType::Slider)
+		{
+			float MinValue = WidgetCustomization.bHasMinValue ? WidgetCustomization.MinValue : 0;
+			float MaxValue = WidgetCustomization.bHasMaxValue ? WidgetCustomization.MaxValue : 1;
+			ChildSlot
+			[
+				SNew(SHorizontalBox)
+				+ SHorizontalBox::Slot()
+				.VAlign(VAlign_Center)
+				.HAlign(HAlign_Center)
+				.AutoWidth()
+				[
+					SNew(SImage)
+					.ColorAndOpacity(UEdGraphSchema_Niagara::GetTypeColor(FNiagaraTypeDefinition::GetIntDef()))
+					.Image(FNiagaraEditorStyle::Get().GetBrush("NiagaraEditor.Module.TypeIconPill"))
+				]
+				+ SHorizontalBox::Slot().AutoWidth()
+				[
+					SNew(SBox).WidthOverride(100.0f)
+					[
+						SNew(SSlider)
+						.MinValue(MinValue)
+						.MaxValue(MaxValue)
+						.Value(this, &SNiagaraIntegerParameterEditor::GetSliderValue)
+						.OnValueChanged_Lambda([this, WidgetCustomization](float NewVal)
+						{
+							SliderValue = NewVal;
+
+							// slider only works with float, so we round for the actual parameter value
+							int32 Resolution = WidgetCustomization.bHasStepWidth && WidgetCustomization.StepWidth >= 1 ? WidgetCustomization.StepWidth : 1;
+							IntValue = FMath::RoundToInt(NewVal / Resolution) * Resolution;
+							ExecuteOnValueChanged();
+						})
+						.OnMouseCaptureBegin(this, &SNiagaraIntegerParameterEditor::ExecuteOnBeginValueChange)
+						.OnMouseCaptureEnd(this, &SNiagaraIntegerParameterEditor::ExecuteOnEndValueChange)
+					]
+				]
+				+ SHorizontalBox::Slot().AutoWidth()
+				[
+					SNew(SBox).WidthOverride(75.0f)
+					[
+						SNew(SNumericEntryBox<int32>)
+						.Font(FAppStyle::Get().GetFontStyle("PropertyWindow.NormalFont"))
+						.MinValue(MinValue)
+						.MaxValue(MaxValue)
+						.Value(this, &SNiagaraIntegerParameterEditor::GetValue)
+						.OnValueChanged(this, &SNiagaraIntegerParameterEditor::ValueChanged)
+						.OnValueCommitted(this, &SNiagaraIntegerParameterEditor::ValueCommitted)
+						.TypeInterface(GetTypeInterface<int32>(DisplayUnit))
+						.AllowSpin(false)
+						.Delta(WidgetCustomization.bHasStepWidth ? WidgetCustomization.StepWidth : 0)
+					]
+				]
+			];
+		}
+		else if (WidgetCustomization.WidgetType == ENiagaraInputWidgetType::NumericDropdown && WidgetCustomization.InputDropdownValues.Num() > 0)
+		{
+			TArray<SNiagaraNumericDropDown<int32>::FNamedValue> DropDownValues;
+			for (const FWidgetNamedInputValue& Value : WidgetCustomization.InputDropdownValues)
+			{
+				DropDownValues.Add(SNiagaraNumericDropDown<int32>::FNamedValue(Value.Value, Value.DisplayName, Value.Tooltip));
+			}
+
+			ChildSlot
 			[
-				SNumericEntryBox<int32>::BuildNarrowColorLabel(Settings->IntPinTypeColor)
-			]
-		];
+				SNew(SNiagaraNumericDropDown<int32>)
+				.DropDownValues(DropDownValues)
+				.bShowNamedValue(true)
+				.MinDesiredValueWidth(75)
+				.PillType(FNiagaraTypeDefinition::GetIntDef())
+				.Value_Lambda([this]() 
+				{ 
+					return IntValue;
+				})
+				.OnValueChanged_Lambda([this](int32 NewVal)
+				{
+					IntValue = NewVal;
+					ExecuteOnValueChanged();
+				})
+			];
+		}
+		else
+		{
+			TOptional<int32> MinValue;
+			TOptional<int32> MaxValue;
+			if (WidgetCustomization.bHasMinValue)
+			{
+				MinValue = WidgetCustomization.MinValue;
+			}
+			if (WidgetCustomization.bHasMaxValue)
+			{
+				MaxValue = WidgetCustomization.MaxValue;
+			}
+			
+			ChildSlot
+			[
+				SNew(SNumericEntryBox<int32>)
+				.Font(FAppStyle::Get().GetFontStyle("PropertyWindow.NormalFont"))
+				.MinValue(MinValue)
+				.MaxValue(MaxValue)
+				.MinSliderValue(MinValue)
+				.MaxSliderValue(MaxValue)
+				.Value(this, &SNiagaraIntegerParameterEditor::GetValue)
+				.OnValueChanged(this, &SNiagaraIntegerParameterEditor::ValueChanged)
+				.OnValueCommitted(this, &SNiagaraIntegerParameterEditor::ValueCommitted)
+				.OnBeginSliderMovement(this, &SNiagaraIntegerParameterEditor::BeginSliderMovement)
+				.OnEndSliderMovement(this, &SNiagaraIntegerParameterEditor::EndSliderMovement)
+				.TypeInterface(GetTypeInterface<int32>(DisplayUnit))
+				.AllowSpin(true)
+				.LabelPadding(FMargin(3))
+				.LabelLocation(SNumericEntryBox<int32>::ELabelLocation::Inside)
+				.Label()
+				[
+					SNumericEntryBox<int32>::BuildNarrowColorLabel(Settings->IntPinTypeColor)
+				]
+			];
+		}
 	}
 
 	virtual void UpdateInternalValueFromStruct(TSharedRef<FStructOnScope> Struct) override
 	{
 		checkf(Struct->GetStruct() == FNiagaraTypeDefinition::GetIntStruct(), TEXT("Struct type not supported."));
-		IntValue = ((FNiagaraInt32*)Struct->GetStructMemory())->Value;
+		IntValue = reinterpret_cast<FNiagaraInt32*>(Struct->GetStructMemory())->Value;
+		SliderValue = IntValue;
 	}
 
 	virtual void UpdateStructFromInternalValue(TSharedRef<FStructOnScope> Struct) override
 	{
 		checkf(Struct->GetStruct() == FNiagaraTypeDefinition::GetIntStruct(), TEXT("Struct type not supported."));
-		((FNiagaraInt32*)Struct->GetStructMemory())->Value = IntValue;
+		reinterpret_cast<FNiagaraInt32*>(Struct->GetStructMemory())->Value = IntValue;
 	}
 
 	virtual bool CanChangeContinuously() const override { return true; }
@@ -75,6 +175,11 @@ private:
 	{
 		return IntValue;
 	}
+	
+	float GetSliderValue() const
+    {
+    	return SliderValue;
+    }
 
 	void ValueChanged(int32 Value)
 	{
@@ -90,13 +195,13 @@ private:
 		}
 	}
 
-private:
 	int32 IntValue = 0;
+	float SliderValue = 0;
 };
 
-TSharedPtr<SNiagaraParameterEditor> FNiagaraEditorIntegerTypeUtilities::CreateParameterEditor(const FNiagaraTypeDefinition& ParameterType, EUnit DisplayUnit) const
+TSharedPtr<SNiagaraParameterEditor> FNiagaraEditorIntegerTypeUtilities::CreateParameterEditor(const FNiagaraTypeDefinition& ParameterType, EUnit DisplayUnit, const FNiagaraInputParameterCustomization& WidgetCustomization) const
 {
-	return SNew(SNiagaraIntegerParameterEditor, DisplayUnit);
+	return SNew(SNiagaraIntegerParameterEditor, DisplayUnit, WidgetCustomization);
 }
 
 bool FNiagaraEditorIntegerTypeUtilities::CanHandlePinDefaults() const
diff --git a/Engine/Plugins/FX/Niagara/Source/NiagaraEditor/Private/TypeEditorUtilities/NiagaraIntegerTypeEditorUtilities.h b/Engine/Plugins/FX/Niagara/Source/NiagaraEditor/Private/TypeEditorUtilities/NiagaraIntegerTypeEditorUtilities.h
--- a/Engine/Plugins/FX/Niagara/Source/NiagaraEditor/Private/TypeEditorUtilities/NiagaraIntegerTypeEditorUtilities.h
+++ b/Engine/Plugins/FX/Niagara/Source/NiagaraEditor/Private/TypeEditorUtilities/NiagaraIntegerTypeEditorUtilities.h
@@ -11,7 +11,7 @@ class FNiagaraEditorIntegerTypeUtilities : public FNiagaraEditorTypeUtilities
 public:
 	//~ INiagaraEditorTypeUtilities interface.
 	virtual bool CanCreateParameterEditor() const override { return true; }
-	virtual TSharedPtr<SNiagaraParameterEditor> CreateParameterEditor(const FNiagaraTypeDefinition& ParameterType, EUnit DisplayUnit) const override;
+	virtual TSharedPtr<SNiagaraParameterEditor> CreateParameterEditor(const FNiagaraTypeDefinition& ParameterType, EUnit DisplayUnit, const FNiagaraInputParameterCustomization& WidgetCustomization) const override;
 	virtual bool CanHandlePinDefaults() const override;
 	virtual FString GetPinDefaultStringFromValue(const FNiagaraVariable& AllocatedVariable) const override;
 	virtual bool SetValueFromPinDefaultString(const FString& StringValue, FNiagaraVariable& Variable) const override;
diff --git a/Engine/Plugins/FX/Niagara/Source/NiagaraEditor/Private/TypeEditorUtilities/NiagaraVectorTypeEditorUtilities.cpp b/Engine/Plugins/FX/Niagara/Source/NiagaraEditor/Private/TypeEditorUtilities/NiagaraVectorTypeEditorUtilities.cpp
--- a/Engine/Plugins/FX/Niagara/Source/NiagaraEditor/Private/TypeEditorUtilities/NiagaraVectorTypeEditorUtilities.cpp
+++ b/Engine/Plugins/FX/Niagara/Source/NiagaraEditor/Private/TypeEditorUtilities/NiagaraVectorTypeEditorUtilities.cpp
@@ -15,7 +15,7 @@ public:
 		SLATE_ARGUMENT(int32, ComponentCount)
 	SLATE_END_ARGS();
 
-	void Construct(const FArguments& InArgs, EUnit DisplayUnit)
+	void Construct(const FArguments& InArgs, EUnit DisplayUnit, const FNiagaraInputParameterCustomization& WidgetCustomization)
 	{
 		SNiagaraParameterEditor::Construct(SNiagaraParameterEditor::FArguments()
 			.MinimumDesiredWidth(DefaultInputSize * InArgs._ComponentCount)
@@ -27,7 +27,7 @@ public:
 			ComponentBox->AddSlot()
 			.Padding(ComponentIndex == 0 ? 0 : 3, 0, 0, 0)
 			[
-				ConstructComponentWidget(ComponentIndex, DisplayUnit)
+				ConstructComponentWidget(ComponentIndex, DisplayUnit, WidgetCustomization)
 			];
 		}
 
@@ -45,7 +45,7 @@ protected:
 	virtual void SetValue(int32 Index, float Value) = 0;
 
 private:
-	TSharedRef<SWidget> ConstructComponentWidget(int32 Index, EUnit DisplayUnit)
+	TSharedRef<SWidget> ConstructComponentWidget(int32 Index, EUnit DisplayUnit, const FNiagaraInputParameterCustomization& WidgetCustomization)
 	{
 		const static TArray Colors = {
 			&SNumericEntryBox<int32>::RedLabelBackgroundColor,
@@ -53,16 +53,26 @@ private:
 			&SNumericEntryBox<int32>::BlueLabelBackgroundColor,
 			&FLinearColor::White
 		};
+		TOptional<float> MinValue;
+		TOptional<float> MaxValue;
+		if (WidgetCustomization.bHasMinValue)
+		{
+			MinValue = WidgetCustomization.MinValue;
+		}
+		if (WidgetCustomization.bHasMaxValue)
+		{
+			MaxValue = WidgetCustomization.MaxValue;
+		}
 		
 		TSharedRef<SWidget>  LabelWidget = SNumericEntryBox<float>::BuildNarrowColorLabel(*Colors[Index]);
 		
 		return SNew(SNumericEntryBox<float>)
 		.Font(FAppStyle::Get().GetFontStyle("PropertyWindow.NormalFont"))
-		.MinValue(TOptional<float>())
-		.MaxValue(TOptional<float>())
-		.MaxSliderValue(TOptional<float>())
-		.MinSliderValue(TOptional<float>())
-		.Delta(0.0f)
+		.MinValue(MinValue)
+		.MaxValue(MaxValue)
+		.MaxSliderValue(MaxValue)
+		.MinSliderValue(MinValue)
+		.Delta(WidgetCustomization.bHasStepWidth ? WidgetCustomization.StepWidth : 0)
 		.Value(this, &SNiagaraVectorParameterEditorBase::GetValueInternal, Index)
 		.OnValueChanged(this, &SNiagaraVectorParameterEditorBase::ValueChanged, Index)
 		.OnValueCommitted(this, &SNiagaraVectorParameterEditorBase::ValueCommitted, Index)
@@ -115,11 +125,11 @@ public:
 	SLATE_BEGIN_ARGS(SNiagaraVector2ParameterEditor) { }
 	SLATE_END_ARGS();
 
-	void Construct(const FArguments& InArgs, EUnit DisplayUnit)
+	void Construct(const FArguments& InArgs, EUnit DisplayUnit, const FNiagaraInputParameterCustomization& WidgetCustomization)
 	{
 		SNiagaraVectorParameterEditorBase::Construct(
 			SNiagaraVectorParameterEditorBase::FArguments()
-			.ComponentCount(2), DisplayUnit);
+			.ComponentCount(2), DisplayUnit, WidgetCustomization);
 	}
 
 	virtual void UpdateInternalValueFromStruct(TSharedRef<FStructOnScope> Struct) override
@@ -149,9 +159,9 @@ private:
 	FVector2f VectorValue;
 };
 
-TSharedPtr<SNiagaraParameterEditor> FNiagaraEditorVector2TypeUtilities::CreateParameterEditor(const FNiagaraTypeDefinition& ParameterType, EUnit DisplayUnit) const
+TSharedPtr<SNiagaraParameterEditor> FNiagaraEditorVector2TypeUtilities::CreateParameterEditor(const FNiagaraTypeDefinition& ParameterType, EUnit DisplayUnit, const FNiagaraInputParameterCustomization& WidgetCustomization) const
 {
-	return SNew(SNiagaraVector2ParameterEditor, DisplayUnit);
+	return SNew(SNiagaraVector2ParameterEditor, DisplayUnit, WidgetCustomization);
 }
 
 bool FNiagaraEditorVector2TypeUtilities::CanHandlePinDefaults() const
@@ -193,11 +203,11 @@ public:
 	SLATE_BEGIN_ARGS(SNiagaraVector3ParameterEditor) { }
 	SLATE_END_ARGS();
 
-	void Construct(const FArguments& InArgs, EUnit DisplayUnit)
+	void Construct(const FArguments& InArgs, EUnit DisplayUnit, const FNiagaraInputParameterCustomization& WidgetCustomization)
 	{
 		SNiagaraVectorParameterEditorBase::Construct(
 			SNiagaraVectorParameterEditorBase::FArguments()
-			.ComponentCount(3), DisplayUnit);
+			.ComponentCount(3), DisplayUnit, WidgetCustomization);
 	}
 
 	virtual void UpdateInternalValueFromStruct(TSharedRef<FStructOnScope> Struct) override
@@ -227,9 +237,9 @@ private:
 	FVector3f VectorValue;
 };
 
-TSharedPtr<SNiagaraParameterEditor> FNiagaraEditorVector3TypeUtilities::CreateParameterEditor(const FNiagaraTypeDefinition& ParameterType, EUnit DisplayUnit) const
+TSharedPtr<SNiagaraParameterEditor> FNiagaraEditorVector3TypeUtilities::CreateParameterEditor(const FNiagaraTypeDefinition& ParameterType, EUnit DisplayUnit, const FNiagaraInputParameterCustomization& WidgetCustomization) const
 {
-	return SNew(SNiagaraVector3ParameterEditor, DisplayUnit);
+	return SNew(SNiagaraVector3ParameterEditor, DisplayUnit, WidgetCustomization);
 }
 
 bool FNiagaraEditorVector3TypeUtilities::CanHandlePinDefaults() const
@@ -275,11 +285,11 @@ public:
 	SLATE_BEGIN_ARGS(SNiagaraVector4ParameterEditor) { }
 	SLATE_END_ARGS();
 
-	void Construct(const FArguments& InArgs, EUnit DisplayUnit)
+	void Construct(const FArguments& InArgs, EUnit DisplayUnit, const FNiagaraInputParameterCustomization& WidgetCustomization)
 	{
 		SNiagaraVectorParameterEditorBase::Construct(
 			SNiagaraVectorParameterEditorBase::FArguments()
-			.ComponentCount(4), DisplayUnit);
+			.ComponentCount(4), DisplayUnit, WidgetCustomization);
 	}
 
 	virtual void UpdateInternalValueFromStruct(TSharedRef<FStructOnScope> Struct) override
@@ -309,9 +319,9 @@ private:
 	FVector4f VectorValue;
 };
 
-TSharedPtr<SNiagaraParameterEditor> FNiagaraEditorVector4TypeUtilities::CreateParameterEditor(const FNiagaraTypeDefinition& ParameterType, EUnit DisplayUnit) const
+TSharedPtr<SNiagaraParameterEditor> FNiagaraEditorVector4TypeUtilities::CreateParameterEditor(const FNiagaraTypeDefinition& ParameterType, EUnit DisplayUnit, const FNiagaraInputParameterCustomization& WidgetCustomization) const
 {
-	return SNew(SNiagaraVector4ParameterEditor, DisplayUnit);
+	return SNew(SNiagaraVector4ParameterEditor, DisplayUnit, WidgetCustomization);
 }
 
 bool FNiagaraEditorVector4TypeUtilities::CanHandlePinDefaults() const
@@ -357,11 +367,11 @@ public:
 	SLATE_BEGIN_ARGS(SNiagaraQuatParameterEditor) { }
 	SLATE_END_ARGS();
 
-	void Construct(const FArguments& InArgs, EUnit DisplayUnit)
+	void Construct(const FArguments& InArgs, EUnit DisplayUnit, const FNiagaraInputParameterCustomization& WidgetCustomization)
 	{
 		SNiagaraVectorParameterEditorBase::Construct(
 			SNiagaraVectorParameterEditorBase::FArguments()
-			.ComponentCount(4), DisplayUnit);
+			.ComponentCount(4), DisplayUnit, WidgetCustomization);
 	}
 
 	virtual void UpdateInternalValueFromStruct(TSharedRef<FStructOnScope> Struct) override
@@ -420,9 +430,9 @@ private:
 	FQuat4f VectorValue;
 };
 
-TSharedPtr<SNiagaraParameterEditor> FNiagaraEditorQuatTypeUtilities::CreateParameterEditor(const FNiagaraTypeDefinition& ParameterType, EUnit DisplayUnit) const
+TSharedPtr<SNiagaraParameterEditor> FNiagaraEditorQuatTypeUtilities::CreateParameterEditor(const FNiagaraTypeDefinition& ParameterType, EUnit DisplayUnit, const FNiagaraInputParameterCustomization& WidgetCustomization) const
 {
-	return SNew(SNiagaraQuatParameterEditor, DisplayUnit);
+	return SNew(SNiagaraQuatParameterEditor, DisplayUnit, WidgetCustomization);
 }
 
 bool FNiagaraEditorQuatTypeUtilities::CanHandlePinDefaults() const
diff --git a/Engine/Plugins/FX/Niagara/Source/NiagaraEditor/Private/TypeEditorUtilities/NiagaraVectorTypeEditorUtilities.h b/Engine/Plugins/FX/Niagara/Source/NiagaraEditor/Private/TypeEditorUtilities/NiagaraVectorTypeEditorUtilities.h
--- a/Engine/Plugins/FX/Niagara/Source/NiagaraEditor/Private/TypeEditorUtilities/NiagaraVectorTypeEditorUtilities.h
+++ b/Engine/Plugins/FX/Niagara/Source/NiagaraEditor/Private/TypeEditorUtilities/NiagaraVectorTypeEditorUtilities.h
@@ -12,7 +12,7 @@ class FNiagaraEditorVector2TypeUtilities : public FNiagaraEditorTypeUtilities
 public:
 	//~ INiagaraEditorTypeUtilities interface.
 	virtual bool CanCreateParameterEditor() const override { return true; }
-	virtual TSharedPtr<SNiagaraParameterEditor> CreateParameterEditor(const FNiagaraTypeDefinition& ParameterType, EUnit DisplayUnit) const override;
+	virtual TSharedPtr<SNiagaraParameterEditor> CreateParameterEditor(const FNiagaraTypeDefinition& ParameterType, EUnit DisplayUnit, const FNiagaraInputParameterCustomization& WidgetCustomization) const override;
 	virtual bool CanHandlePinDefaults() const override;
 	virtual FString GetPinDefaultStringFromValue(const FNiagaraVariable& AllocatedVariable) const override;
 	virtual bool SetValueFromPinDefaultString(const FString& StringValue, FNiagaraVariable& Variable) const override;
@@ -26,7 +26,7 @@ class FNiagaraEditorVector3TypeUtilities : public FNiagaraEditorTypeUtilities
 public:
 	//~ INiagaraEditorTypeUtilities interface.
 	virtual bool CanCreateParameterEditor() const override { return true; }
-	virtual TSharedPtr<SNiagaraParameterEditor> CreateParameterEditor(const FNiagaraTypeDefinition& ParameterType, EUnit DisplayUnit) const override;
+	virtual TSharedPtr<SNiagaraParameterEditor> CreateParameterEditor(const FNiagaraTypeDefinition& ParameterType, EUnit DisplayUnit, const FNiagaraInputParameterCustomization& WidgetCustomization) const override;
 	virtual bool CanHandlePinDefaults() const override;
 	virtual FString GetPinDefaultStringFromValue(const FNiagaraVariable& AllocatedVariable) const override;
 	virtual bool SetValueFromPinDefaultString(const FString& StringValue, FNiagaraVariable& Variable) const override;
@@ -40,7 +40,7 @@ class FNiagaraEditorVector4TypeUtilities : public FNiagaraEditorTypeUtilities
 public:
 	//~ INiagaraEditorTypeUtilities interface.
 	virtual bool CanCreateParameterEditor() const override { return true; }
-	virtual TSharedPtr<SNiagaraParameterEditor> CreateParameterEditor(const FNiagaraTypeDefinition& ParameterType, EUnit DisplayUnit) const override;
+	virtual TSharedPtr<SNiagaraParameterEditor> CreateParameterEditor(const FNiagaraTypeDefinition& ParameterType, EUnit DisplayUnit, const FNiagaraInputParameterCustomization& WidgetCustomization) const override;
 	virtual bool CanHandlePinDefaults() const override;
 	virtual FString GetPinDefaultStringFromValue(const FNiagaraVariable& AllocatedVariable) const override;
 	virtual bool SetValueFromPinDefaultString(const FString& StringValue, FNiagaraVariable& Variable) const override;
@@ -56,7 +56,7 @@ public:
 	virtual bool CanProvideDefaultValue() const override { return true; }
 	virtual void UpdateVariableWithDefaultValue(FNiagaraVariable& Variable) const override;
 	virtual bool CanCreateParameterEditor() const override { return true; }
-	virtual TSharedPtr<SNiagaraParameterEditor> CreateParameterEditor(const FNiagaraTypeDefinition& ParameterType, EUnit DisplayUnit) const override;
+	virtual TSharedPtr<SNiagaraParameterEditor> CreateParameterEditor(const FNiagaraTypeDefinition& ParameterType, EUnit DisplayUnit, const FNiagaraInputParameterCustomization& WidgetCustomization) const override;
 	virtual bool CanHandlePinDefaults() const override;
 	virtual FString GetPinDefaultStringFromValue(const FNiagaraVariable& AllocatedVariable) const override;
 	virtual bool SetValueFromPinDefaultString(const FString& StringValue, FNiagaraVariable& Variable) const override;
diff --git a/Engine/Plugins/FX/Niagara/Source/NiagaraEditor/Private/ViewModels/Stack/NiagaraStackFunctionInput.cpp b/Engine/Plugins/FX/Niagara/Source/NiagaraEditor/Private/ViewModels/Stack/NiagaraStackFunctionInput.cpp
--- a/Engine/Plugins/FX/Niagara/Source/NiagaraEditor/Private/ViewModels/Stack/NiagaraStackFunctionInput.cpp
+++ b/Engine/Plugins/FX/Niagara/Source/NiagaraEditor/Private/ViewModels/Stack/NiagaraStackFunctionInput.cpp
@@ -235,6 +235,11 @@ EUnit UNiagaraStackFunctionInput::GetInputDisplayUnit() const
 	return InputMetaData.IsSet() ? InputMetaData->DisplayUnit : EUnit::Unspecified;
 }
 
+FNiagaraInputParameterCustomization UNiagaraStackFunctionInput::GetInputWidgetCustomization() const
+{
+	return InputMetaData.IsSet() ? InputMetaData->WidgetCustomization : FNiagaraInputParameterCustomization();
+}
+
 FText UNiagaraStackFunctionInput::GetTooltipText() const
 {
 	FText Description = SummaryViewTooltipOverride.IsSet() && !SummaryViewTooltipOverride->Get().IsEmptyOrWhitespace() ? SummaryViewTooltipOverride->Get() : InputMetaData.IsSet() ? InputMetaData->Description : FText::GetEmpty();
diff --git a/Engine/Plugins/FX/Niagara/Source/NiagaraEditor/Private/Widgets/SNiagaraNumericDropDown.h b/Engine/Plugins/FX/Niagara/Source/NiagaraEditor/Private/Widgets/SNiagaraNumericDropDown.h
new file mode 100644
--- /dev/null
+++ b/Engine/Plugins/FX/Niagara/Source/NiagaraEditor/Private/Widgets/SNiagaraNumericDropDown.h
@@ -0,0 +1,180 @@
+// Copyright Epic Games, Inc. All Rights Reserved.
+
+#pragma once
+
+#include "CoreMinimal.h"
+#include "NiagaraTypes.h"
+#include "NiagaraEditorStyle.h"
+#include "EdGraphSchema_Niagara.h"
+#include "Framework/Application/SlateApplication.h"
+#include "Framework/Commands/UIAction.h"
+#include "Framework/MultiBox/MultiBoxBuilder.h"
+#include "Input/Events.h"
+#include "Layout/Margin.h"
+#include "Layout/Visibility.h"
+#include "Misc/Attribute.h"
+#include "Textures/SlateIcon.h"
+#include "Widgets/DeclarativeSyntaxSupport.h"
+#include "Widgets/SBoxPanel.h"
+#include "Widgets/SCompoundWidget.h"
+#include "Widgets/SWidget.h"
+#include "Widgets/Input/SComboButton.h"
+#include "Widgets/Input/SEditableTextBox.h"
+#include "Widgets/Text/STextBlock.h"
+
+/** A widget which allows the user to enter a digit or choose a number from a drop down menu. Adapted from SNumericDropDown. */
+template<typename NumericType>
+class SNiagaraNumericDropDown : public SCompoundWidget
+{
+public:
+	DECLARE_DELEGATE_OneParam( FOnValueChanged, NumericType );
+
+public:
+	/** Represents a named numeric value for display in the drop down menu. */
+	class FNamedValue
+	{
+	public:
+		/** Creates a new FNamedValue 
+		  * @InValue The numeric value to be assigned
+		  * @InName The display text for the value in the UI 
+		  * @InDescription The description of the value used in tooltips or wherever a longer description is needed. */
+		FNamedValue( NumericType InValue, FText InName, FText InDescription )
+		{
+			Value = InValue;
+			Name = InName;
+			Description = InDescription;
+		}
+
+		NumericType GetValue() const
+		{
+			return Value;
+		}
+
+		FText GetName() const
+		{
+			return Name;
+		}
+
+		FText GetDescription() const
+		{
+			return Description;
+		}
+
+	private:
+		NumericType Value;
+		FText Name;
+		FText Description;
+	};
+
+public:
+	SLATE_BEGIN_ARGS( SNiagaraNumericDropDown<NumericType> )
+		: _DropDownValues()
+		, _MinDesiredValueWidth( 40 )
+		, _bShowNamedValue(false)
+		{}
+
+		/** The values which are used to populate the drop down menu. */
+		SLATE_ARGUMENT( TArray<FNamedValue>, DropDownValues )
+		/** Controls the minimum width for the text box portion of the control. */
+		SLATE_ATTRIBUTE( float, MinDesiredValueWidth )
+		/** Toggle to show the drop down text value if the value matches the numeric value. */
+		SLATE_ATTRIBUTE( bool, bShowNamedValue )
+		/** The type to display for the pill icon. */
+		SLATE_ATTRIBUTE( FNiagaraTypeDefinition, PillType )
+		/** The value displayed by the control. */
+		SLATE_ATTRIBUTE( NumericType, Value )
+		/** The callback for when the value changes. */
+		SLATE_EVENT( FOnValueChanged, OnValueChanged )
+
+	SLATE_END_ARGS()
+
+	void Construct( const FArguments& InArgs )
+	{
+		DropDownValues = InArgs._DropDownValues;
+		bShowNamedValue = InArgs._bShowNamedValue;
+		Value = InArgs._Value;
+		OnValueChanged = InArgs._OnValueChanged;
+
+		ChildSlot
+		[
+			SNew( SHorizontalBox )
+			+ SHorizontalBox::Slot()
+			.VAlign( EVerticalAlignment::VAlign_Center )
+			.AutoWidth()
+			.Padding( FMargin( 0, 0, 3, 0 ) )
+			[
+				SNew(SImage)
+				.ColorAndOpacity(UEdGraphSchema_Niagara::GetTypeColor(InArgs._PillType.Get(FNiagaraTypeDefinition::GetFloatDef())))
+				.Image(FNiagaraEditorStyle::Get().GetBrush("NiagaraEditor.Module.TypeIconPill"))
+			]
+			+ SHorizontalBox::Slot()
+			.AutoWidth()
+			[
+				SNew( SComboButton )
+				.OnGetMenuContent( this, &SNiagaraNumericDropDown<NumericType>::BuildMenu )
+				.ButtonContent()
+				[
+					SNew(SEditableTextBox)
+					.Padding(0)
+					.Style(FNiagaraEditorStyle::Get(), "NiagaraEditor.Stack.NumericDropdownInput")
+					.MinDesiredWidth( InArgs._MinDesiredValueWidth )
+					.RevertTextOnEscape(true)
+					.SelectAllTextWhenFocused(true)
+					.Text( this, &SNiagaraNumericDropDown<NumericType>::GetValueText )
+					.OnTextCommitted( this, &SNiagaraNumericDropDown<NumericType>::ValueTextComitted )
+				]
+			]
+		];
+	}
+
+private:
+
+	FText GetValueText() const
+	{
+		if (bShowNamedValue.Get())
+		{
+			for ( FNamedValue DropDownValue : DropDownValues )
+			{
+				if (FMath::IsNearlyEqual(static_cast<float>(DropDownValue.GetValue()), static_cast<float>(Value.Get())))
+				{
+					return DropDownValue.GetName();
+				}
+			}
+		}
+		return FText::AsNumber( Value.Get() );
+	}
+
+	void ValueTextComitted( const FText& InNewText, ETextCommit::Type InTextCommit )
+	{
+		if ( InNewText.IsNumeric() )
+		{
+			NumericType NewValue;
+			TTypeFromString<NumericType>::FromString( NewValue, *InNewText.ToString() );
+			OnValueChanged.ExecuteIfBound(NewValue);
+		}
+	}
+
+	TSharedRef<SWidget> BuildMenu()
+	{
+		FMenuBuilder MenuBuilder( true, NULL );
+
+		for ( FNamedValue DropDownValue : DropDownValues )
+		{
+			FUIAction MenuAction(FExecuteAction::CreateSP(this, &SNiagaraNumericDropDown::SetValue, DropDownValue.GetValue()));
+			MenuBuilder.AddMenuEntry(DropDownValue.GetName(), DropDownValue.GetDescription(), FSlateIcon(), MenuAction);
+		}
+
+		return MenuBuilder.MakeWidget();
+	}
+
+	void SetValue( NumericType InValue )
+	{
+		FSlateApplication::Get().ClearKeyboardFocus( EFocusCause::Cleared );
+		OnValueChanged.ExecuteIfBound(InValue);
+	}
+
+	TArray<FNamedValue> DropDownValues;
+	TAttribute<bool> bShowNamedValue;
+	TAttribute<NumericType> Value;
+	FOnValueChanged OnValueChanged;
+};
diff --git a/Engine/Plugins/FX/Niagara/Source/NiagaraEditor/Public/INiagaraEditorTypeUtilities.h b/Engine/Plugins/FX/Niagara/Source/NiagaraEditor/Public/INiagaraEditorTypeUtilities.h
--- a/Engine/Plugins/FX/Niagara/Source/NiagaraEditor/Public/INiagaraEditorTypeUtilities.h
+++ b/Engine/Plugins/FX/Niagara/Source/NiagaraEditor/Public/INiagaraEditorTypeUtilities.h
@@ -5,12 +5,14 @@
 #include "Delegates/Delegate.h"
 #include "Internationalization/Text.h"
 #include "Math/UnitConversion.h"
+#include "NiagaraTypes.h"
 
 class FStructOnScope;
 class SNiagaraParameterEditor;
 class SWidget;
 struct FNiagaraTypeDefinition;
 struct FNiagaraVariable;
+struct FNiagaraInputParameterCustomization;
 
 class INiagaraEditorTypeUtilities
 {
@@ -25,7 +27,7 @@ public:
 
 	virtual bool CanCreateParameterEditor() const = 0;
 
-	virtual TSharedPtr<SNiagaraParameterEditor> CreateParameterEditor(const FNiagaraTypeDefinition& ParameterType, EUnit DisplayUnit = EUnit::Unspecified) const = 0;
+	virtual TSharedPtr<SNiagaraParameterEditor> CreateParameterEditor(const FNiagaraTypeDefinition& ParameterType, EUnit DisplayUnit = EUnit::Unspecified, const FNiagaraInputParameterCustomization& WidgetCustomization = FNiagaraInputParameterCustomization()) const = 0;
 
 	virtual bool CanCreateDataInterfaceEditor() const = 0;
 
@@ -55,7 +57,7 @@ public:
 	virtual bool CanProvideDefaultValue() const override { return false; }
 	virtual void UpdateVariableWithDefaultValue(FNiagaraVariable& Variable) const override { }
 	virtual bool CanCreateParameterEditor() const override { return false; }
-	virtual TSharedPtr<SNiagaraParameterEditor> CreateParameterEditor(const FNiagaraTypeDefinition& ParameterType, EUnit DisplayUnit) const override { return TSharedPtr<SNiagaraParameterEditor>(); }
+	virtual TSharedPtr<SNiagaraParameterEditor> CreateParameterEditor(const FNiagaraTypeDefinition& ParameterType, EUnit DisplayUnit, const FNiagaraInputParameterCustomization& WidgetCustomization) const override { return TSharedPtr<SNiagaraParameterEditor>(); }
 	virtual bool CanCreateDataInterfaceEditor() const override { return false; };
 	virtual TSharedPtr<SWidget> CreateDataInterfaceEditor(UObject* DataInterface, FNotifyValueChanged DataInterfaceChangedHandler) const override { return TSharedPtr<SWidget>(); }
 	virtual bool CanHandlePinDefaults() const override { return false; }
diff --git a/Engine/Plugins/FX/Niagara/Source/NiagaraEditor/Public/NiagaraEditorStyle.h b/Engine/Plugins/FX/Niagara/Source/NiagaraEditor/Public/NiagaraEditorStyle.h
--- a/Engine/Plugins/FX/Niagara/Source/NiagaraEditor/Public/NiagaraEditorStyle.h
+++ b/Engine/Plugins/FX/Niagara/Source/NiagaraEditor/Public/NiagaraEditorStyle.h
@@ -41,6 +41,7 @@ private:
 	void InitThumbnails();
 	void InitClassIcon();
 	void InitStackIcons();
+	void InitStackWidgets();
 	void InitNiagaraSequence();
 	void InitPlatformSet();
 	void InitDropTarget();
diff --git a/Engine/Plugins/FX/Niagara/Source/NiagaraEditor/Public/ViewModels/Stack/NiagaraStackFunctionInput.h b/Engine/Plugins/FX/Niagara/Source/NiagaraEditor/Public/ViewModels/Stack/NiagaraStackFunctionInput.h
--- a/Engine/Plugins/FX/Niagara/Source/NiagaraEditor/Public/ViewModels/Stack/NiagaraStackFunctionInput.h
+++ b/Engine/Plugins/FX/Niagara/Source/NiagaraEditor/Public/ViewModels/Stack/NiagaraStackFunctionInput.h
@@ -97,6 +97,8 @@ public:
 	/** Gets the unit of this input. */
 	EUnit GetInputDisplayUnit() const;
 
+	FNiagaraInputParameterCustomization GetInputWidgetCustomization() const;
+
 	//~ UNiagaraStackEntry interface
 	virtual FText GetDisplayName() const override;
 	virtual FText GetTooltipText() const override;
diff --git a/Engine/Plugins/FX/Niagara/Source/NiagaraEditorWidgets/Private/Stack/SNiagaraStackFunctionInputValue.cpp b/Engine/Plugins/FX/Niagara/Source/NiagaraEditorWidgets/Private/Stack/SNiagaraStackFunctionInputValue.cpp
--- a/Engine/Plugins/FX/Niagara/Source/NiagaraEditorWidgets/Private/Stack/SNiagaraStackFunctionInputValue.cpp
+++ b/Engine/Plugins/FX/Niagara/Source/NiagaraEditorWidgets/Private/Stack/SNiagaraStackFunctionInputValue.cpp
@@ -522,7 +522,7 @@ TSharedRef<SWidget> SNiagaraStackFunctionInputValue::ConstructLocalValueStructWi
 		TSharedPtr<INiagaraEditorTypeUtilities, ESPMode::ThreadSafe> TypeEditorUtilities = NiagaraEditorModule.GetTypeUtilities(FunctionInput->GetInputType());
 		if (TypeEditorUtilities.IsValid() && TypeEditorUtilities->CanCreateParameterEditor())
 		{
-			TSharedPtr<SNiagaraParameterEditor> ParameterEditor = TypeEditorUtilities->CreateParameterEditor(FunctionInput->GetInputType(), FunctionInput->GetInputDisplayUnit());
+			TSharedPtr<SNiagaraParameterEditor> ParameterEditor = TypeEditorUtilities->CreateParameterEditor(FunctionInput->GetInputType(), FunctionInput->GetInputDisplayUnit(), FunctionInput->GetInputWidgetCustomization());
 			ParameterEditor->UpdateInternalValueFromStruct(DisplayedLocalValueStruct.ToSharedRef());
 			ParameterEditor->SetOnBeginValueChange(SNiagaraParameterEditor::FOnValueChange::CreateSP(
 				this, &SNiagaraStackFunctionInputValue::ParameterBeginValueChange));
