Maintain existing common properties when changing instnaced data channel type.

diff --git a/Engine/Plugins/FX/Niagara/Source/Niagara/Private/NiagaraDataChannel.cpp b/Engine/Plugins/FX/Niagara/Source/Niagara/Private/NiagaraDataChannel.cpp
--- a/Engine/Plugins/FX/Niagara/Source/Niagara/Private/NiagaraDataChannel.cpp
+++ b/Engine/Plugins/FX/Niagara/Source/Niagara/Private/NiagaraDataChannel.cpp
@@ -11,6 +11,7 @@
 #include "UObject/Package.h"
 #include "NiagaraDataChannelAccessor.h"
 #include "NiagaraGpuComputeDispatchInterface.h"
+#include "Engine/Engine.h"
 
 #include UE_INLINE_GENERATED_CPP_BY_NAME(NiagaraDataChannel)
 
@@ -1011,6 +1012,31 @@ UNiagaraDataChannelReader* UNiagaraDataChannelLibrary::ReadFromNiagaraDataChanne
 	return nullptr;
 }
 
+//////////////////////////////////////////////////////////////////////////
+
+#if WITH_EDITOR
+
+void UNiagaraDataChannelAsset::PreEditChange(class FProperty* PropertyAboutToChange)
+{
+	if(PropertyAboutToChange->GetFName() == GET_MEMBER_NAME_CHECKED(UNiagaraDataChannelAsset, DataChannel))
+	{
+		CachedPreChangeDataChannel = DataChannel;
+	}
+}
+
+void UNiagaraDataChannelAsset::PostEditChangeProperty(FPropertyChangedEvent& PropertyChangedEvent)
+{
+	if(CachedPreChangeDataChannel)
+	{
+		UEngine::FCopyPropertiesForUnrelatedObjectsParams Params;
+		UEngine::CopyPropertiesForUnrelatedObjects(CachedPreChangeDataChannel, DataChannel, Params);
+		CachedPreChangeDataChannel = nullptr;
+	}
+}
+
+#endif
+
+
 //////////////////////////////////////////////////////////////////////////
 
 #if !UE_BUILD_SHIPPING
diff --git a/Engine/Plugins/FX/Niagara/Source/Niagara/Public/NiagaraDataChannelPublic.h b/Engine/Plugins/FX/Niagara/Source/Niagara/Public/NiagaraDataChannelPublic.h
--- a/Engine/Plugins/FX/Niagara/Source/Niagara/Public/NiagaraDataChannelPublic.h
+++ b/Engine/Plugins/FX/Niagara/Source/Niagara/Public/NiagaraDataChannelPublic.h
@@ -23,9 +23,20 @@ class NIAGARA_API UNiagaraDataChannelAsset : public UObject
 
 	UPROPERTY(EditAnywhere, Category = DataChannel, Instanced)
 	TObjectPtr<UNiagaraDataChannel> DataChannel;
+	
+	#if WITH_EDITORONLY_DATA
+	/** When changing data channel types we cache the old channel and attempt to copy over any common properties from one to the other. */
+	UPROPERTY(Transient)
+	TObjectPtr<UNiagaraDataChannel> CachedPreChangeDataChannel;
+	#endif
 
 public:
 
+#if WITH_EDITOR
+	virtual void PreEditChange(class FProperty* PropertyAboutToChange)override;
+	virtual void PostEditChangeProperty(FPropertyChangedEvent& PropertyChangedEvent)override;
+#endif
+
 	UNiagaraDataChannel* Get() const { return DataChannel; }
 };
 
