Some basic debug drawing for island data channels.

diff --git a/Engine/Plugins/FX/Niagara/Source/Niagara/Private/NiagaraDataChannel_Islands.cpp b/Engine/Plugins/FX/Niagara/Source/Niagara/Private/NiagaraDataChannel_Islands.cpp
--- a/Engine/Plugins/FX/Niagara/Source/Niagara/Private/NiagaraDataChannel_Islands.cpp
+++ b/Engine/Plugins/FX/Niagara/Source/Niagara/Private/NiagaraDataChannel_Islands.cpp
@@ -10,6 +10,8 @@
 #include "NiagaraWorldManager.h"
 #include "NiagaraDataChannelCommon.h"
 
+#include "DrawDebugHelpers.h"
+
 UNiagaraDataChannelHandler* UNiagaraDataChannel_Islands::CreateHandler(UWorld* OwningWorld)const
 {
 	UNiagaraDataChannelHandler* NewHandler = NewObject<UNiagaraDataChannelHandler_Islands>(OwningWorld);
@@ -86,6 +88,11 @@ void UNiagaraDataChannelHandler_Islands::Tick(float DeltaTime, ETickingGroup Tic
 		FNDCIsland& Island = IslandPool[*It];
 		Island.Tick();
 
+		if(IslandChannel->DebugDrawSettings.ShowBounds())
+		{
+			Island.DebugDrawBounds();
+		}
+
 		//Free up islands that are not being used.
 		//We do this immediately as we leave it up to the handler systems to delay their cleanup as long as they want to.
 		if (Island.IsBeingUsed() == false)
@@ -411,3 +418,8 @@ bool FNDCIsland::TryGrow(FVector Point, FVector PerElementExtents, FVector MaxIs
 
 	return true;
 }
+
+void FNDCIsland::DebugDrawBounds()
+{
+	DrawDebugBox(Owner->GetWorld(), Bounds.Origin, Bounds.BoxExtent, FQuat::Identity, FColor::Red);
+}
\ No newline at end of file
diff --git a/Engine/Plugins/FX/Niagara/Source/Niagara/Public/NiagaraDataChannel_Islands.h b/Engine/Plugins/FX/Niagara/Source/Niagara/Public/NiagaraDataChannel_Islands.h
--- a/Engine/Plugins/FX/Niagara/Source/Niagara/Public/NiagaraDataChannel_Islands.h
+++ b/Engine/Plugins/FX/Niagara/Source/Niagara/Public/NiagaraDataChannel_Islands.h
@@ -57,6 +57,9 @@ public:
 	bool IsBeingUsed()const;
 
 	FNiagaraDataChannelDataPtr GetData()const { return Data; }
+
+	void DebugDrawBounds();
+
 private:
 
 	/** The owning handler for this island. */
@@ -84,6 +87,25 @@ struct TStructOpsTypeTraits<FNDCIsland> : public TStructOpsTypeTraitsBase2<FNDCI
 	enum{ WithCopy = false };
 };
 
+USTRUCT()
+struct FNDCIslandDebugDrawSettings
+{
+	GENERATED_BODY()
+	
+	FNDCIslandDebugDrawSettings()
+	: bEnabled(false)
+	, bShowIslandBounds(false)
+	{}
+
+	UPROPERTY(EditAnywhere, Category = "Debug Drawing")
+	uint32 bEnabled : 1;
+
+	UPROPERTY(EditAnywhere, Category = "Debug Drawing")
+	uint32 bShowIslandBounds : 1;
+
+	bool ShowBounds()const { return bEnabled && bShowIslandBounds; }
+};
+
 UCLASS(Experimental)
 class NIAGARA_API UNiagaraDataChannel_Islands : public UNiagaraDataChannel
 {
@@ -133,6 +155,9 @@ public:
 	/** How many pre-allocated islands to keep in the pool. Higher values will incur a larger standing memory cost but will reduce activation times for new islands. */
 	UPROPERTY(EditAnywhere, Category = "Islands")
 	int32 IslandPoolSize = 16;
+
+	UPROPERTY(EditAnywhere, Category = "Debug Rendering")
+	FNDCIslandDebugDrawSettings DebugDrawSettings;
 };
 
 UCLASS(Experimental, BlueprintType)
