Remove Data Channels from all world managers on asset pre delete.
Fixes issue preventing data channel deletion as they objects are always referenced in the editor.

diff --git a/Engine/Plugins/FX/Niagara/Source/Niagara/Private/NiagaraDataChannel.cpp b/Engine/Plugins/FX/Niagara/Source/Niagara/Private/NiagaraDataChannel.cpp
--- a/Engine/Plugins/FX/Niagara/Source/Niagara/Private/NiagaraDataChannel.cpp
+++ b/Engine/Plugins/FX/Niagara/Source/Niagara/Private/NiagaraDataChannel.cpp
@@ -783,8 +783,7 @@ void UNiagaraDataChannel::PostInitProperties()
 	FNiagaraWorldManager::ForAllWorldManagers(
 		[DataChannel = this](FNiagaraWorldManager& WorldMan)
 		{
-			WorldMan.WaitForAsyncWork();
-			WorldMan.GetDataChannelManager().InitDataChannel(DataChannel, true);
+			WorldMan.InitDataChannel(DataChannel, true);
 		});
 }
 
@@ -825,8 +824,7 @@ void UNiagaraDataChannel::PostLoad()
 	FNiagaraWorldManager::ForAllWorldManagers(
 		[DataChannel = this](FNiagaraWorldManager& WorldMan)
 		{
-			WorldMan.WaitForAsyncWork();
-			WorldMan.GetDataChannelManager().InitDataChannel(DataChannel, true);
+			WorldMan.InitDataChannel(DataChannel, true);
 		});
 }
 
@@ -837,8 +835,7 @@ void UNiagaraDataChannel::BeginDestroy()
 	FNiagaraWorldManager::ForAllWorldManagers(
 		[DataChannel = this](FNiagaraWorldManager& WorldMan)
 		{
-			WorldMan.WaitForAsyncWork();
-			WorldMan.GetDataChannelManager().RemoveDataChannel(DataChannel);
+			WorldMan.RemoveDataChannel(DataChannel);
 		});
 }
 
@@ -849,8 +846,7 @@ void UNiagaraDataChannel::PreEditChange(FProperty* PropertyAboutToChange)
 	FNiagaraWorldManager::ForAllWorldManagers(
 		[DataChannel = this](FNiagaraWorldManager& WorldMan)
 		{
-			WorldMan.WaitForAsyncWork();
-			WorldMan.GetDataChannelManager().RemoveDataChannel(DataChannel);
+			WorldMan.RemoveDataChannel(DataChannel);
 		});
 }
 
@@ -871,7 +867,7 @@ void UNiagaraDataChannel::PostEditChangeProperty(FPropertyChangedEvent& Property
 	FNiagaraWorldManager::ForAllWorldManagers(
 		[DataChannel = this](FNiagaraWorldManager& WorldMan)
 		{
-			WorldMan.GetDataChannelManager().InitDataChannel(DataChannel, true);
+			WorldMan.InitDataChannel(DataChannel, true);
 		});
 }
 
diff --git a/Engine/Plugins/FX/Niagara/Source/Niagara/Private/NiagaraWorldManager.cpp b/Engine/Plugins/FX/Niagara/Source/Niagara/Private/NiagaraWorldManager.cpp
--- a/Engine/Plugins/FX/Niagara/Source/Niagara/Private/NiagaraWorldManager.cpp
+++ b/Engine/Plugins/FX/Niagara/Source/Niagara/Private/NiagaraWorldManager.cpp
@@ -37,7 +37,7 @@ DECLARE_CYCLE_STAT(TEXT("Niagara Manager Wait On Render [GT]"), STAT_NiagaraWorl
 DECLARE_CYCLE_STAT(TEXT("Niagara Manager Wait Pre Garbage Collect [GT]"), STAT_NiagaraWorldManWaitPreGC, STATGROUP_Niagara); 
 DECLARE_CYCLE_STAT(TEXT("Niagara Manager Refresh Owner Allows Scalability"), STAT_NiagaraWorldManRefreshOwnerAllowsScalability, STATGROUP_Niagara);
 
-static int GNiagaraAllowAsyncWorkToEndOfFrame = 1;
+static int GNiagaraAllowAsyncWorkToEndOfFrame = 0;
 static FAutoConsoleVariableRef CVarNiagaraAllowAsyncWorkToEndOfFrame(
 	TEXT("fx.Niagara.AllowAsyncWorkToEndOfFrame"),
 	GNiagaraAllowAsyncWorkToEndOfFrame,
@@ -730,6 +730,18 @@ void FNiagaraWorldManager::RefreshOwnerAllowsScalability()
 	}
 }
 
+void FNiagaraWorldManager::InitDataChannel(const UNiagaraDataChannel* Channel, bool bForce)
+{
+	WaitForAsyncWork();
+	GetDataChannelManager().InitDataChannel(Channel, bForce);
+}
+
+void FNiagaraWorldManager::RemoveDataChannel(const UNiagaraDataChannel* Channel)
+{
+	WaitForAsyncWork();
+	GetDataChannelManager().RemoveDataChannel(Channel);
+}
+
 void FNiagaraWorldManager::WaitForAsyncWork()
 {
 	for (int TG = 0; TG < NiagaraNumTickGroups; ++TG)
diff --git a/Engine/Plugins/FX/Niagara/Source/Niagara/Public/NiagaraDataChannelHandler.h b/Engine/Plugins/FX/Niagara/Source/Niagara/Public/NiagaraDataChannelHandler.h
--- a/Engine/Plugins/FX/Niagara/Source/Niagara/Public/NiagaraDataChannelHandler.h
+++ b/Engine/Plugins/FX/Niagara/Source/Niagara/Public/NiagaraDataChannelHandler.h
@@ -46,7 +46,7 @@ public:
 	UNiagaraDataChannelReader* GetDataChannelReader();
 
 	template<typename T> 
-	T* GetChannelTyped()const{ return CastChecked<T>(DataChannel); }
+	T* GetChannelTyped()const{ return Cast<T>(DataChannel); }
 
 	FNiagaraDataChannelDataPtr CreateData();
 
diff --git a/Engine/Plugins/FX/Niagara/Source/Niagara/Public/NiagaraWorldManager.h b/Engine/Plugins/FX/Niagara/Source/Niagara/Public/NiagaraWorldManager.h
--- a/Engine/Plugins/FX/Niagara/Source/Niagara/Public/NiagaraWorldManager.h
+++ b/Engine/Plugins/FX/Niagara/Source/Niagara/Public/NiagaraWorldManager.h
@@ -221,6 +221,8 @@ public:
 	NIAGARA_API static ENiagaraScalabilityCullingMode GetScalabilityCullingMode() { return ScalabilityCullingMode; }
 
 	FNiagaraDataChannelManager& GetDataChannelManager(){ return *DataChannelManager; }
+	NIAGARA_API void InitDataChannel(const UNiagaraDataChannel* Channel, bool bForce);
+	NIAGARA_API void RemoveDataChannel(const UNiagaraDataChannel* Channel);
 
 	/** Waits for all currently in flight async work to be completed. */
 	void WaitForAsyncWork();
@@ -289,7 +291,7 @@ private:
 	static FDelegateHandle PreGCBeginDestroyHandle;
 	static FDelegateHandle ViewTargetChangedHandle;
 
-	static TMap<class UWorld*, class FNiagaraWorldManager*> WorldManagers;
+	static NIAGARA_API TMap<class UWorld*, class FNiagaraWorldManager*> WorldManagers;
 
 	UWorld* World = nullptr;
 
diff --git a/Engine/Plugins/FX/Niagara/Source/NiagaraEditor/Private/NiagaraEditorModule.cpp b/Engine/Plugins/FX/Niagara/Source/NiagaraEditor/Private/NiagaraEditorModule.cpp
--- a/Engine/Plugins/FX/Niagara/Source/NiagaraEditor/Private/NiagaraEditorModule.cpp
+++ b/Engine/Plugins/FX/Niagara/Source/NiagaraEditor/Private/NiagaraEditorModule.cpp
@@ -929,6 +929,9 @@ void FNiagaraEditorModule::OnPreExit()
 #endif
 
 	UDeviceProfileManager::Get().OnManagerUpdated().Remove(DeviceProfileManagerUpdatedHandle);
+
+	FEditorDelegates::OnAssetsPreDelete.Remove(OnAssetsPreDeleteHandle);
+
 	if (GEditor)
 	{
 		CastChecked<UEditorEngine>(GEngine)->OnPreviewPlatformChanged().Remove(PreviewPlatformChangedHandle);
@@ -993,6 +996,8 @@ void FNiagaraEditorModule::StartupModule()
 
 	DeviceProfileManagerUpdatedHandle = UDeviceProfileManager::Get().OnManagerUpdated().AddRaw(this, &FNiagaraEditorModule::OnDeviceProfileManagerUpdated);
 	FCoreDelegates::OnEnginePreExit.AddRaw(this, &FNiagaraEditorModule::OnPreExit);
+
+	OnAssetsPreDeleteHandle = FEditorDelegates::OnAssetsPreDelete.AddRaw(this, &FNiagaraEditorModule::OnAssetsPreDelete);
 	
 	// register details customization
 	FPropertyEditorModule& PropertyModule = FModuleManager::LoadModuleChecked<FPropertyEditorModule>("PropertyEditor");
@@ -2103,6 +2108,21 @@ void FNiagaraEditorModule::OnAssetRegistryLoadComplete()
 	}
 }
 
+void FNiagaraEditorModule::OnAssetsPreDelete(const TArray<UObject*>& Objects)
+{
+	for(UObject* Obj : Objects)
+	{
+		if(UNiagaraDataChannelAsset* NDCAsset = Cast<UNiagaraDataChannelAsset>(Obj))
+		{
+			FNiagaraWorldManager::ForAllWorldManagers(
+				[DataChannel = NDCAsset->Get()](FNiagaraWorldManager& WorldMan)
+			{
+				WorldMan.RemoveDataChannel(DataChannel);
+			});
+		}
+	}
+}
+
 FNiagaraEditorModule::FOnScriptApplied& FNiagaraEditorModule::OnScriptApplied()
 {
 	return OnScriptAppliedDelegate;
diff --git a/Engine/Plugins/FX/Niagara/Source/NiagaraEditor/Public/NiagaraEditorModule.h b/Engine/Plugins/FX/Niagara/Source/NiagaraEditor/Public/NiagaraEditorModule.h
--- a/Engine/Plugins/FX/Niagara/Source/NiagaraEditor/Public/NiagaraEditorModule.h
+++ b/Engine/Plugins/FX/Niagara/Source/NiagaraEditor/Public/NiagaraEditorModule.h
@@ -336,6 +336,8 @@ private:
 	}
 
 	void OnAssetRegistryLoadComplete();
+	
+	void OnAssetsPreDelete(const TArray<UObject*>& Objects);
 
 private:
 	TSharedPtr<FExtensibilityManager> MenuExtensibilityManager;
@@ -372,6 +374,8 @@ private:
 
 	FDelegateHandle AssetRegistryOnLoadCompleteHandle;
 	
+	FDelegateHandle OnAssetsPreDeleteHandle;
+	
 	TObjectPtr<USequencerSettings> SequencerSettings;
 
 	TSharedPtr<INiagaraEditorWidgetProvider> WidgetProvider;
