Fixes and improvements to write DI.

Fixing Same Frame Reads. We currently write to local data and then publish in post simulate, however post simulate tick can be moved to the end to the frame which means it's impossible to get same frame reads in the subsequent tick phases.
- Moved Begin/End simulate of write DI dataset to new Pre/Post stage tick path.
- Write DI now publishes from post stage so made publishing to data channel data threadsafe.
- Write DI now grabs it's data channel data once in presimulate rather than in post simulate.
- Fixed issue preventing proper init of Pre/Post simulate tick lists.

Write and Append now updating NumInstances via atomic in post stage to fix threadsafety with parallem VM exec.

Added optimized path through Append() in cases where we're emitting every particle.
This required a slight tweak to the DI handlers.

diff --git a/Engine/Plugins/FX/Niagara/Source/Niagara/Classes/NiagaraDataInterface.h b/Engine/Plugins/FX/Niagara/Source/Niagara/Classes/NiagaraDataInterface.h
--- a/Engine/Plugins/FX/Niagara/Source/Niagara/Classes/NiagaraDataInterface.h
+++ b/Engine/Plugins/FX/Niagara/Source/Niagara/Classes/NiagaraDataInterface.h
@@ -496,7 +496,7 @@ private:
 /** Context for pre and post stage ticks for DIs on the CPU. A similar tick flow also occurs for stages on the GPU via the proxy and dispatcher. */
 struct FNDICpuPrePostStageContext
 {
-	const FNiagaraSystemInstance* SystemInstance = nullptr;
+	FNiagaraSystemInstance* SystemInstance = nullptr;
 	void* PerInstanceData = nullptr;
 	float DeltaSeconds = 0.0f;
 	ENiagaraScriptUsage Usage;
@@ -519,6 +519,11 @@ struct FNDIStageTickHandler
 	void Init(UNiagaraScript* Script, FNiagaraSystemInstance* Instance);
 	void PreStageTick(FNiagaraSystemInstance* Instance, float DeltaSeconds);
 	void PostStageTick(FNiagaraSystemInstance* Instance, float DeltaSeconds);
+	void Empty()
+	{
+		PreStageTickList.Empty();
+		PostStageTickList.Empty();
+	}
 };
 
 //////////////////////////////////////////////////////////////////////////
@@ -1078,7 +1083,7 @@ struct FNDIInputParam
 	FORCEINLINE void Init(FVectorVMExternalFunctionContext& Context){ new(this)FNDIInputParam<T>(Context); }
 	FORCEINLINE T GetAndAdvance() { return Data.GetAndAdvance(); }
 	FORCEINLINE T Get() { return Data.Get(); }
-	FORCEINLINE void Advance() { return Data.Advance(); }
+	FORCEINLINE void Advance(int32 Count=1) { return Data.Advance(Count); }
 	FORCEINLINE bool IsConstant() const { return Data.IsConstant(); }
 	FORCEINLINE void Reset() { Data.Reset(); }
 };
@@ -1092,7 +1097,7 @@ struct FNDIInputParam<FFloat16>
 	FORCEINLINE void Init(FVectorVMExternalFunctionContext& Context) { new(this)FNDIInputParam<FFloat16>(Context); }
 	FORCEINLINE FFloat16 GetAndAdvance() { return Data.GetAndAdvance(); }
 	FORCEINLINE FFloat16 Get() { return Data.Get(); }
-	FORCEINLINE void Advance() { return Data.Advance(); }
+	FORCEINLINE void Advance(int32 Count=1) { return Data.Advance(Count); }
 	FORCEINLINE bool IsConstant() const { return Data.IsConstant(); }
 	FORCEINLINE void Reset() { Data.Reset(); }
 };
@@ -1106,7 +1111,7 @@ struct FNDIInputParam<FNiagaraBool>
 	FORCEINLINE void Init(FVectorVMExternalFunctionContext& Context) { new(this)FNDIInputParam<FNiagaraBool>(Context); }
 	FORCEINLINE bool GetAndAdvance() { return Data.GetAndAdvance().GetValue(); }
 	FORCEINLINE bool Get() { return Data.Get().GetValue(); }
-	FORCEINLINE void Advance() { return Data.Advance(); }
+	FORCEINLINE void Advance(int32 Count=1) { return Data.Advance(Count); }
 	FORCEINLINE bool IsConstant() const { return Data.IsConstant(); }
 	FORCEINLINE void Reset() { Data.Reset(); }
 };
@@ -1120,7 +1125,7 @@ struct FNDIInputParam<bool>
 	FORCEINLINE void Init(FVectorVMExternalFunctionContext& Context) { new(this)FNDIInputParam<bool>(Context); }
 	FORCEINLINE bool GetAndAdvance() { return Data.GetAndAdvance().GetValue(); }
 	FORCEINLINE bool Get() { return Data.Get().GetValue(); }
-	FORCEINLINE void Advance() { return Data.Advance(); }
+	FORCEINLINE void Advance(int32 Count=1) { return Data.Advance(Count); }
 	FORCEINLINE bool IsConstant() const { return Data.IsConstant(); }
 	FORCEINLINE void Reset() { Data.Reset(); }
 };
@@ -1135,7 +1140,7 @@ struct FNDIInputParam<FVector2f>
 	FORCEINLINE void Init(FVectorVMExternalFunctionContext& Context) { new(this)FNDIInputParam<FVector2f>(Context); }
 	FORCEINLINE FVector2f GetAndAdvance() { return FVector2f(X.GetAndAdvance(), Y.GetAndAdvance()); }
 	FORCEINLINE FVector2f Get() { return FVector2f(X.Get(), Y.Get()); }
-	FORCEINLINE void Advance() { X.Advance(), Y.Advance(); }
+	FORCEINLINE void Advance(int32 Count=1) { X.Advance(Count), Y.Advance(Count); }
 	FORCEINLINE bool IsConstant() const { return X.IsConstant() && Y.IsConstant(); }
 	FORCEINLINE void Reset() { X.Reset(); Y.Reset(); }
 };
@@ -1151,7 +1156,7 @@ struct FNDIInputParam<FVector3f>
 	FORCEINLINE void Init(FVectorVMExternalFunctionContext& Context) { new(this)FNDIInputParam<FVector3f>(Context); }
 	FORCEINLINE FVector3f GetAndAdvance() { return FVector3f(X.GetAndAdvance(), Y.GetAndAdvance(), Z.GetAndAdvance()); }
 	FORCEINLINE FVector3f Get() { return FVector3f(X.Get(), Y.Get(), Z.Get()); }
-	FORCEINLINE void Advance() { X.Advance(); Y.Advance(); Z.Advance(); }
+	FORCEINLINE void Advance(int32 Count=1) { X.Advance(Count); Y.Advance(Count); Z.Advance(Count); }
 	FORCEINLINE bool IsConstant() const { return X.IsConstant() && Y.IsConstant() && Z.IsConstant(); }
 	FORCEINLINE void Reset() { X.Reset(); Y.Reset(); Z.Reset(); }
 };
@@ -1167,7 +1172,7 @@ struct FNDIInputParam<FNiagaraPosition>
 	FORCEINLINE void Init(FVectorVMExternalFunctionContext& Context) { new(this)FNDIInputParam<FNiagaraPosition>(Context); }
 	FORCEINLINE FNiagaraPosition GetAndAdvance() { return FNiagaraPosition(X.GetAndAdvance(), Y.GetAndAdvance(), Z.GetAndAdvance()); }
 	FORCEINLINE FNiagaraPosition Get() { return FNiagaraPosition(X.Get(), Y.Get(), Z.Get()); }
-	FORCEINLINE void Advance() { X.Advance(); Y.Advance(); Z.Advance(); }
+	FORCEINLINE void Advance(int32 Count=1) { X.Advance(Count); Y.Advance(Count); Z.Advance(Count); }
 	FORCEINLINE void Reset() { X.Reset(); Y.Reset(); Z.Reset(); }
 };
 template<>
@@ -1182,7 +1187,7 @@ struct FNDIInputParam<FVector4f>
 	FORCEINLINE void Init(FVectorVMExternalFunctionContext& Context) { new(this)FNDIInputParam<FVector4f>(Context); }
 	FORCEINLINE FVector4f GetAndAdvance() { return FVector4f(X.GetAndAdvance(), Y.GetAndAdvance(), Z.GetAndAdvance(), W.GetAndAdvance()); }
 	FORCEINLINE FVector4f Get() { return FVector4f(X.Get(), Y.Get(), Z.Get(), W.Get()); }
-	FORCEINLINE void Advance() { return X.Advance(); Y.Advance(); Z.Advance(); W.Advance(); }
+	FORCEINLINE void Advance(int32 Count=1) { return X.Advance(Count); Y.Advance(Count); Z.Advance(Count); W.Advance(Count); }
 	FORCEINLINE bool IsConstant() const { return X.IsConstant() && Y.IsConstant() && Z.IsConstant() && W.IsConstant(); }
 	FORCEINLINE void Reset() { X.Reset(); Y.Reset(); Z.Reset(); W.Reset(); }
 };
@@ -1199,7 +1204,7 @@ struct FNDIInputParam<FQuat4f>
 	FORCEINLINE void Init(FVectorVMExternalFunctionContext& Context) { new(this)FNDIInputParam<FQuat4f>(Context); }
 	FORCEINLINE FQuat4f GetAndAdvance() { return FQuat4f(X.GetAndAdvance(), Y.GetAndAdvance(), Z.GetAndAdvance(), W.GetAndAdvance()); }
 	FORCEINLINE FQuat4f Get() { return FQuat4f(X.Get(), Y.Get(), Z.Get(), W.Get()); }
-	FORCEINLINE void Advance() { X.Advance(); Y.Advance(); Z.Advance(); W.Advance(); }
+	FORCEINLINE void Advance(int32 Count=1) { X.Advance(Count); Y.Advance(Count); Z.Advance(Count); W.Advance(Count); }
 	FORCEINLINE bool IsConstant() const { return X.IsConstant() && Y.IsConstant() && Z.IsConstant() && W.IsConstant(); }
 	FORCEINLINE void Reset() { X.Reset(); Y.Reset(); Z.Reset(); W.Reset(); }
 };
@@ -1216,7 +1221,7 @@ struct FNDIInputParam<FLinearColor>
 	FORCEINLINE void Init(FVectorVMExternalFunctionContext& Context) { new(this)FNDIInputParam<FLinearColor>(Context); }
 	FORCEINLINE FLinearColor GetAndAdvance() { return FLinearColor(R.GetAndAdvance(), G.GetAndAdvance(), B.GetAndAdvance(), A.GetAndAdvance()); }
 	FORCEINLINE FLinearColor Get() { return FLinearColor(R.Get(), G.Get(), B.Get(), A.Get()); }
-	FORCEINLINE void Advance() { R.Advance(); G.Advance(); B.Advance(); A.Advance(); }
+	FORCEINLINE void Advance(int32 Count=1) { R.Advance(Count); G.Advance(Count); B.Advance(Count); A.Advance(Count); }
 	FORCEINLINE bool IsConstant() const { return R.IsConstant() && G.IsConstant() && B.IsConstant() && A.IsConstant(); }
 	FORCEINLINE void Reset() { R.Reset(); G.Reset(); B.Reset(); A.Reset(); }
 };
@@ -1231,7 +1236,7 @@ struct FNDIInputParam<FNiagaraID>
 	FORCEINLINE void Init(FVectorVMExternalFunctionContext& Context) { new(this)FNDIInputParam<FNiagaraID>(Context); }
 	FORCEINLINE FNiagaraID GetAndAdvance() { return FNiagaraID(Index.GetAndAdvance(), AcquireTag.GetAndAdvance()); }
 	FORCEINLINE FNiagaraID Get() { return FNiagaraID(Index.Get(), AcquireTag.Get()); }
-	FORCEINLINE void Advance() { return Index.Advance(); AcquireTag.Advance(); }
+	FORCEINLINE void Advance(int32 Count=1) { return Index.Advance(Count); AcquireTag.Advance(Count); }
 	FORCEINLINE bool IsConstant() const { return Index.IsConstant() && AcquireTag.IsConstant(); }
 	FORCEINLINE void Reset() { Index.Reset(); AcquireTag.Reset(); }
 };
@@ -1248,7 +1253,7 @@ struct FNDIInputParam<FNiagaraSpawnInfo>
 	FORCEINLINE void Init(FVectorVMExternalFunctionContext& Context) { new(this)FNDIInputParam<FNiagaraSpawnInfo>(Context); }
 	FORCEINLINE FNiagaraSpawnInfo GetAndAdvance() { return FNiagaraSpawnInfo(Count.GetAndAdvance(), InterpStartDt.GetAndAdvance(),IntervalDt.GetAndAdvance(),SpawnGroup.GetAndAdvance()); }
 	FORCEINLINE FNiagaraSpawnInfo Get() { return FNiagaraSpawnInfo(Count.Get(), InterpStartDt.Get(),IntervalDt.Get(),SpawnGroup.Get()); }
-	FORCEINLINE void Advance() { return Count.Advance(); InterpStartDt.Advance();  IntervalDt.Advance();  SpawnGroup.Advance(); }
+	FORCEINLINE void Advance(int32 InCount=1) { return Count.Advance(InCount); InterpStartDt.Advance(InCount);  IntervalDt.Advance(InCount);  SpawnGroup.Advance(InCount); }
 	FORCEINLINE bool IsConstant() const { return Count.IsConstant() && InterpStartDt.IsConstant() && IntervalDt.IsConstant() && SpawnGroup.IsConstant(); }
 	FORCEINLINE void Reset() { Count.Reset(); InterpStartDt.Reset();  IntervalDt.Reset(); SpawnGroup.Reset();}
 };
diff --git a/Engine/Plugins/FX/Niagara/Source/Niagara/Classes/NiagaraEmitterInstance.h b/Engine/Plugins/FX/Niagara/Source/Niagara/Classes/NiagaraEmitterInstance.h
--- a/Engine/Plugins/FX/Niagara/Source/Niagara/Classes/NiagaraEmitterInstance.h
+++ b/Engine/Plugins/FX/Niagara/Source/Niagara/Classes/NiagaraEmitterInstance.h
@@ -48,6 +48,7 @@ public:
 	virtual ~FNiagaraEmitterInstance();
 
 	void Init(int32 InEmitterIdx, FNiagaraSystemInstanceID SystemInstanceID);
+	void InitDITickLists();
 
 	void ResetSimulation(bool bKillExisting = true);
 
diff --git a/Engine/Plugins/FX/Niagara/Source/Niagara/Classes/NiagaraScriptExecutionContext.h b/Engine/Plugins/FX/Niagara/Source/Niagara/Classes/NiagaraScriptExecutionContext.h
--- a/Engine/Plugins/FX/Niagara/Source/Niagara/Classes/NiagaraScriptExecutionContext.h
+++ b/Engine/Plugins/FX/Niagara/Source/Niagara/Classes/NiagaraScriptExecutionContext.h
@@ -171,6 +171,7 @@ struct FNiagaraScriptExecutionContextBase
 	virtual ~FNiagaraScriptExecutionContextBase();
 
 	virtual bool Init(class FNiagaraSystemInstance* Instance, UNiagaraScript* InScript, ENiagaraSimTarget InTarget);
+	virtual void InitDITickLists(class FNiagaraSystemInstance* Instance);
 	virtual bool Tick(class FNiagaraSystemInstance* Instance, ENiagaraSimTarget SimTarget) = 0;
 	virtual bool Execute(class FNiagaraSystemInstance* Instance, float DeltaSeconds, uint32 NumInstances, const FScriptExecutionConstantBufferTable& ConstantBufferTable);
 
@@ -181,6 +182,8 @@ struct FNiagaraScriptExecutionContextBase
 
 	TArrayView<const uint8> GetScriptLiterals() const;
 
+	FNDIStageTickHandler& GetNDIStageTickHander(){ return DIStageTickHandler; }
+
 	void DirtyDataInterfaces();
 	void PostTick();
 
diff --git a/Engine/Plugins/FX/Niagara/Source/Niagara/Internal/DataInterface/NiagaraDataInterfaceDataChannelCommon.h b/Engine/Plugins/FX/Niagara/Source/Niagara/Internal/DataInterface/NiagaraDataInterfaceDataChannelCommon.h
--- a/Engine/Plugins/FX/Niagara/Source/Niagara/Internal/DataInterface/NiagaraDataInterfaceDataChannelCommon.h
+++ b/Engine/Plugins/FX/Niagara/Source/Niagara/Internal/DataInterface/NiagaraDataInterfaceDataChannelCommon.h
@@ -244,15 +244,15 @@ struct FNDIVariadicInputHandler
 		for (FNDIInputParam<FFloat16>& Input : HalfInputs) { Input.Reset(); }
 	}
 
-	void Advance()
+	void Advance(int32 Count = 1)
 	{
-		for (FNDIInputParam<float>& Input : FloatInputs) { Input.Advance(); }
-		for (FNDIInputParam<int32>& Input : IntInputs) { Input.Advance(); }
-		for (FNDIInputParam<FFloat16>& Input : HalfInputs) { Input.Advance(); }
+		for (FNDIInputParam<float>& Input : FloatInputs) { Input.Advance(Count); }
+		for (FNDIInputParam<int32>& Input : IntInputs) { Input.Advance(Count); }
+		for (FNDIInputParam<FFloat16>& Input : HalfInputs) { Input.Advance(Count); }
 	}
 
 	template<typename TFloatAction, typename TIntAction, typename THalfAction>
-	bool Process(bool bProcess, const FNDIDataChannel_FunctionToDataSetBinding* BindingInfo, TFloatAction FloatFunc, TIntAction IntFunc, THalfAction HalfFunc)
+	bool Process(bool bProcess, int32 Count, const FNDIDataChannel_FunctionToDataSetBinding* BindingInfo, TFloatAction FloatFunc, TIntAction IntFunc, THalfAction HalfFunc)
 	{
 		if (BindingInfo && bProcess)
 		{
@@ -261,17 +261,17 @@ struct FNDIVariadicInputHandler
 			{
 				switch(VMBinding.DataType)
 				{
-					case (int32)ENiagaraBaseTypes::Float: FloatFunc(VMBinding, FloatInputs[VMBinding.FunctionRegisterIndex].GetAndAdvance()); break;
-					case (int32)ENiagaraBaseTypes::Int32: IntFunc(VMBinding, IntInputs[VMBinding.FunctionRegisterIndex].GetAndAdvance()); break;
-					case (int32)ENiagaraBaseTypes::Bool: IntFunc(VMBinding, IntInputs[VMBinding.FunctionRegisterIndex].GetAndAdvance()); break;
-					case (int32)ENiagaraBaseTypes::Half: HalfFunc(VMBinding, HalfInputs[VMBinding.FunctionRegisterIndex].GetAndAdvance()); break;
+					case (int32)ENiagaraBaseTypes::Float: FloatFunc(VMBinding, FloatInputs[VMBinding.FunctionRegisterIndex]); break;
+					case (int32)ENiagaraBaseTypes::Int32: IntFunc(VMBinding, IntInputs[VMBinding.FunctionRegisterIndex]); break;
+					case (int32)ENiagaraBaseTypes::Bool: IntFunc(VMBinding, IntInputs[VMBinding.FunctionRegisterIndex]); break;
+					case (int32)ENiagaraBaseTypes::Half: HalfFunc(VMBinding, HalfInputs[VMBinding.FunctionRegisterIndex]); break;
 					default: check(0);
 				};
 			}
 			return true;
 		}
 		
-		Advance();
+		Advance(Count);
 		return false;
 	}
 };
@@ -307,7 +307,7 @@ struct FNDIVariadicOutputHandler
 	}
 
 	template<typename TFloatAction, typename TIntAction, typename THalfAction>
-	bool Process(bool bProcess, const FNDIDataChannel_FunctionToDataSetBinding* BindingInfo, TFloatAction FloatFunc, TIntAction IntFunc, THalfAction HalfFunc)
+	bool Process(bool bProcess, int32 Count, const FNDIDataChannel_FunctionToDataSetBinding* BindingInfo, TFloatAction FloatFunc, TIntAction IntFunc, THalfAction HalfFunc)
 	{
 		if (BindingInfo && bProcess)
 		{
@@ -316,17 +316,17 @@ struct FNDIVariadicOutputHandler
 			{
 				switch (VMBinding.DataType)
 				{
-				case (int32)ENiagaraBaseTypes::Float: FloatFunc(VMBinding, FloatOutputs[VMBinding.FunctionRegisterIndex].GetDestAndAdvance()); break;
-				case (int32)ENiagaraBaseTypes::Int32: IntFunc(VMBinding, IntOutputs[VMBinding.FunctionRegisterIndex].GetDestAndAdvance()); break;
-				case (int32)ENiagaraBaseTypes::Bool: IntFunc(VMBinding, IntOutputs[VMBinding.FunctionRegisterIndex].GetDestAndAdvance()); break;
-				case (int32)ENiagaraBaseTypes::Half: HalfFunc(VMBinding, HalfOutputs[VMBinding.FunctionRegisterIndex].GetDestAndAdvance()); break;
+				case (int32)ENiagaraBaseTypes::Float: FloatFunc(VMBinding, FloatOutputs[VMBinding.FunctionRegisterIndex]); break;
+				case (int32)ENiagaraBaseTypes::Int32: IntFunc(VMBinding, IntOutputs[VMBinding.FunctionRegisterIndex]); break;
+				case (int32)ENiagaraBaseTypes::Bool: IntFunc(VMBinding, IntOutputs[VMBinding.FunctionRegisterIndex]); break;
+				case (int32)ENiagaraBaseTypes::Half: HalfFunc(VMBinding, HalfOutputs[VMBinding.FunctionRegisterIndex]); break;
 				default: check(0);
 				};
 			}
 			return true;
 		}
 
-		Fallback(1);
+		Fallback(Count);
 
 		return false;
 	}
diff --git a/Engine/Plugins/FX/Niagara/Source/Niagara/Internal/DataInterface/NiagaraDataInterfaceDataChannelWrite.h b/Engine/Plugins/FX/Niagara/Source/Niagara/Internal/DataInterface/NiagaraDataInterfaceDataChannelWrite.h
--- a/Engine/Plugins/FX/Niagara/Source/Niagara/Internal/DataInterface/NiagaraDataInterfaceDataChannelWrite.h
+++ b/Engine/Plugins/FX/Niagara/Source/Niagara/Internal/DataInterface/NiagaraDataInterfaceDataChannelWrite.h
@@ -104,9 +104,12 @@ public:
 	virtual void DestroyPerInstanceData(void* PerInstanceData, FNiagaraSystemInstance* SystemInstance) override;
 	virtual int32 PerInstanceDataSize() const override;
 	virtual bool HasPreSimulateTick() const override { return true; }
-	virtual bool HasPostSimulateTick() const override { return true; }
+	virtual bool HasPreStageTick(ENiagaraScriptUsage Usage) const override { return true; }
+	virtual bool HasPostStageTick(ENiagaraScriptUsage Usage) const override { return true; }
 	virtual bool PerInstanceTick(void* PerInstanceData, FNiagaraSystemInstance* SystemInstance, float DeltaSeconds) override;
 	virtual bool PerInstanceTickPostSimulate(void* PerInstanceData, FNiagaraSystemInstance* SystemInstance, float DeltaSeconds) override;
+	virtual void PreStageTick(FNDICpuPreStageContext& Context) override;
+	virtual void PostStageTick(FNDICpuPostStageContext& Context) override;
 	virtual void ProvidePerInstanceDataForRenderThread(void* DataForRenderThread, void* PerInstanceData, const FNiagaraSystemInstanceID& SystemInstance) override;
 
 	//We cannot overlapt frames as we must correctly sync up with the data channel manager on Begin/End frame etc.
diff --git a/Engine/Plugins/FX/Niagara/Source/Niagara/Internal/NiagaraDataChannelCommon.h b/Engine/Plugins/FX/Niagara/Source/Niagara/Internal/NiagaraDataChannelCommon.h
--- a/Engine/Plugins/FX/Niagara/Source/Niagara/Internal/NiagaraDataChannelCommon.h
+++ b/Engine/Plugins/FX/Niagara/Source/Niagara/Internal/NiagaraDataChannelCommon.h
@@ -247,4 +247,7 @@ private:
 	TArray<FNiagaraDataChannelPublishRequest> PublishRequests;
 
 	FVector3f LwcTile = FVector3f::ZeroVector;
+
+	/** Critical section protecting shared state for multiple writers publishing from different threads. */
+	FCriticalSection PublishCritSec;
 };
\ No newline at end of file
diff --git a/Engine/Plugins/FX/Niagara/Source/Niagara/Private/DataInterface/NiagaraDataInterfaceDataChannelRead.cpp b/Engine/Plugins/FX/Niagara/Source/Niagara/Private/DataInterface/NiagaraDataInterfaceDataChannelRead.cpp
--- a/Engine/Plugins/FX/Niagara/Source/Niagara/Private/DataInterface/NiagaraDataInterfaceDataChannelRead.cpp
+++ b/Engine/Plugins/FX/Niagara/Source/Niagara/Private/DataInterface/NiagaraDataInterfaceDataChannelRead.cpp
@@ -895,8 +895,8 @@ void UNiagaraDataInterfaceDataChannelRead::Read(FVectorVMExternalFunctionContext
 	FNiagaraDataBuffer* Data = InstData->GetReadBufferCPU(bReadCurrentFrame == false);
 	if (Data && BindingInfo && INiagaraModule::DataChannelsEnabled())
  	{
-// 		FString Label = TEXT("NDIDataChannelRead::Read() - ");
-// 		Data->Dump(0, Data->GetNumInstances(), Label);
+ 		//FString Label = TEXT("NDIDataChannelRead::Read() - ");
+ 		//Data->Dump(0, Data->GetNumInstances(), Label);
 
 		for (int32 i = 0; i < Context.GetNumInstances(); ++i)
 		{
@@ -904,28 +904,28 @@ void UNiagaraDataInterfaceDataChannelRead::Read(FVectorVMExternalFunctionContext
 
 			bool bProcess = (uint32)Index < Data->GetNumInstances();
 
-			auto FloatFunc = [Data, Index](const FNDIDataChannelRegisterBinding& VMBinding, float* FloatData)
+			auto FloatFunc = [Data, Index](const FNDIDataChannelRegisterBinding& VMBinding, VectorVM::FExternalFuncRegisterHandler<float>& FloatData)
 			{
 				if (VMBinding.DataSetRegisterIndex != INDEX_NONE)
 				{
 					float* Src = Data->GetInstancePtrFloat(VMBinding.DataSetRegisterIndex, (uint32)Index);
-					*FloatData = *Src;
+					*FloatData.GetDestAndAdvance() = *Src;
 				}
 			};
-			auto IntFunc = [Data, Index](const FNDIDataChannelRegisterBinding& VMBinding, int32* IntData)
+			auto IntFunc = [Data, Index](const FNDIDataChannelRegisterBinding& VMBinding, VectorVM::FExternalFuncRegisterHandler<int32>& IntData)
 			{
 				if (VMBinding.DataSetRegisterIndex != INDEX_NONE)
 				{
 					int32* Src = Data->GetInstancePtrInt32(VMBinding.DataSetRegisterIndex, (uint32)Index);
-					*IntData = *Src;
+					*IntData.GetDestAndAdvance() = *Src;
 				}
 			};
-			auto HalfFunc = [Data, Index](const FNDIDataChannelRegisterBinding& VMBinding, FFloat16* HalfData)
+			auto HalfFunc = [Data, Index](const FNDIDataChannelRegisterBinding& VMBinding, VectorVM::FExternalFuncRegisterHandler<FFloat16>& HalfData)
 			{
 				if (VMBinding.DataSetRegisterIndex != INDEX_NONE)
-					*HalfData = *Data->GetInstancePtrHalf(VMBinding.DataSetRegisterIndex, (uint32)Index);
+					*HalfData.GetDestAndAdvance() = *Data->GetInstancePtrHalf(VMBinding.DataSetRegisterIndex, (uint32)Index);
 			};
-			bool bSuccess = VariadicOutputs.Process(bProcess, BindingInfo, FloatFunc, IntFunc, HalfFunc);
+			bool bSuccess = VariadicOutputs.Process(bProcess, 1, BindingInfo, FloatFunc, IntFunc, HalfFunc);
 
 			if (OutSuccess.IsValid())
 			{
@@ -976,22 +976,22 @@ void UNiagaraDataInterfaceDataChannelRead::Consume(FVectorVMExternalFunctionCont
 
 				bConsume &= (uint32)Index < Data->GetNumInstances();
 
-				auto FloatFunc = [Data, Index](const FNDIDataChannelRegisterBinding& VMBinding, float* FloatData)
+				auto FloatFunc = [Data, Index](const FNDIDataChannelRegisterBinding& VMBinding, VectorVM::FExternalFuncRegisterHandler<float>& FloatData)
 				{
 					if (VMBinding.DataSetRegisterIndex != INDEX_NONE)
-						*FloatData = *Data->GetInstancePtrFloat(VMBinding.DataSetRegisterIndex, (uint32)Index);
+						*FloatData.GetDestAndAdvance() = *Data->GetInstancePtrFloat(VMBinding.DataSetRegisterIndex, (uint32)Index);
 				};
-				auto IntFunc = [Data, Index](const FNDIDataChannelRegisterBinding& VMBinding, int32* IntData)
+				auto IntFunc = [Data, Index](const FNDIDataChannelRegisterBinding& VMBinding, VectorVM::FExternalFuncRegisterHandler<int32>& IntData)
 				{
 					if (VMBinding.DataSetRegisterIndex != INDEX_NONE)
-						*IntData = *Data->GetInstancePtrInt32(VMBinding.DataSetRegisterIndex, (uint32)Index);
+						*IntData.GetDestAndAdvance() = *Data->GetInstancePtrInt32(VMBinding.DataSetRegisterIndex, (uint32)Index);
 				};
-				auto HalfFunc = [Data, Index](const FNDIDataChannelRegisterBinding& VMBinding, FFloat16* HalfData)
+				auto HalfFunc = [Data, Index](const FNDIDataChannelRegisterBinding& VMBinding, VectorVM::FExternalFuncRegisterHandler<FFloat16>& HalfData)
 				{
 					if (VMBinding.DataSetRegisterIndex != INDEX_NONE)
-						*HalfData = *Data->GetInstancePtrHalf(VMBinding.DataSetRegisterIndex, (uint32)Index);
+						*HalfData.GetDestAndAdvance() = *Data->GetInstancePtrHalf(VMBinding.DataSetRegisterIndex, (uint32)Index);
 				};
-				bSuccess = VariadicOutputs.Process(bConsume, BindingInfo, FloatFunc, IntFunc, HalfFunc);
+				bSuccess = VariadicOutputs.Process(bConsume, 1, BindingInfo, FloatFunc, IntFunc, HalfFunc);
 			}
 			else
 			{
@@ -1120,22 +1120,22 @@ void UNiagaraDataInterfaceDataChannelRead::SpawnConditional(FVectorVMExternalFun
 		for (int32 DataChannelIdx = 0; DataChannelIdx < NumDataChannelInstances; ++DataChannelIdx)
 		{
 			bool bConditionsPass = true;
-			auto FloatFunc = [&bConditionsPass, Op, Data, DataChannelIdx](const FNDIDataChannelRegisterBinding& VMBinding, float FloatData)
+			auto FloatFunc = [&bConditionsPass, Op, Data, DataChannelIdx](const FNDIDataChannelRegisterBinding& VMBinding, FNDIInputParam<float>& FloatData)
 			{
 				if (VMBinding.DataSetRegisterIndex != INDEX_NONE)
-					bConditionsPass &= EvalConditional(Op, FloatData, *Data->GetInstancePtrFloat(VMBinding.DataSetRegisterIndex, (uint32)DataChannelIdx));
+					bConditionsPass &= EvalConditional(Op, FloatData.GetAndAdvance(), *Data->GetInstancePtrFloat(VMBinding.DataSetRegisterIndex, (uint32)DataChannelIdx));
 			};
-			auto IntFunc = [&bConditionsPass, Op, Data, DataChannelIdx](const FNDIDataChannelRegisterBinding& VMBinding, int32 IntData)
+			auto IntFunc = [&bConditionsPass, Op, Data, DataChannelIdx](const FNDIDataChannelRegisterBinding& VMBinding, FNDIInputParam<int32>& IntData)
 			{
 				if (VMBinding.DataSetRegisterIndex != INDEX_NONE)
-					bConditionsPass &= EvalConditional(Op, IntData, *Data->GetInstancePtrInt32(VMBinding.DataSetRegisterIndex, (uint32)DataChannelIdx));
+					bConditionsPass &= EvalConditional(Op, IntData.GetAndAdvance(), *Data->GetInstancePtrInt32(VMBinding.DataSetRegisterIndex, (uint32)DataChannelIdx));
 			};
-			auto HalfFunc = [&bConditionsPass, Op, Data, DataChannelIdx](const FNDIDataChannelRegisterBinding& VMBinding, FFloat16 HalfData)
+			auto HalfFunc = [&bConditionsPass, Op, Data, DataChannelIdx](const FNDIDataChannelRegisterBinding& VMBinding, FNDIInputParam<FFloat16>& HalfData)
 			{
 				if (VMBinding.DataSetRegisterIndex != INDEX_NONE)
-					bConditionsPass &= EvalConditional(Op, HalfData, *Data->GetInstancePtrHalf(VMBinding.DataSetRegisterIndex, (uint32)DataChannelIdx));
+					bConditionsPass &= EvalConditional(Op, HalfData.GetAndAdvance(), *Data->GetInstancePtrHalf(VMBinding.DataSetRegisterIndex, (uint32)DataChannelIdx));
 			};
-			VariadicInputs.Process(true, BindingInfo, FloatFunc, IntFunc, HalfFunc);
+			VariadicInputs.Process(true, 1, BindingInfo, FloatFunc, IntFunc, HalfFunc);
 			VariadicInputs.Reset();
 
 			if (bConditionsPass)
diff --git a/Engine/Plugins/FX/Niagara/Source/Niagara/Private/DataInterface/NiagaraDataInterfaceDataChannelWrite.cpp b/Engine/Plugins/FX/Niagara/Source/Niagara/Private/DataInterface/NiagaraDataInterfaceDataChannelWrite.cpp
--- a/Engine/Plugins/FX/Niagara/Source/Niagara/Private/DataInterface/NiagaraDataInterfaceDataChannelWrite.cpp
+++ b/Engine/Plugins/FX/Niagara/Source/Niagara/Private/DataInterface/NiagaraDataInterfaceDataChannelWrite.cpp
@@ -28,6 +28,8 @@ DECLARE_CYCLE_STAT(TEXT("NDIDataChannelWrite Write"), STAT_NDIDataChannelWrite_W
 DECLARE_CYCLE_STAT(TEXT("NDIDataChannelWrite Append"), STAT_NDIDataChannelWrite_Append, STATGROUP_NiagaraDataChannels);
 DECLARE_CYCLE_STAT(TEXT("NDIDataChannelWrite Tick"), STAT_NDIDataChannelWrite_Tick, STATGROUP_NiagaraDataChannels);
 DECLARE_CYCLE_STAT(TEXT("NDIDataChannelWrite PostTick"), STAT_NDIDataChannelWrite_PostTick, STATGROUP_NiagaraDataChannels);
+DECLARE_CYCLE_STAT(TEXT("NDIDataChannelWrite PreStageTick"), STAT_NDIDataChannelWrite_PreStageTick, STATGROUP_NiagaraDataChannels);
+DECLARE_CYCLE_STAT(TEXT("NDIDataChannelWrite PostStageTick"), STAT_NDIDataChannelWrite_PostStageTick, STATGROUP_NiagaraDataChannels);
 
 int32 GbDebugDumpWriter = 0;
 static FAutoConsoleVariableRef CVarDebugDumpWriterDI(
@@ -116,6 +118,9 @@ struct FNDIDataChannelWriteInstanceData
 
 	TArray<FNDIDataChannel_FuncToDataSetBindingPtr, TInlineAllocator<8>> FunctionToDatSetBindingInfo;
 
+	//Atomic uint for tracking num instances of the target data buffer when writing from multiple threads in the VM.
+	std::atomic<uint32> AtomicNumInstances;
+
 	~FNDIDataChannelWriteInstanceData()
 	{
 		if (Data && DataChannelData.IsValid())
@@ -202,13 +207,46 @@ struct FNDIDataChannelWriteInstanceData
 					}
 				}
 			}
+
+			if (DataChannelPtr)
+			{
+				if (DataChannelData == nullptr || Interface->bUpdateDestinationDataEveryTick)
+				{
+					FNiagaraDataChannelSearchParameters SearchParams;
+					SearchParams.OwningComponent = Instance->GetAttachComponent();
+					DataChannelData = DataChannelPtr->FindData(SearchParams, ENiagaraResourceAccess::WriteOnly);
+				}
+			}
+		}		
+
+		//Verify our function info.
+		if(!ensure(Interface->GetCompiledData().GetFunctionInfo().Num() == FunctionToDatSetBindingInfo.Num()))
+		{
+			UE_LOG(LogNiagara, Warning, TEXT("Invalid Bindings for Niagara Data Interface Data Channel Write: %s"), *Interface->Channel.GetName());
+			return false;			
+		}
+
+		for(const auto& Binding : FunctionToDatSetBindingInfo)
+		{
+			if(Binding.IsValid() == false)
+			{
+				UE_LOG(LogNiagara, Warning, TEXT("Invalid Bindings for Niagara Data Interface Data Channel Write: %s"), *Interface->Channel.GetName());
+				return false;			
+			}
 		}
 
-		if(FNiagaraDataBuffer* CurrBuff = Data->GetCurrentData())
-		{	
+		return true;
+	}
+
+	void PreStageTick(UNiagaraDataInterfaceDataChannelWrite* Interface, FNiagaraSystemInstance* Instance)
+	{
+		if (FNiagaraDataBuffer* CurrBuff = Data->GetCurrentData())
+		{
 			FNiagaraDataBuffer& DestBuff = Data->BeginSimulate(true);
 		}
 
+		//TODO: Currently allocating for every stage using this DI whether it's a write or not. We should limit this to write functions.
+		AtomicNumInstances = 0;
 		if (Interface->AllocationMode == ENiagaraDataChannelAllocationMode::Static)
 		{
 			Data->GetDestinationData()->Allocate(Interface->AllocationCount);
@@ -225,30 +263,14 @@ struct FNDIDataChannelWriteInstanceData
 		{
 			check(0);
 		}
-
-		//Verify our function info.
-		if(!ensure(Interface->GetCompiledData().GetFunctionInfo().Num() == FunctionToDatSetBindingInfo.Num()))
-		{
-			UE_LOG(LogNiagara, Warning, TEXT("Invalid Bindings for Niagara Data Interface Data Channel Write: %s"), *Interface->Channel.GetName());
-			return false;			
-		}
-
-		for(auto& Binding : FunctionToDatSetBindingInfo)
-		{
-			if(Binding.IsValid() == false)
-			{
-				UE_LOG(LogNiagara, Warning, TEXT("Invalid Bindings for Niagara Data Interface Data Channel Write: %s"), *Interface->Channel.GetName());
-				return false;			
-			}
-		}
-
-		return true;
 	}
 
-	bool PostTick(UNiagaraDataInterfaceDataChannelWrite* Interface, FNiagaraSystemInstance* Instance)
+	void PostStageTick(UNiagaraDataInterfaceDataChannelWrite* Interface, FNiagaraSystemInstance* Instance)
 	{
 		if (Data && Data->GetDestinationData())
 		{
+			//TOOD: Move allocation/end to pre/post stage?
+			Data->GetDestinationData()->SetNumInstances(AtomicNumInstances.load(std::memory_order_seq_cst));
 			Data->EndSimulate();
 
 			if (GbDebugDumpWriter)
@@ -257,28 +279,17 @@ struct FNDIDataChannelWriteInstanceData
 				Buffer.Dump(0, Buffer.GetNumInstances(), FString::Printf(TEXT("=== Data Channle Write: %d Elements --> %s ==="), Buffer.GetNumInstances(), *Interface->Channel.GetName()));
 			}
 
-			if (Interface->ShouldPublish() && Data->GetCurrentData()->GetNumInstances() > 0)
+			if (DataChannelData && Interface->ShouldPublish() && Data->GetCurrentData()->GetNumInstances() > 0)
 			{
-				if (UNiagaraDataChannelHandler* Channel = DataChannel.Get())
-				{
-					if(DataChannelData == nullptr || Interface->bUpdateDestinationDataEveryTick)
-					{
-						FNiagaraDataChannelSearchParameters SearchParams;
-						SearchParams.OwningComponent = Instance->GetAttachComponent();
-						DataChannelData = Channel->FindData(SearchParams, ENiagaraResourceAccess::WriteOnly);
-					}
-
-					FNiagaraDataChannelPublishRequest PublishRequest(Data->GetCurrentData());
-					PublishRequest.bVisibleToGame = Interface->bPublishToGame;
-					PublishRequest.bVisibleToCPUSims = Interface->bPublishToCPU;
-					PublishRequest.bVisibleToGPUSims = Interface->bPublishToGPU;
-					PublishRequest.Data = Data->GetCurrentData();
-					PublishRequest.LwcTile = Instance->GetLWCTile();
-					DataChannelData->Publish(PublishRequest);
-				}
+				FNiagaraDataChannelPublishRequest PublishRequest(Data->GetCurrentData());
+				PublishRequest.bVisibleToGame = Interface->bPublishToGame;
+				PublishRequest.bVisibleToCPUSims = Interface->bPublishToCPU;
+				PublishRequest.bVisibleToGPUSims = Interface->bPublishToGPU;
+				PublishRequest.Data = Data->GetCurrentData();
+				PublishRequest.LwcTile = Instance->GetLWCTile();
+				DataChannelData->Publish(PublishRequest);
 			}
 		}
-		return true;
 	}
 };
 
@@ -363,25 +374,44 @@ bool UNiagaraDataInterfaceDataChannelWrite::PerInstanceTick(void* PerInstanceDat
 }
 
 bool UNiagaraDataInterfaceDataChannelWrite::PerInstanceTickPostSimulate(void* PerInstanceData, FNiagaraSystemInstance* SystemInstance, float DeltaSeconds)
+{
+ 	return false;
+}
+
+void UNiagaraDataInterfaceDataChannelWrite::PreStageTick(FNDICpuPostStageContext& Context)
 {
 	if (INiagaraModule::DataChannelsEnabled() == false)
 	{
-		return true;
+		return;
 	}
 
-	SCOPE_CYCLE_COUNTER(STAT_NDIDataChannelWrite_PostTick);
-	check(SystemInstance);
-	FNDIDataChannelWriteInstanceData* InstanceData = static_cast<FNDIDataChannelWriteInstanceData*>(PerInstanceData);
+	SCOPE_CYCLE_COUNTER(STAT_NDIDataChannelWrite_PreStageTick);
+	check(Context.SystemInstance);
+	FNDIDataChannelWriteInstanceData* InstanceData = static_cast<FNDIDataChannelWriteInstanceData*>(Context.PerInstanceData);
 	if (!InstanceData)
 	{
-		return true;
+		return;
 	}
 
-	if (InstanceData->PostTick(this, SystemInstance) == false)
+	InstanceData->PreStageTick(this, Context.SystemInstance);
+}
+
+void UNiagaraDataInterfaceDataChannelWrite::PostStageTick(FNDICpuPostStageContext& Context)
+{
+	if (INiagaraModule::DataChannelsEnabled() == false)
 	{
-		return true;
+		return;
 	}
-	return false;
+
+	SCOPE_CYCLE_COUNTER(STAT_NDIDataChannelWrite_PostStageTick);
+	check(Context.PerInstanceData);
+	FNDIDataChannelWriteInstanceData* InstanceData = static_cast<FNDIDataChannelWriteInstanceData*>(Context.PerInstanceData);
+	if (!Context.SystemInstance)
+	{
+		return;
+	}
+	
+	InstanceData->PostStageTick(this, Context.SystemInstance);
 }
 
 void UNiagaraDataInterfaceDataChannelWrite::ProvidePerInstanceDataForRenderThread(void* DataForRenderThread, void* PerInstanceData, const FNiagaraSystemInstanceID& SystemInstance)
@@ -654,53 +684,71 @@ void UNiagaraDataInterfaceDataChannelWrite::Write(FVectorVMExternalFunctionConte
 {
 	SCOPE_CYCLE_COUNTER(STAT_NDIDataChannelWrite_Write);
 	VectorVM::FUserPtrHandler<FNDIDataChannelWriteInstanceData> InstData(Context);
-	FNDIInputParam<bool> InEmit(Context);
+	FNDIInputParam<FNiagaraBool> InEmit(Context);
 	FNDIInputParam<int32> InIndex(Context);
 
 	const FNDIDataChannelFunctionInfo& FuncInfo = CompiledData.GetFunctionInfo()[FuncIdx];
 	const FNDIDataChannel_FunctionToDataSetBinding* BindingInfo = InstData->FunctionToDatSetBindingInfo.IsValidIndex(FuncIdx) ? InstData->FunctionToDatSetBindingInfo[FuncIdx].Get() : nullptr;
 	FNDIVariadicInputHandler<16> VariadicInputs(Context, BindingInfo);//TODO: Make static / avoid allocation
 
-	FNDIOutputParam<bool> OutSuccess(Context);
+	FNDIOutputParam<FNiagaraBool> OutSuccess(Context);
+
+	std::atomic<uint32>& AtomicNumInstances = InstData->AtomicNumInstances;
 
+	bool bAllFailedFallback = true;
 	if (InstData->Data && BindingInfo && INiagaraModule::DataChannelsEnabled())
 	{
-		FNiagaraDataBuffer* Data = InstData->Data->GetDestinationData();
+		if(FNiagaraDataBuffer* Data = InstData->Data->GetDestinationData())
+		{			
+			bAllFailedFallback = false;
+			uint32 MaxLocalIndex = 0;
+			int32 NumAllocated = (int32)Data->GetNumInstancesAllocated();
+			for (int32 i = 0; i < Context.GetNumInstances(); ++i)
+			{
+				int32 Index = InIndex.GetAndAdvance();
+				bool bEmit = InEmit.GetAndAdvance() && Index >= 0 && Index < NumAllocated;
 
-		for (int32 i = 0; i < Context.GetNumInstances(); ++i)
-		{
-			int32 Index = InIndex.GetAndAdvance();
-			bool bEmit = InEmit.GetAndAdvance();
-			bEmit &= (uint32)Index < Data->GetNumInstances();
+				MaxLocalIndex = bEmit ? FMath::Max((uint32)Index, MaxLocalIndex) : MaxLocalIndex;
 
-			bool bSuccess = false;
+				bool bSuccess = false;
 
-			//TODO: Optimize case where emit is constant
-			//TODO: Optimize for runs of sequential true emits.
-			auto FloatFunc = [Data, Index](const FNDIDataChannelRegisterBinding& VMBinding, float FloatData)
-			{
-				if(VMBinding.DataSetRegisterIndex != INDEX_NONE)
-					*Data->GetInstancePtrFloat(VMBinding.DataSetRegisterIndex, (uint32)Index) = FloatData; 
-			};
-			auto IntFunc = [Data, Index](const FNDIDataChannelRegisterBinding& VMBinding, int32 IntData)
-			{
-				if(VMBinding.DataSetRegisterIndex != INDEX_NONE)
-					*Data->GetInstancePtrInt32(VMBinding.DataSetRegisterIndex, (uint32)Index) = IntData; 
-			};
-			auto HalfFunc = [Data, Index](const FNDIDataChannelRegisterBinding& VMBinding, FFloat16 HalfData)
-			{
-				if(VMBinding.DataSetRegisterIndex != INDEX_NONE)
-					*Data->GetInstancePtrHalf(VMBinding.DataSetRegisterIndex, (uint32)Index) = HalfData; 
-			};
-			bSuccess = VariadicInputs.Process(bEmit, BindingInfo, FloatFunc, IntFunc, HalfFunc);
+				//TODO: Optimize case where emit is constant
+				//TODO: Optimize for runs of sequential true emits.
+				auto FloatFunc = [Data, Index](const FNDIDataChannelRegisterBinding& VMBinding, FNDIInputParam<float>& FloatData)
+				{
+					if (VMBinding.DataSetRegisterIndex != INDEX_NONE)
+						*Data->GetInstancePtrFloat(VMBinding.DataSetRegisterIndex, (uint32)Index) = FloatData.GetAndAdvance();
+				};
+				auto IntFunc = [Data, Index](const FNDIDataChannelRegisterBinding& VMBinding, FNDIInputParam<int32>& IntData)
+				{
+					if (VMBinding.DataSetRegisterIndex != INDEX_NONE)
+						*Data->GetInstancePtrInt32(VMBinding.DataSetRegisterIndex, (uint32)Index) = IntData.GetAndAdvance();
+				};
+				auto HalfFunc = [Data, Index](const FNDIDataChannelRegisterBinding& VMBinding, FNDIInputParam<FFloat16>& HalfData)
+				{
+					if (VMBinding.DataSetRegisterIndex != INDEX_NONE)
+						*Data->GetInstancePtrHalf(VMBinding.DataSetRegisterIndex, (uint32)Index) = HalfData.GetAndAdvance();
+				};
 
-			if (OutSuccess.IsValid())
+				bSuccess = VariadicInputs.Process(bEmit, 1, BindingInfo, FloatFunc, IntFunc, HalfFunc);
+
+				if (OutSuccess.IsValid())
+				{
+					OutSuccess.SetAndAdvance(bSuccess);
+				}
+			}
+
+			//Update the shared instance count with an updated max.
+			bool bAtomicNumInstancesWritten = false;
+			uint32 CurrNumInstances = AtomicNumInstances;
+			while(CurrNumInstances < MaxLocalIndex && !AtomicNumInstances.compare_exchange_weak(CurrNumInstances, MaxLocalIndex))
 			{
-				OutSuccess.SetAndAdvance(bSuccess);
+				CurrNumInstances = AtomicNumInstances;
 			}
 		}
 	}
-	else
+
+	if (bAllFailedFallback)
 	{
 		for (int32 i = 0; i < Context.GetNumInstances(); ++i)
 		{
@@ -716,61 +764,161 @@ void UNiagaraDataInterfaceDataChannelWrite::Append(FVectorVMExternalFunctionCont
 {
 	SCOPE_CYCLE_COUNTER(STAT_NDIDataChannelWrite_Append);
 	VectorVM::FUserPtrHandler<FNDIDataChannelWriteInstanceData> InstData(Context);
-	FNDIInputParam<bool> InEmit(Context);
+	FNDIInputParam<FNiagaraBool> InEmit(Context);
 
 	const FNDIDataChannelFunctionInfo& FuncInfo = CompiledData.GetFunctionInfo()[FuncIdx];
 	const FNDIDataChannel_FunctionToDataSetBinding* BindingInfo = InstData->FunctionToDatSetBindingInfo.IsValidIndex(FuncIdx) ? InstData->FunctionToDatSetBindingInfo[FuncIdx].Get() : nullptr;
 	FNDIVariadicInputHandler<16> VariadicInputs(Context, BindingInfo);//TODO: Make static / avoid allocation
 
-	FNDIOutputParam<bool> OutSuccess(Context);
+	FNDIOutputParam<FNiagaraBool> OutSuccess(Context);
 
+	std::atomic<uint32>& AtomicNumInstances = InstData->AtomicNumInstances;
+
+	bool bAllFailedFallback = true;
 	if (InstData->Data && BindingInfo &&  INiagaraModule::DataChannelsEnabled())
 	{
-		FNiagaraDataBuffer* Data = InstData->Data->GetDestinationData();
-		for (int32 i = 0; i < Context.GetNumInstances(); ++i)
+		if(FNiagaraDataBuffer* Data = InstData->Data->GetDestinationData())
 		{
-			bool bEmit = InEmit.GetAndAdvance();
-
-			bool bSuccess = false;
-
-			//TODO: Optimize case where emit is constant
-			//TODO: Optimize for runs of sequential true emits.
-			if (bEmit && Data)
+			//Get the total number to emit.
+			//Allows going via a faster write path if we're emiting every instance.
+			//Aslo needed to update the atomic num instances and get our start index for writing.
+			uint32 LocalNumToEmit = 0;
+			if(InEmit.IsConstant())
 			{
-				int32 Num = Data->GetNumInstances();
-				bEmit &= (int32)Data->GetNumInstancesAllocated() > Num;
-
-				if (bEmit)
+				bool bEmit = InEmit.GetAndAdvance();
+				LocalNumToEmit = bEmit ? Context.GetNumInstances() : 0;
+			}
+			else
+			{
+				for (int32 i = 0; i < Context.GetNumInstances(); ++i)
 				{
-					Data->SetNumInstances(Num + 1);
+					if (InEmit.GetAndAdvance())
+					{
+						++LocalNumToEmit;
+					}
 				}
+			}			
 
-				int32 Index = Num;
-				auto FloatFunc = [Data, Index](const FNDIDataChannelRegisterBinding& VMBinding, float FloatData)
-				{
-					if (VMBinding.DataSetRegisterIndex != INDEX_NONE)
-						*Data->GetInstancePtrFloat(VMBinding.DataSetRegisterIndex, (uint32)Index) = FloatData;
-				};
-				auto IntFunc = [Data, Index](const FNDIDataChannelRegisterBinding& VMBinding, int32 IntData)
+			if (LocalNumToEmit > 0)
+			{
+				int32 NumAllocated = Data->GetNumInstancesAllocated();
+				InEmit.Reset();
+
+				//Update the shared atomic instance count and grab the current index at which we can write.
+				uint32 CurrNumInstances = AtomicNumInstances.fetch_add(LocalNumToEmit);
+
+				bAllFailedFallback = false;
+
+				bool bEmitAll = LocalNumToEmit == Context.GetNumInstances();
+
+				if(bEmitAll)
 				{
-					if (VMBinding.DataSetRegisterIndex != INDEX_NONE)
-						*Data->GetInstancePtrInt32(VMBinding.DataSetRegisterIndex, (uint32)Index) = IntData;
-				};
-				auto HalfFunc = [Data, Index](const FNDIDataChannelRegisterBinding& VMBinding, FFloat16 HalfData)
+					//If we're writing all instances then we can do a memcpy instead of slower loop copies.
+					bool bSuccess = false;
+					int32 Index = CurrNumInstances;
+					auto FloatFunc = [Data, Index, LocalNumToEmit](const FNDIDataChannelRegisterBinding& VMBinding, FNDIInputParam<float>& FloatData)
+					{
+						if (VMBinding.DataSetRegisterIndex != INDEX_NONE)
+						{
+							float* Dest = Data->GetInstancePtrFloat(VMBinding.DataSetRegisterIndex, (uint32)Index);
+							if (FloatData.IsConstant())
+							{
+								float Value = FloatData.GetAndAdvance();
+								for(uint32 i=0; i<LocalNumToEmit; ++i){ Dest[i] = Value; }
+							}
+							else
+							{
+								 const float* Src = FloatData.Data.GetDest();
+								 FMemory::Memcpy(Dest, Src, LocalNumToEmit * sizeof(float));
+							}
+						}
+					};
+					auto IntFunc = [Data, Index, LocalNumToEmit](const FNDIDataChannelRegisterBinding& VMBinding, FNDIInputParam<int32>& IntData)
+					{
+						if (VMBinding.DataSetRegisterIndex != INDEX_NONE)
+						{
+							int32* Dest = Data->GetInstancePtrInt32(VMBinding.DataSetRegisterIndex, (uint32)Index);
+							if (IntData.IsConstant())
+							{
+								int32 Value = IntData.GetAndAdvance();
+								for (uint32 i = 0; i < LocalNumToEmit; ++i) { Dest[i] = Value; }
+							}
+							else
+							{
+								const int32* Src = IntData.Data.GetDest();
+								FMemory::Memcpy(Dest, Src, LocalNumToEmit * sizeof(int32));
+							}
+						}
+					};
+					auto HalfFunc = [Data, Index, LocalNumToEmit](const FNDIDataChannelRegisterBinding& VMBinding, FNDIInputParam<FFloat16>& HalfData)
+					{
+						if (VMBinding.DataSetRegisterIndex != INDEX_NONE)
+						{
+							FFloat16* Dest = Data->GetInstancePtrHalf(VMBinding.DataSetRegisterIndex, (uint32)Index);
+							if (HalfData.IsConstant())
+							{
+								FFloat16 Value = HalfData.GetAndAdvance();
+								for (uint32 i = 0; i < LocalNumToEmit; ++i) { Dest[i] = Value; }
+							}
+							else
+							{
+								const FFloat16* Src = HalfData.Data.GetDest();
+								FMemory::Memcpy(Dest, Src, LocalNumToEmit * sizeof(FFloat16));
+							}
+						}
+					};
+					bSuccess = VariadicInputs.Process(true, Context.GetNumInstances(), BindingInfo, FloatFunc, IntFunc, HalfFunc);
+
+					if (OutSuccess.IsValid())
+					{
+						for (int32 i = 0; i < Context.GetNumInstances(); ++i)
+						{
+							OutSuccess.SetAndAdvance(bSuccess);
+						}
+					}
+				}
+				else
 				{
-					if (VMBinding.DataSetRegisterIndex != INDEX_NONE)
-						*Data->GetInstancePtrHalf(VMBinding.DataSetRegisterIndex, (uint32)Index) = HalfData;
-				};
-				bSuccess = VariadicInputs.Process(bEmit, BindingInfo, FloatFunc, IntFunc, HalfFunc);
-			}
+					for (int32 i = 0; i < Context.GetNumInstances(); ++i)
+					{
+						bool bEmit = InEmit.GetAndAdvance();
+						bool bSuccess = false;
 
-			if (OutSuccess.IsValid())
-			{
-				OutSuccess.SetAndAdvance(bSuccess);
+						int32 Index = CurrNumInstances;
+
+						if(bEmit)
+						{
+							++CurrNumInstances;
+						}
+
+						auto FloatFunc = [Data, Index](const FNDIDataChannelRegisterBinding& VMBinding, FNDIInputParam<float>& FloatData)
+						{
+							if (VMBinding.DataSetRegisterIndex != INDEX_NONE)
+								*Data->GetInstancePtrFloat(VMBinding.DataSetRegisterIndex, (uint32)Index) = FloatData.GetAndAdvance();
+						};
+						auto IntFunc = [Data, Index](const FNDIDataChannelRegisterBinding& VMBinding, FNDIInputParam<int32>& IntData)
+						{
+							if (VMBinding.DataSetRegisterIndex != INDEX_NONE)
+								*Data->GetInstancePtrInt32(VMBinding.DataSetRegisterIndex, (uint32)Index) = IntData.GetAndAdvance();
+						};
+						auto HalfFunc = [Data, Index](const FNDIDataChannelRegisterBinding& VMBinding, FNDIInputParam<FFloat16>& HalfData)
+						{
+							if (VMBinding.DataSetRegisterIndex != INDEX_NONE)
+								*Data->GetInstancePtrHalf(VMBinding.DataSetRegisterIndex, (uint32)Index) = HalfData.GetAndAdvance();
+						};
+						bSuccess = VariadicInputs.Process(bEmit, 1, BindingInfo, FloatFunc, IntFunc, HalfFunc);
+
+						if (OutSuccess.IsValid())
+						{
+							OutSuccess.SetAndAdvance(bSuccess);
+						}
+					}
+				}
 			}
 		}
 	}
-	else
+
+	if(bAllFailedFallback)
 	{
 		for (int32 i = 0; i < Context.GetNumInstances(); ++i)
 		{
@@ -796,7 +944,7 @@ bool UNiagaraDataInterfaceDataChannelWrite::AppendCompileHash(FNiagaraCompileHas
 // 	bSuccess &= InVisitor->UpdateShaderParameters<NDIDataChannelWriteLocal::FShaderParameters>();
 // 	return bSuccess;
 
-	return false;
+	return true;
 }
 
 void UNiagaraDataInterfaceDataChannelWrite::GetCommonHLSL(FString& OutHLSL)
diff --git a/Engine/Plugins/FX/Niagara/Source/Niagara/Private/NiagaraDataChannel.cpp b/Engine/Plugins/FX/Niagara/Source/Niagara/Private/NiagaraDataChannel.cpp
--- a/Engine/Plugins/FX/Niagara/Source/Niagara/Private/NiagaraDataChannel.cpp
+++ b/Engine/Plugins/FX/Niagara/Source/Niagara/Private/NiagaraDataChannel.cpp
@@ -501,6 +501,7 @@ FNiagaraDataChannelData::~FNiagaraDataChannelData()
 
 void FNiagaraDataChannelData::Reset()
 {
+	FScopeLock Lock(&PublishCritSec);
 	PublishRequests.Reset();
 	BuffersForGPU.Reset();
 	
@@ -563,6 +564,9 @@ void FNiagaraDataChannelData::ConsumePublishRequests(UNiagaraDataChannelHandler*
 {
 	check(IsValid(Owner));
 
+	//There should be no access on other threads at this point anyway but lock just to be safe.
+	FScopeLock Lock(&PublishCritSec);
+
 	if(PublishRequests.Num() == 0)
 	{
 		return;
@@ -676,8 +680,14 @@ void FNiagaraDataChannelData::ConsumePublishRequests(UNiagaraDataChannelHandler*
 #if !UE_BUILD_SHIPPING
 	if(Owner->GetDataChannel()->GetVerboseLogging())
 	{
-		FString Label = FString::Printf(TEXT("Data Channel %s"), *Owner->GetDataChannel()->GetName());
+		FString Label = FString::Printf(TEXT("Data Channel %s - CURR"), *Owner->GetDataChannel()->GetName());
 		CPUSimData->GetCurrentData()->Dump(0, CPUSimData->GetCurrentData()->GetNumInstances(), Label);
+
+		if(PrevCPUSimData)
+		{
+			FString LabelPrev = FString::Printf(TEXT("Data Channel %s - PREV"), *Owner->GetDataChannel()->GetName());
+			PrevCPUSimData->Dump(0, PrevCPUSimData->GetNumInstances(), LabelPrev);
+		}
 	}
 #endif
 
@@ -703,11 +713,14 @@ FNiagaraDataChannelGameData* FNiagaraDataChannelData::GetGameData()
 
 void FNiagaraDataChannelData::Publish(const FNiagaraDataChannelPublishRequest& Request)
 {
+	FScopeLock Lock(&PublishCritSec);
+
 	PublishRequests.Add(Request);
 }
 
 void FNiagaraDataChannelData::RemovePublishRequests(const FNiagaraDataSet* DataSet)
 {
+	FScopeLock Lock(&PublishCritSec);
 	//TODO: Have to ensure lifetime of anything we're holding
 	for (auto It = PublishRequests.CreateIterator(); It; ++It)
 	{
diff --git a/Engine/Plugins/FX/Niagara/Source/Niagara/Private/NiagaraDataInterface.cpp b/Engine/Plugins/FX/Niagara/Source/Niagara/Private/NiagaraDataInterface.cpp
--- a/Engine/Plugins/FX/Niagara/Source/Niagara/Private/NiagaraDataInterface.cpp
+++ b/Engine/Plugins/FX/Niagara/Source/Niagara/Private/NiagaraDataInterface.cpp
@@ -24,6 +24,8 @@ void FNDIStageTickHandler::Init(UNiagaraScript* Script, FNiagaraSystemInstance*
 	check(Script);
 	check(Instance);
 	Usage = Script->GetUsage();
+	PreStageTickList.Empty();
+	PostStageTickList.Empty();
 	for (auto& ResolvedDIInfo : Script->GetResolvedDataInterfaces())
 	{
 		if (ResolvedDIInfo.ResolvedDataInterface)
diff --git a/Engine/Plugins/FX/Niagara/Source/Niagara/Private/NiagaraEmitterInstance.cpp b/Engine/Plugins/FX/Niagara/Source/Niagara/Private/NiagaraEmitterInstance.cpp
--- a/Engine/Plugins/FX/Niagara/Source/Niagara/Private/NiagaraEmitterInstance.cpp
+++ b/Engine/Plugins/FX/Niagara/Source/Niagara/Private/NiagaraEmitterInstance.cpp
@@ -500,6 +500,32 @@ void FNiagaraEmitterInstance::Init(int32 InEmitterIdx, FNiagaraSystemInstanceID
 	bCombineEventSpawn = GbNiagaraAllowEventSpawnCombine && (EmitterData->bCombineEventSpawn || (GbNiagaraAllowEventSpawnCombine == 2));
 }
 
+void FNiagaraEmitterInstance::InitDITickLists()
+{
+	if(IsDisabled())
+	{
+		// Make sure that our function tables need to be regenerated...
+		SpawnExecContext.GetNDIStageTickHander().Empty();
+		UpdateExecContext.GetNDIStageTickHander().Empty();
+
+		for (FNiagaraScriptExecutionContext& EventContext : GetEventExecutionContexts())
+		{
+			EventContext.GetNDIStageTickHander().Empty();
+		}
+		return;
+	}
+
+	// Make sure that our function tables need to be regenerated...
+	SpawnExecContext.InitDITickLists(ParentSystemInstance);
+	UpdateExecContext.InitDITickLists(ParentSystemInstance);
+
+	for (FNiagaraScriptExecutionContext& EventContext : GetEventExecutionContexts())
+	{
+		EventContext.InitDITickLists(ParentSystemInstance);
+	}
+}
+
+
 void FNiagaraEmitterInstance::ResetSimulation(bool bKillExisting /*= true*/)
 {
 	EmitterAge = 0;
diff --git a/Engine/Plugins/FX/Niagara/Source/Niagara/Private/NiagaraScriptExecutionContext.cpp b/Engine/Plugins/FX/Niagara/Source/Niagara/Private/NiagaraScriptExecutionContext.cpp
--- a/Engine/Plugins/FX/Niagara/Source/Niagara/Private/NiagaraScriptExecutionContext.cpp
+++ b/Engine/Plugins/FX/Niagara/Source/Niagara/Private/NiagaraScriptExecutionContext.cpp
@@ -88,26 +88,47 @@ bool FNiagaraScriptExecutionContextBase::Init(FNiagaraSystemInstance* Instance,
 	VectorVMState = AllocVectorVMState(&OptimizeContext);
 #endif
 
+	//If the instance is null, aka a system script, we need to pre calculate whether we'll have pre/post ticks or not.
+	if(Instance == nullptr)
+	{
+		ENiagaraScriptUsage Usage = Script->GetUsage();
+		for (auto& ResolvedDIInfo : Script->GetResolvedDataInterfaces())
+		{
+			if (ResolvedDIInfo.ResolvedDataInterface)
+			{
+				bHasDIsWithPreStageTick |= (ResolvedDIInfo.ResolvedDataInterface->HasPreStageTick(Usage));
+				bHasDIsWithPostStageTick |= (ResolvedDIInfo.ResolvedDataInterface->HasPostStageTick(Usage));
+			}
+			if (bHasDIsWithPreStageTick && bHasDIsWithPostStageTick)
+			{
+				break;
+			}
+		}
+	}
+
+	return true;
+}
+
+void FNiagaraScriptExecutionContextBase::InitDITickLists(class FNiagaraSystemInstance* Instance)
+{
 	ENiagaraScriptUsage Usage = Script->GetUsage();
-	for(auto& ResolvedDIInfo : Script->GetResolvedDataInterfaces())
+	for (auto& ResolvedDIInfo : Script->GetResolvedDataInterfaces())
 	{
-		if(ResolvedDIInfo.ResolvedDataInterface)
+		if (ResolvedDIInfo.ResolvedDataInterface)
 		{
 			bHasDIsWithPreStageTick |= (ResolvedDIInfo.ResolvedDataInterface->HasPreStageTick(Usage));
 			bHasDIsWithPostStageTick |= (ResolvedDIInfo.ResolvedDataInterface->HasPostStageTick(Usage));
 		}
-		if(bHasDIsWithPreStageTick && bHasDIsWithPostStageTick)
+		if (bHasDIsWithPreStageTick && bHasDIsWithPostStageTick)
 		{
 			break;
 		}
 	}
 
-	if(Instance && (bHasDIsWithPreStageTick || bHasDIsWithPostStageTick))
+	if (Instance && (bHasDIsWithPreStageTick || bHasDIsWithPostStageTick))
 	{
 		DIStageTickHandler.Init(Script, Instance);
 	}
-
-	return true;
 }
 
 void FNiagaraScriptExecutionContextBase::BindData(int32 Index, FNiagaraDataSet& DataSet, int32 StartInstance, bool bUpdateInstanceCounts)
diff --git a/Engine/Plugins/FX/Niagara/Source/Niagara/Private/NiagaraSystemInstance.cpp b/Engine/Plugins/FX/Niagara/Source/Niagara/Private/NiagaraSystemInstance.cpp
--- a/Engine/Plugins/FX/Niagara/Source/Niagara/Private/NiagaraSystemInstance.cpp
+++ b/Engine/Plugins/FX/Niagara/Source/Niagara/Private/NiagaraSystemInstance.cpp
@@ -1614,6 +1614,10 @@ void FNiagaraSystemInstance::InitDataInterfaces()
 	//If needed, the system script execution will use these to perform per instance pre and post tick operations on our DIs.
 	SystemSpawnDIStageTickHandler.Init(SystemSimulation->GetSpawnExecutionContext()->Script, this);
 	SystemUpdateDIStageTickHandler.Init(SystemSimulation->GetUpdateExecutionContext()->Script, this);
+	for (TSharedRef<FNiagaraEmitterInstance, ESPMode::ThreadSafe>& Sim : Emitters)
+	{
+		Sim->InitDITickLists();
+	}
 }
 
 void FNiagaraSystemInstance::ResolveUserDataInterfaceBindings()
diff --git a/Engine/Source/Runtime/VectorVM/Public/VectorVM.h b/Engine/Source/Runtime/VectorVM/Public/VectorVM.h
--- a/Engine/Source/Runtime/VectorVM/Public/VectorVM.h
+++ b/Engine/Source/Runtime/VectorVM/Public/VectorVM.h
@@ -266,7 +266,7 @@ namespace VectorVM
 
 		FORCEINLINE const T Get() { return *InputPtr; }
 		FORCEINLINE const T* GetDest() { return InputPtr; }
-		FORCEINLINE void Advance() { InputPtr += AdvanceOffset; }
+		FORCEINLINE void Advance(int32 Count=1) { InputPtr += AdvanceOffset * Count; }
 		FORCEINLINE const T GetAndAdvance()
 		{
 			const T* Ret = InputPtr;
