Niagara - When postloading disabled emitters, make sure that their compile results have been invalidated.

diff --git a/Engine/Plugins/FX/Niagara/Source/Niagara/Private/NiagaraEmitterHandle.cpp b/Engine/Plugins/FX/Niagara/Source/Niagara/Private/NiagaraEmitterHandle.cpp
--- a/Engine/Plugins/FX/Niagara/Source/Niagara/Private/NiagaraEmitterHandle.cpp
+++ b/Engine/Plugins/FX/Niagara/Source/Niagara/Private/NiagaraEmitterHandle.cpp
@@ -227,6 +227,15 @@ void FNiagaraEmitterHandle::ConditionalPostLoad(int32 NiagaraCustomVersion)
 				SetName(Name, *OwningSystem);
 			}
 		}
+
+		if (bIsEnabled == false)
+		{
+			FVersionedNiagaraEmitterData* EmitterData = VersionedInstance.GetEmitterData();
+			if (EmitterData != nullptr)
+			{
+				EmitterData->InvalidateCompileResults();
+			}
+		}
 	}
 }
 
