- Add editor only compile tags
- Move attribute trimming information into editor only compile tags to reduce memory

diff --git a/Engine/Plugins/FX/Niagara/Source/Niagara/Classes/NiagaraScript.h b/Engine/Plugins/FX/Niagara/Source/Niagara/Classes/NiagaraScript.h
--- a/Engine/Plugins/FX/Niagara/Source/Niagara/Classes/NiagaraScript.h
+++ b/Engine/Plugins/FX/Niagara/Source/Niagara/Classes/NiagaraScript.h
@@ -458,6 +458,9 @@ public:
 	    in RI parameters this will be the values that are baked in, otherwise it will be the set of defaults based on the graphs. */
 	UPROPERTY()
 	TArray<FNiagaraVariable> BakedRapidIterationParameters;
+
+	UPROPERTY()
+	TArray<FNiagaraCompilerTag> CompileTagsEditorOnly;
 #endif
 
 	UPROPERTY()
diff --git a/Engine/Plugins/FX/Niagara/Source/Niagara/Private/NiagaraScript.cpp b/Engine/Plugins/FX/Niagara/Source/Niagara/Private/NiagaraScript.cpp
--- a/Engine/Plugins/FX/Niagara/Source/Niagara/Private/NiagaraScript.cpp
+++ b/Engine/Plugins/FX/Niagara/Source/Niagara/Private/NiagaraScript.cpp
@@ -3567,9 +3567,25 @@ void UNiagaraScript::SyncAliases(const FNiagaraAliasContext& ResolveAliasesConte
 			{
 				FNiagaraVariable NewVar = FNiagaraUtilities::ResolveAliases(FNiagaraVariable(FNiagaraTypeDefinition::GetFloatDef(), *Name), ResolveAliasesContext);
 				if (NewVar.GetName() != *Name)
+				{
 					GetVMExecutableData().CompileTags[i].StringValue = NewVar.GetName().ToString();
+				}
+			}
+		}
+	#if WITH_EDITORONLY_DATA
+		for (int32 i = 0; i < GetVMExecutableData().CompileTagsEditorOnly.Num(); i++)
+		{
+			const FString& Name = GetVMExecutableData().CompileTagsEditorOnly[i].StringValue;
+			if (Name.Len())
+			{
+				FNiagaraVariable NewVar = FNiagaraUtilities::ResolveAliases(FNiagaraVariable(FNiagaraTypeDefinition::GetFloatDef(), *Name), ResolveAliasesContext);
+				if (NewVar.GetName() != *Name)
+				{
+					GetVMExecutableData().CompileTagsEditorOnly[i].StringValue = NewVar.GetName().ToString();
+				}
 			}
 		}
+	#endif
 	}
 
 	InvalidateExecutionReadyParameterStores();
diff --git a/Engine/Plugins/FX/Niagara/Source/NiagaraEditor/Private/NiagaraCompiler.cpp b/Engine/Plugins/FX/Niagara/Source/NiagaraEditor/Private/NiagaraCompiler.cpp
--- a/Engine/Plugins/FX/Niagara/Source/NiagaraEditor/Private/NiagaraCompiler.cpp
+++ b/Engine/Plugins/FX/Niagara/Source/NiagaraEditor/Private/NiagaraCompiler.cpp
@@ -1830,6 +1830,7 @@ int32 FHlslNiagaraCompiler::CompileScript(const FStringView GroupName, const FNi
 	CompileResults.Data->LastCompileEvents.Append(InTranslateResults.CompileEvents);
 	CompileResults.Data->ExternalDependencies = InTranslateResults.CompileDependencies;
 	CompileResults.Data->CompileTags = InTranslateResults.CompileTags;
+	CompileResults.Data->CompileTagsEditorOnly = InTranslateResults.CompileTagsEditorOnly;
 
 	// Early out if compiling a GPU particle script as we do not need to submit a CPU compile request.
 	// This must be done after we add in the translator errors etc so tha they are passed to the compile job correctly.
@@ -1844,6 +1845,7 @@ int32 FHlslNiagaraCompiler::CompileScript(const FStringView GroupName, const FNi
 	CompilationJob->TranslatorOutput.ScriptData.LastHlslTranslation = TranslatedHLSL;
 	CompilationJob->TranslatorOutput.ScriptData.ExternalDependencies = InTranslateResults.CompileDependencies;
 	CompilationJob->TranslatorOutput.ScriptData.CompileTags = InTranslateResults.CompileTags;
+	CompilationJob->TranslatorOutput.ScriptData.CompileTagsEditorOnly = InTranslateResults.CompileTagsEditorOnly;
 
 	bool bJobScheduled = false;
 	if (CompileResults.bVMSucceeded)
@@ -2136,6 +2138,7 @@ TOptional<FNiagaraCompileResults> FHlslNiagaraCompiler::GetCompileResult(int32 J
 	Results.NumErrors = CompileResults.NumErrors;
 	Results.CompileEvents = CompileResults.CompileEvents;
 	Results.Data->CompileTags = CompileResults.Data->CompileTags;
+	Results.Data->CompileTagsEditorOnly = CompileResults.Data->CompileTagsEditorOnly;
 
 	CompilationJob.Reset();
 	return Results;
diff --git a/Engine/Plugins/FX/Niagara/Source/NiagaraEditor/Private/NiagaraHlslTranslator.cpp b/Engine/Plugins/FX/Niagara/Source/NiagaraEditor/Private/NiagaraHlslTranslator.cpp
--- a/Engine/Plugins/FX/Niagara/Source/NiagaraEditor/Private/NiagaraHlslTranslator.cpp
+++ b/Engine/Plugins/FX/Niagara/Source/NiagaraEditor/Private/NiagaraHlslTranslator.cpp
@@ -2141,7 +2141,15 @@ FNiagaraTranslateResults TNiagaraHlslTranslator<GraphBridge>::Translate(const FN
 			{
 				Preamble.Appendf(TEXT("//\tVariable: \"%s\" StringValue: \"%s\" \n"), *Tag.Variable.ToString(),  *Tag.StringValue);
 			}
+		}
+		if (TranslateResults.CompileTags.Num() != 0)
+		{
+			Preamble.Appendf(TEXT("// Compile Tags Editor Only: \n"));
 
+			for (const FNiagaraCompilerTag& Tag : TranslateResults.CompileTagsEditorOnly)
+			{
+				Preamble.Appendf(TEXT("//\tVariable: \"%s\" StringValue: \"%s\" \n"), *Tag.Variable.ToString(), *Tag.StringValue);
+			}
 		}
 
 		// Display the computed compile tags in the source hlsl to make checking easier.
@@ -7471,7 +7479,7 @@ void TNiagaraHlslTranslator<GraphBridge>::FunctionCall(const FFunctionCallNode*
 				FNiagaraTypeDefinition TypeDef = GraphBridge::GetPinType(CallInputs[i], ENiagaraStructConversion::UserFacing);
 				if (!(TypeDef.IsDataInterface() || TypeDef == FNiagaraTypeDefinition::GetParameterMapDef()))
 				{
-					WriteCompilerTag(Inputs[i], CallInputs[i], false, FNiagaraCompileEventSeverity::Display, ResolvedVariable.GetName().ToString());
+					WriteCompilerTag(Inputs[i], CallInputs[i], false, false, FNiagaraCompileEventSeverity::Display, ResolvedVariable.GetName().ToString());
 				}
 			}
 		}
@@ -9457,7 +9465,7 @@ void TNiagaraHlslTranslator<GraphBridge>::FindConstantValue(int32 InputCompileRe
 }
 
 template<typename GraphBridge>
-void TNiagaraHlslTranslator<GraphBridge>::WriteCompilerTag(int32 InputCompileResult, const FPin* Pin, bool bEmitMessageOnFailure, FNiagaraCompileEventSeverity FailureSeverity, const FString& Prefix)
+void TNiagaraHlslTranslator<GraphBridge>::WriteCompilerTag(int32 InputCompileResult, const FPin* Pin, bool bEditorOnly, bool bEmitMessageOnFailure, FNiagaraCompileEventSeverity FailureSeverity, const FString& Prefix)
 {
 	FString Value;
 	FNiagaraTypeDefinition TypeDef = GraphBridge::GetPinType(Pin, ENiagaraStructConversion::UserFacing);
@@ -9485,15 +9493,17 @@ void TNiagaraHlslTranslator<GraphBridge>::WriteCompilerTag(int32 InputCompileRes
 	}
 	else
 	{
+		TArray<FNiagaraCompilerTag>& CompileTags = bEditorOnly ? TranslateResults.CompileTagsEditorOnly : TranslateResults.CompileTags;
+
 		// Always use the latest output value for the tag.
-		if (FNiagaraCompilerTag* Tag = FNiagaraCompilerTag::FindTag(TranslateResults.CompileTags, Variable))
+		if (FNiagaraCompilerTag* Tag = FNiagaraCompilerTag::FindTag(CompileTags, Variable))
 		{
 			Tag->StringValue = Value;
 			Tag->Variable = Variable;
 		}
 		else
 		{
-			TranslateResults.CompileTags.Emplace(Variable, Value);
+			CompileTags.Emplace(Variable, Value);
 		}
 	}
 }
diff --git a/Engine/Plugins/FX/Niagara/Source/NiagaraEditor/Private/NiagaraHlslTranslator.h b/Engine/Plugins/FX/Niagara/Source/NiagaraEditor/Private/NiagaraHlslTranslator.h
--- a/Engine/Plugins/FX/Niagara/Source/NiagaraEditor/Private/NiagaraHlslTranslator.h
+++ b/Engine/Plugins/FX/Niagara/Source/NiagaraEditor/Private/NiagaraHlslTranslator.h
@@ -651,7 +651,7 @@ public:
 
 	void CullMapSetInputPin(const FPin* InputPin);
 
-	void WriteCompilerTag(int32 InputCompileResult, const FPin* Pin, bool bEmitMessageOnFailure, FNiagaraCompileEventSeverity FailureSeverity, const FString& Prefix = FString());
+	void WriteCompilerTag(int32 InputCompileResult, const FPin* Pin, bool bEditorOnly, bool bEmitMessageOnFailure, FNiagaraCompileEventSeverity FailureSeverity, const FString& Prefix = FString());
 
 	void Message(FNiagaraCompileEventSeverity Severity, FText MessageText, const FNode* Node, const FPin* Pin, FStringView ShortDescription = FStringView());
 	void Error(FText ErrorText, const FNode* Node, const FPin* Pin, FStringView ShortDescription = FStringView());
diff --git a/Engine/Plugins/FX/Niagara/Source/NiagaraEditor/Private/NiagaraNodeOutputTag.cpp b/Engine/Plugins/FX/Niagara/Source/NiagaraEditor/Private/NiagaraNodeOutputTag.cpp
--- a/Engine/Plugins/FX/Niagara/Source/NiagaraEditor/Private/NiagaraNodeOutputTag.cpp
+++ b/Engine/Plugins/FX/Niagara/Source/NiagaraEditor/Private/NiagaraNodeOutputTag.cpp
@@ -58,7 +58,7 @@ void UNiagaraNodeOutputTag::Compile(FTranslator* Translator, TArray<int32>& Outp
 		}
 		else
 		{
-			Translator->WriteCompilerTag(CompiledInput, InputPin, bEmitMessageOnFailure, FailureSeverity);
+			Translator->WriteCompilerTag(CompiledInput, InputPin, bEditorOnly, bEmitMessageOnFailure, FailureSeverity);
 		}
 	}
 
diff --git a/Engine/Plugins/FX/Niagara/Source/NiagaraEditor/Private/NiagaraNodeOutputTag.h b/Engine/Plugins/FX/Niagara/Source/NiagaraEditor/Private/NiagaraNodeOutputTag.h
--- a/Engine/Plugins/FX/Niagara/Source/NiagaraEditor/Private/NiagaraNodeOutputTag.h
+++ b/Engine/Plugins/FX/Niagara/Source/NiagaraEditor/Private/NiagaraNodeOutputTag.h
@@ -30,6 +30,9 @@ public:
 	virtual bool CancelEditablePinName(const FText& InName, UEdGraphPin* InGraphPinObj) override;
 	virtual void BuildParameterMapHistory(FNiagaraParameterMapHistoryBuilder& OutHistory, bool bRecursive = true, bool bFilterForCompilation = true) const override;
 
+	UPROPERTY(EditAnywhere, Category = Tag)
+	bool bEditorOnly = false;
+
 	UPROPERTY(EditAnywhere, Category=Tag)
 	bool bEmitMessageOnFailure = true;
 
diff --git a/Engine/Plugins/FX/Niagara/Source/NiagaraEditor/Public/INiagaraCompiler.h b/Engine/Plugins/FX/Niagara/Source/NiagaraEditor/Public/INiagaraCompiler.h
--- a/Engine/Plugins/FX/Niagara/Source/NiagaraEditor/Public/INiagaraCompiler.h
+++ b/Engine/Plugins/FX/Niagara/Source/NiagaraEditor/Public/INiagaraCompiler.h
@@ -96,7 +96,9 @@ struct FNiagaraTranslateResults
 
 	TArray<FNiagaraCompileDependency> CompileDependencies;
 
-	TArray<FNiagaraCompilerTag > CompileTags;
+	TArray<FNiagaraCompilerTag> CompileTags;
+
+	TArray<FNiagaraCompilerTag> CompileTagsEditorOnly;
 
 	/** A string representation of the compilation output. */
 	FString OutputHLSL;
diff --git a/Engine/Source/Runtime/Core/Private/UObject/FortniteReleaseBranchCustomObjectVersion.cpp b/Engine/Source/Runtime/Core/Private/UObject/FortniteReleaseBranchCustomObjectVersion.cpp
--- a/Engine/Source/Runtime/Core/Private/UObject/FortniteReleaseBranchCustomObjectVersion.cpp
+++ b/Engine/Source/Runtime/Core/Private/UObject/FortniteReleaseBranchCustomObjectVersion.cpp
@@ -6,6 +6,7 @@ TMap<FGuid, FGuid> FFortniteReleaseBranchCustomObjectVersion::GetSystemGuids()
 	TMap<FGuid, FGuid> SystemGuids;
 	const FDevSystemGuids& DevGuids = FDevSystemGuids::Get();
 
+	SystemGuids.Add(DevGuids.Niagara_LatestScriptCompileVersion, FGuid("11BB3AC044D97344A3F554562F60FB1C"));
 	SystemGuids.Add(DevGuids.SkeletalMeshDerivedDataVersion, FGuid("5705956EC7134274A5A1BD3FC74DB3A9"));
 	return SystemGuids;
 }
