- Ensure we post load the system spawn / update before potentially accessing the parameter stores
- Don't allow any pooling until the system is considered fully loaded

#rnx

diff --git a/Engine/Plugins/FX/Niagara/Source/Niagara/Private/NiagaraComponentPool.cpp b/Engine/Plugins/FX/Niagara/Source/Niagara/Private/NiagaraComponentPool.cpp
--- a/Engine/Plugins/FX/Niagara/Source/Niagara/Private/NiagaraComponentPool.cpp
+++ b/Engine/Plugins/FX/Niagara/Source/Niagara/Private/NiagaraComponentPool.cpp
@@ -280,6 +280,11 @@ void UNiagaraComponentPool::PrimePool(UNiagaraSystem* Template, UWorld* World)
 		return;
 	}
 
+	if (!Template->bFullyLoaded)
+	{
+		return;
+	}
+
 	const uint32 PrimeCount = FMath::Min(Template->PoolPrimeSize, Template->MaxPoolSize);
 	if (PrimeCount == 0)
 	{
diff --git a/Engine/Plugins/FX/Niagara/Source/Niagara/Private/NiagaraSystem.cpp b/Engine/Plugins/FX/Niagara/Source/Niagara/Private/NiagaraSystem.cpp
--- a/Engine/Plugins/FX/Niagara/Source/Niagara/Private/NiagaraSystem.cpp
+++ b/Engine/Plugins/FX/Niagara/Source/Niagara/Private/NiagaraSystem.cpp
@@ -431,11 +431,19 @@ void UNiagaraSystem::UpdateSystemAfterLoad()
 	}
 	bFullyLoaded = true;
 
-#if WITH_EDITORONLY_DATA
-	if (SystemSpawnScript && SystemUpdateScript && !GetOutermost()->bIsCookedForEditor)
+	// Ensure parameter stores have been post loaded inside the script otherwise searching them will be invalid
+	if (SystemSpawnScript)
 	{
 		SystemSpawnScript->ConditionalPostLoad();
+	}
+	if (SystemUpdateScript)
+	{
 		SystemUpdateScript->ConditionalPostLoad();
+	}
+
+#if WITH_EDITORONLY_DATA
+	if (SystemSpawnScript && SystemUpdateScript && !GetOutermost()->bIsCookedForEditor)
+	{
 		ensure(SystemSpawnScript->GetLatestSource() != nullptr);
 		ensure(SystemSpawnScript->GetLatestSource() == SystemUpdateScript->GetLatestSource());
 		SystemSpawnScript->GetLatestSource()->OnChanged().AddUObject(this, &UNiagaraSystem::GraphSourceChanged);
@@ -475,7 +483,6 @@ void UNiagaraSystem::UpdateSystemAfterLoad()
 		}
 		else
 		{
-			SystemSpawnScript->ConditionalPostLoad();
 			SystemScriptSource = SystemSpawnScript->GetLatestSource();
 		}
 		AllSystemScripts.Add(SystemSpawnScript);
@@ -486,10 +493,6 @@ void UNiagaraSystem::UpdateSystemAfterLoad()
 			SystemUpdateScript->SetUsage(ENiagaraScriptUsage::SystemUpdateScript);
 			SystemUpdateScript->SetLatestSource(SystemScriptSource);
 		}
-		else
-		{
-			SystemUpdateScript->ConditionalPostLoad();
-		}
 		AllSystemScripts.Add(SystemUpdateScript);
 
 		// Synchronize with parameter definitions
diff --git a/Engine/Plugins/FX/Niagara/Source/Niagara/Private/NiagaraWorldManager.cpp b/Engine/Plugins/FX/Niagara/Source/Niagara/Private/NiagaraWorldManager.cpp
--- a/Engine/Plugins/FX/Niagara/Source/Niagara/Private/NiagaraWorldManager.cpp
+++ b/Engine/Plugins/FX/Niagara/Source/Niagara/Private/NiagaraWorldManager.cpp
@@ -1864,7 +1864,7 @@ void FNiagaraWorldManager::PrimePoolForAllSystems()
 		for (TObjectIterator<UNiagaraSystem> It; It; ++It)
 		{
 			UNiagaraSystem* Sys = *It;
-			if (Sys && Sys->bFullyLoaded)
+			if (Sys)
 			{
 				ComponentPool->PrimePool(Sys, World);
 			}
