- Convert FNiagaraDataSetCompiledData Variables to use Base
- Pack variable layout data more efficiently to avoid 16 bytes per component type * 2
- Move some deprecate properties into editor only data
- Approximate savings are ~15% per asset varies depending on number of attributes

diff --git a/Engine/Plugins/FX/Niagara/Source/Niagara/Classes/NiagaraDataSet.h b/Engine/Plugins/FX/Niagara/Source/Niagara/Classes/NiagaraDataSet.h
--- a/Engine/Plugins/FX/Niagara/Source/Niagara/Classes/NiagaraDataSet.h
+++ b/Engine/Plugins/FX/Niagara/Source/Niagara/Classes/NiagaraDataSet.h
@@ -288,13 +288,13 @@ public:
 	FORCEINLINE FRWBuffer& GetGPUFreeIDs() { return GPUFreeIDs; }
 	FORCEINLINE uint32 GetGPUNumAllocatedIDs() const { return GPUNumAllocatedIDs; }
 
-	FORCEINLINE const TArray<FNiagaraVariable>& GetVariables()const { return CompiledData->Variables; }
-	FORCEINLINE uint32 GetNumVariables()const { return CompiledData->Variables.Num(); }
-	FORCEINLINE bool HasVariable(const FNiagaraVariable& Var)const { return CompiledData->Variables.Contains(Var); }
-	FORCEINLINE bool HasVariable(const FName& Name)const { return CompiledData->Variables.FindByPredicate([&](const FNiagaraVariable& VarInfo) { return VarInfo.GetName() == Name; }) != nullptr; }
-	FORCEINLINE uint32 GetNumFloatComponents()const { return CompiledData->TotalFloatComponents; }
-	FORCEINLINE uint32 GetNumInt32Components()const { return CompiledData->TotalInt32Components; }
-	FORCEINLINE uint32 GetNumHalfComponents()const { return CompiledData->TotalHalfComponents; }
+	FORCEINLINE const TArray<FNiagaraVariableBase>& GetVariables() const { return CompiledData->Variables; }
+	FORCEINLINE uint32 GetNumVariables() const { return CompiledData->Variables.Num(); }
+	FORCEINLINE bool HasVariable(const FNiagaraVariable& Var) const { return CompiledData->Variables.Contains(Var); }
+	FORCEINLINE bool HasVariable(const FName& Name) const { return CompiledData->Variables.FindByPredicate([&](const FNiagaraVariable& VarInfo) { return VarInfo.GetName() == Name; }) != nullptr; }
+	FORCEINLINE uint32 GetNumFloatComponents() const { return CompiledData->TotalFloatComponents; }
+	FORCEINLINE uint32 GetNumInt32Components() const { return CompiledData->TotalInt32Components; }
+	FORCEINLINE uint32 GetNumHalfComponents() const { return CompiledData->TotalHalfComponents; }
 
 	const TArray<FNiagaraVariableLayoutInfo>& GetVariableLayouts() const { return CompiledData->VariableLayouts; }
 	const FNiagaraVariableLayoutInfo* GetVariableLayout(const FNiagaraVariable& Var) const;
@@ -413,7 +413,7 @@ struct FNiagaraDataVariableIterator
 		: Data(InData)
 		, CurrIdx(StartIdx)
 	{
-		Variables = Data->GetOwner()->GetVariables();
+		FNiagaraVariable::ConvertFromBaseArray(Data->GetOwner()->GetVariables(), Variables);
 	}
 
 	void Get()
@@ -428,25 +428,25 @@ struct FNiagaraDataVariableIterator
 
 			for (uint32 CompIdx = 0; CompIdx < Layout.GetNumFloatComponents(); ++CompIdx)
 			{
-				uint32 CompBufferOffset = Layout.FloatComponentStart + CompIdx;
+				uint32 CompBufferOffset = Layout.GetFloatComponentStart() + CompIdx;
 				const float* Src = Data->GetInstancePtrFloat(CompBufferOffset, CurrIdx);
-				float* Dst = (float*)(ValuePtr + Layout.LayoutInfo.FloatComponentByteOffsets[CompIdx]);
+				float* Dst = (float*)(ValuePtr + Layout.LayoutInfo.GetFloatComponentByteOffset(CompIdx));
 				*Dst = *Src;
 			}
 
 			for (uint32 CompIdx = 0; CompIdx < Layout.GetNumInt32Components(); ++CompIdx)
 			{
-				uint32 CompBufferOffset = Layout.Int32ComponentStart + CompIdx;
+				uint32 CompBufferOffset = Layout.GetInt32ComponentStart() + CompIdx;
 				const int32* Src = Data->GetInstancePtrInt32(CompBufferOffset, CurrIdx);
-				int32* Dst = (int32*)(ValuePtr + Layout.LayoutInfo.Int32ComponentByteOffsets[CompIdx]);
+				int32* Dst = (int32*)(ValuePtr + Layout.LayoutInfo.GetInt32ComponentByteOffset(CompIdx));
 				*Dst = *Src;
 			}
 
 			for (uint32 CompIdx = 0; CompIdx < Layout.GetNumHalfComponents(); ++CompIdx)
 			{
-				uint32 CompBufferOffset = Layout.HalfComponentStart + CompIdx;
+				uint32 CompBufferOffset = Layout.GetHalfComponentStart() + CompIdx;
 				const FFloat16* Src = Data->GetInstancePtrHalf(CompBufferOffset, CurrIdx);
-				FFloat16* Dst = (FFloat16*)(ValuePtr + Layout.LayoutInfo.HalfComponentByteOffsets[CompIdx]);
+				FFloat16* Dst = (FFloat16*)(ValuePtr + Layout.LayoutInfo.GetHalfComponentByteOffset(CompIdx));
 				*Dst = *Src;
 			}
 		}
diff --git a/Engine/Plugins/FX/Niagara/Source/Niagara/Classes/NiagaraDataSetAccessor.h b/Engine/Plugins/FX/Niagara/Source/Niagara/Classes/NiagaraDataSetAccessor.h
--- a/Engine/Plugins/FX/Niagara/Source/Niagara/Classes/NiagaraDataSetAccessor.h
+++ b/Engine/Plugins/FX/Niagara/Source/Niagara/Classes/NiagaraDataSetAccessor.h
@@ -213,7 +213,7 @@ struct FNiagaraDataSetAccessorFloat
 		if (FloatVariableIndex != INDEX_NONE)
 		{
 			bIsFloat = true;
-			ComponentIndex = DataSetCompiledData.VariableLayouts[FloatVariableIndex].FloatComponentStart;
+			ComponentIndex = DataSetCompiledData.VariableLayouts[FloatVariableIndex].GetFloatComponentStart();
 			return;
 		}
 
@@ -223,7 +223,7 @@ struct FNiagaraDataSetAccessorFloat
 			if (AlternateVariableIndex != INDEX_NONE)
 			{
 				bIsFloat = true;
-				ComponentIndex = DataSetCompiledData.VariableLayouts[AlternateVariableIndex].FloatComponentStart;
+				ComponentIndex = DataSetCompiledData.VariableLayouts[AlternateVariableIndex].GetFloatComponentStart();
 				return;
 			}
 		}
@@ -234,7 +234,7 @@ struct FNiagaraDataSetAccessorFloat
 			if (HalfVariableIndex != INDEX_NONE)
 			{
 				bIsFloat = false;
-				ComponentIndex = DataSetCompiledData.VariableLayouts[HalfVariableIndex].HalfComponentStart;
+				ComponentIndex = DataSetCompiledData.VariableLayouts[HalfVariableIndex].GetHalfComponentStart();
 			}
 		}
 	}
@@ -478,7 +478,7 @@ struct FNiagaraDataSetAccessorInt32
 		const int32 VariableIndex = DataSetCompiledData.Variables.IndexOfByPredicate(FNiagaraVariableMatch(FNiagaraDataSetAccessorTypeInfo<TType>::GetIntType(), VariableName));
 		if (VariableIndex != INDEX_NONE)
 		{
-			ComponentIndex = DataSetCompiledData.VariableLayouts[VariableIndex].Int32ComponentStart;
+			ComponentIndex = DataSetCompiledData.VariableLayouts[VariableIndex].GetInt32ComponentStart();
 		}
 	}
 
@@ -623,8 +623,8 @@ struct FNiagaraDataSetAccessorStruct
 		if (VariableIndex != INDEX_NONE)
 		{
 			bIsValid = true;
-			FloatComponentIndex = DataSetCompiledData.VariableLayouts[VariableIndex].FloatComponentStart;
-			Int32ComponentIndex = DataSetCompiledData.VariableLayouts[VariableIndex].Int32ComponentStart;
+			FloatComponentIndex = DataSetCompiledData.VariableLayouts[VariableIndex].GetFloatComponentStart();
+			Int32ComponentIndex = DataSetCompiledData.VariableLayouts[VariableIndex].GetInt32ComponentStart();
 		}
 	}
 
diff --git a/Engine/Plugins/FX/Niagara/Source/Niagara/Classes/NiagaraDataSetCompiledData.h b/Engine/Plugins/FX/Niagara/Source/Niagara/Classes/NiagaraDataSetCompiledData.h
--- a/Engine/Plugins/FX/Niagara/Source/Niagara/Classes/NiagaraDataSetCompiledData.h
+++ b/Engine/Plugins/FX/Niagara/Source/Niagara/Classes/NiagaraDataSetCompiledData.h
@@ -10,21 +10,31 @@ struct FNiagaraVariableLayoutInfo
 {
 	GENERATED_BODY()
 
+private:
 	/** Start index for the float components in the main buffer. */
 	UPROPERTY()
-	uint32 FloatComponentStart = 0;
+	uint16 FloatComponentStart = 0;
 
 	/** Start index for the int32 components in the main buffer. */
 	UPROPERTY()
-	uint32 Int32ComponentStart = 0;
+	uint16 Int32ComponentStart = 0;
 
 	/** Start index for the half components in the main buffer. */
 	UPROPERTY()
-	uint32 HalfComponentStart = 0;
+	uint16 HalfComponentStart = 0;
+
+public:
+	uint32 GetFloatComponentStart() const { return FloatComponentStart; }
+	uint32 GetInt32ComponentStart() const { return Int32ComponentStart; }
+	uint32 GetHalfComponentStart() const { return HalfComponentStart; }
+
+	void SetFloatComponentStart(uint32 Offset) { check(Offset <= TNumericLimits<uint16>::Max()); FloatComponentStart = Offset; }
+	void SetInt32ComponentStart(uint32 Offset) { check(Offset <= TNumericLimits<uint16>::Max()); Int32ComponentStart = Offset; }
+	void SetHalfComponentStart(uint32 Offset) { check(Offset <= TNumericLimits<uint16>::Max()); HalfComponentStart = Offset; }
 
-	uint32 GetNumFloatComponents()const { return LayoutInfo.FloatComponentByteOffsets.Num(); }
-	uint32 GetNumInt32Components()const { return LayoutInfo.Int32ComponentByteOffsets.Num(); }
-	uint32 GetNumHalfComponents()const { return LayoutInfo.HalfComponentByteOffsets.Num(); }
+	uint32 GetNumFloatComponents() const { return LayoutInfo.GetNumFloatComponents(); }
+	uint32 GetNumInt32Components() const { return LayoutInfo.GetNumInt32Components(); }
+	uint32 GetNumHalfComponents() const { return LayoutInfo.GetNumHalfComponents(); }
 
 	/** This variable's type layout info. */
 	UPROPERTY()
@@ -38,7 +48,7 @@ struct NIAGARA_API FNiagaraDataSetCompiledData
 
 	/** Variables in the data set. */
 	UPROPERTY()
-	TArray<FNiagaraVariable> Variables;
+	TArray<FNiagaraVariableBase> Variables;
 
 	/** Data describing the layout of variable data. */
 	UPROPERTY()
diff --git a/Engine/Plugins/FX/Niagara/Source/Niagara/Classes/NiagaraEmitter.h b/Engine/Plugins/FX/Niagara/Source/Niagara/Classes/NiagaraEmitter.h
--- a/Engine/Plugins/FX/Niagara/Source/Niagara/Classes/NiagaraEmitter.h
+++ b/Engine/Plugins/FX/Niagara/Source/Niagara/Classes/NiagaraEmitter.h
@@ -490,7 +490,7 @@ struct NIAGARA_API FVersionedNiagaraEmitterData
 	const TCHAR* GetDebugSimName() const { return TEXT(""); }
 #endif
 
-	void GatherCompiledParticleAttributes(TArray<FNiagaraVariable>& OutVariables) const;
+	void GatherCompiledParticleAttributes(TArray<FNiagaraVariableBase>& OutVariables) const;
 	
 private:
 	UPROPERTY()
diff --git a/Engine/Plugins/FX/Niagara/Source/Niagara/Classes/NiagaraScript.h b/Engine/Plugins/FX/Niagara/Source/Niagara/Classes/NiagaraScript.h
--- a/Engine/Plugins/FX/Niagara/Source/Niagara/Classes/NiagaraScript.h
+++ b/Engine/Plugins/FX/Niagara/Source/Niagara/Classes/NiagaraScript.h
@@ -471,7 +471,7 @@ public:
 
 	/** Attributes used by this script. */
 	UPROPERTY()
-	TArray<FNiagaraVariable> Attributes;
+	TArray<FNiagaraVariableBase> Attributes;
 
 	/** Contains various usage information for this script. */
 	UPROPERTY()
diff --git a/Engine/Plugins/FX/Niagara/Source/Niagara/Private/DataInterface/NiagaraDataInterfaceDataChannelCommon.cpp b/Engine/Plugins/FX/Niagara/Source/Niagara/Private/DataInterface/NiagaraDataInterfaceDataChannelCommon.cpp
--- a/Engine/Plugins/FX/Niagara/Source/Niagara/Private/DataInterface/NiagaraDataInterfaceDataChannelCommon.cpp
+++ b/Engine/Plugins/FX/Niagara/Source/Niagara/Private/DataInterface/NiagaraDataInterfaceDataChannelCommon.cpp
@@ -38,9 +38,9 @@ FNDIDataChannel_FunctionToDataSetBinding::FNDIDataChannel_FunctionToDataSetBindi
 			const FNiagaraVariableBase& Param = Parameters[ParamIdx];
 			if (const FNiagaraVariableLayoutInfo* DataSetVariableLayout = DataSetLayout.FindVariableLayoutInfo(Param))
 			{
-				uint32 DataSetFloatRegister = DataSetVariableLayout->FloatComponentStart;
-				uint32 DataSetIntRegister = DataSetVariableLayout->Int32ComponentStart;
-				uint32 DataSetHalfRegister = DataSetVariableLayout->HalfComponentStart;
+				uint32 DataSetFloatRegister = DataSetVariableLayout->GetFloatComponentStart();
+				uint32 DataSetIntRegister = DataSetVariableLayout->GetInt32ComponentStart();
+				uint32 DataSetHalfRegister = DataSetVariableLayout->GetHalfComponentStart();
 
 				GenVMBindings(Param, Param.GetType().GetStruct(), NumFloatComponents, NumInt32Components, NumHalfComponents, DataSetFloatRegister, DataSetIntRegister, DataSetHalfRegister);
 			}
diff --git a/Engine/Plugins/FX/Niagara/Source/Niagara/Private/DataInterface/NiagaraDataInterfaceDataChannelRead.cpp b/Engine/Plugins/FX/Niagara/Source/Niagara/Private/DataInterface/NiagaraDataInterfaceDataChannelRead.cpp
--- a/Engine/Plugins/FX/Niagara/Source/Niagara/Private/DataInterface/NiagaraDataInterfaceDataChannelRead.cpp
+++ b/Engine/Plugins/FX/Niagara/Source/Niagara/Private/DataInterface/NiagaraDataInterfaceDataChannelRead.cpp
@@ -483,9 +483,9 @@ void UNiagaraDataInterfaceDataChannelRead::ProvidePerInstanceDataForRenderThread
 			{
 				if (const FNiagaraVariableLayoutInfo* LayoutInfo = GPUCompiledData.FindVariableLayoutInfo(Param))
 				{
-					TargetData->GPUScriptParameterOffsetTable.Add(LayoutInfo->GetNumFloatComponents() > 0 ? LayoutInfo->FloatComponentStart : INDEX_NONE);
-					TargetData->GPUScriptParameterOffsetTable.Add(LayoutInfo->GetNumInt32Components() > 0 ? LayoutInfo->Int32ComponentStart : INDEX_NONE);
-					TargetData->GPUScriptParameterOffsetTable.Add(LayoutInfo->GetNumHalfComponents() > 0 ? LayoutInfo->HalfComponentStart : INDEX_NONE);
+					TargetData->GPUScriptParameterOffsetTable.Add(LayoutInfo->GetNumFloatComponents() > 0 ? LayoutInfo->GetFloatComponentStart() : INDEX_NONE);
+					TargetData->GPUScriptParameterOffsetTable.Add(LayoutInfo->GetNumInt32Components() > 0 ? LayoutInfo->GetInt32ComponentStart() : INDEX_NONE);
+					TargetData->GPUScriptParameterOffsetTable.Add(LayoutInfo->GetNumHalfComponents() > 0 ? LayoutInfo->GetHalfComponentStart() : INDEX_NONE);
 				}
 				else
 				{
diff --git a/Engine/Plugins/FX/Niagara/Source/Niagara/Private/NiagaraDataChannel.cpp b/Engine/Plugins/FX/Niagara/Source/Niagara/Private/NiagaraDataChannel.cpp
--- a/Engine/Plugins/FX/Niagara/Source/Niagara/Private/NiagaraDataChannel.cpp
+++ b/Engine/Plugins/FX/Niagara/Source/Niagara/Private/NiagaraDataChannel.cpp
@@ -154,9 +154,9 @@ void FNiagaraDataChannelGameData::WriteToDataSet(FNiagaraDataBuffer* DestBuffer,
 		int32 DestVarSize = SimVar.GetSizeInBytes();
 		const FNiagaraVariableLayoutInfo& SimLayout = CompiledData.VariableLayouts[SimVarIndex];
 
-		int32 FloatCompIdx = SimLayout.FloatComponentStart;
-		int32 IntCompIdx = SimLayout.Int32ComponentStart;
-		int32 HalfCompIdx = SimLayout.HalfComponentStart;
+		int32 FloatCompIdx = SimLayout.GetFloatComponentStart();
+		int32 IntCompIdx = SimLayout.GetInt32ComponentStart();
+		int32 HalfCompIdx = SimLayout.GetHalfComponentStart();
 
 		TFunction<void(UScriptStruct*, UScriptStruct*, uint8*)> WriteData;
 		WriteData = [&](UScriptStruct* SrcStruct, UScriptStruct* DestStruct, uint8* SrcPropertyBase)
@@ -310,9 +310,9 @@ void FNiagaraDataChannelGameData::AppendFromDataSet(const FNiagaraDataBuffer* Sr
 
 		const FNiagaraVariableLayoutInfo& SimLayout = CompiledData.VariableLayouts[SimVarIndex];
 		
-		int32 FloatCompIdx = SimLayout.FloatComponentStart;
-		int32 IntCompIdx = SimLayout.Int32ComponentStart;
-		int32 HalfCompIdx = SimLayout.HalfComponentStart;
+		int32 FloatCompIdx = SimLayout.GetFloatComponentStart();
+		int32 IntCompIdx = SimLayout.GetInt32ComponentStart();
+		int32 HalfCompIdx = SimLayout.GetHalfComponentStart();
 		
 		TFunction<void(UScriptStruct*, UScriptStruct*, uint8*)> ReadData;
 		ReadData = [&](UScriptStruct* SrcStruct, UScriptStruct* DestStruct, uint8* DestData)
diff --git a/Engine/Plugins/FX/Niagara/Source/Niagara/Private/NiagaraDataInterfaceParticleRead.cpp b/Engine/Plugins/FX/Niagara/Source/Niagara/Private/NiagaraDataInterfaceParticleRead.cpp
--- a/Engine/Plugins/FX/Niagara/Source/Niagara/Private/NiagaraDataInterfaceParticleRead.cpp
+++ b/Engine/Plugins/FX/Niagara/Source/Niagara/Private/NiagaraDataInterfaceParticleRead.cpp
@@ -291,7 +291,7 @@ static void NDIParticleRead_FindAttributeIndices(FNDIParticleRead_RenderInstance
 	InstanceData->AttributeCompressed.SetNumUninitialized(NumAttrIndices);
 
 	// Find the register index for each named attribute in the source emitter.
-	const TArray<FNiagaraVariable>& SourceEmitterVariables = SourceDataSet->GetVariables();
+	const TArray<FNiagaraVariableBase>& SourceEmitterVariables = SourceDataSet->GetVariables();
 	const TArray<FNiagaraVariableLayoutInfo>& SourceEmitterVariableLayouts = SourceDataSet->GetVariableLayouts();
 	for (int AttrNameIdx = 0; AttrNameIdx < ShaderStorage.AttributeNames.Num(); ++AttrNameIdx)
 	{
@@ -306,7 +306,7 @@ static void NDIParticleRead_FindAttributeIndices(FNDIParticleRead_RenderInstance
 		bool FoundVariable = false;
 		for (int VarIdx = 0; VarIdx < SourceEmitterVariables.Num(); ++VarIdx)
 		{
-			const FNiagaraVariable& Var = SourceEmitterVariables[VarIdx];
+			const FNiagaraVariableBase& Var = SourceEmitterVariables[VarIdx];
 			if (Var.GetName() == AttrName)
 			{
 				ENiagaraParticleDataValueType AttributeType = ShaderStorage.AttributeTypes[AttrNameIdx];
@@ -314,13 +314,13 @@ static void NDIParticleRead_FindAttributeIndices(FNDIParticleRead_RenderInstance
 				{
 					const FNiagaraVariableLayoutInfo& Layout = SourceEmitterVariableLayouts[VarIdx];
 					InstanceData->AttributeIndices[AttrNameIdx] =
-						(AttributeType == ENiagaraParticleDataValueType::Int || AttributeType == ENiagaraParticleDataValueType::Bool || AttributeType == ENiagaraParticleDataValueType::ID) ? Layout.Int32ComponentStart : Layout.FloatComponentStart;
+						(AttributeType == ENiagaraParticleDataValueType::Int || AttributeType == ENiagaraParticleDataValueType::Bool || AttributeType == ENiagaraParticleDataValueType::ID) ? Layout.GetInt32ComponentStart() : Layout.GetFloatComponentStart();
 					InstanceData->AttributeCompressed[AttrNameIdx] = 0;
 				}
 				else if (CheckHalfVariableType(Var.GetType(), AttributeType))
 				{
 					const FNiagaraVariableLayoutInfo& Layout = SourceEmitterVariableLayouts[VarIdx];
-					InstanceData->AttributeIndices[AttrNameIdx] = Layout.HalfComponentStart;
+					InstanceData->AttributeIndices[AttrNameIdx] = Layout.GetHalfComponentStart();
 					InstanceData->AttributeCompressed[AttrNameIdx] = 1;
 				}
 				else
@@ -354,7 +354,7 @@ static void NDIParticleRead_FindAttributeIndices(FNDIParticleRead_RenderInstance
 			const FNiagaraVariable& Var = SourceEmitterVariables[VarIdx];
 			if (Var.GetName() == FName_ID)
 			{
-				InstanceData->AcquireTagRegisterIndex = SourceEmitterVariableLayouts[VarIdx].Int32ComponentStart + 1;
+				InstanceData->AcquireTagRegisterIndex = SourceEmitterVariableLayouts[VarIdx].GetInt32ComponentStart() + 1;
 				break;
 			}
 		}
@@ -956,9 +956,9 @@ DEFINE_NDI_DIRECT_FUNC_BINDER_WITH_PAYLOAD(UNiagaraDataInterfaceParticleRead, Re
 DEFINE_NDI_DIRECT_FUNC_BINDER_WITH_PAYLOAD(UNiagaraDataInterfaceParticleRead, ReadQuatByIndex);
 DEFINE_NDI_DIRECT_FUNC_BINDER_WITH_PAYLOAD(UNiagaraDataInterfaceParticleRead, ReadIDByIndex);
 
-static bool HasMatchingVariable(TArrayView<const FNiagaraVariable> Variables, FName AttributeName, const FNiagaraTypeDefinition& ValidType)
+static bool HasMatchingVariable(TArrayView<const FNiagaraVariableBase> Variables, FName AttributeName, const FNiagaraTypeDefinition& ValidType)
 {
-	return Variables.Find(FNiagaraVariable(ValidType, AttributeName)) != INDEX_NONE;
+	return Variables.Find(FNiagaraVariableBase(ValidType, AttributeName)) != INDEX_NONE;
 }
 
 void UNiagaraDataInterfaceParticleRead::GetVMExternalFunction(const FVMExternalFunctionBindingInfo& BindingInfo, void* InstanceData, FVMExternalFunction& OutFunc)
@@ -1013,7 +1013,7 @@ void UNiagaraDataInterfaceParticleRead::GetVMExternalFunction(const FVMExternalF
 		return;
 	}
 
-	TArrayView<const FNiagaraVariable> EmitterVariables;
+	TArrayView<const FNiagaraVariableBase> EmitterVariables;
 	if (PIData->EmitterInstance)
 	{
 		EmitterVariables = PIData->EmitterInstance->GetData().GetVariables();
@@ -2465,7 +2465,7 @@ void UNiagaraDataInterfaceParticleRead::GetFeedback(UNiagaraSystem* Asset, UNiag
 									return Var.GetName() == *AttributeName;
 								};
 
-								TArray<FNiagaraVariable> Variables;
+								TArray<FNiagaraVariableBase> Variables;
 								FoundSourceEmitter.GetEmitterData()->GatherCompiledParticleAttributes(Variables);
 								const FNiagaraVariableBase* FoundVar = Variables.FindByPredicate(AttribFilter);
 								if (FoundVar && !CheckVariableType(FoundVar->GetType(), AttributeType))
diff --git a/Engine/Plugins/FX/Niagara/Source/Niagara/Private/NiagaraDataSet.cpp b/Engine/Plugins/FX/Niagara/Source/Niagara/Private/NiagaraDataSet.cpp
--- a/Engine/Plugins/FX/Niagara/Source/Niagara/Private/NiagaraDataSet.cpp
+++ b/Engine/Plugins/FX/Niagara/Source/Niagara/Private/NiagaraDataSet.cpp
@@ -484,9 +484,9 @@ bool FNiagaraDataSet::GetVariableComponentOffsets(const FNiagaraVariable& Var, i
 	const FNiagaraVariableLayoutInfo *Info = GetVariableLayout(Var);
 	if (Info)
 	{
-		FloatStart = Info->FloatComponentStart;
-		IntStart = Info->Int32ComponentStart;
-		HalfStart = Info->HalfComponentStart;
+		FloatStart = Info->GetFloatComponentStart();
+		IntStart = Info->GetInt32ComponentStart();
+		HalfStart = Info->GetHalfComponentStart();
 		return true;
 	}
 
@@ -1063,8 +1063,8 @@ void FNiagaraDataBuffer::CopyToUnrelated(FNiagaraDataBuffer& DestBuffer, int32 S
 				check(FloatComponents == DestVarLayout.GetNumFloatComponents());
 				for (uint32 CompIdx = 0; CompIdx < FloatComponents; ++CompIdx)
 				{
-					const int32 SrcCompOffset = SrcVarLayout.FloatComponentStart + CompIdx;
-					const int32 DestCompOffest = DestVarLayout.FloatComponentStart + CompIdx;
+					const int32 SrcCompOffset = SrcVarLayout.GetFloatComponentStart() + CompIdx;
+					const int32 DestCompOffest = DestVarLayout.GetFloatComponentStart() + CompIdx;
 					const float* SrcStart = GetInstancePtrFloat(SrcCompOffset, StartIdx);
 					const float* SrcEnd = GetInstancePtrFloat(SrcCompOffset, StartIdx + InstancesToCopy);
 					float* Dst = DestBuffer.GetInstancePtrFloat(DestCompOffest, DestStartIdx);
@@ -1083,8 +1083,8 @@ void FNiagaraDataBuffer::CopyToUnrelated(FNiagaraDataBuffer& DestBuffer, int32 S
 				check(IntComponents == DestVarLayout.GetNumInt32Components());
 				for (uint32 CompIdx = 0; CompIdx < IntComponents; ++CompIdx)
 				{
-					const int32 SrcCompOffset = SrcVarLayout.Int32ComponentStart + CompIdx;
-					const int32 DestCompOffest = DestVarLayout.Int32ComponentStart + CompIdx;
+					const int32 SrcCompOffset = SrcVarLayout.GetInt32ComponentStart() + CompIdx;
+					const int32 DestCompOffest = DestVarLayout.GetInt32ComponentStart() + CompIdx;
 					const int32* SrcStart = GetInstancePtrInt32(SrcCompOffset, StartIdx);
 					const int32* SrcEnd = GetInstancePtrInt32(SrcCompOffset, StartIdx + InstancesToCopy);
 					int32* Dst = DestBuffer.GetInstancePtrInt32(DestCompOffest, DestStartIdx);
@@ -1103,8 +1103,8 @@ void FNiagaraDataBuffer::CopyToUnrelated(FNiagaraDataBuffer& DestBuffer, int32 S
 				check(HalfComponents == DestVarLayout.GetNumHalfComponents());
 				for (uint32 CompIdx = 0; CompIdx < HalfComponents; ++CompIdx)
 				{
-					const int32 SrcCompOffset = SrcVarLayout.HalfComponentStart + CompIdx;
-					const int32 DestCompOffest = DestVarLayout.HalfComponentStart + CompIdx;
+					const int32 SrcCompOffset = SrcVarLayout.GetHalfComponentStart() + CompIdx;
+					const int32 DestCompOffest = DestVarLayout.GetHalfComponentStart() + CompIdx;
 					const FFloat16* SrcStart = GetInstancePtrHalf(SrcCompOffset, StartIdx);
 					const FFloat16* SrcEnd = GetInstancePtrHalf(SrcCompOffset, StartIdx + InstancesToCopy);
 					FFloat16* Dst = DestBuffer.GetInstancePtrHalf(DestCompOffest, DestStartIdx);
@@ -1523,8 +1523,8 @@ void FNiagaraDataBuffer::BuildRegisterTable()
 		const int32 NumFloats = VarLayout.GetNumFloatComponents();
 		for (int32 CompIdx = 0; CompIdx < NumFloats; ++CompIdx)
 		{
-			uint32 CompBufferOffset = VarLayout.FloatComponentStart + CompIdx;
-			uint32 CompRegisterOffset = VarLayout.LayoutInfo.FloatComponentRegisterOffsets[CompIdx];
+			uint32 CompBufferOffset = VarLayout.GetFloatComponentStart() + CompIdx;
+			uint32 CompRegisterOffset = VarLayout.LayoutInfo.GetFloatComponentRegisterOffset(CompIdx);
 			RegisterTable[NumRegisters + CompRegisterOffset] = (uint8*)GetComponentPtrFloat(CompBufferOffset);
 			check(RegisterTable[NumRegisters + CompRegisterOffset]);
 		}
@@ -1537,8 +1537,8 @@ void FNiagaraDataBuffer::BuildRegisterTable()
 		const int32 NumInts = VarLayout.GetNumInt32Components();
 		for (int32 CompIdx = 0; CompIdx < NumInts; ++CompIdx)
 		{
-			uint32 CompBufferOffset = VarLayout.Int32ComponentStart + CompIdx;
-			uint32 CompRegisterOffset = VarLayout.LayoutInfo.Int32ComponentRegisterOffsets[CompIdx];
+			uint32 CompBufferOffset = VarLayout.GetInt32ComponentStart() + CompIdx;
+			uint32 CompRegisterOffset = VarLayout.LayoutInfo.GetInt32ComponentRegisterOffset(CompIdx);
 			RegisterTable[NumRegisters + CompRegisterOffset] = (uint8*)GetComponentPtrInt32(CompBufferOffset);
 			check(RegisterTable[NumRegisters + CompRegisterOffset]);
 		}
@@ -1551,8 +1551,8 @@ void FNiagaraDataBuffer::BuildRegisterTable()
 		const int32 NumHalfs = VarLayout.GetNumHalfComponents();
 		for (int32 CompIdx = 0; CompIdx < NumHalfs; ++CompIdx)
 		{
-			uint32 CompBufferOffset = VarLayout.HalfComponentStart + CompIdx;
-			uint32 CompRegisterOffset = VarLayout.LayoutInfo.HalfComponentRegisterOffsets[CompIdx];
+			uint32 CompBufferOffset = VarLayout.GetHalfComponentStart() + CompIdx;
+			uint32 CompRegisterOffset = VarLayout.LayoutInfo.GetHalfComponentRegisterOffset(CompIdx);
 			RegisterTable[NumRegisters + CompRegisterOffset] = (uint8*)GetComponentPtrHalf(CompBufferOffset);
 			check(RegisterTable[NumRegisters + CompRegisterOffset]);
 		}
@@ -1589,13 +1589,13 @@ void FNiagaraDataSetCompiledData::BuildLayout()
 	TotalInt32Components = 0;
 	TotalHalfComponents = 0;
 
-	for (FNiagaraVariable& Var : Variables)
+	for (const FNiagaraVariableBase& Var : Variables)
 	{
 		FNiagaraVariableLayoutInfo& VarInfo = VariableLayouts[VariableLayouts.AddDefaulted()];
-		FNiagaraTypeLayoutInfo::GenerateLayoutInfo(VarInfo.LayoutInfo, Var.GetType().GetScriptStruct());
-		VarInfo.FloatComponentStart = TotalFloatComponents;
-		VarInfo.Int32ComponentStart = TotalInt32Components;
-		VarInfo.HalfComponentStart = TotalHalfComponents;
+		VarInfo.LayoutInfo.GenerateLayoutInfo(Var.GetType().GetScriptStruct());
+		VarInfo.SetFloatComponentStart(TotalFloatComponents);
+		VarInfo.SetInt32ComponentStart(TotalInt32Components);
+		VarInfo.SetHalfComponentStart(TotalHalfComponents);
 		TotalFloatComponents += VarInfo.GetNumFloatComponents();
 		TotalInt32Components += VarInfo.GetNumInt32Components();
 		TotalHalfComponents += VarInfo.GetNumHalfComponents();
diff --git a/Engine/Plugins/FX/Niagara/Source/Niagara/Private/NiagaraDataSetDebugAccessor.cpp b/Engine/Plugins/FX/Niagara/Source/Niagara/Private/NiagaraDataSetDebugAccessor.cpp
--- a/Engine/Plugins/FX/Niagara/Source/Niagara/Private/NiagaraDataSetDebugAccessor.cpp
+++ b/Engine/Plugins/FX/Niagara/Source/Niagara/Private/NiagaraDataSetDebugAccessor.cpp
@@ -26,9 +26,9 @@ bool FNiagaraDataSetDebugAccessor::Init(const FNiagaraDataSetCompiledData& Compi
 		NumComponentsInt32 = VariableLayout.GetNumInt32Components();
 		NumComponentsFloat = VariableLayout.GetNumFloatComponents();
 		NumComponentsHalf = VariableLayout.GetNumHalfComponents();
-		ComponentIndexInt32 = VariableLayout.Int32ComponentStart; 
-		ComponentIndexFloat = VariableLayout.FloatComponentStart;
-		ComponentIndexHalf = VariableLayout.HalfComponentStart;
+		ComponentIndexInt32 = VariableLayout.GetInt32ComponentStart();
+		ComponentIndexFloat = VariableLayout.GetFloatComponentStart();
+		ComponentIndexHalf = VariableLayout.GetHalfComponentStart();
 		NiagaraType = Variable.GetType();
 
 		bIsBool = Variable.GetType().IsSameBaseDefinition(FNiagaraTypeDefinition::GetBoolDef());
@@ -84,7 +84,7 @@ bool FNiagaraDataSetDebugAccessor::ValidateDataBuffer(const FNiagaraDataSetCompi
 		// Look over float data
 		if (CompiledData.VariableLayouts[iVariable].GetNumFloatComponents() > 0)
 		{
-			const int32 ComponentIndex = CompiledData.VariableLayouts[iVariable].FloatComponentStart;
+			const int32 ComponentIndex = CompiledData.VariableLayouts[iVariable].GetFloatComponentStart();
 			const int32 NumComponents = CompiledData.VariableLayouts[iVariable].GetNumFloatComponents();
 
 			for (int32 iComponent = 0; iComponent < NumComponents; ++iComponent)
@@ -101,7 +101,7 @@ bool FNiagaraDataSetDebugAccessor::ValidateDataBuffer(const FNiagaraDataSetCompi
 		}
 		else if (CompiledData.VariableLayouts[iVariable].GetNumHalfComponents() > 0)
 		{
-			const int32 ComponentIndex = CompiledData.VariableLayouts[iVariable].HalfComponentStart;
+			const int32 ComponentIndex = CompiledData.VariableLayouts[iVariable].GetHalfComponentStart();
 			const int32 NumComponents = CompiledData.VariableLayouts[iVariable].GetNumHalfComponents();
 
 			for (int32 iComponent = 0; iComponent < NumComponents; ++iComponent)
@@ -132,7 +132,7 @@ bool FNiagaraDataSetDebugAccessor::ValidateDataBuffer(const FNiagaraDataSetCompi
 		// Look over float data
 		if (CompiledData.VariableLayouts[iVariable].GetNumFloatComponents() > 0)
 		{
-			const int32 ComponentIndex = CompiledData.VariableLayouts[iVariable].FloatComponentStart;
+			const int32 ComponentIndex = CompiledData.VariableLayouts[iVariable].GetFloatComponentStart();
 			const int32 NumComponents = CompiledData.VariableLayouts[iVariable].GetNumFloatComponents();
 
 			for ( int32 iComponent=0; iComponent < NumComponents; ++iComponent )
@@ -152,7 +152,7 @@ bool FNiagaraDataSetDebugAccessor::ValidateDataBuffer(const FNiagaraDataSetCompi
 		}
 		else if (CompiledData.VariableLayouts[iVariable].GetNumHalfComponents() > 0)
 		{
-			const int32 ComponentIndex = CompiledData.VariableLayouts[iVariable].HalfComponentStart;
+			const int32 ComponentIndex = CompiledData.VariableLayouts[iVariable].GetHalfComponentStart();
 			const int32 NumComponents = CompiledData.VariableLayouts[iVariable].GetNumHalfComponents();
 
 			for (int32 iComponent = 0; iComponent < NumComponents; ++iComponent)
diff --git a/Engine/Plugins/FX/Niagara/Source/Niagara/Private/NiagaraEmitter.cpp b/Engine/Plugins/FX/Niagara/Source/Niagara/Private/NiagaraEmitter.cpp
--- a/Engine/Plugins/FX/Niagara/Source/Niagara/Private/NiagaraEmitter.cpp
+++ b/Engine/Plugins/FX/Niagara/Source/Niagara/Private/NiagaraEmitter.cpp
@@ -1738,7 +1738,7 @@ void FVersionedNiagaraEmitterData::UpdateDebugName(const UNiagaraEmitter& Emitte
 			{
 				if (const FNiagaraVariableLayoutInfo* VariableInfo = CompiledData->FindVariableLayoutInfo(FNiagaraVariableBase(FNiagaraTypeDefinition::GetIntDef(), SimStageMetaData.ParticleIterationStateBinding)))
 				{
-					SimStageMetaData.ParticleIterationStateComponentIndex = VariableInfo->Int32ComponentStart;
+					SimStageMetaData.ParticleIterationStateComponentIndex = VariableInfo->GetInt32ComponentStart();
 				}
 			}
 		}
@@ -2273,7 +2273,7 @@ void UNiagaraEmitter::RebindNotifications()
 #endif
 
 
-void FVersionedNiagaraEmitterData::GatherCompiledParticleAttributes(TArray<FNiagaraVariable>& OutVariables) const
+void FVersionedNiagaraEmitterData::GatherCompiledParticleAttributes(TArray<FNiagaraVariableBase>& OutVariables) const
 {
 	if (SimTarget == ENiagaraSimTarget::GPUComputeSim)
 	{
diff --git a/Engine/Plugins/FX/Niagara/Source/Niagara/Private/NiagaraEmitterHandle.cpp b/Engine/Plugins/FX/Niagara/Source/Niagara/Private/NiagaraEmitterHandle.cpp
--- a/Engine/Plugins/FX/Niagara/Source/Niagara/Private/NiagaraEmitterHandle.cpp
+++ b/Engine/Plugins/FX/Niagara/Source/Niagara/Private/NiagaraEmitterHandle.cpp
@@ -17,8 +17,8 @@ FNiagaraEmitterHandle::FNiagaraEmitterHandle()
 	, Source_DEPRECATED(nullptr)
 	, LastMergedSource_DEPRECATED(nullptr)
 	, bIsolated(false)
-#endif
 	, Instance_DEPRECATED(nullptr)
+#endif
 {
 }
 
@@ -190,11 +190,13 @@ bool FNiagaraEmitterHandle::NeedsRecompile() const
 
 void FNiagaraEmitterHandle::ConditionalPostLoad(int32 NiagaraCustomVersion)
 {
+#if WITH_EDITORONLY_DATA
 	if (Instance_DEPRECATED != nullptr)
 	{
 		VersionedInstance.Emitter = Instance_DEPRECATED;
 		VersionedInstance.Version = Instance_DEPRECATED->IsVersioningEnabled() ? Instance_DEPRECATED->GetExposedVersion().VersionGuid : FGuid(); 
 	}
+#endif
 
 	if (UNiagaraEmitter* InstanceEmitter = VersionedInstance.Emitter)
 	{
diff --git a/Engine/Plugins/FX/Niagara/Source/Niagara/Private/NiagaraEmitterInstance.cpp b/Engine/Plugins/FX/Niagara/Source/Niagara/Private/NiagaraEmitterInstance.cpp
--- a/Engine/Plugins/FX/Niagara/Source/Niagara/Private/NiagaraEmitterInstance.cpp
+++ b/Engine/Plugins/FX/Niagara/Source/Niagara/Private/NiagaraEmitterInstance.cpp
@@ -331,7 +331,7 @@ void FNiagaraEmitterInstance::Init(int32 InEmitterIdx, FNiagaraSystemInstanceID
 #if !UE_BUILD_SHIPPING && !UE_BUILD_TEST 
 		if (bVerboseAttributeLogging)
 		{
-			for (FNiagaraVariable& Attr : EmitterData->UpdateScriptProps.Script->GetVMExecutableData().Attributes)
+			for (const FNiagaraVariableBase& Attr : EmitterData->UpdateScriptProps.Script->GetVMExecutableData().Attributes)
 			{
 				int32 FoundIdx;
 				if (!EmitterData->SpawnScriptProps.Script->GetVMExecutableData().Attributes.Find(Attr, FoundIdx))
diff --git a/Engine/Plugins/FX/Niagara/Source/Niagara/Private/NiagaraRendererProperties.cpp b/Engine/Plugins/FX/Niagara/Source/Niagara/Private/NiagaraRendererProperties.cpp
--- a/Engine/Plugins/FX/Niagara/Source/Niagara/Private/NiagaraRendererProperties.cpp
+++ b/Engine/Plugins/FX/Niagara/Source/Niagara/Private/NiagaraRendererProperties.cpp
@@ -126,7 +126,7 @@ bool FNiagaraRendererLayout::SetVariable(const FNiagaraDataSetCompiledData* Comp
 	const FNiagaraVariableLayoutInfo& DataSetVariableLayout = CompiledData->VariableLayouts[VariableIndex];
 	const int32 VarSize = bHalfVariable ? sizeof(FFloat16) : sizeof(float);
 	const int32 NumComponents = DataSetVariable.GetSizeInBytes() / VarSize;
-	const int32 Offset = bHalfVariable ? DataSetVariableLayout.HalfComponentStart : DataSetVariableLayout.FloatComponentStart;
+	const int32 Offset = bHalfVariable ? DataSetVariableLayout.GetHalfComponentStart() : DataSetVariableLayout.GetFloatComponentStart();
 	uint16& TotalVFComponents = bHalfVariable ? TotalHalfComponents_GT : TotalFloatComponents_GT;
 
 	int32 GPULocation = INDEX_NONE;
@@ -745,17 +745,17 @@ uint32 UNiagaraRendererProperties::ComputeMaxUsedComponents(const FNiagaraDataSe
 
 			if (const uint32 FloatCount = DataSetVarLayout.GetNumFloatComponents())
 			{
-				AccumulateUniqueComponents(BaseType_Float, FloatCount, DataSetVarLayout.FloatComponentStart);
+				AccumulateUniqueComponents(BaseType_Float, FloatCount, DataSetVarLayout.GetFloatComponentStart());
 			}
 
 			if (const uint32 IntCount = DataSetVarLayout.GetNumInt32Components())
 			{
-				AccumulateUniqueComponents(BaseType_Int, IntCount, DataSetVarLayout.Int32ComponentStart);
+				AccumulateUniqueComponents(BaseType_Int, IntCount, DataSetVarLayout.GetInt32ComponentStart());
 			}
 
 			if (const uint32 HalfCount = DataSetVarLayout.GetNumHalfComponents())
 			{
-				AccumulateUniqueComponents(BaseType_Half, HalfCount, DataSetVarLayout.HalfComponentStart);
+				AccumulateUniqueComponents(BaseType_Half, HalfCount, DataSetVarLayout.GetHalfComponentStart());
 			}
 		}
 	}
diff --git a/Engine/Plugins/FX/Niagara/Source/Niagara/Private/NiagaraSimCacheHelper.h b/Engine/Plugins/FX/Niagara/Source/Niagara/Private/NiagaraSimCacheHelper.h
--- a/Engine/Plugins/FX/Niagara/Source/Niagara/Private/NiagaraSimCacheHelper.h
+++ b/Engine/Plugins/FX/Niagara/Source/Niagara/Private/NiagaraSimCacheHelper.h
@@ -205,7 +205,7 @@ struct FNiagaraSimCacheHelper
 			{
 				const FNiagaraVariableLayoutInfo& DataSetVariableLayout = CompiledData.VariableLayouts[CacheLayout.CacheBufferWriteInfo.ComponentUniqueID];
 				check(DataSetVariableLayout.GetNumInt32Components() == 1);
-				CacheLayout.CacheBufferWriteInfo.ComponentUniqueID = DataSetVariableLayout.Int32ComponentStart;
+				CacheLayout.CacheBufferWriteInfo.ComponentUniqueID = DataSetVariableLayout.GetInt32ComponentStart();
 				CacheLayout.bAllowInterpolation = true;
 			}
 		}
@@ -226,19 +226,19 @@ struct FNiagaraSimCacheHelper
 
 			for (int32 iComponent=0; iComponent < CacheVariable.FloatCount; ++iComponent)
 			{
-				CacheLayout.CacheBufferWriteInfo.ComponentMappingsFromDataBuffer[FloatOffset] = uint16(DataSetVariableLayout.FloatComponentStart + iComponent);
+				CacheLayout.CacheBufferWriteInfo.ComponentMappingsFromDataBuffer[FloatOffset] = uint16(DataSetVariableLayout.GetFloatComponentStart() + iComponent);
 				++FloatOffset;
 			}
 
 			for (int32 iComponent=0; iComponent < CacheVariable.HalfCount; ++iComponent)
 			{
-				CacheLayout.CacheBufferWriteInfo.ComponentMappingsFromDataBuffer[HalfOffset] = uint16(DataSetVariableLayout.HalfComponentStart + iComponent);
+				CacheLayout.CacheBufferWriteInfo.ComponentMappingsFromDataBuffer[HalfOffset] = uint16(DataSetVariableLayout.GetHalfComponentStart() + iComponent);
 				++HalfOffset;
 			}
 
 			for (int32 iComponent=0; iComponent < CacheVariable.Int32Count; ++iComponent)
 			{
-				CacheLayout.CacheBufferWriteInfo.ComponentMappingsFromDataBuffer[Int32Offset] = uint16(DataSetVariableLayout.Int32ComponentStart + iComponent);
+				CacheLayout.CacheBufferWriteInfo.ComponentMappingsFromDataBuffer[Int32Offset] = uint16(DataSetVariableLayout.GetInt32ComponentStart() + iComponent);
 				++Int32Offset;
 			}
 		}
@@ -526,23 +526,23 @@ struct FNiagaraSimCacheHelper
 
 								CacheBufferReadInfo.VariableCopyMappingsToDataBuffer.Emplace(
 									SourceVariable.FloatOffset, PrevSourceVariable->FloatOffset,
-									uint16(DestVariableLayout->FloatComponentStart), uint16(CompiledData.VariableLayouts[PreviousDataSetVariableIndex].FloatComponentStart),
+									uint16(DestVariableLayout->GetFloatComponentStart()), uint16(CompiledData.VariableLayouts[PreviousDataSetVariableIndex].GetFloatComponentStart()),
 									bRebaseVariable ? &FNiagaraSimCacheHelper::InterpPositions<true, true> : &FNiagaraSimCacheHelper::InterpPositions<false, true>
 									
 								);
 							}
 							else
 							{
-								CacheBufferReadInfo.VariableCopyMappingsToDataBuffer.Emplace(SourceVariable.FloatOffset, uint16(DestVariableLayout->FloatComponentStart), bRebaseVariable ? &FNiagaraSimCacheHelper::InterpPositions<true, false> : &FNiagaraSimCacheHelper::InterpPositions<false, false>);
+								CacheBufferReadInfo.VariableCopyMappingsToDataBuffer.Emplace(SourceVariable.FloatOffset, uint16(DestVariableLayout->GetFloatComponentStart()), bRebaseVariable ? &FNiagaraSimCacheHelper::InterpPositions<true, false> : &FNiagaraSimCacheHelper::InterpPositions<false, false>);
 							}
 						}
 						else if ( CacheLayout.bAllowVelocityExtrapolation )
 						{
-							CacheBufferReadInfo.VariableCopyMappingsToDataBuffer.Emplace(SourceVariable.FloatOffset, uint16(DestVariableLayout->FloatComponentStart), bRebaseVariable ? &FNiagaraSimCacheHelper::ExtrapolatePositions<true> : &FNiagaraSimCacheHelper::ExtrapolatePositions<false>);
+							CacheBufferReadInfo.VariableCopyMappingsToDataBuffer.Emplace(SourceVariable.FloatOffset, uint16(DestVariableLayout->GetFloatComponentStart()), bRebaseVariable ? &FNiagaraSimCacheHelper::ExtrapolatePositions<true> : &FNiagaraSimCacheHelper::ExtrapolatePositions<false>);
 						}
 						else
 						{
-							CacheBufferReadInfo.VariableCopyMappingsToDataBuffer.Emplace(SourceVariable.FloatOffset, uint16(DestVariableLayout->FloatComponentStart), &FNiagaraSimCacheHelper::RebasePositions);
+							CacheBufferReadInfo.VariableCopyMappingsToDataBuffer.Emplace(SourceVariable.FloatOffset, uint16(DestVariableLayout->GetFloatComponentStart()), &FNiagaraSimCacheHelper::RebasePositions);
 						}
 						DestVariableLayout = nullptr;
 					}
@@ -562,19 +562,19 @@ struct FNiagaraSimCacheHelper
 
 								CacheBufferReadInfo.VariableCopyMappingsToDataBuffer.Emplace(
 									SourceVariable.FloatOffset, PrevSourceVariable->FloatOffset,
-									uint16(DestVariableLayout->FloatComponentStart), uint16(CompiledData.VariableLayouts[PreviousDataSetVariableIndex].FloatComponentStart),
+									uint16(DestVariableLayout->GetFloatComponentStart()), uint16(CompiledData.VariableLayouts[PreviousDataSetVariableIndex].GetFloatComponentStart()),
 									bRebaseVariable ? &FNiagaraSimCacheHelper::InterpQuaternions<true, true> : &FNiagaraSimCacheHelper::InterpQuaternions<false, true>
 									
 								);
 							}
 							else
 							{
-								CacheBufferReadInfo.VariableCopyMappingsToDataBuffer.Emplace(SourceVariable.FloatOffset, uint16(DestVariableLayout->FloatComponentStart), bInterpVariable  ? &FNiagaraSimCacheHelper::InterpQuaternions<true, false> : &FNiagaraSimCacheHelper::InterpQuaternions<false, false>);
+								CacheBufferReadInfo.VariableCopyMappingsToDataBuffer.Emplace(SourceVariable.FloatOffset, uint16(DestVariableLayout->GetFloatComponentStart()), bInterpVariable  ? &FNiagaraSimCacheHelper::InterpQuaternions<true, false> : &FNiagaraSimCacheHelper::InterpQuaternions<false, false>);
 							}
 						}
 						else
 						{
-							CacheBufferReadInfo.VariableCopyMappingsToDataBuffer.Emplace(SourceVariable.FloatOffset, uint16(DestVariableLayout->FloatComponentStart), &FNiagaraSimCacheHelper::RebaseQuaternions);
+							CacheBufferReadInfo.VariableCopyMappingsToDataBuffer.Emplace(SourceVariable.FloatOffset, uint16(DestVariableLayout->GetFloatComponentStart()), &FNiagaraSimCacheHelper::RebaseQuaternions);
 						}
 						DestVariableLayout = nullptr;
 					}
@@ -583,7 +583,7 @@ struct FNiagaraSimCacheHelper
 						if (bRebaseVariable)
 						{
 							check(SourceVariable.FloatCount == 16);
-							CacheBufferReadInfo.VariableCopyMappingsToDataBuffer.Emplace(SourceVariable.FloatOffset, uint16(DestVariableLayout->FloatComponentStart), &FNiagaraSimCacheHelper::RebaseMatrices);
+							CacheBufferReadInfo.VariableCopyMappingsToDataBuffer.Emplace(SourceVariable.FloatOffset, uint16(DestVariableLayout->GetFloatComponentStart()), &FNiagaraSimCacheHelper::RebaseMatrices);
 							DestVariableLayout = nullptr;
 						}
 					}
@@ -592,17 +592,17 @@ struct FNiagaraSimCacheHelper
 
 			for (int32 i = 0; i < SourceVariable.FloatCount; ++i)
 			{
-				CacheBufferReadInfo.ComponentMappingsToDataBuffer[FloatOffset++] = uint16(DestVariableLayout ? DestVariableLayout->FloatComponentStart + i : INDEX_NONE);
+				CacheBufferReadInfo.ComponentMappingsToDataBuffer[FloatOffset++] = uint16(DestVariableLayout ? DestVariableLayout->GetFloatComponentStart() + i : INDEX_NONE);
 			}
 
 			for (int32 i = 0; i < SourceVariable.HalfCount; ++i)
 			{
-				CacheBufferReadInfo.ComponentMappingsToDataBuffer[HalfOffset++] = uint16(DestVariableLayout ? DestVariableLayout->HalfComponentStart + i : INDEX_NONE);
+				CacheBufferReadInfo.ComponentMappingsToDataBuffer[HalfOffset++] = uint16(DestVariableLayout ? DestVariableLayout->GetHalfComponentStart() + i : INDEX_NONE);
 			}
 
 			for (int32 i = 0; i < SourceVariable.Int32Count; ++i)
 			{
-				CacheBufferReadInfo.ComponentMappingsToDataBuffer[Int32Offset++] = uint16(DestVariableLayout ? DestVariableLayout->Int32ComponentStart + i : INDEX_NONE);
+				CacheBufferReadInfo.ComponentMappingsToDataBuffer[Int32Offset++] = uint16(DestVariableLayout ? DestVariableLayout->GetInt32ComponentStart() + i : INDEX_NONE);
 			}
 		}
 
diff --git a/Engine/Plugins/FX/Niagara/Source/Niagara/Private/NiagaraSystem.cpp b/Engine/Plugins/FX/Niagara/Source/Niagara/Private/NiagaraSystem.cpp
--- a/Engine/Plugins/FX/Niagara/Source/Niagara/Private/NiagaraSystem.cpp
+++ b/Engine/Plugins/FX/Niagara/Source/Niagara/Private/NiagaraSystem.cpp
@@ -903,6 +903,7 @@ void UNiagaraSystem::PostLoad()
 	}
 
 	const int32 NiagaraVer = GetLinkerCustomVersion(FNiagaraCustomVersion::GUID);
+#if WITH_EDITORONLY_DATA
 	if (NiagaraVer < FNiagaraCustomVersion::PlatformScalingRefactor)
 	{
 		for (int32 DL = 0; DL < ScalabilityOverrides_DEPRECATED.Num(); ++DL)
@@ -913,6 +914,7 @@ void UNiagaraSystem::PostLoad()
 			NewOverride.Platforms = FNiagaraPlatformSet(FNiagaraPlatformSet::CreateQualityLevelMask(DL));
 		}
 	}
+#endif
 
 	for (FNiagaraSystemScalabilityOverride& Override : SystemScalabilityOverrides.Overrides)
 	{
@@ -3059,7 +3061,7 @@ void UNiagaraSystem::InitEmitterCompiledData()
 
 		FNiagaraTypeDefinition SpawnInfoDef = FNiagaraTypeDefinition(FNiagaraSpawnInfo::StaticStruct());
 
-		for (FNiagaraVariable& Var : SystemSpawnScript->GetVMExecutableData().Attributes)
+		for (const FNiagaraVariableBase& Var : SystemSpawnScript->GetVMExecutableData().Attributes)
 		{
 			for (int32 EmitterIdx = 0; EmitterIdx < EmitterHandles.Num(); ++EmitterIdx)
 			{
@@ -3075,7 +3077,7 @@ void UNiagaraSystem::InitEmitterCompiledData()
 			}
 		}
 
-		for (FNiagaraVariable& Var : SystemUpdateScript->GetVMExecutableData().Attributes)
+		for (const FNiagaraVariableBase& Var : SystemUpdateScript->GetVMExecutableData().Attributes)
 		{
 			for (int32 EmitterIdx = 0; EmitterIdx < EmitterHandles.Num(); ++EmitterIdx)
 			{
@@ -3116,7 +3118,8 @@ void UNiagaraSystem::InitSystemCompiledData()
 
 	ExposedParameters.CopyParametersTo(SystemCompiledData.InstanceParamStore, false, FNiagaraParameterStore::EDataInterfaceCopyMethod::Reference);
 
-	auto CreateDataSetCompiledData = [&](FNiagaraDataSetCompiledData& CompiledData, TConstArrayView<FNiagaraVariable> Vars)
+	//-TODO:FNiagaraVariableBase Verify if EngineParamsSpawn->Parameters can also be moved to base
+	auto CreateDataSetCompiledData = [&](FNiagaraDataSetCompiledData& CompiledData, TConstStridedView<FNiagaraVariableBase> Vars)
 	{
 		CompiledData.Empty();
 
@@ -3136,12 +3139,12 @@ void UNiagaraSystem::InitSystemCompiledData()
 	const FNiagaraVMExecutableData& SystemSpawnScriptData = GetSystemSpawnScript()->GetVMExecutableData();
 	const FNiagaraVMExecutableData& SystemUpdateScriptData = GetSystemUpdateScript()->GetVMExecutableData();
 
-	CreateDataSetCompiledData(SystemCompiledData.DataSetCompiledData, SystemUpdateScriptData.Attributes);
+	CreateDataSetCompiledData(SystemCompiledData.DataSetCompiledData, MakeStridedView(SystemUpdateScriptData.Attributes));
 
 	const FNiagaraParameters* EngineParamsSpawn = SystemSpawnScriptData.DataSetToParameters.Find(TEXT("Engine"));
-	CreateDataSetCompiledData(SystemCompiledData.SpawnInstanceParamsDataSetCompiledData, EngineParamsSpawn ? TConstArrayView<FNiagaraVariable>(EngineParamsSpawn->Parameters) : TArrayView<FNiagaraVariable>());
+	CreateDataSetCompiledData(SystemCompiledData.SpawnInstanceParamsDataSetCompiledData, EngineParamsSpawn ? MakeConstStridedViewOfBase<FNiagaraVariableBase, FNiagaraVariable>(EngineParamsSpawn->Parameters) : MakeConstStridedView<FNiagaraVariableBase>(0, nullptr, 0));
 	const FNiagaraParameters* EngineParamsUpdate = SystemUpdateScriptData.DataSetToParameters.Find(TEXT("Engine"));
-	CreateDataSetCompiledData(SystemCompiledData.UpdateInstanceParamsDataSetCompiledData, EngineParamsUpdate ? TConstArrayView<FNiagaraVariable>(EngineParamsUpdate->Parameters) : TArrayView<FNiagaraVariable>());
+	CreateDataSetCompiledData(SystemCompiledData.UpdateInstanceParamsDataSetCompiledData, EngineParamsUpdate ? MakeConstStridedViewOfBase<FNiagaraVariableBase, FNiagaraVariable>(EngineParamsUpdate->Parameters) : MakeConstStridedView<FNiagaraVariableBase>(0, nullptr, 0));
 
 	// create the bindings to be used with our constant buffers; geenrating the offsets to/from the data sets; we need
 	// editor data to build these bindings because of the constant buffer structs only having their variable definitions
@@ -3662,16 +3665,16 @@ void FNiagaraParameterDataSetBindingCollection::BuildInternal(const TArray<FNiag
 
 			for (uint32 CompIdx = 0; CompIdx < Layout.GetNumFloatComponents(); ++CompIdx)
 			{
-				int32 ParamOffset = ParameterOffset + Layout.LayoutInfo.FloatComponentByteOffsets[CompIdx];
-				int32 DataSetOffset = Layout.FloatComponentStart + NumFloats++;
+				int32 ParamOffset = ParameterOffset + Layout.LayoutInfo.GetFloatComponentByteOffset(CompIdx);
+				int32 DataSetOffset = Layout.GetFloatComponentStart() + NumFloats++;
 				auto& Binding = FloatOffsets.AddDefaulted_GetRef();
 				Binding.ParameterOffset = ParamOffset;
 				Binding.DataSetComponentOffset = DataSetOffset;
 			}
 			for (uint32 CompIdx = 0; CompIdx < Layout.GetNumInt32Components(); ++CompIdx)
 			{
-				int32 ParamOffset = ParameterOffset + Layout.LayoutInfo.Int32ComponentByteOffsets[CompIdx];
-				int32 DataSetOffset = Layout.Int32ComponentStart + NumInts++;
+				int32 ParamOffset = ParameterOffset + Layout.LayoutInfo.GetInt32ComponentByteOffset(CompIdx);
+				int32 DataSetOffset = Layout.GetInt32ComponentStart() + NumInts++;
 				auto& Binding = Int32Offsets.AddDefaulted_GetRef();
 				Binding.ParameterOffset = ParamOffset;
 				Binding.DataSetComponentOffset = DataSetOffset;
diff --git a/Engine/Plugins/FX/Niagara/Source/Niagara/Private/NiagaraTypes.cpp b/Engine/Plugins/FX/Niagara/Source/Niagara/Private/NiagaraTypes.cpp
--- a/Engine/Plugins/FX/Niagara/Source/Niagara/Private/NiagaraTypes.cpp
+++ b/Engine/Plugins/FX/Niagara/Source/Niagara/Private/NiagaraTypes.cpp
@@ -56,6 +56,20 @@ bool FNiagaraVariableBase::ReplaceRootNamespace(const FStringView& ExpectedNames
 	return false;
 }
 
+bool FNiagaraVariableBase::SerializeFromMismatchedTag(const struct FPropertyTag& Tag, FStructuredArchive::FSlot Slot)
+{
+	if (Tag.Type == NAME_StructProperty && Tag.StructName == FNiagaraVariable::StaticStruct()->GetFName())
+	{
+		FNiagaraVariable Var;
+		FNiagaraVariable::StaticStruct()->SerializeItem(Slot, &Var, nullptr);
+		Name = Var.Name;
+		TypeDefHandle = Var.TypeDefHandle;
+		//ensureMsgf(Var.IsDataAllocated() == false, TEXT("NiagaraVariable has been demoted to Base but contains data, this is probably an error."));
+		return true;
+	}
+	return false;
+}
+
 void FNiagaraVariableMetaData::CopyUserEditableMetaData(const FNiagaraVariableMetaData& OtherMetaData)
 {
 	for (const FProperty* ChildProperty : TFieldRange<FProperty>(StaticStruct()))
@@ -784,3 +798,64 @@ bool FNiagaraTypeHelper::IsNiagaraFriendlyTopLevelStruct(UScriptStruct* InStruct
 
 /////////////////////////////////////////////////////////////////////////////
 
+void FNiagaraTypeLayoutInfo::GenerateLayoutInfo(const UScriptStruct* Struct)
+{
+	int32 FloatCount = 0;
+	int32 Int32Count = 0;
+	int32 HalfCount = 0;
+	GenerateLayoutInfoInternal(Struct, FloatCount, Int32Count, HalfCount, true);
+	NumFloatComponents = FloatCount;
+	NumInt32Components = Int32Count;
+	NumHalfComponents = HalfCount;
+
+	FloatCount = 0;
+	Int32Count = 0;
+	HalfCount = 0;
+	ComponentsOffsets.SetNum(GetTotalComponents() * 2);
+	GenerateLayoutInfoInternal(Struct, FloatCount, Int32Count, HalfCount, false);
+}
+
+void FNiagaraTypeLayoutInfo::GenerateLayoutInfoInternal(const UScriptStruct* Struct, int32& FloatCount, int32& Int32Count, int32& HalfCount, bool bCountComponentsOnly, int32 BaseOffset)
+{
+	for (TFieldIterator<FProperty> PropertyIt(Struct, EFieldIteratorFlags::IncludeSuper); PropertyIt; ++PropertyIt)
+	{
+		const FProperty* Property = *PropertyIt;
+		const int32 PropOffset = BaseOffset + Property->GetOffset_ForInternal();
+		if (Property->IsA(FFloatProperty::StaticClass()))
+		{
+			if (bCountComponentsOnly == false)
+			{
+				ComponentsOffsets[GetFloatArrayByteOffset() + FloatCount] = PropOffset;
+				ComponentsOffsets[GetFloatArrayRegisterOffset() + FloatCount] = FloatCount;
+			}
+			++FloatCount;
+		}
+		else if (Property->IsA(FUInt16Property::StaticClass()))
+		{
+			if (bCountComponentsOnly == false)
+			{
+				ComponentsOffsets[GetHalfArrayByteOffset() + HalfCount] = PropOffset;
+				ComponentsOffsets[GetHalfArrayRegisterOffset() + HalfCount] = HalfCount;
+			}
+			++HalfCount;
+		}
+		else if (Property->IsA(FIntProperty::StaticClass()) || Property->IsA(FBoolProperty::StaticClass()))
+		{
+			if (bCountComponentsOnly == false)
+			{
+				ComponentsOffsets[GetInt32ArrayByteOffset() + Int32Count] = PropOffset;
+				ComponentsOffsets[GetInt32ArrayRegisterOffset() + Int32Count] = Int32Count;
+			}
+			++Int32Count;
+		}
+		//Should be able to support double easily enough
+		else if (const FStructProperty* StructProp = CastField<FStructProperty>(Property))
+		{
+			GenerateLayoutInfoInternal(FNiagaraTypeHelper::FindNiagaraFriendlyTopLevelStruct(StructProp->Struct, ENiagaraStructConversion::Simulation), FloatCount, Int32Count, HalfCount, bCountComponentsOnly, PropOffset);
+		}
+		else
+		{
+			checkf(false, TEXT("Property(%s) Class(%s) is not a supported type"), *Property->GetName(), *Property->GetClass()->GetName());
+		}
+	}
+}
diff --git a/Engine/Plugins/FX/Niagara/Source/Niagara/Public/NiagaraCommon.h b/Engine/Plugins/FX/Niagara/Source/Niagara/Public/NiagaraCommon.h
--- a/Engine/Plugins/FX/Niagara/Source/Niagara/Public/NiagaraCommon.h
+++ b/Engine/Plugins/FX/Niagara/Source/Niagara/Public/NiagaraCommon.h
@@ -300,7 +300,7 @@ struct FNiagaraDataSetProperties
 	FNiagaraDataSetID ID;
 
 	UPROPERTY()
-	TArray<FNiagaraVariable> Variables;
+	TArray<FNiagaraVariableBase> Variables;
 };
 
 /** Information about an input or output of a Niagara operation node. */
diff --git a/Engine/Plugins/FX/Niagara/Source/Niagara/Public/NiagaraParameterStore.h b/Engine/Plugins/FX/Niagara/Source/Niagara/Public/NiagaraParameterStore.h
--- a/Engine/Plugins/FX/Niagara/Source/Niagara/Public/NiagaraParameterStore.h
+++ b/Engine/Plugins/FX/Niagara/Source/Niagara/Public/NiagaraParameterStore.h
@@ -23,7 +23,7 @@ struct FNiagaraBoundParameter
 	GENERATED_USTRUCT_BODY()
 
 	UPROPERTY()
-	FNiagaraVariable Parameter;
+	FNiagaraVariableBase Parameter;
 	UPROPERTY()
 	int32 SrcOffset = 0;
 	UPROPERTY()
@@ -276,6 +276,19 @@ public:
 	void Dump();
 	void DumpParameters(bool bDumpBindings = false)const;
 
+	SIZE_T GetResourceSize() const
+	{
+		SIZE_T ResourceSize = sizeof(FNiagaraParameterStore);
+		ResourceSize += SortedParameterOffsets.GetAllocatedSize();
+		ResourceSize += ParameterData.GetAllocatedSize();
+		ResourceSize += DataInterfaces.GetAllocatedSize();
+		ResourceSize += UObjects.GetAllocatedSize();
+		ResourceSize += OriginalPositionData.GetAllocatedSize();
+		ResourceSize += Bindings.GetAllocatedSize();
+		ResourceSize += SourceStores.GetAllocatedSize();
+		return ResourceSize;
+	}
+
 	FORCEINLINE uint32 GetParametersDirty() const { return bParametersDirty; }
 	FORCEINLINE uint32 GetInterfacesDirty() const { return bInterfacesDirty; }
 	FORCEINLINE uint32 GetUObjectsDirty() const { return bUObjectsDirty; }
diff --git a/Engine/Plugins/FX/Niagara/Source/Niagara/Public/NiagaraScriptExecutionParameterStore.h b/Engine/Plugins/FX/Niagara/Source/Niagara/Public/NiagaraScriptExecutionParameterStore.h
--- a/Engine/Plugins/FX/Niagara/Source/Niagara/Public/NiagaraScriptExecutionParameterStore.h
+++ b/Engine/Plugins/FX/Niagara/Source/Niagara/Public/NiagaraScriptExecutionParameterStore.h
@@ -137,6 +137,15 @@ public:
 	TArray<uint8> CachedScriptLiterals;
 #endif
 
+	SIZE_T GetResourceSize() const
+	{
+		SIZE_T ResourceSize = sizeof(FNiagaraScriptExecutionParameterStore) - sizeof(FNiagaraParameterStore);
+
+		ResourceSize += FNiagaraParameterStore::GetResourceSize();
+		ResourceSize += PaddingInfo.GetAllocatedSize();
+		return ResourceSize;
+	}
+
 protected:
 	void AddPaddedParamSize(const FNiagaraTypeDefinition& InParamType, uint32 InOffset);
 	void AddAlignmentPadding();
diff --git a/Engine/Plugins/FX/Niagara/Source/Niagara/Public/NiagaraSystemSimulation.h b/Engine/Plugins/FX/Niagara/Source/Niagara/Public/NiagaraSystemSimulation.h
--- a/Engine/Plugins/FX/Niagara/Source/Niagara/Public/NiagaraSystemSimulation.h
+++ b/Engine/Plugins/FX/Niagara/Source/Niagara/Public/NiagaraSystemSimulation.h
@@ -91,14 +91,14 @@ struct FNiagaraParameterStoreToDataSetBinding
 			const int32 ParameterOffset = *ParameterOffsetPtr;
 			for (uint32 CompIdx = 0; CompIdx < Layout.GetNumFloatComponents(); ++CompIdx)
 			{
-				const int32 ParamOffset = ParameterOffset + Layout.LayoutInfo.FloatComponentByteOffsets[CompIdx];
-				const int32 DataSetOffset = Layout.FloatComponentStart + CompIdx;
+				const int32 ParamOffset = ParameterOffset + Layout.LayoutInfo.GetFloatComponentByteOffset(CompIdx);
+				const int32 DataSetOffset = Layout.GetFloatComponentStart() + CompIdx;
 				FloatOffsets.Emplace(ParamOffset, DataSetOffset);
 			}
 			for (uint32 CompIdx = 0; CompIdx < Layout.GetNumInt32Components(); ++CompIdx)
 			{
-				const int32 ParamOffset = ParameterOffset + Layout.LayoutInfo.Int32ComponentByteOffsets[CompIdx];
-				const int32 DataSetOffset = Layout.Int32ComponentStart + CompIdx;
+				const int32 ParamOffset = ParameterOffset + Layout.LayoutInfo.GetInt32ComponentByteOffset(CompIdx);
+				const int32 DataSetOffset = Layout.GetInt32ComponentStart() + CompIdx;
 				Int32Offsets.Emplace(ParamOffset, DataSetOffset);
 			}
 			for (uint32 CompIdx = 0; CompIdx < Layout.GetNumHalfComponents(); ++CompIdx)
@@ -107,8 +107,8 @@ struct FNiagaraParameterStoreToDataSetBinding
 
 				if (ParameterSetsSupportHalf)
 				{
-					const int32 ParamOffset = ParameterOffset + Layout.LayoutInfo.HalfComponentByteOffsets[CompIdx];
-					const int32 DataSetOffset = Layout.HalfComponentStart + CompIdx;
+					const int32 ParamOffset = ParameterOffset + Layout.LayoutInfo.GetHalfComponentByteOffset(CompIdx);
+					const int32 DataSetOffset = Layout.GetHalfComponentStart() + CompIdx;
 					HalfOffsets.Emplace(ParamOffset, DataSetOffset, !ParameterSetsSupportHalf);
 				}
 				else
@@ -116,8 +116,8 @@ struct FNiagaraParameterStoreToDataSetBinding
 					// if parameter sets don't support half, then we need to write in floats into the parameter set, and
 					// for that we need to adjust the offset based on the difference in stride between float & half
 					// In reality
-					const int32 ParamOffset = ParameterOffset + sizeof(float) * Layout.LayoutInfo.HalfComponentByteOffsets[CompIdx] / sizeof(FFloat16);
-					const int32 DataSetOffset = Layout.HalfComponentStart + CompIdx;
+					const int32 ParamOffset = ParameterOffset + sizeof(float) * Layout.LayoutInfo.GetHalfComponentByteOffset(CompIdx) / sizeof(FFloat16);
+					const int32 DataSetOffset = Layout.GetHalfComponentStart() + CompIdx;
 					HalfOffsets.Emplace(ParamOffset, DataSetOffset, !ParameterSetsSupportHalf);
 				}
 			}
diff --git a/Engine/Plugins/FX/Niagara/Source/Niagara/Public/NiagaraTypes.h b/Engine/Plugins/FX/Niagara/Source/Niagara/Public/NiagaraTypes.h
--- a/Engine/Plugins/FX/Niagara/Source/Niagara/Public/NiagaraTypes.h
+++ b/Engine/Plugins/FX/Niagara/Source/Niagara/Public/NiagaraTypes.h
@@ -463,77 +463,47 @@ struct FNiagaraTypeLayoutInfo
 {
 	GENERATED_BODY()
 
-	FNiagaraTypeLayoutInfo()
-	{}
-
-	/** Byte offset of each float component in a structured layout. */
+private:
 	UPROPERTY()
-	TArray<uint32> FloatComponentByteOffsets;
+	uint16 NumFloatComponents = 0;
 
-	/** Offset into register table for each float component. */
 	UPROPERTY()
-	TArray<uint32> FloatComponentRegisterOffsets;
+	uint16 NumInt32Components = 0;
 
-	/** Byte offset of each int32 component in a structured layout. */
 	UPROPERTY()
-	TArray<uint32> Int32ComponentByteOffsets;
+	uint16 NumHalfComponents = 0;
 
-	/** Offset into register table for each int32 component. */
 	UPROPERTY()
-	TArray<uint32> Int32ComponentRegisterOffsets;
+	TArray<uint32> ComponentsOffsets;
 
-	/** Byte offset of each half component in a structured layout. */
-	UPROPERTY()
-	TArray<uint32> HalfComponentByteOffsets;
+public:
+	uint32 GetNumFloatComponents() const { return NumFloatComponents; }
+	uint32 GetNumInt32Components() const { return NumInt32Components; }
+	uint32 GetNumHalfComponents() const { return NumHalfComponents; }
+	uint32 GetTotalComponents() const { return GetNumFloatComponents() + GetNumInt32Components() + GetNumHalfComponents(); }
 
-	/** Offset into register table for each half component. */
-	UPROPERTY()
-	TArray<uint32> HalfComponentRegisterOffsets;
+	uint32 GetFloatComponentByteOffset(int32 Component) const { return ComponentsOffsets[GetFloatArrayByteOffset() + Component]; }
+	uint32 GetFloatComponentRegisterOffset(int32 Component) const { return ComponentsOffsets[GetFloatArrayRegisterOffset() + Component]; }
 
-	static void GenerateLayoutInfo(FNiagaraTypeLayoutInfo& Layout, const UScriptStruct* Struct)
-	{
-		Layout.FloatComponentByteOffsets.Empty();
-		Layout.FloatComponentRegisterOffsets.Empty();
-		Layout.Int32ComponentByteOffsets.Empty();
-		Layout.Int32ComponentRegisterOffsets.Empty();
-		Layout.HalfComponentByteOffsets.Empty();
-		Layout.HalfComponentRegisterOffsets.Empty();
-		GenerateLayoutInfoInternal(Layout, Struct);
-	}
+	uint32 GetInt32ComponentByteOffset(int32 Component) const { return ComponentsOffsets[GetInt32ArrayByteOffset() + Component]; }
+	uint32 GetInt32ComponentRegisterOffset(int32 Component) const { return ComponentsOffsets[GetInt32ArrayRegisterOffset() + Component]; }
+
+	uint32 GetHalfComponentByteOffset(int32 Component) const { return ComponentsOffsets[GetHalfArrayByteOffset() + Component]; }
+	uint32 GetHalfComponentRegisterOffset(int32 Component) const { return ComponentsOffsets[GetHalfArrayRegisterOffset() + Component]; }
+
+	NIAGARA_API void GenerateLayoutInfo(const UScriptStruct* Struct);
 
 private:
-	static void GenerateLayoutInfoInternal(FNiagaraTypeLayoutInfo& Layout, const UScriptStruct* Struct, int32 BaseOffest = 0)
-	{
-		for (TFieldIterator<FProperty> PropertyIt(Struct, EFieldIteratorFlags::IncludeSuper); PropertyIt; ++PropertyIt)
-		{
-			FProperty* Property = *PropertyIt;
-			int32 PropOffset = BaseOffest + Property->GetOffset_ForInternal();
-			if (Property->IsA(FFloatProperty::StaticClass()))
-			{
-				Layout.FloatComponentRegisterOffsets.Add(Layout.FloatComponentByteOffsets.Num());
-				Layout.FloatComponentByteOffsets.Add(PropOffset);
-			}
-			else if (Property->IsA(FUInt16Property::StaticClass()))
-			{
-				Layout.HalfComponentRegisterOffsets.Add(Layout.HalfComponentByteOffsets.Num());
-				Layout.HalfComponentByteOffsets.Add(PropOffset);
-			}
-			else if (Property->IsA(FIntProperty::StaticClass()) || Property->IsA(FBoolProperty::StaticClass()))
-			{
-				Layout.Int32ComponentRegisterOffsets.Add(Layout.Int32ComponentByteOffsets.Num());
-				Layout.Int32ComponentByteOffsets.Add(PropOffset);
-			}
-			//Should be able to support double easily enough
-			else if (FStructProperty* StructProp = CastField<FStructProperty>(Property))
-			{
-				GenerateLayoutInfoInternal(Layout, FNiagaraTypeHelper::FindNiagaraFriendlyTopLevelStruct(StructProp->Struct, ENiagaraStructConversion::Simulation), PropOffset);
-			}
-			else
-			{
-				checkf(false, TEXT("Property(%s) Class(%s) is not a supported type"), *Property->GetName(), *Property->GetClass()->GetName());
-			}
-		}
-	}
+	int32 GetFloatArrayByteOffset() const { return 0; }
+	int32 GetFloatArrayRegisterOffset() const { return NumFloatComponents; }
+
+	int32 GetInt32ArrayByteOffset() const { return NumFloatComponents * 2; }
+	int32 GetInt32ArrayRegisterOffset() const { return (NumFloatComponents * 2) + NumInt32Components; }
+
+	int32 GetHalfArrayByteOffset() const { return (NumFloatComponents * 2) + (NumInt32Components * 2); }
+	int32 GetHalfArrayRegisterOffset() const { return (NumFloatComponents * 2) + (NumInt32Components * 2) + NumHalfComponents; }
+
+	void GenerateLayoutInfoInternal(const UScriptStruct* Struct, int32& FloatCount, int32& Int32Count, int32& HalfCount, bool bCountComponentsOnly, int32 BaseOffset = 0);
 };
 
 
@@ -1874,6 +1844,7 @@ struct FNiagaraVariableBase
 #if WITH_EDITORONLY_DATA
 	void PostSerialize(const FArchive& Ar);
 #endif
+	bool SerializeFromMismatchedTag(const struct FPropertyTag& Tag, FStructuredArchive::FSlot Slot);
 
 	bool NIAGARA_API AppendCompileHash(FNiagaraCompileHashVisitor* InVisitor) const;
 
@@ -1904,6 +1875,7 @@ struct TStructOpsTypeTraits<FNiagaraVariableBase> : public TStructOpsTypeTraitsB
 #if WITH_EDITORONLY_DATA
 		WithPostSerialize = true,
 #endif
+		WithStructuredSerializeFromMismatchedTag = true,
 	};
 };
 
@@ -2076,6 +2048,24 @@ struct FNiagaraVariable : public FNiagaraVariableBase
 	void PostSerialize(const FArchive& Ar);
 #endif
 
+	static void ConvertFromBaseArray(TConstArrayView<FNiagaraVariableBase> FromVariables, TArray<FNiagaraVariable>& ToVariables)
+	{
+		ToVariables.Reserve(ToVariables.Num() + FromVariables.Num());
+		for (const FNiagaraVariableBase& BaseVar : FromVariables)
+		{
+			ToVariables.Emplace(BaseVar);
+		}
+	}
+
+	static void ConvertToBaseArray(TConstArrayView<FNiagaraVariable> FromVariables, TArray<FNiagaraVariableBase>& ToVariables)
+	{
+		ToVariables.Reserve(ToVariables.Num() + FromVariables.Num());
+		for (const FNiagaraVariableBase& BaseVar : FromVariables)
+		{
+			ToVariables.Emplace(BaseVar);
+		}
+	}
+
 private:
 	//This gets serialized but do we need to worry about endianness doing things like this? If not, where does that get handled?
 	//TODO: Remove storage here entirely and move everything to an FNiagaraParameterStore.
diff --git a/Engine/Plugins/FX/Niagara/Source/NiagaraEditor/Private/NiagaraBakerRenderer.cpp b/Engine/Plugins/FX/Niagara/Source/NiagaraEditor/Private/NiagaraBakerRenderer.cpp
--- a/Engine/Plugins/FX/Niagara/Source/NiagaraEditor/Private/NiagaraBakerRenderer.cpp
+++ b/Engine/Plugins/FX/Niagara/Source/NiagaraEditor/Private/NiagaraBakerRenderer.cpp
@@ -563,10 +563,10 @@ void FNiagaraBakerRenderer::RenderParticleAttribute(UTextureRenderTarget2D* Rend
 		const FNiagaraVariableLayoutInfo& VariableInfo = ParticleDataSet.GetCompiledData().VariableLayouts[VariableIndex];
 	
 		float* FloatChannels[4];
-		FloatChannels[0] = (float*)ParticleDataBuffer->GetComponentPtrFloat(VariableInfo.FloatComponentStart);
-		FloatChannels[1] = VariableInfo.GetNumFloatComponents() > 1 ? (float*)ParticleDataBuffer->GetComponentPtrFloat(VariableInfo.FloatComponentStart + 1) : nullptr;
-		FloatChannels[2] = VariableInfo.GetNumFloatComponents() > 2 ? (float*)ParticleDataBuffer->GetComponentPtrFloat(VariableInfo.FloatComponentStart + 2) : nullptr;
-		FloatChannels[3] = VariableInfo.GetNumFloatComponents() > 3 ? (float*)ParticleDataBuffer->GetComponentPtrFloat(VariableInfo.FloatComponentStart + 3) : nullptr;
+		FloatChannels[0] = (float*)ParticleDataBuffer->GetComponentPtrFloat(VariableInfo.GetFloatComponentStart());
+		FloatChannels[1] = VariableInfo.GetNumFloatComponents() > 1 ? (float*)ParticleDataBuffer->GetComponentPtrFloat(VariableInfo.GetFloatComponentStart() + 1) : nullptr;
+		FloatChannels[2] = VariableInfo.GetNumFloatComponents() > 2 ? (float*)ParticleDataBuffer->GetComponentPtrFloat(VariableInfo.GetFloatComponentStart() + 2) : nullptr;
+		FloatChannels[3] = VariableInfo.GetNumFloatComponents() > 3 ? (float*)ParticleDataBuffer->GetComponentPtrFloat(VariableInfo.GetFloatComponentStart() + 3) : nullptr;
 	
 		const FIntPoint RenderTargetSize(RenderTarget->GetSurfaceWidth(), RenderTarget->GetSurfaceHeight());
 		const int32 ParticleBufferStore = RenderTargetSize.X * RenderTargetSize.Y;
diff --git a/Engine/Plugins/FX/Niagara/Source/NiagaraEditor/Private/NiagaraHlslTranslator.cpp b/Engine/Plugins/FX/Niagara/Source/NiagaraEditor/Private/NiagaraHlslTranslator.cpp
--- a/Engine/Plugins/FX/Niagara/Source/NiagaraEditor/Private/NiagaraHlslTranslator.cpp
+++ b/Engine/Plugins/FX/Niagara/Source/NiagaraEditor/Private/NiagaraHlslTranslator.cpp
@@ -2017,7 +2017,9 @@ FNiagaraTranslateResults TNiagaraHlslTranslator<GraphBridge>::Translate(const FN
 			SetProps.ID = InfoPair.Key;
 			for (TPair <int32, FDataSetAccessInfo> IndexPair : InfoPair.Value)
 			{
-				SetProps.Variables = IndexPair.Value.Variables;
+				//-TODO:FNiagaraVariableBase the source can be FNiagaraVariableBase also
+				SetProps.Variables.Empty();
+				FNiagaraVariable::ConvertToBaseArray(IndexPair.Value.Variables, SetProps.Variables);
 			}
 
 			CompilationOutput.ScriptData.WriteDataSets.Add(SetProps);
diff --git a/Engine/Plugins/FX/Niagara/Source/NiagaraEditor/Private/Widgets/SNiagaraSpreadsheetView.cpp b/Engine/Plugins/FX/Niagara/Source/NiagaraEditor/Private/Widgets/SNiagaraSpreadsheetView.cpp
--- a/Engine/Plugins/FX/Niagara/Source/NiagaraEditor/Private/Widgets/SNiagaraSpreadsheetView.cpp
+++ b/Engine/Plugins/FX/Niagara/Source/NiagaraEditor/Private/Widgets/SNiagaraSpreadsheetView.cpp
@@ -31,6 +31,16 @@
 #define VALUE_COLUMN_NAME TEXT("Value")
 #define FILLER_COLUMN_NAME TEXT("__FILLER__")
 
+struct FNiagaraSpreadsheetTypeLayoutInfo
+{
+	TArray<uint32> FloatComponentByteOffsets;
+	TArray<uint32> FloatComponentRegisterOffsets;
+	TArray<uint32> Int32ComponentByteOffsets;
+	TArray<uint32> Int32ComponentRegisterOffsets;
+	TArray<uint32> HalfComponentByteOffsets;
+	TArray<uint32> HalfComponentRegisterOffsets;
+};
+
 void SNiagaraSpreadsheetRow::Construct(const FArguments& InArgs, const TSharedRef<STableViewBase>& InOwnerTableView)
 {
 	RowIndex = InArgs._RowIndex;
@@ -840,14 +850,13 @@ TSharedRef<SWidget> SNiagaraSpreadsheetView::GetOutputFilterMenu(EUITab Tab)
 		NAME_None,
 		EUserInterfaceActionType::ToggleButton);
 
-	TArray<FNiagaraVariable> Variables = CaptureData[(int32)Tab].DataSet.GetVariables();
-	for (const FNiagaraVariable& Var : Variables)
+	for (const FNiagaraVariableBase& Var : CaptureData[(int32)Tab].DataSet.GetVariables())
 	{
 		const FNiagaraTypeDefinition& TypeDef = Var.GetType();
 		const UScriptStruct* Struct = TypeDef.GetScriptStruct();
 		const UEnum* Enum = TypeDef.GetEnum();
 
-		FNiagaraTypeLayoutInfo Layout;
+		FNiagaraSpreadsheetTypeLayoutInfo Layout;
 		TArray<FName> PropertyNames;
 		TArray<SNiagaraSpreadsheetView::FieldInfo> FieldInfos;
 
@@ -1166,7 +1175,7 @@ void SNiagaraSpreadsheetView::ResetEntries(EUITab Tab)
 	}
 }
 
-void SNiagaraSpreadsheetView::GenerateLayoutInfo(FNiagaraTypeLayoutInfo& Layout, const UScriptStruct* Struct, const UEnum* Enum,  FName BaseName, TArray<FName>& PropertyNames, TArray<SNiagaraSpreadsheetView::FieldInfo>& FieldInfo)
+void SNiagaraSpreadsheetView::GenerateLayoutInfo(FNiagaraSpreadsheetTypeLayoutInfo& Layout, const UScriptStruct* Struct, const UEnum* Enum,  FName BaseName, TArray<FName>& PropertyNames, TArray<SNiagaraSpreadsheetView::FieldInfo>& FieldInfo)
 {
 	TFieldIterator<FProperty> PropertyCountIt(Struct, EFieldIteratorFlags::IncludeSuper);
 	int32 NumProperties = 0;
@@ -1263,8 +1272,6 @@ void SNiagaraSpreadsheetView::ResetColumns(EUITab Tab)
 			uint32 TotalInt32Components = 0;
 			uint32 TotalHalfComponents = 0;
 
-			TArray<FNiagaraVariable> Variables = CaptureData[i].DataSet.GetVariables();
-
 			TArray<FName> ColumnNames;
 
 			if (CaptureData[i].bOutputColumnsAreAttributes)
@@ -1279,13 +1286,13 @@ void SNiagaraSpreadsheetView::ResetColumns(EUITab Tab)
 				ColumnNames.Add(FILLER_COLUMN_NAME);
 			}
 
-			for (const FNiagaraVariable& Var : Variables)
+			for (const FNiagaraVariableBase& Var : CaptureData[i].DataSet.GetVariables())
 			{
 				FNiagaraTypeDefinition TypeDef = Var.GetType();
 				const UScriptStruct* Struct = TypeDef.GetScriptStruct();
 				const UEnum* Enum = TypeDef.GetEnum();
 
-				FNiagaraTypeLayoutInfo Layout;
+				FNiagaraSpreadsheetTypeLayoutInfo Layout;
 				TArray<FName> PropertyNames;
 				TArray<SNiagaraSpreadsheetView::FieldInfo> FieldInfos;
 
@@ -1402,7 +1409,7 @@ void SNiagaraSpreadsheetView::ResetColumns(EUITab Tab)
 				const UScriptStruct* Struct = TypeDef.GetScriptStruct();
 				const UEnum* Enum = TypeDef.GetEnum();
 
-				FNiagaraTypeLayoutInfo Layout;
+				FNiagaraSpreadsheetTypeLayoutInfo Layout;
 				TArray<FName> PropertyNames;
 				TArray<SNiagaraSpreadsheetView::FieldInfo> FieldInfos;
 
diff --git a/Engine/Plugins/FX/Niagara/Source/NiagaraEditor/Private/Widgets/SNiagaraSpreadsheetView.h b/Engine/Plugins/FX/Niagara/Source/NiagaraEditor/Private/Widgets/SNiagaraSpreadsheetView.h
--- a/Engine/Plugins/FX/Niagara/Source/NiagaraEditor/Private/Widgets/SNiagaraSpreadsheetView.h
+++ b/Engine/Plugins/FX/Niagara/Source/NiagaraEditor/Private/Widgets/SNiagaraSpreadsheetView.h
@@ -71,7 +71,7 @@ protected:
 	void SetTarget(UNiagaraComponent* InComponent);
 	void GetNameAndTooltip(const UNiagaraComponent*, FText& OutText, FText& OutTooltip) const;
 
-	void GenerateLayoutInfo(FNiagaraTypeLayoutInfo& Layout, const UScriptStruct* Struct, const UEnum* Enum, FName BaseName, TArray<FName>& PropertyNames, TArray<FieldInfo>& FieldInfo);
+	void GenerateLayoutInfo(struct FNiagaraSpreadsheetTypeLayoutInfo& Layout, const UScriptStruct* Struct, const UEnum* Enum, FName BaseName, TArray<FName>& PropertyNames, TArray<FieldInfo>& FieldInfo);
 
 	EUITab TabState;
 
diff --git a/Engine/Plugins/Runtime/HairStrands/Source/HairStrandsCore/Private/Niagara/NiagaraDataInterfaceHairStrands.cpp b/Engine/Plugins/Runtime/HairStrands/Source/HairStrandsCore/Private/Niagara/NiagaraDataInterfaceHairStrands.cpp
--- a/Engine/Plugins/Runtime/HairStrands/Source/HairStrandsCore/Private/Niagara/NiagaraDataInterfaceHairStrands.cpp
+++ b/Engine/Plugins/Runtime/HairStrands/Source/HairStrandsCore/Private/Niagara/NiagaraDataInterfaceHairStrands.cpp
@@ -971,8 +971,8 @@ void UNiagaraDataInterfaceHairStrands::SimCachePostReadFrame(void* OptionalPerIn
 			if( EmitterDatas.GetCompiledData().VariableLayouts.IsValidIndex(NodePositionIndex) &&
 				EmitterDatas.GetCompiledData().VariableLayouts.IsValidIndex(RestPositionIndex) && InstanceData)
 			{
-				const uint32 NodePositionComponent = EmitterDatas.GetCompiledData().VariableLayouts[NodePositionIndex].FloatComponentStart;
-				const uint32 RestPositionComponent = EmitterDatas.GetCompiledData().VariableLayouts[RestPositionIndex].FloatComponentStart;
+				const uint32 NodePositionComponent = EmitterDatas.GetCompiledData().VariableLayouts[NodePositionIndex].GetFloatComponentStart();
+				const uint32 RestPositionComponent = EmitterDatas.GetCompiledData().VariableLayouts[RestPositionIndex].GetFloatComponentStart();
 				
 				FNiagaraDataBuffer* ParticlesBuffer = EmitterDatas.GetCurrentData();
 				FNDIHairStrandsBuffer* HairstrandsBuffer = InstanceData->HairStrandsBuffer;
diff --git a/Engine/Source/Runtime/Core/Private/UObject/FortniteReleaseBranchCustomObjectVersion.cpp b/Engine/Source/Runtime/Core/Private/UObject/FortniteReleaseBranchCustomObjectVersion.cpp
--- a/Engine/Source/Runtime/Core/Private/UObject/FortniteReleaseBranchCustomObjectVersion.cpp
+++ b/Engine/Source/Runtime/Core/Private/UObject/FortniteReleaseBranchCustomObjectVersion.cpp
@@ -6,7 +6,7 @@ TMap<FGuid, FGuid> FFortniteReleaseBranchCustomObjectVersion::GetSystemGuids()
 	TMap<FGuid, FGuid> SystemGuids;
 	const FDevSystemGuids& DevGuids = FDevSystemGuids::Get();
 
-	SystemGuids.Add(DevGuids.Niagara_LatestScriptCompileVersion, FGuid("11BB3AC044D97344A3F554562F60FB1C"));
+	SystemGuids.Add(DevGuids.Niagara_LatestScriptCompileVersion, FGuid("C2B4733DD8CEFC43B07A1ADD16A8978B"));
 	SystemGuids.Add(DevGuids.SkeletalMeshDerivedDataVersion, FGuid("5705956EC7134274A5A1BD3FC74DB3A9"));
 	return SystemGuids;
 }
