- Variable -> VariableBase fixes

diff --git a/Engine/Plugins/FX/Niagara/Source/Niagara/Private/NiagaraEmitter.cpp b/Engine/Plugins/FX/Niagara/Source/Niagara/Private/NiagaraEmitter.cpp
--- a/Engine/Plugins/FX/Niagara/Source/Niagara/Private/NiagaraEmitter.cpp
+++ b/Engine/Plugins/FX/Niagara/Source/Niagara/Private/NiagaraEmitter.cpp
@@ -2283,7 +2283,7 @@ void FVersionedNiagaraEmitterData::GatherCompiledParticleAttributes(TArray<FNiag
 	{
 		OutVariables = UpdateScriptProps.Script->GetVMExecutableData().Attributes;
 
-		for (const FNiagaraVariable& Var : SpawnScriptProps.Script->GetVMExecutableData().Attributes)
+		for (const FNiagaraVariableBase& Var : SpawnScriptProps.Script->GetVMExecutableData().Attributes)
 		{
 			OutVariables.AddUnique(Var);
 		}
diff --git a/Engine/Plugins/FX/Niagara/Source/Niagara/Private/NiagaraSystem.cpp b/Engine/Plugins/FX/Niagara/Source/Niagara/Private/NiagaraSystem.cpp
--- a/Engine/Plugins/FX/Niagara/Source/Niagara/Private/NiagaraSystem.cpp
+++ b/Engine/Plugins/FX/Niagara/Source/Niagara/Private/NiagaraSystem.cpp
@@ -3124,7 +3124,7 @@ void UNiagaraSystem::InitSystemCompiledData()
 		CompiledData.Empty();
 
 		CompiledData.Variables.Reset(Vars.Num());
-		for (const FNiagaraVariable& Var : Vars)
+		for (const FNiagaraVariableBase& Var : Vars)
 		{
 			CompiledData.Variables.AddUnique(Var);
 		}
