Niagara - When adding messages to a message store, don't replace existing messages if they are equal.  This is neccessary so that FTexts in messages don't get recreated at cook time.

diff --git a/Engine/Plugins/FX/Niagara/Source/Niagara/Private/NiagaraMessageDataBase.cpp b/Engine/Plugins/FX/Niagara/Source/Niagara/Private/NiagaraMessageDataBase.cpp
new file mode 100644
--- /dev/null
+++ b/Engine/Plugins/FX/Niagara/Source/Niagara/Private/NiagaraMessageDataBase.cpp
@@ -0,0 +1,25 @@
+// Copyright Epic Games, Inc. All Rights Reserved.
+
+#include "NiagaraMessageDataBase.h"
+
+#if WITH_EDITORONLY_DATA
+bool UNiagaraMessageDataBase::Equals(const UNiagaraMessageDataBase* Other) const
+{
+	if (GetClass() != Other->GetClass())
+	{
+		return false;
+	}
+
+	for (TFieldIterator<FProperty> PropertyIterator(GetClass()); PropertyIterator; ++PropertyIterator)
+	{
+		if (PropertyIterator->Identical(
+			PropertyIterator->ContainerPtrToValuePtr<void>(this),
+			PropertyIterator->ContainerPtrToValuePtr<void>(Other), PPF_DeepComparison) == false)
+		{
+			return false;
+		}
+	}
+
+	return true;
+}
+#endif
diff --git a/Engine/Plugins/FX/Niagara/Source/Niagara/Private/NiagaraMessageStore.cpp b/Engine/Plugins/FX/Niagara/Source/Niagara/Private/NiagaraMessageStore.cpp
--- a/Engine/Plugins/FX/Niagara/Source/Niagara/Private/NiagaraMessageStore.cpp
+++ b/Engine/Plugins/FX/Niagara/Source/Niagara/Private/NiagaraMessageStore.cpp
@@ -16,7 +16,11 @@ void FNiagaraMessageStore::SetMessages(const TMap<FGuid, TObjectPtr<UNiagaraMess
 
 void FNiagaraMessageStore::AddMessage(const FGuid& MessageKey, UNiagaraMessageDataBase* NewMessage)
 {
-	MessageKeyToMessageMap.Add(MessageKey, NewMessage);
+	TObjectPtr<UNiagaraMessageDataBase>* ExistingMessagePtr = MessageKeyToMessageMap.Find(MessageKey);
+	if (ExistingMessagePtr == nullptr || (*ExistingMessagePtr)->Equals(NewMessage) == false)
+	{
+		MessageKeyToMessageMap.Add(MessageKey, NewMessage);
+	}
 }
 
 void FNiagaraMessageStore::RemoveMessage(const FGuid& MessageKey) 
diff --git a/Engine/Plugins/FX/Niagara/Source/Niagara/Public/NiagaraMessageDataBase.h b/Engine/Plugins/FX/Niagara/Source/Niagara/Public/NiagaraMessageDataBase.h
--- a/Engine/Plugins/FX/Niagara/Source/Niagara/Public/NiagaraMessageDataBase.h
+++ b/Engine/Plugins/FX/Niagara/Source/Niagara/Public/NiagaraMessageDataBase.h
@@ -14,4 +14,8 @@ public:
 	UNiagaraMessageDataBase() = default;
 
 	virtual bool GetAllowDismissal() const { return false; }
+
+#if WITH_EDITORONLY_DATA
+	virtual bool Equals(const UNiagaraMessageDataBase* Other) const;
+#endif
 };
