Niagara - Make the data and methods on FNiagaraMessageStore only available with editor only data.

diff --git a/Engine/Plugins/FX/Niagara/Source/Niagara/Private/NiagaraMessageStore.cpp b/Engine/Plugins/FX/Niagara/Source/Niagara/Private/NiagaraMessageStore.cpp
--- a/Engine/Plugins/FX/Niagara/Source/Niagara/Private/NiagaraMessageStore.cpp
+++ b/Engine/Plugins/FX/Niagara/Source/Niagara/Private/NiagaraMessageStore.cpp
@@ -3,6 +3,8 @@
 #include "NiagaraMessageStore.h"
 #include "NiagaraMessageDataBase.h"
 
+#if WITH_EDITORONLY_DATA
+
 const TMap<FGuid, TObjectPtr<UNiagaraMessageDataBase>>& FNiagaraMessageStore::GetMessages() const
 { 
 	return MessageKeyToMessageMap;
@@ -63,4 +65,6 @@ UObject* FNiagaraMessageSourceAndStore::GetSource() const
 FNiagaraMessageStore* FNiagaraMessageSourceAndStore::GetStore() const
 {
 	return SourceWeak.IsValid() ? Store : nullptr;
-}
\ No newline at end of file
+}
+
+#endif // WITH_EDITORONLY_DATA
\ No newline at end of file
diff --git a/Engine/Plugins/FX/Niagara/Source/Niagara/Public/NiagaraMessageStore.h b/Engine/Plugins/FX/Niagara/Source/Niagara/Public/NiagaraMessageStore.h
--- a/Engine/Plugins/FX/Niagara/Source/Niagara/Public/NiagaraMessageStore.h
+++ b/Engine/Plugins/FX/Niagara/Source/Niagara/Public/NiagaraMessageStore.h
@@ -11,6 +11,7 @@ struct NIAGARA_API FNiagaraMessageStore
 {
 	GENERATED_BODY();
 
+#if WITH_EDITORONLY_DATA
 	const TMap<FGuid, TObjectPtr<UNiagaraMessageDataBase>>& GetMessages() const;
 	void SetMessages(const TMap<FGuid, TObjectPtr<UNiagaraMessageDataBase>>& InMessageKeyToMessageMap);
 	void AddMessage(const FGuid& MessageKey, UNiagaraMessageDataBase* NewMessage);
@@ -26,8 +27,11 @@ private:
 
 	UPROPERTY()
 	TArray<FGuid> DismissedMessageKeys;
+#endif
 };
 
+#if WITH_EDITORONLY_DATA
+
 struct NIAGARA_API FNiagaraMessageSourceAndStore
 {
 public:
@@ -41,3 +45,5 @@ private:
 	TWeakObjectPtr<UObject> SourceWeak;
 	FNiagaraMessageStore* Store;
 };
+
+#endif
