Added duplicate script variable validation & fix logic.
Validation & Fix can be run via command line in editor
Commandlet can be run via project startup commandline and will print out all scripts with their given guid conflicts

diff --git a/Engine/Plugins/FX/Niagara/Source/NiagaraEditor/Private/Commandlets/NiagaraScriptValidationCommandlet.cpp b/Engine/Plugins/FX/Niagara/Source/NiagaraEditor/Private/Commandlets/NiagaraScriptValidationCommandlet.cpp
new file mode 100644
--- /dev/null
+++ b/Engine/Plugins/FX/Niagara/Source/NiagaraEditor/Private/Commandlets/NiagaraScriptValidationCommandlet.cpp
@@ -0,0 +1,237 @@
+// Copyright Epic Games, Inc. All Rights Reserved.
+
+#include "Commandlets/NiagaraScriptValidationCommandlet.h"
+#include "AssetRegistry/ARFilter.h"
+#include "AssetRegistry/AssetData.h"
+#include "AssetRegistry/AssetRegistryModule.h"
+#include "CollectionManagerModule.h"
+#include "CollectionManagerTypes.h"
+#include "HAL/FileManager.h"
+#include "ICollectionManager.h"
+#include "NiagaraEditorUtilities.h"
+#include "Misc/OutputDeviceArchiveWrapper.h"
+#include "Misc/PackageName.h"
+#include "Misc/Paths.h"
+#include "Modules/ModuleManager.h"
+#include "NiagaraSystem.h"
+#include "NiagaraScriptSource.h"
+#include "NiagaraGraph.h"
+#include "NiagaraScriptVariable.h"
+#include "NiagaraParameterDefinitions.h"
+
+#include UE_INLINE_GENERATED_CPP_BY_NAME(NiagaraScriptValidationCommandlet)
+
+DEFINE_LOG_CATEGORY_STATIC(LogNiagaraScriptValidation, Log, All);
+
+UNiagaraScriptValidationCommandlet::UNiagaraScriptValidationCommandlet(const FObjectInitializer& ObjectInitializer)
+	: Super(ObjectInitializer)
+{
+}
+
+int32 UNiagaraScriptValidationCommandlet::Main(const FString& Params)
+{
+	TArray<FString> Tokens;
+	TArray<FString> Switches;
+
+	ParseCommandLine(*Params, Tokens, Switches);
+
+	if (!FParse::Value(*Params, TEXT("AuditOutputFolder="), AuditOutputFolder))
+	{
+		// No output folder specified. Use the default folder.
+		AuditOutputFolder = FPaths::ProjectSavedDir() / TEXT("Audit");
+	}
+	
+	// Add a timestamp to the folder
+	AuditOutputFolder /= FDateTime::Now().ToString();
+
+	FParse::Value(*Params, TEXT("FilterCollection="), FilterCollection);
+
+	// Package Paths
+	FString PackagePathsString;
+	if (FParse::Value(*Params, TEXT("PackagePaths="), PackagePathsString, false))
+	{
+		TArray<FString> PackagePathsStrings;
+		PackagePathsString.ParseIntoArray(PackagePathsStrings, TEXT(","));
+		for (const FString& v : PackagePathsStrings)
+		{
+			PackagePaths.Add(FName(v));
+		}
+	}
+
+	if (PackagePaths.Num() == 0)
+	{
+		//PackagePaths.Add(FName(TEXT("/Game")));
+		PackagePaths.Add(FName(TEXT("/Plugins")));
+		PackagePaths.Add(FName(TEXT("/Niagara")));
+		PackagePaths.Add(FName(TEXT("/NiagaraFluids")));
+		PackagePaths.Add(FName(TEXT("/Engine")));
+	}
+
+	ProcessNiagaraScripts();
+
+	return 0;
+}
+
+void UNiagaraScriptValidationCommandlet::ProcessNiagaraScripts()
+{
+	FAssetRegistryModule& AssetRegistryModule = FModuleManager::LoadModuleChecked<FAssetRegistryModule>(TEXT("AssetRegistry"));
+	IAssetRegistry& AssetRegistry = AssetRegistryModule.Get();
+	AssetRegistry.SearchAllAssets(true);
+
+	FARFilter Filter;
+	Filter.PackagePaths = PackagePaths;
+	Filter.bRecursivePaths = true;
+
+	Filter.ClassPaths.Add(UNiagaraScript::StaticClass()->GetClassPathName());
+	if (!FilterCollection.IsEmpty())
+	{
+		FCollectionManagerModule& CollectionManagerModule = FCollectionManagerModule::GetModule();
+		CollectionManagerModule.Get().GetObjectsInCollection(FName(*FilterCollection), ECollectionShareType::CST_All, Filter.SoftObjectPaths, ECollectionRecursionFlags::SelfAndChildren);
+	}
+
+	TArray<FAssetData> AssetList;
+	AssetRegistry.GetAssets(Filter, AssetList);
+	
+	const double StartProcessNiagaraSystemsTime = FPlatformTime::Seconds();
+
+	//  Iterate over all scripts
+	const FString DevelopersFolder = FPackageName::FilenameToLongPackageName(FPaths::GameDevelopersDir().LeftChop(1));
+	FString LastPackageName = TEXT("");
+	UPackage* CurrentPackage = nullptr;
+
+	IFileManager::Get().MakeDirectory(*(AuditOutputFolder / "ScriptValidation"));
+
+	for (const FAssetData& AssetIt : AssetList)
+	{
+		const FString SystemName = AssetIt.GetObjectPathString();
+		const FString PackageName = AssetIt.PackageName.ToString();
+
+		if (PackageName.StartsWith(DevelopersFolder))
+		{
+			// Skip developer folders
+			continue;
+		}
+
+		if (PackageName != LastPackageName)
+		{
+			UPackage* Package = ::LoadPackage(nullptr, *PackageName, LOAD_None);
+			if (Package != nullptr)
+			{
+				LastPackageName = PackageName;
+				Package->FullyLoad();
+				CurrentPackage = Package;
+			}
+			else
+			{
+				UE_LOG(LogNiagaraScriptValidation, Warning, TEXT("Failed to load package %s processing %s"), *PackageName, *SystemName);
+				CurrentPackage = nullptr;
+			}
+		}
+
+		const FString ShorterScriptName = AssetIt.AssetName.ToString();
+		UNiagaraScript* NiagaraScript = FindObject<UNiagaraScript>(CurrentPackage, *ShorterScriptName);
+		if (NiagaraScript == nullptr)
+		{
+			UE_LOG(LogNiagaraScriptValidation, Warning, TEXT("Failed to load Niagara script %s"), *SystemName);
+			continue;
+		}
+		
+		TArray<FNiagaraAssetVersion> Versions = NiagaraScript->GetAllAvailableVersions();
+		
+		TMap<FNiagaraAssetVersion, TMap<FGuid, TArray<FNiagaraVariableBase>>> AllResultsForScript;
+		bool bShouldWriteLogFile = false;
+		
+		for(FNiagaraAssetVersion Version : Versions)
+		{
+			const UNiagaraScriptSource* ScriptSrc = Cast<UNiagaraScriptSource>(NiagaraScript->GetSource(Version.VersionGuid));
+			TSharedPtr<FNiagaraVersionDataAccessor> VersionDataAccessor = NiagaraScript->GetVersionDataAccessor(Version.VersionGuid);
+			if(ScriptSrc)
+			{
+				if(UNiagaraGraph* FunctionGraph = ScriptSrc->NodeGraph)
+				{
+					const TMap<FNiagaraVariable, TObjectPtr<UNiagaraScriptVariable>>& AllVars = FunctionGraph->GetAllMetaData();
+					TArray<TObjectPtr<UNiagaraScriptVariable>> ScriptPtrVars;
+					AllVars.GenerateValueArray(ScriptPtrVars);
+					TArray<UNiagaraScriptVariable*> ScriptVars;
+
+					for(const TObjectPtr<UNiagaraScriptVariable>& Var : ScriptPtrVars)
+					{
+						if(Var == nullptr)
+						{
+							continue;								
+						}
+						
+						ScriptVars.Add(Var);
+					}
+					
+					ScriptVars.StableSort([](const UNiagaraScriptVariable& Lhs, const UNiagaraScriptVariable& Rhs)
+					{
+						return Lhs.Variable.GetName().ToString().Compare(Rhs.Variable.GetName().ToString(), ESearchCase::IgnoreCase) < 0;
+					});
+					
+					TMap<FGuid, TArray<FNiagaraVariableBase>> Results = FNiagaraEditorUtilities::Scripts::Validation::ValidateScriptVariableIds(NiagaraScript, Version.VersionGuid);
+					for(const auto& Result : Results)
+					{
+						if(Result.Value.Num() > 1)
+						{
+							bShouldWriteLogFile = true;
+							AllResultsForScript.Add(Version, Results);
+						}
+					}
+				}
+			}
+		}
+		
+		if(bShouldWriteLogFile)
+		{
+			const FString MetaDataFileName = AuditOutputFolder / "ScriptValidation" / ShorterScriptName + TEXT(".xml");
+			TUniquePtr<FArchive> FileArchive(IFileManager::Get().CreateDebugFileWriter(*MetaDataFileName));
+			TUniquePtr<FOutputDeviceArchiveWrapper> OutputStream(new FOutputDeviceArchiveWrapper(FileArchive.Get()));
+
+			OutputStream->Log(TEXT("<?xml version='1.0' ?>"));
+			OutputStream->Log(TEXT("<Info>"));
+			OutputStream->Logf(TEXT("\t<Path>%s</Path>"), *GetPathNameSafe(NiagaraScript));
+			for(const auto& ResultForScript : AllResultsForScript)
+			{
+				bool bSkipLogVersion = true;
+				for(const auto& GuidConflict : ResultForScript.Value)
+				{
+					if(GuidConflict.Value.Num() > 1)
+					{
+						bSkipLogVersion = false;
+						break;
+					}
+				}
+
+				if(bSkipLogVersion)
+				{
+					continue;
+				}
+			
+				OutputStream->Logf(TEXT("\t<VersionNumber>%d . %d</VersionNumber>"), ResultForScript.Key.MajorVersion, ResultForScript.Key.MinorVersion);
+				for(const auto& GuidConflict : ResultForScript.Value)
+				{
+					if(GuidConflict.Value.Num() > 1)
+					{
+						OutputStream->Log(TEXT("\t\t<GuidConflict>"));
+						OutputStream->Logf(TEXT("\t\t\t<Guid>%s</Guid>"), *GuidConflict.Key.ToString());
+						for(const FNiagaraVariableBase& VariableBase : GuidConflict.Value)
+						{
+							OutputStream->Logf(TEXT("\t\t\t<Name>%s</Name>"), *VariableBase.GetName().ToString());
+						}
+						OutputStream->Log(TEXT("\t\t</GuidConflict>"));
+					}
+				}						
+			}
+			OutputStream->Log(TEXT("</Info>"));	
+		}
+	}
+
+	
+	// Probably don't need to do this, but just in case we have any 'hanging' packages 
+	// and more processing steps are added later, let's clean up everything...
+	::CollectGarbage(RF_NoFlags);
+
+	double ProcessNiagaraSystemsTime = FPlatformTime::Seconds() - StartProcessNiagaraSystemsTime;
+	UE_LOG(LogNiagaraScriptValidation, Log, TEXT("Took %5.3f seconds to process referenced Niagara scripts..."), ProcessNiagaraSystemsTime);
+}
diff --git a/Engine/Plugins/FX/Niagara/Source/NiagaraEditor/Private/NiagaraEditorModule.cpp b/Engine/Plugins/FX/Niagara/Source/NiagaraEditor/Private/NiagaraEditorModule.cpp
--- a/Engine/Plugins/FX/Niagara/Source/NiagaraEditor/Private/NiagaraEditorModule.cpp
+++ b/Engine/Plugins/FX/Niagara/Source/NiagaraEditor/Private/NiagaraEditorModule.cpp
@@ -1337,6 +1337,16 @@ void FNiagaraEditorModule::StartupModule()
 		TEXT("Compiles the specified script on disk for the niagara vector vm"),
 		FConsoleCommandWithArgsDelegate::CreateRaw(this, &FNiagaraEditorModule::TestCompileScriptFromConsole));
 
+	ValidateScriptVariableGuidsCommand = IConsoleManager::Get().RegisterConsoleCommand(
+		TEXT("fx.Niagara.ValidateDuplicateVariableGuids"),
+		TEXT("Validate the script guids of a given script."),
+		FConsoleCommandWithArgsDelegate::CreateRaw(this, &FNiagaraEditorModule::ValidateScriptVariableIds, false));
+		
+	ValidateAndFixScriptVariableGuidsCommand = IConsoleManager::Get().RegisterConsoleCommand(
+		TEXT("fx.Niagara.FixDuplicateVariableGuids"),
+		TEXT("Validates and fixes the script guids of a given script, if duplicates exist."),
+		FConsoleCommandWithArgsDelegate::CreateRaw(this, &FNiagaraEditorModule::ValidateScriptVariableIds, true));
+
 	DumpRapidIterationParametersForAsset = IConsoleManager::Get().RegisterConsoleCommand(
 		TEXT("fx.DumpRapidIterationParametersForAsset"),
 		TEXT("Dumps the values of the rapid iteration parameters for the specified asset by path."),
@@ -2156,6 +2166,32 @@ UNiagaraParameterCollection* FNiagaraEditorModule::FindCollectionForVariable(con
 	return FindCachedCollectionByPrefix(VariableName);
 }
 
+void FNiagaraEditorModule::ValidateScriptVariableIds(const TArray<FString>& ScriptPathArgs, bool bFix)
+{
+	for(const FString& ScriptPath : ScriptPathArgs)
+	{
+		FSoftObjectPath ScriptSoftPath(ScriptPath);
+		if(ScriptSoftPath.IsAsset())
+		{
+			UObject* ScriptObject = ScriptSoftPath.TryLoad();
+			if(UNiagaraScript* Script = Cast<UNiagaraScript>(ScriptObject))
+			{
+				if(bFix)
+				{
+					FNiagaraEditorUtilities::Scripts::Validation::FixupDuplicateScriptVariableGuids(Script);
+				}
+				else
+				{
+					for(FNiagaraAssetVersion& AssetVersion : Script->GetAllAvailableVersions())
+					{
+						FNiagaraEditorUtilities::Scripts::Validation::ValidateScriptVariableIds(Script, AssetVersion.VersionGuid);
+					}
+				}
+			}
+		}
+	}	
+}
+
 #if NIAGARA_PERF_BASELINES
 void FNiagaraEditorModule::GeneratePerfBaselines(TArray<UNiagaraEffectType*>& BaselinesToGenerate)
 {
diff --git a/Engine/Plugins/FX/Niagara/Source/NiagaraEditor/Private/NiagaraEditorUtilities.cpp b/Engine/Plugins/FX/Niagara/Source/NiagaraEditor/Private/NiagaraEditorUtilities.cpp
--- a/Engine/Plugins/FX/Niagara/Source/NiagaraEditor/Private/NiagaraEditorUtilities.cpp
+++ b/Engine/Plugins/FX/Niagara/Source/NiagaraEditor/Private/NiagaraEditorUtilities.cpp
@@ -4583,6 +4583,111 @@ UNiagaraDataInterface* FNiagaraEditorUtilities::GetResolvedRuntimeInstanceForEdi
 	return nullptr;
 }
 
+TMap<FGuid, TArray<FNiagaraVariableBase>> FNiagaraEditorUtilities::Scripts::Validation::ValidateScriptVariableIds(UNiagaraScript* Script, FGuid VersionGuid)
+{
+	bool bReturnResults = false;
+	
+	TMap<FGuid, TArray<FNiagaraVariableBase>> Results;
+
+	if(!VersionGuid.IsValid())
+	{
+		VersionGuid = Script->GetExposedVersion().VersionGuid;
+	}
+	
+	if(UNiagaraScriptSourceBase* SourceBase = Script->GetSource(VersionGuid))
+	{
+		UNiagaraScriptSource* Source = CastChecked<UNiagaraScriptSource>(SourceBase);
+		
+		for(const auto& ScriptVariableEntry : Source->NodeGraph->GetAllMetaData())
+		{
+			Results.FindOrAdd(ScriptVariableEntry.Value->Metadata.GetVariableGuid()).Add(ScriptVariableEntry.Key);
+		}
+
+		for(const auto& Result : Results)
+		{
+			if(Result.Value.Num() > 1)
+			{
+				bReturnResults = true;
+				break;
+			}
+		}
+
+		if(bReturnResults)
+		{
+			UE_LOG(LogNiagaraEditor, Error, TEXT("---- Duplicate variable guids found in script %s, Version %s ---- :"), *Script->GetName(), *VersionGuid.ToString());
+			for(const auto& Result : Results)
+			{
+				if(Result.Value.Num() > 1)
+				{
+					FString ResultString("Duplicate guids found in: ");
+					for(int32 StringConstructorIndex = 0; StringConstructorIndex < Result.Value.Num(); StringConstructorIndex++)
+					{
+						ResultString.Appendf(TEXT("%s, "), *Result.Value[StringConstructorIndex].GetName().ToString());
+					}
+			
+					UE_LOG(LogNiagaraEditor, Error, TEXT("%s"), *ResultString);
+				}
+			}
+	
+			return Results;
+		}
+	}
+
+	UE_LOG(LogNiagaraEditor, Log, TEXT("---- No duplicate variable guids in script %s, version %s ----"), *Script->GetName(), *VersionGuid.ToString());
+	return {};
+}
+
+TMap<FNiagaraVariableBase, FGuid> FNiagaraEditorUtilities::Scripts::Validation::FixupDuplicateScriptVariableGuids(UNiagaraScript* Script)
+{
+	/** We keep track of already reassigned guids here to reuse them if the same parameter appears in another version. */
+	TMap<FNiagaraVariableBase, FGuid> OldGuids;
+	TMap<FNiagaraVariableBase, FGuid> RemappedGuids;
+	
+	for(FNiagaraAssetVersion& AssetVersion : Script->GetAllAvailableVersions())
+	{
+		UNiagaraScriptSource* ScriptSource = Cast<UNiagaraScriptSource>(Script->GetSource(AssetVersion.VersionGuid));
+		TMap<FNiagaraVariable, TObjectPtr<UNiagaraScriptVariable>>& VariableMap = ScriptSource->NodeGraph->GetAllMetaData();
+		
+		TMap<FGuid, TArray<FNiagaraVariableBase>> Results = Validation::ValidateScriptVariableIds(Script, AssetVersion.VersionGuid);
+		for(const auto& Result : Results)
+		{
+			// we only want to fix any guids if we have more than 1 variable with the same guid
+			if(Result.Value.Num() > 1)
+			{
+				// if a parameter has already been given a new guid before (in a different version), we reuse the same new guid
+				for(const FNiagaraVariableBase& VariableWithDuplicateGuid : Result.Value)
+				{
+					if(RemappedGuids.Contains(VariableWithDuplicateGuid))
+					{
+						FGuid NewGuid = RemappedGuids[VariableWithDuplicateGuid];
+						VariableMap[VariableWithDuplicateGuid]->Metadata.SetVariableGuid(NewGuid);
+					}
+					else
+					{
+						Script->Modify(true);
+						FGuid NewGuid = FGuid::NewGuid();
+						OldGuids.Add(VariableWithDuplicateGuid, VariableMap[VariableWithDuplicateGuid]->Metadata.GetVariableGuid());
+						VariableMap[VariableWithDuplicateGuid]->Metadata.SetVariableGuid(NewGuid);
+						RemappedGuids.Add(VariableWithDuplicateGuid, NewGuid);
+					}
+				}
+			}
+		}
+	}
+
+	if(RemappedGuids.Num() > 0)
+	{
+		UE_LOG(LogNiagaraEditor, Log, TEXT("----------- Variables with new guids in script %s ---------"), *Script->GetName());
+	}
+	
+	for(const auto& RemappedGuid : RemappedGuids)
+	{
+		UE_LOG(LogNiagaraEditor, Log, TEXT("Variable: %s. Guid: %s -> %s"), *RemappedGuid.Key.GetName().ToString(), *OldGuids[RemappedGuid.Key].ToString(), *RemappedGuids[RemappedGuid.Key].ToString())
+	}
+
+	return RemappedGuids;
+}
+
 void FNiagaraParameterUtilities::FilterToRelevantStaticVariables(TConstArrayView<FNiagaraVariable> InVars, TArray<FNiagaraVariable>& OutVars, FName InOldEmitterAlias, FName InNewEmitterAlias, bool bFilterByEmitterAliasAndConvertToUnaliased)
 {
 	FNiagaraAliasContext RenameContext(ENiagaraScriptUsage::ParticleSpawnScript);
diff --git a/Engine/Plugins/FX/Niagara/Source/NiagaraEditor/Public/Commandlets/NiagaraScriptValidationCommandlet.h b/Engine/Plugins/FX/Niagara/Source/NiagaraEditor/Public/Commandlets/NiagaraScriptValidationCommandlet.h
new file mode 100644
--- /dev/null
+++ b/Engine/Plugins/FX/Niagara/Source/NiagaraEditor/Public/Commandlets/NiagaraScriptValidationCommandlet.h
@@ -0,0 +1,30 @@
+// Copyright Epic Games, Inc. All Rights Reserved.
+
+#pragma once
+
+#include "CoreMinimal.h"
+#include "UObject/ObjectMacros.h"
+#include "Commandlets/Commandlet.h"
+#include "NiagaraScriptValidationCommandlet.generated.h"
+
+class UNiagaraScript;
+
+UCLASS(config=Editor)
+class UNiagaraScriptValidationCommandlet : public UCommandlet
+{
+	GENERATED_UCLASS_BODY()
+
+	virtual int32 Main(const FString& Params) override;
+
+	/** The folder in which the commandlet's output files will be stored */
+	FString AuditOutputFolder;
+
+	/** Only assets in this collection will be considered. If this is left blank, no assets will be filtered by collection */
+	FString FilterCollection;
+
+	/** Package paths to include */
+	TArray<FName> PackagePaths;
+
+private:
+	void ProcessNiagaraScripts();
+};
\ No newline at end of file
diff --git a/Engine/Plugins/FX/Niagara/Source/NiagaraEditor/Public/NiagaraEditorModule.h b/Engine/Plugins/FX/Niagara/Source/NiagaraEditor/Public/NiagaraEditorModule.h
--- a/Engine/Plugins/FX/Niagara/Source/NiagaraEditor/Public/NiagaraEditorModule.h
+++ b/Engine/Plugins/FX/Niagara/Source/NiagaraEditor/Public/NiagaraEditorModule.h
@@ -324,6 +324,9 @@ private:
 	}
 
 	void TestCompileScriptFromConsole(const TArray<FString>& Arguments);
+
+	void ValidateScriptVariableIds(const TArray<FString>& ScriptPathArgs, bool bFix = false);
+	
 	void ReinitializeStyle();
 
 	void EnqueueObjectForDeferredDestructionInternal(FDeferredDestructionContainerBase* InObjectToDestruct);
@@ -389,6 +392,8 @@ private:
 	TMap<const UScriptStruct*, FOnCreateMovieSceneTrackForParameter> TypeToParameterTrackCreatorMap;
 
 	IConsoleCommand* TestCompileScriptCommand;
+	IConsoleCommand* ValidateScriptVariableGuidsCommand;
+	IConsoleCommand* ValidateAndFixScriptVariableGuidsCommand;
 	IConsoleCommand* DumpRapidIterationParametersForAsset;
 	IConsoleCommand* PreventSystemRecompileCommand;
 	IConsoleCommand* PreventAllSystemRecompilesCommand;
diff --git a/Engine/Plugins/FX/Niagara/Source/NiagaraEditor/Public/NiagaraEditorUtilities.h b/Engine/Plugins/FX/Niagara/Source/NiagaraEditor/Public/NiagaraEditorUtilities.h
--- a/Engine/Plugins/FX/Niagara/Source/NiagaraEditor/Public/NiagaraEditorUtilities.h
+++ b/Engine/Plugins/FX/Niagara/Source/NiagaraEditor/Public/NiagaraEditorUtilities.h
@@ -413,6 +413,15 @@ namespace FNiagaraEditorUtilities
 	NIAGARAEDITOR_API bool IsEditorDataInterfaceInstance(const UNiagaraDataInterface* DataInterface);
 
 	NIAGARAEDITOR_API UNiagaraDataInterface* GetResolvedRuntimeInstanceForEditorDataInterfaceInstance(const UNiagaraSystem& OwningSystem, UNiagaraDataInterface& EditorDataInterfaceInstance);
+
+	namespace Scripts
+	{
+		namespace Validation
+		{
+			TMap<FGuid, TArray<FNiagaraVariableBase>> ValidateScriptVariableIds(UNiagaraScript* Script, FGuid VersionGuid);
+			TMap<FNiagaraVariableBase, FGuid> FixupDuplicateScriptVariableGuids(UNiagaraScript* Script);
+		}
+	}
 }
 
 namespace FNiagaraParameterUtilities
