Made detection of newly spawned FX more robust in the scalability manager to fix issues where scalability updates were not performed for new systems.

diff --git a/Engine/Plugins/FX/Niagara/Source/Niagara/Private/NiagaraScalabilityManager.cpp b/Engine/Plugins/FX/Niagara/Source/Niagara/Private/NiagaraScalabilityManager.cpp
--- a/Engine/Plugins/FX/Niagara/Source/Niagara/Private/NiagaraScalabilityManager.cpp
+++ b/Engine/Plugins/FX/Niagara/Source/Niagara/Private/NiagaraScalabilityManager.cpp
@@ -205,12 +205,14 @@ bool FNiagaraScalabilityManager::EvaluateCullState(FNiagaraWorldManager* WorldMa
 		return false;
 	}
 
+	FNiagaraScalabilityState& CompState = State[ComponentIndex];
+
 	//Don't update if we're doing new systems only and this is not new.
 	//Saves the potential cost of reavaluating every effect in every tick group something new is added.
 	//Though this does mean the sorted significance values will be using out of date distances etc.
 	//I'm somewhat on the fence currently as to whether it's better to pay this cost for correctness.
 	const bool UpdateScalability = Component->ScalabilityManagerHandle == ComponentIndex
-		&& (!Context.bNewOnly || Component->GetSystemInstanceController()->IsPendingSpawn());
+		&& (!Context.bNewOnly || CompState.bNewlyRegistered);
 
 	if (UpdateScalability)
 	{
@@ -222,8 +224,8 @@ bool FNiagaraScalabilityManager::EvaluateCullState(FNiagaraWorldManager* WorldMa
 			Unregister(Component);
 			return false;
 		}
+		CompState.bNewlyRegistered = false;
 
-		FNiagaraScalabilityState& CompState = State[ComponentIndex];
 		const FNiagaraSystemScalabilitySettings& Scalability = System->GetScalabilitySettings();
 
 		const FNiagaraScalabilitySystemData& SysData = GetSystemData(ComponentIndex);
diff --git a/Engine/Plugins/FX/Niagara/Source/Niagara/Public/NiagaraScalabilityState.h b/Engine/Plugins/FX/Niagara/Source/Niagara/Public/NiagaraScalabilityState.h
--- a/Engine/Plugins/FX/Niagara/Source/Niagara/Public/NiagaraScalabilityState.h
+++ b/Engine/Plugins/FX/Niagara/Source/Niagara/Public/NiagaraScalabilityState.h
@@ -13,6 +13,7 @@ struct FNiagaraScalabilityState
 		: Significance(1.0f)
 		, LastVisibleTime(0.0f)
 		, SystemDataIndex(INDEX_NONE)
+		, bNewlyRegistered(1)
 		, bCulled(0)
 		, bPreviousCulled(0)
 		, bCulledByDistance(0)
@@ -26,6 +27,7 @@ struct FNiagaraScalabilityState
 		: Significance(InSignificance)
 		, LastVisibleTime(0.0f)
 		, SystemDataIndex(INDEX_NONE)
+		, bNewlyRegistered(1)
 		, bCulled(InCulled)
 		, bPreviousCulled(InPreviousCulled)
 		, bCulledByDistance(0)
@@ -46,6 +48,9 @@ struct FNiagaraScalabilityState
 
 	int16 SystemDataIndex;
 
+	UPROPERTY(VisibleAnywhere, Category = "Scalability")
+	uint8 bNewlyRegistered : 1;
+
 	UPROPERTY(VisibleAnywhere, Category = "Scalability")
 	uint8 bCulled : 1;
 
