- Add display for approximate memory information to Niagara System Editor
- Uses archive count to determine system / emitter
- Component use resource count which will include limited things like particle data

diff --git a/Engine/Plugins/FX/Niagara/Source/Niagara/Internal/NiagaraSystemMemReport.h b/Engine/Plugins/FX/Niagara/Source/Niagara/Internal/NiagaraSystemMemReport.h
new file mode 100644
--- /dev/null
+++ b/Engine/Plugins/FX/Niagara/Source/Niagara/Internal/NiagaraSystemMemReport.h
@@ -0,0 +1,39 @@
+// Copyright Epic Games, Inc. All Rights Reserved.
+
+#include "CoreMinimal.h"
+
+class UObject;
+class UNiagaraSystem;
+
+struct FNiagaraSystemMemReport
+{
+	enum class EReportType
+	{
+		Basic,
+		Verbose,
+		Max
+	};
+
+	struct FNode
+	{
+		explicit FNode(UObject* Object, uint32 InDepth);
+		explicit FNode(FName CustomName, uint32 InDepth, uint32 ByteSize);
+
+		FName	ObjectName;
+		FName	ObjectClass;
+		uint32	Depth = 0;
+		uint64	ExclusiveSizeBytes = 0;
+		uint64	InclusiveSizeBytes = 0;
+	};
+
+	NIAGARA_API void GenerateReport(EReportType ReportType, UNiagaraSystem* System);
+
+	TConstArrayView<FNode> GetNodes() const { return Nodes; }
+	uint64 GetDataInterfaceSizeBytes() const { return DataInterfaceSizeBytes; }
+
+private:
+	uint64 GatherResourceMemory(EReportType ReportType, UObject* Object, uint32 Depth);
+
+	TArray<FNode> Nodes;
+	uint64 DataInterfaceSizeBytes = 0;
+};
diff --git a/Engine/Plugins/FX/Niagara/Source/Niagara/Private/NiagaraSystemMemReport.cpp b/Engine/Plugins/FX/Niagara/Source/Niagara/Private/NiagaraSystemMemReport.cpp
new file mode 100644
--- /dev/null
+++ b/Engine/Plugins/FX/Niagara/Source/Niagara/Private/NiagaraSystemMemReport.cpp
@@ -0,0 +1,259 @@
+// Copyright Epic Games, Inc. All Rights Reserved.
+
+#include "NiagaraSystemMemReport.h"
+
+#include "NiagaraDataInterface.h"
+#include "NiagaraEmitter.h"
+#include "Misc/StringBuilder.h"
+#include "UObject/UObjectIterator.h"
+
+#include "Serialization/ArchiveCountMem.h"
+
+#include "NiagaraSystem.h"
+#include "NiagaraSystemImpl.h"
+
+namespace NiagaraMemReportInternal
+{
+	const FName NAME_ByteCode("ByteCode");
+	const FName NAME_RapidIterationParameters("RapidIterationParameters");
+	const FName NAME_ScriptExecutionParamStore("ScriptExecutionParamStore");
+
+	static FAutoConsoleCommandWithWorldAndArgs CmdReportSystemMem(
+		TEXT("NiagaraReportSystemMemory"),
+		TEXT("Dumps some rough information about system memory breakdown"),
+		FConsoleCommandWithWorldAndArgsDelegate::CreateLambda(
+			[](const TArray<FString>& Args, UWorld*)
+			{
+				FNiagaraSystemMemReport::EReportType ReportType = FNiagaraSystemMemReport::EReportType::Basic;
+				FString WildcardFilter;
+	
+				static const FString FilterPrefix("-filter=");
+				static const FString TypePrefix("-type=");
+				for (const FString& Arg : Args)
+				{
+					if (Arg.StartsWith(FilterPrefix))
+					{
+						WildcardFilter = Arg.Mid(FilterPrefix.Len());
+					}
+					else if (Arg.StartsWith(TypePrefix))
+					{
+						FString ValueString = Arg.Mid(TypePrefix.Len());
+						ReportType = FNiagaraSystemMemReport::EReportType(FMath::Clamp(FCString::Atoi(*ValueString), 0, int32(FNiagaraSystemMemReport::EReportType::Max) - 1));
+					}
+				}
+				
+				UE_LOG(LogNiagara, Log, TEXT("============ Niagara MemReport"));
+				for (TObjectIterator<UNiagaraSystem> SystemIt; SystemIt; ++SystemIt)
+				{
+					UNiagaraSystem* NiagaraSystem = *SystemIt;
+					if (::IsValid(NiagaraSystem) == false)
+					{
+						continue;
+					}
+
+					if (!WildcardFilter.IsEmpty())
+					{
+						const FString PathName = NiagaraSystem->GetPathName();
+						if ( !PathName.MatchesWildcard(WildcardFilter) )
+						{
+							continue;
+						}
+					}
+
+					FNiagaraSystemMemReport MemReport;
+					MemReport.GenerateReport(ReportType, NiagaraSystem);
+					for ( const FNiagaraSystemMemReport::FNode& Node : MemReport.GetNodes() )
+					{
+						UE_LOG(LogNiagara, Log, TEXT("%*s = %s %s (%u, %u)"), Node.Depth, TEXT(""), *Node.ObjectClass.ToString(), *Node.ObjectName.ToString(), uint32(Node.ExclusiveSizeBytes), uint32(Node.InclusiveSizeBytes));
+					}
+					UE_LOG(LogNiagara, Log, TEXT(""));
+				}
+			}
+		)
+	);
+
+#if WITH_EDITORONLY_DATA
+	void InternalHandleProperty(TSet<UObject*>& EditorOnlyObjects, TSet<UObject*>& RuntimeObjects, const FProperty* Property, const void* Container, bool bIsEditorOnly, int32 ArrayIndex = 0)
+	{
+		UObject* FoundObject = nullptr;
+		if (const FWeakObjectProperty* WeakObjectProperty = CastField<const FWeakObjectProperty>(Property))
+		{
+			const FWeakObjectPtr& Value = WeakObjectProperty->GetPropertyValue_InContainer(Container, ArrayIndex);
+			FoundObject = Value.Get();
+		}
+		else if (const FObjectProperty* ObjectProperty = CastField<const FObjectProperty>(Property))
+		{
+			const TObjectPtr<UObject>& Value = ObjectProperty->GetPropertyValue_InContainer(Container, ArrayIndex);
+			FoundObject = Value.Get();
+		}
+		else if (const FSoftObjectProperty* SoftObjectProperty = CastField<const FSoftObjectProperty>(Property))
+		{
+			const FSoftObjectPtr& Value = SoftObjectProperty->GetPropertyValue_InContainer(Container, ArrayIndex);
+			FoundObject = Value.Get();
+		}
+
+		if (FoundObject != nullptr)
+		{
+			if (bIsEditorOnly)
+			{
+				EditorOnlyObjects.Add(FoundObject);
+			}
+			else
+			{
+				RuntimeObjects.Add(FoundObject);
+			}
+		}
+	}
+
+	void InternalFindSubObjects(TSet<UObject*>& EditorOnlyObjects, TSet<UObject*>& RuntimeObjects, const UStruct* Struct, const void* Container, bool bIsParentEditorOnly)
+	{
+		for (TFieldIterator<const FProperty> PropertyIt(Struct); PropertyIt; ++PropertyIt)
+		{
+			const FProperty* Property = *PropertyIt;
+			const bool bIsEditorOnly = Property->IsEditorOnlyProperty() || bIsParentEditorOnly;
+
+			if (const FStructProperty* StructProperty = CastField<const FStructProperty>(Property))
+			{
+				if ( StructProperty->Struct )
+				{
+					InternalFindSubObjects(EditorOnlyObjects, RuntimeObjects, StructProperty->Struct, StructProperty->ContainerPtrToValuePtr<void>(Container, 0), bIsEditorOnly);
+				}
+			}
+			else if (const FArrayProperty* ArrayProperty = CastField<const FArrayProperty>(Property))
+			{
+				const FProperty* ArrayInnerProp = ArrayProperty->Inner;
+
+				if (ArrayInnerProp && (CastField<const FWeakObjectProperty>(ArrayInnerProp) || CastField<const FObjectProperty>(ArrayInnerProp) ||CastField<const FSoftObjectProperty>(ArrayInnerProp)))
+				{
+					FScriptArrayHelper ArrayHelper(ArrayProperty, ArrayProperty->ContainerPtrToValuePtr<void>(Container, 0));
+					for (int32 i=0; i < ArrayHelper.Num(); ++i)
+					{
+						const void* ArrayItemContainer = ArrayHelper.GetRawPtr(i);
+						InternalHandleProperty(EditorOnlyObjects, RuntimeObjects, ArrayInnerProp, ArrayItemContainer, bIsEditorOnly);
+					}
+				}
+			}
+			else
+			{
+				InternalHandleProperty(EditorOnlyObjects, RuntimeObjects, Property, Container, bIsEditorOnly);
+			}
+		}
+	}
+
+	void RemoveEditorOnlyObjects(TArray<UObject*>& SubObjects, UObject* Object)
+	{
+		SubObjects.RemoveAll([](UObject* SubObject) { return SubObject->IsEditorOnly() || (!SubObject->NeedsLoadForClient() && !SubObject->NeedsLoadForServer()); });
+		if (SubObjects.Num() == 0)
+		{
+			return;
+		}
+
+		TSet<UObject*> EditorOnlyObjects;
+		TSet<UObject*> RuntimeObjects;
+		InternalFindSubObjects(EditorOnlyObjects, RuntimeObjects, Object->GetClass(), Object, false);
+
+		SubObjects.RemoveAll([&EditorOnlyObjects, &RuntimeObjects](UObject* SubObject) { return EditorOnlyObjects.Contains(SubObject) && !RuntimeObjects.Contains(SubObject); });
+	}
+#endif
+}
+
+FNiagaraSystemMemReport::FNode::FNode(UObject* Object, uint32 InDepth)
+{
+	FResourceSizeEx ResourceSize(EResourceSizeMode::Exclusive);
+	ResourceSize.AddDedicatedSystemMemoryBytes(FArchiveCountMem(Object, true).GetMax());
+	Object->GetResourceSizeEx(ResourceSize);
+
+	ObjectName			= Object->GetFName();
+	ObjectClass			= Object->GetClass()->GetFName();
+	Depth				= InDepth;
+	ExclusiveSizeBytes	= ResourceSize.GetTotalMemoryBytes();
+	InclusiveSizeBytes	= ResourceSize.GetTotalMemoryBytes();
+
+	if (UNiagaraEmitter* Emitter = Cast<UNiagaraEmitter>(Object))
+	{
+		ObjectName = FName(*Emitter->GetUniqueEmitterName());
+	}
+}
+
+FNiagaraSystemMemReport::FNode::FNode(FName CustomName, uint32 InDepth, uint32 ByteSize)
+{
+	ObjectName			= CustomName;
+	ObjectClass			= NAME_None;
+	Depth				= InDepth;
+	ExclusiveSizeBytes	= ByteSize;
+	InclusiveSizeBytes	= ByteSize;
+}
+
+void FNiagaraSystemMemReport::GenerateReport(EReportType ReportType, UNiagaraSystem* System)
+{
+	Nodes.Empty();
+	DataInterfaceSizeBytes = 0;
+	GatherResourceMemory(ReportType, System, 0);
+}
+
+uint64 FNiagaraSystemMemReport::GatherResourceMemory(EReportType ReportType, UObject* Object, uint32 Depth)
+{
+	using namespace NiagaraMemReportInternal;
+
+	const int32 NodeIndex = Nodes.Emplace(Object, Depth);
+
+	if (ReportType == EReportType::Verbose)
+	{
+		if (UNiagaraScript* NiagaraScript = Cast<UNiagaraScript>(Object))
+		{
+			{
+				const FNiagaraVMExecutableData& VMData = NiagaraScript->GetVMExecutableData();
+				Nodes.Emplace(NAME_ByteCode, Depth + 1, VMData.ByteCode.GetLength() + VMData.OptimizedByteCode.GetLength());
+				Nodes[NodeIndex].ExclusiveSizeBytes -= Nodes.Last().ExclusiveSizeBytes;
+			}
+
+			Nodes.Emplace(NAME_RapidIterationParameters, Depth + 1, NiagaraScript->RapidIterationParameters.GetResourceSize());
+			Nodes[NodeIndex].ExclusiveSizeBytes -= Nodes.Last().ExclusiveSizeBytes;
+
+			ENiagaraSimTarget SimTarget = ENiagaraSimTarget::CPUSim;
+			const FVersionedNiagaraEmitter VersionedNiagaraEmitter = NiagaraScript->GetOuterEmitter();
+			if (const FVersionedNiagaraEmitterData* EmitterData = VersionedNiagaraEmitter.GetEmitterData())
+			{
+				SimTarget = EmitterData->SimTarget;
+			}
+			if (const FNiagaraScriptExecutionParameterStore* ExecParameterStore = NiagaraScript->GetExecutionReadyParameterStore(SimTarget))
+			{
+				Nodes.Emplace(NAME_ScriptExecutionParamStore, Depth + 1, ExecParameterStore->GetResourceSize());
+				Nodes[NodeIndex].ExclusiveSizeBytes -= Nodes.Last().ExclusiveSizeBytes;
+			}
+		}
+	}
+
+	TArray<UObject*> SubObjects;
+	GetObjectsWithOuter(Object, SubObjects, false);
+	if (SubObjects.Num() > 0)
+	{
+	#if WITH_EDITORONLY_DATA
+		RemoveEditorOnlyObjects(SubObjects, Object);
+	#endif
+
+		for (UObject* SubObject : SubObjects)
+		{
+			Nodes[NodeIndex].InclusiveSizeBytes += GatherResourceMemory(ReportType, SubObject, Depth + 1);
+		}
+	}
+
+	// Keep track of data interface memory
+	if (Object->IsA<UNiagaraDataInterface>())
+	{
+		DataInterfaceSizeBytes += Nodes[NodeIndex].InclusiveSizeBytes;
+	}
+
+	// If in basic mode remove all the nodes we added below us as we don't want them
+	const uint64 InclusiveSizeBytes = Nodes[NodeIndex].InclusiveSizeBytes;
+	if (ReportType == EReportType::Basic)
+	{
+		const int32 NumToRemove = Nodes.Num() - NodeIndex;
+		if (NumToRemove > 0 && !Object->IsA<UNiagaraSystem>() && !Object->IsA<UNiagaraEmitter>())
+		{
+			Nodes.RemoveAt(NodeIndex, NumToRemove);
+		}
+	}
+
+	return InclusiveSizeBytes;
+}
diff --git a/Engine/Plugins/FX/Niagara/Source/NiagaraEditor/Private/NiagaraEditorCommands.cpp b/Engine/Plugins/FX/Niagara/Source/NiagaraEditor/Private/NiagaraEditorCommands.cpp
--- a/Engine/Plugins/FX/Niagara/Source/NiagaraEditor/Private/NiagaraEditorCommands.cpp
+++ b/Engine/Plugins/FX/Niagara/Source/NiagaraEditor/Private/NiagaraEditorCommands.cpp
@@ -26,6 +26,7 @@ void FNiagaraEditorCommands::RegisterCommands()
 	UI_COMMAND(ToggleParticleCounts, "ParticleCounts", "Display Particle Counts", EUserInterfaceActionType::ToggleButton, FInputChord());
 	UI_COMMAND(ToggleEmitterExecutionOrder, "Emitter Execution Order", "Display Emitter Execution Order", EUserInterfaceActionType::ToggleButton, FInputChord());
 	UI_COMMAND(ToggleGpuTickInformation, "Gpu Tick Information", "Display Gpu Tick Information", EUserInterfaceActionType::ToggleButton, FInputChord());
+	UI_COMMAND(ToggleMemoryInfo, "Memory", "Show approximate memory usage", EUserInterfaceActionType::ToggleButton, FInputChord());
 	UI_COMMAND(TogglePreviewBackground, "Background", "Toggles the preview pane's background.", EUserInterfaceActionType::ToggleButton, FInputChord());
 	UI_COMMAND(ToggleUnlockToChanges, "Lock/Unlock To Changes", "Toggles whether or not changes in the source asset get pulled into this asset automatically.", EUserInterfaceActionType::Button, FInputChord());
 	UI_COMMAND(ToggleStatPerformance, "Performance", "Show runtime performance for particle scripts.", EUserInterfaceActionType::ToggleButton, FInputChord());
diff --git a/Engine/Plugins/FX/Niagara/Source/NiagaraEditor/Private/NiagaraEditorSettings.cpp b/Engine/Plugins/FX/Niagara/Source/NiagaraEditor/Private/NiagaraEditorSettings.cpp
--- a/Engine/Plugins/FX/Niagara/Source/NiagaraEditor/Private/NiagaraEditorSettings.cpp
+++ b/Engine/Plugins/FX/Niagara/Source/NiagaraEditor/Private/NiagaraEditorSettings.cpp
@@ -375,6 +375,20 @@ void UNiagaraEditorSettings::SetShowGpuTickInformation(bool bInShowGpuTickInform
 	}
 }
 
+bool UNiagaraEditorSettings::IsShowMemoryInfo() const
+{
+	return bShowMemoryInfo;
+}
+
+void UNiagaraEditorSettings::SetShowMemoryInfo(bool bInShowInfo)
+{
+	if (bShowMemoryInfo != bInShowInfo)
+	{
+		bShowMemoryInfo = bInShowInfo;
+		SaveConfig();
+	}
+}
+
 TArray<float> UNiagaraEditorSettings::GetPlaybackSpeeds() const
 {
 	if(!CachedPlaybackSpeeds.IsSet())
diff --git a/Engine/Plugins/FX/Niagara/Source/NiagaraEditor/Private/Widgets/SNiagaraSystemViewport.cpp b/Engine/Plugins/FX/Niagara/Source/NiagaraEditor/Private/Widgets/SNiagaraSystemViewport.cpp
--- a/Engine/Plugins/FX/Niagara/Source/NiagaraEditor/Private/Widgets/SNiagaraSystemViewport.cpp
+++ b/Engine/Plugins/FX/Niagara/Source/NiagaraEditor/Private/Widgets/SNiagaraSystemViewport.cpp
@@ -13,6 +13,7 @@
 #include "NiagaraEditorSettings.h"
 #include "NiagaraEditorStyle.h"
 #include "NiagaraEffectType.h"
+#include "NiagaraSystemMemReport.h"
 #include "NiagaraPerfBaseline.h"
 #include "NiagaraSettings.h"
 #include "NiagaraShader.h"
@@ -28,6 +29,7 @@
 #include "Engine/Canvas.h"
 #include "Engine/Font.h"
 #include "Engine/TextureCube.h"
+#include "Serialization/ArchiveCountMem.h"
 #include "Slate/SceneViewport.h"
 #include "ThumbnailRendering/ThumbnailManager.h"
 #include "Widgets/Docking/SDockTab.h"
@@ -63,6 +65,7 @@ public:
 	void DrawParticleCounts(UNiagaraComponent* Component, FCanvas* Canvas, float& CurrentX, float& CurrentY, UFont* Font, const float FontHeight);
 	void DrawEmitterExecutionOrder(UNiagaraComponent* Component, FCanvas* Canvas, float& CurrentX, float& CurrentY, UFont* Font, const float FontHeight);
 	void DrawGpuTickInformation(UNiagaraComponent* Component, FCanvas* Canvas, float& CurrentX, float& CurrentY, UFont* Font, const float FontHeight);
+	void DrawMemoryInfo(UNiagaraComponent* Component, FCanvas* Canvas, float& CurrentX, float& CurrentY, UFont* Font, const float FontHeight);
 	void SetUpdateViewportFocus(bool bUpdate) { bUpdateViewportFocus = bUpdate; }
 
 	TWeakPtr<SNiagaraSystemViewport> NiagaraViewportPtr;
@@ -75,6 +78,9 @@ public:
 
 private:
 	bool bUpdateViewportFocus = false;
+
+	FNiagaraSystemMemReport	MemReport;
+	double					MemReportLastCheckedTime = 0.0;
 };
 
 FNiagaraSystemViewportClient::FNiagaraSystemViewportClient(FAdvancedPreviewScene& InPreviewScene, const TSharedRef<SNiagaraSystemViewport>& InNiagaraEditorViewport, FOnScreenShotCaptured InOnScreenShotCaptured)
@@ -157,24 +163,29 @@ void FNiagaraSystemViewportClient::Draw(FViewport* InViewport, FCanvas* Canvas)
 		UFont* Font = GEngine->GetSmallFont();
 		const float FontHeight = Font->GetMaxCharHeight() * 1.1f;
 
-		if (NiagaraViewport->GetDrawElement(SNiagaraSystemViewport::EDrawElements::InstructionCounts) && ParticleSystem)
+		if (NiagaraViewport->GetDrawElement(SNiagaraSystemViewport::EDrawElements::EmitterExecutionOrder) && Component)
 		{
-			DrawInstructionCounts(ParticleSystem, Canvas, CurrentX, CurrentY, Font, FontHeight);
+			DrawEmitterExecutionOrder(Component, Canvas, CurrentX, CurrentY, Font, FontHeight);
 			CurrentY += FontHeight;
 		}
-		if (NiagaraViewport->GetDrawElement(SNiagaraSystemViewport::EDrawElements::ParticleCounts) && Component)
+		if (NiagaraViewport->GetDrawElement(SNiagaraSystemViewport::EDrawElements::GpuTickInformation) && Component)
 		{
-			DrawParticleCounts(Component, Canvas, CurrentX, CurrentY, Font, FontHeight);
+			DrawGpuTickInformation(Component, Canvas, CurrentX, CurrentY, Font, FontHeight);
 			CurrentY += FontHeight;
 		}
-		if (NiagaraViewport->GetDrawElement(SNiagaraSystemViewport::EDrawElements::EmitterExecutionOrder) && Component)
+		if (NiagaraViewport->GetDrawElement(SNiagaraSystemViewport::EDrawElements::InstructionCounts) && ParticleSystem)
 		{
-			DrawEmitterExecutionOrder(Component, Canvas, CurrentX, CurrentY, Font, FontHeight);
+			DrawInstructionCounts(ParticleSystem, Canvas, CurrentX, CurrentY, Font, FontHeight);
 			CurrentY += FontHeight;
 		}
-		if (NiagaraViewport->GetDrawElement(SNiagaraSystemViewport::EDrawElements::GpuTickInformation) && Component)
+		if (NiagaraViewport->GetDrawElement(SNiagaraSystemViewport::EDrawElements::MemoryInfo) && Component)
 		{
-			DrawGpuTickInformation(Component, Canvas, CurrentX, CurrentY, Font, FontHeight);
+			DrawMemoryInfo(Component, Canvas, CurrentX, CurrentY, Font, FontHeight);
+			CurrentY += FontHeight;
+		}
+		if (NiagaraViewport->GetDrawElement(SNiagaraSystemViewport::EDrawElements::ParticleCounts) && Component)
+		{
+			DrawParticleCounts(Component, Canvas, CurrentX, CurrentY, Font, FontHeight);
 			CurrentY += FontHeight;
 		}
 	}
@@ -406,6 +417,55 @@ void FNiagaraSystemViewportClient::DrawGpuTickInformation(UNiagaraComponent* Com
 	}
 }
 
+void FNiagaraSystemViewportClient::DrawMemoryInfo(UNiagaraComponent* Component, FCanvas* Canvas, float& CurrentX, float& CurrentY, UFont* Font, const float FontHeight)
+{
+	// Update costs
+	const double CurrentTime = FPlatformTime::Seconds();
+	const double RefreshTime = 2.0;
+	if ((CurrentTime - MemReportLastCheckedTime) > RefreshTime)
+	{
+		MemReportLastCheckedTime = CurrentTime;
+		MemReport.GenerateReport(FNiagaraSystemMemReport::EReportType::Basic, Component->GetAsset());
+	}
+
+	FResourceSizeEx ComponentResSize = FResourceSizeEx(EResourceSizeMode::EstimatedTotal);
+	Component->GetResourceSizeEx(ComponentResSize);
+	const uint64 MemComponentCost = ComponentResSize.GetTotalMemoryBytes();
+
+	// Title
+	Canvas->DrawShadowedString(CurrentX, CurrentY, TEXT("Memory Info (Approximate)"), Font, FLinearColor::White);
+	CurrentY += FontHeight;
+	CurrentX += 8.0f;
+
+	// System Mem Information
+	for (const FNiagaraSystemMemReport::FNode& Node : MemReport.GetNodes())
+	{
+		Canvas->DrawShadowedString(
+			CurrentX + (Node.Depth * 8.0f), CurrentY,
+			*FString::Printf(TEXT("- %s %s = %skb (%s bytes)"), *Node.ObjectClass.ToString(), *Node.ObjectName.ToString(), *FString::FormatAsNumber(FMath::DivideAndRoundUp(uint32(Node.InclusiveSizeBytes), 1024u)), *FString::FormatAsNumber(Node.InclusiveSizeBytes)),
+			Font, FLinearColor::White
+		);
+		CurrentY += FontHeight;
+	}
+
+	// Data Interface Mem
+	const uint32 DIMemory = uint32(MemReport.GetDataInterfaceSizeBytes());
+	Canvas->DrawShadowedString(
+		CurrentX, CurrentY,
+		*FString::Printf(TEXT("- DataInterface = %skb (%s bytes) - Included In NiagaraSystem Mem"), *FString::FormatAsNumber(FMath::DivideAndRoundUp(DIMemory, 1024u)), *FString::FormatAsNumber(DIMemory)),
+		Font, FLinearColor::White
+	);
+	CurrentY += FontHeight;
+
+	// Component Mem Information
+	Canvas->DrawShadowedString(
+		CurrentX, CurrentY,
+		*FString::Printf(TEXT("- Component = %skb (%s bytes)"), *FString::FormatAsNumber(FMath::DivideAndRoundUp(uint32(MemComponentCost), 1024u)), *FString::FormatAsNumber(uint32(MemComponentCost))),
+		Font, FLinearColor::White
+	);
+	CurrentY += FontHeight;
+}
+
 void FNiagaraSystemViewportClient::SetOrbitModeFromSettings()
 {
 	const UNiagaraSettings* Settings = GetDefault<UNiagaraSettings>();
@@ -484,6 +544,7 @@ void SNiagaraSystemViewport::Construct(const FArguments& InArgs)
 	DrawFlags |= Settings->IsShowInstructionsCount() ? EDrawElements::InstructionCounts : 0;
 	DrawFlags |= Settings->IsShowEmitterExecutionOrder() ? EDrawElements::EmitterExecutionOrder : 0;
 	DrawFlags |= Settings->IsShowGpuTickInformation() ? EDrawElements::GpuTickInformation : 0;
+	DrawFlags |= Settings->IsShowMemoryInfo() ? EDrawElements::MemoryInfo : 0;
 
 	bShowBackground = false;
 	PreviewComponent = nullptr;
@@ -749,6 +810,18 @@ void SNiagaraSystemViewport::BindCommands()
 		FIsActionChecked::CreateLambda([Viewport = this]() -> bool { return Viewport->GetDrawElement(EDrawElements::GpuTickInformation); })
 	);
 
+	CommandList->MapAction(
+		Commands.ToggleMemoryInfo,
+		FExecuteAction::CreateLambda([Viewport = this]()
+		{
+			Viewport->ToggleDrawElement(EDrawElements::MemoryInfo);
+			GetMutableDefault<UNiagaraEditorSettings>()->SetShowMemoryInfo(Viewport->GetDrawElement(EDrawElements::MemoryInfo));
+			Viewport->RefreshViewport();
+		}),
+		FCanExecuteAction(),
+		FIsActionChecked::CreateLambda([Viewport = this]() -> bool { return Viewport->GetDrawElement(EDrawElements::MemoryInfo); })
+	);
+
 	CommandList->MapAction(
 		Commands.TogglePreviewBackground,
 		FExecuteAction::CreateSP( this, &SNiagaraSystemViewport::TogglePreviewBackground ),
diff --git a/Engine/Plugins/FX/Niagara/Source/NiagaraEditor/Private/Widgets/SNiagaraSystemViewport.h b/Engine/Plugins/FX/Niagara/Source/NiagaraEditor/Private/Widgets/SNiagaraSystemViewport.h
--- a/Engine/Plugins/FX/Niagara/Source/NiagaraEditor/Private/Widgets/SNiagaraSystemViewport.h
+++ b/Engine/Plugins/FX/Niagara/Source/NiagaraEditor/Private/Widgets/SNiagaraSystemViewport.h
@@ -77,11 +77,12 @@ public:
 	/** Draw flag types */
 	enum EDrawElements
 	{
-		Bounds = 0x020,
-		InstructionCounts = 0x040,
-		ParticleCounts = 0x080,
-		EmitterExecutionOrder = 0x100,
-		GpuTickInformation = 0x200,
+		Bounds					= 0x020,
+		InstructionCounts		= 0x040,
+		ParticleCounts			= 0x080,
+		EmitterExecutionOrder	= 0x100,
+		GpuTickInformation		= 0x200,
+		MemoryInfo				= 0x400,
 	};
 
 	bool GetDrawElement(EDrawElements Element) const;
diff --git a/Engine/Plugins/FX/Niagara/Source/NiagaraEditor/Private/Widgets/SNiagaraSystemViewportToolBar.cpp b/Engine/Plugins/FX/Niagara/Source/NiagaraEditor/Private/Widgets/SNiagaraSystemViewportToolBar.cpp
--- a/Engine/Plugins/FX/Niagara/Source/NiagaraEditor/Private/Widgets/SNiagaraSystemViewportToolBar.cpp
+++ b/Engine/Plugins/FX/Niagara/Source/NiagaraEditor/Private/Widgets/SNiagaraSystemViewportToolBar.cpp
@@ -36,17 +36,13 @@ TSharedRef<SWidget> SNiagaraSystemViewportToolBar::GenerateShowMenu() const
 	{
 		auto Commands = FNiagaraEditorCommands::Get();
 
-		//ShowMenuBuilder.AddMenuEntry(Commands.ToggleMaterialStats);
-		//ShowMenuBuilder.AddMenuEntry(Commands.ToggleMobileStats);
-
-		ShowMenuBuilder.AddMenuSeparator();
-
 		ShowMenuBuilder.AddMenuEntry(Commands.TogglePreviewGrid);
-		ShowMenuBuilder.AddMenuEntry(Commands.ToggleInstructionCounts);
-		ShowMenuBuilder.AddMenuEntry(Commands.ToggleParticleCounts);
+		ShowMenuBuilder.AddMenuSeparator();
 		ShowMenuBuilder.AddMenuEntry(Commands.ToggleEmitterExecutionOrder);
 		ShowMenuBuilder.AddMenuEntry(Commands.ToggleGpuTickInformation);
-		//ShowMenuBuilder.AddMenuEntry(Commands.TogglePreviewBackground);
+		ShowMenuBuilder.AddMenuEntry(Commands.ToggleInstructionCounts);
+		ShowMenuBuilder.AddMenuEntry(Commands.ToggleMemoryInfo);
+		ShowMenuBuilder.AddMenuEntry(Commands.ToggleParticleCounts);
 	}
 
 	return ShowMenuBuilder.MakeWidget();
diff --git a/Engine/Plugins/FX/Niagara/Source/NiagaraEditor/Public/NiagaraEditorCommands.h b/Engine/Plugins/FX/Niagara/Source/NiagaraEditor/Public/NiagaraEditorCommands.h
--- a/Engine/Plugins/FX/Niagara/Source/NiagaraEditor/Public/NiagaraEditorCommands.h
+++ b/Engine/Plugins/FX/Niagara/Source/NiagaraEditor/Public/NiagaraEditorCommands.h
@@ -41,6 +41,7 @@ public:
 	TSharedPtr<FUICommandInfo> ToggleParticleCounts;
 	TSharedPtr<FUICommandInfo> ToggleEmitterExecutionOrder;
 	TSharedPtr<FUICommandInfo> ToggleGpuTickInformation;
+	TSharedPtr<FUICommandInfo> ToggleMemoryInfo;
 
 	/** Toggles the preview pane's background */
 	TSharedPtr< FUICommandInfo > TogglePreviewBackground;
diff --git a/Engine/Plugins/FX/Niagara/Source/NiagaraEditor/Public/NiagaraEditorSettings.h b/Engine/Plugins/FX/Niagara/Source/NiagaraEditor/Public/NiagaraEditorSettings.h
--- a/Engine/Plugins/FX/Niagara/Source/NiagaraEditor/Public/NiagaraEditorSettings.h
+++ b/Engine/Plugins/FX/Niagara/Source/NiagaraEditor/Public/NiagaraEditorSettings.h
@@ -496,6 +496,8 @@ private:
 	UPROPERTY(config)
 	bool bShowGpuTickInformation;
 
+	UPROPERTY(config)
+	bool bShowMemoryInfo;
 
 	UPROPERTY(config)
 	TArray<FNiagaraParameterPanelSectionStorage> SystemParameterPanelSectionData;
@@ -513,4 +515,6 @@ public:
 	void SetShowEmitterExecutionOrder(bool bShowEmitterExecutionOrder);
 	bool IsShowGpuTickInformation() const;
 	void SetShowGpuTickInformation(bool bShowGpuTickInformation);
+	bool IsShowMemoryInfo() const;
+	void SetShowMemoryInfo(bool bInShowInfo);
 };
