- Add missing pragma once

diff --git a/Engine/Plugins/FX/Niagara/Source/Niagara/Internal/NiagaraSystemMemReport.h b/Engine/Plugins/FX/Niagara/Source/Niagara/Internal/NiagaraSystemMemReport.h
--- a/Engine/Plugins/FX/Niagara/Source/Niagara/Internal/NiagaraSystemMemReport.h
+++ b/Engine/Plugins/FX/Niagara/Source/Niagara/Internal/NiagaraSystemMemReport.h
@@ -1,5 +1,7 @@
 // Copyright Epic Games, Inc. All Rights Reserved.
 
+#pragma once
+
 #include "CoreMinimal.h"
 
 class UObject;
