Change Niagara to Scene UB rather than exposing GPU-Scene internals in public interface.

diff --git a/Engine/Plugins/FX/Niagara/Source/Niagara/Private/NiagaraAsyncGpuTraceHelper.cpp b/Engine/Plugins/FX/Niagara/Source/Niagara/Private/NiagaraAsyncGpuTraceHelper.cpp
--- a/Engine/Plugins/FX/Niagara/Source/Niagara/Private/NiagaraAsyncGpuTraceHelper.cpp
+++ b/Engine/Plugins/FX/Niagara/Source/Niagara/Private/NiagaraAsyncGpuTraceHelper.cpp
@@ -119,7 +119,7 @@ void FNiagaraAsyncGpuTraceHelper::BeginFrame(FRHICommandList& RHICmdList, FNiaga
 	}
 }
 
-void FNiagaraAsyncGpuTraceHelper::PostRenderOpaque(FRHICommandList& RHICmdList, FNiagaraGpuComputeDispatchInterface* Dispatcher, TConstStridedView<FSceneView> Views)
+void FNiagaraAsyncGpuTraceHelper::PostRenderOpaque(FRHICommandList& RHICmdList, FNiagaraGpuComputeDispatchInterface* Dispatcher, TConstStridedView<FSceneView> Views, TUniformBufferRef<FSceneUniformParameters> SceneUniformBufferRHI)
 {
 #if NIAGARA_ASYNC_GPU_TRACE_COLLISION_GROUPS
 	if (bCollisionGroupMapDirty)
@@ -135,7 +135,7 @@ void FNiagaraAsyncGpuTraceHelper::PostRenderOpaque(FRHICommandList& RHICmdList,
 	{
 		if (TraceProvider->IsAvailable())
 		{
-			TraceProvider->PostRenderOpaque(RHICmdList, Views, CollisionHashMap);
+			TraceProvider->PostRenderOpaque(RHICmdList, Views, SceneUniformBufferRHI, CollisionHashMap);
 		}
 	}
 }
@@ -157,7 +157,7 @@ FNiagaraAsyncGpuTraceProvider* FNiagaraAsyncGpuTraceHelper::GetTraceProvider(END
 	return nullptr;
 }
 
-void FNiagaraAsyncGpuTraceHelper::EndFrame(FRHICommandList& RHICmdList, FNiagaraGpuComputeDispatchInterface* Dispatcher)
+void FNiagaraAsyncGpuTraceHelper::EndFrame(FRHICommandList& RHICmdList, FNiagaraGpuComputeDispatchInterface* Dispatcher, TUniformBufferRef<FSceneUniformParameters> SceneUniformBufferRHI)
 {
 	if (!Dispatches.Num())
 	{
@@ -187,7 +187,7 @@ void FNiagaraAsyncGpuTraceHelper::EndFrame(FRHICommandList& RHICmdList, FNiagara
 
 			if (FNiagaraAsyncGpuTraceProvider* TraceProvider = GetTraceProvider(DispatchInfo.ProviderType))
 			{
-				TraceProvider->IssueTraces(RHICmdList, Request, CollisionHashMap);
+				TraceProvider->IssueTraces(RHICmdList, Request, SceneUniformBufferRHI, CollisionHashMap);
 			}
 			else
 			{
diff --git a/Engine/Plugins/FX/Niagara/Source/Niagara/Private/NiagaraAsyncGpuTraceHelper.h b/Engine/Plugins/FX/Niagara/Source/Niagara/Private/NiagaraAsyncGpuTraceHelper.h
--- a/Engine/Plugins/FX/Niagara/Source/Niagara/Private/NiagaraAsyncGpuTraceHelper.h
+++ b/Engine/Plugins/FX/Niagara/Source/Niagara/Private/NiagaraAsyncGpuTraceHelper.h
@@ -65,8 +65,8 @@ public:
 	void Reset();
 
 	void BeginFrame(FRHICommandList& RHICmdList, FNiagaraGpuComputeDispatchInterface* Dispatcher);
-	void PostRenderOpaque(FRHICommandList& RHICmdList, FNiagaraGpuComputeDispatchInterface* Dispatcher, TConstStridedView<FSceneView> Views);
-	void EndFrame(FRHICommandList& RHICmdList, FNiagaraGpuComputeDispatchInterface* Dispatcher);
+	void PostRenderOpaque(FRHICommandList& RHICmdList, FNiagaraGpuComputeDispatchInterface* Dispatcher, TConstStridedView<FSceneView> Views, TUniformBufferRef<FSceneUniformParameters> SceneUniformBufferRHI);
+	void EndFrame(FRHICommandList& RHICmdList, FNiagaraGpuComputeDispatchInterface* Dispatcher, TUniformBufferRef<FSceneUniformParameters> SceneUniformBufferRHI);
 
 	/** Accumulates ray requests from user DIs into single dispatches per DI. */
 	void AddToDispatch(FNiagaraDataInterfaceProxy* DispatchKey, uint32 MaxRays, int32 MaxRetraces, ENDICollisionQuery_AsyncGpuTraceProvider::Type ProviderType);
diff --git a/Engine/Plugins/FX/Niagara/Source/Niagara/Private/NiagaraGpuComputeDispatch.cpp b/Engine/Plugins/FX/Niagara/Source/Niagara/Private/NiagaraGpuComputeDispatch.cpp
--- a/Engine/Plugins/FX/Niagara/Source/Niagara/Private/NiagaraGpuComputeDispatch.cpp
+++ b/Engine/Plugins/FX/Niagara/Source/Niagara/Private/NiagaraGpuComputeDispatch.cpp
@@ -446,7 +446,8 @@ void FNiagaraGpuComputeDispatch::ProcessPendingTicksFlush(FRHICommandListImmedia
 					}
 				);
 				PostInitViews(GraphBuilder, DummyViews, bAllowGPUParticleUpdate);
-				PostRenderOpaque(GraphBuilder, DummyViews, bAllowGPUParticleUpdate);
+				FSceneUniformBuffer &SceneUB = UE::FXRenderingUtils::CreateSceneUniformBuffer(GraphBuilder, GetSceneInterface());
+				PostRenderOpaque(GraphBuilder, DummyViews, SceneUB, bAllowGPUParticleUpdate);
 				AddPass(GraphBuilder, RDG_EVENT_NAME("UpdateDrawIndirectBuffers - PostOpaque"),
 					[this](FRHICommandListImmediate& RHICmdList)
 					{
@@ -1976,7 +1977,18 @@ void FNiagaraGpuComputeDispatch::PostInitViews(FRDGBuilder& GraphBuilder, TConst
 	}
 }
 
-void FNiagaraGpuComputeDispatch::PostRenderOpaque(FRDGBuilder& GraphBuilder, TConstStridedView<FSceneView> Views, bool bAllowGPUParticleUpdate)
+BEGIN_SHADER_PARAMETER_STRUCT(FNiagaraGpuComputeDispatchParameters,)
+	SHADER_PARAMETER_RDG_UNIFORM_BUFFER(FSceneUniformParameters, Scene)
+END_SHADER_PARAMETER_STRUCT()
+
+static FNiagaraGpuComputeDispatchParameters *MakeParameters(FRDGBuilder& GraphBuilder, TRDGUniformBufferRef<FSceneUniformParameters> SceneUniformBufferRDG)
+{
+	FNiagaraGpuComputeDispatchParameters *Parameters = GraphBuilder.AllocParameters<FNiagaraGpuComputeDispatchParameters>();
+	Parameters->Scene = SceneUniformBufferRDG;
+	return Parameters;
+}
+
+void FNiagaraGpuComputeDispatch::PostRenderOpaque(FRDGBuilder& GraphBuilder, TConstStridedView<FSceneView> Views, FSceneUniformBuffer &SceneUniformBuffer, bool bAllowGPUParticleUpdate)
 {
 	LLM_SCOPE(ELLMTag::Niagara);
 
@@ -1997,45 +2009,52 @@ void FNiagaraGpuComputeDispatch::PostRenderOpaque(FRDGBuilder& GraphBuilder, TCo
 		CachedViewInitOptions.ProjectionMatrix		= Views[0].SceneViewInitOptions.ProjectionMatrix;
 	}
 
+	TRDGUniformBufferRef<FSceneUniformParameters> SceneUniformBufferRDG = UE::FXRenderingUtils::GetSceneUniformBuffer(GraphBuilder, SceneUniformBuffer);
+
 	if (bAllowGPUParticleUpdate && FNiagaraUtilities::AllowGPUParticles(GetShaderPlatform()))
 	{
-		AddPass(
-			GraphBuilder,
+		auto PassParameters = MakeParameters(GraphBuilder, SceneUniformBufferRDG);
+		GraphBuilder.AddPass(
 			RDG_EVENT_NAME("AsyncGpuTraceHelper::PostRenderOpaque"),
-			[this, Views](FRHICommandListImmediate& RHICmdList)
+			PassParameters,
+			ERDGPassFlags::Compute | ERDGPassFlags::NeverCull,
+			[this, Views, PassParameters](FRHICommandListImmediate& RHICmdList)
 			{
-				AsyncGpuTraceHelper->PostRenderOpaque(RHICmdList, this, Views);
+				AsyncGpuTraceHelper->PostRenderOpaque(RHICmdList, this, Views, PassParameters->Scene->GetRHIRef());
 			}
 		);
 
 		ExecuteTicks(GraphBuilder, Views, ENiagaraGpuComputeTickStage::PostOpaqueRender);
 	}
-
-	AddPass(
-		GraphBuilder,
-		RDG_EVENT_NAME("Niagara::PostRenderFinish"),
-		[this, bExecuteReadback=bRequiresReadback, bAllowGPUParticleUpdate](FRHICommandListImmediate& RHICmdList)
-		{
-			if (bAllowGPUParticleUpdate)
+	{
+		auto PassParameters = MakeParameters(GraphBuilder, SceneUniformBufferRDG);
+		GraphBuilder.AddPass(
+			RDG_EVENT_NAME("Niagara::PostRenderFinish"),
+			PassParameters,
+			ERDGPassFlags::Compute | ERDGPassFlags::NeverCull,
+			[this, bExecuteReadback=bRequiresReadback, PassParameters, bAllowGPUParticleUpdate](FRHICommandListImmediate& RHICmdList)
 			{
-				FinishDispatches();
+				if (bAllowGPUParticleUpdate)
+				{
+					FinishDispatches();
 
-				AsyncGpuTraceHelper->EndFrame(RHICmdList, this);
-			}
+					AsyncGpuTraceHelper->EndFrame(RHICmdList, this, PassParameters->Scene->GetRHIRef());
+				}
 
-			ProcessDebugReadbacks(RHICmdList, false);
+				ProcessDebugReadbacks(RHICmdList, false);
 
-			if (bExecuteReadback)
-			{
-				check(!GPUInstanceCounterManager.HasPendingGPUReadback());
-				GPUInstanceCounterManager.EnqueueGPUReadback(RHICmdList);
-			}
+				if (bExecuteReadback)
+				{
+					check(!GPUInstanceCounterManager.HasPendingGPUReadback());
+					GPUInstanceCounterManager.EnqueueGPUReadback(RHICmdList);
+				}
 
-		#if WITH_NIAGARA_GPU_PROFILER
-			GPUProfilerPtr->EndFrame(RHICmdList);
-		#endif
-		}
-	);
+			#if WITH_NIAGARA_GPU_PROFILER
+				GPUProfilerPtr->EndFrame(RHICmdList);
+			#endif
+			}
+		);
+	}
 	bRequiresReadback = false;
 
 	if (FNiagaraGpuComputeDispatchLocal::CsvStatsEnabled())
@@ -2126,7 +2145,7 @@ bool FNiagaraGpuComputeDispatch::RequiresRayTracingScene() const
 	return NumProxiesThatRequireRayTracingScene > 0;
 }
 
-void FNiagaraGpuComputeDispatch::PreRender(FRDGBuilder& GraphBuilder, TConstStridedView<FSceneView> Views, bool bAllowGPUParticleUpdate)
+void FNiagaraGpuComputeDispatch::PreRender(FRDGBuilder& GraphBuilder, TConstStridedView<FSceneView> Views, FSceneUniformBuffer &SceneUniformBuffer, bool bAllowGPUParticleUpdate)
 {
 	if (!FNiagaraUtilities::AllowGPUParticles(GetShaderPlatform()))
 	{
diff --git a/Engine/Plugins/FX/Niagara/Source/Niagara/Private/NiagaraGpuComputeDispatch.h b/Engine/Plugins/FX/Niagara/Source/Niagara/Private/NiagaraGpuComputeDispatch.h
--- a/Engine/Plugins/FX/Niagara/Source/Niagara/Private/NiagaraGpuComputeDispatch.h
+++ b/Engine/Plugins/FX/Niagara/Source/Niagara/Private/NiagaraGpuComputeDispatch.h
@@ -115,12 +115,12 @@ public:
 	virtual bool UsesDepthBuffer() const override;
 	virtual bool RequiresEarlyViewUniformBuffer() const override;
 	virtual bool RequiresRayTracingScene() const override;
-	virtual void PreRender(FRDGBuilder& GraphBuilder, TConstStridedView<FSceneView> Views, bool bAllowGPUParticleUpdate) override;
+	virtual void PreRender(FRDGBuilder& GraphBuilder, TConstStridedView<FSceneView> Views, FSceneUniformBuffer &SceneUniformBuffer, bool bAllowGPUParticleUpdate) override;
 	virtual void OnDestroy() override; // Called on the gamethread to delete the batcher on the renderthread.
 
 	virtual void Tick(UWorld* World, float DeltaTime) override;
 
-	virtual void PostRenderOpaque(FRDGBuilder& GraphBuilder, TConstStridedView<FSceneView> Views, bool bAllowGPUParticleUpdate) override;
+	virtual void PostRenderOpaque(FRDGBuilder& GraphBuilder, TConstStridedView<FSceneView> Views, FSceneUniformBuffer &SceneUniformBuffer, bool bAllowGPUParticleUpdate) override;
 
 	// FNiagaraGpuComputeDispatchInterface Impl
 	virtual void FlushPendingTicks_GameThread() override;
diff --git a/Engine/Plugins/FX/Niagara/Source/NiagaraShader/Private/NiagaraAsyncGpuTraceProvider.cpp b/Engine/Plugins/FX/Niagara/Source/NiagaraShader/Private/NiagaraAsyncGpuTraceProvider.cpp
--- a/Engine/Plugins/FX/Niagara/Source/NiagaraShader/Private/NiagaraAsyncGpuTraceProvider.cpp
+++ b/Engine/Plugins/FX/Niagara/Source/NiagaraShader/Private/NiagaraAsyncGpuTraceProvider.cpp
@@ -282,6 +282,6 @@ void FNiagaraAsyncGpuTraceProvider::BuildCollisionGroupHashMap(FRHICommandList&
 	RHICmdList.Transition(FRHITransitionInfo(Result.HashToCollisionGroups.UAV, ERHIAccess::UAVCompute, ERHIAccess::SRVCompute));
 }
 
-void FNiagaraAsyncGpuTraceProvider::PostRenderOpaque(FRHICommandList& RHICmdList, TConstStridedView<FSceneView> Views, FCollisionGroupHashMap* CollisionGroupHash)
+void FNiagaraAsyncGpuTraceProvider::PostRenderOpaque(FRHICommandList& RHICmdList, TConstStridedView<FSceneView> Views, TUniformBufferRef<FSceneUniformParameters> SceneUniformBufferRHI, FCollisionGroupHashMap* CollisionGroupHash)
 {
 }
diff --git a/Engine/Plugins/FX/Niagara/Source/NiagaraShader/Private/NiagaraAsyncGpuTraceProviderGsdf.cpp b/Engine/Plugins/FX/Niagara/Source/NiagaraShader/Private/NiagaraAsyncGpuTraceProviderGsdf.cpp
--- a/Engine/Plugins/FX/Niagara/Source/NiagaraShader/Private/NiagaraAsyncGpuTraceProviderGsdf.cpp
+++ b/Engine/Plugins/FX/Niagara/Source/NiagaraShader/Private/NiagaraAsyncGpuTraceProviderGsdf.cpp
@@ -107,7 +107,7 @@ bool FNiagaraAsyncGpuTraceProviderGsdf::IsAvailable() const
 	return true;
 }
 
-void FNiagaraAsyncGpuTraceProviderGsdf::PostRenderOpaque(FRHICommandList& RHICmdList, TConstStridedView<FSceneView> Views, FCollisionGroupHashMap* CollisionGroupHash)
+void FNiagaraAsyncGpuTraceProviderGsdf::PostRenderOpaque(FRHICommandList& RHICmdList, TConstStridedView<FSceneView> Views, TUniformBufferRef<FSceneUniformParameters> SceneUniformBufferRHI, FCollisionGroupHashMap* CollisionGroupHash)
 {
 	if (const FGlobalDistanceFieldParameterData* DistanceFieldData = UE::FXRenderingUtils::GetGlobalDistanceFieldParameterData(Views))
 	{
@@ -120,7 +120,7 @@ void FNiagaraAsyncGpuTraceProviderGsdf::PostRenderOpaque(FRHICommandList& RHICmd
 	m_ViewUniformBuffer = Views[0].ViewUniformBuffer;
 }
 
-void FNiagaraAsyncGpuTraceProviderGsdf::IssueTraces(FRHICommandList& RHICmdList, const FDispatchRequest& Request, FCollisionGroupHashMap* CollisionHashMap)
+void FNiagaraAsyncGpuTraceProviderGsdf::IssueTraces(FRHICommandList& RHICmdList, const FDispatchRequest& Request, TUniformBufferRef<FSceneUniformParameters> SceneUniformBufferRHI, FCollisionGroupHashMap* CollisionHashMap)
 {
 	check(IsAvailable());
 
diff --git a/Engine/Plugins/FX/Niagara/Source/NiagaraShader/Private/NiagaraAsyncGpuTraceProviderGsdf.h b/Engine/Plugins/FX/Niagara/Source/NiagaraShader/Private/NiagaraAsyncGpuTraceProviderGsdf.h
--- a/Engine/Plugins/FX/Niagara/Source/NiagaraShader/Private/NiagaraAsyncGpuTraceProviderGsdf.h
+++ b/Engine/Plugins/FX/Niagara/Source/NiagaraShader/Private/NiagaraAsyncGpuTraceProviderGsdf.h
@@ -15,9 +15,9 @@ public:
 	FNiagaraAsyncGpuTraceProviderGsdf(EShaderPlatform InShaderPlatform, FNiagaraGpuComputeDispatchInterface* Dispatcher);
 
 	static bool IsSupported();
-	virtual void PostRenderOpaque(FRHICommandList& RHICmdList, TConstStridedView<FSceneView> Views, FCollisionGroupHashMap* CollisionGroupHash) override;
+	virtual void PostRenderOpaque(FRHICommandList& RHICmdList, TConstStridedView<FSceneView> Views, TUniformBufferRef<FSceneUniformParameters> SceneUniformBufferRHI, FCollisionGroupHashMap* CollisionGroupHash) override;
 	virtual bool IsAvailable() const override;
-	virtual void IssueTraces(FRHICommandList& RHICmdList, const FDispatchRequest& Request, FCollisionGroupHashMap* CollisionGroupHash) override;
+	virtual void IssueTraces(FRHICommandList& RHICmdList, const FDispatchRequest& Request, TUniformBufferRef<FSceneUniformParameters> SceneUniformBufferRHI, FCollisionGroupHashMap* CollisionGroupHash) override;
 	virtual void Reset() override;
 	virtual EProviderType GetType() const override { return Type; }
 
diff --git a/Engine/Plugins/FX/Niagara/Source/NiagaraShader/Private/NiagaraAsyncGpuTraceProviderHwrt.cpp b/Engine/Plugins/FX/Niagara/Source/NiagaraShader/Private/NiagaraAsyncGpuTraceProviderHwrt.cpp
--- a/Engine/Plugins/FX/Niagara/Source/NiagaraShader/Private/NiagaraAsyncGpuTraceProviderHwrt.cpp
+++ b/Engine/Plugins/FX/Niagara/Source/NiagaraShader/Private/NiagaraAsyncGpuTraceProviderHwrt.cpp
@@ -40,12 +40,6 @@ struct FVFXTracePayload
 
 IMPLEMENT_RT_PAYLOAD_TYPE(ERayTracingPayloadType::VFX, sizeof(FVFXTracePayload));
 
-BEGIN_SHADER_PARAMETER_STRUCT(FNiagaraGPUSceneParameters, )
-	SHADER_PARAMETER_SRV(StructuredBuffer<float4>, GPUSceneInstanceSceneData)
-	SHADER_PARAMETER_SRV(StructuredBuffer<float4>, GPUSceneInstancePayloadData)
-	SHADER_PARAMETER_SRV(StructuredBuffer<float4>, GPUScenePrimitiveSceneData)
-	SHADER_PARAMETER(uint32, GPUSceneFrameNumber)
-END_SHADER_PARAMETER_STRUCT()
 
 class FNiagaraCollisionRayTraceRG : public FGlobalShader
 {
@@ -54,7 +48,7 @@ class FNiagaraCollisionRayTraceRG : public FGlobalShader
 
 	BEGIN_SHADER_PARAMETER_STRUCT(FParameters, )
 		SHADER_PARAMETER_STRUCT_REF(FViewUniformShaderParameters, View)
-		SHADER_PARAMETER_STRUCT_INCLUDE(FNiagaraGPUSceneParameters, GPUSceneParameters)
+		SHADER_PARAMETER_STRUCT_REF(FSceneUniformParameters, Scene)
 
 		SHADER_PARAMETER_SRV(Buffer<UINT>, HashTable)
 		SHADER_PARAMETER_SRV(Buffer<UINT>, HashToCollisionGroups)
@@ -125,7 +119,6 @@ void FNiagaraCollisionRayTraceRG::ModifyCompilationEnvironment(const FGlobalShad
 	FGlobalShader::ModifyCompilationEnvironment(Parameters, OutEnvironment);
 	OutEnvironment.SetDefine(TEXT("NIAGARA_SUPPORTS_RAY_TRACING"), 1);
 	OutEnvironment.SetDefine(TEXT("VF_SUPPORTS_PRIMITIVE_SCENE_DATA"), 1);
-	OutEnvironment.SetDefine(TEXT("USE_GLOBAL_GPU_SCENE_DATA"), 1);
 }
 
 TShaderRef<FNiagaraCollisionRayTraceRG> FNiagaraCollisionRayTraceRG::GetShader(FGlobalShaderMap* ShaderMap, bool SupportsCollisionGroups)
@@ -309,7 +302,7 @@ bool FNiagaraAsyncGpuTraceProviderHwrt::IsAvailable() const
 	return UE::FXRenderingUtils::RayTracing::HasRayTracingScene(Dispatcher->GetSceneInterface());
 }
 
-void FNiagaraAsyncGpuTraceProviderHwrt::PostRenderOpaque(FRHICommandList& RHICmdList, TConstStridedView<FSceneView> Views, FCollisionGroupHashMap* CollisionGroupHash)
+void FNiagaraAsyncGpuTraceProviderHwrt::PostRenderOpaque(FRHICommandList& RHICmdList, TConstStridedView<FSceneView> Views, TUniformBufferRef<FSceneUniformParameters> SceneUniformBufferRHI, FCollisionGroupHashMap* CollisionGroupHash)
 {
 	check(IsAvailable());
 	check(Views.Num() > 0);
@@ -354,7 +347,7 @@ void FNiagaraAsyncGpuTraceProviderHwrt::PostRenderOpaque(FRHICommandList& RHICmd
 	}
 }
 
-void FNiagaraAsyncGpuTraceProviderHwrt::IssueTraces(FRHICommandList& RHICmdList, const FDispatchRequest& Request, FCollisionGroupHashMap* CollisionGroupHash)
+void FNiagaraAsyncGpuTraceProviderHwrt::IssueTraces(FRHICommandList& RHICmdList, const FDispatchRequest& Request, TUniformBufferRef<FSceneUniformParameters> SceneUniformBufferRHI, FCollisionGroupHashMap* CollisionGroupHash)
 {
 	check(IsAvailable());
 	check(RayTracingPipelineState);
@@ -375,13 +368,7 @@ void FNiagaraAsyncGpuTraceProviderHwrt::IssueTraces(FRHICommandList& RHICmdList,
 	FNiagaraCollisionRayTraceRG::FParameters Params;
 
 	Params.View = GetShaderBinding(ViewUniformBuffer);
-
-	UE::FXRenderingUtils::GPUScene::FBuffers GPUSceneBuffers = UE::FXRenderingUtils::GPUScene::GetBuffers(Dispatcher->GetSceneInterface());
-
-	Params.GPUSceneParameters.GPUSceneInstanceSceneData = GPUSceneBuffers.InstanceSceneDataBuffer;
-	Params.GPUSceneParameters.GPUSceneInstancePayloadData = GPUSceneBuffers.InstancePayloadDataBuffer;
-	Params.GPUSceneParameters.GPUScenePrimitiveSceneData = GPUSceneBuffers.PrimitiveBuffer;
-	Params.GPUSceneParameters.GPUSceneFrameNumber = GPUSceneBuffers.SceneFrameNumber;
+	Params.Scene = SceneUniformBufferRHI;
 
 	if (CollisionGroupHash)
 	{
diff --git a/Engine/Plugins/FX/Niagara/Source/NiagaraShader/Private/NiagaraAsyncGpuTraceProviderHwrt.h b/Engine/Plugins/FX/Niagara/Source/NiagaraShader/Private/NiagaraAsyncGpuTraceProviderHwrt.h
--- a/Engine/Plugins/FX/Niagara/Source/NiagaraShader/Private/NiagaraAsyncGpuTraceProviderHwrt.h
+++ b/Engine/Plugins/FX/Niagara/Source/NiagaraShader/Private/NiagaraAsyncGpuTraceProviderHwrt.h
@@ -17,9 +17,9 @@ public:
 	FNiagaraAsyncGpuTraceProviderHwrt(EShaderPlatform InShaderPlatform, FNiagaraGpuComputeDispatchInterface* Dispatcher);
 
 	static bool IsSupported();
-	virtual void PostRenderOpaque(FRHICommandList& RHICmdList, TConstStridedView<FSceneView> Views, FCollisionGroupHashMap* CollisionGroupHash) override;
+	virtual void PostRenderOpaque(FRHICommandList& RHICmdList, TConstStridedView<FSceneView> Views, TUniformBufferRef<FSceneUniformParameters> SceneUniformBufferRHI, FCollisionGroupHashMap* CollisionGroupHash) override;
 	virtual bool IsAvailable() const override;
-	virtual void IssueTraces(FRHICommandList& RHICmdList, const FDispatchRequest& Request, FCollisionGroupHashMap* CollisionGroupHash) override;
+	virtual void IssueTraces(FRHICommandList& RHICmdList, const FDispatchRequest& Request, TUniformBufferRef<FSceneUniformParameters> SceneUniformBufferRHI, FCollisionGroupHashMap* CollisionGroupHash) override;
 	virtual void Reset() override;
 	virtual EProviderType GetType() const override { return Type; }
 
diff --git a/Engine/Plugins/FX/Niagara/Source/NiagaraShader/Public/NiagaraAsyncGpuTraceProvider.h b/Engine/Plugins/FX/Niagara/Source/NiagaraShader/Public/NiagaraAsyncGpuTraceProvider.h
--- a/Engine/Plugins/FX/Niagara/Source/NiagaraShader/Public/NiagaraAsyncGpuTraceProvider.h
+++ b/Engine/Plugins/FX/Niagara/Source/NiagaraShader/Public/NiagaraAsyncGpuTraceProvider.h
@@ -3,6 +3,7 @@
 
 #include "RHIDefinitions.h"
 #include "RHIUtilities.h"
+#include "FXRenderingUtils.h"
 
 class FPrimitiveComponentId;
 class FSceneInterface;
@@ -79,9 +80,9 @@ public:
 	virtual bool IsAvailable() const = 0;
 	virtual EProviderType GetType() const = 0;
 
-	virtual void PostRenderOpaque(FRHICommandList& RHICmdList, TConstStridedView<FSceneView> Views, FCollisionGroupHashMap* CollisionGroupHash);
+	virtual void PostRenderOpaque(FRHICommandList& RHICmdList, TConstStridedView<FSceneView> Views, TUniformBufferRef<FSceneUniformParameters> SceneUniformBufferRHI, FCollisionGroupHashMap* CollisionGroupHash);
 
-	virtual void IssueTraces(FRHICommandList& RHICmdList, const FDispatchRequest& Request, FCollisionGroupHashMap* CollisionGroupHash)
+	virtual void IssueTraces(FRHICommandList& RHICmdList, const FDispatchRequest& Request, TUniformBufferRef<FSceneUniformParameters> SceneUniformBufferRHI, FCollisionGroupHashMap* CollisionGroupHash)
 	{
 	}
 
diff --git a/Engine/Source/Runtime/Engine/Private/Particles/FXSystem.cpp b/Engine/Source/Runtime/Engine/Private/Particles/FXSystem.cpp
--- a/Engine/Source/Runtime/Engine/Private/Particles/FXSystem.cpp
+++ b/Engine/Source/Runtime/Engine/Private/Particles/FXSystem.cpp
@@ -486,7 +486,7 @@ DECLARE_CYCLE_STAT(TEXT("FXPreRender_FinalizeCDF"), STAT_CLM_FXPreRender_Finaliz
 DECLARE_GPU_DRAWCALL_STAT(FXSystemPreRender);
 DECLARE_GPU_DRAWCALL_STAT(FXSystemPostRenderOpaque);
 
-void FFXSystem::PreRender(FRDGBuilder& GraphBuilder, TConstStridedView<FSceneView> Views, bool bAllowGPUParticleSceneUpdate)
+void FFXSystem::PreRender(FRDGBuilder& GraphBuilder, TConstStridedView<FSceneView> Views, FSceneUniformBuffer &SceneUniformBuffer, bool bAllowGPUParticleSceneUpdate)
 {
 	bAllowGPUParticleSceneUpdate = bAllowGPUParticleSceneUpdate && Views.Num() > 0 && Views[0].AllowGPUParticleUpdate();
 
@@ -533,7 +533,7 @@ void FFXSystem::PreRender(FRDGBuilder& GraphBuilder, TConstStridedView<FSceneVie
     }
 }
 
-void FFXSystem::PostRenderOpaque(FRDGBuilder& GraphBuilder, TConstStridedView<FSceneView> Views, bool bAllowGPUParticleUpdate)
+void FFXSystem::PostRenderOpaque(FRDGBuilder& GraphBuilder, TConstStridedView<FSceneView> Views, FSceneUniformBuffer &SceneUniformBuffer, bool bAllowGPUParticleUpdate)
 {
 	bAllowGPUParticleUpdate = bAllowGPUParticleUpdate && Views.Num() > 0 && Views[0].AllowGPUParticleUpdate();
 
diff --git a/Engine/Source/Runtime/Engine/Private/Particles/FXSystemPrivate.h b/Engine/Source/Runtime/Engine/Private/Particles/FXSystemPrivate.h
--- a/Engine/Source/Runtime/Engine/Private/Particles/FXSystemPrivate.h
+++ b/Engine/Source/Runtime/Engine/Private/Particles/FXSystemPrivate.h
@@ -208,8 +208,8 @@ public:
 	virtual bool UsesDepthBuffer() const override;
 	virtual bool RequiresEarlyViewUniformBuffer() const override;
 	virtual bool RequiresRayTracingScene() const override;
-	virtual void PreRender(FRDGBuilder& GraphBuilder, TConstStridedView<FSceneView> Views, bool bAllowGPUParticleUpdate) override;
-	virtual void PostRenderOpaque(FRDGBuilder& GraphBuilder, TConstStridedView<FSceneView> Views, bool bAllowGPUParticleUpdate) override;
+	virtual void PreRender(FRDGBuilder& GraphBuilder, TConstStridedView<FSceneView> Views, FSceneUniformBuffer &SceneUniformBuffer, bool bAllowGPUParticleUpdate) override;
+	virtual void PostRenderOpaque(FRDGBuilder& GraphBuilder, TConstStridedView<FSceneView> Views, FSceneUniformBuffer &SceneUniformBuffer, bool bAllowGPUParticleUpdate) override;
 	// End FFXSystemInterface.
 
 	/*--------------------------------------------------------------------------
diff --git a/Engine/Source/Runtime/Engine/Private/Particles/FXSystemSet.cpp b/Engine/Source/Runtime/Engine/Private/Particles/FXSystemSet.cpp
--- a/Engine/Source/Runtime/Engine/Private/Particles/FXSystemSet.cpp
+++ b/Engine/Source/Runtime/Engine/Private/Particles/FXSystemSet.cpp
@@ -196,12 +196,12 @@ bool FFXSystemSet::RequiresRayTracingScene() const
 	return false;
 }
 
-void FFXSystemSet::PreRender(FRDGBuilder& GraphBuilder, TConstStridedView<FSceneView> Views, bool bAllowGPUParticleSceneUpdate)
+void FFXSystemSet::PreRender(FRDGBuilder& GraphBuilder, TConstStridedView<FSceneView> Views, FSceneUniformBuffer &SceneUniformBuffer, bool bAllowGPUParticleSceneUpdate)
 {
 	for (FFXSystemInterface* FXSystem : FXSystems)
 	{
 		check(FXSystem);
-		FXSystem->PreRender(GraphBuilder, Views, bAllowGPUParticleSceneUpdate);
+		FXSystem->PreRender(GraphBuilder, Views, SceneUniformBuffer, bAllowGPUParticleSceneUpdate);
 	}
 }
 
@@ -214,12 +214,12 @@ void FFXSystemSet::SetSceneTexturesUniformBuffer(const TUniformBufferRef<FSceneT
 	}
 }
 
-void FFXSystemSet::PostRenderOpaque(FRDGBuilder& GraphBuilder, TConstStridedView<FSceneView> Views, bool bAllowGPUParticleSceneUpdate)
+void FFXSystemSet::PostRenderOpaque(FRDGBuilder& GraphBuilder, TConstStridedView<FSceneView> Views, FSceneUniformBuffer &SceneUniformBuffer, bool bAllowGPUParticleSceneUpdate)
 {
 	for (FFXSystemInterface* FXSystem : FXSystems)
 	{
 		check(FXSystem);
-		FXSystem->PostRenderOpaque(GraphBuilder, Views, bAllowGPUParticleSceneUpdate);
+		FXSystem->PostRenderOpaque(GraphBuilder, Views, SceneUniformBuffer, bAllowGPUParticleSceneUpdate);
 	}
 }
 
diff --git a/Engine/Source/Runtime/Engine/Private/Particles/FXSystemSet.h b/Engine/Source/Runtime/Engine/Private/Particles/FXSystemSet.h
--- a/Engine/Source/Runtime/Engine/Private/Particles/FXSystemSet.h
+++ b/Engine/Source/Runtime/Engine/Private/Particles/FXSystemSet.h
@@ -45,8 +45,8 @@ public:
 	virtual bool UsesDepthBuffer() const override;
 	virtual bool RequiresEarlyViewUniformBuffer() const override;
 	virtual bool RequiresRayTracingScene() const override;
-	virtual void PreRender(FRDGBuilder& GraphBuilder, TConstStridedView<FSceneView> Views, bool bAllowGPUParticleSceneUpdate) override;
-	virtual void PostRenderOpaque(FRDGBuilder& GraphBuilder, TConstStridedView<FSceneView> Views, bool bAllowGPUParticleSceneUpdate) override;
+	virtual void PreRender(FRDGBuilder& GraphBuilder, TConstStridedView<FSceneView> Views, FSceneUniformBuffer &SceneUniformBuffer, bool bAllowGPUParticleSceneUpdate) override;
+	virtual void PostRenderOpaque(FRDGBuilder& GraphBuilder, TConstStridedView<FSceneView> Views, FSceneUniformBuffer &SceneUniformBuffer, bool bAllowGPUParticleSceneUpdate) override;
 
 	virtual void OnDestroy() override;
 	virtual void DestroyGPUSimulation() override;
diff --git a/Engine/Source/Runtime/Engine/Public/FXSystem.h b/Engine/Source/Runtime/Engine/Public/FXSystem.h
--- a/Engine/Source/Runtime/Engine/Public/FXSystem.h
+++ b/Engine/Source/Runtime/Engine/Public/FXSystem.h
@@ -23,7 +23,7 @@ class FScene;
 class FSceneInterface;
 class UVectorFieldComponent;
 class FGlobalDistanceFieldParameterData;
-
+class FSceneUniformBuffer;
 /*-----------------------------------------------------------------------------
 	Forward declarations.
 -----------------------------------------------------------------------------*/
@@ -261,7 +261,7 @@ public:
 	 * Notification from the renderer that it is about to draw FX belonging to
 	 * this system.
 	 */
-	virtual void PreRender(FRDGBuilder& GraphBuilder, TConstStridedView<FSceneView> Views, bool bAllowGPUParticleUpdate) = 0;
+	virtual void PreRender(FRDGBuilder& GraphBuilder, TConstStridedView<FSceneView> Views, FSceneUniformBuffer &SceneUniformBuffer, bool bAllowGPUParticleUpdate) = 0;
 
 	UE_DEPRECATED(5.3, "Passing an array of FViewInfo is deprecated in favor of FSceneView. This function will do nothing.")
 	inline void PreRender(FRDGBuilder& GraphBuilder, TConstArrayView<FViewInfo> Views, bool bAllowGPUParticleUpdate) {}
@@ -269,7 +269,7 @@ public:
 	/**
 	 * Notification from the renderer that opaque primitives have rendered.
 	 */
-	virtual void PostRenderOpaque(FRDGBuilder& GraphBuilder, TConstStridedView<FSceneView> Views, bool bAllowGPUParticleUpdate) = 0;
+	virtual void PostRenderOpaque(FRDGBuilder& GraphBuilder, TConstStridedView<FSceneView> Views, FSceneUniformBuffer &SceneUniformBuffer, bool bAllowGPUParticleUpdate) = 0;
 
 	UE_DEPRECATED(5.3, "Passing an array of FViewInfo is deprecated in favor of FSceneView. This function will do nothing.")
 	inline void PostRenderOpaque(FRDGBuilder& GraphBuilder, TConstArrayView<FViewInfo> Views, bool bAllowGPUParticleUpdate) {}
diff --git a/Engine/Source/Runtime/Renderer/Private/DeferredShadingRenderer.cpp b/Engine/Source/Runtime/Renderer/Private/DeferredShadingRenderer.cpp
--- a/Engine/Source/Runtime/Renderer/Private/DeferredShadingRenderer.cpp
+++ b/Engine/Source/Runtime/Renderer/Private/DeferredShadingRenderer.cpp
@@ -636,6 +636,7 @@ END_SHADER_PARAMETER_STRUCT()
 static void RenderOpaqueFX(
 	FRDGBuilder& GraphBuilder,
 	TConstStridedView<FSceneView> Views,
+	FSceneUniformBuffer &SceneUniformBuffer,
 	FFXSystemInterface* FXSystem,
 	TRDGUniformBufferRef<FSceneTextureUniformParameters> SceneTexturesUniformBuffer)
 {
@@ -655,7 +656,7 @@ static void RenderOpaqueFX(
 			FXSystem->SetSceneTexturesUniformBuffer(ExtractUBPassParameters->SceneTextures->GetRHIRef());
 		});
 
-		FXSystem->PostRenderOpaque(GraphBuilder, Views, true /*bAllowGPUParticleUpdate*/);
+		FXSystem->PostRenderOpaque(GraphBuilder, Views, SceneUniformBuffer, true /*bAllowGPUParticleUpdate*/);
 
 		// Clear the scene textures UB pointer on the FX system. Use the same pass parameters to extend resource lifetimes.
 		GraphBuilder.AddPass(RDG_EVENT_NAME("UnsetSceneTexturesUniformBuffer"), ExtractUBPassParameters, UBPassFlags, [FXSystem](FRHICommandListImmediate&)
@@ -2868,7 +2869,7 @@ void FDeferredShadingSceneRenderer::Render(FRDGBuilder& GraphBuilder)
 		{
 			SCOPE_CYCLE_COUNTER(STAT_FDeferredShadingSceneRenderer_FXSystem_PreRender);
 			GraphBuilder.SetCommandListStat(GET_STATID(STAT_CLM_FXPreRender));
-			FXSystem->PreRender(GraphBuilder, GetSceneViews(), bIsFirstSceneRenderer /*bAllowGPUParticleUpdate*/);
+			FXSystem->PreRender(GraphBuilder, GetSceneViews(), GetSceneUniforms(), bIsFirstSceneRenderer /*bAllowGPUParticleUpdate*/);
 			if (FGPUSortManager* GPUSortManager = FXSystem->GetGPUSortManager())
 			{
 				GPUSortManager->OnPreRender(GraphBuilder);
@@ -4081,7 +4082,7 @@ void FDeferredShadingSceneRenderer::Render(FRDGBuilder& GraphBuilder)
 
 		FRDGTextureRef ExposureIlluminance = AddCalculateExposureIlluminancePass(GraphBuilder, Views, SceneTextures, TranslucencyLightingVolumeTextures, ExposureIlluminanceSetup);
 
-		RenderOpaqueFX(GraphBuilder, GetSceneViews(), FXSystem, SceneTextures.UniformBuffer);
+		RenderOpaqueFX(GraphBuilder, GetSceneViews(), GetSceneUniforms(), FXSystem, SceneTextures.UniformBuffer);
 
 		FRendererModule& RendererModule = static_cast<FRendererModule&>(GetRendererModule());
 		RendererModule.RenderPostOpaqueExtensions(GraphBuilder, Views, SceneTextures);
diff --git a/Engine/Source/Runtime/Renderer/Private/FXRenderingUtils.cpp b/Engine/Source/Runtime/Renderer/Private/FXRenderingUtils.cpp
--- a/Engine/Source/Runtime/Renderer/Private/FXRenderingUtils.cpp
+++ b/Engine/Source/Runtime/Renderer/Private/FXRenderingUtils.cpp
@@ -211,21 +211,21 @@ void UE::FXRenderingUtils::DistanceFields::SetupAtlasParameters(FRDGBuilder& Gra
 	}
 }
 
-UE::FXRenderingUtils::GPUScene::FBuffers UE::FXRenderingUtils::GPUScene::GetBuffers(const FSceneInterface* InScene)
+FSceneUniformBuffer &UE::FXRenderingUtils::CreateSceneUniformBuffer(FRDGBuilder& GraphBuilder, const FSceneInterface* InScene)
 {
-	UE::FXRenderingUtils::GPUScene::FBuffers Buffers{};
-
+	FSceneUniformBuffer *Result = GraphBuilder.AllocObject<FSceneUniformBuffer>();
 	if (const FScene* Scene = InScene->GetRenderScene())
 	{
-		Buffers.InstanceSceneDataBuffer = Scene->GPUScene.InstanceSceneDataBuffer->GetSRV();
-		Buffers.InstancePayloadDataBuffer = Scene->GPUScene.InstancePayloadDataBuffer->GetSRV();
-		Buffers.PrimitiveBuffer = Scene->GPUScene.PrimitiveBuffer->GetSRV();
-		Buffers.SceneFrameNumber = Scene->GPUScene.GetSceneFrameNumber();
+		FSceneUniformBuffer SceneUniformBuffer;
+		Scene->GPUScene.FillSceneUniformBuffer(GraphBuilder, SceneUniformBuffer);
 	}
-
-	return Buffers;
+	return *Result;
 }
 
+TRDGUniformBufferRef<FSceneUniformParameters> UE::FXRenderingUtils::GetSceneUniformBuffer(FRDGBuilder& GraphBuilder, FSceneUniformBuffer &SceneUniformBuffer)
+{
+	return SceneUniformBuffer.GetBuffer(GraphBuilder);
+}
 
 #if RHI_RAYTRACING
 
diff --git a/Engine/Source/Runtime/Renderer/Private/MobileShadingRenderer.cpp b/Engine/Source/Runtime/Renderer/Private/MobileShadingRenderer.cpp
--- a/Engine/Source/Runtime/Renderer/Private/MobileShadingRenderer.cpp
+++ b/Engine/Source/Runtime/Renderer/Private/MobileShadingRenderer.cpp
@@ -170,13 +170,14 @@ static FMobileCustomDepthStencilUsage GetCustomDepthStencilUsage(const FViewInfo
 static void RenderOpaqueFX(
 	FRDGBuilder& GraphBuilder,
 	TConstStridedView<FSceneView> Views,
+	FSceneUniformBuffer &SceneUniformBuffer,
 	FFXSystemInterface* FXSystem,
 	TRDGUniformBufferRef<FMobileSceneTextureUniformParameters> MobileSceneTexturesUniformBuffer)
 {
 	// Notify the FX system that opaque primitives have been rendered and we now have a valid depth buffer.
 	if (FXSystem && Views.Num() > 0)
 	{
-		FXSystem->PostRenderOpaque(GraphBuilder, Views, true /*bAllowGPUParticleUpdate*/);
+		FXSystem->PostRenderOpaque(GraphBuilder, Views, SceneUniformBuffer, true /*bAllowGPUParticleUpdate*/);
 
 		if (FGPUSortManager* GPUSortManager = FXSystem->GetGPUSortManager())
 		{
@@ -946,7 +947,7 @@ void FMobileSceneRenderer::Render(FRDGBuilder& GraphBuilder)
 	// Notify the FX system that the scene is about to be rendered.
 	if (FXSystem)
 	{
-		FXSystem->PreRender(GraphBuilder, GetSceneViews(), true /*bAllowGPUParticleUpdate*/);
+		FXSystem->PreRender(GraphBuilder, GetSceneViews(), GetSceneUniforms(), true /*bAllowGPUParticleUpdate*/);
 		if (FGPUSortManager* GPUSortManager = FXSystem->GetGPUSortManager())
 		{
 			// if GPUSortManager::OnPostRenderOpaque is called below (from RenderOpaqueFX) we must also call OnPreRender (as it sets up
@@ -1067,7 +1068,7 @@ void FMobileSceneRenderer::Render(FRDGBuilder& GraphBuilder)
 	FRendererModule& RendererModule = static_cast<FRendererModule&>(GetRendererModule());
 	RendererModule.RenderPostOpaqueExtensions(GraphBuilder, Views, SceneTextures);
 
-	RenderOpaqueFX(GraphBuilder, GetSceneViews(), FXSystem, SceneTextures.MobileUniformBuffer);
+	RenderOpaqueFX(GraphBuilder, GetSceneViews(), GetSceneUniforms(), FXSystem, SceneTextures.MobileUniformBuffer);
 
 	if (bRequiresPixelProjectedPlanarRelfectionPass)
 	{
diff --git a/Engine/Source/Runtime/Renderer/Private/SceneHitProxyRendering.cpp b/Engine/Source/Runtime/Renderer/Private/SceneHitProxyRendering.cpp
--- a/Engine/Source/Runtime/Renderer/Private/SceneHitProxyRendering.cpp
+++ b/Engine/Source/Runtime/Renderer/Private/SceneHitProxyRendering.cpp
@@ -667,14 +667,14 @@ void FDeferredShadingSceneRenderer::RenderHitProxies(FRDGBuilder& GraphBuilder)
 	if (FXSystem && Views.IsValidIndex(0))
 	{
 		FGPUSortManager* GPUSortManager = FXSystem->GetGPUSortManager();
-		FXSystem->PreRender(GraphBuilder, GetSceneViews(), false);
+		FXSystem->PreRender(GraphBuilder, GetSceneViews(), GetSceneUniforms(), false);
 		if (GPUSortManager)
 		{
 			GPUSortManager->OnPreRender(GraphBuilder);
 		}
 		// Call PostRenderOpaque now as this is irrelevant for when rendering hit proxies.
 		// because we don't tick the particles in the render loop (see last param being "false").
-		FXSystem->PostRenderOpaque(GraphBuilder, GetSceneViews(), false /*bAllowGPUParticleUpdate*/);
+		FXSystem->PostRenderOpaque(GraphBuilder, GetSceneViews(), GetSceneUniforms(), false /*bAllowGPUParticleUpdate*/);
 		if (GPUSortManager)
 		{
 			GPUSortManager->OnPostRenderOpaque(GraphBuilder);
diff --git a/Engine/Source/Runtime/Renderer/Public/FXRenderingUtils.h b/Engine/Source/Runtime/Renderer/Public/FXRenderingUtils.h
--- a/Engine/Source/Runtime/Renderer/Public/FXRenderingUtils.h
+++ b/Engine/Source/Runtime/Renderer/Public/FXRenderingUtils.h
@@ -20,6 +20,8 @@ class FSceneInterface;
 class FSceneView;
 class FSceneViewFamily;
 class FShaderParametersMetadata;
+class FSceneUniformBuffer;
+DECLARE_UNIFORM_BUFFER_STRUCT(FSceneUniformParameters, RENDERER_API)
 
 #if RHI_RAYTRACING
 class FRHIRayTracingScene;
@@ -58,18 +60,14 @@ namespace UE::FXRenderingUtils
 		RENDERER_API void SetupAtlasParameters(FRDGBuilder& GraphBuilder, uint8* DestinationData, const FSceneView* View);
 	}
 
-	namespace GPUScene
-	{
-		struct FBuffers
-		{
-			FRHIShaderResourceView* InstanceSceneDataBuffer;
-			FRHIShaderResourceView* InstancePayloadDataBuffer;
-			FRHIShaderResourceView* PrimitiveBuffer;
-			uint32                  SceneFrameNumber;
-		};
-
-		RENDERER_API FBuffers GetBuffers(const FSceneInterface* Scene);
-	}
+	/**
+	 * Creates a minimal Scene Uniform buffer for the given Scene, allocated using the FRDGBuilder (making it safe to keep reference).
+	 */
+	RENDERER_API FSceneUniformBuffer &CreateSceneUniformBuffer(FRDGBuilder& GraphBuilder, const FSceneInterface* Scene);
+	/**
+	 * Get the RDG uniform buffer from the FSceneUniformBuffer (not exposed in public API).
+	 */
+	RENDERER_API TRDGUniformBufferRef<FSceneUniformParameters> GetSceneUniformBuffer(FRDGBuilder& GraphBuilder, FSceneUniformBuffer &SceneUniformBuffer);
 
 #if RHI_RAYTRACING
 	namespace RayTracing
