Fixed missing include for non unity compile

diff --git a/Engine/Plugins/FX/Niagara/Source/NiagaraEditor/Private/ViewModels/NiagaraScratchPadScriptViewModel.cpp b/Engine/Plugins/FX/Niagara/Source/NiagaraEditor/Private/ViewModels/NiagaraScratchPadScriptViewModel.cpp
--- a/Engine/Plugins/FX/Niagara/Source/NiagaraEditor/Private/ViewModels/NiagaraScratchPadScriptViewModel.cpp
+++ b/Engine/Plugins/FX/Niagara/Source/NiagaraEditor/Private/ViewModels/NiagaraScratchPadScriptViewModel.cpp
@@ -12,6 +12,7 @@
 #include "NiagaraEditorModule.h"
 #include "NiagaraObjectSelection.h"
 #include "NiagaraParameterDefinitions.h"
+#include "NiagaraScriptSource.h"
 #include "ViewModels/NiagaraScratchPadUtilities.h"
 #include "ViewModels/NiagaraScriptGraphViewModel.h"
 #include "ViewModels/NiagaraSystemViewModel.h"
