* Small fix to clamp value for slider widgets
* display value as backup for empty dropdown names

diff --git a/Engine/Plugins/FX/Niagara/Source/NiagaraEditor/Private/TypeEditorUtilities/NiagaraFloatTypeEditorUtilities.cpp b/Engine/Plugins/FX/Niagara/Source/NiagaraEditor/Private/TypeEditorUtilities/NiagaraFloatTypeEditorUtilities.cpp
--- a/Engine/Plugins/FX/Niagara/Source/NiagaraEditor/Private/TypeEditorUtilities/NiagaraFloatTypeEditorUtilities.cpp
+++ b/Engine/Plugins/FX/Niagara/Source/NiagaraEditor/Private/TypeEditorUtilities/NiagaraFloatTypeEditorUtilities.cpp
@@ -74,8 +74,14 @@ public:
 						.MinValue(MinValue)
 						.MaxValue(MaxValue)
 						.Value(this, &SNiagaraFloatParameterEditor::GetValue)
-						.OnValueChanged(this, &SNiagaraFloatParameterEditor::ValueChanged)
-						.OnValueCommitted(this, &SNiagaraFloatParameterEditor::ValueCommitted)
+						.OnValueChanged_Lambda([this, MinValue, MaxValue](float NewVal)
+						{
+							ValueChanged(FMath::Clamp(NewVal, MinValue, MaxValue));
+						})
+						.OnValueCommitted_Lambda([this, MinValue, MaxValue](float NewVal, ETextCommit::Type CommitType)
+						{
+							ValueCommitted(FMath::Clamp(NewVal, MinValue, MaxValue), CommitType);
+						})
 						.TypeInterface(GetTypeInterface<float>(DisplayUnit))
 						.AllowSpin(false)
 						.Delta(WidgetCustomization.bHasStepWidth ? WidgetCustomization.StepWidth : 0)
@@ -109,7 +115,7 @@ public:
 			TArray<SNiagaraNumericDropDown<float>::FNamedValue> DropDownValues;
 			for (const FWidgetNamedInputValue& Value : WidgetCustomization.InputDropdownValues)
 			{
-				DropDownValues.Add(SNiagaraNumericDropDown<float>::FNamedValue(Value.Value, Value.DisplayName, Value.Tooltip));
+				DropDownValues.Add(SNiagaraNumericDropDown<float>::FNamedValue(Value.Value, Value.DisplayName.IsEmpty() ? FText::AsNumber(Value.Value) : Value.DisplayName, Value.Tooltip));
 			}
 
 			ChildSlot
diff --git a/Engine/Plugins/FX/Niagara/Source/NiagaraEditor/Private/TypeEditorUtilities/NiagaraIntegerTypeEditorUtilities.cpp b/Engine/Plugins/FX/Niagara/Source/NiagaraEditor/Private/TypeEditorUtilities/NiagaraIntegerTypeEditorUtilities.cpp
--- a/Engine/Plugins/FX/Niagara/Source/NiagaraEditor/Private/TypeEditorUtilities/NiagaraIntegerTypeEditorUtilities.cpp
+++ b/Engine/Plugins/FX/Niagara/Source/NiagaraEditor/Private/TypeEditorUtilities/NiagaraIntegerTypeEditorUtilities.cpp
@@ -72,8 +72,14 @@ public:
 						.MinValue(MinValue)
 						.MaxValue(MaxValue)
 						.Value(this, &SNiagaraIntegerParameterEditor::GetValue)
-						.OnValueChanged(this, &SNiagaraIntegerParameterEditor::ValueChanged)
-						.OnValueCommitted(this, &SNiagaraIntegerParameterEditor::ValueCommitted)
+						.OnValueChanged_Lambda([this, MinValue, MaxValue](int32 NewVal)
+						{
+							ValueChanged(FMath::Clamp(NewVal, MinValue, MaxValue));
+						})
+						.OnValueCommitted_Lambda([this, MinValue, MaxValue](int32 NewVal, ETextCommit::Type CommitType)
+						{
+							ValueCommitted(FMath::Clamp(NewVal, MinValue, MaxValue), CommitType);
+						})
 						.TypeInterface(GetTypeInterface<int32>(DisplayUnit))
 						.AllowSpin(false)
 						.Delta(WidgetCustomization.bHasStepWidth ? WidgetCustomization.StepWidth : 0)
@@ -86,7 +92,7 @@ public:
 			TArray<SNiagaraNumericDropDown<int32>::FNamedValue> DropDownValues;
 			for (const FWidgetNamedInputValue& Value : WidgetCustomization.InputDropdownValues)
 			{
-				DropDownValues.Add(SNiagaraNumericDropDown<int32>::FNamedValue(Value.Value, Value.DisplayName, Value.Tooltip));
+				DropDownValues.Add(SNiagaraNumericDropDown<int32>::FNamedValue(Value.Value, Value.DisplayName.IsEmpty() ? FText::AsNumber(Value.Value) : Value.DisplayName, Value.Tooltip));
 			}
 
 			ChildSlot
