Temp hack to fix recent breakage due to data set mem savings format change.

diff --git a/Engine/Plugins/FX/Niagara/Source/Niagara/Private/NiagaraDataChannel.cpp b/Engine/Plugins/FX/Niagara/Source/Niagara/Private/NiagaraDataChannel.cpp
--- a/Engine/Plugins/FX/Niagara/Source/Niagara/Private/NiagaraDataChannel.cpp
+++ b/Engine/Plugins/FX/Niagara/Source/Niagara/Private/NiagaraDataChannel.cpp
@@ -831,6 +831,10 @@ void UNiagaraDataChannel::PostLoad()
 		}
 	}
 
+	//Init compiled data. These are not currently serialized as we have no mechanism to rebuild then on internal data format changes like those in scripts.
+	GetCompiledData(ENiagaraSimTarget::CPUSim);
+	GetCompiledData(ENiagaraSimTarget::GPUComputeSim);
+
 	//TODO: Can serialize?
 	GameDataLayout.Init(Variables);
 
diff --git a/Engine/Plugins/FX/Niagara/Source/Niagara/Public/NiagaraDataChannel.h b/Engine/Plugins/FX/Niagara/Source/Niagara/Public/NiagaraDataChannel.h
--- a/Engine/Plugins/FX/Niagara/Source/Niagara/Public/NiagaraDataChannel.h
+++ b/Engine/Plugins/FX/Niagara/Source/Niagara/Public/NiagaraDataChannel.h
@@ -136,10 +136,10 @@ private:
 	/**
 	Data layout for payloads in Niagara datasets.
 	*/
-	UPROPERTY()
+	UPROPERTY(Transient)
 	mutable FNiagaraDataSetCompiledData CompiledData;
 
-	UPROPERTY()
+	UPROPERTY(Transient)
 	mutable FNiagaraDataSetCompiledData CompiledDataGPU;
 
 	/** Layout information for any data stored at the "Game" level. i.e. From game code/BP. AoS layout and LWC types. */
