Added a setting to make numeric inputs feel snappier by updating the sim after every input change

diff --git a/Engine/Plugins/FX/Niagara/Source/NiagaraEditor/Private/TypeEditorUtilities/NiagaraColorTypeEditorUtilities.cpp b/Engine/Plugins/FX/Niagara/Source/NiagaraEditor/Private/TypeEditorUtilities/NiagaraColorTypeEditorUtilities.cpp
--- a/Engine/Plugins/FX/Niagara/Source/NiagaraEditor/Private/TypeEditorUtilities/NiagaraColorTypeEditorUtilities.cpp
+++ b/Engine/Plugins/FX/Niagara/Source/NiagaraEditor/Private/TypeEditorUtilities/NiagaraColorTypeEditorUtilities.cpp
@@ -1,6 +1,8 @@
 // Copyright Epic Games, Inc. All Rights Reserved.
 
 #include "NiagaraColorTypeEditorUtilities.h"
+
+#include "NiagaraEditorSettings.h"
 #include "SNiagaraParameterEditor.h"
 #include "NiagaraTypes.h"
 #include "NiagaraEditorStyle.h"
@@ -90,6 +92,7 @@ private:
 		.OnBeginSliderMovement(this, &SNiagaraColorParameterEditor::BeginSliderMovement)
 		.OnEndSliderMovement(this, &SNiagaraColorParameterEditor::EndSliderMovement)
 		.AllowSpin(true)
+		.BroadcastValueChangesPerKey(!GetDefault<UNiagaraEditorSettings>()->GetUpdateStackValuesOnCommitOnly())
 		.LabelVAlign(EVerticalAlignment::VAlign_Center)
 		.MinDesiredValueWidth(30)
 		.Label()
diff --git a/Engine/Plugins/FX/Niagara/Source/NiagaraEditor/Private/TypeEditorUtilities/NiagaraFloatTypeEditorUtilities.cpp b/Engine/Plugins/FX/Niagara/Source/NiagaraEditor/Private/TypeEditorUtilities/NiagaraFloatTypeEditorUtilities.cpp
--- a/Engine/Plugins/FX/Niagara/Source/NiagaraEditor/Private/TypeEditorUtilities/NiagaraFloatTypeEditorUtilities.cpp
+++ b/Engine/Plugins/FX/Niagara/Source/NiagaraEditor/Private/TypeEditorUtilities/NiagaraFloatTypeEditorUtilities.cpp
@@ -4,6 +4,7 @@
 
 #include "EdGraphSchema_Niagara.h"
 #include "GraphEditorSettings.h"
+#include "NiagaraEditorSettings.h"
 #include "NiagaraEditorStyle.h"
 #include "SNiagaraParameterEditor.h"
 #include "NiagaraTypes.h"
@@ -26,6 +27,7 @@ public:
 			.MaximumDesiredWidth(DefaultInputSize));
 
 		const UGraphEditorSettings* Settings = GetDefault<UGraphEditorSettings>();
+		const UNiagaraEditorSettings* NiagaraSettings = GetDefault<UNiagaraEditorSettings>();
 
 		if (WidgetCustomization.WidgetType == ENiagaraInputWidgetType::Slider)
 		{
@@ -74,14 +76,8 @@ public:
 						.MinValue(MinValue)
 						.MaxValue(MaxValue)
 						.Value(this, &SNiagaraFloatParameterEditor::GetValue)
-						.OnValueChanged_Lambda([this, MinValue, MaxValue](float NewVal)
-						{
-							ValueChanged(FMath::Clamp(NewVal, MinValue, MaxValue));
-						})
-						.OnValueCommitted_Lambda([this, MinValue, MaxValue](float NewVal, ETextCommit::Type CommitType)
-						{
-							ValueCommitted(FMath::Clamp(NewVal, MinValue, MaxValue), CommitType);
-						})
+						.OnValueChanged(this, &SNiagaraFloatParameterEditor::ValueChanged)
+						.OnValueCommitted(this, &SNiagaraFloatParameterEditor::ValueCommitted)
 						.TypeInterface(GetTypeInterface<float>(DisplayUnit))
 						.AllowSpin(false)
 						.Delta(WidgetCustomization.bHasStepWidth ? WidgetCustomization.StepWidth : 0)
@@ -164,6 +160,7 @@ public:
 				.OnEndSliderMovement(this, &SNiagaraFloatParameterEditor::EndSliderMovement)
 				.TypeInterface(GetTypeInterface<float>(DisplayUnit))
 				.AllowSpin(true)
+				.BroadcastValueChangesPerKey(!NiagaraSettings->GetUpdateStackValuesOnCommitOnly())
 				.LabelPadding(FMargin(3))
 				.Delta(WidgetCustomization.bHasStepWidth ? WidgetCustomization.StepWidth : 0)
 				.LabelLocation(SNumericEntryBox<float>::ELabelLocation::Inside)
diff --git a/Engine/Plugins/FX/Niagara/Source/NiagaraEditor/Private/TypeEditorUtilities/NiagaraIntegerTypeEditorUtilities.cpp b/Engine/Plugins/FX/Niagara/Source/NiagaraEditor/Private/TypeEditorUtilities/NiagaraIntegerTypeEditorUtilities.cpp
--- a/Engine/Plugins/FX/Niagara/Source/NiagaraEditor/Private/TypeEditorUtilities/NiagaraIntegerTypeEditorUtilities.cpp
+++ b/Engine/Plugins/FX/Niagara/Source/NiagaraEditor/Private/TypeEditorUtilities/NiagaraIntegerTypeEditorUtilities.cpp
@@ -4,6 +4,7 @@
 
 #include "EdGraphSchema_Niagara.h"
 #include "GraphEditorSettings.h"
+#include "NiagaraEditorSettings.h"
 #include "SNiagaraParameterEditor.h"
 #include "NiagaraTypes.h"
 #include "NiagaraEditorStyle.h"
@@ -72,14 +73,8 @@ public:
 						.MinValue(MinValue)
 						.MaxValue(MaxValue)
 						.Value(this, &SNiagaraIntegerParameterEditor::GetValue)
-						.OnValueChanged_Lambda([this, MinValue, MaxValue](int32 NewVal)
-						{
-							ValueChanged(FMath::Clamp(NewVal, MinValue, MaxValue));
-						})
-						.OnValueCommitted_Lambda([this, MinValue, MaxValue](int32 NewVal, ETextCommit::Type CommitType)
-						{
-							ValueCommitted(FMath::Clamp(NewVal, MinValue, MaxValue), CommitType);
-						})
+						.OnValueChanged(this, &SNiagaraIntegerParameterEditor::ValueChanged)
+						.OnValueCommitted(this, &SNiagaraIntegerParameterEditor::ValueCommitted)
 						.TypeInterface(GetTypeInterface<int32>(DisplayUnit))
 						.AllowSpin(false)
 						.Delta(WidgetCustomization.bHasStepWidth ? WidgetCustomization.StepWidth : 0)
@@ -141,6 +136,7 @@ public:
 				.OnEndSliderMovement(this, &SNiagaraIntegerParameterEditor::EndSliderMovement)
 				.TypeInterface(GetTypeInterface<int32>(DisplayUnit))
 				.AllowSpin(true)
+				.BroadcastValueChangesPerKey(!GetDefault<UNiagaraEditorSettings>()->GetUpdateStackValuesOnCommitOnly())
 				.LabelPadding(FMargin(3))
 				.LabelLocation(SNumericEntryBox<int32>::ELabelLocation::Inside)
 				.Label()
diff --git a/Engine/Plugins/FX/Niagara/Source/NiagaraEditor/Private/TypeEditorUtilities/NiagaraVectorTypeEditorUtilities.cpp b/Engine/Plugins/FX/Niagara/Source/NiagaraEditor/Private/TypeEditorUtilities/NiagaraVectorTypeEditorUtilities.cpp
--- a/Engine/Plugins/FX/Niagara/Source/NiagaraEditor/Private/TypeEditorUtilities/NiagaraVectorTypeEditorUtilities.cpp
+++ b/Engine/Plugins/FX/Niagara/Source/NiagaraEditor/Private/TypeEditorUtilities/NiagaraVectorTypeEditorUtilities.cpp
@@ -1,6 +1,8 @@
 // Copyright Epic Games, Inc. All Rights Reserved.
 
 #include "NiagaraVectorTypeEditorUtilities.h"
+
+#include "NiagaraEditorSettings.h"
 #include "SNiagaraParameterEditor.h"
 #include "NiagaraEditorStyle.h"
 #include "NiagaraTypes.h"
@@ -79,6 +81,7 @@ private:
 		.OnBeginSliderMovement(this, &SNiagaraVectorParameterEditorBase::BeginSliderMovement)
 		.OnEndSliderMovement(this, &SNiagaraVectorParameterEditorBase::EndSliderMovement)
 		.TypeInterface(GetTypeInterface<float>(DisplayUnit))
+		.BroadcastValueChangesPerKey(!GetDefault<UNiagaraEditorSettings>()->GetUpdateStackValuesOnCommitOnly())
 		.AllowSpin(true)
 		.LabelPadding( FMargin(3) )
 		.LabelLocation( SNumericEntryBox<float>::ELabelLocation::Inside )
diff --git a/Engine/Plugins/FX/Niagara/Source/NiagaraEditor/Public/NiagaraEditorSettings.h b/Engine/Plugins/FX/Niagara/Source/NiagaraEditor/Public/NiagaraEditorSettings.h
--- a/Engine/Plugins/FX/Niagara/Source/NiagaraEditor/Public/NiagaraEditorSettings.h
+++ b/Engine/Plugins/FX/Niagara/Source/NiagaraEditor/Public/NiagaraEditorSettings.h
@@ -409,6 +409,8 @@ public:
 	bool IsAllowedAssetByClassUsage(const FAssetData& InAssetData) const;
 	bool IsAllowedAssetObjectByClassUsage(const UObject& InAssetObject) const;
 
+	bool GetUpdateStackValuesOnCommitOnly() const;
+
 private:
 	bool IsAllowedAssetObjectByClassUsageInternal(const UObject& InAssetObject, TSet<const UObject*>& CheckedAssetObjects) const;
 
@@ -450,6 +452,10 @@ private:
 	UPROPERTY(config, EditAnywhere, Category=Niagara, meta=(EditCondition="bDisplayAffectedAssetStats", ClampMin=1))
 	int32 AffectedAssetSearchLimit = 25;
 
+	/** This affects numeric inputs for modules. When set to false, the inputs update the simulation while typing. When set to true, the simulation is only updated after submitting the change by pressing Enter. */
+	UPROPERTY(config, EditAnywhere, Category = Niagara)
+	bool bUpdateStackValuesOnCommitOnly = false;
+
 	/** Speeds used for slowing down and speeding up the playback speeds */
 	UPROPERTY(config, EditAnywhere, Category = Niagara)
 	TArray<float> PlaybackSpeeds;
