Added missing config change file

diff --git a/Engine/Plugins/FX/Niagara/Source/NiagaraEditor/Private/NiagaraEditorSettings.cpp b/Engine/Plugins/FX/Niagara/Source/NiagaraEditor/Private/NiagaraEditorSettings.cpp
--- a/Engine/Plugins/FX/Niagara/Source/NiagaraEditor/Private/NiagaraEditorSettings.cpp
+++ b/Engine/Plugins/FX/Niagara/Source/NiagaraEditor/Private/NiagaraEditorSettings.cpp
@@ -842,5 +842,10 @@ bool UNiagaraEditorSettings::IsAllowedAssetObjectByClassUsage(const UObject& InA
 	return IsAllowedAssetObjectByClassUsageInternal(InAssetObject, CheckedAssetObjects);
 }
 
+bool UNiagaraEditorSettings::GetUpdateStackValuesOnCommitOnly() const
+{
+	return bUpdateStackValuesOnCommitOnly;
+}
+
 #undef LOCTEXT_NAMESPACE
 
