Fixed a bug where the simulation was always reset after a change, even when paused

diff --git a/Engine/Plugins/FX/Niagara/Source/NiagaraEditor/Private/ViewModels/NiagaraSystemViewModel.cpp b/Engine/Plugins/FX/Niagara/Source/NiagaraEditor/Private/ViewModels/NiagaraSystemViewModel.cpp
--- a/Engine/Plugins/FX/Niagara/Source/NiagaraEditor/Private/ViewModels/NiagaraSystemViewModel.cpp
+++ b/Engine/Plugins/FX/Niagara/Source/NiagaraEditor/Private/ViewModels/NiagaraSystemViewModel.cpp
@@ -2277,7 +2277,8 @@ void FNiagaraSystemViewModel::ResetSystem(ETimeResetMode TimeResetMode, EMultiRe
 	bool bResetAge = false;
 	if (Sequencer.IsValid())
 	{
-		bResetAge = TimeResetMode == ETimeResetMode::AllowResetTime && (Sequencer->GetPlaybackStatus() == EMovieScenePlayerStatus::Playing || Sequencer->GetPlaybackStatus() == EMovieScenePlayerStatus::Stopped || EditorSettings->GetResimulateOnChangeWhilePaused() == false);
+		// important to note here that EMovieScenePlayerStatus::Stopped means the simulation is *paused* in sequencer.
+		bResetAge = TimeResetMode == ETimeResetMode::AllowResetTime && (Sequencer->GetPlaybackStatus() == EMovieScenePlayerStatus::Playing || (Sequencer->GetPlaybackStatus() == EMovieScenePlayerStatus::Stopped && EditorSettings->GetResimulateOnChangeWhilePaused() == false));
 		if (bResetAge)
 		{
 			TGuardValue<bool> Guard(bSettingSequencerTimeDirectly, true);
@@ -2300,7 +2301,7 @@ void FNiagaraSystemViewModel::ResetSystem(ETimeResetMode TimeResetMode, EMultiRe
 	UpdateContext.GetPostWork().BindLambda(
 		[ReinitMode, bResetAge](UNiagaraComponent* Component)
 		{
-			if (ReinitMode == EReinitMode::ResetSystem && bResetAge  && Component->GetAgeUpdateMode() == ENiagaraAgeUpdateMode::DesiredAge)
+			if (ReinitMode == EReinitMode::ResetSystem && bResetAge && Component->GetAgeUpdateMode() == ENiagaraAgeUpdateMode::DesiredAge)
 			{
 				Component->SetDesiredAge(0);
 			}
