- Ensure we handle null EmitterData for variable binding

diff --git a/Engine/Plugins/FX/Niagara/Source/NiagaraEditor/Private/Customizations/NiagaraTypeCustomizations.cpp b/Engine/Plugins/FX/Niagara/Source/NiagaraEditor/Private/Customizations/NiagaraTypeCustomizations.cpp
--- a/Engine/Plugins/FX/Niagara/Source/NiagaraEditor/Private/Customizations/NiagaraTypeCustomizations.cpp
+++ b/Engine/Plugins/FX/Niagara/Source/NiagaraEditor/Private/Customizations/NiagaraTypeCustomizations.cpp
@@ -136,18 +136,19 @@ TArray<FNiagaraVariableBase> FNiagaraStackAssetAction_VarBind::FindVariables(con
 	TArray<FNiagaraParameterMapHistory> Histories;
 
 	UNiagaraEmitter* Emitter = InEmitter.Emitter;
-	UNiagaraScriptSource* Source = Cast<UNiagaraScriptSource>(InEmitter.GetEmitterData()->GraphSource);
-	if (Source)
+	if (FVersionedNiagaraEmitterData* EmitterData = InEmitter.GetEmitterData())
 	{
-		Histories.Append(UNiagaraNodeParameterMapBase::GetParameterMaps(Source->NodeGraph));
+		if ( UNiagaraScriptSource* Source = Cast<UNiagaraScriptSource>(EmitterData->GraphSource) )
+		{
+			Histories.Append(UNiagaraNodeParameterMapBase::GetParameterMaps(Source->NodeGraph));
+		}
 	}
 
 	if (bSystem || bEmitter)
 	{
 		if (UNiagaraSystem* Sys = Emitter->GetTypedOuter<UNiagaraSystem>())
 		{
-			Source = Cast<UNiagaraScriptSource>(Sys->GetSystemUpdateScript()->GetLatestSource()); 
-			if (Source)
+			if ( UNiagaraScriptSource* Source = Cast<UNiagaraScriptSource>(Sys->GetSystemUpdateScript()->GetLatestSource()) )
 			{
 				Histories.Append(UNiagaraNodeParameterMapBase::GetParameterMaps(Source->NodeGraph));
 			}
@@ -272,8 +273,10 @@ TArray<FName> FNiagaraVariableAttributeBindingCustomization::GetNames(const FVer
 	}
 	
 	TArray<UNiagaraGraph*> GraphsToTraverse;
-
-	GraphsToTraverse.Add(Cast<UNiagaraScriptSource>(OwningVersionedEmitter.GetEmitterData()->GraphSource)->NodeGraph);
+	if (FVersionedNiagaraEmitterData* OwningEmitterData = OwningVersionedEmitter.GetEmitterData())
+	{
+		GraphsToTraverse.Add(Cast<UNiagaraScriptSource>(OwningEmitterData->GraphSource)->NodeGraph);
+	}
 
 	TSet<FNiagaraVariableBase> Vars;
 
@@ -977,7 +980,10 @@ TArray<TPair<FNiagaraVariableBase, FNiagaraVariableBase> > FNiagaraMaterialAttri
 		TArray<UNiagaraScript*> Scripts;
 		Scripts.Add(BaseSystem->GetSystemUpdateScript());
 		Scripts.Add(BaseSystem->GetSystemSpawnScript());
-		BaseEmitter.GetEmitterData()->GetScripts(Scripts, false);
+		if (FVersionedNiagaraEmitterData* EmitterData = BaseEmitter.GetEmitterData())
+		{
+			EmitterData->GetScripts(Scripts, false);
+		}
 
 		TMap<FString, FString> EmitterAlias;
 		EmitterAlias.Emplace(FNiagaraConstants::EmitterNamespace.ToString(), BaseEmitter.Emitter->GetUniqueEmitterName());
