- Fix for UVs other than layer 0 being incorrect on none manual fetch platforms

#rnx

diff --git a/Engine/Plugins/FX/Niagara/Shaders/Private/NiagaraMeshVertexFactory.ush b/Engine/Plugins/FX/Niagara/Shaders/Private/NiagaraMeshVertexFactory.ush
--- a/Engine/Plugins/FX/Niagara/Shaders/Private/NiagaraMeshVertexFactory.ush
+++ b/Engine/Plugins/FX/Niagara/Shaders/Private/NiagaraMeshVertexFactory.ush
@@ -42,27 +42,25 @@ struct FVertexFactoryInput
 	half4	VertexColor : ATTRIBUTE3;
 
 	#if NUM_MATERIAL_TEXCOORDS_VERTEX
-		float2	TexCoords0 : ATTRIBUTE4;
-		#if NUM_MATERIAL_TEXCOORDS_VERTEX > 1
-			float2	TexCoords1 : ATTRIBUTE5;
-		#endif
-		#if NUM_MATERIAL_TEXCOORDS_VERTEX > 2
-			float2	TexCoords2 : ATTRIBUTE6;
-		#endif
-		#if NUM_MATERIAL_TEXCOORDS_VERTEX > 3
-			float2	TexCoords3 : ATTRIBUTE7;
-		#endif
-		#if NUM_MATERIAL_TEXCOORDS_VERTEX > 4
-			float2	TexCoords4 : ATTRIBUTE8;
+		#if NUM_MATERIAL_TEXCOORDS_VERTEX == 1
+			float2	TexCoords0 : ATTRIBUTE4;
+		#elif NUM_MATERIAL_TEXCOORDS_VERTEX >= 2
+			float4	TexCoords01 : ATTRIBUTE4;
 		#endif
-		#if NUM_MATERIAL_TEXCOORDS_VERTEX > 5
-			float2	TexCoords5 : ATTRIBUTE9;
+		#if NUM_MATERIAL_TEXCOORDS_VERTEX == 3
+			float2	TexCoords2 : ATTRIBUTE5;
+		#elif NUM_MATERIAL_TEXCOORDS_VERTEX >= 4
+			float4	TexCoords23 : ATTRIBUTE5;
 		#endif
-		#if NUM_MATERIAL_TEXCOORDS_VERTEX > 6
-			float2	TexCoords6 : ATTRIBUTE10;
+		#if NUM_MATERIAL_TEXCOORDS_VERTEX == 5
+			float2	TexCoords4 : ATTRIBUTE6;
+		#elif NUM_MATERIAL_TEXCOORDS_VERTEX >= 6
+			float4	TexCoords45 : ATTRIBUTE6;
 		#endif
-		#if NUM_MATERIAL_TEXCOORDS_VERTEX > 7
-			float2	TexCoords7 : ATTRIBUTE11;
+		#if NUM_MATERIAL_TEXCOORDS_VERTEX == 7
+			float2	TexCoords6 : ATTRIBUTE7;
+		#elif NUM_MATERIAL_TEXCOORDS_VERTEX >= 8
+			float4	TexCoords67 : ATTRIBUTE7;
 		#endif
 	#elif USE_PARTICLE_SUBUVS
 		float2	TexCoords0 : ATTRIBUTE4;
@@ -574,27 +572,29 @@ FVertexFactoryIntermediates GetVertexFactoryIntermediates(FVertexFactoryInput In
 			Intermediates.TexCoords[CoordinateIndex] = NiagaraMeshVF.VertexFetch_TexCoordBuffer[NumFetchTexCoords * Input.VertexId + ClampedCoordinateIndex];
 		}
 	#else
-		Intermediates.TexCoords[0] = Input.TexCoords0;
-		#if NUM_MATERIAL_TEXCOORDS_VERTEX > 1
-			Intermediates.TexCoords[1] = Input.TexCoords1;
+		#if NUM_MATERIAL_TEXCOORDS_VERTEX == 1
+			Intermediates.TexCoords[0] = Input.TexCoords0;
+		#elif NUM_MATERIAL_TEXCOORDS_VERTEX >= 2
+			Intermediates.TexCoords[0] = Input.TexCoords01.xy;
+			Intermediates.TexCoords[1] = Input.TexCoords01.zw;
 		#endif
-		#if NUM_MATERIAL_TEXCOORDS_VERTEX > 2
+		#if NUM_MATERIAL_TEXCOORDS_VERTEX == 3
 			Intermediates.TexCoords[2] = Input.TexCoords2;
+		#elif NUM_MATERIAL_TEXCOORDS_VERTEX >= 4
+			Intermediates.TexCoords[2] = Input.TexCoords23.xy;
+			Intermediates.TexCoords[3] = Input.TexCoords23.zw;
 		#endif
-		#if NUM_MATERIAL_TEXCOORDS_VERTEX > 3
-			Intermediates.TexCoords[3] = Input.TexCoords3;
-		#endif
-		#if NUM_MATERIAL_TEXCOORDS_VERTEX > 4
+		#if NUM_MATERIAL_TEXCOORDS_VERTEX == 5
 			Intermediates.TexCoords[4] = Input.TexCoords4;
+		#elif NUM_MATERIAL_TEXCOORDS_VERTEX >= 6
+			Intermediates.TexCoords[4] = Input.TexCoords45.xy;
+			Intermediates.TexCoords[5] = Input.TexCoords45.zw;
 		#endif
-		#if NUM_MATERIAL_TEXCOORDS_VERTEX > 5
-			Intermediates.TexCoords[5] = Input.TexCoords5;
-		#endif
-		#if NUM_MATERIAL_TEXCOORDS_VERTEX > 6
+		#if NUM_MATERIAL_TEXCOORDS_VERTEX == 7
 			Intermediates.TexCoords[6] = Input.TexCoords6;
-		#endif
-		#if NUM_MATERIAL_TEXCOORDS_VERTEX > 7
-			Intermediates.TexCoords[7] = Input.TexCoords7;
+		#elif NUM_MATERIAL_TEXCOORDS_VERTEX >= 8
+			Intermediates.TexCoords[6] = Input.TexCoords67.xy;
+			Intermediates.TexCoords[7] = Input.TexCoords67.zw;
 		#endif
 	#endif // MANUAL_VERTEX_FETCH
 #endif // NUM_MATERIAL_TEXCOORDS_VERTEX
