Incorporated ParameterMaphHistory & Builder into graph bridge layer that was added to the translator.  Abstracts graph logic into a layer where we can replace it with our new Digest format.

diff --git a/Engine/Plugins/FX/Niagara/Source/Niagara/Classes/NiagaraParameterCollection.h b/Engine/Plugins/FX/Niagara/Source/Niagara/Classes/NiagaraParameterCollection.h
--- a/Engine/Plugins/FX/Niagara/Source/Niagara/Classes/NiagaraParameterCollection.h
+++ b/Engine/Plugins/FX/Niagara/Source/Niagara/Classes/NiagaraParameterCollection.h
@@ -28,6 +28,7 @@ public:
 	UNiagaraParameterCollection* GetParent()const { return Collection; }
 
 	FNiagaraParameterStore& GetParameterStore() { return ParameterStorage; }
+	const FNiagaraParameterStore& GetParameterStore() const { return ParameterStorage; }
 
 	bool AddParameter(const FNiagaraVariable& Parameter);
 	bool RemoveParameter(const FNiagaraVariable& Parameter);
@@ -147,6 +148,7 @@ public:
 	void RenameParameter(FNiagaraVariable& Parameter, FName NewName);
 
 	TArray<FNiagaraVariable>& GetParameters() { return Parameters; }
+	const TArray<FNiagaraVariable>& GetParameters() const { return Parameters; }
 
 	const UMaterialParameterCollection* GetSourceCollection() const { return SourceMaterialCollection; }
 
@@ -154,6 +156,7 @@ public:
 	//TMap<FString, UNiagaraParameterCollectionOverride> PerPlatformOverrides;
 
 	FORCEINLINE UNiagaraParameterCollectionInstance* GetDefaultInstance() { return DefaultInstance; }
+	FORCEINLINE const UNiagaraParameterCollectionInstance* GetDefaultInstance() const { return DefaultInstance; }
 	
 	/**
 	Takes the friendly name presented to the UI and converts to the real parameter name used under the hood.
diff --git a/Engine/Plugins/FX/Niagara/Source/Niagara/Private/NiagaraParameterCollection.cpp b/Engine/Plugins/FX/Niagara/Source/Niagara/Private/NiagaraParameterCollection.cpp
--- a/Engine/Plugins/FX/Niagara/Source/Niagara/Private/NiagaraParameterCollection.cpp
+++ b/Engine/Plugins/FX/Niagara/Source/Niagara/Private/NiagaraParameterCollection.cpp
@@ -457,6 +457,7 @@ int32 UNiagaraParameterCollection::AddParameter(const FNiagaraVariable& Paramete
 	if (Idx == INDEX_NONE)
 	{
 		Modify();
+		CompileId = FGuid::NewGuid(); // Any scripts could now be referencing this new parameter and so would need to be changed.
 
 		Idx = Parameters.Add(Parameter);
 		DefaultInstance->AddParameter(Parameter);
diff --git a/Engine/Plugins/FX/Niagara/Source/NiagaraEditor/Private/Customizations/NiagaraComponentRendererPropertiesDetails.cpp b/Engine/Plugins/FX/Niagara/Source/NiagaraEditor/Private/Customizations/NiagaraComponentRendererPropertiesDetails.cpp
--- a/Engine/Plugins/FX/Niagara/Source/NiagaraEditor/Private/Customizations/NiagaraComponentRendererPropertiesDetails.cpp
+++ b/Engine/Plugins/FX/Niagara/Source/NiagaraEditor/Private/Customizations/NiagaraComponentRendererPropertiesDetails.cpp
@@ -417,7 +417,7 @@ TArray<FNiagaraVariable> FNiagaraComponentRendererPropertiesDetails::GetPossible
 	{
 		for (const FNiagaraVariable& Var : History.Variables)
 		{
-			if (FNiagaraParameterMapHistory::IsAttribute(Var) && UNiagaraComponentRendererProperties::IsConvertible(Var.GetType(), PropertyType))
+			if (FNiagaraParameterUtilities::IsAttribute(Var) && UNiagaraComponentRendererProperties::IsConvertible(Var.GetType(), PropertyType))
 			{
 				Options.AddUnique(Var);
 			}
diff --git a/Engine/Plugins/FX/Niagara/Source/NiagaraEditor/Private/Customizations/NiagaraTypeCustomizations.cpp b/Engine/Plugins/FX/Niagara/Source/NiagaraEditor/Private/Customizations/NiagaraTypeCustomizations.cpp
--- a/Engine/Plugins/FX/Niagara/Source/NiagaraEditor/Private/Customizations/NiagaraTypeCustomizations.cpp
+++ b/Engine/Plugins/FX/Niagara/Source/NiagaraEditor/Private/Customizations/NiagaraTypeCustomizations.cpp
@@ -162,11 +162,11 @@ TArray<FNiagaraVariableBase> FNiagaraStackAssetAction_VarBind::FindVariables(con
 			if (Var.GetType().IsStatic() && !bAllowStatic)
 				continue;
 
-			if (FNiagaraParameterMapHistory::IsAttribute(Var) && bParticles)
+			if (FNiagaraParameterUtilities::IsAttribute(Var) && bParticles)
 			{
 				Bindings.AddUnique(Var);
 			}
-			else if (FNiagaraParameterMapHistory::IsSystemParameter(Var) && bSystem)
+			else if (FNiagaraParameterUtilities::IsSystemParameter(Var) && bSystem)
 			{
 				Bindings.AddUnique(Var);
 			}
@@ -175,7 +175,7 @@ TArray<FNiagaraVariableBase> FNiagaraStackAssetAction_VarBind::FindVariables(con
 				Bindings.AddUnique(FNiagaraUtilities::ResolveAliases(Var, FNiagaraAliasContext()
 					.ChangeEmitterNameToEmitter(Emitter->GetUniqueEmitterName())));
 			}
-			else if (FNiagaraParameterMapHistory::IsAliasedEmitterParameter(Var) && bEmitter)
+			else if (FNiagaraParameterUtilities::IsAliasedEmitterParameter(Var) && bEmitter)
 			{
 				Bindings.AddUnique(Var);
 			}
@@ -183,7 +183,7 @@ TArray<FNiagaraVariableBase> FNiagaraStackAssetAction_VarBind::FindVariables(con
 			{
 				Bindings.AddUnique(Var);
 			}
-			else if (FNiagaraParameterMapHistory::IsUserParameter(Var) && bUser)
+			else if (FNiagaraParameterUtilities::IsUserParameter(Var) && bUser)
 			{
 				Bindings.AddUnique(Var);
 			}
@@ -647,7 +647,7 @@ TArray<FName> FNiagaraUserParameterBindingCustomization::GetNames() const
 		UNiagaraSystem* BaseSystem = BaseSystemWeakPtr.Get();
 		for (const FNiagaraVariable Var : BaseSystem->GetExposedParameters().ReadParameterVariables())
 		{
-			if (FNiagaraParameterMapHistory::IsUserParameter(Var) && Var.GetType() == TargetUserParameterBinding->Parameter.GetType())
+			if (FNiagaraParameterUtilities::IsUserParameter(Var) && Var.GetType() == TargetUserParameterBinding->Parameter.GetType())
 			{
 				Names.AddUnique(Var.GetName());
 			}
@@ -1760,7 +1760,7 @@ TArray<TSharedPtr<FEdGraphSchemaAction>> FNiagaraScriptVariableBindingCustomizat
 	{
 		for (const FNiagaraVariable& Var : History.Variables)
 		{
-			FString Namespace = FNiagaraParameterMapHistory::GetNamespace(Var);
+			FString Namespace = FNiagaraParameterUtilities::GetNamespace(Var);
 			if (Namespace == TEXT("Module."))
 			{
 				// TODO: Skip module inputs for now. Does it make sense to bind module inputs to module inputs?
@@ -1775,7 +1775,7 @@ TArray<TSharedPtr<FEdGraphSchemaAction>> FNiagaraScriptVariableBindingCustomizat
 
 	for (const auto& Var : Graph->GetParameterReferenceMap())
 	{
-		FString Namespace = FNiagaraParameterMapHistory::GetNamespace(Var.Key);
+		FString Namespace = FNiagaraParameterUtilities::GetNamespace(Var.Key);
 		if (Namespace == TEXT("Module."))
 		{
 			// TODO: Skip module inputs for now. Does it make sense to bind module inputs to module inputs?
diff --git a/Engine/Plugins/FX/Niagara/Source/NiagaraEditor/Private/NiagaraAttributeTrimmer.cpp b/Engine/Plugins/FX/Niagara/Source/NiagaraEditor/Private/NiagaraAttributeTrimmer.cpp
--- a/Engine/Plugins/FX/Niagara/Source/NiagaraEditor/Private/NiagaraAttributeTrimmer.cpp
+++ b/Engine/Plugins/FX/Niagara/Source/NiagaraEditor/Private/NiagaraAttributeTrimmer.cpp
@@ -5,7 +5,7 @@
 #include "Algo/ForEach.h"
 #include "Algo/RemoveIf.h"
 #include "EdGraphSchema_Niagara.h"
-#include "NiagaraCompilationGraphBridge.h"
+#include "NiagaraCompilationBridge.h"
 #include "NiagaraConstants.h"
 #include "NiagaraModule.h"
 #include "NiagaraNodeCustomHlsl.h"
@@ -22,23 +22,23 @@ class FNiagaraAttributeTrimmerHelper<GraphBridge>::FFunctionInputResolver
 public:
 	FFunctionInputResolver(const FGraph* Graph, ENiagaraScriptUsage Usage)
 	{
-		TMap<FName, FScopedPin> UnresolvedInputs;
+		TMap<FName, FModuleScopedPin> UnresolvedInputs;
 
 		Traverse(TEXT(""), Graph, Usage, UnresolvedInputs);
 	}
 
-	FScopedPin ResolveInput(const FScopedPin& Input) const
+	FModuleScopedPin ResolveInput(const FModuleScopedPin& Input) const
 	{
-		if (const FScopedPin* FunctionInput = FunctionInputMap.Find(Input))
+		if (const FModuleScopedPin* FunctionInput = FunctionInputMap.Find(Input))
 		{
 			return *FunctionInput;
 		}
 
-		return FScopedPin();
+		return FModuleScopedPin();
 	}
 
 private:
-	void Traverse(FString Namespace, const FGraph* Graph, ENiagaraScriptUsage Usage, TMap<FName, FScopedPin>& UnresolvedInputs)
+	void Traverse(FString Namespace, const FGraph* Graph, ENiagaraScriptUsage Usage, TMap<FName, FModuleScopedPin>& UnresolvedInputs)
 	{
 		const FString NamespacePrefix = (Namespace.Len() > 0) ? (Namespace + TEXT(".")) : TEXT("");
 
@@ -52,25 +52,25 @@ private:
 
 			for (const FNode* Node : TraversedNodes)
 			{
-				if (const FFunctionCallNode* FunctionNode = GraphBridge::template AsNodeType<FFunctionCallNode>(Node))
+				if (const FFunctionCallNode* FunctionNode = GraphBridge::AsFunctionCallNode(Node))
 				{
 					if (const FGraph* FunctionGraph = GraphBridge::GetFunctionNodeGraph(FunctionNode))
 					{
-						TMap<FName, FScopedPin> FunctionInputs;
+						TMap<FName, FModuleScopedPin> FunctionInputs;
 						Traverse(NamespacePrefix + GraphBridge::GetFunctionName(FunctionNode), FunctionGraph, GraphBridge::GetFunctionUsage(FunctionNode), FunctionInputs);
 
 						// resolve any of the UnresolvedInputs that we've accumulated
 						for (const FInputPin* NodePin : GraphBridge::GetInputPins(FunctionNode))
 						{
-							if (FScopedPin* InnerPin = FunctionInputs.Find(NodePin->PinName))
+							if (FModuleScopedPin* InnerPin = FunctionInputs.Find(NodePin->PinName))
 							{
-								FunctionInputMap.Add(*InnerPin, FScopedPin(NodePin, *Namespace));
+								FunctionInputMap.Add(*InnerPin, FModuleScopedPin(NodePin, *Namespace));
 								FunctionInputs.Remove(NodePin->PinName);
 							}
 						}
 					}
 				}
-				if (const FInputNode* InputNode = GraphBridge::template AsNodeType<FInputNode>(Node))
+				if (const FInputNode* InputNode = GraphBridge::AsInputNode(Node))
 				{
 					int32 OutputPinCount = 0;
 					for (const FOutputPin* InputNodePin : GraphBridge::GetOutputPins(InputNode))
@@ -78,14 +78,14 @@ private:
 						check(OutputPinCount == 0);
 						++OutputPinCount;
 
-						UnresolvedInputs.Add(GraphBridge::GetInputVariable(InputNode).GetName(), FScopedPin(InputNodePin, *Namespace));
+						UnresolvedInputs.Add(GraphBridge::GetInputVariable(InputNode).GetName(), FModuleScopedPin(InputNodePin, *Namespace));
 					}
 				}
 			}
 		}
 	}
 
-	TMap<FScopedPin, FScopedPin> FunctionInputMap;
+	TMap<FModuleScopedPin, FModuleScopedPin> FunctionInputMap;
 };
 
 template<typename GraphBridge>
@@ -97,7 +97,7 @@ public:
 		Traverse(TEXT(""), Graph, Usage);
 	}
 
-	const TArray<FScopedPin>& ReadFunctionInputs() const
+	const TArray<FModuleScopedPin>& ReadFunctionInputs() const
 	{
 		return ImpureFunctionInputs;
 	}
@@ -117,7 +117,7 @@ private:
 
 			for (const FNode* Node : TraversedNodes)
 			{
-				if (const FFunctionCallNode* FunctionNode = GraphBridge::template AsNodeType<FFunctionCallNode>(Node))
+				if (const FFunctionCallNode* FunctionNode = GraphBridge::AsFunctionCallNode(Node))
 				{
 					if (const FGraph* FunctionGraph = GraphBridge::GetFunctionNodeGraph(FunctionNode))
 					{
@@ -148,7 +148,7 @@ private:
 		}
 	}
 
-	TArray<FScopedPin> ImpureFunctionInputs;
+	TArray<FModuleScopedPin> ImpureFunctionInputs;
 };
 
 template<typename GraphBridge>
@@ -166,7 +166,7 @@ public:
 	//		An expression is the set of nodes that define what is being set in a MapSet.  It does not traverse through
 	//		Exec pins as those nodes would be part of another expression
 	//	-generates the list of custom HLSL nodes that are encountered in the traversal
-	void FindDependencies(const FScopedPin& EndPin, FDependencyChain& Dependencies)
+	void FindDependencies(const FModuleScopedPin& EndPin, FDependencyChain& Dependencies)
 	{
 		FindDependenciesInternal(EndPin, Dependencies, nullptr);
 	}
@@ -177,14 +177,14 @@ private:
 	//		An expression is the set of nodes that define what is being set in a MapSet.  It does not traverse through
 	//		Exec pins as those nodes would be part of another expression
 	//	-generates the list of custom HLSL nodes that are encountered in the traversal
-	void FindDependenciesInternal(const FScopedPin& EndPin, FDependencyChain& Dependencies, TArray<const FInputNode*>* InputNodes)
+	void FindDependenciesInternal(const FModuleScopedPin& EndPin, FDependencyChain& Dependencies, TArray<const FInputNode*>* InputNodes)
 	{
 		TStringBuilder<1024> NamespaceBuilder;
 
-		TArray<FScopedPin> PinsToEvaluate;
+		TArray<FModuleScopedPin> PinsToEvaluate;
 		TSet<const FNode*> EvaluatedNodes;
 
-		FScopedPin CurrentPin = EndPin;
+		FModuleScopedPin CurrentPin = EndPin;
 		while (CurrentPin.Pin)
 		{
 			if (CurrentPin.Pin->Direction == EEdGraphPinDirection::EGPD_Output)
@@ -199,11 +199,11 @@ private:
 				bool AlreadyEvaluated = false;
 
 				// record any output pins to a parameter get node
-				if (const FParamMapGetNode* ParamGetNode = GraphBridge::template AsNodeType<FParamMapGetNode>(CurrentNode))
+				if (const FParamMapGetNode* ParamGetNode = GraphBridge::AsParamMapGetNode(CurrentNode))
 				{
-					Dependencies.Pins.Emplace(FScopedPin(TracedPin, CurrentPin.ModuleName));
+					Dependencies.Pins.Emplace(FModuleScopedPin(TracedPin, CurrentPin.ModuleName));
 				}
-				else if (const FFunctionCallNode* FunctionCallNode = GraphBridge::template AsNodeType<FFunctionCallNode>(CurrentNode))
+				else if (const FFunctionCallNode* FunctionCallNode = GraphBridge::AsFunctionCallNode(CurrentNode))
 				{
 					if (const FGraph* NodeGraph = GraphBridge::GetFunctionNodeGraph(FunctionCallNode))
 					{
@@ -227,7 +227,7 @@ private:
 								if (InnerOutputNodePin->PinName == TracedPin->PinName)
 								{
 									TArray<const FInputNode*> InnerInputNodes;
-									FindDependenciesInternal(FScopedPin(InnerOutputNodePin, AggregateModuleName), Dependencies, &InnerInputNodes);
+									FindDependenciesInternal(FModuleScopedPin(InnerOutputNodePin, AggregateModuleName), Dependencies, &InnerInputNodes);
 
 									// map the InputNodes we found to the input pins on the FunctionCallNode
 									for (const FInputNode* InnerInputNode : InnerInputNodes)
@@ -250,7 +250,7 @@ private:
 						AlreadyEvaluated = true;
 					}
 				}
-				else if (const FInputNode* InputNode = GraphBridge::template AsNodeType<FInputNode>(CurrentNode))
+				else if (const FInputNode* InputNode = GraphBridge::AsInputNode(CurrentNode))
 				{
 					if (InputNodes)
 					{
@@ -258,7 +258,7 @@ private:
 					}
 					else
 					{
-						const FScopedPin ResolvedPin = InputResolver.ResolveInput(FScopedPin(TracedPin, CurrentPin.ModuleName));
+						const FModuleScopedPin ResolvedPin = InputResolver.ResolveInput(FModuleScopedPin(TracedPin, CurrentPin.ModuleName));
 						if (ResolvedPin.Pin)
 						{
 							PinsToEvaluate.Add(ResolvedPin);
@@ -288,7 +288,7 @@ private:
 						PinsToEvaluate.Emplace(NodePin, CurrentPin.ModuleName);
 					}
 
-					if (const FCustomHlslNode* CustomHlslNode = GraphBridge::template AsNodeType<FCustomHlslNode>(CurrentNode))
+					if (const FCustomHlslNode* CustomHlslNode = GraphBridge::AsCustomHlslNode(CurrentNode))
 					{
 						if (!Dependencies.CustomNodes.Contains(CustomHlslNode))
 						{
@@ -372,11 +372,11 @@ private:
 // For a specific read of a variable finds the corresponding PreviousWritePin if one exists (only considers actual writes
 // rather than default pins on a MapGet)
 template<typename GraphBridge>
-typename FNiagaraAttributeTrimmerHelper<GraphBridge>::FScopedPin FNiagaraAttributeTrimmerHelper<GraphBridge>::FindWriteForRead(const FParamMapHistory& ParamMap, const FScopedPin& ReadPin, int32* OutVariableIndex)
+typename FNiagaraAttributeTrimmerHelper<GraphBridge>::FModuleScopedPin FNiagaraAttributeTrimmerHelper<GraphBridge>::FindWriteForRead(const FParamMapHistory& ParamMap, const FModuleScopedPin& ReadPin, int32* OutVariableIndex)
 {
 	if (ReadPin.Pin == nullptr)
 	{
-		return FScopedPin();
+		return FModuleScopedPin();
 	}
 
 	int32 VariableCount = ParamMap.Variables.Num();
@@ -398,22 +398,22 @@ typename FNiagaraAttributeTrimmerHelper<GraphBridge>::FScopedPin FNiagaraAttribu
 				if (ReadHistoryEntry.PreviousWritePin.Pin)
 				{
 					// we only care about writes (MapSets).  This will exclude default variables that are set as the input to the MapGet
-					if (const FParamMapSetNode* WriteNode = GraphBridge::template AsNodeType<FParamMapSetNode>(GraphBridge::GetOwningNode(ReadHistoryEntry.PreviousWritePin.Pin)))
+					if (const FParamMapSetNode* WriteNode = GraphBridge::AsParamMapSetNode(GraphBridge::GetOwningNode(ReadHistoryEntry.PreviousWritePin.Pin)))
 					{
 						return ReadHistoryEntry.PreviousWritePin;
 					}
 				}
 
-				return FScopedPin();
+				return FModuleScopedPin();
 			}
 		}
 	}
 
-	return FScopedPin();
+	return FModuleScopedPin();
 }
 
 template<typename GraphBridge>
-int32 FNiagaraAttributeTrimmerHelper<GraphBridge>::FindDefaultBinding(const FParamMapHistory& ParamMap, int32 VariableIndex, const FScopedPin& ReadPin)
+int32 FNiagaraAttributeTrimmerHelper<GraphBridge>::FindDefaultBinding(const FParamMapHistory& ParamMap, int32 VariableIndex, const FModuleScopedPin& ReadPin)
 {
 	if (const FGraph* OwnerGraph = GraphBridge::GetOwningGraph(GraphBridge::GetOwningNode(ReadPin.Pin)))
 	{
@@ -438,7 +438,7 @@ int32 FNiagaraAttributeTrimmerHelper<GraphBridge>::FindDefaultBinding(const FPar
 
 // for a specific read of a variable finds the actual name of the attribute being read
 template<typename GraphBridge>
-FName FNiagaraAttributeTrimmerHelper<GraphBridge>::FindAttributeForRead(const FParamMapHistory& ParamMap, const FScopedPin& ReadPin)
+FName FNiagaraAttributeTrimmerHelper<GraphBridge>::FindAttributeForRead(const FParamMapHistory& ParamMap, const FModuleScopedPin& ReadPin)
 {
 	const int32 VariableCount = ParamMap.Variables.Num();
 
@@ -468,18 +468,18 @@ void FNiagaraAttributeTrimmerHelper<GraphBridge>::ResolveDependencyChain(const F
 	{
 		if (ParamMap.PerVariableWriteHistory[VariableIndex].Num())
 		{
-			const FScopedPin& LastWritePin = ParamMap.PerVariableWriteHistory[VariableIndex].Last();
+			const FModuleScopedPin& LastWritePin = ParamMap.PerVariableWriteHistory[VariableIndex].Last();
 
-			TArray<FScopedPin> PinsToResolve;
+			TArray<FModuleScopedPin> PinsToResolve;
 			PinsToResolve.Add(LastWritePin);
 
 			while (PinsToResolve.Num() > 0)
 			{
-				const FScopedPin PinToResolve = PinsToResolve.Pop(false);
+				const FModuleScopedPin PinToResolve = PinsToResolve.Pop(false);
 
 				if (const FDependencyChain* Dependencies = DependencyData.Find(PinToResolve))
 				{
-					for (const FScopedPin& ReadPin : Dependencies->Pins)
+					for (const FModuleScopedPin& ReadPin : Dependencies->Pins)
 					{
 						bool AlreadyAdded = false;
 						ResolvedDependencies.Pins.Add(ReadPin, &AlreadyAdded);
@@ -487,7 +487,7 @@ void FNiagaraAttributeTrimmerHelper<GraphBridge>::ResolveDependencyChain(const F
 						if (!AlreadyAdded)
 						{
 							int32 SourceVariableIndex = INDEX_NONE;
-							const FScopedPin WritePin = FindWriteForRead(ParamMap, ReadPin, &SourceVariableIndex);
+							const FModuleScopedPin WritePin = FindWriteForRead(ParamMap, ReadPin, &SourceVariableIndex);
 
 							if (WritePin.Pin)
 							{
@@ -509,7 +509,7 @@ void FNiagaraAttributeTrimmerHelper<GraphBridge>::ResolveDependencyChain(const F
 										const int32 VariableWriteCount = ParamMap.PerVariableWriteHistory[DefaultBoundVariable].Num();
 										for (int32 VariableWriteIt = VariableWriteCount - 1; VariableWriteIt >= 0; --VariableWriteIt)
 										{
-											const FScopedPin& DefaultWritePin = ParamMap.PerVariableWriteHistory[DefaultBoundVariable][VariableWriteIt];
+											const FModuleScopedPin& DefaultWritePin = ParamMap.PerVariableWriteHistory[DefaultBoundVariable][VariableWriteIt];
 											const FNode* DefaultWritePinOwningNode = GraphBridge::GetOwningNode(DefaultWritePin.Pin);
 											const int32 DefaultWriteNodeVisitationIndex = ParamMap.MapNodeVisitations.IndexOfByKey(DefaultWritePinOwningNode);
 											if (DefaultWriteNodeVisitationIndex != INDEX_NONE && DefaultWriteNodeVisitationIndex < ReadNodeVisitationIndex)
@@ -561,7 +561,7 @@ void FNiagaraAttributeTrimmerHelper<GraphBridge>::TrimAttributes_Safe(TConstArra
 	{
 		for (const FPin* Pin : ParamMap->MapPinHistory)
 		{
-			if (const FCustomHlslNode* CustomHlslNode = GraphBridge::template AsNodeType<FCustomHlslNode>(GraphBridge::GetOwningNode(Pin)))
+			if (const FCustomHlslNode* CustomHlslNode = GraphBridge::AsCustomHlslNode(GraphBridge::GetOwningNode(Pin)))
 			{
 				CustomHlslNodes.AddUnique(CustomHlslNode);
 			}
@@ -630,7 +630,7 @@ void FNiagaraAttributeTrimmerHelper<GraphBridge>::TrimAttributes_Aggressive(cons
 		{
 			const FNiagaraVariableBase& Var = ParamMap->Variables[VariableIt];
 
-			for (const FScopedPin& WritePin : ParamMap->PerVariableWriteHistory[VariableIt])
+			for (const FModuleScopedPin& WritePin : ParamMap->PerVariableWriteHistory[VariableIt])
 			{
 				FDependencyChain& Dependencies = PerVariableDependencySets.Add(WritePin);
 
@@ -648,14 +648,14 @@ void FNiagaraAttributeTrimmerHelper<GraphBridge>::TrimAttributes_Aggressive(cons
 		}
 
 		FImpureFunctionParser ImpureFunctionParser(GraphBridge::GetGraph(CompileDuplicateData), ParamMap->OriginatingScriptUsage);
-		for (const FScopedPin& RequiredFunctionInput : ImpureFunctionParser.ReadFunctionInputs())
+		for (const FModuleScopedPin& RequiredFunctionInput : ImpureFunctionParser.ReadFunctionInputs())
 		{
 			FDependencyChain Dependencies;
 
 			FExpressionBuilder Builder(*ParamMap, InputResolver);
 			Builder.FindDependencies(RequiredFunctionInput, Dependencies);
 
-			for (const FScopedPin& DependentPin : Dependencies.Pins)
+			for (const FModuleScopedPin& DependentPin : Dependencies.Pins)
 			{
 				// see if this variable corresponds to something in our parameter map's list of variables
 				// if it does then we need to mark it as something that needs to be preserved
@@ -709,7 +709,7 @@ void FNiagaraAttributeTrimmerHelper<GraphBridge>::TrimAttributes_Aggressive(cons
 			FDependencyChain ResolvedChain;
 			ResolveDependencyChain(*ParamMap, PerVariableDependencySets, AttributeName, ResolvedChain);
 
-			for (const FScopedPin& ParameterRead : ResolvedChain.Pins)
+			for (const FModuleScopedPin& ParameterRead : ResolvedChain.Pins)
 			{
 				const FName DependentAttributeName = FindAttributeForRead(*ParamMap, ParameterRead);
 
@@ -724,15 +724,15 @@ void FNiagaraAttributeTrimmerHelper<GraphBridge>::TrimAttributes_Aggressive(cons
 					}
 
 					// if we're dependent on a Previous value, then make sure we add the attribute to the list of things to be preserved as well
-					if (FParamMapHistory::IsPreviousName(DependentAttributeName))
+					if (FNiagaraParameterUtilities::IsPreviousName(DependentAttributeName))
 					{
-						ConditionalAddAttribute(FParamMapHistory::GetSourceForPreviousValue(DependentAttributeName));
+						ConditionalAddAttribute(FNiagaraParameterUtilities::GetSourceForPreviousValue(DependentAttributeName));
 					}
 
 					// likewise for Initial value
-					if (FParamMapHistory::IsInitialName(DependentAttributeName))
+					if (FNiagaraParameterUtilities::IsInitialName(DependentAttributeName))
 					{
-						ConditionalAddAttribute(FParamMapHistory::GetSourceForInitialValue(DependentAttributeName));
+						ConditionalAddAttribute(FNiagaraParameterUtilities::GetSourceForInitialValue(DependentAttributeName));
 					}
 				}
 			}
diff --git a/Engine/Plugins/FX/Niagara/Source/NiagaraEditor/Private/NiagaraAttributeTrimmer.h b/Engine/Plugins/FX/Niagara/Source/NiagaraEditor/Private/NiagaraAttributeTrimmer.h
--- a/Engine/Plugins/FX/Niagara/Source/NiagaraEditor/Private/NiagaraAttributeTrimmer.h
+++ b/Engine/Plugins/FX/Niagara/Source/NiagaraEditor/Private/NiagaraAttributeTrimmer.h
@@ -10,7 +10,6 @@ class FNiagaraAttributeTrimmerHelper : public GraphBridge
 {
 	using typename GraphBridge::FGraph;
 	using typename GraphBridge::FPin;
-	using typename GraphBridge::FScopedPin;
 	using typename GraphBridge::FCustomHlslNode;
 	using typename GraphBridge::FNode;
 	using typename GraphBridge::FOutputNode;
@@ -22,8 +21,9 @@ class FNiagaraAttributeTrimmerHelper : public GraphBridge
 	using typename GraphBridge::FOutputPin;
 	using typename GraphBridge::FParamMapHistory;
 	using typename GraphBridge::FCompilationCopy;
+	using typename GraphBridge::FModuleScopedPin;
 
-	using FDependencySet = TSet<FScopedPin>;
+	using FDependencySet = TSet<FModuleScopedPin>;
 	struct FCustomHlslNodeInfo
 	{
 		bool bHasDataInterfaceInputs = false;
@@ -38,7 +38,7 @@ class FNiagaraAttributeTrimmerHelper : public GraphBridge
 		FCustomHlslNodeMap CustomNodes;
 	};
 
-	using FDependencyMap = TMap<FScopedPin, FDependencyChain>;
+	using FDependencyMap = TMap<FModuleScopedPin, FDependencyChain>;
 
 	class FFunctionInputResolver;
 	class FImpureFunctionParser;
@@ -46,11 +46,11 @@ class FNiagaraAttributeTrimmerHelper : public GraphBridge
 
 	// For a specific read of a variable finds the corresponding PreviousWritePin if one exists (only considers actual writes
 	// rather than default pins on a MapGet)
-	static FScopedPin FindWriteForRead(const FParamMapHistory& ParamMap, const FScopedPin& ReadPin, int32* OutVariableIndex);
-	static int32 FindDefaultBinding(const FParamMapHistory& ParamMap, int32 VariableIndex, const FScopedPin& ReadPin);
+	static FModuleScopedPin FindWriteForRead(const FParamMapHistory& ParamMap, const FModuleScopedPin& ReadPin, int32* OutVariableIndex);
+	static int32 FindDefaultBinding(const FParamMapHistory& ParamMap, int32 VariableIndex, const FModuleScopedPin& ReadPin);
 
 	// for a specific read of a variable finds the actual name of the attribute being read
-	static FName FindAttributeForRead(const FParamMapHistory& ParamMap, const FScopedPin& ReadPin);
+	static FName FindAttributeForRead(const FParamMapHistory& ParamMap, const FModuleScopedPin& ReadPin);
 
 	// given the set of expressions (as defined in FindDependencies above) we resolve the named attribute aggregating the dependent reads and custom nodes
 	static void ResolveDependencyChain(const FParamMapHistory& ParamMap, const FDependencyMap& DependencyData, const FName& AttributeName, FDependencyChain& ResolvedDependencies);
diff --git a/Engine/Plugins/FX/Niagara/Source/NiagaraEditor/Private/NiagaraCompilationGraphBridge.h b/Engine/Plugins/FX/Niagara/Source/NiagaraEditor/Private/NiagaraCompilationGraphBridge.h
deleted file mode 100644
diff --git a/Engine/Plugins/FX/Niagara/Source/NiagaraEditor/Private/NiagaraCompilationGraphBridgeImpl.cpp b/Engine/Plugins/FX/Niagara/Source/NiagaraEditor/Private/NiagaraCompilationGraphBridgeImpl.cpp
new file mode 100644
--- /dev/null
+++ b/Engine/Plugins/FX/Niagara/Source/NiagaraEditor/Private/NiagaraCompilationGraphBridgeImpl.cpp
@@ -0,0 +1,538 @@
+// Copyright Epic Games, Inc. All Rights Reserved.
+
+#include "NiagaraCompilationBridge.h"
+
+#include "EdGraphSchema_Niagara.h"
+#include "NiagaraCompilationPrivate.h"
+#include "NiagaraConstants.h"
+#include "NiagaraModule.h"
+#include "NiagaraNodeConvert.h"
+#include "NiagaraNodeCustomHlsl.h"
+#include "NiagaraNodeEmitter.h"
+#include "NiagaraNodeFunctionCall.h"
+#include "NiagaraNodeIf.h"
+#include "NiagaraNodeInput.h"
+#include "NiagaraNodeOp.h"
+#include "NiagaraNodeOutput.h"
+#include "NiagaraNodeParameterMapFor.h"
+#include "NiagaraNodeParameterMapGet.h"
+#include "NiagaraNodeParameterMapSet.h"
+#include "NiagaraNodeSelect.h"
+#include "NiagaraScriptVariable.h"
+
+namespace NiagaraCompilationGraphBridgeImpl
+{
+
+template<typename ArrayType>
+static void AppendCompilationPins(const UNiagaraNode* Node, EEdGraphPinDirection PinDirection, ArrayType& OutPins)
+{
+	const UNiagaraNodeWithDynamicPins* DynNode = Cast<const UNiagaraNodeWithDynamicPins>(Node);
+	for (UEdGraphPin* Pin : Node->Pins)
+	{
+		if (Pin->Direction == PinDirection)
+		{
+			if (Pin->bOrphanedPin)
+			{
+				continue;
+			}
+
+			if (DynNode != nullptr && DynNode->IsAddPin(Pin))
+			{
+				continue;
+			}
+
+			OutPins.Add(Pin);
+		}
+	}
+}
+
+}
+
+
+FNiagaraCompilationGraphBridge::FBuilderExtraData::FBuilderExtraData()
+{
+	AvailableCollections = MakeUnique<FAvailableParameterCollections>();
+}
+
+void FNiagaraCompilationGraphBridge::FParameterCollectionStore::Append(const FParameterCollectionStore& Other)
+{
+	Collections.Append(Other.Collections);
+	CollectionVariables.Append(Other.CollectionVariables);
+	CollectionNamespaces.Append(Other.CollectionNamespaces);
+}
+
+const UNiagaraGraph* FNiagaraCompilationGraphBridge::GetGraph(const FCompilationCopy* CompilationCopy)
+{
+	return CompilationCopy->NodeGraphDeepCopy.Get();
+}
+
+const UNiagaraNode* FNiagaraCompilationGraphBridge::GetOwningNode(const FPin* Pin)
+{
+	return Cast<const FNode>(Pin->GetOwningNode());
+}
+
+UNiagaraNode* FNiagaraCompilationGraphBridge::GetMutableOwningNode(const FPin* Pin)
+{
+	return Cast<FNode>(Pin->GetOwningNode());
+}
+
+const UNiagaraGraph* FNiagaraCompilationGraphBridge::GetOwningGraph(const FNode* Node)
+{
+	if (const UNiagaraNode* NiagaraNode = Cast<const UNiagaraNode>(Node))
+	{
+		return NiagaraNode->GetNiagaraGraph();
+	}
+	return nullptr;
+}
+
+bool FNiagaraCompilationGraphBridge::CustomHlslReferencesTokens(const FCustomHlslNode* CustomNode, TConstArrayView<FStringView> TokenStrings)
+{
+	TArray<FString> StringTokens;
+	UNiagaraNodeCustomHlsl::GetTokensFromString(CustomNode->GetCustomHlsl(), StringTokens, false, false);
+
+	for (FStringView Token : TokenStrings)
+	{
+		const bool TokenMatched = StringTokens.ContainsByPredicate([&Token](const FString& HlslToken) -> bool
+			{
+				return HlslToken.Contains(Token);
+			});
+
+		if (TokenMatched)
+		{
+			return true;
+		}
+	}
+
+	return false;
+}
+
+ENiagaraScriptUsage FNiagaraCompilationGraphBridge::GetCustomHlslUsage(const FCustomHlslNode* CustomNode)
+{
+	return CustomNode->ScriptUsage;
+}
+
+FString FNiagaraCompilationGraphBridge::GetCustomHlslString(const FCustomHlslNode* CustomNode)
+{
+	return CustomNode->GetCustomHlsl();
+}
+
+void FNiagaraCompilationGraphBridge::GetCustomHlslIncludePaths(const FCustomHlslNode* CustomNode, TArray<FNiagaraCustomHlslInclude>& Includes)
+{
+	CustomNode->GetIncludeFilePaths(Includes);
+}
+
+const TArray<FNiagaraConvertConnection>& FNiagaraCompilationGraphBridge::GetConvertConnections(const FConvertNode* ConvertNode)
+{
+	return ConvertNode->GetConnections();
+}
+
+const UNiagaraNodeFunctionCall* FNiagaraCompilationGraphBridge::AsFunctionCallNode(const FNode* Node)
+{
+	return Cast<const UNiagaraNodeFunctionCall>(Node);
+}
+
+const UNiagaraNodeInput* FNiagaraCompilationGraphBridge::AsInputNode(const FNode* Node)
+{
+	return Cast<const UNiagaraNodeInput>(Node);
+}
+
+const UNiagaraNodeParameterMapGet* FNiagaraCompilationGraphBridge::AsParamMapGetNode(const FNode* Node)
+{
+	return Cast<const UNiagaraNodeParameterMapGet>(Node);
+}
+
+const UNiagaraNodeCustomHlsl* FNiagaraCompilationGraphBridge::AsCustomHlslNode(const FNode* Node)
+{
+	return Cast<const UNiagaraNodeCustomHlsl>(Node);
+}
+
+const UNiagaraNodeParameterMapSet* FNiagaraCompilationGraphBridge::AsParamMapSetNode(const FNode* Node)
+{
+	return Cast<const UNiagaraNodeParameterMapSet>(Node);
+}
+
+bool FNiagaraCompilationGraphBridge::GraphHasParametersOfType(const FGraph* Graph, const FNiagaraTypeDefinition& TypeDef)
+{
+	TArray<FNiagaraVariable> Inputs;
+	TArray<FNiagaraVariable> Outputs;
+	Graph->GetParameters(Inputs, Outputs);
+
+	auto ContainsVariableOfType = [TypeDef](const FNiagaraVariable& Variable) -> bool
+	{
+		return Variable.GetType() == TypeDef;
+	};
+
+	return Inputs.ContainsByPredicate(ContainsVariableOfType)
+		|| Outputs.ContainsByPredicate(ContainsVariableOfType);
+}
+
+TArray<FNiagaraVariableBase> FNiagaraCompilationGraphBridge::GraphGetStaticSwitchInputs(const FGraph* Graph)
+{
+	TArray<FNiagaraVariableBase> StaticSwitchInputs;
+	const TArray<FNiagaraVariable>& Variables = Graph->FindStaticSwitchInputs();
+
+	StaticSwitchInputs.Reserve(Variables.Num());
+	Algo::Transform(Variables, StaticSwitchInputs, [](const FNiagaraVariable& InVar) -> FNiagaraVariableBase
+		{
+			return InVar;
+		});
+
+	return StaticSwitchInputs;
+}
+
+void FNiagaraCompilationGraphBridge::FindOutputNodes(const FGraph* Graph, ENiagaraScriptUsage ScriptUsage, TArray<const FOutputNode*>& OutputNodes)
+{
+	TArray<FOutputNode*> MutableOutputNodes;
+	Graph->FindOutputNodes(ScriptUsage, MutableOutputNodes);
+	OutputNodes.Reserve(MutableOutputNodes.Num());
+	Algo::Transform(MutableOutputNodes, OutputNodes, [](FOutputNode* OutputNode) -> const FOutputNode*
+		{
+			return OutputNode;
+		});
+}
+
+void FNiagaraCompilationGraphBridge::FindOutputNodes(const FGraph* Graph, TArray<const FOutputNode*>& OutputNodes)
+{
+	TArray<FOutputNode*> MutableOutputNodes;
+	Graph->FindOutputNodes(MutableOutputNodes);
+	OutputNodes.Reserve(MutableOutputNodes.Num());
+	Algo::Transform(MutableOutputNodes, OutputNodes, [](FOutputNode* OutputNode) -> const FOutputNode*
+		{
+			return OutputNode;
+		});
+}
+
+void FNiagaraCompilationGraphBridge::BuildTraversal(const FGraph* Graph, const FNode* OutputNode, TArray<const FNode*>& TraversedNodes)
+{
+	TArray<UNiagaraNode*> MutableTraversedNodes;
+	Graph->BuildTraversal(MutableTraversedNodes, const_cast<FNode*>(OutputNode), true);
+	TraversedNodes.Reserve(MutableTraversedNodes.Num());
+	Algo::Transform(MutableTraversedNodes, TraversedNodes, [](UNiagaraNode* Node) -> const FNode*
+		{
+			return Node;
+		});
+}
+
+const UNiagaraGraph* FNiagaraCompilationGraphBridge::GetEmitterGraph(const FEmitterNode* EmitterNode)
+{
+	return EmitterNode->GetCalledGraph();
+}
+
+FString FNiagaraCompilationGraphBridge::GetEmitterUniqueName(const FEmitterNode* EmitterNode)
+{
+	return EmitterNode->GetEmitterUniqueName();
+}
+
+ENiagaraScriptUsage FNiagaraCompilationGraphBridge::GetEmitterUsage(const FEmitterNode* EmitterNode)
+{
+	return EmitterNode->GetUsage();
+}
+
+FString FNiagaraCompilationGraphBridge::GetEmitterName(const FEmitterNode* EmitterNode)
+{
+	return EmitterNode->GetName();
+}
+
+FString FNiagaraCompilationGraphBridge::GetEmitterPathName(const FEmitterNode* EmitterNode)
+{
+	return EmitterNode->GetFullName();
+}
+
+FString FNiagaraCompilationGraphBridge::GetEmitterHandleIdString(const FEmitterNode* EmitterNode)
+{
+	return EmitterNode->GetEmitterHandleId().ToString(EGuidFormats::Digits);
+}
+
+const UNiagaraGraph* FNiagaraCompilationGraphBridge::GetFunctionNodeGraph(const FFunctionCallNode* FunctionCall)
+{
+	return FunctionCall->GetCalledGraph();
+}
+
+FString FNiagaraCompilationGraphBridge::GetFunctionFullName(const FFunctionCallNode* FunctionCall)
+{
+	return FunctionCall->FunctionScript ? FunctionCall->FunctionScript->GetFullName() : FString();
+}
+
+FString FNiagaraCompilationGraphBridge::GetFunctionScriptName(const FFunctionCallNode* FunctionCall)
+{
+	return FunctionCall->FunctionScript ? FunctionCall->FunctionScript->GetName() : FString();
+}
+
+FString FNiagaraCompilationGraphBridge::GetFunctionName(const FFunctionCallNode* FunctionCall)
+{
+	return FunctionCall ? FunctionCall->GetFunctionName() : FString();
+}
+
+ENiagaraScriptUsage FNiagaraCompilationGraphBridge::GetFunctionUsage(const FFunctionCallNode* FunctionCall)
+{
+	return FunctionCall->GetCalledUsage();
+}
+
+TOptional<ENiagaraDefaultMode> FNiagaraCompilationGraphBridge::GetGraphDefaultMode(const FGraph* Graph, const FNiagaraVariable& Variable, FNiagaraScriptVariableBinding& Binding)
+{
+	return Graph->GetDefaultMode(Variable, &Binding);
+}
+
+const UEdGraphPin* FNiagaraCompilationGraphBridge::GetDefaultPin(const FParamMapGetNode* GetNode, const FOutputPin* OutputPin)
+{
+	return GetNode->GetDefaultPin(const_cast<FOutputPin*>(OutputPin));
+}
+
+bool FNiagaraCompilationGraphBridge::IsStaticPin(const FPin* Pin)
+{
+	return UEdGraphSchema_Niagara::IsStaticPin(Pin);
+}
+
+// retrieves all input pins (excluding any add pins that may be present)
+TArray<const UEdGraphPin*> FNiagaraCompilationGraphBridge::GetInputPins(const FNode* Node)
+{
+	TArray<const FInputPin*> InputPins;
+	NiagaraCompilationGraphBridgeImpl::AppendCompilationPins(Node, EEdGraphPinDirection::EGPD_Input, InputPins);
+	return InputPins;
+}
+
+// retrieves all output pins (excluding both orphaned pins and add pins)
+TArray<const UEdGraphPin*> FNiagaraCompilationGraphBridge::GetOutputPins(const FNode* Node)
+{
+	TArray<const FInputPin*> OutputPins;
+	NiagaraCompilationGraphBridgeImpl::AppendCompilationPins(Node, EEdGraphPinDirection::EGPD_Output, OutputPins);
+	return OutputPins;
+}
+
+// gets all pins assoicated with the node
+TArray<const UEdGraphPin*> FNiagaraCompilationGraphBridge::GetPins(const FNode* Node)
+{
+	TArray<const FPin*> Pins;
+	Pins.Reserve(Node->Pins.Num());
+	Algo::Transform(Node->Pins, Pins, [](FPin* InPin) -> const FPin*
+		{
+			return InPin;
+		});
+	return Pins;
+}
+
+FNiagaraTypeDefinition FNiagaraCompilationGraphBridge::GetPinType(const FPin* Pin, ENiagaraStructConversion Conversion)
+{
+	return UEdGraphSchema_Niagara::PinToTypeDefinition(Pin, Conversion);
+}
+
+FText FNiagaraCompilationGraphBridge::GetPinFriendlyName(const FPin* Pin)
+{
+	return Pin->PinFriendlyName;
+}
+
+FText FNiagaraCompilationGraphBridge::GetPinDisplayName(const FPin* Pin)
+{
+	return Pin->GetDisplayName();
+}
+
+FNiagaraVariable FNiagaraCompilationGraphBridge::GetPinVariable(const FPin* Pin, bool bNeedsValue, ENiagaraStructConversion Conversion)
+{
+	return UEdGraphSchema_Niagara::PinToNiagaraVariable(Pin, bNeedsValue, Conversion);
+}
+
+const UEdGraphPin* FNiagaraCompilationGraphBridge::GetPinAsInput(const FPin* Pin)
+{
+	if (Pin && Pin->Direction == EGPD_Input)
+	{
+		return Pin;
+	}
+	return nullptr;
+}
+
+FNiagaraVariable FNiagaraCompilationGraphBridge::GetInputVariable(const FInputNode* InputNode)
+{
+	return InputNode->Input;
+}
+
+const TArray<FNiagaraVariable>& FNiagaraCompilationGraphBridge::GetOutputVariables(const FOutputNode* OutputNode)
+{
+	return OutputNode->GetOutputs();
+}
+
+TArray<FNiagaraVariable> FNiagaraCompilationGraphBridge::GetGraphOutputNodeVariables(const FGraph* Graph, ENiagaraScriptUsage Usage)
+{
+	TArray<FNiagaraVariable> OutputNodeVariables;
+	Graph->GetOutputNodeVariables(Usage, OutputNodeVariables);
+	return OutputNodeVariables;
+}
+
+TArray<const UNiagaraNodeInput*> FNiagaraCompilationGraphBridge::GetGraphInputNodes(const FGraph* Graph, const UNiagaraGraph::FFindInputNodeOptions& Options)
+{
+	TArray<FInputNode*> MutableInputNodes;
+	Graph->FindInputNodes(MutableInputNodes, Options);
+
+	TArray<const FInputNode*> InputNodes;
+	InputNodes.Reserve(MutableInputNodes.Num());
+
+	Algo::Transform(MutableInputNodes, InputNodes, [](FInputNode* InputNode) -> const FInputNode*
+		{
+			return InputNode;
+		});
+
+	return InputNodes;
+}
+
+const UEdGraphPin* FNiagaraCompilationGraphBridge::GetLinkedOutputPin(const FInputPin* InputPin)
+{
+	if (InputPin && !InputPin->LinkedTo.IsEmpty())
+	{
+		return InputPin->LinkedTo[0];
+	}
+	return nullptr;
+}
+
+bool FNiagaraCompilationGraphBridge::CanCreateConnection(const FOutputPin* OutputPin, const FInputPin* InputPin, FText& FailureMessage)
+{
+	FPinConnectionResponse PinResponse = GetDefault<UEdGraphSchema_Niagara>()->CanCreateConnection(OutputPin, InputPin);
+	if (PinResponse.Response == ECanCreateConnectionResponse::CONNECT_RESPONSE_DISALLOW)
+	{
+		FailureMessage = PinResponse.Message;
+		return false;
+	}
+
+	return true;
+}
+
+ENiagaraScriptUsage FNiagaraCompilationGraphBridge::GetOutputNodeUsage(const FOutputNode* OutputNode)
+{
+	return OutputNode->GetUsage();
+}
+
+FGuid FNiagaraCompilationGraphBridge::GetOutputNodeUsageId(const FOutputNode* OutputNode)
+{
+	return OutputNode->GetUsageId();
+}
+
+ENiagaraScriptUsage FNiagaraCompilationGraphBridge::GetOutputNodeScriptType(const FOutputNode* OutputNode)
+{
+	return OutputNode->ScriptType;
+}
+
+FGuid FNiagaraCompilationGraphBridge::GetOutputNodeScriptTypeId(const FOutputNode* OutputNode)
+{
+	return OutputNode->ScriptTypeId;
+}
+
+bool FNiagaraCompilationGraphBridge::IsGraphEmpty(const FGraph* Graph)
+{
+	return Graph->IsEmpty();
+}
+
+void FNiagaraCompilationGraphBridge::AddCollectionPaths(const FParamMapHistory& History, TArray<FString>& Paths)
+{
+	for (UNiagaraParameterCollection* Collection : History.EncounteredParameterCollections.Collections)
+	{
+		Paths.AddUnique(FSoftObjectPath(Collection).ToString());
+	}
+}
+
+bool FNiagaraCompilationGraphBridge::NodeIsEnabled(const FNode* Node)
+{
+	return Node->IsNodeEnabled();
+}
+
+TOptional<ENiagaraDefaultMode> FNiagaraCompilationGraphBridge::GraphGetDefaultMode(const FGraph* Graph, const FNiagaraVariableBase& Variable, FNiagaraScriptVariableBinding& Binding)
+{
+	return Graph->GetDefaultMode(Variable, &Binding);
+}
+
+const UEdGraphPin* FNiagaraCompilationGraphBridge::GetSelectOutputPin(const FSelectNode* SelectNode, const FNiagaraVariableBase& Variable)
+{
+	return SelectNode->GetOutputPin(Variable);
+}
+
+FString FNiagaraCompilationGraphBridge::GetNodeName(const FNode* Node)
+{
+	return Node->GetName();
+}
+
+FString FNiagaraCompilationGraphBridge::GetNodeTitle(const FNode* Node)
+{
+	return Node->GetNodeTitle(ENodeTitleType::FullTitle).ToString();
+}
+
+const UEdGraphPin* FNiagaraCompilationGraphBridge::GetInputPin(const FNode* Node, int32 PinIndex)
+{
+	return Node->GetInputPin(PinIndex);
+}
+
+int32 FNiagaraCompilationGraphBridge::GetPinIndexById(TConstArrayView<const FPin*> Pins, const FGuid& PinId)
+{
+	const int32 PinCount = Pins.Num();
+	for (int32 PinIt = 0; PinIt < PinCount; ++PinIt)
+	{
+		if (Pins[PinIt]->PinId == PinId)
+		{
+			return PinIt;
+		}
+	}
+	return INDEX_NONE;
+}
+
+FString FNiagaraCompilationGraphBridge::GetCollectionFullName(FParameterCollection Collection)
+{
+	return GetFullNameSafe(Collection);
+}
+
+bool FNiagaraCompilationGraphBridge::IsCollectionValid(FParameterCollection Collection)
+{
+	return IsValid(Collection);
+}
+
+UNiagaraDataInterface* FNiagaraCompilationGraphBridge::GetCollectionDataInterface(FParameterCollection Collection, const FNiagaraVariable& Variable)
+{
+	return Collection->GetDefaultInstance()->GetParameterStore().GetDataInterface(Variable);
+}
+
+UObject* FNiagaraCompilationGraphBridge::GetCollectionUObject(FParameterCollection Collection, const FNiagaraVariable& Variable)
+{
+	return Collection->GetDefaultInstance()->GetParameterStore().GetUObject(Variable);
+}
+
+const UNiagaraNodeOutput* FNiagaraCompilationGraphBridge::AsOutputNode(const FNode* Node)
+{
+	return Cast<const FOutputNode>(Node);
+}
+
+bool FNiagaraCompilationGraphBridge::IsParameterMapGet(const FNode* Node)
+{
+	return Node->IsA<UNiagaraNodeParameterMapGet>();
+}
+
+TOptional<FName> FNiagaraCompilationGraphBridge::GetOutputNodeStackContextOverride(const FOutputNode* OutputNode)
+{
+	return OutputNode->GetStackContextOverride();
+}
+
+FString FNiagaraCompilationGraphBridge::GetNodeClassName(const FNode* Node)
+{
+	return Node->GetClass()->GetName();
+}
+
+bool FNiagaraCompilationGraphBridge::IsParameterMapPin(const FPin* Pin)
+{
+	return UEdGraphSchema_Niagara::PinToTypeDefinition(Pin) == FNiagaraTypeDefinition::GetParameterMapDef();
+}
+
+bool FNiagaraCompilationGraphBridge::GetGraphReferencesStaticVariables(const FGraph* Graph, FNiagaraStaticVariableSearchContext& StaticVariableContext)
+{
+	return Graph->ReferencesStaticVariable(StaticVariableContext);
+}
+
+const UNiagaraNodeEmitter* FNiagaraCompilationGraphBridge::GetNodeAsEmitter(const FNode* Node)
+{
+	return Cast<const UNiagaraNodeEmitter>(Node);
+}
+
+UNiagaraParameterCollection* FNiagaraCompilationGraphBridge::FAvailableParameterCollections::FindCollection(const FNiagaraVariable& Variable) const
+{
+	const UEdGraphSchema_Niagara* Schema = GetDefault<UEdGraphSchema_Niagara>();
+	return Schema->VariableIsFromParameterCollection(Variable);
+}
+
+UNiagaraParameterCollection* FNiagaraCompilationGraphBridge::FAvailableParameterCollections::FindMatchingCollection(FName VariableName, bool bAllowPartialMatch, FNiagaraVariable& OutVar) const
+{
+	const UEdGraphSchema_Niagara* Schema = GetDefault<UEdGraphSchema_Niagara>();
+	return Schema->VariableIsFromParameterCollection(VariableName.ToString(), true, OutVar);
+}
diff --git a/Engine/Plugins/FX/Niagara/Source/NiagaraEditor/Private/NiagaraCompilationPrivate.h b/Engine/Plugins/FX/Niagara/Source/NiagaraEditor/Private/NiagaraCompilationPrivate.h
--- a/Engine/Plugins/FX/Niagara/Source/NiagaraEditor/Private/NiagaraCompilationPrivate.h
+++ b/Engine/Plugins/FX/Niagara/Source/NiagaraEditor/Private/NiagaraCompilationPrivate.h
@@ -11,18 +11,12 @@
 #include "NiagaraTypes.h"
 #include "Templates/SharedPointer.h"
 
-class FNiagaraFixedConstantResolver;
+class FCompileConstantResolver;
 class FNiagaraCompileRequestData;
 class FNiagaraCompileRequestDuplicateData;
+class UNiagaraGraph;
 struct FNiagaraGraphCachedBuiltHistory;
 class UNiagaraScriptSourceBase;
-class UNiagaraGraph;
-class FNiagaraCompilationGraph;
-class FCompileConstantResolver;
-class FNiagaraParameterMapHistoryBuilder;
-class FNiagaraCompilationNodeEmitter;
-enum class ENiagaraFunctionDebugState : uint8;
-template<typename GraphBridge> class TNiagaraHlslTranslator;
 
 struct FNiagaraSimulationStageInfo
 {
diff --git a/Engine/Plugins/FX/Niagara/Source/NiagaraEditor/Private/NiagaraCompiler.cpp b/Engine/Plugins/FX/Niagara/Source/NiagaraEditor/Private/NiagaraCompiler.cpp
--- a/Engine/Plugins/FX/Niagara/Source/NiagaraEditor/Private/NiagaraCompiler.cpp
+++ b/Engine/Plugins/FX/Niagara/Source/NiagaraEditor/Private/NiagaraCompiler.cpp
@@ -339,7 +339,7 @@ void FNiagaraCompileRequestData::SortOutputNodesByDependencies(TArray<class UNia
 
 FName FNiagaraCompileRequestData::ResolveEmitterAlias(FName VariableName) const
 {
-	return FNiagaraParameterMapHistory::ResolveEmitterAlias(VariableName, EmitterUniqueName);
+	return FNiagaraParameterUtilities::ResolveEmitterAlias(VariableName, EmitterUniqueName);
 }
 
 bool FNiagaraCompileRequestDuplicateData::IsDuplicateDataFor(UNiagaraSystem* InSystem, UNiagaraEmitter* InEmitter, UNiagaraScript* InScript) const
@@ -397,7 +397,7 @@ void FNiagaraCompileRequestData::GatherPreCompiledVariables(const FString& InNam
 	{
 		for (const FNiagaraVariable& EncounteredVariable : EncounteredVariables)
 		{
-			if (FNiagaraParameterMapHistory::IsInNamespace(EncounteredVariable, InNamespaceFilter))
+			if (FNiagaraParameterUtilities::IsInNamespace(EncounteredVariable, InNamespaceFilter))
 			{
 				FNiagaraVariable NewVar = EncounteredVariable;
 				if (NewVar.IsDataAllocated() == false && !NewVar.IsDataInterface() && !NewVar.IsUObject())
@@ -643,8 +643,8 @@ void FNiagaraCompileRequestData::FinishPrecompile(const TArray<FNiagaraVariable>
 			if (bStageEnabled)
 			{
 				// Map all for this output node
-				FNiagaraParameterMapHistoryWithMetaDataBuilder Builder;
-				Builder.ConstantResolver = ConstantResolver;
+				TNiagaraParameterMapHistoryWithMetaDataBuilder<FNiagaraCompilationGraphBridge> Builder;
+				*Builder.ConstantResolver = ConstantResolver;
 				Builder.AddGraphToCallingGraphContextStack(Source->NodeGraph);
 				Builder.RegisterEncounterableVariables(EncounterableVariables);
 				Builder.RegisterExternalStaticVariables(StaticVariables);
@@ -783,7 +783,7 @@ void FNiagaraCompileRequestData::FinishPrecompile(const TArray<FNiagaraVariable>
 					{
 						int32 VariableIndex = DataInterfaceVariableIndices[i];
 						const FNiagaraVariable& Variable = History.Variables[VariableIndex];
-						for (const FModuleScopedPin& WritePin : History.PerVariableWriteHistory[VariableIndex])
+						for (const FNiagaraParameterMapHistory::FModuleScopedPin& WritePin : History.PerVariableWriteHistory[VariableIndex])
 						{
 							if (WritePin.Pin != nullptr && WritePin.Pin->LinkedTo.Num() == 1 && WritePin.Pin->LinkedTo[0] != nullptr)
 							{
@@ -884,7 +884,7 @@ void FNiagaraCompileRequestDuplicateData::FinishPrecompileDuplicate(const TArray
 		{
 			// Map all for this output node
 			FNiagaraParameterMapHistoryWithMetaDataBuilder Builder;
-			Builder.ConstantResolver = ConstantResolver;
+			*Builder.ConstantResolver = ConstantResolver;
 			Builder.AddGraphToCallingGraphContextStack(NodeGraphDeepCopy.Get());
 			Builder.RegisterEncounterableVariables(EncounterableVariables);
 			Builder.RegisterExternalStaticVariables(InStaticVariables);
@@ -1587,7 +1587,7 @@ TSharedPtr<FNiagaraGraphCachedDataBase, ESPMode::ThreadSafe> FNiagaraEditorModul
 
 			// Map all for this output node
 			FNiagaraParameterMapHistoryWithMetaDataBuilder Builder;
-			Builder.ConstantResolver = ConstantResolver;
+			*Builder.ConstantResolver = ConstantResolver;
 			if (ScriptSource != nullptr)
 			{
 				Builder.AddGraphToCallingGraphContextStack(ScriptSource->NodeGraph);
diff --git a/Engine/Plugins/FX/Niagara/Source/NiagaraEditor/Private/NiagaraGraph.cpp b/Engine/Plugins/FX/Niagara/Source/NiagaraEditor/Private/NiagaraGraph.cpp
--- a/Engine/Plugins/FX/Niagara/Source/NiagaraEditor/Private/NiagaraGraph.cpp
+++ b/Engine/Plugins/FX/Niagara/Source/NiagaraEditor/Private/NiagaraGraph.cpp
@@ -3923,7 +3923,7 @@ void UNiagaraGraph::RefreshParameterReferences() const
 		for (int32 Index = 0; Index < History.VariablesWithOriginalAliasesIntact.Num(); Index++)
 		{
 			const FNiagaraVariable& Parameter = History.VariablesWithOriginalAliasesIntact[Index];
-			for (const FModuleScopedPin& WriteEvent : History.PerVariableWriteHistory[Index])
+			for (const FNiagaraParameterMapHistory::FModuleScopedPin& WriteEvent : History.PerVariableWriteHistory[Index])
 			{
 				if (WriteEvent.Pin->PinType.PinSubCategory == UNiagaraNodeParameterMapBase::ParameterPinSubCategory)
 				{
diff --git a/Engine/Plugins/FX/Niagara/Source/NiagaraEditor/Private/NiagaraGraphHlslTranslator.h b/Engine/Plugins/FX/Niagara/Source/NiagaraEditor/Private/NiagaraGraphHlslTranslator.h
deleted file mode 100644
diff --git a/Engine/Plugins/FX/Niagara/Source/NiagaraEditor/Private/NiagaraHlslTranslator.cpp b/Engine/Plugins/FX/Niagara/Source/NiagaraEditor/Private/NiagaraHlslTranslator.cpp
--- a/Engine/Plugins/FX/Niagara/Source/NiagaraEditor/Private/NiagaraHlslTranslator.cpp
+++ b/Engine/Plugins/FX/Niagara/Source/NiagaraEditor/Private/NiagaraHlslTranslator.cpp
@@ -7,6 +7,7 @@
 #include "Modules/ModuleManager.h"
 #include "NiagaraAttributeTrimmer.h"
 #include "NiagaraCommon.h"
+#include "NiagaraCompilationPrivate.h"
 #include "NiagaraConstants.h"
 #include "NiagaraDataInterface.h"
 #include "NiagaraDataInterfaceVector2DCurve.h"
@@ -17,7 +18,6 @@
 #include "NiagaraFunctionLibrary.h"
 #include "NiagaraNodeEmitter.h"
 #include "NiagaraGraph.h"
-#include "NiagaraGraphHlslTranslator.h"
 #include "NiagaraNode.h"
 #include "NiagaraNodeConvert.h"
 #include "NiagaraNodeCustomHlsl.h"
@@ -673,9 +673,9 @@ void TNiagaraHlslTranslator<GraphBridge>::BuildMissingDefaults()
 			// Now go through and initialize any "Particles.Initial." variables
 			for (FNiagaraVariable& Var : InitialNamespaceVariablesMissingDefault)
 			{
-				if (FParamMapHistory::IsInitialValue(Var))
+				if (FNiagaraParameterUtilities::IsInitialValue(Var))
 				{
-					FNiagaraVariable SourceForInitialValue = FParamMapHistory::GetSourceForInitialValue(Var);
+					FNiagaraVariable SourceForInitialValue = FNiagaraParameterUtilities::GetSourceForInitialValue(Var);
 					FString ParameterMapInstanceName = GetParameterMapInstanceName(0);
 					FString Value = FString::Printf(TEXT("%s.%s = %s.%s;\n"), *ParameterMapInstanceName, *GetSanitizedSymbolName(Var.GetName().ToString()),
 						*ParameterMapInstanceName, *GetSanitizedSymbolName(SourceForInitialValue.GetName().ToString()));
@@ -717,7 +717,7 @@ FString TNiagaraHlslTranslator<GraphBridge>::BuildParameterMapHlslDefinitions(TA
 		}
 
 		// ignore those that are rapid iteration parameters as those will be read in directly from the cbuffer
-		if (FParamMapHistory::IsRapidIterationParameter(Var))
+		if (FNiagaraParameterUtilities::IsRapidIterationParameter(Var))
 		{
 			continue;
 		}
@@ -782,9 +782,9 @@ FString TNiagaraHlslTranslator<GraphBridge>::BuildParameterMapHlslDefinitions(TA
 		for (const FNiagaraVariableBase& Var : CompileOptions.AdditionalVariables)
 		{
 			bool bFoundSource = false;
-			if (FParamMapHistory::IsPreviousValue(Var))
+			if (FNiagaraParameterUtilities::IsPreviousValue(Var))
 			{
-				FNiagaraVariable Source = FParamMapHistory::GetSourceForPreviousValue(FNiagaraVariable(Var));
+				FNiagaraVariable Source = FNiagaraParameterUtilities::GetSourceForPreviousValue(FNiagaraVariable(Var));
 				const FNiagaraTypeDefinition& SourceType = Source.GetType();
 
 				for (int32 ParamMapIdxTest = 0; !bFoundSource && ParamMapIdxTest < OtherOutputParamMapHistories.Num(); ParamMapIdxTest++)
@@ -845,7 +845,7 @@ FString TNiagaraHlslTranslator<GraphBridge>::BuildParameterMapHlslDefinitions(TA
 	for (int32 UniqueVarIdx = 0; UniqueVarIdx < UniqueVariables.Num(); UniqueVarIdx++)
 	{
 		const FNiagaraVariable& NiagaraVariable = UniqueVariables[UniqueVarIdx];
-		if (FParamMapHistory::IsAttribute(NiagaraVariable))
+		if (FNiagaraParameterUtilities::IsAttribute(NiagaraVariable))
 		{
 			FString VariableName = GetSanitizedSymbolName(NiagaraVariable.GetName().ToString());
 			ReplaceNamespaceInline(VariableName, PARAM_MAP_ATTRIBUTE_STR, PARAM_MAP_INDICES_STR);
@@ -974,7 +974,7 @@ void FNiagaraHlslTranslator::HandleNamespacedExternalVariablesToDataSetRead(TArr
 {
 	for (const FNiagaraVariable& Var : ExternalVariablesForBulkUsage)
 	{
-		if (FNiagaraParameterMapHistory::IsInNamespace(Var, InNamespaceStr))
+		if (FNiagaraParameterUtilities::IsInNamespace(Var, InNamespaceStr))
 		{
 			InDataSetVars.Add(Var);
 		}
@@ -1102,9 +1102,9 @@ static void ConvertFloatToHalf(const FNiagaraCompileOptions& InCompileOptions, T
 
 			// also we'll check if the current attribute is a previous version of an exception because we wouldn't want
 			// those to be mismatched
-			if (FNiagaraParameterMapHistory::IsPreviousValue(Attribute))
+			if (FNiagaraParameterUtilities::IsPreviousValue(Attribute))
 			{
-				FNiagaraVariable SrcAttribute = FNiagaraParameterMapHistory::GetSourceForPreviousValue(Attribute);
+				FNiagaraVariable SrcAttribute = FNiagaraParameterUtilities::GetSourceForPreviousValue(Attribute);
 
 				if (FNiagaraVariable::SearchArrayForPartialNameMatch(ConvertExceptions, SrcAttribute.GetName()) != INDEX_NONE)
 				{
@@ -1357,15 +1357,15 @@ FNiagaraTranslateResults TNiagaraHlslTranslator<GraphBridge>::Translate(const FN
 					for (int32 iVar = 0; iVar < FoundHistory.VariableMetaData.Num(); ++iVar)
 					{
 						// Particle attribute?
-						if (!FParamMapHistory::IsAttribute(FoundHistory.Variables[iVar]))
+						if (!FNiagaraParameterUtilities::IsAttribute(FoundHistory.Variables[iVar]))
 						{
 							continue;
 						}
 
 						// Is this an output?
-						const bool bIsOutput = FoundHistory.PerVariableWriteHistory[iVar].ContainsByPredicate([](const FScopedPin& InPin) -> bool
+						const bool bIsOutput = FoundHistory.PerVariableWriteHistory[iVar].ContainsByPredicate([](const FModuleScopedPin& InPin) -> bool
 						{
-							return GraphBridge::template AsNodeType<FParamMapSetNode>(GraphBridge::GetOwningNode(InPin.Pin)) != nullptr;
+							return GraphBridge::AsParamMapSetNode(GraphBridge::GetOwningNode(InPin.Pin)) != nullptr;
 						});
 
 						if (!bIsOutput)
@@ -2094,9 +2094,9 @@ FNiagaraTranslateResults TNiagaraHlslTranslator<GraphBridge>::Translate(const FN
 		//Get full list of instance data accessed by the script as the VM binding assumes same for input and output.
 		for (FNiagaraVariable& Var : DataSetVariables[InstanceReadVarsIndex])
 		{
-			if (FParamMapHistory::IsAttribute(Var))
+			if (FNiagaraParameterUtilities::IsAttribute(Var))
 			{
-				FNiagaraVariable BasicAttribVar = FParamMapHistory::ResolveAsBasicAttribute(Var, false);
+				FNiagaraVariable BasicAttribVar = FNiagaraParameterUtilities::ResolveAsBasicAttribute(Var, false);
 				CompilationOutput.ScriptData.Attributes.AddUnique(BasicAttribVar);
 			}
 			else
@@ -2648,11 +2648,11 @@ void FNiagaraHlslTranslator::DefinePreviousParametersFunction(FString& HlslOutpu
 			const TArray<FNiagaraVariable>& NiagaraVariables = DataSetVariables[DataSetWrites[DataSetID]];
 			for (const FNiagaraVariable& Var : NiagaraVariables)
 			{
-				if (FNiagaraParameterMapHistory::IsPreviousValue(Var))
+				if (FNiagaraParameterUtilities::IsPreviousValue(Var))
 				{
 					FString CurMap = TranslationStages[0].PassNamespace;
 					FString Value = FString::Printf(TEXT("Context.%s.%s = Context.%s.%s;\n"), *CurMap, *GetSanitizedSymbolName(Var.GetName().ToString()),
-						*CurMap, *GetSanitizedSymbolName(FNiagaraParameterMapHistory::GetSourceForPreviousValue(Var).GetName().ToString()));
+						*CurMap, *GetSanitizedSymbolName(FNiagaraParameterUtilities::GetSourceForPreviousValue(Var).GetName().ToString()));
 					HlslOutputString += Value + TEXT("\n");
 				}
 			}
@@ -3188,7 +3188,7 @@ void TNiagaraHlslTranslator<GraphBridge>::DefineMainGPUFunctions(
 					{
 						VarFmt = VarName + TEXT("{0} = {4};\n");
 					}
-					else if (DataSetIndex == 0 && FParamMapHistory::IsPreviousValue(Var) && TranslationStages[i].ScriptUsage == ENiagaraScriptUsage::ParticleUpdateScript)
+					else if (DataSetIndex == 0 && FNiagaraParameterUtilities::IsPreviousValue(Var) && TranslationStages[i].ScriptUsage == ENiagaraScriptUsage::ParticleUpdateScript)
 					{
 						GatheredPreviousVariables.AddUnique(Var);
 						bWrite = false; // We need to bump the read indices forwards, but not actually add the read.
@@ -3197,7 +3197,7 @@ void TNiagaraHlslTranslator<GraphBridge>::DefineMainGPUFunctions(
 					{
 						VarFmt = VarName + TEXT("{0} = InputData{1}({2}, {3}, InstanceIdx);\n");
 
-						if (FParamMapHistory::IsAttribute(Var))
+						if (FNiagaraParameterUtilities::IsAttribute(Var))
 						{
 							FString RegisterName = VarName;
 							ReplaceNamespaceInline(RegisterName, PARAM_MAP_ATTRIBUTE_STR, PARAM_MAP_INDICES_STR);
@@ -3213,7 +3213,7 @@ void TNiagaraHlslTranslator<GraphBridge>::DefineMainGPUFunctions(
 			// Put any gathered previous variables into the list here so that we can use them by recording the last value from the parent variable on load.
 			for (FNiagaraVariable VarPrevious : GatheredPreviousVariables)
 			{
-				FNiagaraVariable SrcVar = FParamMapHistory::GetSourceForPreviousValue(VarPrevious);
+				FNiagaraVariable SrcVar = FNiagaraParameterUtilities::GetSourceForPreviousValue(VarPrevious);
 				const FString VarName = ContextName + GetSanitizedSymbolName(SrcVar.GetName().ToString());
 				const FString VarPrevName = ContextName + GetSanitizedSymbolName(VarPrevious.GetName().ToString());
 				HlslOutput += VarPrevName + TEXT(" = ") + VarName + TEXT(";\n");
@@ -3274,9 +3274,9 @@ void TNiagaraHlslTranslator<GraphBridge>::DefineMainGPUFunctions(
 					for (const FVarAndDefaultSource& VarAndDefaultSource : VarAndDefaultSourceArray)
 					{
 						const FNiagaraVariable& Var = VarAndDefaultSource.Variable;
-						if (FParamMapHistory::IsPreviousValue(Var))
+						if (FNiagaraParameterUtilities::IsPreviousValue(Var))
 						{
-							FNiagaraVariable SrcVar = FParamMapHistory::GetSourceForPreviousValue(Var);
+							FNiagaraVariable SrcVar = FNiagaraParameterUtilities::GetSourceForPreviousValue(Var);
 							const FString VarName = GetSanitizedSymbolName(SrcVar.GetName().ToString());
 							const FString VarPrevName = GetSanitizedSymbolName(Var.GetName().ToString());
 							HlslOutput += TEXT("\t\tContext.") + TranslationStages[i].PassNamespace + TEXT(".") + VarPrevName + TEXT(" = Context.") + TranslationStages[i-1].PassNamespace + TEXT(".") + VarName + TEXT(";\n");
@@ -3763,9 +3763,9 @@ void FNiagaraHlslTranslator::DefineMain(FString &OutHlslOutput,
 					for (const FVarAndDefaultSource& VarAndDefaultSource : VarAndDefaultSourceArray)
 					{
 						const FNiagaraVariable& Var = VarAndDefaultSource.Variable;
-						if (FNiagaraParameterMapHistory::IsPreviousValue(Var))
+						if (FNiagaraParameterUtilities::IsPreviousValue(Var))
 						{
-							FNiagaraVariable SrcVar = FNiagaraParameterMapHistory::GetSourceForPreviousValue(Var);
+							FNiagaraVariable SrcVar = FNiagaraParameterUtilities::GetSourceForPreviousValue(Var);
 							const FString VarName =  GetSanitizedSymbolName(SrcVar.GetName().ToString());
 							const FString VarPrevName =  GetSanitizedSymbolName(Var.GetName().ToString());
 							OutHlslOutput += TEXT("\t\tContext.") + TranslationStages[StageIdx + 1].PassNamespace + TEXT(".") + VarPrevName + TEXT(" = Context.") + TranslationStages[StageIdx].PassNamespace + TEXT(".")+ VarName +  TEXT(";\n");
@@ -3970,7 +3970,7 @@ void FNiagaraHlslTranslator::DefineDataSetVariableReads(FString &OutHlslOutput,
 			{
 				Fmt = VariableName + TEXT("{0} = {4};\n");
 			}
-			else if (DataSetIndex == 0 && FNiagaraParameterMapHistory::IsPreviousValue(Var) && bIsUpdateScript)
+			else if (DataSetIndex == 0 && FNiagaraParameterUtilities::IsPreviousValue(Var) && bIsUpdateScript)
 			{
 				GatheredPreviousVariables.AddUnique(Var);
 				bWrite = false; // We need to bump the read indices forwards, but not actually add the read.
@@ -3979,7 +3979,7 @@ void FNiagaraHlslTranslator::DefineDataSetVariableReads(FString &OutHlslOutput,
 			{
 				Fmt = VariableName + TEXT("{0} = InputData{1}({2}, {3});\n");
 
-				if (FNiagaraParameterMapHistory::IsAttribute(Var))
+				if (FNiagaraParameterUtilities::IsAttribute(Var))
 				{
 					FString RegisterName = VariableName;
 					ReplaceNamespaceInline(RegisterName, PARAM_MAP_ATTRIBUTE_STR, PARAM_MAP_INDICES_STR);
@@ -3995,7 +3995,7 @@ void FNiagaraHlslTranslator::DefineDataSetVariableReads(FString &OutHlslOutput,
 		// Put any gathered previous variables into the list here so that we can use them by recording the last value from the parent variable on load.
 		for (FNiagaraVariable VarPrevious : GatheredPreviousVariables)
 		{
-			FNiagaraVariable SrcVar = FNiagaraParameterMapHistory::GetSourceForPreviousValue(VarPrevious);
+			FNiagaraVariable SrcVar = FNiagaraParameterUtilities::GetSourceForPreviousValue(VarPrevious);
 			const FString VarName = ContextName + GetSanitizedSymbolName(SrcVar.GetName().ToString());
 			const FString VarPrevName = ContextName + GetSanitizedSymbolName(VarPrevious.GetName().ToString());
 			HlslOutput += VarPrevName + TEXT(" = ") + VarName + TEXT(";\n");
@@ -4546,7 +4546,7 @@ bool FNiagaraHlslTranslator::ShouldInterpolateParameter(const FNiagaraVariable&
 		return false;
 	}
 
-	if (FNiagaraParameterMapHistory::IsRapidIterationParameter(Parameter))
+	if (FNiagaraParameterUtilities::IsRapidIterationParameter(Parameter))
 	{
 		return false;
 	}
@@ -4578,9 +4578,9 @@ bool FNiagaraHlslTranslator::ShouldInterpolateParameter(const FNiagaraVariable&
 
 // specialization handling the case where the translation process treats the graph as mutable
 template<>
-void TNiagaraHlslTranslator<FNiagaraCompilationGraphBridge>::UpdateStaticSwitchConstants(FNode* Node)
+void TNiagaraHlslTranslator<FNiagaraCompilationGraphBridge>::UpdateStaticSwitchConstants(const FPin* ChildPin)
 {
-	if (UNiagaraNodeStaticSwitch* SwitchNode = Cast<UNiagaraNodeStaticSwitch>(Node))
+	if (UNiagaraNodeStaticSwitch* SwitchNode = Cast<UNiagaraNodeStaticSwitch>(ChildPin->GetOwningNode()))
 	{
 		SwitchNode->CheckForOutdatedEnum(this);
 
@@ -4753,7 +4753,7 @@ int32 TNiagaraHlslTranslator<GraphBridge>::GetParameter(const FNiagaraVariable&
 
 				if (Pin != nullptr)
 				{
-					const FInputNode* InputNode = GraphBridge::template AsNodeType<FInputNode>(GraphBridge::GetOwningNode(Pin));
+					const FInputNode* InputNode = GraphBridge::AsInputNode(GraphBridge::GetOwningNode(Pin));
 					if (InputNode != nullptr && GraphBridge::GetInputVariable(InputNode) == Parameter)
 					{
 						if (CurrentDefaultPinTraversal.Num() == 0)
@@ -4775,7 +4775,7 @@ int32 TNiagaraHlslTranslator<GraphBridge>::GetParameter(const FNiagaraVariable&
 	//Not a in a function or not a valid function parameter so grab from the main uniforms.
 	int32 OutputChunkIdx = INDEX_NONE;
 	FNiagaraVariable OutputVariable = Parameter;
-	if (FParamMapHistory::IsInNamespace(OutputVariable, PARAM_MAP_ATTRIBUTE_STR) || FParamMapHistory::IsExternalConstantNamespace(OutputVariable, CompileOptions.TargetUsage, CompileOptions.GetTargetUsageBitmask()))
+	if (FNiagaraParameterUtilities::IsInNamespace(OutputVariable, PARAM_MAP_ATTRIBUTE_STR) || FNiagaraParameterUtilities::IsExternalConstantNamespace(OutputVariable, CompileOptions.TargetUsage, CompileOptions.GetTargetUsageBitmask()))
 	{
 		if (!ParameterMapRegisterExternalConstantNamespaceVariable(OutputVariable, nullptr, 0, OutputChunkIdx, nullptr))
 		{
@@ -4784,7 +4784,7 @@ int32 TNiagaraHlslTranslator<GraphBridge>::GetParameter(const FNiagaraVariable&
 	}
 	else
 	{
-		OutputVariable = FParamMapHistory::MoveToExternalConstantNamespaceVariable(OutputVariable, CompileOptions.TargetUsage);
+		OutputVariable = FNiagaraParameterUtilities::MoveToExternalConstantNamespaceVariable(OutputVariable, CompileOptions.TargetUsage);
 		if (!ParameterMapRegisterExternalConstantNamespaceVariable(OutputVariable, nullptr, 0, OutputChunkIdx, nullptr))
 		{
 			OutputChunkIdx = INDEX_NONE;
@@ -5040,7 +5040,7 @@ bool FNiagaraHlslTranslationStage::ShouldDoSpawnOnlyLogic() const
 
 bool FNiagaraHlslTranslationStage::IsExternalConstantNamespace(const FNiagaraVariable& InVar, ENiagaraScriptUsage InTargetUsage, uint32 InTargetBitmask)
 {
-	if (FNiagaraParameterMapHistory::IsExternalConstantNamespace(InVar, InTargetUsage, InTargetBitmask))
+	if (FNiagaraParameterUtilities::IsExternalConstantNamespace(InVar, InTargetUsage, InTargetBitmask))
 	{
 		if (IterationSourceType == ENiagaraIterationSource::DataInterface && InVar.IsInNameSpace(IterationDataInterface))
 			return false;
@@ -5154,14 +5154,14 @@ void TNiagaraHlslTranslator<GraphBridge>::InitializeParameterMapDefaults(int32 P
 			// During the initial pass, only support constants for the default pin and non-bound variables
 			const bool bDefaultPinConnected = GraphBridge::GetLinkedOutputPin(DefaultPin) != nullptr;
 
-			if (!GraphBridge::FParamMapHistory::IsInitialValue(Var) && !bDefaultPinConnected && !(DefaultMode.IsSet() && (*DefaultMode == ENiagaraDefaultMode::Binding || *DefaultMode == ENiagaraDefaultMode::FailIfPreviouslyNotSet)))
+			if (!FNiagaraParameterUtilities::IsInitialValue(Var) && !bDefaultPinConnected && !(DefaultMode.IsSet() && (*DefaultMode == ENiagaraDefaultMode::Binding || *DefaultMode == ENiagaraDefaultMode::FailIfPreviouslyNotSet)))
 			{
 				HandleParameterRead(ParamMapHistoryIdx, Var, DefaultPin, DefaultPinNode, OutputChunkId, TOptional<ENiagaraDefaultMode>(), TOptional<FNiagaraScriptVariableBinding>(), !bWriteToParamMapEntries);
 				UniqueVarToChunk.Add(Var, OutputChunkId);
 			}
-			else if (GraphBridge::FParamMapHistory::IsInitialValue(Var))
+			else if (FNiagaraParameterUtilities::IsInitialValue(Var))
 			{
-				FNiagaraVariable SourceForInitialValue = GraphBridge::FParamMapHistory::GetSourceForInitialValue(Var);
+				FNiagaraVariable SourceForInitialValue = FNiagaraParameterUtilities::GetSourceForInitialValue(Var);
 				if (!UniqueVars.Contains(SourceForInitialValue))
 				{
 					//@todo(ng) disabled pending investigation UE-150159
@@ -5246,7 +5246,7 @@ void TNiagaraHlslTranslator<GraphBridge>::Output(const FOutputNode* OutputNode,
 
 				if (Var.GetType() != FNiagaraTypeDefinition::GetParameterMapDef())
 				{
-					FNiagaraVariable VarNamespaced = FParamMapHistory::BasicAttributeToNamespacedAttribute(Var);
+					FNiagaraVariable VarNamespaced = FNiagaraParameterUtilities::BasicAttributeToNamespacedAttribute(Var);
 					FString ParameterMapInstanceName = GetParameterMapInstanceName(0);
 					int32 ChunkIdx = AddBodyChunk(ParameterMapInstanceName + TEXT(".") + GetSanitizedSymbolName(VarNamespaced.GetName().ToString()), TEXT("{0}"), VarNamespaced.GetType(), Input, false);
 
@@ -5469,7 +5469,7 @@ void TNiagaraHlslTranslator<GraphBridge>::ParameterMapSet(const FParamMapSetNode
 							{
 								bAllSame = false;
 							}
-							else if (FoundOverrideIdx == INDEX_NONE && FParamMapHistory::IsAliasedEmitterParameter(Var))
+							else if (FoundOverrideIdx == INDEX_NONE && FNiagaraParameterUtilities::IsAliasedEmitterParameter(Var))
 							{
 								FNiagaraAliasContext ResolveAliasesContext(FNiagaraAliasContext::ERapidIterationParameterMode::EmitterOrParticleScript);
 								ResolveAliasesContext.ChangeEmitterToEmitterName(CompileData->GetUniqueEmitterName());
@@ -5524,7 +5524,7 @@ void TNiagaraHlslTranslator<GraphBridge>::ParameterMapSet(const FParamMapSetNode
 					}
 
 					FName UsageName;
-					if (FParamMapHistory::IsAliasedEmitterParameter(Var))
+					if (FNiagaraParameterUtilities::IsAliasedEmitterParameter(Var))
 					{
 						FNiagaraVariable AliasedVar = ActiveHistoryForFunctionCalls.ResolveAliases(Var);
 						UsageName = AliasedVar.GetName();
@@ -5576,7 +5576,7 @@ void TNiagaraHlslTranslator<GraphBridge>::ParameterMapSet(const FParamMapSetNode
 					}
 
 					FName ResolvedName;
-					if (FNiagaraParameterMapHistory::IsAliasedEmitterParameter(Var))
+					if (FNiagaraParameterUtilities::IsAliasedEmitterParameter(Var))
 					{
 						ResolvedName = ActiveHistoryForFunctionCalls.ResolveAliases(Var).GetName();
 					}
@@ -5673,7 +5673,7 @@ bool FNiagaraHlslTranslator::RequiresInterpolation() const
 template<typename GraphBridge>
 bool TNiagaraHlslTranslator<GraphBridge>::GetLiteralConstantVariable(FNiagaraVariable& OutVar) const
 {
-	if (FParamMapHistory::IsInNamespace(OutVar, PARAM_MAP_EMITTER_STR) || FParamMapHistory::IsInNamespace(OutVar, PARAM_MAP_SYSTEM_STR))
+	if (FNiagaraParameterUtilities::IsInNamespace(OutVar, PARAM_MAP_EMITTER_STR) || FNiagaraParameterUtilities::IsInNamespace(OutVar, PARAM_MAP_SYSTEM_STR))
 	{
 		FNiagaraVariable ResolvedVar = ActiveHistoryForFunctionCalls.ResolveAliases(OutVar);
 		if (OutVar == FNiagaraVariable(FNiagaraTypeDefinition::GetBoolDef(), TEXT("Emitter.Localspace")))
@@ -5758,173 +5758,173 @@ bool TNiagaraHlslTranslator<GraphBridge>::HandleBoundConstantVariableToDataSetRe
 template<typename GraphBridge>
 bool TNiagaraHlslTranslator<GraphBridge>::ParameterMapRegisterExternalConstantNamespaceVariable(FNiagaraVariable InVariable, const FNode* InNodeForErrorReporting, int32 InParamMapHistoryIdx, int32& Output, const FPin* InDefaultPin)
 {
+	Output = INDEX_NONE;
+
+	if (!ensure(InParamMapHistoryIdx >= 0))
+	{
+		return false;
+	}
+
 	InVariable = ActiveHistoryForFunctionCalls.ResolveAliases(InVariable);
+
+	if (!InVariable.IsValid())
+	{
+		Error(FText::Format(LOCTEXT("GetSystemConstantFail", "Unknown System constant, Type: {0} Variable: {1}"), InVariable.GetType().GetNameText(), FText::FromName(InVariable.GetName())), InNodeForErrorReporting, nullptr);
+		return false;
+	}
+
+	if (InVariable.GetType().IsStatic())
+	{
+		return false;
+	}
+
 	FString VarName = InVariable.GetName().ToString();
 	FString SymbolName = GetSanitizedSymbolName(VarName);
 	FString FlattenedName = SymbolName.Replace(TEXT("."), TEXT("_"));
 	FString ParameterMapInstanceName = GetParameterMapInstanceName(InParamMapHistoryIdx);
 
-	Output = INDEX_NONE;
-	if (InVariable.IsValid())
+	bool bMissingParameter = false;
+	FParameterCollection Collection = ParamMapHistories[InParamMapHistoryIdx].IsParameterCollectionParameter(InVariable, bMissingParameter);
+	if (GraphBridge::IsCollectionValid(Collection) && bMissingParameter)
 	{
-		bool bMissingParameter = false;
-		UNiagaraParameterCollection* Collection = nullptr;
-		if (InParamMapHistoryIdx >= 0)
-		{
-			Collection = ParamMapHistories[InParamMapHistoryIdx].IsParameterCollectionParameter(InVariable, bMissingParameter);
-			if (Collection && bMissingParameter)
-			{
-				Error(FText::Format(LOCTEXT("MissingNPCParameterError", "Parameter named {0} of type {1} was not found in Parameter Collection {2}"),
-					FText::FromName(InVariable.GetName()), InVariable.GetType().GetNameText(), FText::FromString(Collection->GetFullName())), InNodeForErrorReporting, InDefaultPin);
-				return false;
-			}
-		}
+		Error(FText::Format(LOCTEXT("MissingNPCParameterError", "Parameter named {0} of type {1} was not found in Parameter Collection {2}"),
+			FText::FromName(InVariable.GetName()), InVariable.GetType().GetNameText(), FText::FromString(GraphBridge::GetCollectionFullName(Collection))), InNodeForErrorReporting, InDefaultPin);
+		return false;
+	}
 
-		const bool bIsDataInterface = InVariable.IsDataInterface();
-		const bool bIsUObject = InVariable.IsUObject();
-		const FString* EmitterAlias = ActiveHistoryForFunctionCalls.GetEmitterAlias();
+	const bool bIsDataInterface = InVariable.IsDataInterface();
+	const bool bIsUObject = InVariable.IsUObject();
+	const FString* EmitterAlias = ActiveHistoryForFunctionCalls.GetEmitterAlias();
 
-		bool bIsPerInstanceBulkSystemParam = IsBulkSystemScript() && !bIsDataInterface && !bIsUObject && (FParamMapHistory::IsUserParameter(InVariable) || FParamMapHistory::IsPerInstanceEngineParameter(InVariable, EmitterAlias != nullptr ? *EmitterAlias : TEXT("Emitter")));
-		const bool bIsExternalConstantParameter = FParamMapHistory::IsRapidIterationParameter(InVariable) && !InVariable.GetType().IsStatic();
+	bool bIsPerInstanceBulkSystemParam = IsBulkSystemScript() && !bIsDataInterface && !bIsUObject && (FNiagaraParameterUtilities::IsUserParameter(InVariable) || FNiagaraParameterUtilities::IsPerInstanceEngineParameter(InVariable, EmitterAlias != nullptr ? *EmitterAlias : TEXT("Emitter")));
+	const bool bIsExternalConstantParameter = FNiagaraParameterUtilities::IsRapidIterationParameter(InVariable) && !InVariable.GetType().IsStatic();
 
-		if (InVariable.GetType().IsStatic())
-		{
-			return false;
-		}
+	if (!bIsPerInstanceBulkSystemParam)
+	{
+		int32 UniformChunk = 0;
 
-		if (!bIsPerInstanceBulkSystemParam)
+		if (false == ParamMapDefinedSystemVars.Contains(InVariable.GetName()))
 		{
-			int32 UniformChunk = 0;
+			FString SymbolNameDefined = FlattenedName;
 
-			if (false == ParamMapDefinedSystemVars.Contains(InVariable.GetName()))
+			if (bIsDataInterface)
 			{
-				FString SymbolNameDefined = FlattenedName;
-
-				if (bIsDataInterface)
+				UNiagaraDataInterface* DataInterface = nullptr;
+				if (GraphBridge::IsCollectionValid(Collection))
 				{
-					UNiagaraDataInterface* DataInterface = nullptr;
-					if (Collection)
-					{
-						DataInterface = Collection->GetDefaultInstance()->GetParameterStore().GetDataInterface(InVariable);
-						if (DataInterface == nullptr)
-						{
-							Error(FText::Format(LOCTEXT("ParameterCollectionDataInterfaceNotFoundErrorFormat", "Data interface named {0} of type {1} was not found in Parameter Collection {2}"),
-								FText::FromName(InVariable.GetName()), InVariable.GetType().GetNameText(), FText::FromString(Collection->GetFullName())), InNodeForErrorReporting, InDefaultPin);
-							return false;
-						}
-					}
-					else
+					DataInterface = GraphBridge::GetCollectionDataInterface(Collection, InVariable);
+					if (DataInterface == nullptr)
 					{
-						DataInterface = CompileDuplicateData->GetDuplicatedDataInterfaceCDOForClass(const_cast<UClass*>(InVariable.GetType().GetClass()));
-						if (DataInterface == nullptr)
-						{
-							Error(FText::Format(LOCTEXT("GetDuplicatedDataInterfaceCDOForClassFailed", "GetDuplicatedDataInterfaceCDOForClass failed for Variable({0}) Class({1})"), FText::FromName(InVariable.GetName()), InVariable.GetType().GetNameText()), InNodeForErrorReporting, InDefaultPin);
-							return false;
-						}
-					}
-					if (ensure(DataInterface))
-					{
-						Output = RegisterDataInterface(InVariable, DataInterface, true, true);
-						return true;
+						Error(FText::Format(LOCTEXT("ParameterCollectionDataInterfaceNotFoundErrorFormat", "Data interface named {0} of type {1} was not found in Parameter Collection {2}"),
+							FText::FromName(InVariable.GetName()), InVariable.GetType().GetNameText(), FText::FromString(GraphBridge::GetCollectionFullName(Collection))), InNodeForErrorReporting, InDefaultPin);
+						return false;
 					}
 				}
-				if (bIsUObject)
+				else
 				{
-					UObject* Object = nullptr;
-					if (Collection)
+					DataInterface = CompileDuplicateData->GetDuplicatedDataInterfaceCDOForClass(const_cast<UClass*>(InVariable.GetType().GetClass()));
+					if (DataInterface == nullptr)
 					{
-						Object = Collection->GetDefaultInstance()->GetParameterStore().GetUObject(InVariable);
+						Error(FText::Format(LOCTEXT("GetDuplicatedDataInterfaceCDOForClassFailed", "GetDuplicatedDataInterfaceCDOForClass failed for Variable({0}) Class({1})"), FText::FromName(InVariable.GetName()), InVariable.GetType().GetNameText()), InNodeForErrorReporting, InDefaultPin);
+						return false;
 					}
-					Output = RegisterUObject(InVariable, Object, true);
-					return true;
 				}
-				if (!InVariable.IsDataAllocated() && !InDefaultPin)
+				if (ensure(DataInterface))
 				{
-					FNiagaraEditorUtilities::ResetVariableToDefaultValue(InVariable);
+					Output = RegisterDataInterface(InVariable, DataInterface, true, true);
+					return true;
 				}
-				else if (!InVariable.IsDataAllocated())
+			}
+			if (bIsUObject)
+			{
+				UObject* Object = nullptr;
+				if (GraphBridge::IsCollectionValid(Collection))
 				{
-					FillVariableWithDefaultValue(InVariable, InDefaultPin);
+					Object = GraphBridge::GetCollectionUObject(Collection, InVariable);
 				}
+				Output = RegisterUObject(InVariable, Object, true);
+				return true;
+			}
+			if (!InVariable.IsDataAllocated() && !InDefaultPin)
+			{
+				FNiagaraEditorUtilities::ResetVariableToDefaultValue(InVariable);
+			}
+			else if (!InVariable.IsDataAllocated())
+			{
+				FillVariableWithDefaultValue(InVariable, InDefaultPin);
+			}
 
-				if (InVariable.GetAllocatedSizeInBytes() != InVariable.GetSizeInBytes())
-				{
-					Error(FText::Format(LOCTEXT("GetParameterUnsetParam", "Variable {0} hasn't had its default value set. Required Bytes: {1} vs Allocated Bytes: {2}"), FText::FromName(InVariable.GetName()), FText::AsNumber(InVariable.GetType().GetSize()), FText::AsNumber(InVariable.GetSizeInBytes())), nullptr, nullptr);
-				}
+			if (InVariable.GetAllocatedSizeInBytes() != InVariable.GetSizeInBytes())
+			{
+				Error(FText::Format(LOCTEXT("GetParameterUnsetParam", "Variable {0} hasn't had its default value set. Required Bytes: {1} vs Allocated Bytes: {2}"), FText::FromName(InVariable.GetName()), FText::AsNumber(InVariable.GetType().GetSize()), FText::AsNumber(InVariable.GetSizeInBytes())), nullptr, nullptr);
+			}
 
-				if ( IsVariableInUniformBuffer(InVariable) )
+			if ( IsVariableInUniformBuffer(InVariable) )
+			{
+				// we must ensure that there's a one to one relationship between symbol name and parameter.  The generated VM only
+				// knows about the symbols while the parameter stores knows about the parameters, if these mismatch, then we're going
+				// to be incorrectly addressing the constant table
+				if (!CompilationOutput.ScriptData.Parameters.FindParameter(InVariable))
 				{
-					// we must ensure that there's a one to one relationship between symbol name and parameter.  The generated VM only
-					// knows about the symbols while the parameter stores knows about the parameters, if these mismatch, then we're going
-					// to be incorrectly addressing the constant table
-					if (!CompilationOutput.ScriptData.Parameters.FindParameter(InVariable))
-					{
-						bool AddParameter = true;
+					bool AddParameter = true;
 
-						// add the parameter, but first evaluate whether any of the symbols for existing parameters would conflict
-						for (const FNiagaraVariable& ExistingParameter : CompilationOutput.ScriptData.Parameters.Parameters)
+					// add the parameter, but first evaluate whether any of the symbols for existing parameters would conflict
+					for (const FNiagaraVariable& ExistingParameter : CompilationOutput.ScriptData.Parameters.Parameters)
+					{
+						FNameBuilder ExistingParameterName(ExistingParameter.GetName());
+						if (GetSanitizedSymbolName(ExistingParameterName).Equals(SymbolName))
 						{
-							FNameBuilder ExistingParameterName(ExistingParameter.GetName());
-							if (GetSanitizedSymbolName(ExistingParameterName).Equals(SymbolName))
-							{
-								Error(FText::Format(LOCTEXT("NonUniqueSymbolNames", "Parameters ('{0}' and '{1}') found which resolve to the same HLSL symbol name '{2}'.  These should be disambiguated."),
-									FText::FromName(InVariable.GetName()), FText::FromName(ExistingParameter.GetName()), FText::FromString(SymbolName)), InNodeForErrorReporting, InDefaultPin);
+							Error(FText::Format(LOCTEXT("NonUniqueSymbolNames", "Parameters ('{0}' and '{1}') found which resolve to the same HLSL symbol name '{2}'.  These should be disambiguated."),
+								FText::FromName(InVariable.GetName()), FText::FromName(ExistingParameter.GetName()), FText::FromString(SymbolName)), InNodeForErrorReporting, InDefaultPin);
 
-								AddParameter = false;
-							}
+							AddParameter = false;
 						}
+					}
 
-						if (AddParameter)
-						{
-							CompilationOutput.ScriptData.Parameters.Parameters.Add(InVariable);
-						}
+					if (AddParameter)
+					{
+						CompilationOutput.ScriptData.Parameters.Parameters.Add(InVariable);
 					}
 				}
-
-				UniformChunk = AddUniformChunk(SymbolNameDefined, InVariable, ENiagaraCodeChunkMode::Uniform, UNiagaraScript::IsGPUScript(CompileOptions.TargetUsage));
-			}
-			else
-			{
-				const auto& SystemVar = ParamMapDefinedSystemVars.FindChecked(InVariable.GetName());
-				UniformChunk = SystemVar.ChunkIndex;
-			}
-			
-			if (bIsExternalConstantParameter)
-			{
-				Output = UniformChunk;
-				return true;
-			}
-			else
-			{
-				//Add this separately as the same uniform can appear in the pre sim chunks more than once in different param maps.
-				PerStageMainPreSimulateChunks[ActiveStageIdx].AddUnique(FString::Printf(TEXT("%s.%s = %s;"), *ParameterMapInstanceName, *GetSanitizedSymbolName(VarName), *GetCodeAsSource(UniformChunk)));
 			}
+
+			UniformChunk = AddUniformChunk(SymbolNameDefined, InVariable, ENiagaraCodeChunkMode::Uniform, UNiagaraScript::IsGPUScript(CompileOptions.TargetUsage));
 		}
-		else if (bIsPerInstanceBulkSystemParam && !ExternalVariablesForBulkUsage.Contains(InVariable))
+		else
 		{
-			ExternalVariablesForBulkUsage.Add(InVariable);
+			const auto& SystemVar = ParamMapDefinedSystemVars.FindChecked(InVariable.GetName());
+			UniformChunk = SystemVar.ChunkIndex;
 		}
-
-		Output = AddSourceChunk(ParameterMapInstanceName + TEXT(".") + SymbolName, InVariable.GetType());
-		if (CodeChunks.IsValidIndex(Output))
-		{ // Leave a breadcrumb to resolve for compile tags.
-			CodeChunks[Output].Original = InVariable;
+			
+		if (bIsExternalConstantParameter)
+		{
+			Output = UniformChunk;
+			return true;
 		}
-
-		//Track a few special case reads that the system needs to know about.
-		if(InVariable == SYS_PARAM_ENGINE_SYSTEM_SIGNIFICANCE_INDEX)
+		else
 		{
-			CompilationOutput.ScriptData.bReadsSignificanceIndex = true;
+			//Add this separately as the same uniform can appear in the pre sim chunks more than once in different param maps.
+			PerStageMainPreSimulateChunks[ActiveStageIdx].AddUnique(FString::Printf(TEXT("%s.%s = %s;"), *ParameterMapInstanceName, *GetSanitizedSymbolName(VarName), *GetCodeAsSource(UniformChunk)));
 		}
+	}
+	else if (bIsPerInstanceBulkSystemParam && !ExternalVariablesForBulkUsage.Contains(InVariable))
+	{
+		ExternalVariablesForBulkUsage.Add(InVariable);
+	}
 
-		return true;
+	Output = AddSourceChunk(ParameterMapInstanceName + TEXT(".") + SymbolName, InVariable.GetType());
+	if (CodeChunks.IsValidIndex(Output))
+	{ // Leave a breadcrumb to resolve for compile tags.
+		CodeChunks[Output].Original = InVariable;
 	}
 
-	if (Output == INDEX_NONE)
+	//Track a few special case reads that the system needs to know about.
+	if(InVariable == SYS_PARAM_ENGINE_SYSTEM_SIGNIFICANCE_INDEX)
 	{
-		Error(FText::Format(LOCTEXT("GetSystemConstantFail", "Unknown System constant, Type: {0} Variable: {1}"), InVariable.GetType().GetNameText(), FText::FromName(InVariable.GetName())), InNodeForErrorReporting, nullptr);
+		CompilationOutput.ScriptData.bReadsSignificanceIndex = true;
 	}
-	return false;
+
+	return true;
 }
 
 template<typename GraphBridge>
@@ -6096,7 +6096,7 @@ void FNiagaraHlslTranslator::SetConstantByStaticVariable(int32& OutValue, const
 				{
 					VarWithValue.SetData(StaticVariables[FoundOverrideIdx].GetData());
 				}
-				else if (FNiagaraParameterMapHistory::IsAliasedEmitterParameter(Var))
+				else if (FNiagaraParameterUtilities::IsAliasedEmitterParameter(Var))
 				{
 					FNiagaraAliasContext ResolveAliasesContext(FNiagaraAliasContext::ERapidIterationParameterMode::EmitterOrParticleScript);
 					ResolveAliasesContext.ChangeEmitterToEmitterName(GetUniqueEmitterName());
@@ -6139,7 +6139,7 @@ void FNiagaraHlslTranslator::SetConstantByStaticVariable(int32& OutValue, const
 template<typename GraphBridge>
 bool TNiagaraHlslTranslator<GraphBridge>::ParameterMapRegisterUniformAttributeVariable(const FNiagaraVariable& InVariable, FNode* InNode, int32 InParamMapHistoryIdx, int32& Output)
 {
-	FNiagaraVariable NewVar = FParamMapHistory::BasicAttributeToNamespacedAttribute(InVariable);
+	FNiagaraVariable NewVar = FNiagaraParameterUtilities::BasicAttributeToNamespacedAttribute(InVariable);
 	if (NewVar.IsValid())
 	{
 		return ParameterMapRegisterNamespaceAttributeVariable(NewVar, InNode, InParamMapHistoryIdx, Output);
@@ -6177,7 +6177,7 @@ bool TNiagaraHlslTranslator<GraphBridge>::ParameterMapRegisterNamespaceAttribute
 	FNiagaraVariable NamespaceVar = InVariable;
 
 	Output = INDEX_NONE;
-	FNiagaraVariable BasicVar = FParamMapHistory::ResolveAsBasicAttribute(InVariable);
+	FNiagaraVariable BasicVar = FNiagaraParameterUtilities::ResolveAsBasicAttribute(InVariable);
 	if (BasicVar.IsValid())
 	{
 		if (false == ParamMapDefinedAttributesToUniformChunks.Contains(BasicVar.GetName()))
@@ -6230,6 +6230,30 @@ FString FNiagaraHlslTranslator::GetParameterMapInstanceName(int32 ParamMapHistor
 	return ParameterMapInstanceName;
 }
 
+// simple helper template allowing us to keep Emitter() generic even if the Builder interface may diverge
+template<typename GraphBridge>
+class FScopedBuilderEmitter
+{
+public:
+	FScopedBuilderEmitter(FNiagaraCompilationGraphBridge::FParamMapHistoryBuilder& InBuilder, const typename GraphBridge::FEmitterNode* InEmitterNode)
+		: EmitterNode(InEmitterNode)
+		, EmitterUniqueName(GraphBridge::GetEmitterUniqueName(InEmitterNode))
+		, Builder(InBuilder)
+	{
+		Builder.EnterEmitter(EmitterUniqueName, GraphBridge::GetEmitterGraph(EmitterNode), EmitterNode);
+	}
+
+	~FScopedBuilderEmitter()
+	{
+		Builder.ExitEmitter(EmitterUniqueName, EmitterNode);
+	}
+
+private:
+	const typename GraphBridge::FEmitterNode* EmitterNode;
+	const FString EmitterUniqueName;
+	typename GraphBridge::FParamMapHistoryBuilder& Builder;
+};
+
 template<typename GraphBridge>
 void TNiagaraHlslTranslator<GraphBridge>::Emitter(const FEmitterNode* EmitterNode, TArray<int32>& Inputs, TArray<int32>& Outputs)
 {
@@ -6281,40 +6305,42 @@ void TNiagaraHlslTranslator<GraphBridge>::Emitter(const FEmitterNode* EmitterNod
 		Error(LOCTEXT("EmitterMissingParamMapIndex", "Emitter call missing valid ParameterMap index!"), EmitterNode, nullptr);
 		return;
 	}
-	ActiveHistoryForFunctionCalls.EnterEmitter(EmitterUniqueName, GraphBridge::GetEmitterGraph(EmitterNode), EmitterNode);
 
-	// Clear out the parameter map writes to emitter module parameters as they should not be shared across emitters.
-	if (ParamMapHistoryIdx != -1 && ParamMapHistoryIdx < ParamMapHistories.Num())
 	{
-		for (int32 i = 0; i < ParamMapHistories[ParamMapHistoryIdx].Variables.Num(); i++)
+		FScopedBuilderEmitter<GraphBridge> EmitterBuidlerScope(ActiveHistoryForFunctionCalls, EmitterNode);
+
+		// Clear out the parameter map writes to emitter module parameters as they should not be shared across emitters.
+		if (ParamMapHistoryIdx != -1 && ParamMapHistoryIdx < ParamMapHistories.Num())
 		{
-			check(ParamMapHistories[ParamMapHistoryIdx].VariablesWithOriginalAliasesIntact.Num() > i);
-			FNiagaraVariable Var = ParamMapHistories[ParamMapHistoryIdx].VariablesWithOriginalAliasesIntact[i];
-			if (FParamMapHistory::IsAliasedModuleParameter(Var))
+			for (int32 i = 0; i < ParamMapHistories[ParamMapHistoryIdx].Variables.Num(); i++)
 			{
-				ParamMapSetVariablesToChunks[ParamMapHistoryIdx][i] = INDEX_NONE;
+				check(ParamMapHistories[ParamMapHistoryIdx].VariablesWithOriginalAliasesIntact.Num() > i);
+				FNiagaraVariable Var = ParamMapHistories[ParamMapHistoryIdx].VariablesWithOriginalAliasesIntact[i];
+				if (FNiagaraParameterUtilities::IsAliasedModuleParameter(Var))
+				{
+					ParamMapSetVariablesToChunks[ParamMapHistoryIdx][i] = INDEX_NONE;
+				}
 			}
 		}
-	}
 
-	// We act like a function call here as the semantics are identical.
-	RegisterFunctionCall(ScriptUsage, Name, FullName, EmitterNode->NodeGuid, GraphBridge::GetEmitterHandleIdString(EmitterNode), EmitterGraph, Signature, false, FString(), {}, Inputs, CallInputs, CallOutputs, Signature);
-	GenerateFunctionCall(ScriptUsage, Signature, Inputs, Outputs);
+		// We act like a function call here as the semantics are identical.
+		RegisterFunctionCall(ScriptUsage, Name, FullName, EmitterNode->NodeGuid, GraphBridge::GetEmitterHandleIdString(EmitterNode), EmitterGraph, Signature, false, FString(), {}, Inputs, CallInputs, CallOutputs, Signature);
+		GenerateFunctionCall(ScriptUsage, Signature, Inputs, Outputs);
 
-	// Clear out the parameter map writes to emitter module parameters as they should not be shared across emitters.
-	if (ParamMapHistoryIdx != -1 && ParamMapHistoryIdx < ParamMapHistories.Num())
-	{
-		for (int32 i = 0; i < ParamMapHistories[ParamMapHistoryIdx].Variables.Num(); i++)
+		// Clear out the parameter map writes to emitter module parameters as they should not be shared across emitters.
+		if (ParamMapHistoryIdx != -1 && ParamMapHistoryIdx < ParamMapHistories.Num())
 		{
-			check(ParamMapHistories[ParamMapHistoryIdx].VariablesWithOriginalAliasesIntact.Num() > i);
-			FNiagaraVariable Var = ParamMapHistories[ParamMapHistoryIdx].VariablesWithOriginalAliasesIntact[i];
-			if (ActiveHistoryForFunctionCalls.IsInEncounteredFunctionNamespace(Var) || FParamMapHistory::IsAliasedModuleParameter(Var) || FParamMapHistory::IsInNamespace(Var, PARAM_MAP_TRANSIENT_STR))
+			for (int32 i = 0; i < ParamMapHistories[ParamMapHistoryIdx].Variables.Num(); i++)
 			{
-				ParamMapSetVariablesToChunks[ParamMapHistoryIdx][i] = INDEX_NONE;
+				check(ParamMapHistories[ParamMapHistoryIdx].VariablesWithOriginalAliasesIntact.Num() > i);
+				FNiagaraVariable Var = ParamMapHistories[ParamMapHistoryIdx].VariablesWithOriginalAliasesIntact[i];
+				if (ActiveHistoryForFunctionCalls.IsInEncounteredFunctionNamespace(Var) || FNiagaraParameterUtilities::IsAliasedModuleParameter(Var) || FNiagaraParameterUtilities::IsInNamespace(Var, PARAM_MAP_TRANSIENT_STR))
+				{
+					ParamMapSetVariablesToChunks[ParamMapHistoryIdx][i] = INDEX_NONE;
+				}
 			}
 		}
 	}
-	ActiveHistoryForFunctionCalls.ExitEmitter(EmitterUniqueName, EmitterNode);
 
 	ExitStatsScope();
 }
@@ -6364,7 +6390,7 @@ void TNiagaraHlslTranslator<GraphBridge>::ParameterMapGet(const FParamMapGetNode
 		TOptional<ENiagaraDefaultMode> DefaultMode = GraphBridge::GetGraphDefaultMode(GraphBridge::GetOwningGraph(GetNode), Var, DefaultBinding);
 		if (Var.GetType().IsStatic())
 		{
-			if (FParamMapHistory::IsExternalConstantNamespace(Var, CompileOptions.TargetUsage, CompileOptions.GetTargetUsageBitmask()))
+			if (FNiagaraParameterUtilities::IsExternalConstantNamespace(Var, CompileOptions.TargetUsage, CompileOptions.GetTargetUsageBitmask()))
 			{
 				if (DefaultMode.IsSet() && *DefaultMode == ENiagaraDefaultMode::FailIfPreviouslyNotSet && !Var.IsInNameSpace(FNiagaraConstants::UserNamespaceString))
 				{
@@ -6409,13 +6435,13 @@ void TNiagaraHlslTranslator<GraphBridge>::ValidateFailIfPreviouslyNotSet(const F
 {
 	bFailIfNotSet = false;
 	FNiagaraVariable SearchVar = InVar;
-	if (FParamMapHistory::IsInitialValue(InVar))
+	if (FNiagaraParameterUtilities::IsInitialValue(InVar))
 	{
-		SearchVar = FParamMapHistory::GetSourceForInitialValue(InVar);
+		SearchVar = FNiagaraParameterUtilities::GetSourceForInitialValue(InVar);
 	}
-	else if (FParamMapHistory::IsPreviousValue(InVar))
+	else if (FNiagaraParameterUtilities::IsPreviousValue(InVar))
 	{
-		SearchVar = FParamMapHistory::GetSourceForPreviousValue(InVar);
+		SearchVar = FNiagaraParameterUtilities::GetSourceForPreviousValue(InVar);
 	}
 
 	FVarAndDefaultSource* ParamMapDefinedVarAndDefaultSource = ParamMapDefinedAttributesToNamespaceVars.Find(SearchVar.GetName());
@@ -6462,8 +6488,8 @@ void TNiagaraHlslTranslator<GraphBridge>::HandleParameterRead(int32 ParamMapHist
 	bool bIsCandidateForRapidIteration = false;
 	const FInputPin* InputPin = DefaultPin;
 
-	FString Namespace = FParamMapHistory::GetNamespace(Var);
-	if (!ParamMapHistories[ParamMapHistoryIdx].IsValidNamespaceForReading(CompileOptions.TargetUsage, CompileOptions.TargetUsageBitmask, Namespace))
+	FString Namespace = FNiagaraParameterUtilities::GetNamespace(Var);
+	if (!FNiagaraParameterUtilities::IsValidNamespaceForReading(CompileOptions.TargetUsage, CompileOptions.TargetUsageBitmask, Namespace))
 	{
 		if (UNiagaraScript::IsStandaloneScript(CompileOptions.TargetUsage) && Namespace.StartsWith(PARAM_MAP_ATTRIBUTE_STR))
 		{
@@ -6488,7 +6514,7 @@ void TNiagaraHlslTranslator<GraphBridge>::HandleParameterRead(int32 ParamMapHist
 	}
 
 	
-	if (FParamMapHistory::IsExternalConstantNamespace(Var, CompileOptions.TargetUsage, CompileOptions.GetTargetUsageBitmask()))
+	if (FNiagaraParameterUtilities::IsExternalConstantNamespace(Var, CompileOptions.TargetUsage, CompileOptions.GetTargetUsageBitmask()))
 	{
 		if (DefaultMode.IsSet() && *DefaultMode == ENiagaraDefaultMode::FailIfPreviouslyNotSet && !bIgnoreDefaultSetFirst)
 		{
@@ -6504,7 +6530,7 @@ void TNiagaraHlslTranslator<GraphBridge>::HandleParameterRead(int32 ParamMapHist
 			return;
 		}			
 	}
-	else if (FParamMapHistory::IsAliasedModuleParameter(Var) && ActiveHistoryForFunctionCalls.InTopLevelFunctionCall(CompileOptions.TargetUsage))
+	else if (FNiagaraParameterUtilities::IsAliasedModuleParameter(Var) && ActiveHistoryForFunctionCalls.InTopLevelFunctionCall(CompileOptions.TargetUsage))
 	{
 		if (DefaultMode.IsSet() && *DefaultMode == ENiagaraDefaultMode::Binding && DefaultBinding.IsSet() && DefaultBinding->IsValid())
 		{
@@ -6518,7 +6544,7 @@ void TNiagaraHlslTranslator<GraphBridge>::HandleParameterRead(int32 ParamMapHist
 	}
 
 	FParamMapHistory& History = ParamMapHistories[ParamMapHistoryIdx];
-	bool bWasEmitterAliased = FParamMapHistory::IsAliasedEmitterParameter(Var);
+	bool bWasEmitterAliased = FNiagaraParameterUtilities::IsAliasedEmitterParameter(Var);
 	Var = ActiveHistoryForFunctionCalls.ResolveAliases(Var);
 
 	const FNiagaraKnownConstantInfo ConstantInfo = FNiagaraConstants::GetKnownConstantInfo(Var.GetName(), false);
@@ -6583,12 +6609,12 @@ void TNiagaraHlslTranslator<GraphBridge>::HandleParameterRead(int32 ParamMapHist
 			bool bVarChanged = false;
 			if (!bWasEmitterAliased && ActiveHistoryForFunctionCalls.GetEmitterAlias() != nullptr)
 			{
-				Var = FParamMapHistory::ConvertVariableToRapidIterationConstantName(Var, *(*ActiveHistoryForFunctionCalls.GetEmitterAlias()), GetTargetUsage());
+				Var = FNiagaraParameterUtilities::ConvertVariableToRapidIterationConstantName(Var, *(*ActiveHistoryForFunctionCalls.GetEmitterAlias()), GetTargetUsage());
 				bVarChanged = true;
 			}
 			else if (UNiagaraScript::IsSystemSpawnScript(CompileOptions.TargetUsage) || UNiagaraScript::IsSystemUpdateScript(CompileOptions.TargetUsage))
 			{
-				Var = FParamMapHistory::ConvertVariableToRapidIterationConstantName(Var, nullptr, GetTargetUsage());
+				Var = FNiagaraParameterUtilities::ConvertVariableToRapidIterationConstantName(Var, nullptr, GetTargetUsage());
 				bVarChanged = true;
 			}
 
@@ -6625,7 +6651,7 @@ void TNiagaraHlslTranslator<GraphBridge>::HandleParameterRead(int32 ParamMapHist
 				{
 					if (!bWasEmitterAliased && ActiveHistoryForFunctionCalls.GetEmitterAlias() != nullptr && CompileData != nullptr)
 					{
-						Var = FParamMapHistory::ConvertVariableToRapidIterationConstantName(OriginalVar, *CompileData->EmitterUniqueName, GetTargetUsage());
+						Var = FNiagaraParameterUtilities::ConvertVariableToRapidIterationConstantName(OriginalVar, *CompileData->EmitterUniqueName, GetTargetUsage());
 						bVarChanged = true;
 						FoundIdx = TranslationOptions.OverrideModuleConstants.Find(Var);
 					}
@@ -6696,9 +6722,9 @@ void TNiagaraHlslTranslator<GraphBridge>::HandleParameterRead(int32 ParamMapHist
 
 		if (LastSetChunkIdx == INDEX_NONE && (UNiagaraScript::IsSpawnScript(TranslationStages[ActiveStageIdx].ScriptUsage)))
 		{
-			if (FParamMapHistory::IsInitialValue(Var))
+			if (FNiagaraParameterUtilities::IsInitialValue(Var))
 			{
-				FNiagaraVariable SourceForInitialValue = FParamMapHistory::GetSourceForInitialValue(Var);
+				FNiagaraVariable SourceForInitialValue = FNiagaraParameterUtilities::GetSourceForInitialValue(Var);
 				bool bFoundExistingSet = false;
 				for (int32 OtherParamIdx = 0; OtherParamIdx < OtherOutputParamMapHistories.Num(); OtherParamIdx++)
 				{
@@ -6772,12 +6798,12 @@ void TNiagaraHlslTranslator<GraphBridge>::HandleParameterRead(int32 ParamMapHist
 					FNiagaraVariable RapidIterationConstantVar;
 					if (!bWasEmitterAliased && ActiveHistoryForFunctionCalls.GetEmitterAlias() != nullptr)
 					{
-						RapidIterationConstantVar = FParamMapHistory::ConvertVariableToRapidIterationConstantName(Var, *(*ActiveHistoryForFunctionCalls.GetEmitterAlias()), GetTargetUsage());
+						RapidIterationConstantVar = FNiagaraParameterUtilities::ConvertVariableToRapidIterationConstantName(Var, *(*ActiveHistoryForFunctionCalls.GetEmitterAlias()), GetTargetUsage());
 						bVarChanged = true;
 					}
 					else if (UNiagaraScript::IsSystemSpawnScript(CompileOptions.TargetUsage) || UNiagaraScript::IsSystemUpdateScript(CompileOptions.TargetUsage))
 					{
-						RapidIterationConstantVar = FParamMapHistory::ConvertVariableToRapidIterationConstantName(Var, nullptr, GetTargetUsage());
+						RapidIterationConstantVar = FNiagaraParameterUtilities::ConvertVariableToRapidIterationConstantName(Var, nullptr, GetTargetUsage());
 						bVarChanged = true;
 					}
 
@@ -6811,7 +6837,7 @@ void TNiagaraHlslTranslator<GraphBridge>::HandleParameterRead(int32 ParamMapHist
 					GraphBridge::BuildTraversal(GraphBridge::GetOwningGraph(LinkedOutputPinNode), LinkedOutputPinNode, Nodes);
 					for (const FNode* Node : Nodes)
 					{
-						const FInputNode* InputNode = GraphBridge::template AsNodeType<FInputNode>(Node);
+						const FInputNode* InputNode = GraphBridge::AsInputNode(Node);
 						if (InputNode && GraphBridge::GetInputVariable(InputNode).GetType() == FNiagaraTypeDefinition::GetParameterMapDef())
 						{
 							if (InputNode->Usage != ENiagaraInputNodeUsage::TranslatorConstant)
@@ -7078,7 +7104,7 @@ int32 TNiagaraHlslTranslator<GraphBridge>::RegisterUObject(FNiagaraVariable Vari
 	if (bAddParameterMapRead)
 	{
 		FName ResolvedName;
-		if (FNiagaraParameterMapHistory::IsAliasedEmitterParameter(Variable.GetName().ToString()))
+		if (FNiagaraParameterUtilities::IsAliasedEmitterParameter(Variable.GetName().ToString()))
 		{
 			ResolvedName = ActiveHistoryForFunctionCalls.ResolveAliases(Variable).GetName();
 		}
@@ -7122,7 +7148,7 @@ int32 TNiagaraHlslTranslator<GraphBridge>::RegisterDataInterface(const FNiagaraV
 	FNiagaraVariable ReadVariable;
 	if (bAddParameterMapRead)
 	{
-		ReadVariable = FParamMapHistory::IsAliasedEmitterParameter(Var.GetName().ToString())
+		ReadVariable = FNiagaraParameterUtilities::IsAliasedEmitterParameter(Var.GetName().ToString())
 			? ActiveHistoryForFunctionCalls.ResolveAliases(Var) 
 			: Var;
 	}
@@ -7304,7 +7330,7 @@ public:
 
 	~FScopedBuilderFunctionCall()
 	{
-		Builder.ExitFunction(FunctionCallNode->GetFunctionName(), FunctionCallNode);
+		Builder.ExitFunction(FunctionCallNode);
 	}
 
 private:
@@ -7460,7 +7486,7 @@ void TNiagaraHlslTranslator<GraphBridge>::FunctionCall(const FFunctionCallNode*
 
 					FNiagaraVariable Variable(TypeDef, *(TEXT("Module.") + Prefix));
 
-					if (FParamMapHistory::IsAliasedModuleParameter(Variable) && ActiveHistoryForFunctionCalls.InTopLevelFunctionCall(CompileOptions.TargetUsage))
+					if (FNiagaraParameterUtilities::IsAliasedModuleParameter(Variable) && ActiveHistoryForFunctionCalls.InTopLevelFunctionCall(CompileOptions.TargetUsage))
 					{
 						ResolvedVariable = ActiveHistoryForFunctionCalls.ResolveAliases(Variable);
 					}
@@ -7487,7 +7513,7 @@ void TNiagaraHlslTranslator<GraphBridge>::FunctionCall(const FFunctionCallNode*
 		}
 	}
 
-	if (const FCustomHlslNode* CustomFunctionHlsl = GraphBridge::template AsNodeType<FCustomHlslNode>(FunctionNode))
+	if (const FCustomHlslNode* CustomFunctionHlsl = GraphBridge::AsCustomHlslNode(FunctionNode))
 	{
 		// All of the arguments here are resolved withing the HandleCustomHlsl function..
 		HandleCustomHlslNode(CustomFunctionHlsl, ScriptUsage, Name, FullName, bCustomHlsl, CustomHlsl, CustomHlslIncludeFilePaths, Signature, Inputs);
@@ -7556,7 +7582,7 @@ void TNiagaraHlslTranslator<GraphBridge>::CullMapSetInputPin(const FPin* InputPi
 	}
 
 	// when a map set input is culled that is connected to a function call node (as is the case for dynamic inputs), we also need to cull any upstream pins that set inputs for the culled function call node
-	if (const FFunctionCallNode* FunctionNode = GraphBridge::template AsNodeType<FFunctionCallNode>(GraphBridge::GetOwningNode(LinkedPin)))
+	if (const FFunctionCallNode* FunctionNode = GraphBridge::AsFunctionCallNode(GraphBridge::GetOwningNode(LinkedPin)))
 	{
 		FString FunctionScriptName = GraphBridge::GetFunctionName(FunctionNode);
 		FunctionNodeStack.Last().CulledFunctionNames.Add(FunctionScriptName);
@@ -8031,13 +8057,13 @@ void TNiagaraHlslTranslator<GraphBridge>::ProcessCustomHlsl(const FString& InCus
 		// Clean up any namespaced variables in the token list if they are aliased or promote any tokens that are namespaced to the parent 
 		// parameter map.
 		TArray<FString> PossibleNamespaces;
-		FParamMapHistory::GetValidNamespacesForReading(CompileOptions.TargetUsage, 0, PossibleNamespaces);
+		FNiagaraParameterUtilities::GetValidNamespacesForReading(CompileOptions.TargetUsage, 0, PossibleNamespaces);
 
 		for (FParamMapHistory& History : ParamMapHistories)
 		{
 			for (FNiagaraVariable& Var : History.Variables)
 			{
-				FString Namespace = FParamMapHistory::GetNamespace(Var);
+				FString Namespace = FNiagaraParameterUtilities::GetNamespace(Var);
 				PossibleNamespaces.AddUnique(Namespace);
 			}
 		}
@@ -8052,8 +8078,8 @@ void TNiagaraHlslTranslator<GraphBridge>::ProcessCustomHlsl(const FString& InCus
 			FNiagaraVariable VarAliased = UniqueParamMapEntriesAliasesIntact[VarIdx];
 			FNiagaraVariable VarActual = UniqueParamMapEntries[VarIdx];
 
-			if (FParamMapHistory::IsExternalConstantNamespace(VarAliased, CompileOptions.TargetUsage, CompileOptions.GetTargetUsageBitmask()) ||
-				FParamMapHistory::IsExternalConstantNamespace(VarActual, CompileOptions.TargetUsage, CompileOptions.GetTargetUsageBitmask()))
+			if (FNiagaraParameterUtilities::IsExternalConstantNamespace(VarAliased, CompileOptions.TargetUsage, CompileOptions.GetTargetUsageBitmask()) ||
+				FNiagaraParameterUtilities::IsExternalConstantNamespace(VarActual, CompileOptions.TargetUsage, CompileOptions.GetTargetUsageBitmask()))
 			{
 				int32 TempOutput;
 				if (ParameterMapRegisterExternalConstantNamespaceVariable(VarActual, InNodeForErrorReporting, ParamMapHistoryIdx, TempOutput, nullptr))
@@ -8115,7 +8141,7 @@ void TNiagaraHlslTranslator<GraphBridge>::HandleCustomHlslNode(const FCustomHlsl
 	OutSignature.Name = *OutName; // Force the name to be set to include the node guid for safety...
 	bOutCustomHlsl = true;
 	OutCustomHlsl = GraphBridge::GetCustomHlslString(CustomFunctionHlsl);
-	OutCustomHlslIncludeFilePaths = GraphBridge::GetCustomHlslIncludePaths(CustomFunctionHlsl);
+	GraphBridge::GetCustomHlslIncludePaths(CustomFunctionHlsl, OutCustomHlslIncludeFilePaths);
 
 	FNiagaraFunctionSignature InSignature = CustomFunctionHlsl->Signature;
 	ProcessCustomHlsl(OutCustomHlsl, OutScriptUsage, InSignature, Inputs, CustomFunctionHlsl, OutCustomHlsl, OutSignature);
@@ -8861,9 +8887,9 @@ FName FNiagaraHlslTranslator::GetDataInterfaceName(FName BaseName, const FString
 {
 	if (UniqueEmitterName.IsEmpty() == false)
 	{
-		if (FNiagaraParameterMapHistory::IsAliasedEmitterParameter(BaseName.ToString()))
+		if (FNiagaraParameterUtilities::IsAliasedEmitterParameter(BaseName.ToString()))
 		{
-			return FNiagaraParameterMapHistory::ResolveEmitterAlias(BaseName, UniqueEmitterName);
+			return FNiagaraParameterUtilities::ResolveEmitterAlias(BaseName, UniqueEmitterName);
 		}
 		else if(bIsParameterMapDataInterface == false)
 		{
@@ -9523,11 +9549,11 @@ int32 TNiagaraHlslTranslator<GraphBridge>::CompileInputPin(const FInputPin* Pin)
 
 	if (const FOutputPin* LinkedPin = GraphBridge::GetLinkedOutputPin(Pin))
 	{
-		FPinConnectionResponse ConnectionResponse = GraphBridge::CanCreateConnection(LinkedPin, Pin);
-		if (ConnectionResponse.Response == CONNECT_RESPONSE_DISALLOW)
+		FText ConnectionMessage;
+		if (!GraphBridge::CanCreateConnection(LinkedPin, Pin, ConnectionMessage))
 		{
 			FText OutErrorMessage = FText::Format(LOCTEXT("InputConnectionDisallowed", "Input connection is not allowed! Reason: {0}"),
-				ConnectionResponse.Message);
+				ConnectionMessage);
 
 			Error(OutErrorMessage, GraphBridge::GetOwningNode(Pin), Pin);
 		}
@@ -9567,7 +9593,7 @@ int32 TNiagaraHlslTranslator<GraphBridge>::CompileOutputPin(const FPin* InPin)
 		return INDEX_NONE;
 	}
 
-	UpdateStaticSwitchConstants(GraphBridge::GetMutableOwningNode(InPin));
+	UpdateStaticSwitchConstants(InPin);
 
 	// The incoming pin to compile may be pointing to a reroute node. If so, we just jump over it
 	// to where it really came from.
diff --git a/Engine/Plugins/FX/Niagara/Source/NiagaraEditor/Private/NiagaraHlslTranslator.h b/Engine/Plugins/FX/Niagara/Source/NiagaraEditor/Private/NiagaraHlslTranslator.h
--- a/Engine/Plugins/FX/Niagara/Source/NiagaraEditor/Private/NiagaraHlslTranslator.h
+++ b/Engine/Plugins/FX/Niagara/Source/NiagaraEditor/Private/NiagaraHlslTranslator.h
@@ -4,9 +4,11 @@
 
 #include "CoreMinimal.h"
 
-#include "INiagaraCompiler.h"
 #include "Kismet2/CompilerResultsLog.h"
+#include "INiagaraCompiler.h"
 #include "NiagaraCommon.h"
+#include "NiagaraCompilationBridge.h"
+#include "NiagaraParameterMapHistory.h"
 #include "NiagaraScript.h"
 #include "NiagaraScriptSourceBase.h"
 #include "NiagaraShared.h"
@@ -14,7 +16,17 @@
 
 #include "NiagaraHlslTranslator.generated.h"
 
+class FNiagaraCompileOptions;
+class FNiagaraCompileRequestData;
+class FNiagaraCompileRequestDuplicateData;
+struct FNiagaraGraphHelper;
 struct FNiagaraTranslatorOutput;
+class UNiagaraNode;
+class UNiagaraNodeFunctionCall;
+class UNiagaraNodeOutput;
+class UNiagaraNodeParameterMapGet;
+class UNiagaraNodeParameterMapSet;
+class UNiagaraScriptSource;
 
 UENUM()
 enum class ENiagaraDataSetAccessMode : uint8
@@ -562,7 +574,6 @@ class TNiagaraHlslTranslator : public FNiagaraHlslTranslator, private GraphBridg
 {
 	using typename GraphBridge::FGraph;
 	using typename GraphBridge::FPin;
-	using typename GraphBridge::FScopedPin;
 	using typename GraphBridge::FNode;
 	using typename GraphBridge::FParamMapHistory;
 	using typename GraphBridge::FParamMapHistoryBuilder;
@@ -584,9 +595,9 @@ class TNiagaraHlslTranslator : public FNiagaraHlslTranslator, private GraphBridg
 	using typename GraphBridge::FInputPin;
 	using typename GraphBridge::FOutputPin;
 	using typename GraphBridge::FModuleScopedPin;
-	using typename GraphBridge::FGraphTraversalHandle;
 	using typename GraphBridge::FGraphFunctionAliasContext;
 	using typename GraphBridge::FConvertConnection;
+	using typename GraphBridge::FParameterCollection;
 
 public:
 	using FBridge = GraphBridge;
@@ -711,7 +722,7 @@ protected:
 
 	void ValidateFailIfPreviouslyNotSet(const FNiagaraVariable& InVar, bool& bFailIfNotSet);
 
-	void UpdateStaticSwitchConstants(FNode* Node);
+	void UpdateStaticSwitchConstants(const FPin* Pin);
 
 	virtual const FString& GetEmitterUniqueName() const override;
 	virtual TConstArrayView<FNiagaraVariable> GetStaticVariables() const override;
diff --git a/Engine/Plugins/FX/Niagara/Source/NiagaraEditor/Private/NiagaraNode.cpp b/Engine/Plugins/FX/Niagara/Source/NiagaraEditor/Private/NiagaraNode.cpp
--- a/Engine/Plugins/FX/Niagara/Source/NiagaraEditor/Private/NiagaraNode.cpp
+++ b/Engine/Plugins/FX/Niagara/Source/NiagaraEditor/Private/NiagaraNode.cpp
@@ -3,7 +3,7 @@
 #include "NiagaraNode.h"
 #include "NiagaraGraph.h"
 #include "EdGraphSchema_Niagara.h"
-#include "NiagaraGraphHlslTranslator.h"
+#include "NiagaraHlslTranslator.h"
 #include "GraphEditAction.h"
 #include "Widgets/SNiagaraGraphNode.h"
 #include "Misc/SecureHash.h"
@@ -947,7 +947,7 @@ void UNiagaraNode::RegisterPassthroughPin(FNiagaraParameterMapHistoryBuilder& Ou
 {
 	const UEdGraphSchema_Niagara* Schema = CastChecked<UEdGraphSchema_Niagara>(GetSchema());
 	if (bVisitInputPin)
-		OutHistory.VisitInputPin(InputPin, this, bFilterForCompilation);
+		OutHistory.VisitInputPin(InputPin, bFilterForCompilation);
 
 	FNiagaraTypeDefinition InDef = Schema->PinToTypeDefinition(InputPin);
 	FNiagaraTypeDefinition OutDef = Schema->PinToTypeDefinition(OutputPin);
@@ -1019,4 +1019,26 @@ void UNiagaraNode::GatherForLocalization(FPropertyLocalizationDataGatherer& Prop
 }
 #endif
 
+void UNiagaraNode::GetCompilationInputPins(FPinCollectorArray& InputPins) const
+{
+	for (UEdGraphPin* Pin : Pins)
+	{
+		if (Pin->Direction == EGPD_Input && !Pin->bOrphanedPin)
+		{
+			InputPins.Add(Pin);
+		}
+	}
+}
+
+void UNiagaraNode::GetCompilationOutputPins(FPinCollectorArray& OutputPins) const
+{
+	for (UEdGraphPin* Pin : Pins)
+	{
+		if (Pin->Direction == EGPD_Output && !Pin->bOrphanedPin)
+		{
+			OutputPins.Add(Pin);
+		}
+	}
+}
+
 #undef LOCTEXT_NAMESPACE
diff --git a/Engine/Plugins/FX/Niagara/Source/NiagaraEditor/Private/NiagaraNodeConvert.cpp b/Engine/Plugins/FX/Niagara/Source/NiagaraEditor/Private/NiagaraNodeConvert.cpp
--- a/Engine/Plugins/FX/Niagara/Source/NiagaraEditor/Private/NiagaraNodeConvert.cpp
+++ b/Engine/Plugins/FX/Niagara/Source/NiagaraEditor/Private/NiagaraNodeConvert.cpp
@@ -4,7 +4,7 @@
 #include "EdGraphSchema_Niagara.h"
 #include "NiagaraEditorUtilities.h"
 #include "Widgets/SNiagaraGraphNodeConvert.h"
-#include "NiagaraGraphHlslTranslator.h"
+#include "NiagaraHlslTranslator.h"
 #include "UObject/UnrealType.h"
 #include "Kismet2/StructureEditorUtils.h"
 
@@ -152,7 +152,7 @@ TSharedPtr<SGraphNode> UNiagaraNodeConvert::CreateVisualWidget()
 void UNiagaraNodeConvert::Compile(FTranslator* Translator, TArray<int32>& CompileOutputs) const
 {
 	FPinCollectorArray InputPins;
-	FTranslator::FBridge::GetCompilationInputPins(this, InputPins);
+	GetCompilationInputPins(InputPins);
 
 	TArray<int32, TInlineAllocator<16>> CompileInputs;
 	CompileInputs.Reserve(InputPins.Num());
@@ -176,7 +176,7 @@ void UNiagaraNodeConvert::Compile(FTranslator* Translator, TArray<int32>& Compil
 	check(Schema);
 
 	FPinCollectorArray OutputPins;
-	FTranslator::FBridge::GetCompilationOutputPins(this, OutputPins);
+	GetCompilationOutputPins(OutputPins);
 
 	// Go through all the output nodes and cross-reference them with the connections list. Output errors if any connections are incomplete.
 	{
diff --git a/Engine/Plugins/FX/Niagara/Source/NiagaraEditor/Private/NiagaraNodeCustomHlsl.cpp b/Engine/Plugins/FX/Niagara/Source/NiagaraEditor/Private/NiagaraNodeCustomHlsl.cpp
--- a/Engine/Plugins/FX/Niagara/Source/NiagaraEditor/Private/NiagaraNodeCustomHlsl.cpp
+++ b/Engine/Plugins/FX/Niagara/Source/NiagaraEditor/Private/NiagaraNodeCustomHlsl.cpp
@@ -459,7 +459,7 @@ void UNiagaraNodeCustomHlsl::BuildParameterMapHistory(FNiagaraParameterMapHistor
 		}
 
 		TArray<FString> PossibleNamespaces;
-		FNiagaraParameterMapHistory::GetValidNamespacesForReading(OutHistory.GetBaseUsageContext(), 0, PossibleNamespaces);
+		FNiagaraParameterUtilities::GetValidNamespacesForReading(OutHistory.GetBaseUsageContext(), 0, PossibleNamespaces);
 
 		if ((bHasParamMapOutput || bHasParamMapInput) && ParamMapIdx != INDEX_NONE)
 		{
diff --git a/Engine/Plugins/FX/Niagara/Source/NiagaraEditor/Private/NiagaraNodeEmitter.cpp b/Engine/Plugins/FX/Niagara/Source/NiagaraEditor/Private/NiagaraNodeEmitter.cpp
--- a/Engine/Plugins/FX/Niagara/Source/NiagaraEditor/Private/NiagaraNodeEmitter.cpp
+++ b/Engine/Plugins/FX/Niagara/Source/NiagaraEditor/Private/NiagaraNodeEmitter.cpp
@@ -10,7 +10,7 @@
 #include "NiagaraGraph.h"
 #include "NiagaraNodeInput.h"
 #include "NiagaraNodeOutput.h"
-#include "NiagaraGraphHlslTranslator.h"
+#include "NiagaraHlslTranslator.h"
 #include "Stats/Stats.h"
 #include "NiagaraEditorModule.h"
 
@@ -347,7 +347,7 @@ void UNiagaraNodeEmitter::BuildParameterMapHistory(FNiagaraParameterMapHistoryBu
 
 			// Build up a new parameter map history with all the child graph nodes..
 			FNiagaraParameterMapHistoryBuilder ChildBuilder;
-			ChildBuilder.ConstantResolver = OutHistory.ConstantResolver;
+			*ChildBuilder.ConstantResolver = *OutHistory.ConstantResolver;
 			ChildBuilder.RegisterEncounterableVariables(OutHistory.GetEncounterableVariables());
 			ChildBuilder.EnableScriptAllowList(true, GetUsage());
 
@@ -390,9 +390,7 @@ void UNiagaraNodeEmitter::BuildParameterMapHistory(FNiagaraParameterMapHistoryBu
 						OutHistory.Histories[ParamMapIdx].PerVariableConstantValue[ExistingIdx].AddUnique(ConstantStr);	
 					}
 				}
-				OutHistory.Histories[ParamMapIdx].ParameterCollections.Append(History.ParameterCollections);
-				OutHistory.Histories[ParamMapIdx].ParameterCollectionNamespaces.Append(History.ParameterCollectionNamespaces);
-				OutHistory.Histories[ParamMapIdx].ParameterCollectionVariables.Append(History.ParameterCollectionVariables);
+				OutHistory.Histories[ParamMapIdx].EncounteredParameterCollections.Append(History.EncounteredParameterCollections);
 				OutHistory.Histories[ParamMapIdx].PinToConstantValues.Append(History.PinToConstantValues);
 			}
 
@@ -448,7 +446,7 @@ void UNiagaraNodeEmitter::Compile(FTranslator* Translator, TArray<int32>& Output
 	InputPins.RemoveAll([](UEdGraphPin* InputPin) { return (InputPin->PinType.PinCategory != UEdGraphSchema_Niagara::PinCategoryType) && (InputPin->PinType.PinCategory != UEdGraphSchema_Niagara::PinCategoryEnum); });
 
 	FPinCollectorArray OutputPins;
-	FTranslator::FBridge::GetCompilationOutputPins(this, OutputPins);
+	GetCompilationOutputPins(OutputPins);
 
 	check(Outputs.Num() == 0);
 
diff --git a/Engine/Plugins/FX/Niagara/Source/NiagaraEditor/Private/NiagaraNodeFunctionCall.cpp b/Engine/Plugins/FX/Niagara/Source/NiagaraEditor/Private/NiagaraNodeFunctionCall.cpp
--- a/Engine/Plugins/FX/Niagara/Source/NiagaraEditor/Private/NiagaraNodeFunctionCall.cpp
+++ b/Engine/Plugins/FX/Niagara/Source/NiagaraEditor/Private/NiagaraNodeFunctionCall.cpp
@@ -11,7 +11,7 @@
 #include "Subsystems/AssetEditorSubsystem.h"
 #include "NiagaraEditorUtilities.h"
 #include "NiagaraGraph.h"
-#include "NiagaraGraphHlslTranslator.h"
+#include "NiagaraHlslTranslator.h"
 #include "NiagaraMessages.h"
 #include "NiagaraNodeInput.h"
 #include "NiagaraNodeOutput.h"
@@ -1290,11 +1290,11 @@ void UNiagaraNodeFunctionCall::BuildParameterMapHistory(FNiagaraParameterMapHist
 			// Make sure to register pin constants
 			if (Pin->LinkedTo.Num() == 0 && Schema->IsPinStatic(Pin) && Pin->DefaultValue.Len() != 0)
 			{
-				OutHistory.RegisterConstantFromInputPin(Pin);
+				OutHistory.RegisterConstantFromInputPin(Pin, Pin->DefaultValue);
 			}
 		}
 
-		FCompileConstantResolver FunctionResolver = OutHistory.ConstantResolver.WithDebugState(DebugState);
+		FCompileConstantResolver FunctionResolver = OutHistory.ConstantResolver->WithDebugState(DebugState);
 		if (OutHistory.HasCurrentUsageContext())
 		{
 			// if we traverse a full emitter graph the usage might change during the traversal, so we need to update the constant resolver
@@ -1391,7 +1391,7 @@ void UNiagaraNodeFunctionCall::BuildParameterMapHistory(FNiagaraParameterMapHist
 			OutHistory.EndNodeVisitation(ParamMapIdx, NodeIdx);
 		}
 
-		OutHistory.ExitFunction(GetFunctionName(), this);
+		OutHistory.ExitFunction(this);
 
 		if (DoDepthTraversal)
 		{
diff --git a/Engine/Plugins/FX/Niagara/Source/NiagaraEditor/Private/NiagaraNodeIf.cpp b/Engine/Plugins/FX/Niagara/Source/NiagaraEditor/Private/NiagaraNodeIf.cpp
--- a/Engine/Plugins/FX/Niagara/Source/NiagaraEditor/Private/NiagaraNodeIf.cpp
+++ b/Engine/Plugins/FX/Niagara/Source/NiagaraEditor/Private/NiagaraNodeIf.cpp
@@ -2,7 +2,7 @@
 
 #include "NiagaraNodeIf.h"
 #include "NiagaraEditorUtilities.h"
-#include "NiagaraGraphHlslTranslator.h"
+#include "NiagaraHlslTranslator.h"
 
 #include UE_INLINE_GENERATED_CPP_BY_NAME(NiagaraNodeIf)
 
diff --git a/Engine/Plugins/FX/Niagara/Source/NiagaraEditor/Private/NiagaraNodeInput.cpp b/Engine/Plugins/FX/Niagara/Source/NiagaraEditor/Private/NiagaraNodeInput.cpp
--- a/Engine/Plugins/FX/Niagara/Source/NiagaraEditor/Private/NiagaraNodeInput.cpp
+++ b/Engine/Plugins/FX/Niagara/Source/NiagaraEditor/Private/NiagaraNodeInput.cpp
@@ -2,7 +2,7 @@
 
 #include "NiagaraNodeInput.h"
 #include "UObject/UnrealType.h"
-#include "NiagaraGraphHlslTranslator.h"
+#include "NiagaraHlslTranslator.h"
 #include "NiagaraGraph.h"
 
 #include "NiagaraNodeOutput.h"
diff --git a/Engine/Plugins/FX/Niagara/Source/NiagaraEditor/Private/NiagaraNodeOp.cpp b/Engine/Plugins/FX/Niagara/Source/NiagaraEditor/Private/NiagaraNodeOp.cpp
--- a/Engine/Plugins/FX/Niagara/Source/NiagaraEditor/Private/NiagaraNodeOp.cpp
+++ b/Engine/Plugins/FX/Niagara/Source/NiagaraEditor/Private/NiagaraNodeOp.cpp
@@ -1,7 +1,7 @@
 // Copyright Epic Games, Inc. All Rights Reserved.
 
 #include "NiagaraNodeOp.h"
-#include "NiagaraGraphHlslTranslator.h"
+#include "NiagaraHlslTranslator.h"
 #include "GraphEditorSettings.h"
 
 #include "EdGraphSchema_Niagara.h"
diff --git a/Engine/Plugins/FX/Niagara/Source/NiagaraEditor/Private/NiagaraNodeOutput.cpp b/Engine/Plugins/FX/Niagara/Source/NiagaraEditor/Private/NiagaraNodeOutput.cpp
--- a/Engine/Plugins/FX/Niagara/Source/NiagaraEditor/Private/NiagaraNodeOutput.cpp
+++ b/Engine/Plugins/FX/Niagara/Source/NiagaraEditor/Private/NiagaraNodeOutput.cpp
@@ -2,7 +2,7 @@
 
 #include "NiagaraNodeOutput.h"
 #include "UObject/UnrealType.h"
-#include "NiagaraGraphHlslTranslator.h"
+#include "NiagaraHlslTranslator.h"
 #include "NiagaraEmitter.h"
 #include "NiagaraScript.h"
 #include "NiagaraGraph.h"
@@ -307,7 +307,7 @@ void UNiagaraNodeOutput::BuildParameterMapHistory(FNiagaraParameterMapHistoryBui
 			FNiagaraTypeDefinition InputType = Schema->PinToTypeDefinition(Pin);
 			if (InputType.IsStatic())
 			{
-				OutHistory.RegisterConstantFromInputPin(Pin);
+				OutHistory.RegisterConstantFromInputPin(Pin, Pin->DefaultValue);
 			}
 			else if (InputType == FNiagaraTypeDefinition::GetParameterMapDef())
 			{
diff --git a/Engine/Plugins/FX/Niagara/Source/NiagaraEditor/Private/NiagaraNodeOutputTag.cpp b/Engine/Plugins/FX/Niagara/Source/NiagaraEditor/Private/NiagaraNodeOutputTag.cpp
--- a/Engine/Plugins/FX/Niagara/Source/NiagaraEditor/Private/NiagaraNodeOutputTag.cpp
+++ b/Engine/Plugins/FX/Niagara/Source/NiagaraEditor/Private/NiagaraNodeOutputTag.cpp
@@ -2,7 +2,7 @@
 
 #include "NiagaraNodeOutputTag.h"
 #include "NiagaraCustomVersion.h"
-#include "NiagaraGraphHlslTranslator.h"
+#include "NiagaraHlslTranslator.h"
 #include "Widgets/SNiagaraGraphNodeCustomHlsl.h"
 #include "EdGraphSchema_Niagara.h"
 #include "ScopedTransaction.h"
@@ -193,7 +193,7 @@ void UNiagaraNodeOutputTag::BuildParameterMapHistory(FNiagaraParameterMapHistory
 			continue;
 		}
 
-		OutHistory.VisitInputPin(InputPin, this, bFilterForCompilation);
+		OutHistory.VisitInputPin(InputPin, bFilterForCompilation);
 
 
 		if (!IsNodeEnabled() && OutHistory.GetIgnoreDisabled())
diff --git a/Engine/Plugins/FX/Niagara/Source/NiagaraEditor/Private/NiagaraNodeParameterMapFor.cpp b/Engine/Plugins/FX/Niagara/Source/NiagaraEditor/Private/NiagaraNodeParameterMapFor.cpp
--- a/Engine/Plugins/FX/Niagara/Source/NiagaraEditor/Private/NiagaraNodeParameterMapFor.cpp
+++ b/Engine/Plugins/FX/Niagara/Source/NiagaraEditor/Private/NiagaraNodeParameterMapFor.cpp
@@ -3,7 +3,7 @@
 #include "NiagaraNodeParameterMapFor.h"
 #include "EdGraphSchema_Niagara.h"
 #include "NiagaraEditorUtilities.h"
-#include "NiagaraGraphHlslTranslator.h"
+#include "NiagaraHlslTranslator.h"
 #include "EdGraph/EdGraphNode.h"
 
 #include UE_INLINE_GENERATED_CPP_BY_NAME(NiagaraNodeParameterMapFor)
diff --git a/Engine/Plugins/FX/Niagara/Source/NiagaraEditor/Private/NiagaraNodeParameterMapGet.cpp b/Engine/Plugins/FX/Niagara/Source/NiagaraEditor/Private/NiagaraNodeParameterMapGet.cpp
--- a/Engine/Plugins/FX/Niagara/Source/NiagaraEditor/Private/NiagaraNodeParameterMapGet.cpp
+++ b/Engine/Plugins/FX/Niagara/Source/NiagaraEditor/Private/NiagaraNodeParameterMapGet.cpp
@@ -4,7 +4,7 @@
 #include "EdGraphSchema_Niagara.h"
 #include "NiagaraEditorUtilities.h"
 #include "NiagaraGraph.h"
-#include "NiagaraGraphHlslTranslator.h"
+#include "NiagaraHlslTranslator.h"
 #include "ScopedTransaction.h"
 #include "Widgets/SNiagaraGraphParameterMapGetNode.h"
 #include "NiagaraEditorModule.h"
@@ -415,7 +415,7 @@ void UNiagaraNodeParameterMapGet::SynchronizeDefaultInputPin(UEdGraphPin* Defaul
 	}
 	
 	FNiagaraVariable Var = Schema->PinToNiagaraVariable(OutputPin);
-	if (FNiagaraParameterMapHistory::IsEngineParameter(Var))
+	if (FNiagaraParameterUtilities::IsEngineParameter(Var))
 	{
 		DefaultPin->bDefaultValueIsIgnored = true;
 		DefaultPin->bNotConnectable = true;
@@ -476,7 +476,7 @@ void UNiagaraNodeParameterMapGet::BuildParameterMapHistory(FNiagaraParameterMapH
 {
 	if (bRecursive)
 	{
-		OutHistory.VisitInputPin(GetInputPin(0), this, bFilterForCompilation);
+		OutHistory.VisitInputPin(GetInputPin(0), bFilterForCompilation);
 	}
 
 	if (!IsNodeEnabled() && OutHistory.GetIgnoreDisabled())
@@ -522,9 +522,9 @@ void UNiagaraNodeParameterMapGet::Compile(FTranslator* Translator, TArray<int32>
 	const UEdGraphSchema_Niagara* Schema = CastChecked<UEdGraphSchema_Niagara>(GetSchema());
 
 	FPinCollectorArray InputPins;
-	FTranslator::FBridge::GetCompilationInputPins(this, InputPins);
+	GetCompilationInputPins(InputPins);
 	FPinCollectorArray OutputPins;
-	FTranslator::FBridge::GetCompilationOutputPins(this, OutputPins);
+	GetCompilationOutputPins(OutputPins);
 
 	// Initialize the outputs to invalid values.
 	check(Outputs.Num() == 0);
diff --git a/Engine/Plugins/FX/Niagara/Source/NiagaraEditor/Private/NiagaraNodeParameterMapSet.cpp b/Engine/Plugins/FX/Niagara/Source/NiagaraEditor/Private/NiagaraNodeParameterMapSet.cpp
--- a/Engine/Plugins/FX/Niagara/Source/NiagaraEditor/Private/NiagaraNodeParameterMapSet.cpp
+++ b/Engine/Plugins/FX/Niagara/Source/NiagaraEditor/Private/NiagaraNodeParameterMapSet.cpp
@@ -7,7 +7,7 @@
 #include "EdGraphSchema_Niagara.h"
 #include "NiagaraConstants.h"
 #include "NiagaraEditorUtilities.h"
-#include "NiagaraGraphHlslTranslator.h"
+#include "NiagaraHlslTranslator.h"
 #include "NiagaraGraph.h"
 #include "NiagaraScriptVariable.h"
 #include "ScopedTransaction.h"
@@ -195,10 +195,10 @@ bool UNiagaraNodeParameterMapSet::CommitEditablePinName(const FText& InName, UEd
 void UNiagaraNodeParameterMapSet::Compile(FTranslator* Translator, TArray<int32>& Outputs) const
 {
 	FPinCollectorArray InputPins;
-	FTranslator::FBridge::GetCompilationInputPins(this, InputPins);
+	GetCompilationInputPins(InputPins);
 
 	FPinCollectorArray OutputPins;
-	FTranslator::FBridge::GetCompilationOutputPins(this, OutputPins);
+	GetCompilationOutputPins(OutputPins);
 
 	// Initialize the outputs to invalid values.
 	check(Outputs.Num() == 0);
@@ -267,7 +267,7 @@ void UNiagaraNodeParameterMapSet::BuildParameterMapHistory(FNiagaraParameterMapH
 
 	for (UEdGraphPin* InputPin : InputPins)
 	{
-		OutHistory.VisitInputPin(InputPin, this, bFilterForCompilation);
+		OutHistory.VisitInputPin(InputPin, bFilterForCompilation);
 	}
 
 	if (IsNodeEnabled() || !OutHistory.GetIgnoreDisabled())
diff --git a/Engine/Plugins/FX/Niagara/Source/NiagaraEditor/Private/NiagaraNodeReadDataSet.cpp b/Engine/Plugins/FX/Niagara/Source/NiagaraEditor/Private/NiagaraNodeReadDataSet.cpp
--- a/Engine/Plugins/FX/Niagara/Source/NiagaraEditor/Private/NiagaraNodeReadDataSet.cpp
+++ b/Engine/Plugins/FX/Niagara/Source/NiagaraEditor/Private/NiagaraNodeReadDataSet.cpp
@@ -2,7 +2,7 @@
 
 #include "NiagaraNodeReadDataSet.h"
 #include "UObject/UnrealType.h"
-#include "NiagaraGraphHlslTranslator.h"
+#include "NiagaraHlslTranslator.h"
 #include "NiagaraEvents.h"
 #include "EdGraphSchema_Niagara.h"
 #include "NiagaraGraph.h"
diff --git a/Engine/Plugins/FX/Niagara/Source/NiagaraEditor/Private/NiagaraNodeReroute.cpp b/Engine/Plugins/FX/Niagara/Source/NiagaraEditor/Private/NiagaraNodeReroute.cpp
--- a/Engine/Plugins/FX/Niagara/Source/NiagaraEditor/Private/NiagaraNodeReroute.cpp
+++ b/Engine/Plugins/FX/Niagara/Source/NiagaraEditor/Private/NiagaraNodeReroute.cpp
@@ -2,7 +2,7 @@
 
 #include "NiagaraNodeReroute.h"
 #include "NiagaraEditorUtilities.h"
-#include "NiagaraGraphHlslTranslator.h"
+#include "NiagaraHlslTranslator.h"
 #include "NiagaraConstants.h"
 
 #include UE_INLINE_GENERATED_CPP_BY_NAME(NiagaraNodeReroute)
@@ -109,10 +109,10 @@ bool UNiagaraNodeReroute::ShouldDrawNodeAsControlPointOnly(int32& OutInputPinInd
 void UNiagaraNodeReroute::Compile(FTranslator* Translator, TArray<int32>& Outputs) const
 {
 	FPinCollectorArray InputPins;
-	FTranslator::FBridge::GetCompilationInputPins(this, InputPins);
+	GetCompilationInputPins(InputPins);
 
 	FPinCollectorArray OutputPins;
-	FTranslator::FBridge::GetCompilationOutputPins(this, OutputPins);
+	GetCompilationOutputPins(OutputPins);
 
 	// Initialize the outputs to invalid values.
 	check(Outputs.Num() == 0);
diff --git a/Engine/Plugins/FX/Niagara/Source/NiagaraEditor/Private/NiagaraNodeSelect.cpp b/Engine/Plugins/FX/Niagara/Source/NiagaraEditor/Private/NiagaraNodeSelect.cpp
--- a/Engine/Plugins/FX/Niagara/Source/NiagaraEditor/Private/NiagaraNodeSelect.cpp
+++ b/Engine/Plugins/FX/Niagara/Source/NiagaraEditor/Private/NiagaraNodeSelect.cpp
@@ -6,7 +6,7 @@
 #include "FindInBlueprintManager.h"
 #include "NiagaraEditorStyle.h"
 #include "NiagaraEditorUtilities.h"
-#include "NiagaraGraphHlslTranslator.h"
+#include "NiagaraHlslTranslator.h"
 #include "EdGraph/EdGraphPin.h"
 #include "Widgets/SNiagaraPinTypeSelector.h"
 #include "ScopedTransaction.h"
diff --git a/Engine/Plugins/FX/Niagara/Source/NiagaraEditor/Private/NiagaraNodeSimTargetSelector.cpp b/Engine/Plugins/FX/Niagara/Source/NiagaraEditor/Private/NiagaraNodeSimTargetSelector.cpp
--- a/Engine/Plugins/FX/Niagara/Source/NiagaraEditor/Private/NiagaraNodeSimTargetSelector.cpp
+++ b/Engine/Plugins/FX/Niagara/Source/NiagaraEditor/Private/NiagaraNodeSimTargetSelector.cpp
@@ -2,7 +2,7 @@
 
 #include "NiagaraNodeSimTargetSelector.h"
 #include "NiagaraEditorUtilities.h"
-#include "NiagaraGraphHlslTranslator.h"
+#include "NiagaraHlslTranslator.h"
 
 #include UE_INLINE_GENERATED_CPP_BY_NAME(NiagaraNodeSimTargetSelector)
 
@@ -50,9 +50,9 @@ FString UNiagaraNodeSimTargetSelector::GetInputCaseName(int32 Case) const
 void UNiagaraNodeSimTargetSelector::Compile(FTranslator* Translator, TArray<int32>& Outputs) const
 {
 	FPinCollectorArray InputPins;
-	FTranslator::FBridge::GetCompilationInputPins(this, InputPins);
+	GetCompilationInputPins(InputPins);
 	FPinCollectorArray OutputPins;
-	FTranslator::FBridge::GetCompilationOutputPins(this, OutputPins);
+	GetCompilationOutputPins(OutputPins);
 
 	//ENiagaraSimTarget SimulationTarget = Translator->GetSimulationTarget();
 	bool bCPUSim = Translator->IsCompileOptionDefined(*FNiagaraCompileOptions::CpuScriptDefine);
diff --git a/Engine/Plugins/FX/Niagara/Source/NiagaraEditor/Private/NiagaraNodeStaticSwitch.cpp b/Engine/Plugins/FX/Niagara/Source/NiagaraEditor/Private/NiagaraNodeStaticSwitch.cpp
--- a/Engine/Plugins/FX/Niagara/Source/NiagaraEditor/Private/NiagaraNodeStaticSwitch.cpp
+++ b/Engine/Plugins/FX/Niagara/Source/NiagaraEditor/Private/NiagaraNodeStaticSwitch.cpp
@@ -2,7 +2,7 @@
 
 #include "NiagaraNodeStaticSwitch.h"
 #include "NiagaraEditorUtilities.h"
-#include "NiagaraGraphHlslTranslator.h"
+#include "NiagaraHlslTranslator.h"
 #include "NiagaraConstants.h"
 #include "NiagaraEditorStyle.h"
 #include "NiagaraScriptVariable.h"
@@ -567,7 +567,7 @@ void UNiagaraNodeStaticSwitch::Compile(FTranslator* Translator, TArray<int32>& O
 		}
 	}
 	FPinCollectorArray OutputPins;
-	FTranslator::FBridge::GetCompilationOutputPins(this, OutputPins);
+	GetCompilationOutputPins(OutputPins);
 
 	// Initialize the outputs to invalid values.
 	check(Outputs.Num() == 0);
@@ -824,7 +824,7 @@ void UNiagaraNodeStaticSwitch::BuildParameterMapHistory(FNiagaraParameterMapHist
 		if (IsSetByPin())
 		{
 			UEdGraphPin* SelectorPin = GetSelectorPin();
-			OutHistory.VisitInputPin(SelectorPin, this, bFilterForCompilation);
+			OutHistory.VisitInputPin(SelectorPin, bFilterForCompilation);
 
 			OutHistory.SetConstantByStaticVariable(SelectorValue, SelectorPin);
 
diff --git a/Engine/Plugins/FX/Niagara/Source/NiagaraEditor/Private/NiagaraNodeUsageSelector.cpp b/Engine/Plugins/FX/Niagara/Source/NiagaraEditor/Private/NiagaraNodeUsageSelector.cpp
--- a/Engine/Plugins/FX/Niagara/Source/NiagaraEditor/Private/NiagaraNodeUsageSelector.cpp
+++ b/Engine/Plugins/FX/Niagara/Source/NiagaraEditor/Private/NiagaraNodeUsageSelector.cpp
@@ -2,7 +2,7 @@
 
 #include "NiagaraNodeUsageSelector.h"
 #include "NiagaraEditorUtilities.h"
-#include "NiagaraGraphHlslTranslator.h"
+#include "NiagaraHlslTranslator.h"
 #include "NiagaraScriptSource.h"
 #include "Widgets/Layout/SSeparator.h"
 
@@ -298,9 +298,9 @@ void UNiagaraNodeUsageSelector::AddWidgetsToInputBox(TSharedPtr<SVerticalBox> In
 void UNiagaraNodeUsageSelector::Compile(FTranslator* Translator, TArray<int32>& Outputs) const
 {
 	FPinCollectorArray InputPins;
-	FTranslator::FBridge::GetCompilationInputPins(this, InputPins);
+	GetCompilationInputPins(InputPins);
 	FPinCollectorArray OutputPins;
-	FTranslator::FBridge::GetCompilationOutputPins(this, OutputPins);
+	GetCompilationOutputPins(OutputPins);
 
 	ENiagaraScriptUsage CurrentUsage = Translator->GetCurrentUsage();
 	ENiagaraScriptGroup UsageGroup = ENiagaraScriptGroup::Max;
diff --git a/Engine/Plugins/FX/Niagara/Source/NiagaraEditor/Private/NiagaraNodeWithDynamicPins.cpp b/Engine/Plugins/FX/Niagara/Source/NiagaraEditor/Private/NiagaraNodeWithDynamicPins.cpp
--- a/Engine/Plugins/FX/Niagara/Source/NiagaraEditor/Private/NiagaraNodeWithDynamicPins.cpp
+++ b/Engine/Plugins/FX/Niagara/Source/NiagaraEditor/Private/NiagaraNodeWithDynamicPins.cpp
@@ -507,5 +507,27 @@ void UNiagaraNodeWithDynamicPins::MoveDynamicPinFromMenu(UEdGraphPin* Pin, int32
 	MoveDynamicPin(Pin, DirectionToMove);
 }
 
+void UNiagaraNodeWithDynamicPins::GetCompilationInputPins(FPinCollectorArray& InputPins) const
+{
+	for (UEdGraphPin* Pin : Pins)
+	{
+		if (Pin->Direction == EGPD_Input && !Pin->bOrphanedPin && !IsAddPin(Pin))
+		{
+			InputPins.Add(Pin);
+		}
+	}
+}
+
+void UNiagaraNodeWithDynamicPins::GetCompilationOutputPins(FPinCollectorArray& OutputPins) const
+{
+	for (UEdGraphPin* Pin : Pins)
+	{
+		if (Pin->Direction == EGPD_Output && !Pin->bOrphanedPin && !IsAddPin(Pin))
+		{
+			OutputPins.Add(Pin);
+		}
+	}
+}
+
 #undef LOCTEXT_NAMESPACE
 
diff --git a/Engine/Plugins/FX/Niagara/Source/NiagaraEditor/Private/NiagaraNodeWriteDataSet.cpp b/Engine/Plugins/FX/Niagara/Source/NiagaraEditor/Private/NiagaraNodeWriteDataSet.cpp
--- a/Engine/Plugins/FX/Niagara/Source/NiagaraEditor/Private/NiagaraNodeWriteDataSet.cpp
+++ b/Engine/Plugins/FX/Niagara/Source/NiagaraEditor/Private/NiagaraNodeWriteDataSet.cpp
@@ -2,7 +2,7 @@
 
 #include "NiagaraNodeWriteDataSet.h"
 #include "UObject/UnrealType.h"
-#include "NiagaraGraphHlslTranslator.h"
+#include "NiagaraHlslTranslator.h"
 #include "NiagaraEvents.h"
 #include "Widgets/SNiagaraGraphNodeWriteDataSet.h"
 #include "EdGraphSchema_Niagara.h"
diff --git a/Engine/Plugins/FX/Niagara/Source/NiagaraEditor/Private/NiagaraParameterMapHistory.cpp b/Engine/Plugins/FX/Niagara/Source/NiagaraEditor/Private/NiagaraParameterMapHistory.cpp
--- a/Engine/Plugins/FX/Niagara/Source/NiagaraEditor/Private/NiagaraParameterMapHistory.cpp
+++ b/Engine/Plugins/FX/Niagara/Source/NiagaraEditor/Private/NiagaraParameterMapHistory.cpp
@@ -2,31 +2,29 @@
 
 #include "NiagaraParameterMapHistory.h"
 
-#include "NiagaraEditorCommon.h"
-#include "NiagaraHlslTranslator.h"
-#include "NiagaraSystem.h"
-#include "NiagaraSystemImpl.h"
-#include "NiagaraGraph.h"
 #include "EdGraphSchema_Niagara.h"
-#include "NiagaraNode.h"
-#include "NiagaraGraph.h"
+#include "INiagaraEditorTypeUtilities.h"
+#include "Modules/ModuleManager.h"
+#include "NiagaraCompilationBridge.h"
 #include "NiagaraCompiler.h"
+#include "NiagaraConstants.h"
+#include "NiagaraEditorCommon.h"
+#include "NiagaraEditorModule.h"
 #include "NiagaraEditorUtilities.h"
-#include "NiagaraNodeOutput.h"
+#include "NiagaraGraph.h"
+#include "NiagaraHlslTranslator.h"
+#include "NiagaraNode.h"
 #include "NiagaraNodeEmitter.h"
+#include "NiagaraNodeOutput.h"
 #include "NiagaraNodeParameterMapGet.h"
-#include "ViewModels/Stack/NiagaraStackGraphUtilities.h"
-#include "NiagaraParameterCollection.h"
-#include "NiagaraConstants.h"
+#include "NiagaraNodeParameterMapSet.h"
 #include "NiagaraNodeStaticSwitch.h"
-#include "NiagaraEditorModule.h"
-#include "INiagaraEditorTypeUtilities.h"
-#include "NiagaraEditorModule.h"
+#include "NiagaraParameterCollection.h"
 #include "NiagaraScriptVariable.h"
-#include "NiagaraNodeParameterMapSet.h"
+#include "NiagaraSystem.h"
+#include "NiagaraSystemImpl.h"
 #include "UObject/UObjectThreadContext.h"
-#include "Modules/ModuleManager.h"
-#include "NiagaraEditorModule.h"
+#include "ViewModels/Stack/NiagaraStackGraphUtilities.h"
 
 #define LOCTEXT_NAMESPACE "NiagaraEditor"
 
@@ -103,13 +101,13 @@ void FGraphTraversalHandle::PushPin(const UEdGraphPin* Pin)
 		{
 			Path.Push(Pin->PersistentGuid);
 		}
-		else if (Pin->GetOwningNode())
+		else
 		{
 			int32 Index = Pin->GetOwningNode()->GetPinIndex((UEdGraphPin*)Pin);
 			FGuid PinIndexGuid(Index, 0, 0, 0);
 			Path.Push(PinIndexGuid);
 		}
-	
+
 		FriendlyPath.Push(Pin->GetOwningNode()->GetNodeTitle(ENodeTitleType::FullTitle).ToString() + TEXT("->") + Pin->PinName.ToString());
 	}
 	else
@@ -137,17 +135,18 @@ void FGraphTraversalHandle::PopPin()
 	FriendlyPath.Pop();
 }
 
-FNiagaraParameterMapHistory::FNiagaraParameterMapHistory() 
+template<typename GraphBridge>
+TNiagaraParameterMapHistory<GraphBridge>::TNiagaraParameterMapHistory()
 {
 	OriginatingScriptUsage = ENiagaraScriptUsage::Function;
 }
 
-void FNiagaraParameterMapHistory::GetValidNamespacesForReading(const UNiagaraScript* InScript, TArray<FString>& OutputNamespaces)
+void FNiagaraParameterUtilities::GetValidNamespacesForReading(const UNiagaraScript* InScript, TArray<FString>& OutputNamespaces)
 {
 	GetValidNamespacesForReading(InScript->GetUsage(), 0, OutputNamespaces);
 }
 
-void FNiagaraParameterMapHistory::GetValidNamespacesForReading(ENiagaraScriptUsage InScriptUsage, int32 InUsageBitmask, TArray<FString>& OutputNamespaces)
+void FNiagaraParameterUtilities::GetValidNamespacesForReading(ENiagaraScriptUsage InScriptUsage, int32 InUsageBitmask, TArray<FString>& OutputNamespaces)
 {
 	TArray<ENiagaraScriptUsage> SupportedContexts;
 	SupportedContexts.Add(InScriptUsage);
@@ -174,7 +173,7 @@ void FNiagaraParameterMapHistory::GetValidNamespacesForReading(ENiagaraScriptUsa
 	}
 }
 
-FString FNiagaraParameterMapHistory::GetNamespace(const FNiagaraVariable& InVar, bool bIncludeDelimiter)
+FString FNiagaraParameterUtilities::GetNamespace(const FNiagaraVariable& InVar, bool bIncludeDelimiter)
 {
 	TArray<FString> SplitName;
 	InVar.GetName().ToString().ParseIntoArray(SplitName, TEXT("."));
@@ -191,7 +190,7 @@ FString FNiagaraParameterMapHistory::GetNamespace(const FNiagaraVariable& InVar,
 	}
 }
 
-bool FNiagaraParameterMapHistory::IsValidNamespaceForReading(ENiagaraScriptUsage InScriptUsage, int32 InUsageBitmask, FString Namespace)
+bool FNiagaraParameterUtilities::IsValidNamespaceForReading(ENiagaraScriptUsage InScriptUsage, int32 InUsageBitmask, FString Namespace)
 {
 	TArray<FString> OutputNamespaces;
 	GetValidNamespacesForReading(InScriptUsage, InUsageBitmask, OutputNamespaces);
@@ -234,16 +233,17 @@ bool FNiagaraParameterMapHistory::IsValidNamespaceForReading(ENiagaraScriptUsage
 	return true;
 }
 
-
-int32 FNiagaraParameterMapHistory::RegisterParameterMapPin(const UEdGraphPin* Pin)
+template<typename GraphBridge>
+int32 TNiagaraParameterMapHistory<GraphBridge>::RegisterParameterMapPin(const FPin* Pin)
 {
 	int32 RetIdx =  MapPinHistory.Add(Pin);
 	return RetIdx;
 }
 
-void FNiagaraParameterMapHistory::RegisterConstantPin(const FGraphTraversalHandle& InTraversalPath, const UEdGraphPin* InPin, const FString& InValue)
+template<typename GraphBridge>
+void TNiagaraParameterMapHistory<GraphBridge>::RegisterConstantPin(const FGraphTraversalHandle& InTraversalPath, const FPin* InPin, const FString& InValue)
 {
-	if (UEdGraphSchema_Niagara::IsStaticPin(InPin))
+	if (GraphBridge::IsStaticPin(InPin))
 	{
 		FGraphTraversalHandle Handle = InTraversalPath;
 		Handle.PushPin(InPin);
@@ -263,7 +263,8 @@ void FNiagaraParameterMapHistory::RegisterConstantPin(const FGraphTraversalHandl
 }
 
 
-void FNiagaraParameterMapHistory::RegisterConstantVariableWrite(const FString& InValue, int32 VarIdx, bool bIsSettingDefault, bool bLinkNotValue)
+template<typename GraphBridge>
+void TNiagaraParameterMapHistory<GraphBridge>::RegisterConstantVariableWrite(const FString& InValue, int32 VarIdx, bool bIsSettingDefault, bool bLinkNotValue)
 {
 	
 	FNiagaraEditorModule& NiagaraEditorModule = FModuleManager::GetModuleChecked<FNiagaraEditorModule>("NiagaraEditor");
@@ -290,7 +291,8 @@ void FNiagaraParameterMapHistory::RegisterConstantVariableWrite(const FString& I
 	
 }
 
-uint32 FNiagaraParameterMapHistory::BeginNodeVisitation(const UNiagaraNode* Node)
+template<typename GraphBridge>
+uint32 TNiagaraParameterMapHistory<GraphBridge>::BeginNodeVisitation(const FNode* Node)
 {
 	uint32 AddedIndex = MapNodeVisitations.Add(Node);
 	MapNodeVariableMetaData.Add(TTuple<uint32, uint32>(Variables.Num(), 0));
@@ -298,7 +300,8 @@ uint32 FNiagaraParameterMapHistory::BeginNodeVisitation(const UNiagaraNode* Node
 	return AddedIndex;
 }
 
-void FNiagaraParameterMapHistory::EndNodeVisitation(uint32 IndexFromBeginNode)
+template<typename GraphBridge>
+void TNiagaraParameterMapHistory<GraphBridge>::EndNodeVisitation(uint32 IndexFromBeginNode)
 {
 	check(IndexFromBeginNode < (uint32)MapNodeVisitations.Num());
 	check(MapNodeVisitations.Num() == MapNodeVariableMetaData.Num());
@@ -306,7 +309,8 @@ void FNiagaraParameterMapHistory::EndNodeVisitation(uint32 IndexFromBeginNode)
 }
 
 
-int32 FNiagaraParameterMapHistory::FindVariableByName(const FName& VariableName, bool bAllowPartialMatch) const
+template<typename GraphBridge>
+int32 TNiagaraParameterMapHistory<GraphBridge>::FindVariableByName(const FName& VariableName, bool bAllowPartialMatch) const
 {
 	if (!bAllowPartialMatch)
 	{
@@ -324,7 +328,8 @@ int32 FNiagaraParameterMapHistory::FindVariableByName(const FName& VariableName,
 }
 
 
-int32 FNiagaraParameterMapHistory::FindVariable(const FName& VariableName, const FNiagaraTypeDefinition& Type) const
+template<typename GraphBridge>
+int32 TNiagaraParameterMapHistory<GraphBridge>::FindVariable(const FName& VariableName, const FNiagaraTypeDefinition& Type) const
 {
 	int32 FoundIdx = Variables.IndexOfByPredicate([&](const FNiagaraVariable& InObj) -> bool
 	{
@@ -334,11 +339,12 @@ int32 FNiagaraParameterMapHistory::FindVariable(const FName& VariableName, const
 	return FoundIdx;
 }
 
-int32 FNiagaraParameterMapHistory::AddVariable(
+template<typename GraphBridge>
+int32 TNiagaraParameterMapHistory<GraphBridge>::AddVariable(
 	  const FNiagaraVariable& InVar
 	, const FNiagaraVariable& InAliasedVar
 	, FName ModuleName
-	, const UEdGraphPin* InPin
+	, const FPin* InPin
 	, TOptional<FNiagaraVariableMetaData> InMetaData /*= TOptional<FNiagaraVariableMetaData>()*/)
 {
 	FNiagaraVariable Var = InVar;
@@ -379,12 +385,14 @@ int32 FNiagaraParameterMapHistory::AddVariable(
 	return FoundIdx;
 }
 
-int32 FNiagaraParameterMapHistory::AddExternalVariable(const FNiagaraVariable& Var)
+template<typename GraphBridge>
+int32 TNiagaraParameterMapHistory<GraphBridge>::AddExternalVariable(const FNiagaraVariable& Var)
 {
 	return AddVariable(Var, Var, NAME_None, nullptr);
 }
 
-const UEdGraphPin* FNiagaraParameterMapHistory::GetFinalPin() const
+template<typename GraphBridge>
+const typename GraphBridge::FPin* TNiagaraParameterMapHistory<GraphBridge>::GetFinalPin() const
 {
 	if (MapPinHistory.Num() > 0)
 	{
@@ -393,7 +401,8 @@ const UEdGraphPin* FNiagaraParameterMapHistory::GetFinalPin() const
 	return nullptr;
 }
 
-const UEdGraphPin* FNiagaraParameterMapHistory::GetOriginalPin() const
+template<typename GraphBridge>
+const typename GraphBridge::FPin* TNiagaraParameterMapHistory<GraphBridge>::GetOriginalPin() const
 {
 	if (MapPinHistory.Num() > 0)
 	{
@@ -403,7 +412,7 @@ const UEdGraphPin* FNiagaraParameterMapHistory::GetOriginalPin() const
 }
 
 
-FName FNiagaraParameterMapHistory::ResolveEmitterAlias(const FName& InName, const FString& InAlias)
+FName FNiagaraParameterUtilities::ResolveEmitterAlias(const FName& InName, const FString& InAlias)
 {
 	// If the alias is empty than the name can't be resolved.
 	if (InAlias.IsEmpty())
@@ -420,14 +429,14 @@ FName FNiagaraParameterMapHistory::ResolveEmitterAlias(const FName& InName, cons
 
 
 
-FString FNiagaraParameterMapHistory::MakeSafeNamespaceString(const FString& InStr)
+FString FNiagaraParameterUtilities::MakeSafeNamespaceString(const FString& InStr)
 {
 	FString  Sanitized = FNiagaraHlslTranslator::GetSanitizedSymbolName(InStr);
 	return Sanitized;
 }
 
 
-FNiagaraVariable FNiagaraParameterMapHistory::ResolveAsBasicAttribute(const FNiagaraVariable& InVar, bool bSanitizeString)
+FNiagaraVariable FNiagaraParameterUtilities::ResolveAsBasicAttribute(const FNiagaraVariable& InVar, bool bSanitizeString)
 {
 	if (IsAttribute(InVar))
 	{
@@ -448,7 +457,7 @@ FNiagaraVariable FNiagaraParameterMapHistory::ResolveAsBasicAttribute(const FNia
 	}
 }
 
-FNiagaraVariable FNiagaraParameterMapHistory::BasicAttributeToNamespacedAttribute(const FNiagaraVariable& InVar, bool bSanitizeString)
+FNiagaraVariable FNiagaraParameterUtilities::BasicAttributeToNamespacedAttribute(const FNiagaraVariable& InVar, bool bSanitizeString)
 {
 	FString ParamName = InVar.GetName().ToString();
 	ParamName.InsertAt(0, FString(PARAM_MAP_ATTRIBUTE_STR));
@@ -463,7 +472,7 @@ FNiagaraVariable FNiagaraParameterMapHistory::BasicAttributeToNamespacedAttribut
 	return RetVar;
 }
 
-FNiagaraVariable FNiagaraParameterMapHistory::VariableToNamespacedVariable(const FNiagaraVariable& InVar, FString Namespace)
+FNiagaraVariable FNiagaraParameterUtilities::VariableToNamespacedVariable(const FNiagaraVariable& InVar, FString Namespace)
 {
 	FString ParamName = Namespace;
 	if (Namespace.EndsWith(TEXT(".")))
@@ -481,7 +490,7 @@ FNiagaraVariable FNiagaraParameterMapHistory::VariableToNamespacedVariable(const
 	return RetVar;
 }
 
-bool FNiagaraParameterMapHistory::IsInNamespace(const FNiagaraVariableBase& InVar, const FString& Namespace)
+bool FNiagaraParameterUtilities::IsInNamespace(const FNiagaraVariableBase& InVar, const FString& Namespace)
 {
 	const bool FastTest = InVar.IsInNameSpace(Namespace);
 
@@ -502,50 +511,50 @@ bool FNiagaraParameterMapHistory::IsInNamespace(const FNiagaraVariableBase& InVa
 	return FastTest;
 }
 
-bool FNiagaraParameterMapHistory::IsAliasedModuleParameter(const FNiagaraVariable& InVar)
+bool FNiagaraParameterUtilities::IsAliasedModuleParameter(const FNiagaraVariable& InVar)
 {
 	return IsInNamespace(InVar, PARAM_MAP_MODULE_STR);
 }
 
-bool FNiagaraParameterMapHistory::IsAliasedEmitterParameter(const FNiagaraVariable& InVar)
+bool FNiagaraParameterUtilities::IsAliasedEmitterParameter(const FNiagaraVariable& InVar)
 {
 	return IsInNamespace(InVar, PARAM_MAP_EMITTER_STR);
 }
 
-bool FNiagaraParameterMapHistory::IsAliasedEmitterParameter(const FString& InVarName)
+bool FNiagaraParameterUtilities::IsAliasedEmitterParameter(const FString& InVarName)
 {
 	return IsAliasedEmitterParameter(FNiagaraVariable(FNiagaraTypeDefinition::GetFloatDef(), *InVarName));
 }
 
 
-bool FNiagaraParameterMapHistory::IsSystemParameter(const FNiagaraVariable& InVar)
+bool FNiagaraParameterUtilities::IsSystemParameter(const FNiagaraVariable& InVar)
 {
 	return IsInNamespace(InVar, PARAM_MAP_SYSTEM_STR);
 }
 
-bool FNiagaraParameterMapHistory::IsEngineParameter(const FNiagaraVariable& InVar)
+bool FNiagaraParameterUtilities::IsEngineParameter(const FNiagaraVariable& InVar)
 {
 	return IsInNamespace(InVar, PARAM_MAP_ENGINE_STR);
 }
 
-bool FNiagaraParameterMapHistory::IsPerInstanceEngineParameter(const FNiagaraVariable& InVar, const FString& EmitterAlias)
+bool FNiagaraParameterUtilities::IsPerInstanceEngineParameter(const FNiagaraVariable& InVar, const FString& EmitterAlias)
 {
 	FString EmitterEngineNamespaceAlias = TEXT("Engine.") + EmitterAlias + TEXT(".");
 	return IsInNamespace(InVar, PARAM_MAP_ENGINE_OWNER_STR) || IsInNamespace(InVar, PARAM_MAP_ENGINE_SYSTEM_STR) || IsInNamespace(InVar, PARAM_MAP_ENGINE_EMITTER_STR) || 
 		IsInNamespace(InVar, EmitterEngineNamespaceAlias);
 }
 
-bool FNiagaraParameterMapHistory::IsUserParameter(const FNiagaraVariable& InVar)
+bool FNiagaraParameterUtilities::IsUserParameter(const FNiagaraVariable& InVar)
 {
 	return IsInNamespace(InVar, PARAM_MAP_USER_STR);
 }
 
-bool FNiagaraParameterMapHistory::IsRapidIterationParameter(const FNiagaraVariable& InVar)
+bool FNiagaraParameterUtilities::IsRapidIterationParameter(const FNiagaraVariable& InVar)
 {
 	return IsInNamespace(InVar, PARAM_MAP_RAPID_ITERATION_STR);
 }
 
-bool FNiagaraParameterMapHistory::SplitRapidIterationParameterName(const FNiagaraVariable& InVar, ENiagaraScriptUsage InUsage, FString& EmitterName, FString& FunctionCallName, FString& InputName)
+bool FNiagaraParameterUtilities::SplitRapidIterationParameterName(const FNiagaraVariable& InVar, ENiagaraScriptUsage InUsage, FString& EmitterName, FString& FunctionCallName, FString& InputName)
 {
 	TArray<FString> SplitName;
 	InVar.GetName().ToString().ParseIntoArray(SplitName, TEXT("."));
@@ -573,19 +582,20 @@ bool FNiagaraParameterMapHistory::SplitRapidIterationParameterName(const FNiagar
 	return false;
 }
 
-bool FNiagaraParameterMapHistory::IsAttribute(const FNiagaraVariableBase& InVar)
+bool FNiagaraParameterUtilities::IsAttribute(const FNiagaraVariableBase& InVar)
 {
 	return IsInNamespace(InVar, PARAM_MAP_ATTRIBUTE_STR);
 }
 
-const UEdGraphPin* FNiagaraParameterMapHistory::GetDefaultValuePin(int32 VarIdx) const
+template<typename GraphBridge>
+const typename GraphBridge::FInputPin* TNiagaraParameterMapHistory<GraphBridge>::GetDefaultValuePin(int32 VarIdx) const
 {
 	if (PerVariableWriteHistory[VarIdx].Num() > 0)
 	{
-		const FModuleScopedPin& ScopedPin = PerVariableWriteHistory[VarIdx][0];
-		if (ScopedPin.Pin != nullptr && ScopedPin.Pin->Direction == EEdGraphPinDirection::EGPD_Input && Cast<UNiagaraNodeParameterMapGet>(ScopedPin.Pin->GetOwningNode()) != nullptr)
+		const TModuleScopedPin<FPin>& ScopedPin = PerVariableWriteHistory[VarIdx][0];
+		if (ScopedPin.Pin != nullptr && ScopedPin.Pin->Direction == EEdGraphPinDirection::EGPD_Input && GraphBridge::IsParameterMapGet(GraphBridge::GetOwningNode(ScopedPin.Pin)))
 		{
-			return ScopedPin.Pin;
+			return GraphBridge::GetPinAsInput(ScopedPin.Pin);
 		}
 	}
 	return nullptr;
@@ -629,56 +639,58 @@ static FName ReplaceVariableNamespace(const FName& InVariableName, const TCHAR*
 	return *FString::Join(JoinString, TEXT("."));
 }
 
-bool FNiagaraParameterMapHistory::IsInitialName(const FName& InVariableName)
+bool FNiagaraParameterUtilities::IsInitialName(const FName& InVariableName)
 {
 	return DoesVariableIncludeNamespace(InVariableName, PARAM_MAP_INITIAL_BASE_STR);
 }
 
-bool FNiagaraParameterMapHistory::IsInitialValue(const FNiagaraVariableBase& InVar)
+bool FNiagaraParameterUtilities::IsInitialValue(const FNiagaraVariableBase& InVar)
 {
 	return IsInitialName(InVar.GetName());
 }
 
-bool FNiagaraParameterMapHistory::IsPreviousName(const FName& InVariableName)
+bool FNiagaraParameterUtilities::IsPreviousName(const FName& InVariableName)
 {
 	return DoesVariableIncludeNamespace(InVariableName, PARAM_MAP_PREVIOUS_BASE_STR);
 }
 
-bool FNiagaraParameterMapHistory::IsPreviousValue(const FNiagaraVariableBase& InVar)
+bool FNiagaraParameterUtilities::IsPreviousValue(const FNiagaraVariableBase& InVar)
 {
 	return IsPreviousName(InVar.GetName());
 }
 
-FNiagaraVariable FNiagaraParameterMapHistory::GetSourceForInitialValue(const FNiagaraVariable& InVar)
+FNiagaraVariable FNiagaraParameterUtilities::GetSourceForInitialValue(const FNiagaraVariable& InVar)
 {
 	FNiagaraVariable Var = InVar;
 	Var.SetName(GetSourceForInitialValue(InVar.GetName()));
 	return Var;
 }
 
-FName FNiagaraParameterMapHistory::GetSourceForInitialValue(const FName& InVariableName)
+FName FNiagaraParameterUtilities::GetSourceForInitialValue(const FName& InVariableName)
 {
 	return ReplaceVariableNamespace(InVariableName, PARAM_MAP_INITIAL_BASE_STR);
 }
 
-FNiagaraVariable FNiagaraParameterMapHistory::GetSourceForPreviousValue(const FNiagaraVariable& InVar)
+FNiagaraVariable FNiagaraParameterUtilities::GetSourceForPreviousValue(const FNiagaraVariable& InVar)
 {
 	FNiagaraVariable Var = InVar;
 	Var.SetName(GetSourceForPreviousValue(InVar.GetName()));
 	return Var;
 }
 
-FName FNiagaraParameterMapHistory::GetSourceForPreviousValue(const FName& InVariableName)
+FName FNiagaraParameterUtilities::GetSourceForPreviousValue(const FName& InVariableName)
 {
 	return ReplaceVariableNamespace(InVariableName, PARAM_MAP_PREVIOUS_BASE_STR);
 }
 
-bool FNiagaraParameterMapHistory::IsPrimaryDataSetOutput(const FNiagaraVariable& InVar, const UNiagaraScript* InScript,  bool bAllowDataInterfaces, bool bAllowStatics) const
+template<typename GraphBridge>
+bool TNiagaraParameterMapHistory<GraphBridge>::IsPrimaryDataSetOutput(const FNiagaraVariable& InVar, const UNiagaraScript* InScript,  bool bAllowDataInterfaces, bool bAllowStatics) const
 {
 	return IsPrimaryDataSetOutput(InVar, InScript->GetUsage(),  bAllowDataInterfaces);
 }
 
-bool FNiagaraParameterMapHistory::IsPrimaryDataSetOutput(const FNiagaraVariable& InVar, ENiagaraScriptUsage Usage, bool bAllowDataInterfaces, bool bAllowStatics) const
+template<typename GraphBridge>
+bool TNiagaraParameterMapHistory<GraphBridge>::IsPrimaryDataSetOutput(const FNiagaraVariable& InVar, ENiagaraScriptUsage Usage, bool bAllowDataInterfaces, bool bAllowStatics) const
 {
 	if (bAllowDataInterfaces == false && InVar.GetType().GetClass() != nullptr)
 	{
@@ -697,21 +709,21 @@ bool FNiagaraParameterMapHistory::IsPrimaryDataSetOutput(const FNiagaraVariable&
 		// the child emitters that were encountered.
 		for (FString EmitterEncounteredNamespace : EmitterNamespacesEncountered)
 		{
-			if (IsInNamespace(InVar, EmitterEncounteredNamespace))
+			if (FNiagaraParameterUtilities::IsInNamespace(InVar, EmitterEncounteredNamespace))
 			{
 				return true;
 			}
 		}
-		return IsInNamespace(InVar, PARAM_MAP_SYSTEM_STR) || IsInNamespace(InVar, PARAM_MAP_EMITTER_STR);
+		return FNiagaraParameterUtilities::IsInNamespace(InVar, PARAM_MAP_SYSTEM_STR) || FNiagaraParameterUtilities::IsInNamespace(InVar, PARAM_MAP_EMITTER_STR);
 	}
 	else if (Usage == ENiagaraScriptUsage::Module || Usage == ENiagaraScriptUsage::Function)
 	{
-		return IsInNamespace(InVar, PARAM_MAP_MODULE_STR);
+		return FNiagaraParameterUtilities::IsInNamespace(InVar, PARAM_MAP_MODULE_STR);
 	}
-	return IsInNamespace(InVar, PARAM_MAP_ATTRIBUTE_STR);
+	return FNiagaraParameterUtilities::IsInNamespace(InVar, PARAM_MAP_ATTRIBUTE_STR);
 }
 
-bool FNiagaraParameterMapHistory::IsWrittenToScriptUsage(const FNiagaraVariable& InVar, ENiagaraScriptUsage Usage, bool bAllowDataInterfaces)
+bool FNiagaraParameterUtilities::IsWrittenToScriptUsage(const FNiagaraVariable& InVar, ENiagaraScriptUsage Usage, bool bAllowDataInterfaces)
 {
 	if (bAllowDataInterfaces == false && InVar.GetType().GetClass() != nullptr)
 	{
@@ -730,7 +742,7 @@ bool FNiagaraParameterMapHistory::IsWrittenToScriptUsage(const FNiagaraVariable&
 	return IsInNamespace(InVar, PARAM_MAP_ATTRIBUTE_STR);
 }
 
-FNiagaraVariable FNiagaraParameterMapHistory::MoveToExternalConstantNamespaceVariable(const FNiagaraVariable& InVar, ENiagaraScriptUsage InUsage)
+FNiagaraVariable FNiagaraParameterUtilities::MoveToExternalConstantNamespaceVariable(const FNiagaraVariable& InVar, ENiagaraScriptUsage InUsage)
 {
 	if (UNiagaraScript::IsParticleScript(InUsage))
 	{
@@ -747,12 +759,12 @@ FNiagaraVariable FNiagaraParameterMapHistory::MoveToExternalConstantNamespaceVar
 	return InVar;
 }
 
-FNiagaraVariable FNiagaraParameterMapHistory::MoveToExternalConstantNamespaceVariable(const FNiagaraVariable& InVar, const UNiagaraScript* InScript)
+FNiagaraVariable FNiagaraParameterUtilities::MoveToExternalConstantNamespaceVariable(const FNiagaraVariable& InVar, const UNiagaraScript* InScript)
 {
 	return MoveToExternalConstantNamespaceVariable(InVar, InScript->GetUsage());
 }
 
-bool FNiagaraParameterMapHistory::IsExportableExternalConstant(const FNiagaraVariable& InVar, const UNiagaraScript* InScript)
+bool FNiagaraParameterUtilities::IsExportableExternalConstant(const FNiagaraVariable& InVar, const UNiagaraScript* InScript)
 {
 	if (InScript->IsEquivalentUsage(ENiagaraScriptUsage::SystemSpawnScript))
 	{
@@ -764,7 +776,7 @@ bool FNiagaraParameterMapHistory::IsExportableExternalConstant(const FNiagaraVar
 	}
 }
 
-bool FNiagaraParameterMapHistory::IsExternalConstantNamespace(const FNiagaraVariable& InVar, ENiagaraScriptUsage InUsage, int32 InUsageBitmask)
+bool FNiagaraParameterUtilities::IsExternalConstantNamespace(const FNiagaraVariable& InVar, ENiagaraScriptUsage InUsage, int32 InUsageBitmask)
 {
 	// Parameter collections are always constants
 	if (IsInNamespace(InVar, PARAM_MAP_NPC_STR))
@@ -813,32 +825,29 @@ bool FNiagaraParameterMapHistory::IsExternalConstantNamespace(const FNiagaraVari
 	return false;
 }
 
-bool FNiagaraParameterMapHistory::IsExternalConstantNamespace(const FNiagaraVariable& InVar, const UNiagaraScript* InScript, const FGuid& VersionGuid)
+bool FNiagaraParameterUtilities::IsExternalConstantNamespace(const FNiagaraVariable& InVar, const UNiagaraScript* InScript, const FGuid& VersionGuid)
 {
 	return IsExternalConstantNamespace(InVar, InScript->GetUsage(), InScript->GetScriptData(VersionGuid)->ModuleUsageBitmask);
 }
 
-const UNiagaraNodeOutput* FNiagaraParameterMapHistory::GetFinalOutputNode() const
+template<typename GraphBridge>
+const typename GraphBridge::FOutputNode* TNiagaraParameterMapHistory<GraphBridge>::GetFinalOutputNode() const
 {
-	const UEdGraphPin* Pin = GetFinalPin();
-	if (Pin != nullptr)
+	if (const FPin* Pin = GetFinalPin())
 	{
-		const UNiagaraNodeOutput* Output = Cast<const UNiagaraNodeOutput>(Pin->GetOwningNode());
-		if (Output != nullptr)
-		{
-			return Output;
-		}
+		return GraphBridge::AsOutputNode(GraphBridge::GetOwningNode(Pin));
 	}
 
 	return nullptr;
 }
 
-FNiagaraVariable FNiagaraParameterMapHistory::ConvertVariableToRapidIterationConstantName(FNiagaraVariable InVar, const TCHAR* InEmitterName, ENiagaraScriptUsage InUsage)
+FNiagaraVariable FNiagaraParameterUtilities::ConvertVariableToRapidIterationConstantName(FNiagaraVariable InVar, const TCHAR* InEmitterName, ENiagaraScriptUsage InUsage)
 {
 	return FNiagaraUtilities::ConvertVariableToRapidIterationConstantName(InVar, InEmitterName, InUsage);
 }
 
-UNiagaraParameterCollection* FNiagaraParameterMapHistory::IsParameterCollectionParameter(FNiagaraVariable& InVar, bool& bMissingParameter)
+template<>
+FNiagaraCompilationGraphBridge::FParameterCollection TNiagaraParameterMapHistory<FNiagaraCompilationGraphBridge>::IsParameterCollectionParameter(FNiagaraVariable& InVar, bool& bMissingParameter)
 {
 	if (bParameterCollectionsSkipped)
 	{
@@ -848,18 +857,19 @@ UNiagaraParameterCollection* FNiagaraParameterMapHistory::IsParameterCollectionP
 
 	bMissingParameter = false;
 	FString VarName = InVar.GetName().ToString();
-	for (int32 i = 0; i < ParameterCollections.Num(); ++i)
+	for (int32 i = 0; i < EncounteredParameterCollections.Collections.Num(); ++i)
 	{
-		if (VarName.StartsWith(ParameterCollectionNamespaces[i]))
+		if (VarName.StartsWith(EncounteredParameterCollections.CollectionNamespaces[i]))
 		{
-			bMissingParameter = !ParameterCollectionVariables[i].Contains(InVar);
-			return ParameterCollections[i];
+			bMissingParameter = !EncounteredParameterCollections.CollectionVariables[i].Contains(InVar);
+			return EncounteredParameterCollections.Collections[i];
 		}
 	}
 	return nullptr;
 }
 
-bool FNiagaraParameterMapHistory::ShouldIgnoreVariableDefault(const FNiagaraVariable& Var)const
+template<typename GraphBridge>
+bool TNiagaraParameterMapHistory<GraphBridge>::ShouldIgnoreVariableDefault(const FNiagaraVariable& Var)const
 {
 	// NOTE(mv): Used for variables that are explicitly assigned to (on spawn) and should not be default initialized
 	//           These are explicitly written to in NiagaraHlslTranslator::DefineMain
@@ -869,7 +879,8 @@ bool FNiagaraParameterMapHistory::ShouldIgnoreVariableDefault(const FNiagaraVari
 	return bShouldBeIgnored;
 }
 
-bool FNiagaraParameterMapHistory::IsVariableFromCustomIterationNamespaceOverride(const FNiagaraVariable& InVar) const
+template<typename GraphBridge>
+bool TNiagaraParameterMapHistory<GraphBridge>::IsVariableFromCustomIterationNamespaceOverride(const FNiagaraVariable& InVar) const
 {
 	for (const FName& Name : IterationNamespaceOverridesEncountered)
 	{
@@ -882,8 +893,8 @@ bool FNiagaraParameterMapHistory::IsVariableFromCustomIterationNamespaceOverride
 	return false;
 }
 
-
-FNiagaraParameterMapHistoryBuilder::FNiagaraParameterMapHistoryBuilder()
+template<typename GraphBridge>
+TNiagaraParameterMapHistoryBuilder<GraphBridge>::TNiagaraParameterMapHistoryBuilder()
 {
 	ContextuallyVisitedNodes.AddDefaulted(1);
 	PinToParameterMapIndices.AddDefaulted(1);
@@ -893,6 +904,8 @@ FNiagaraParameterMapHistoryBuilder::FNiagaraParameterMapHistoryBuilder()
 	bIgnoreDisabled = true;
 	FilterScriptType = ENiagaraScriptUsage::Function;
 
+	ConstantResolver = MakePimpl<FConstantResolver, EPimplPtrMode::DeepCopy>();
+
 	// depending on when this builder is created, we could be doing a PostLoad.  In this scenario
 	// we need to avoid triggering any loads for dependent NiagaraParameterCollection.  Mark the builder
 	// to skip those steps if necessary and we'll ensure that the generated histories don't require the
@@ -900,33 +913,40 @@ FNiagaraParameterMapHistoryBuilder::FNiagaraParameterMapHistoryBuilder()
 	bIncludeParameterCollectionInfo = !FUObjectThreadContext::Get().IsRoutingPostLoad;
 }
 
-void FNiagaraParameterMapHistoryBuilder::BuildParameterMaps(const UNiagaraNodeOutput* OutputNode, bool bRecursive)
+template<typename GraphBridge>
+void TNiagaraParameterMapHistoryBuilder<GraphBridge>::BuildParameterMaps(const FOutputNode* OutputNode, bool bRecursive)
 {
+	const ENiagaraScriptUsage NodeUsage = GraphBridge::GetOutputNodeUsage(OutputNode);
+
 	if (UNiagaraScript::LogCompileStaticVars > 0)
 	{
-		UE_LOG(LogNiagaraEditor, Log, TEXT("************** BuildParameterMaps %d  **************"), OutputNode->GetUsage());
+		UE_LOG(LogNiagaraEditor, Log, TEXT("************** BuildParameterMaps %d  **************"), NodeUsage);
 	}
-	TOptional<FName> StackContextAlias = OutputNode->GetStackContextOverride();
-	BeginUsage(OutputNode->GetUsage(), StackContextAlias.Get(ScriptUsageContextNameStack.Num() > 0 ? ScriptUsageContextNameStack.Top() : NAME_None));
+	TOptional<FName> StackContextAlias = GraphBridge::GetOutputNodeStackContextOverride(OutputNode);
+	BeginUsage(NodeUsage, StackContextAlias.Get(ScriptUsageContextNameStack.Num() > 0 ? ScriptUsageContextNameStack.Top() : NAME_None));
 
-	OutputNode->BuildParameterMapHistory(*this, bRecursive);
+	constexpr bool bFilterForCompilation = true;
+	OutputNode->BuildParameterMapHistory(*this, bRecursive, bFilterForCompilation);
 
 	EndUsage();
 }
 
-void FNiagaraParameterMapHistoryBuilder::EnableScriptAllowList(bool bInEnable, ENiagaraScriptUsage InScriptType)
+template<typename GraphBridge>
+void TNiagaraParameterMapHistoryBuilder<GraphBridge>::EnableScriptAllowList(bool bInEnable, ENiagaraScriptUsage InScriptType)
 {
 	bFilterByScriptAllowList = bInEnable;
 	FilterScriptType = InScriptType;
 }
 
-bool FNiagaraParameterMapHistoryBuilder::HasCurrentUsageContext() const
+template<typename GraphBridge>
+bool TNiagaraParameterMapHistoryBuilder<GraphBridge>::HasCurrentUsageContext() const
 {
 	return RelevantScriptUsageContext.Num() > 0;
 }
 
 
-bool FNiagaraParameterMapHistoryBuilder::ContextContains(ENiagaraScriptUsage InUsage) const
+template<typename GraphBridge>
+bool TNiagaraParameterMapHistoryBuilder<GraphBridge>::ContextContains(ENiagaraScriptUsage InUsage) const
 {
 	if (RelevantScriptUsageContext.Num() == 0)
 	{
@@ -938,25 +958,27 @@ bool FNiagaraParameterMapHistoryBuilder::ContextContains(ENiagaraScriptUsage InU
 	}
 }
 
-
-ENiagaraScriptUsage FNiagaraParameterMapHistoryBuilder::GetCurrentUsageContext()const
+template<typename GraphBridge>
+ENiagaraScriptUsage TNiagaraParameterMapHistoryBuilder<GraphBridge>::GetCurrentUsageContext()const
 {
 	if (RelevantScriptUsageContext.Num() == 0)
 		return ENiagaraScriptUsage::Function;
 	return RelevantScriptUsageContext.Last();
 }
 
-ENiagaraScriptUsage FNiagaraParameterMapHistoryBuilder::GetBaseUsageContext()const
+template<typename GraphBridge>
+ENiagaraScriptUsage TNiagaraParameterMapHistoryBuilder<GraphBridge>::GetBaseUsageContext()const
 {
 	if (RelevantScriptUsageContext.Num() == 0)
 		return ENiagaraScriptUsage::Function;
 	return RelevantScriptUsageContext[0];
 }
 
-int32 FNiagaraParameterMapHistoryBuilder::CreateParameterMap()
+template<typename GraphBridge>
+int32 TNiagaraParameterMapHistoryBuilder<GraphBridge>::CreateParameterMap()
 {
 	int32 RetValue = Histories.AddDefaulted(1);
-	FNiagaraParameterMapHistory& History = Histories[RetValue];
+	FHistory& History = Histories[RetValue];
 
 	FName StackName = ScriptUsageContextNameStack.Num() > 0 ? ScriptUsageContextNameStack.Top() : NAME_None;
 
@@ -970,7 +992,8 @@ int32 FNiagaraParameterMapHistoryBuilder::CreateParameterMap()
 	return RetValue;
 }
 
-uint32 FNiagaraParameterMapHistoryBuilder::BeginNodeVisitation(int32 WhichParameterMap, const class UNiagaraNode* Node)
+template<typename GraphBridge>
+uint32 TNiagaraParameterMapHistoryBuilder<GraphBridge>::BeginNodeVisitation(int32 WhichParameterMap, const FNode* Node)
 {
 	if (WhichParameterMap != INDEX_NONE)
 	{
@@ -982,7 +1005,8 @@ uint32 FNiagaraParameterMapHistoryBuilder::BeginNodeVisitation(int32 WhichParame
 	}
 }
 
-void FNiagaraParameterMapHistoryBuilder::EndNodeVisitation(int32 WhichParameterMap, uint32 IndexFromBeginNode)
+template<typename GraphBridge>
+void TNiagaraParameterMapHistoryBuilder<GraphBridge>::EndNodeVisitation(int32 WhichParameterMap, uint32 IndexFromBeginNode)
 {
 	if (WhichParameterMap != INDEX_NONE)
 	{
@@ -990,7 +1014,8 @@ void FNiagaraParameterMapHistoryBuilder::EndNodeVisitation(int32 WhichParameterM
 	}
 }
 
-void FNiagaraParameterMapHistoryBuilder::RegisterDataSetWrite(int32 WhichParameterMap, const FNiagaraDataSetID& DataSet)
+template<typename GraphBridge>
+void TNiagaraParameterMapHistoryBuilder<GraphBridge>::RegisterDataSetWrite(int32 WhichParameterMap, const FNiagaraDataSetID& DataSet)
 {
 	if (Histories.IsValidIndex(WhichParameterMap))
 	{
@@ -998,11 +1023,12 @@ void FNiagaraParameterMapHistoryBuilder::RegisterDataSetWrite(int32 WhichParamet
 	}
 }
 
-int32 FNiagaraParameterMapHistoryBuilder::RegisterParameterMapPin(int32 WhichParameterMap, const UEdGraphPin* Pin)
+template<typename GraphBridge>
+int32 TNiagaraParameterMapHistoryBuilder<GraphBridge>::RegisterParameterMapPin(int32 WhichParameterMap, const FPin* Pin)
 {
 	if (Pin && UNiagaraScript::LogCompileStaticVars > 0)
 	{
-		UE_LOG(LogNiagaraEditor, Log, TEXT("FNiagaraParameterMapHistoryBuilder::RegisterParameterMapPin %d > %s .... %p"), WhichParameterMap, *Pin->GetOwningNode()->GetNodeTitle(ENodeTitleType::FullTitle).ToString(), Pin);
+		UE_LOG(LogNiagaraEditor, Log, TEXT("FNiagaraParameterMapHistoryBuilder::RegisterParameterMapPin %d > %s .... %p"), WhichParameterMap, *GraphBridge::GetNodeTitle(GraphBridge::GetOwningNode(Pin)), Pin);
 	}
 	if (Pin && WhichParameterMap != INDEX_NONE)
 	{
@@ -1027,7 +1053,8 @@ int32 FNiagaraParameterMapHistoryBuilder::RegisterParameterMapPin(int32 WhichPar
 	}
 }
 
-int32 FNiagaraParameterMapHistoryBuilder::RegisterConstantPin(int32 WhichConstant, const UEdGraphPin* Pin)
+template<typename GraphBridge>
+int32 TNiagaraParameterMapHistoryBuilder<GraphBridge>::RegisterConstantPin(int32 WhichConstant, const FPin* Pin)
 {
 	if (WhichConstant != INDEX_NONE)
 	{
@@ -1044,26 +1071,27 @@ int32 FNiagaraParameterMapHistoryBuilder::RegisterConstantPin(int32 WhichConstan
 	}
 }
 
-bool FNiagaraParameterMapHistoryBuilder::IsStaticVariableExportableToOuterScopeBasedOnCurrentContext(const FNiagaraVariable& Var) const
+template<typename GraphBridge>
+bool TNiagaraParameterMapHistoryBuilder<GraphBridge>::IsStaticVariableExportableToOuterScopeBasedOnCurrentContext(const FNiagaraVariable& Var) const
 {
-	if ((IsInEncounteredEmitterNamespace(Var) || FNiagaraParameterMapHistory::IsAliasedEmitterParameter(Var)) && (ContextContains(ENiagaraScriptUsage::EmitterSpawnScript) || ContextContains(ENiagaraScriptUsage::EmitterUpdateScript)))
+	if ((IsInEncounteredEmitterNamespace(Var) || FNiagaraParameterUtilities::IsAliasedEmitterParameter(Var)) && (ContextContains(ENiagaraScriptUsage::EmitterSpawnScript) || ContextContains(ENiagaraScriptUsage::EmitterUpdateScript)))
 	{
 		return true;
 	}
-	else if (FNiagaraParameterMapHistory::IsAttribute(Var) && (UNiagaraScript::IsParticleScript(GetBaseUsageContext())))
+	else if (FNiagaraParameterUtilities::IsAttribute(Var) && (UNiagaraScript::IsParticleScript(GetBaseUsageContext())))
 	{
 		return true;
 	}
-	else if (FNiagaraParameterMapHistory::IsSystemParameter(Var) && (UNiagaraScript::IsSystemScript(GetBaseUsageContext()) || UNiagaraScript::IsEmitterScript(GetBaseUsageContext())))
+	else if (FNiagaraParameterUtilities::IsSystemParameter(Var) && (UNiagaraScript::IsSystemScript(GetBaseUsageContext()) || UNiagaraScript::IsEmitterScript(GetBaseUsageContext())))
 	{
 		return true;
 	}
 	return false;
 }
 
-int32 FNiagaraParameterMapHistoryBuilder::RegisterConstantVariableWrite(int32 WhichParamMapIdx, int32 WhichConstant,  int32 WhichVarIdx, bool bIsSettingDefault, bool bIsLinkNotValue)
+template<typename GraphBridge>
+int32 TNiagaraParameterMapHistoryBuilder<GraphBridge>::RegisterConstantVariableWrite(int32 WhichParamMapIdx, int32 WhichConstant,  int32 WhichVarIdx, bool bIsSettingDefault, bool bIsLinkNotValue)
 {
-
 	FNiagaraEditorModule& NiagaraEditorModule = FModuleManager::GetModuleChecked<FNiagaraEditorModule>("NiagaraEditor");
 	if (WhichConstant != INDEX_NONE)
 	{
@@ -1119,35 +1147,39 @@ int32 FNiagaraParameterMapHistoryBuilder::RegisterConstantVariableWrite(int32 Wh
 }
 
 
-FString FNiagaraParameterMapHistoryBuilder::GetConstant(int32 InIndex)
+template<typename GraphBridge>
+FString TNiagaraParameterMapHistoryBuilder<GraphBridge>::GetConstant(int32 InIndex)
 {
 	if (Constants.IsValidIndex(InIndex))
 		return Constants[InIndex];
 	return FString();
 }
 
-int32 FNiagaraParameterMapHistoryBuilder::AddOrGetConstantFromValue(const FString& Value)
+template<typename GraphBridge>
+int32 TNiagaraParameterMapHistoryBuilder<GraphBridge>::AddOrGetConstantFromValue(const FString& Value)
 {
 	return Constants.AddUnique(Value);
 }
 
-int32 FNiagaraParameterMapHistoryBuilder::RegisterConstantFromInputPin(const UEdGraphPin* InputPin) 
+template<typename GraphBridge>
+int32 TNiagaraParameterMapHistoryBuilder<GraphBridge>::RegisterConstantFromInputPin(const FInputPin* InputPin, const FString& PinDefaultValue)
 {
 	int32 ConstantIdx = INDEX_NONE;
 	if (InputPin && InputPin->Direction == EEdGraphPinDirection::EGPD_Input)
 	{
-		FString Value;
-		if (InputPin && InputPin->LinkedTo.Num() > 0 && InputPin->LinkedTo[0] != nullptr)
+		if (const FOutputPin* ConnectedOutputPin = GraphBridge::GetLinkedOutputPin(InputPin))
 		{
-			UEdGraphPin* ConnectedOutputPin = InputPin->LinkedTo[0];
 			ConstantIdx = GetConstantFromOutputPin(ConnectedOutputPin);
 			RegisterConstantPin(ConstantIdx, InputPin);
 		}
-		else if (InputPin && InputPin->Direction == EEdGraphPinDirection::EGPD_Input && InputPin->LinkedTo.Num() == 0 && InputPin->DefaultValue.Len() != 0 && !InputPin->PinName.IsNone())
+		else
 		{
-			Value = InputPin->DefaultValue;
-			ConstantIdx = AddOrGetConstantFromValue(Value);
-			RegisterConstantPin(ConstantIdx, InputPin);
+			ensure(InputPin->Direction == EEdGraphPinDirection::EGPD_Input);
+			if (InputPin->PinName.IsNone())
+			{
+				ConstantIdx = AddOrGetConstantFromValue(PinDefaultValue);
+				RegisterConstantPin(ConstantIdx, InputPin);
+			}
 		}
 	}
 
@@ -1161,7 +1193,8 @@ int32 FNiagaraParameterMapHistoryBuilder::RegisterConstantFromInputPin(const UEd
 	return ConstantIdx;
 }
 
-int32 FNiagaraParameterMapHistoryBuilder::GetConstantFromInputPin(const UEdGraphPin* InputPin) const
+template<typename GraphBridge>
+int32 TNiagaraParameterMapHistoryBuilder<GraphBridge>::GetConstantFromInputPin(const FInputPin* InputPin) const
 {
 	if (InputPin && InputPin->Direction == EEdGraphPinDirection::EGPD_Input)
 	{
@@ -1183,7 +1216,8 @@ int32 FNiagaraParameterMapHistoryBuilder::GetConstantFromInputPin(const UEdGraph
 	return INDEX_NONE;
 }
 
-int32 FNiagaraParameterMapHistoryBuilder::GetConstantFromOutputPin(const UEdGraphPin* OutputPin) const
+template<typename GraphBridge>
+int32 TNiagaraParameterMapHistoryBuilder<GraphBridge>::GetConstantFromOutputPin(const FOutputPin* OutputPin) const
 {
 	if (OutputPin && OutputPin->Direction == EEdGraphPinDirection::EGPD_Output)
 	{
@@ -1205,8 +1239,8 @@ int32 FNiagaraParameterMapHistoryBuilder::GetConstantFromOutputPin(const UEdGrap
 	return INDEX_NONE;
 }
 
-
-int32 FNiagaraParameterMapHistoryBuilder::GetConstantFromVariableRead(const FNiagaraVariableBase& InVar)
+template<typename GraphBridge>
+int32 TNiagaraParameterMapHistoryBuilder<GraphBridge>::GetConstantFromVariableRead(const FNiagaraVariableBase& InVar)
 {
 	if (InVar.IsValid())
 	{
@@ -1227,7 +1261,8 @@ int32 FNiagaraParameterMapHistoryBuilder::GetConstantFromVariableRead(const FNia
 	return INDEX_NONE;
 }
 
-int32 FNiagaraParameterMapHistoryBuilder::TraceParameterMapOutputPin(const UEdGraphPin* OutputPin)
+template<typename GraphBridge>
+int32 TNiagaraParameterMapHistoryBuilder<GraphBridge>::TraceParameterMapOutputPin(const FOutputPin* OutputPin)
 {
 	if (OutputPin && OutputPin->Direction == EEdGraphPinDirection::EGPD_Output)
 	{
@@ -1243,118 +1278,94 @@ int32 FNiagaraParameterMapHistoryBuilder::TraceParameterMapOutputPin(const UEdGr
 	return INDEX_NONE;
 }
 
-bool FNiagaraParameterMapHistoryBuilder::GetPinPreviouslyVisited(const UEdGraphPin* InPin) const
+template<typename GraphBridge>
+bool TNiagaraParameterMapHistoryBuilder<GraphBridge>::GetPinPreviouslyVisited(const FPin* InPin) const
 {
 	if (InPin != nullptr)
-		return GetNodePreviouslyVisited(CastChecked<UNiagaraNode>(InPin->GetOwningNode()));
-	else
-		return true;
+	{
+		return GetNodePreviouslyVisited(GraphBridge::GetOwningNode(InPin));
+	}
+
+	return true;
 }
 
 
-bool FNiagaraParameterMapHistoryBuilder::GetNodePreviouslyVisited(const class UNiagaraNode* Node) const
+template<typename GraphBridge>
+bool TNiagaraParameterMapHistoryBuilder<GraphBridge>::GetNodePreviouslyVisited(const FNode* Node) const
 {
 	return ContextuallyVisitedNodes.Last().Contains(Node);
 }
 
-int32 FNiagaraParameterMapHistoryBuilder::FindMatchingParameterMapFromContextInputs(const FNiagaraVariable& InVar) const
+template<typename GraphBridge>
+int32 TNiagaraParameterMapHistoryBuilder<GraphBridge>::FindMatchingParameterMapFromContextInputs(const FNiagaraVariable& InVar) const
 {
 	if (CallingContext.Num() == 0)
 	{
 		return INDEX_NONE;
 	}
-	const UNiagaraNode* Node = CallingContext.Last();
-	FPinCollectorArray Inputs;
-	Node->GetInputPins(Inputs);
-	const UEdGraphSchema_Niagara* Schema = CastChecked<UEdGraphSchema_Niagara>(Node->GetSchema());
+	const FNode* Node = CallingContext.Last();
+	TArray<const FInputPin*> Inputs = GraphBridge::GetInputPins(Node);
 
 	for (int32 i = 0; i < Inputs.Num(); i++)
 	{
-		FNiagaraVariable CallInputVar = Schema->PinToNiagaraVariable(Inputs[i]);
+		FNiagaraVariable CallInputVar = GraphBridge::GetPinVariable(Inputs[i], false, ENiagaraStructConversion::UserFacing);
 		if (CallInputVar.IsEquivalent(InVar) && CallInputVar.GetType() == FNiagaraTypeDefinition::GetParameterMapDef())
 		{
-			if (Inputs[i]->LinkedTo.Num() != 0 && PinToParameterMapIndices.Num() >= 2)
+			const FOutputPin* OutputPin = GraphBridge::GetLinkedOutputPin(Inputs[i]);
+			if (OutputPin && PinToParameterMapIndices.Num() >= 2)
 			{
-				const UEdGraphPin* OutputPin = Inputs[i]->LinkedTo[0];
 				const int32* ParamMapIdxPtr = PinToParameterMapIndices[PinToParameterMapIndices.Num() - 2].Find(OutputPin);
 				if (ParamMapIdxPtr != nullptr)
 				{
 					return *ParamMapIdxPtr;
 				}
-				else
-				{
-					FString ScriptUsageDisplayName;
-					const UNiagaraNodeOutput* ContextOutputNode = FNiagaraStackGraphUtilities::GetEmitterOutputNodeForStackNode(*Node);
-					if (ContextOutputNode != nullptr)
-					{
-						UEnum* NiagaraScriptUsageEnum = FindObjectChecked<UEnum>(nullptr, TEXT("/Script/Niagara.ENiagaraScriptUsage"), true);
-						ScriptUsageDisplayName = NiagaraScriptUsageEnum->GetDisplayNameTextByValue((uint64)ContextOutputNode->GetUsage()).ToString();
-					}
-					else
-					{
-						ScriptUsageDisplayName = "Unknown";
-					}
-					FString FunctionDisplayName;
-					if (FunctionNameContextStack.Num() > 0)
-					{
-						FunctionDisplayName = FunctionNameContextStack.Last().ToString();
-					}
-					else
-					{
-						FunctionDisplayName = Node->GetName();
-					}
-
-					
-					/* 
-					UE_LOG(LogNiagaraEditor, Error, TEXT("Unable to find matching parameter map for variable.  Name:%s  Function:%s  Usage:%s  Package:%s"),
-						*InVar.GetName().ToString(), *FunctionDisplayName, *ScriptUsageDisplayName, *Node->GetOutermost()->GetName());
-					*/
-					
-				}
 			}
 		}
 	}
 	return INDEX_NONE;
 }
 
-int32 FNiagaraParameterMapHistoryBuilder::FindMatchingStaticFromContextInputs(const FNiagaraVariable& InVar) const
+template<typename GraphBridge>
+int32 TNiagaraParameterMapHistoryBuilder<GraphBridge>::FindMatchingStaticFromContextInputs(const FNiagaraVariable& InVar) const
 {
-	if (CallingContext.Num() == 0)
+	if (CallingContext.Num() == 0 || PinToConstantIndices.Num() < 2)
 	{
 		return INDEX_NONE;
 	}
-	const UNiagaraNode* Node = CallingContext.Last();
-	FPinCollectorArray Inputs;
-	Node->GetInputPins(Inputs);
-	const UEdGraphSchema_Niagara* Schema = CastChecked<UEdGraphSchema_Niagara>(Node->GetSchema());
+
+	const TMap<const FPin*, int32>& ConstantIndexMapping = PinToConstantIndices[PinToConstantIndices.Num() - 2];
+	const FNode* Node = CallingContext.Last();
+	TArray<const FInputPin*> Inputs = GraphBridge::GetInputPins(Node);
 
 	for (int32 i = 0; i < Inputs.Num(); i++)
 	{
-		FNiagaraVariable CallInputVar = Schema->PinToNiagaraVariable(Inputs[i]);
+		FNiagaraVariable CallInputVar = GraphBridge::GetPinVariable(Inputs[i], false, ENiagaraStructConversion::UserFacing);
 		if (CallInputVar.IsEquivalent(InVar))
 		{
-			// If linked to something, then the output we are linked to will carry the constant.
-			if (Inputs[i]->LinkedTo.Num() != 0 && PinToConstantIndices.Num() >= 2)
+			const int32* IdxPtr = nullptr;
+
+			if (const FOutputPin* OutputPin = GraphBridge::GetLinkedOutputPin(Inputs[i]))
 			{
-				const UEdGraphPin* OutputPin = Inputs[i]->LinkedTo[0];
-				//int32 ConstantIdx = GetConstantFromOutputPin(OutputPin);
-				const int32* IdxPtr = PinToConstantIndices[PinToConstantIndices.Num() - 2].Find(OutputPin);
-				if (IdxPtr != nullptr)
-					return *IdxPtr;
+				// If linked to something, then the output we are linked to will carry the constant.
+				IdxPtr = ConstantIndexMapping.Find(OutputPin);
 			}
-			// If just using the default, the input pin will carry the constant.
-			else if (Inputs[i]->LinkedTo.Num() == 0 && PinToConstantIndices.Num() >= 2)
-			{				
-				const int32* IdxPtr = PinToConstantIndices[PinToConstantIndices.Num() - 2].Find(Inputs[i]);
-				if (IdxPtr != nullptr)
-					return *IdxPtr;
+			else
+			{
+				// If just using the default, the input pin will carry the constant.
+				IdxPtr = ConstantIndexMapping.Find(Inputs[i]);
+			}
+
+			if (IdxPtr)
+			{
+				return *IdxPtr;
 			}
 		}
 	}
 	return INDEX_NONE;
 }
 
-void FNiagaraParameterMapHistoryBuilder::BeginTranslation(const UNiagaraScript* Script)
+template<typename GraphBridge>
+void TNiagaraParameterMapHistoryBuilder<GraphBridge>::BeginTranslation(const UNiagaraScript* Script)
 {
 	//For now this will just tell particle scripts what emitter they're being compiled as part of but maybe we want to do more here.
 	//This is mainly so that parameter names match up between System/Emitter scripts and the parameters they drive within particle scripts.
@@ -1364,12 +1375,14 @@ void FNiagaraParameterMapHistoryBuilder::BeginTranslation(const UNiagaraScript*
 	BeginTranslation(Emitter);
 }
 
-void FNiagaraParameterMapHistoryBuilder::EndTranslation(const UNiagaraScript* Script)
+template<typename GraphBridge>
+void TNiagaraParameterMapHistoryBuilder<GraphBridge>::EndTranslation(const UNiagaraScript* Script)
 {
 	EmitterNameContextStack.Reset();
 }
 
-void FNiagaraParameterMapHistoryBuilder::BeginTranslation(const UNiagaraEmitter* Emitter)
+template<typename GraphBridge>
+void TNiagaraParameterMapHistoryBuilder<GraphBridge>::BeginTranslation(const UNiagaraEmitter* Emitter)
 {
 	//For now this will just tell particle scripts what emitter they're being compiled as part of but maybe we want to do more here.
 	//This is mainly so that parameter names match up between System/Emitter scripts and the parameters they drive within particle scripts.
@@ -1383,12 +1396,14 @@ void FNiagaraParameterMapHistoryBuilder::BeginTranslation(const UNiagaraEmitter*
 	BuildCurrentAliases();
 }
 
-void FNiagaraParameterMapHistoryBuilder::EndTranslation(const UNiagaraEmitter* Emitter)
+template<typename GraphBridge>
+void TNiagaraParameterMapHistoryBuilder<GraphBridge>::EndTranslation(const UNiagaraEmitter* Emitter)
 {
 	EmitterNameContextStack.Reset();
 }
 
-void FNiagaraParameterMapHistoryBuilder::BeginTranslation(const FString& EmitterUniqueName)
+template<typename GraphBridge>
+void TNiagaraParameterMapHistoryBuilder<GraphBridge>::BeginTranslation(const FString& EmitterUniqueName)
 {
 	//For now this will just tell particle scripts what emitter they're being compiled as part of but maybe we want to do more here.
 	//This is mainly so that parameter names match up between System/Emitter scripts and the parameters they drive within particle scripts.
@@ -1400,12 +1415,14 @@ void FNiagaraParameterMapHistoryBuilder::BeginTranslation(const FString& Emitter
 	}
 	BuildCurrentAliases();
 }
-void FNiagaraParameterMapHistoryBuilder::EndTranslation(const FString& EmitterUniqueName)
+template<typename GraphBridge>
+void TNiagaraParameterMapHistoryBuilder<GraphBridge>::EndTranslation(const FString& EmitterUniqueName)
 {
 	EmitterNameContextStack.Reset();
 }
 
-void FNiagaraParameterMapHistoryBuilder::BeginUsage(ENiagaraScriptUsage InUsage, FName InStageName)
+template<typename GraphBridge>
+void TNiagaraParameterMapHistoryBuilder<GraphBridge>::BeginUsage(ENiagaraScriptUsage InUsage, FName InStageName)
 {
 	RelevantScriptUsageContext.Push(InUsage);
 	ScriptUsageContextNameStack.Push(InStageName);
@@ -1413,25 +1430,21 @@ void FNiagaraParameterMapHistoryBuilder::BeginUsage(ENiagaraScriptUsage InUsage,
 	BuildCurrentAliases();
 }
 
-void FNiagaraParameterMapHistoryBuilder::EndUsage()
+template<typename GraphBridge>
+void TNiagaraParameterMapHistoryBuilder<GraphBridge>::EndUsage()
 {
 	RelevantScriptUsageContext.Pop();
 	ScriptUsageContextNameStack.Pop();
 }
 
-const UNiagaraNode* FNiagaraParameterMapHistoryBuilder::GetCallingContext() const
+template<typename GraphBridge>
+const typename GraphBridge::FNode* TNiagaraParameterMapHistoryBuilder<GraphBridge>::GetCallingContext() const
 {
-	if (CallingContext.Num() == 0)
-	{
-		return nullptr;
-	}
-	else
-	{
-		return CallingContext.Last();
-	}
+	return CallingContext.IsEmpty() ? nullptr : CallingContext.Last();
 }
 
-bool FNiagaraParameterMapHistoryBuilder::InTopLevelFunctionCall(ENiagaraScriptUsage InFilterScriptType) const
+template<typename GraphBridge>
+bool TNiagaraParameterMapHistoryBuilder<GraphBridge>::InTopLevelFunctionCall(ENiagaraScriptUsage InFilterScriptType) const
 {
 	if (InFilterScriptType == ENiagaraScriptUsage::EmitterSpawnScript || InFilterScriptType == ENiagaraScriptUsage::EmitterUpdateScript || InFilterScriptType == ENiagaraScriptUsage::SystemSpawnScript || InFilterScriptType == ENiagaraScriptUsage::SystemUpdateScript)
 	{
@@ -1439,11 +1452,10 @@ bool FNiagaraParameterMapHistoryBuilder::InTopLevelFunctionCall(ENiagaraScriptUs
 		{
 			return true;
 		}
-		else if (CallingContext.Num() <= 2 && Cast<UNiagaraNodeEmitter>(CallingContext[0]) != nullptr) // Handle a function call off of an emitter
+		else if (CallingContext.Num() <= 2 && GraphBridge::GetNodeAsEmitter(CallingContext[0]) != nullptr) // Handle a function call off of an emitter
 		{
 			return true;
 		}
-
 	}
 	else if (UNiagaraScript::IsParticleScript(InFilterScriptType))
 	{
@@ -1457,7 +1469,8 @@ bool FNiagaraParameterMapHistoryBuilder::InTopLevelFunctionCall(ENiagaraScriptUs
 }
 
 
-void FNiagaraParameterMapHistoryBuilder::EnterFunction(const FString& InNodeName, const UNiagaraGraph* InGraph, const UNiagaraNode* Node)
+template<typename GraphBridge>
+void TNiagaraParameterMapHistoryBuilder<GraphBridge>::EnterFunction(const FString& InNodeName, const FGraph* InGraph, const FNode* Node)
 {
 	RegisterNodeVisitation(Node);
 	CallingContext.Push(Node);
@@ -1474,7 +1487,8 @@ void FNiagaraParameterMapHistoryBuilder::EnterFunction(const FString& InNodeName
 	ContextuallyVisitedNodes.Emplace();
 }
 
-void FNiagaraParameterMapHistoryBuilder::ExitFunction(const FString& InNodeName, const UNiagaraNode* Node)
+template<typename GraphBridge>
+void TNiagaraParameterMapHistoryBuilder<GraphBridge>::ExitFunction(const FNode* Node)
 {
 	CallingContext.Pop();
 	CallingGraphContext.Pop();
@@ -1486,7 +1500,8 @@ void FNiagaraParameterMapHistoryBuilder::ExitFunction(const FString& InNodeName,
 	ContextuallyVisitedNodes.Pop();
 }
 
-void FNiagaraParameterMapHistoryBuilder::EnterEmitter(const FString& InEmitterName, const UNiagaraGraph* InGraph, const UNiagaraNode* Node)
+template<typename GraphBridge>
+void TNiagaraParameterMapHistoryBuilder<GraphBridge>::EnterEmitter(const FString& InEmitterName, const FGraph* InGraph, const FNode* Node)
 {
 	RegisterNodeVisitation(Node);
 	CallingContext.Push(Node);
@@ -1498,17 +1513,17 @@ void FNiagaraParameterMapHistoryBuilder::EnterEmitter(const FString& InEmitterNa
 	// Emitters must record their namespaces to their histories as well as
 	// make sure to record their current usage type is so that we can filter variables
 	// for relevance downstream.
-	const UNiagaraNodeEmitter* EmitterNode = Cast<UNiagaraNodeEmitter>(Node);
+	const FEmitterNode* EmitterNode = GraphBridge::GetNodeAsEmitter(Node);
 	if (EmitterNode != nullptr)
 	{
-		RelevantScriptUsageContext.Emplace(EmitterNode->GetUsage());
+		RelevantScriptUsageContext.Emplace(GraphBridge::GetEmitterUsage(EmitterNode));
 	}
 	else
 	{
 		RelevantScriptUsageContext.Emplace(ENiagaraScriptUsage::EmitterSpawnScript);
 	}
 
-	for (FNiagaraParameterMapHistory& History : Histories)
+	for (FHistory& History : Histories)
 	{
 		History.EmitterNamespacesEncountered.AddUnique(InEmitterName);
 	}
@@ -1518,7 +1533,8 @@ void FNiagaraParameterMapHistoryBuilder::EnterEmitter(const FString& InEmitterNa
 	ContextuallyVisitedNodes.Emplace();
 }
 
-void FNiagaraParameterMapHistoryBuilder::ExitEmitter(const FString& InEmitterName, const UNiagaraNode* Node)
+template<typename GraphBridge>
+void TNiagaraParameterMapHistoryBuilder<GraphBridge>::ExitEmitter(const FString& InEmitterName, const FNode* Node)
 {
 	CallingContext.Pop();
 	CallingGraphContext.Pop();
@@ -1530,13 +1546,14 @@ void FNiagaraParameterMapHistoryBuilder::ExitEmitter(const FString& InEmitterNam
 }
 
 
-bool FNiagaraParameterMapHistoryBuilder::IsInEncounteredFunctionNamespace(const FNiagaraVariable& InVar) const
+template<typename GraphBridge>
+bool TNiagaraParameterMapHistoryBuilder<GraphBridge>::IsInEncounteredFunctionNamespace(const FNiagaraVariable& InVar) const
 {
 	if (EncounteredFunctionNames.Num() != 0)
 	{
 		for (FString EncounteredNamespace : EncounteredFunctionNames.Last())
 		{
-			if (FNiagaraParameterMapHistory::IsInNamespace(InVar, EncounteredNamespace))
+			if (FNiagaraParameterUtilities::IsInNamespace(InVar, EncounteredNamespace))
 			{
 				return true;
 			}
@@ -1545,11 +1562,12 @@ bool FNiagaraParameterMapHistoryBuilder::IsInEncounteredFunctionNamespace(const
 	return false;
 }
 
-bool FNiagaraParameterMapHistoryBuilder::IsInEncounteredEmitterNamespace(const FNiagaraVariable& InVar) const
+template<typename GraphBridge>
+bool TNiagaraParameterMapHistoryBuilder<GraphBridge>::IsInEncounteredEmitterNamespace(const FNiagaraVariable& InVar) const
 {
 	for (FString EmitterEncounteredNamespace : EncounteredEmitterNames)
 	{
-		if (FNiagaraParameterMapHistory::IsInNamespace(InVar, EmitterEncounteredNamespace))
+		if (FNiagaraParameterUtilities::IsInNamespace(InVar, EmitterEncounteredNamespace))
 		{
 			return true;
 		}
@@ -1560,96 +1578,95 @@ bool FNiagaraParameterMapHistoryBuilder::IsInEncounteredEmitterNamespace(const F
 /**
 * Use the current alias map to resolve any aliases in this input variable name.
 */
-FNiagaraVariable FNiagaraParameterMapHistoryBuilder::ResolveAliases(const FNiagaraVariable& InVar) const
+template<typename GraphBridge>
+FNiagaraVariable TNiagaraParameterMapHistoryBuilder<GraphBridge>::ResolveAliases(const FNiagaraVariable& InVar) const
 {
 	FNiagaraVariable Var = FNiagaraUtilities::ResolveAliases(InVar, ResolveAliasContext);
 	//ensure(!Var.IsInNameSpace(FNiagaraConstants::StackContextNamespace));
 	return Var;
 }
 
-void FNiagaraParameterMapHistoryBuilder::RegisterNodeVisitation(const UEdGraphNode* Node)
+template<typename GraphBridge>
+void TNiagaraParameterMapHistoryBuilder<GraphBridge>::RegisterNodeVisitation(const FNode* Node)
 {
-	ContextuallyVisitedNodes.Last().AddUnique(CastChecked<UNiagaraNode>(Node));
+	ContextuallyVisitedNodes.Last().AddUnique(Node);
 }
 
 
-const FString* FNiagaraParameterMapHistoryBuilder::GetModuleAlias() const
+template<typename GraphBridge>
+const FString* TNiagaraParameterMapHistoryBuilder<GraphBridge>::GetModuleAlias() const
 {
 	return ResolveAliasContext.GetModuleName().IsSet()
 		? &ResolveAliasContext.GetModuleName().GetValue()
 		: nullptr;
 }
 
-const FString* FNiagaraParameterMapHistoryBuilder::GetEmitterAlias() const
+template<typename GraphBridge>
+const FString* TNiagaraParameterMapHistoryBuilder<GraphBridge>::GetEmitterAlias() const
 {
 	return ResolveAliasContext.GetEmitterName().IsSet()
 		? &ResolveAliasContext.GetEmitterName().GetValue()
 		: nullptr;
 }
 
-void FNiagaraParameterMapHistoryBuilder::VisitInputPin(const UEdGraphPin* Pin, const class UNiagaraNode* InNode, bool bFilterForCompilation)
+template<typename GraphBridge>
+void TNiagaraParameterMapHistoryBuilder<GraphBridge>::VisitInputPin(const FInputPin* Pin, bool bFilterForCompilation)
 {
-	const UEdGraphSchema_Niagara* Schema = GetDefault<UEdGraphSchema_Niagara>();
-
-	if (Pin != nullptr && Pin->Direction == EEdGraphPinDirection::EGPD_Input)
+	if (Pin && ensure(Pin->Direction == EEdGraphPinDirection::EGPD_Input))
 	{
-		for (int32 j = 0; j < Pin->LinkedTo.Num(); j++)
+		if (const FOutputPin* LinkedPin = GraphBridge::GetLinkedOutputPin(Pin))
 		{
-			UEdGraphPin* LinkedPin = Pin->LinkedTo[j];	
+			const FNode* Node = GraphBridge::GetOwningNode(LinkedPin);
 
-			if (LinkedPin)
+			if (!GetNodePreviouslyVisited(Node))
 			{
-				UNiagaraNode* Node = CastChecked<UNiagaraNode>(LinkedPin->GetOwningNode());
-
-				if (!GetNodePreviouslyVisited(Node))
+				if (UNiagaraScript::LogCompileStaticVars > 0)
 				{
-				
-					if (UNiagaraScript::LogCompileStaticVars > 0)
-					{
-						UE_LOG(LogNiagaraEditor, Log, TEXT("Build Parameter Map History: %s %s"), *Node->GetClass()->GetName(), *Node->GetNodeTitle(ENodeTitleType::FullTitle).ToString());
-					}
-					Node->BuildParameterMapHistory(*this, true, bFilterForCompilation);
-					RegisterNodeVisitation(Node);
+					UE_LOG(LogNiagaraEditor, Log, TEXT("Build Parameter Map History: %s %s"), *GraphBridge::GetNodeClassName(Node), *GraphBridge::GetNodeTitle(Node));
 				}
+				Node->BuildParameterMapHistory(*this, true, bFilterForCompilation);
+				RegisterNodeVisitation(Node);
+			}
 
-				if (Schema->PinToTypeDefinition(Pin) == FNiagaraTypeDefinition::GetParameterMapDef())
-				{
-					int32 ParamMapIdx = TraceParameterMapOutputPin(LinkedPin);
+			if (GraphBridge::IsParameterMapPin(Pin))
+			{
+				int32 ParamMapIdx = TraceParameterMapOutputPin(LinkedPin);
 					
-					if (UNiagaraScript::LogCompileStaticVars > 0)
+				if (UNiagaraScript::LogCompileStaticVars > 0)
+				{
+					if (ParamMapIdx == INDEX_NONE)
 					{
-						if (ParamMapIdx == INDEX_NONE)
-						{
-							UE_LOG(LogNiagaraEditor, Log, TEXT("Found bad pin!"));
-						}
-
-						UE_LOG(LogNiagaraEditor, Log, TEXT("Build Parameter Map History: %s %s PMapIdx: %d"), *Node->GetClass()->GetName(), *Node->GetNodeTitle(ENodeTitleType::FullTitle).ToString(), ParamMapIdx);
+						UE_LOG(LogNiagaraEditor, Log, TEXT("Found bad pin!"));
 					}
 
-					RegisterParameterMapPin(ParamMapIdx, Pin);
+					UE_LOG(LogNiagaraEditor, Log, TEXT("Build Parameter Map History: %s %s PMapIdx: %d"), *GraphBridge::GetNodeClassName(Node), *GraphBridge::GetNodeTitle(Node), ParamMapIdx);
 				}
+
+				RegisterParameterMapPin(ParamMapIdx, Pin);
 			}
 		} 
 	}
 }
 
-void FNiagaraParameterMapHistoryBuilder::VisitInputPins(const class UNiagaraNode* InNode, bool bFilterForCompilation)
+template<typename GraphBridge>
+void TNiagaraParameterMapHistoryBuilder<GraphBridge>::VisitInputPins(const FNode* InNode, bool bFilterForCompilation)
 {
-	FPinCollectorArray InputPins;
-	InNode->GetInputPins(InputPins);
+	TArray<const FInputPin*> InputPins = GraphBridge::GetInputPins(InNode);
 
 	for (int32 i = 0; i < InputPins.Num(); i++)
 	{
-		VisitInputPin(InputPins[i], InNode, bFilterForCompilation);
+		VisitInputPin(InputPins[i], bFilterForCompilation);
 	}
 }
 
-bool FNiagaraParameterMapHistoryBuilder::IsNamespacedVariableRelevantToScriptType(const FNiagaraVariable& InVar, ENiagaraScriptUsage InFilterScriptType)
+template<typename GraphBridge>
+bool TNiagaraParameterMapHistoryBuilder<GraphBridge>::IsNamespacedVariableRelevantToScriptType(const FNiagaraVariable& InVar, ENiagaraScriptUsage InFilterScriptType)
 {
 	return true;
 }
 
-bool FNiagaraParameterMapHistoryBuilder::ShouldTrackVariable(const FNiagaraVariable& InVar)
+template<typename GraphBridge>
+bool TNiagaraParameterMapHistoryBuilder<GraphBridge>::ShouldTrackVariable(const FNiagaraVariable& InVar)
 {
 	if (!bFilterByScriptAllowList)
 	{
@@ -1662,10 +1679,10 @@ bool FNiagaraParameterMapHistoryBuilder::ShouldTrackVariable(const FNiagaraVaria
 	return false;
 }
 
-int32 FNiagaraParameterMapHistoryBuilder::HandleVariableWrite(int32 ParamMapIdx, const UEdGraphPin* InPin)
+template<typename GraphBridge>
+int32 TNiagaraParameterMapHistoryBuilder<GraphBridge>::HandleVariableWrite(int32 ParamMapIdx, const FInputPin* InPin)
 {
-	const UEdGraphSchema_Niagara* Schema = GetDefault<UEdGraphSchema_Niagara>();
-	FNiagaraVariable Var = Schema->PinToNiagaraVariable(InPin);
+	FNiagaraVariable Var = GraphBridge::GetPinVariable(InPin, false, ENiagaraStructConversion::UserFacing);
 
 	if (!ShouldTrackVariable(Var))
 	{
@@ -1675,21 +1692,14 @@ int32 FNiagaraParameterMapHistoryBuilder::HandleVariableWrite(int32 ParamMapIdx,
 	Var = ResolveAliases(Var);
 
 	int32 FoundIdx = AddVariableToHistory(Histories[ParamMapIdx], Var, AliasedVar, InPin);
-	if (InPin && Var.GetType().IsStatic() && InPin->LinkedTo.Num() > 0)
+	if (InPin && Var.GetType().IsStatic())
 	{
-		int32 ConstantIdx = GetConstantFromOutputPin(InPin->LinkedTo[0]);
-		RegisterConstantPin(ConstantIdx, InPin);
+		const FOutputPin* LinkedPin = GraphBridge::GetLinkedOutputPin(InPin);
+		const int32 ConstantIdx = LinkedPin
+			? GetConstantFromOutputPin(LinkedPin)
+			: AddOrGetConstantFromValue(InPin->DefaultValue);
 
-		if (FoundIdx != INDEX_NONE && ParamMapIdx != INDEX_NONE && ConstantIdx != INDEX_NONE)
-		{
-			RegisterConstantVariableWrite(ParamMapIdx, ConstantIdx, FoundIdx, false, false);
-		}
-	}
-	else if (InPin && Var.GetType().IsStatic() && InPin->LinkedTo.Num() == 0)
-	{
-		int32 ConstantIdx = AddOrGetConstantFromValue(InPin->DefaultValue);
 		RegisterConstantPin(ConstantIdx, InPin);
-
 		if (FoundIdx != INDEX_NONE && ParamMapIdx != INDEX_NONE && ConstantIdx != INDEX_NONE)
 		{
 			RegisterConstantVariableWrite(ParamMapIdx, ConstantIdx, FoundIdx, false, false);
@@ -1699,10 +1709,18 @@ int32 FNiagaraParameterMapHistoryBuilder::HandleVariableWrite(int32 ParamMapIdx,
 	return FoundIdx;
 }
 
-int32 FNiagaraParameterMapHistoryBuilder::HandleVariableWrite(int32 ParameterMapIndex, const FNiagaraVariable& Var)
+template<typename GraphBridge>
+void TNiagaraParameterMapHistoryBuilder<GraphBridge>::HandleParameterCollectionRead(const FNiagaraVariableBase& InVar, int32 ParamMapIndex)
 {
-	const UEdGraphSchema_Niagara* Schema = GetDefault<UEdGraphSchema_Niagara>();
+	if (bIncludeParameterCollectionInfo)
+	{
+		AddParameterCollection(this->AvailableCollections->FindCollection(InVar), ParamMapIndex);
+	}
+}
 
+template<typename GraphBridge>
+int32 TNiagaraParameterMapHistoryBuilder<GraphBridge>::HandleVariableWrite(int32 ParameterMapIndex, const FNiagaraVariable& Var)
+{
 	if (!ShouldTrackVariable(Var))
 	{
 		return INDEX_NONE;
@@ -1714,14 +1732,14 @@ int32 FNiagaraParameterMapHistoryBuilder::HandleVariableWrite(int32 ParameterMap
 }
 
 
-int32 FNiagaraParameterMapHistoryBuilder::HandleVariableRead(int32 ParamMapIdx, const UEdGraphPin* InPin, bool RegisterReadsAsVariables, const UEdGraphPin* InDefaultPin, bool bFilterForCompilation, bool& OutUsedDefault)
+template<typename GraphBridge>
+int32 TNiagaraParameterMapHistoryBuilder<GraphBridge>::HandleVariableRead(int32 ParamMapIdx, const FPin* InPin, bool RegisterReadsAsVariables, const FInputPin* InDefaultPin, bool bFilterForCompilation, bool& OutUsedDefault)
 {
 	FString DefaultValue;
 	OutUsedDefault = false;
-	const UEdGraphSchema_Niagara* Schema = GetDefault<UEdGraphSchema_Niagara>();
-	FNiagaraVariable Var = Schema->PinToNiagaraVariable(InPin);
+	FNiagaraVariable Var = GraphBridge::GetPinVariable(InPin, false, ENiagaraStructConversion::UserFacing);
 
-	FNiagaraParameterMapHistory& History = Histories[ParamMapIdx];
+	FHistory& History = Histories[ParamMapIdx];
 
 	if (!ShouldTrackVariable(Var))
 	{
@@ -1731,17 +1749,7 @@ int32 FNiagaraParameterMapHistoryBuilder::HandleVariableRead(int32 ParamMapIdx,
 	Var = ResolveAliases(Var);
 
 	//Track any parameter collections we're referencing.
-	if (bIncludeParameterCollectionInfo)
-	{
-		if (UNiagaraParameterCollection* Collection = Schema->VariableIsFromParameterCollection(Var))
-		{
-			int32 Index = History.ParameterCollections.AddUnique(Collection);
-			History.ParameterCollectionNamespaces.SetNum(History.ParameterCollections.Num());
-			History.ParameterCollectionVariables.SetNum(History.ParameterCollections.Num());
-			History.ParameterCollectionNamespaces[Index] = Collection->GetFullNamespace();
-			History.ParameterCollectionVariables[Index] = Collection->GetParameters();
-		}
-	}
+	HandleParameterCollectionRead(Var, ParamMapIdx);
 
 	const FString* ModuleAlias = GetModuleAlias();
 	FName ModuleName = ModuleAlias ? FName(**ModuleAlias) : NAME_None;
@@ -1750,7 +1758,7 @@ int32 FNiagaraParameterMapHistoryBuilder::HandleVariableRead(int32 ParamMapIdx,
 	bool bUseRapidIterationParam = true;
 
 	int32 FoundIdx = History.FindVariable(Var.GetName(), Var.GetType());
-	if (FoundIdx == -1)
+	if (FoundIdx == INDEX_NONE)
 	{
 		if (RegisterReadsAsVariables)
 		{
@@ -1758,7 +1766,7 @@ int32 FNiagaraParameterMapHistoryBuilder::HandleVariableRead(int32 ParamMapIdx,
 
 			if (InDefaultPin)
 			{
-				VisitInputPin(InDefaultPin, Cast<UNiagaraNode>(InDefaultPin->GetOwningNode()), bFilterForCompilation);
+				VisitInputPin(InDefaultPin, bFilterForCompilation);
 				DefaultValue = InDefaultPin->DefaultValue;
 			}
 
@@ -1766,10 +1774,10 @@ int32 FNiagaraParameterMapHistoryBuilder::HandleVariableRead(int32 ParamMapIdx,
 			AddWriteHistory = false;
 
 			// Add the default binding as well to the parameter history, if used.
-			if (UNiagaraGraph* Graph = Cast<UNiagaraGraph>(InPin->GetOwningNode()->GetGraph()))
+			if (const FGraph* Graph = GraphBridge::GetOwningGraph(GraphBridge::GetOwningNode(InPin)))
 			{
 				FNiagaraScriptVariableBinding VariableBinding;
-				TOptional<ENiagaraDefaultMode> DefaultMode = Graph->GetDefaultMode(AliasedVar, &VariableBinding);
+				TOptional<ENiagaraDefaultMode> DefaultMode = GraphBridge::GetGraphDefaultMode(Graph, AliasedVar, VariableBinding);
 				if (DefaultMode.IsSet() && *DefaultMode == ENiagaraDefaultMode::Binding && VariableBinding.IsValid())
 				{
 					FNiagaraVariable TempVar = FNiagaraVariable(Var.GetType(), VariableBinding.GetName());
@@ -1783,7 +1791,7 @@ int32 FNiagaraParameterMapHistoryBuilder::HandleVariableRead(int32 ParamMapIdx,
 					DefaultValue = History.Variables[FoundIdxBinding].GetName().ToString();
 					bUseRapidIterationParam = false;
 
-					History.PerVariableReadHistory[FoundIdxBinding].AddDefaulted_GetRef().ReadPin = FModuleScopedPin(InDefaultPin, ModuleName);
+					History.PerVariableReadHistory[FoundIdxBinding].AddDefaulted_GetRef().ReadPin = typename FHistory::FModuleScopedPin(InDefaultPin, ModuleName);
 				}
 			}
 		}
@@ -1796,8 +1804,8 @@ int32 FNiagaraParameterMapHistoryBuilder::HandleVariableRead(int32 ParamMapIdx,
 		}
 	}
 
-	FNiagaraParameterMapHistory::FReadHistory& ReadEntry = History.PerVariableReadHistory[FoundIdx].AddDefaulted_GetRef();
-	ReadEntry.ReadPin = FModuleScopedPin(InPin, ModuleName);
+	typename FHistory::FReadHistory& ReadEntry = History.PerVariableReadHistory[FoundIdx].AddDefaulted_GetRef();
+	ReadEntry.ReadPin = typename FHistory::FModuleScopedPin(InPin, ModuleName);
 	if (AddWriteHistory && History.PerVariableWriteHistory[FoundIdx].Num())
 	{
 		ReadEntry.PreviousWritePin = History.PerVariableWriteHistory[FoundIdx].Last();
@@ -1809,7 +1817,7 @@ int32 FNiagaraParameterMapHistoryBuilder::HandleVariableRead(int32 ParamMapIdx,
 	check(History.Variables.Num() == History.PerVariableConstantValue.Num());
 	check(History.Variables.Num() == History.PerVariableConstantDefaultValue.Num());
 
-	if (OutUsedDefault && InDefaultPin && UEdGraphSchema_Niagara::IsStaticPin(InDefaultPin))
+	if (OutUsedDefault && InDefaultPin && GraphBridge::IsStaticPin(InDefaultPin))
 	{
 		int32 ConstantIdx = INDEX_NONE;
 		ENiagaraScriptUsage Usage = GetCurrentUsageContext();
@@ -1829,9 +1837,9 @@ int32 FNiagaraParameterMapHistoryBuilder::HandleVariableRead(int32 ParamMapIdx,
 
 		bool bWroteVarNameToDefaultValue = false;
 
-		if (bUseRapidIterationParam && /*InTopLevelFunctionCall(Usage) && */FNiagaraParameterMapHistory::IsAliasedModuleParameter(Histories[ParamMapIdx].VariablesWithOriginalAliasesIntact[FoundIdx]))
+		if (bUseRapidIterationParam && /*InTopLevelFunctionCall(Usage) && */FNiagaraParameterUtilities::IsAliasedModuleParameter(Histories[ParamMapIdx].VariablesWithOriginalAliasesIntact[FoundIdx]))
 		{					
-			FNiagaraVariable VarRapid = FNiagaraParameterMapHistory::ConvertVariableToRapidIterationConstantName(Histories[ParamMapIdx].Variables[FoundIdx], EmitterNamespace.Len() == 0 ? nullptr : *EmitterNamespace, Usage);
+			FNiagaraVariable VarRapid = FNiagaraParameterUtilities::ConvertVariableToRapidIterationConstantName(Histories[ParamMapIdx].Variables[FoundIdx], EmitterNamespace.Len() == 0 ? nullptr : *EmitterNamespace, Usage);
 			
 			int32 StaticIdx = INDEX_NONE;
 			
@@ -1853,7 +1861,7 @@ int32 FNiagaraParameterMapHistoryBuilder::HandleVariableRead(int32 ParamMapIdx,
 					});;
 			}
 
-			if (ConstantResolver.ResolveConstant(VarRapid))
+			if (ConstantResolver->ResolveConstant(VarRapid))
 			{
 				DefaultValue = VarRapid.GetType().ToString(VarRapid.GetData());
 			}
@@ -1878,15 +1886,15 @@ int32 FNiagaraParameterMapHistoryBuilder::HandleVariableRead(int32 ParamMapIdx,
 			
 
 		}
-		else if (FNiagaraParameterMapHistory::IsExternalConstantNamespace(Var, Usage, 0) || Var.IsInNameSpace(EmitterNamespace) || Var.IsInNameSpace(FNiagaraConstants::SystemNamespace) || Var.IsInNameSpace(FNiagaraConstants::ParticleAttributeNamespace))
+		else if (FNiagaraParameterUtilities::IsExternalConstantNamespace(Var, Usage, 0) || Var.IsInNameSpace(EmitterNamespace) || Var.IsInNameSpace(FNiagaraConstants::SystemNamespace) || Var.IsInNameSpace(FNiagaraConstants::ParticleAttributeNamespace))
 		{
 			const FNiagaraVariable& VarRapid = Histories[ParamMapIdx].Variables[FoundIdx];
 			int32 StaticIdx = StaticVariables.IndexOfByPredicate([&](const FNiagaraVariable& InObj) -> bool
-				{
-					return (InObj.GetName() == VarRapid.GetName());
-				});;
+			{
+				return (InObj.GetName() == VarRapid.GetName());
+			});;
 
-			if (ConstantResolver.ResolveConstant(Var))
+			if (ConstantResolver->ResolveConstant(Var))
 			{
 				DefaultValue = Var.GetType().ToString(Var.GetData());
 			}
@@ -1916,7 +1924,7 @@ int32 FNiagaraParameterMapHistoryBuilder::HandleVariableRead(int32 ParamMapIdx,
 			RegisterConstantVariableWrite(ParamMapIdx, ConstantIdx, FoundIdx, true, bWroteVarNameToDefaultValue);
 		}
 	}
-	else if (FoundIdx != INDEX_NONE && ParamMapIdx != INDEX_NONE && UEdGraphSchema_Niagara::IsStaticPin(InPin))
+	else if (FoundIdx != INDEX_NONE && ParamMapIdx != INDEX_NONE && GraphBridge::IsStaticPin(InPin))
 	{
 		int32 ConstantIdx = GetConstantFromVariableRead(Histories[ParamMapIdx].Variables[FoundIdx]);
 		RegisterConstantPin(ConstantIdx, InPin);
@@ -1925,12 +1933,14 @@ int32 FNiagaraParameterMapHistoryBuilder::HandleVariableRead(int32 ParamMapIdx,
 	return FoundIdx;
 }
 
-void FNiagaraParameterMapHistoryBuilder::RegisterEncounterableVariables(TConstArrayView<FNiagaraVariable> Variables)
+template<typename GraphBridge>
+void TNiagaraParameterMapHistoryBuilder<GraphBridge>::RegisterEncounterableVariables(TConstArrayView<FNiagaraVariable> Variables)
 {
 	EncounterableExternalVariables.Append(Variables.GetData(), Variables.Num());
 }
 
-void FNiagaraParameterMapHistoryBuilder::RegisterExternalStaticVariables(TConstArrayView<FNiagaraVariable> Variables)
+template<typename GraphBridge>
+void TNiagaraParameterMapHistoryBuilder<GraphBridge>::RegisterExternalStaticVariables(TConstArrayView<FNiagaraVariable> Variables)
 {
 	int32 LastIdx = StaticVariables.Num();
 	StaticVariables.Append(Variables.GetData(), Variables.Num());
@@ -1987,11 +1997,12 @@ void UpdateAliasedVariable(FNiagaraVariable& AliasedVar, const FNiagaraVariable&
 	AliasedVar.SetName(*OutVarStrName);
 }
 
-int32 FNiagaraParameterMapHistoryBuilder::HandleExternalVariableRead(int32 ParamMapIdx, const FName& Name)
+template<typename GraphBridge>
+int32 TNiagaraParameterMapHistoryBuilder<GraphBridge>::HandleExternalVariableRead(int32 ParamMapIdx, const FName& Name)
 {
 	const UEdGraphSchema_Niagara* Schema = GetDefault<UEdGraphSchema_Niagara>();
 
-	FNiagaraParameterMapHistory& History = Histories[ParamMapIdx];
+	FHistory& History = Histories[ParamMapIdx];
 
 	FNiagaraVariable Var(FNiagaraTypeDefinition(), Name);
 
@@ -2007,13 +2018,10 @@ int32 FNiagaraParameterMapHistoryBuilder::HandleExternalVariableRead(int32 Param
 	if (bIncludeParameterCollectionInfo)
 	{
 		FNiagaraVariable FoundTempVar;
-		if (UNiagaraParameterCollection* Collection = Schema->VariableIsFromParameterCollection(Name.ToString(), true, FoundTempVar))
+		FParameterCollection Collection = this->AvailableCollections->FindMatchingCollection(Name, true, FoundTempVar);
+
+		if (AddParameterCollection(Collection, ParamMapIdx))
 		{
-			int32 Index = History.ParameterCollections.AddUnique(Collection);
-			History.ParameterCollectionNamespaces.SetNum(History.ParameterCollections.Num());
-			History.ParameterCollectionVariables.SetNum(History.ParameterCollections.Num());
-			History.ParameterCollectionNamespaces[Index] = Collection->GetFullNamespace();
-			History.ParameterCollectionVariables[Index] = Collection->GetParameters();
 			Var = FoundTempVar;
 			UpdateAliasedVariable(AliasedVar, OriginalUnaliasedVar, Var);
 		}
@@ -2061,30 +2069,48 @@ int32 FNiagaraParameterMapHistoryBuilder::HandleExternalVariableRead(int32 Param
 	return FoundIdx;
 }
 
+template<>
+bool TNiagaraParameterMapHistoryBuilder<FNiagaraCompilationGraphBridge>::AddParameterCollection(FParameterCollection Collection, int32 ParamMapIndex)
+{
+	if (!Collection)
+	{
+		return false;
+	}
 
-void FNiagaraParameterMapHistoryBuilder::SetConstantByStaticVariable(int32& OutValue, const UEdGraphPin* InDefaultPin)
+	FNiagaraCompilationGraphBridge::FParameterCollectionStore& CollectionStore = Histories[ParamMapIndex].EncounteredParameterCollections;
+
+	int32 Index = CollectionStore.Collections.AddUnique(Collection);
+	CollectionStore.CollectionNamespaces.SetNum(CollectionStore.Collections.Num());
+	CollectionStore.CollectionVariables.SetNum(CollectionStore.Collections.Num());
+	CollectionStore.CollectionNamespaces[Index] = Collection->GetFullNamespace();
+	CollectionStore.CollectionVariables[Index] = Collection->GetParameters();
+
+	return true;
+}
+
+template<typename GraphBridge>
+void TNiagaraParameterMapHistoryBuilder<GraphBridge>::SetConstantByStaticVariable(int32& OutValue, const FInputPin* InDefaultPin)
 {
 	const UEdGraphSchema_Niagara* Schema = GetDefault<UEdGraphSchema_Niagara>();
 
 	FNiagaraEditorModule& NiagaraEditorModule = FModuleManager::GetModuleChecked<FNiagaraEditorModule>("NiagaraEditor");
 	OutValue = 0;
-	FNiagaraVariable Var = Schema->PinToNiagaraVariable(InDefaultPin, true);
+	FNiagaraVariable Var = GraphBridge::GetPinVariable(InDefaultPin, true, ENiagaraStructConversion::UserFacing);
 	FNiagaraVariable VarDefault = Var;
 	FNiagaraEditorUtilities::ResetVariableToDefaultValue(VarDefault);// Do this to handle non-zero defaults
 	if (VarDefault.GetType().IsStatic())
 	{
 		FNiagaraVariable VarWithValue = FNiagaraVariable(Var.GetType(), Var.GetName());
 		FString Value;
-		if (InDefaultPin && InDefaultPin->Direction == EEdGraphPinDirection::EGPD_Input && InDefaultPin->LinkedTo.Num() > 0 && InDefaultPin->LinkedTo[0] != nullptr)
+		if (const FOutputPin* ConnectedOutputPin = GraphBridge::GetLinkedOutputPin(InDefaultPin))
 		{
-			UEdGraphPin* ConnectedOutputPin = InDefaultPin->LinkedTo[0];
 			int32 InputConstant = GetConstantFromOutputPin(ConnectedOutputPin);
 			if (InputConstant != INDEX_NONE)
 			{
 				Value = Constants[InputConstant];
 			}
 		}
-		else if (InDefaultPin && InDefaultPin->Direction == EEdGraphPinDirection::EGPD_Input && InDefaultPin->LinkedTo.Num() == 0 && InDefaultPin->DefaultValue.Len() != 0)
+		else if (ensure(InDefaultPin->Direction == EEdGraphPinDirection::EGPD_Input))
 		{
 			Value = InDefaultPin->DefaultValue;
 		}
@@ -2113,7 +2139,8 @@ void FNiagaraParameterMapHistoryBuilder::SetConstantByStaticVariable(int32& OutV
 	}
 }
 
-int32 FNiagaraParameterMapHistoryBuilder::FindStaticVariable(const FNiagaraVariable& Var) const
+template<typename GraphBridge>
+int32 TNiagaraParameterMapHistoryBuilder<GraphBridge>::FindStaticVariable(const FNiagaraVariable& Var) const
 {
 	FString Value = Var.GetName().ToString();
 	int32 FoundOverrideIdx = StaticVariables.IndexOfByPredicate([&](const FNiagaraVariable& InObj) -> bool
@@ -2124,7 +2151,8 @@ int32 FNiagaraParameterMapHistoryBuilder::FindStaticVariable(const FNiagaraVaria
 	return FoundOverrideIdx;
 }
 
-bool FNiagaraParameterMapHistoryBuilder::ShouldProcessDepthTraversal(const UNiagaraGraph* Graph)
+template<typename GraphBridge>
+bool TNiagaraParameterMapHistoryBuilder<GraphBridge>::ShouldProcessDepthTraversal(const FGraph* Graph)
 {
 	// first check if there are limits setup for how deep we should look
 	if (MaxGraphDepthTraversal == INDEX_NONE || CurrentGraphDepth < MaxGraphDepthTraversal)
@@ -2132,7 +2160,7 @@ bool FNiagaraParameterMapHistoryBuilder::ShouldProcessDepthTraversal(const UNiag
 		// next, if we only care about static variables, check if the graph we're going to go into actually includes any
 		if (StaticVariableSearchContext.bCollectingStaticVariablesOnly)
 		{
-			return Graph->ReferencesStaticVariable(StaticVariableSearchContext);
+			return GraphBridge::GetGraphReferencesStaticVariables(Graph, StaticVariableSearchContext);
 		}
 
 		return true;
@@ -2141,7 +2169,8 @@ bool FNiagaraParameterMapHistoryBuilder::ShouldProcessDepthTraversal(const UNiag
 	return false;
 }
 
-void FNiagaraParameterMapHistoryBuilder::SetConstantByStaticVariable(int32& OutValue, const FNiagaraVariable& Var)
+template<typename GraphBridge>
+void TNiagaraParameterMapHistoryBuilder<GraphBridge>::SetConstantByStaticVariable(int32& OutValue, const FNiagaraVariable& Var)
 {
 	
 	const UEdGraphSchema_Niagara* Schema = GetDefault<UEdGraphSchema_Niagara>();
@@ -2188,7 +2217,8 @@ void FNiagaraParameterMapHistoryBuilder::SetConstantByStaticVariable(int32& OutV
 }
 
 
-int32 FNiagaraParameterMapHistoryBuilder::AddVariableToHistory(FNiagaraParameterMapHistory& History, const FNiagaraVariable& InVar, const FNiagaraVariable& InAliasedVar, const UEdGraphPin* InPin)
+template<typename GraphBridge>
+int32 TNiagaraParameterMapHistoryBuilder<GraphBridge>::AddVariableToHistory(FHistory& History, const FNiagaraVariable& InVar, const FNiagaraVariable& InAliasedVar, const FPin* InPin)
 {
 	const FString* ModuleAlias = GetModuleAlias();
 	const FName ModuleName = ModuleAlias ? FName(**ModuleAlias) : NAME_None;
@@ -2196,7 +2226,8 @@ int32 FNiagaraParameterMapHistoryBuilder::AddVariableToHistory(FNiagaraParameter
 	return History.AddVariable(InVar, InAliasedVar, ModuleName, InPin);
 }
 
-void FNiagaraParameterMapHistoryBuilder::BuildCurrentAliases()
+template<typename GraphBridge>
+void TNiagaraParameterMapHistoryBuilder<GraphBridge>::BuildCurrentAliases()
 {
 	if(RelevantScriptUsageContext.Num() > 0)
 	{
@@ -2468,7 +2499,8 @@ uint32 FCompileConstantResolver::BuildTypeHash() const
 	return HashUnion.Hash;
 }
 
-int32 FNiagaraParameterMapHistoryWithMetaDataBuilder::AddVariableToHistory(FNiagaraParameterMapHistory& History, const FNiagaraVariable& InVar, const FNiagaraVariable& InAliasedVar, const UEdGraphPin* InPin)
+template<typename GraphBuilder>
+int32 TNiagaraParameterMapHistoryWithMetaDataBuilder<GraphBuilder>::AddVariableToHistory(FHistory& History, const FNiagaraVariable& InVar, const FNiagaraVariable& InAliasedVar, const FPin* InPin)
 {
 	const FString* ModuleAlias = GetModuleAlias();
 	const FName ModuleName = ModuleAlias ? FName(**ModuleAlias) : NAME_None;
@@ -2543,4 +2575,18 @@ void FNiagaraGraphCachedBuiltHistory::SetSourceEmitter(const FVersionedNiagaraEm
 	}
 }
 
+template<typename GraphBridge>
+void TNiagaraParameterMapHistoryWithMetaDataBuilder<GraphBridge>::AddGraphToCallingGraphContextStack(const FGraph* InGraph)
+{
+	CallingGraphContext.Add(InGraph);
+};
+
+//////////////////////////////////////////////////////////////////////////
+
+template struct TNiagaraParameterMapHistory<FNiagaraCompilationGraphBridge>;
+
+template class TNiagaraParameterMapHistoryBuilder<FNiagaraCompilationGraphBridge>;
+
+template class TNiagaraParameterMapHistoryWithMetaDataBuilder<FNiagaraCompilationGraphBridge>;
+
 #undef LOCTEXT_NAMESPACE
diff --git a/Engine/Plugins/FX/Niagara/Source/NiagaraEditor/Private/ViewModels/HierarchyEditor/NiagaraSummaryViewViewModel.cpp b/Engine/Plugins/FX/Niagara/Source/NiagaraEditor/Private/ViewModels/HierarchyEditor/NiagaraSummaryViewViewModel.cpp
--- a/Engine/Plugins/FX/Niagara/Source/NiagaraEditor/Private/ViewModels/HierarchyEditor/NiagaraSummaryViewViewModel.cpp
+++ b/Engine/Plugins/FX/Niagara/Source/NiagaraEditor/Private/ViewModels/HierarchyEditor/NiagaraSummaryViewViewModel.cpp
@@ -1,6 +1,7 @@
 // Copyright Epic Games, Inc. All Rights Reserved.
 
 #include "ViewModels/HierarchyEditor/NiagaraSummaryViewViewModel.h"
+#include "GraphEditAction.h"
 #include "NiagaraConstants.h"
 #include "NiagaraEditorStyle.h"
 #include "NiagaraEmitterEditorData.h"
diff --git a/Engine/Plugins/FX/Niagara/Source/NiagaraEditor/Private/ViewModels/NiagaraParameterPanelViewModel.cpp b/Engine/Plugins/FX/Niagara/Source/NiagaraEditor/Private/ViewModels/NiagaraParameterPanelViewModel.cpp
--- a/Engine/Plugins/FX/Niagara/Source/NiagaraEditor/Private/ViewModels/NiagaraParameterPanelViewModel.cpp
+++ b/Engine/Plugins/FX/Niagara/Source/NiagaraEditor/Private/ViewModels/NiagaraParameterPanelViewModel.cpp
@@ -19,6 +19,7 @@
 #include "NiagaraObjectSelection.h"
 #include "NiagaraParameterDefinitions.h"
 #include "NiagaraScriptGraphViewModel.h"
+#include "NiagaraScriptSource.h"
 #include "NiagaraScriptVariable.h"
 #include "NiagaraSimulationStageBase.h"
 #include "NiagaraSystem.h"
@@ -1965,7 +1966,7 @@ TArray<FNiagaraParameterPanelItem> FNiagaraSystemToolkitParameterPanelViewModel:
 					: FCompileConstantResolver();
 				 
 				Builder.SetIgnoreDisabled(bIgnoreDisabled);
-				Builder.ConstantResolver = ConstantResolver;
+				*Builder.ConstantResolver = ConstantResolver;
 				FName StageName;
 				ENiagaraScriptUsage StageUsage = OutputNode->GetUsage();
 				if (StageUsage == ENiagaraScriptUsage::ParticleSimulationStageScript && GraphOwningEmitter.Emitter)
@@ -2054,7 +2055,7 @@ TArray<FNiagaraParameterPanelItem> FNiagaraSystemToolkitParameterPanelViewModel:
 						bool bVarOnlyInTopLevelGraph = true;
 						if (!bForceScript)
 						{
-							for (FModuleScopedPin& WritePin : Builder.Histories[0].PerVariableWriteHistory[VariableIndex])
+							for (FNiagaraParameterMapHistory::FModuleScopedPin& WritePin : Builder.Histories[0].PerVariableWriteHistory[VariableIndex])
 							{
 								UEdGraphNode* VariableOwningNode = WritePin.Pin->GetOwningNode();
 								bVarOnlyInTopLevelGraph &= AllGraphs.Contains(static_cast<const UNiagaraGraph*>(VariableOwningNode->GetGraph()));
diff --git a/Engine/Plugins/FX/Niagara/Source/NiagaraEditor/Private/ViewModels/NiagaraSystemViewModel.cpp b/Engine/Plugins/FX/Niagara/Source/NiagaraEditor/Private/ViewModels/NiagaraSystemViewModel.cpp
--- a/Engine/Plugins/FX/Niagara/Source/NiagaraEditor/Private/ViewModels/NiagaraSystemViewModel.cpp
+++ b/Engine/Plugins/FX/Niagara/Source/NiagaraEditor/Private/ViewModels/NiagaraSystemViewModel.cpp
@@ -1499,7 +1499,7 @@ TArray<FNiagaraGraphParameterReference> FNiagaraSystemViewModel::GetGraphParamet
 				: FCompileConstantResolver();
 
 			Builder.SetIgnoreDisabled(bIgnoreDisabled);
-			Builder.ConstantResolver = ConstantResolver;
+			*Builder.ConstantResolver = ConstantResolver;
 			FName StageName;
 			ENiagaraScriptUsage StageUsage = OutputNode->GetUsage();
 			if (StageUsage == ENiagaraScriptUsage::ParticleSimulationStageScript && GraphOwningEmitter.Emitter)
@@ -1537,7 +1537,7 @@ TArray<FNiagaraGraphParameterReference> FNiagaraSystemViewModel::GetGraphParamet
 						}
 					}
 
-					for (const FModuleScopedPin& Write : Builder.Histories[0].PerVariableWriteHistory[VariableIndex])
+					for (const FNiagaraParameterMapHistory::FModuleScopedPin& Write : Builder.Histories[0].PerVariableWriteHistory[VariableIndex])
 					{
 						if (Write.Pin->GetOwningNode() != nullptr)
 						{
diff --git a/Engine/Plugins/FX/Niagara/Source/NiagaraEditor/Private/ViewModels/Stack/NiagaraStackFunctionInput.cpp b/Engine/Plugins/FX/Niagara/Source/NiagaraEditor/Private/ViewModels/Stack/NiagaraStackFunctionInput.cpp
--- a/Engine/Plugins/FX/Niagara/Source/NiagaraEditor/Private/ViewModels/Stack/NiagaraStackFunctionInput.cpp
+++ b/Engine/Plugins/FX/Niagara/Source/NiagaraEditor/Private/ViewModels/Stack/NiagaraStackFunctionInput.cpp
@@ -1593,19 +1593,19 @@ void UNiagaraStackFunctionInput::GetAvailableParameterHandles(TArray<FNiagaraPar
 			for (int32 i = 1; i < MaxGroupIndex; i++)
 			{
 				UNiagaraNodeFunctionCall* ModuleToCheck = Cast<UNiagaraNodeFunctionCall>(StackGroups[i].EndNode);
-				FNiagaraParameterMapHistoryBuilder Builder;
-				FNiagaraStackGraphUtilities::BuildParameterMapHistoryWithStackContextResolution(Emitter, OutputNode, ModuleToCheck, Builder, false);
+				TArray<FNiagaraParameterMapHistory> Histories;
+				FNiagaraStackGraphUtilities::BuildParameterMapHistoryWithStackContextResolution(Emitter, OutputNode, ModuleToCheck, Histories, false);
 
-				if (Builder.Histories.Num() == 1)
+				if (Histories.Num() == 1)
 				{
-					for (int32 j = 0; j < Builder.Histories[0].Variables.Num(); j++)
+					for (int32 j = 0; j < Histories[0].Variables.Num(); j++)
 					{
-						FNiagaraVariable& HistoryVariable = Builder.Histories[0].Variables[j];
+						FNiagaraVariable& HistoryVariable = Histories[0].Variables[j];
 						FNiagaraParameterHandle AvailableHandle = FNiagaraParameterHandle(HistoryVariable.GetName());
 
 						// check if the variable was written to
 						bool bWritten = false;
-						for (const FModuleScopedPin& WritePin : Builder.Histories[0].PerVariableWriteHistory[j])
+						for (const FNiagaraParameterMapHistory::FModuleScopedPin& WritePin : Histories[0].PerVariableWriteHistory[j])
 						{
 							if (Cast<UNiagaraNodeParameterMapSet>(WritePin.Pin->GetOwningNode()) != nullptr)
 							{
diff --git a/Engine/Plugins/FX/Niagara/Source/NiagaraEditor/Private/ViewModels/Stack/NiagaraStackGraphUtilities.cpp b/Engine/Plugins/FX/Niagara/Source/NiagaraEditor/Private/ViewModels/Stack/NiagaraStackGraphUtilities.cpp
--- a/Engine/Plugins/FX/Niagara/Source/NiagaraEditor/Private/ViewModels/Stack/NiagaraStackGraphUtilities.cpp
+++ b/Engine/Plugins/FX/Niagara/Source/NiagaraEditor/Private/ViewModels/Stack/NiagaraStackGraphUtilities.cpp
@@ -514,8 +514,10 @@ TArray<FName> FNiagaraStackGraphUtilities::StackContextResolution(FVersionedNiag
 	return PossibleRootNames;
 }
 
-void FNiagaraStackGraphUtilities::BuildParameterMapHistoryWithStackContextResolution(FVersionedNiagaraEmitter OwningEmitter, UNiagaraNodeOutput* OutputNodeInChain, UNiagaraNode* NodeToVisit, FNiagaraParameterMapHistoryBuilder& Builder, bool bRecursive /*= true*/, bool bFilterForCompilation /*= true*/)
+void FNiagaraStackGraphUtilities::BuildParameterMapHistoryWithStackContextResolution(FVersionedNiagaraEmitter OwningEmitter, UNiagaraNodeOutput* OutputNodeInChain, UNiagaraNode* NodeToVisit, TArray<FNiagaraParameterMapHistory>& OutHistories, bool bRecursive /*= true*/, bool bFilterForCompilation /*= true*/)
 {
+	FNiagaraParameterMapHistoryBuilder Builder;
+
 	bool bSetUsage = false;
 	FVersionedNiagaraEmitterData* EmitterData = OwningEmitter.GetEmitterData();
 	if (EmitterData && OutputNodeInChain)
@@ -553,6 +555,8 @@ void FNiagaraStackGraphUtilities::BuildParameterMapHistoryWithStackContextResolu
 	{
 		Builder.EndUsage();
 	}
+
+	OutHistories = MoveTemp(Builder.Histories);
 }
 
 
@@ -1194,7 +1198,7 @@ void FNiagaraStackGraphUtilities::GetStackFunctionInputPinsWithoutCache(
 {
 	FNiagaraParameterMapHistoryBuilder Builder;
 	Builder.SetIgnoreDisabled(bIgnoreDisabled);
-	Builder.ConstantResolver = ConstantResolver;
+	*Builder.ConstantResolver = ConstantResolver;
 	Builder.RegisterExternalStaticVariables(StaticVars);
 	
 	/* The line below represents a compromise. It basically says don't look up the static variable from the 
@@ -1826,7 +1830,7 @@ void FNiagaraStackGraphUtilities::GetStackFunctionOutputVariables(UNiagaraNodeFu
 {
 	FNiagaraParameterMapHistoryBuilder Builder;
 	Builder.SetIgnoreDisabled(false);
-	Builder.ConstantResolver = ConstantResolver;
+	*Builder.ConstantResolver = ConstantResolver;
 
 	FVersionedNiagaraEmitter OuterEmitter = FNiagaraStackGraphUtilitiesImpl::GetOuterEmitter(&FunctionCallNode);
 	UNiagaraSystem* System = FunctionCallNode.GetTypedOuter<UNiagaraSystem>();
@@ -1855,7 +1859,7 @@ void FNiagaraStackGraphUtilities::GetStackFunctionOutputVariables(UNiagaraNodeFu
 		for (int32 i = 0; i < Builder.Histories[0].Variables.Num(); i++)
 		{
 			bool bHasParameterMapSetWrite = false;
-			for (const FModuleScopedPin& WritePin : Builder.Histories[0].PerVariableWriteHistory[i])
+			for (const FNiagaraParameterMapHistory::FModuleScopedPin& WritePin : Builder.Histories[0].PerVariableWriteHistory[i])
 			{
 				if (WritePin.Pin != nullptr && WritePin.Pin->GetOwningNode() != nullptr && WritePin.Pin->GetOwningNode()->IsA<UNiagaraNodeParameterMapSet>())
 				{
@@ -1879,7 +1883,7 @@ bool FNiagaraStackGraphUtilities::GetStackFunctionInputAndOutputVariables(UNiaga
 {
 	FNiagaraParameterMapHistoryBuilder Builder;
 	Builder.SetIgnoreDisabled(false);
-	Builder.ConstantResolver = ConstantResolver;
+	*Builder.ConstantResolver = ConstantResolver;
 
 	FVersionedNiagaraEmitter OuterEmitter = FNiagaraStackGraphUtilitiesImpl::GetOuterEmitter(&FunctionCallNode);
 	UNiagaraSystem* System = FunctionCallNode.GetTypedOuter<UNiagaraSystem>();
@@ -2880,11 +2884,11 @@ FNiagaraVariable FNiagaraStackGraphUtilities::CreateRapidIterationParameter(cons
 	FNiagaraVariable RapidIterationVariable;
 	if (ScriptUsage == ENiagaraScriptUsage::SystemSpawnScript || ScriptUsage == ENiagaraScriptUsage::SystemUpdateScript)
 	{
-		RapidIterationVariable = FNiagaraParameterMapHistory::ConvertVariableToRapidIterationConstantName(InputVariable, nullptr, ScriptUsage); // These names *should* have the emitter baked in...
+		RapidIterationVariable = FNiagaraParameterUtilities::ConvertVariableToRapidIterationConstantName(InputVariable, nullptr, ScriptUsage); // These names *should* have the emitter baked in...
 	}
 	else
 	{
-		RapidIterationVariable = FNiagaraParameterMapHistory::ConvertVariableToRapidIterationConstantName(InputVariable, *UniqueEmitterName, ScriptUsage);
+		RapidIterationVariable = FNiagaraParameterUtilities::ConvertVariableToRapidIterationConstantName(InputVariable, *UniqueEmitterName, ScriptUsage);
 	}
 
 	return RapidIterationVariable;
@@ -2905,7 +2909,7 @@ void FNiagaraStackGraphUtilities::CleanUpStaleRapidIterationParameters(UNiagaraS
 			FString EmitterName;
 			FString FunctionCallName;
 			FString InputName;
-			if (FNiagaraParameterMapHistory::SplitRapidIterationParameterName(RapidIterationParameter, Script.GetUsage(), EmitterName, FunctionCallName, InputName))
+			if (FNiagaraParameterUtilities::SplitRapidIterationParameterName(RapidIterationParameter, Script.GetUsage(), EmitterName, FunctionCallName, InputName))
 			{
 				if (EmitterName != OwningEmitter.Emitter->GetUniqueEmitterName() || ValidFunctionCallNames.Contains(FunctionCallName) == false)
 				{
@@ -3169,7 +3173,7 @@ void FNiagaraStackGraphUtilities::MoveModule(UNiagaraScript& SourceScript, UNiag
 			FString EmitterName;
 			FString FunctionCallName;
 			FString InputName;
-			if (FNiagaraParameterMapHistory::SplitRapidIterationParameterName(
+			if (FNiagaraParameterUtilities::SplitRapidIterationParameterName(
 				ScriptRapidIterationParameter, SourceScript.GetUsage(), EmitterName, FunctionCallName, InputName))
 			{
 				FGuid* NodeIdPtr = FunctionCallNameToNodeIdMap.Find(FunctionCallName);
@@ -3339,7 +3343,7 @@ void FNiagaraStackGraphUtilities::MoveModule(UNiagaraScript& SourceScript, UNiag
 					FString OldEmitterName;
 					FString OldFunctionCallName;
 					FString InputName;
-					FNiagaraParameterMapHistory::SplitRapidIterationParameterName(RapidIterationParameter, SourceScript.GetUsage(), OldEmitterName, OldFunctionCallName, InputName);
+					FNiagaraParameterUtilities::SplitRapidIterationParameterName(RapidIterationParameter, SourceScript.GetUsage(), OldEmitterName, OldFunctionCallName, InputName);
 					FNiagaraParameterHandle ModuleHandle = FNiagaraParameterHandle::CreateModuleParameterHandle(*InputName);
 					FNiagaraParameterHandle AliasedModuleHandle = FNiagaraParameterHandle::CreateAliasedModuleParameterHandle(ModuleHandle, TargetFunctionNode);
 					FNiagaraVariable TargetRapidIterationParameter = CreateRapidIterationParameter(EmitterName, TargetUsage, AliasedModuleHandle.GetParameterHandleString(), RapidIterationParameter.GetType());
@@ -3645,7 +3649,7 @@ void FNiagaraStackGraphUtilities::RenameReferencingParameters(UNiagaraSystem* Sy
 		for (FNiagaraVariable& Variable : RapidIterationVariables)
 		{
 			FString EmitterName, FunctionCallName, InputName;
-			if (FNiagaraParameterMapHistory::SplitRapidIterationParameterName(Variable, Script->GetUsage(), EmitterName, FunctionCallName, InputName))
+			if (FNiagaraParameterUtilities::SplitRapidIterationParameterName(Variable, Script->GetUsage(), EmitterName, FunctionCallName, InputName))
 			{
 				if (EmitterName == OwningEmitterName && FunctionCallName == OldModuleName)
 				{
diff --git a/Engine/Plugins/FX/Niagara/Source/NiagaraEditor/Private/ViewModels/Stack/NiagaraStackModuleItemLinkedInputCollection.cpp b/Engine/Plugins/FX/Niagara/Source/NiagaraEditor/Private/ViewModels/Stack/NiagaraStackModuleItemLinkedInputCollection.cpp
--- a/Engine/Plugins/FX/Niagara/Source/NiagaraEditor/Private/ViewModels/Stack/NiagaraStackModuleItemLinkedInputCollection.cpp
+++ b/Engine/Plugins/FX/Niagara/Source/NiagaraEditor/Private/ViewModels/Stack/NiagaraStackModuleItemLinkedInputCollection.cpp
@@ -77,7 +77,7 @@ void UNiagaraStackModuleItemLinkedInputCollection::RefreshChildrenInternal(const
 	{
 		FNiagaraParameterMapHistoryBuilder Builder;
 		Builder.SetIgnoreDisabled(false);
-		Builder.ConstantResolver = GetEmitterViewModel().IsValid() 
+		*Builder.ConstantResolver = GetEmitterViewModel().IsValid()
 			? FCompileConstantResolver(GetEmitterViewModel()->GetEmitter(), FNiagaraStackGraphUtilities::GetOutputNodeUsage(*FunctionCallNode))
 			: FCompileConstantResolver();
 		FunctionCallNode->BuildParameterMapHistory(Builder, false);
diff --git a/Engine/Plugins/FX/Niagara/Source/NiagaraEditor/Public/INiagaraCompiler.h b/Engine/Plugins/FX/Niagara/Source/NiagaraEditor/Public/INiagaraCompiler.h
--- a/Engine/Plugins/FX/Niagara/Source/NiagaraEditor/Public/INiagaraCompiler.h
+++ b/Engine/Plugins/FX/Niagara/Source/NiagaraEditor/Public/INiagaraCompiler.h
@@ -134,7 +134,6 @@ public:
 	virtual const FString& GetTranslatedHLSL() const = 0;
 
 	static TUniquePtr<INiagaraHlslTranslator> CreateTranslator(const FNiagaraCompileRequestDataBase* InCompileData, const FNiagaraCompileRequestDuplicateDataBase* InDuplicateData);
-	static TUniquePtr<INiagaraHlslTranslator> CreateTranslator(const FNiagaraPrecompileData* InCompileData, const FNiagaraCompilationCopyData* InDuplicateData);
 
 	static FString BuildHLSLStructDecl(const FNiagaraTypeDefinition& Type, FText& OutErrorMessage, bool bGpuScript);
 	static bool IsBuiltInHlslType(const FNiagaraTypeDefinition& Type);
diff --git a/Engine/Plugins/FX/Niagara/Source/NiagaraEditor/Public/NiagaraCompilationBridge.h b/Engine/Plugins/FX/Niagara/Source/NiagaraEditor/Public/NiagaraCompilationBridge.h
new file mode 100644
--- /dev/null
+++ b/Engine/Plugins/FX/Niagara/Source/NiagaraEditor/Public/NiagaraCompilationBridge.h
@@ -0,0 +1,173 @@
+// Copyright Epic Games, Inc. All Rights Reserved.
+
+#pragma once
+
+#include "NiagaraCommon.h"
+
+class FCompileConstantResolver;
+class FNiagaraCompileRequestData;
+class FNiagaraCompileRequestDuplicateData;
+class UEdGraphPin;
+class UNiagaraGraph;
+class UNiagaraNode;
+class UNiagaraNodeConvert;
+class UNiagaraNodeCustomHlsl;
+class UNiagaraNodeEmitter;
+class UNiagaraNodeFunctionCall;
+class UNiagaraNodeIf;
+class UNiagaraNodeInput;
+class UNiagaraNodeOp;
+class UNiagaraNodeOutput;
+class UNiagaraNodeParameterMapFor;
+class UNiagaraNodeParameterMapGet;
+class UNiagaraNodeParameterMapSet;
+class UNiagaraNodeSelect;
+class UNiagaraNodeStaticSwitch;
+class UNiagaraParameterCollection;
+struct FGraphTraversalHandle;
+struct FNiagaraConvertConnection;
+struct FNiagaraCustomHlslInclude;
+struct FNiagaraFindInputNodeOptions;
+struct FNiagaraGraphFunctionAliasContext;
+struct FNiagaraStaticVariableSearchContext;
+template<typename GraphBridge> class TNiagaraParameterMapHistoryBuilder;
+template<typename GraphBridge> struct TNiagaraParameterMapHistory;
+template<typename PinType> struct TModuleScopedPin;
+
+struct FNiagaraCompilationGraphBridge
+{
+	// base types
+	using FGraph = UNiagaraGraph;
+	using FPin = UEdGraphPin;
+	using FCustomHlslNode = UNiagaraNodeCustomHlsl;
+	using FNode = UNiagaraNode;
+	using FOutputNode = UNiagaraNodeOutput;
+	using FInputNode = UNiagaraNodeInput;
+	using FOpNode = UNiagaraNodeOp;
+	using FEmitterNode = UNiagaraNodeEmitter;
+	using FIfNode = UNiagaraNodeIf;
+	using FConvertNode = UNiagaraNodeConvert;
+	using FSelectNode = UNiagaraNodeSelect;
+	using FFunctionCallNode = UNiagaraNodeFunctionCall;
+	using FParamMapGetNode = UNiagaraNodeParameterMapGet;
+	using FParamMapSetNode = UNiagaraNodeParameterMapSet;
+	using FParamMapForNode = UNiagaraNodeParameterMapFor;
+	using FStaticSwitchNode = UNiagaraNodeStaticSwitch;
+	using FInputPin = UEdGraphPin;
+	using FOutputPin = UEdGraphPin;
+	using FParamMapHistory = TNiagaraParameterMapHistory<FNiagaraCompilationGraphBridge>;
+	using FParamMapHistoryBuilder = TNiagaraParameterMapHistoryBuilder<FNiagaraCompilationGraphBridge>;
+	using FPrecompileData = FNiagaraCompileRequestData;
+	using FCompilationCopy = FNiagaraCompileRequestDuplicateData;
+	using FModuleScopedPin = TModuleScopedPin<FPin>;
+	using FGraphTraversalHandle = FGraphTraversalHandle;
+	using FGraphFunctionAliasContext = FNiagaraGraphFunctionAliasContext;
+	using FConvertConnection = FNiagaraConvertConnection;
+	using FParameterCollection = UNiagaraParameterCollection*;
+	using FConstantResolver = FCompileConstantResolver;
+
+	// additional data for extending the ParameterMapHistoryBuilder
+	struct FParameterCollectionStore
+	{
+		void Append(const FParameterCollectionStore& Other);
+		void Add(UNiagaraParameterCollection* Collection);
+
+		TArray<UNiagaraParameterCollection*> Collections;
+		/** Cached off contents of used parameter collections, in case they change during threaded compilation. */
+		TArray<TArray<FNiagaraVariable>> CollectionVariables;
+		/** Cached off contents of used parameter collections, in case they change during threaded compilation. */
+		TArray<FString> CollectionNamespaces;
+	};
+
+	struct FAvailableParameterCollections
+	{
+		UNiagaraParameterCollection* FindCollection(const FNiagaraVariable& Variable) const;
+		UNiagaraParameterCollection* FindMatchingCollection(FName VariableName, bool bAllowPartialMatch, FNiagaraVariable& OutVar) const;
+	};
+
+	class FBuilderExtraData
+	{
+	public:
+		FBuilderExtraData();
+		TUniquePtr<FAvailableParameterCollections> AvailableCollections;
+	};
+
+	static const FGraph* GetGraph(const FCompilationCopy* CompilationCopy);
+	static const FNode* GetOwningNode(const FPin* Pin);
+	static FNode* GetMutableOwningNode(const FPin* Pin);
+	static const FGraph* GetOwningGraph(const FNode* Node);
+	static bool CustomHlslReferencesTokens(const FCustomHlslNode* CustomNode, TConstArrayView<FStringView> TokenStrings);
+	static ENiagaraScriptUsage GetCustomHlslUsage(const FCustomHlslNode* CustomNode);
+	static FString GetCustomHlslString(const FCustomHlslNode* CustomNode);
+	static void GetCustomHlslIncludePaths(const FCustomHlslNode* CustomNode, TArray<FNiagaraCustomHlslInclude>& Includes);
+	static const TArray<FConvertConnection>& GetConvertConnections(const FConvertNode* ConvertNode);
+
+	// various cast functions
+	static const FFunctionCallNode* AsFunctionCallNode(const FNode* Node);
+	static const FInputNode* AsInputNode(const FNode* Node);
+	static const FParamMapGetNode* AsParamMapGetNode(const FNode* Node);
+	static const FCustomHlslNode* AsCustomHlslNode(const FNode* Node);
+	static const FParamMapSetNode* AsParamMapSetNode(const FNode* Node);
+
+	static bool GraphHasParametersOfType(const FGraph* Graph, const FNiagaraTypeDefinition& TypeDef);
+	static TArray<FNiagaraVariableBase> GraphGetStaticSwitchInputs(const FGraph* Graph);
+	static void FindOutputNodes(const FGraph* Graph, ENiagaraScriptUsage ScriptUsage, TArray<const FOutputNode*>& OutputNodes);
+	static void FindOutputNodes(const FGraph* Graph, TArray<const FOutputNode*>& OutputNodes);
+	static void BuildTraversal(const FGraph* Graph, const FNode* OutputNode, TArray<const FNode*>& TraversedNodes);
+	static const FGraph* GetEmitterGraph(const FEmitterNode* EmitterNode);
+	static FString GetEmitterUniqueName(const FEmitterNode* EmitterNode);
+	static ENiagaraScriptUsage GetEmitterUsage(const FEmitterNode* EmitterNode);
+	static FString GetEmitterName(const FEmitterNode* EmitterNode);
+	static FString GetEmitterPathName(const FEmitterNode* EmitterNode);
+	static FString GetEmitterHandleIdString(const FEmitterNode* EmitterNode);
+	static const FGraph* GetFunctionNodeGraph(const FFunctionCallNode* FunctionCall);
+	static FString GetFunctionFullName(const FFunctionCallNode* FunctionCall);
+	static FString GetFunctionScriptName(const FFunctionCallNode* FunctionCall);
+	static FString GetFunctionName(const FFunctionCallNode* FunctionCall);
+	static ENiagaraScriptUsage GetFunctionUsage(const FFunctionCallNode* FunctionCall);
+	static TOptional<ENiagaraDefaultMode> GetGraphDefaultMode(const FGraph* Graph, const FNiagaraVariable& Variable, FNiagaraScriptVariableBinding& Binding);
+	static const FInputPin* GetDefaultPin(const FParamMapGetNode* GetNode, const FOutputPin* OutputPin);
+	static bool IsStaticPin(const FPin* Pin);
+	// retrieves all input pins (excluding any add pins that may be present)
+	static TArray<const FInputPin*> GetInputPins(const FNode* Node);
+	// retrieves all output pins (excluding both orphaned pins and add pins)
+	static TArray<const FOutputPin*> GetOutputPins(const FNode* Node);
+	// gets all pins assoicated with the node
+	static TArray<const FPin*> GetPins(const FNode* Node);
+	static FNiagaraTypeDefinition GetPinType(const FPin* Pin, ENiagaraStructConversion Conversion);
+	static FText GetPinFriendlyName(const FPin* Pin);
+	static FText GetPinDisplayName(const FPin* Pin);
+	static FNiagaraVariable GetPinVariable(const FPin* Pin, bool bNeedsValue, ENiagaraStructConversion Conversion);
+	static const FInputPin* GetPinAsInput(const FPin* Pin);
+	static FNiagaraVariable GetInputVariable(const FInputNode* InputNode);
+	static const TArray<FNiagaraVariable>& GetOutputVariables(const FOutputNode* OutputNode);
+	static TArray<FNiagaraVariable> GetGraphOutputNodeVariables(const FGraph* Graph, ENiagaraScriptUsage Usage);
+	static TArray<const FInputNode*> GetGraphInputNodes(const FGraph* Graph, const FNiagaraFindInputNodeOptions& Options);
+	static const FOutputPin* GetLinkedOutputPin(const FInputPin* InputPin);
+	static bool CanCreateConnection(const FOutputPin* OutputPin, const FInputPin* InputPin, FText& FailureMessage);
+	static ENiagaraScriptUsage GetOutputNodeUsage(const FOutputNode* OutputNode);
+	static FGuid GetOutputNodeUsageId(const FOutputNode* OutputNode);
+	static ENiagaraScriptUsage GetOutputNodeScriptType(const FOutputNode* OutputNode);
+	static FGuid GetOutputNodeScriptTypeId(const FOutputNode* OutputNode);
+	static bool IsGraphEmpty(const FGraph* Graph);
+	static void AddCollectionPaths(const FParamMapHistory& History, TArray<FString>& Paths);
+	static bool NodeIsEnabled(const FNode* Node);
+	static TOptional<ENiagaraDefaultMode> GraphGetDefaultMode(const FGraph* Graph, const FNiagaraVariableBase& Variable, FNiagaraScriptVariableBinding& Binding);
+	static const FOutputPin* GetSelectOutputPin(const FSelectNode* SelectNode, const FNiagaraVariableBase& Variable);
+	static FString GetNodeName(const FNode* Node);
+	static FString GetNodeTitle(const FNode* Node);
+	static const FInputPin* GetInputPin(const FNode* Node, int32 PinIndex);
+	static int32 GetPinIndexById(TConstArrayView<const FPin*> Pins, const FGuid& PinId);
+	static FString GetCollectionFullName(FParameterCollection Collection);
+	static bool IsCollectionValid(FParameterCollection Collection);
+	static UNiagaraDataInterface* GetCollectionDataInterface(FParameterCollection Collection, const FNiagaraVariable& Variable);
+	static UObject* GetCollectionUObject(FParameterCollection Collection, const FNiagaraVariable& Variable);
+
+	static const FOutputNode* AsOutputNode(const FNode* Node);
+	static bool IsParameterMapGet(const FNode* Node);
+	static TOptional<FName> GetOutputNodeStackContextOverride(const FOutputNode* OutputNode);
+	static FString GetNodeClassName(const FNode* Node);
+	static bool IsParameterMapPin(const FPin* Pin);
+	static bool GetGraphReferencesStaticVariables(const FGraph* Graph, FNiagaraStaticVariableSearchContext& StaticVariableContext);
+	static const FEmitterNode* GetNodeAsEmitter(const FNode* Node);
+};
diff --git a/Engine/Plugins/FX/Niagara/Source/NiagaraEditor/Public/NiagaraGraph.h b/Engine/Plugins/FX/Niagara/Source/NiagaraEditor/Public/NiagaraGraph.h
--- a/Engine/Plugins/FX/Niagara/Source/NiagaraEditor/Public/NiagaraGraph.h
+++ b/Engine/Plugins/FX/Niagara/Source/NiagaraEditor/Public/NiagaraGraph.h
@@ -140,6 +140,41 @@ struct FNiagaraGraphFunctionAliasContext
 	TArray<const UEdGraphPin*> StaticSwitchValues;
 };
 
+/** Options for the FindInputNodes function */
+struct FNiagaraFindInputNodeOptions
+{
+	FNiagaraFindInputNodeOptions()
+		: bSort(false)
+		, bIncludeParameters(true)
+		, bIncludeAttributes(true)
+		, bIncludeSystemConstants(true)
+		, bIncludeTranslatorConstants(false)
+		, bFilterDuplicates(false)
+		, bFilterByScriptUsage(false)
+		, TargetScriptUsage(ENiagaraScriptUsage::Function)
+	{
+	}
+
+	/** Whether or not to sort the nodes, defaults to false. */
+	bool bSort;
+	/** Whether or not to include parameters, defaults to true. */
+	bool bIncludeParameters;
+	/** Whether or not to include attributes, defaults to true. */
+	bool bIncludeAttributes;
+	/** Whether or not to include system parameters, defaults to true. */
+	bool bIncludeSystemConstants;
+	/** Whether or not to include translator parameters, defaults to false. */
+	bool bIncludeTranslatorConstants;
+	/** Whether of not to filter out duplicate nodes, defaults to false. */
+	bool bFilterDuplicates;
+	/** Whether or not to limit to nodes connected to an output node of the specified script type.*/
+	bool bFilterByScriptUsage;
+	/** The specified script usage required for an input.*/
+	ENiagaraScriptUsage TargetScriptUsage;
+	/** The specified id within the graph of the script usage*/
+	FGuid TargetScriptUsageId;
+};
+
 UCLASS(MinimalAPI)
 class UNiagaraGraph : public UEdGraph
 {
@@ -180,40 +215,7 @@ class UNiagaraGraph : public UEdGraph
 	void FindOutputNodes(ENiagaraScriptUsage TargetUsageType, TArray<UNiagaraNodeOutput*>& OutputNodes) const;
 	void FindEquivalentOutputNodes(ENiagaraScriptUsage TargetUsageType, TArray<UNiagaraNodeOutput*>& OutputNodes) const;
 
-	/** Options for the FindInputNodes function */
-	struct FFindInputNodeOptions
-	{
-		FFindInputNodeOptions()
-			: bSort(false)
-			, bIncludeParameters(true)
-			, bIncludeAttributes(true)
-			, bIncludeSystemConstants(true)
-			, bIncludeTranslatorConstants(false)
-			, bFilterDuplicates(false)
-			, bFilterByScriptUsage(false)
-			, TargetScriptUsage(ENiagaraScriptUsage::Function)
-		{
-		}
-
-		/** Whether or not to sort the nodes, defaults to false. */
-		bool bSort;
-		/** Whether or not to include parameters, defaults to true. */
-		bool bIncludeParameters;
-		/** Whether or not to include attributes, defaults to true. */
-		bool bIncludeAttributes;
-		/** Whether or not to include system parameters, defaults to true. */
-		bool bIncludeSystemConstants;
-		/** Whether or not to include translator parameters, defaults to false. */
-		bool bIncludeTranslatorConstants;
-		/** Whether of not to filter out duplicate nodes, defaults to false. */
-		bool bFilterDuplicates;
-		/** Whether or not to limit to nodes connected to an output node of the specified script type.*/
-		bool bFilterByScriptUsage;
-		/** The specified script usage required for an input.*/
-		ENiagaraScriptUsage TargetScriptUsage;
-		/** The specified id within the graph of the script usage*/
-		FGuid TargetScriptUsageId;
-	};
+	using FFindInputNodeOptions = FNiagaraFindInputNodeOptions;
 
 	/** Finds input nodes in the graph with. */
 	void FindInputNodes(TArray<class UNiagaraNodeInput*>& OutInputNodes, FFindInputNodeOptions Options = FFindInputNodeOptions()) const;
diff --git a/Engine/Plugins/FX/Niagara/Source/NiagaraEditor/Public/NiagaraNode.h b/Engine/Plugins/FX/Niagara/Source/NiagaraEditor/Public/NiagaraNode.h
--- a/Engine/Plugins/FX/Niagara/Source/NiagaraEditor/Public/NiagaraNode.h
+++ b/Engine/Plugins/FX/Niagara/Source/NiagaraEditor/Public/NiagaraNode.h
@@ -260,6 +260,9 @@ protected:
 		}
 	}
 
+	virtual void GetCompilationInputPins(FPinCollectorArray& InputPins) const;
+	virtual void GetCompilationOutputPins(FPinCollectorArray& OutputPins) const;
+
 	/** The current change identifier for this node. Used to sync status with UNiagaraScripts.*/
 	UPROPERTY()
 	FGuid ChangeId;
diff --git a/Engine/Plugins/FX/Niagara/Source/NiagaraEditor/Public/NiagaraNodeWithDynamicPins.h b/Engine/Plugins/FX/Niagara/Source/NiagaraEditor/Public/NiagaraNodeWithDynamicPins.h
--- a/Engine/Plugins/FX/Niagara/Source/NiagaraEditor/Public/NiagaraNodeWithDynamicPins.h
+++ b/Engine/Plugins/FX/Niagara/Source/NiagaraEditor/Public/NiagaraNodeWithDynamicPins.h
@@ -92,6 +92,9 @@ protected:
 
 	virtual bool IsValidPinToCompile(UEdGraphPin* Pin) const override;
 
+	virtual void GetCompilationInputPins(FPinCollectorArray& InputPins) const override;
+	virtual void GetCompilationOutputPins(FPinCollectorArray& OutputPins) const override;
+
 private:
 
 	/** Gets the display text for a pin. */
diff --git a/Engine/Plugins/FX/Niagara/Source/NiagaraEditor/Public/NiagaraParameterMapHistory.h b/Engine/Plugins/FX/Niagara/Source/NiagaraEditor/Public/NiagaraParameterMapHistory.h
--- a/Engine/Plugins/FX/Niagara/Source/NiagaraEditor/Public/NiagaraParameterMapHistory.h
+++ b/Engine/Plugins/FX/Niagara/Source/NiagaraEditor/Public/NiagaraParameterMapHistory.h
@@ -4,18 +4,20 @@
 
 #include "CoreMinimal.h"
 #include "NiagaraCommon.h"
+#include "NiagaraCompilationBridge.h"
 #include "NiagaraEmitter.h"
+#include "NiagaraParameterMapHistoryFwd.h"
+#include "NiagaraScript.h"
 #include "Templates/Tuple.h"
 #include "ViewModels/Stack/NiagaraStackGraphUtilities.h"
-#include "NiagaraScript.h"
 
-class UNiagaraNodeOutput;
 class UEdGraphPin;
 class UEdGraphNode;
 class UNiagaraGraph;
-class UNiagaraParameterCollection;
 class FNiagaraHlslTranslator;
-class UNiagaraGraph;
+class UNiagaraNodeEmitter;
+class UNiagaraNodeOutput;
+class UNiagaraParameterCollection;
 
 class FCompileConstantResolver
 {
@@ -50,25 +52,26 @@ private:
 	ENiagaraFunctionDebugState DebugState;
 };
 
-struct FModuleScopedPin
+template<typename PinType>
+struct TModuleScopedPin
 {
-	FModuleScopedPin() = default;
-	FModuleScopedPin(const UEdGraphPin* InPin, FName InName = NAME_None)
+	TModuleScopedPin() = default;
+	TModuleScopedPin(const PinType* InPin, FName InName = NAME_None)
 	: Pin(InPin)
 	, ModuleName(InName)
 	{}
 
-	bool operator==(const FModuleScopedPin& Rhs) const
+	bool operator==(const TModuleScopedPin& Rhs) const
 	{
 		return Pin == Rhs.Pin && ModuleName == Rhs.ModuleName;
 	}
 
-	const UEdGraphPin* Pin = nullptr;
+	const PinType* Pin = nullptr;
 	FName ModuleName = NAME_None;
 
-	friend uint32 GetTypeHash(const FModuleScopedPin& ScopedPin)
+	friend uint32 GetTypeHash(const TModuleScopedPin& ScopedPin)
 	{
-		return GetTypeHash(TTuple<const UEdGraphPin*, FName>(ScopedPin.Pin, ScopedPin.ModuleName));
+		return GetTypeHash(TTuple<const PinType*, FName>(ScopedPin.Pin, ScopedPin.ModuleName));
 	}
 };
 
@@ -133,14 +136,123 @@ struct FNiagaraStaticVariableSearchContext
 	TMap<FGraphKey, bool> CachedResults;
 };
 
+namespace FNiagaraParameterUtilities
+{
+	/** Get the first namespace entry for this variable. Optionally includes the trailing period.*/
+	FString GetNamespace(const FNiagaraVariable& InVar, bool bIncludeDelimiter = true);
+
+
+	FName ResolveEmitterAlias(const FName& InName, const FString& InAlias);
+
+
+	/**
+	* Remove the Particles namespace if it exists.
+	*/
+	FNiagaraVariable ResolveAsBasicAttribute(const FNiagaraVariable& InVar, bool bSanitizeInput = true);
+
+
+	/**
+	* Reverses ResolveAsBasicAttribute.
+	*/
+	FNiagaraVariable BasicAttributeToNamespacedAttribute(const FNiagaraVariable& InVar, bool bSanitizeInput = true);
+
+
+	/** Prepends the namespace string to the variable name.*/
+	FNiagaraVariable VariableToNamespacedVariable(const FNiagaraVariable& InVar, FString Namespace);
+
+	/**
+	* Does this parameter start with the "Module" namespace? Note that the emitter namespace is an alias
+	* for all non-funcion/module script types and will be specialized to the function call node's name using the module.
+	*/
+	bool IsAliasedModuleParameter(const FNiagaraVariable& InVar);
+
+	/**
+	* Does this parameter start with the "Emitter" namespace? Note that the emitter namespace is an alias
+	* for Emitter and System script types and will be specialized to the name of that specific emitter.
+	*/
+	bool IsAliasedEmitterParameter(const FNiagaraVariable& InVar);
+	bool IsAliasedEmitterParameter(const FString& InVarName);
+	/** Is this parameter in the special "System" namespace?*/
+	bool IsSystemParameter(const FNiagaraVariable& InVar);
+	/** Is this parameter in the special "Engine" namespace?*/
+	bool IsEngineParameter(const FNiagaraVariable& InVar);
+	/*** Is per instance engine parameter. */
+	bool IsPerInstanceEngineParameter(const FNiagaraVariable& InVar, const FString& EmitterAlias);
+	bool IsUserParameter(const FNiagaraVariable& InVar);
+	bool IsRapidIterationParameter(const FNiagaraVariable& InVar);
+	bool SplitRapidIterationParameterName(const FNiagaraVariable& InVar, ENiagaraScriptUsage InUsage, FString& EmitterName, FString& FunctionCallName, FString& InputName);
+
+	/** Take an input string and make it hlsl safe.*/
+	FString MakeSafeNamespaceString(const FString& InStr);
+
+	/** Does the variable start with this namespace?*/
+	bool IsInNamespace(const FNiagaraVariableBase& InVar, const FString& Namespace);
+
+	void GetValidNamespacesForReading(const UNiagaraScript* InScript, TArray<FString>& OutputNamespaces);
+	void GetValidNamespacesForReading(ENiagaraScriptUsage InScriptUsage, int32 InUsageBitmask, TArray<FString>& OutputNamespaces);
+	bool IsValidNamespaceForReading(ENiagaraScriptUsage InScriptUsage, int32 InUsageBitmask, FString Namespace);
+	bool IsWrittenToScriptUsage(const FNiagaraVariable& InVar, ENiagaraScriptUsage InUsage, bool bAllowDataInterfaces);
+
+	/** Are we required to export this variable as an external constant?*/
+	bool IsExportableExternalConstant(const FNiagaraVariable& InVar, const UNiagaraScript* InScript);
+
+	/** Does this variable belong in a namespace that needs to come in as an external constant to this script?*/
+	bool IsExternalConstantNamespace(const FNiagaraVariable& InVar, const UNiagaraScript* InScript, const FGuid& VersionGuid);
+	bool IsExternalConstantNamespace(const FNiagaraVariable& InVar, ENiagaraScriptUsage InUsage, int32 InUsageBitmask);
+
+	/** Take a non-namespaced variable and move it to an appropriate external constant namespace for this script type.*/
+	FNiagaraVariable MoveToExternalConstantNamespaceVariable(const FNiagaraVariable& InVar, const UNiagaraScript* InScript);
+	FNiagaraVariable MoveToExternalConstantNamespaceVariable(const FNiagaraVariable& InVar, ENiagaraScriptUsage InUsage);
+	/**
+	* Does this parameter start with the "Particles" namespace?
+	*/
+	bool IsAttribute(const FNiagaraVariableBase& InVar);
+
+	/** Does this parameter contain the "Initial" namespace as one of its intermediate namespaces?*/
+	bool IsInitialValue(const FNiagaraVariableBase& InVar);
+
+	/** Does this variable name contain the "Initial" namespace as one of its intermediate namespaces? */
+	bool IsInitialName(const FName& InVariableName);
+
+	/** Does this parameter contain the "Previous" namespace as one of its intermediate namespaces?*/
+	bool IsPreviousValue(const FNiagaraVariableBase& InVar);
+
+	/** Does this variable name contain the "Previous" namespace as one of its intermediate namespaces? */
+	bool IsPreviousName(const FName& InVariableName);
+
+	/** Does this parameter contain the "Initial" namespace as one of its intermediate namespaces? If so, remove the "Initial" namespace and return the original value.*/
+	FNiagaraVariable GetSourceForInitialValue(const FNiagaraVariable& InVar);
+
+	/** Does this parameter contain the "Initial" namespace as one of its intermediate namespaces? If so, remove the "Initial" namespace and return the original value.*/
+	FName GetSourceForInitialValue(const FName& InVariableName);
+
+	/** Does this parameter contain the "Previous" namespace as one of its intermediate namespaces? If so, remove the "Previous" namespace and return the original value.*/
+	FNiagaraVariable GetSourceForPreviousValue(const FNiagaraVariable& InVar);
+
+	/** Does this parameter contain the "Previous" namespace as one of its intermediate namespaces? If so, remove the "Previous" namespace and return the original value.*/
+	FName GetSourceForPreviousValue(const FName& InVariableName);
+
+	FNiagaraVariable ConvertVariableToRapidIterationConstantName(FNiagaraVariable InVar, const TCHAR* InEmitterName, ENiagaraScriptUsage InUsage);
+};
+
 /** Traverses a Niagara node graph to identify the variables that have been written and read from a parameter map. 
 * 	This class is meant to aid in UI and compilation of the graph. There are several main script types and each one interacts
 *	slightly differently with the history depending on context.
 */
-struct FNiagaraParameterMapHistory
+template<typename GraphBridge>
+struct TNiagaraParameterMapHistory
 {
 public:
-	FNiagaraParameterMapHistory();
+	using FPin = typename GraphBridge::FPin;
+	using FNode = typename GraphBridge::FNode;
+	using FOutputNode = typename GraphBridge::FOutputNode;
+	using FInputPin = typename GraphBridge::FInputPin;
+	using FParameterCollection = typename GraphBridge::FParameterCollection;
+	using FParameterCollectionStore = typename GraphBridge::FParameterCollectionStore;
+
+	using FModuleScopedPin = TModuleScopedPin<FPin>;
+
+	TNiagaraParameterMapHistory();
 
 	ENiagaraScriptUsage OriginatingScriptUsage;
 	FGuid UsageGuid;
@@ -157,34 +269,29 @@ public:
 
 	TArray<FNiagaraVariable> VariablesWithOriginalAliasesIntact;
 
-
 	/** Used parameter collections identified during the traversal. TODO: Need to ensure these cannot be GCd if the asset is deleted while it's being used in an in flight compilation. */
-	TArray<UNiagaraParameterCollection*> ParameterCollections;
-	/** Cached off contents of used parameter collections, in case they change during threaded compilation. */
-	TArray<TArray<FNiagaraVariable>> ParameterCollectionVariables;
-	/** Cached off contents of used parameter collections, in case they change during threaded compilation. */
-	TArray<FString> ParameterCollectionNamespaces;
+	FParameterCollectionStore EncounteredParameterCollections;
 	
 	/** Are there any warnings that were encountered during the traversal of the graph for a given variable? */
 	TArray<FString> PerVariableWarnings;
 
 	/** For each variable that was found, identify the pins that wrote to them in order from first to last write.*/
-	TArray<TArray<FModuleScopedPin> > PerVariableWriteHistory;
+	TArray<TArray<TModuleScopedPin<FPin>>> PerVariableWriteHistory;
 
 	/** For each variable that was found, identify the pins that read them from the map in order from first to last read. First of the pair has the read pin, second of the pair has the last set that wrote to the pin.*/
 	struct FReadHistory
 	{
-		FModuleScopedPin ReadPin;
-		FModuleScopedPin PreviousWritePin;
+		TModuleScopedPin<FPin> ReadPin;
+		TModuleScopedPin<FPin> PreviousWritePin;
 	};
 
 	TArray<TArray<FReadHistory>> PerVariableReadHistory; 
 
 	/** List of pins that manipulated the parameter map from input to output. */
-	TArray<const UEdGraphPin*> MapPinHistory;
+	TArray<const FPin*> MapPinHistory;
 
 	/** List of nodes that manipulated the parameter map from input to output.*/
-	TArray<const class UNiagaraNode*> MapNodeVisitations;
+	TArray<const FNode*> MapNodeVisitations;
 
 	/** For each node in MapNodeVisitations, record the start index and end index of variables added within the body of the node, i.e. a Get node will record just the values it pulls out directly. A function call, however, will record all sub-graph traversals.*/
 	TArray<TTuple<uint32, uint32> > MapNodeVariableMetaData;
@@ -200,18 +307,18 @@ public:
 
 	TMap<FGraphTraversalHandle, FString> PinToConstantValues;
 	bool bParameterCollectionsSkipped = false;
-	void RegisterConstantPin(const FGraphTraversalHandle& InTraversalPath, const UEdGraphPin* InPin, const FString& InValue);
+	void RegisterConstantPin(const FGraphTraversalHandle& InTraversalPath, const FPin* InPin, const FString& InValue);
 
 	bool IsVariableFromCustomIterationNamespaceOverride(const FNiagaraVariable& InVar) const;
 	
 	/**
 	* Called in a depth-first traversal to identify a given Niagara Parameter Map pin that was touched during traversal.
 	*/
-	int32 RegisterParameterMapPin(const UEdGraphPin* Pin);
+	int32 RegisterParameterMapPin(const FPin* Pin);
 
 	void RegisterConstantVariableWrite(const FString& InValue, int32 InVarIdx, bool bIsSettingDefault, bool bIsLinkNotValue);
 
-	uint32 BeginNodeVisitation(const UNiagaraNode* Node);
+	uint32 BeginNodeVisitation(const FNode* Node);
 	void EndNodeVisitation(uint32 IndexFromBeginNode);
 
 	/**
@@ -219,8 +326,6 @@ public:
 	*/
 	int32 FindVariableByName(const FName& VariableName, bool bAllowPartialMatch = false) const;
 
-	static int32 SearchArrayForPartialNameMatch(const TArray<FNiagaraVariable>& Variables, const FName& VariableName);
-	
 	/**
 	* Find a variable by both name and type. 
 	*/
@@ -235,168 +340,83 @@ public:
 	/**
 	* Get the pin that added the parameter map to the graph.
 	*/
-	const UEdGraphPin* GetOriginalPin() const;
+	const FPin* GetOriginalPin() const;
 
 	/**
 	* Get the output pin that we traced to build this history object.
 	*/
-	const UEdGraphPin* GetFinalPin() const;
-
-
-
-	/** Get the first namespace entry for this variable. Optionally includes the trailing period.*/
-	static FString GetNamespace(const FNiagaraVariable& InVar, bool bIncludeDelimiter = true);
-
-
-	static FName ResolveEmitterAlias(const FName& InName, const FString& InAlias);
-
-	
-	/**
-	* Remove the Particles namespace if it exists.
-	*/
-	static FNiagaraVariable ResolveAsBasicAttribute(const FNiagaraVariable& InVar, bool bSanitizeInput = true);
-
-
-	/**
-	* Reverses ResolveAsBasicAttribute.
-	*/
-	static FNiagaraVariable BasicAttributeToNamespacedAttribute(const FNiagaraVariable& InVar, bool bSanitizeInput = true);
-
-	
-	/** Prepends the namespace string to the variable name.*/
-	static FNiagaraVariable VariableToNamespacedVariable(const FNiagaraVariable& InVar, FString Namespace);
-		
-	/**
-	* Does this parameter start with the "Module" namespace? Note that the emitter namespace is an alias
-	* for all non-funcion/module script types and will be specialized to the function call node's name using the module.
-	*/
-	static bool IsAliasedModuleParameter(const FNiagaraVariable& InVar);
+	const FPin* GetFinalPin() const;
 
-	/**
-	* Does this parameter start with the "Emitter" namespace? Note that the emitter namespace is an alias
-	* for Emitter and System script types and will be specialized to the name of that specific emitter.
-	*/
-	static bool IsAliasedEmitterParameter(const FNiagaraVariable& InVar);
-	static bool IsAliasedEmitterParameter(const FString& InVarName);
-	/** Is this parameter in the special "System" namespace?*/
-	static bool IsSystemParameter(const FNiagaraVariable& InVar);
-	/** Is this parameter in the special "Engine" namespace?*/
-	static bool IsEngineParameter(const FNiagaraVariable& InVar);
-	/*** Is per instance engine parameter. */
-	static bool IsPerInstanceEngineParameter(const FNiagaraVariable& InVar, const FString& EmitterAlias);
-	static bool IsUserParameter(const FNiagaraVariable& InVar);
-	static bool IsRapidIterationParameter(const FNiagaraVariable& InVar);
-	static bool SplitRapidIterationParameterName(const FNiagaraVariable& InVar, ENiagaraScriptUsage InUsage, FString& EmitterName, FString& FunctionCallName, FString& InputName);
-	
-	/** Take an input string and make it hlsl safe.*/
-	static FString MakeSafeNamespaceString(const FString& InStr);
-
-	/** Does the variable start with this namespace?*/
-	static bool IsInNamespace(const FNiagaraVariableBase& InVar, const FString& Namespace);
-
-	static void GetValidNamespacesForReading(const UNiagaraScript* InScript, TArray<FString>& OutputNamespaces);
-	static void GetValidNamespacesForReading(ENiagaraScriptUsage InScriptUsage, int32 InUsageBitmask, TArray<FString>& OutputNamespaces);
-	static bool IsValidNamespaceForReading(ENiagaraScriptUsage InScriptUsage, int32 InUsageBitmask, FString Namespace);
 
 	/** Called to determine if a given variable should be output from a script. It is not static as it requires the overall context to include emitter namespaces visited for system scripts.*/
 	bool IsPrimaryDataSetOutput(const FNiagaraVariable& InVar, const UNiagaraScript* InScript,   bool bAllowDataInterfaces = false,  bool bAllowStatics = false) const;
 	bool IsPrimaryDataSetOutput(const FNiagaraVariable& InVar, ENiagaraScriptUsage InUsage, bool bAllowDataInterfaces = false, bool bAllowStatics = false) const;
-	static bool IsWrittenToScriptUsage(const FNiagaraVariable& InVar, ENiagaraScriptUsage InUsage, bool bAllowDataInterfaces);
-
-	/** Are we required to export this variable as an external constant?*/
-	bool IsExportableExternalConstant(const FNiagaraVariable& InVar, const UNiagaraScript* InScript);
-
-	/** Does this variable belong in a namespace that needs to come in as an external constant to this script?*/
-	static bool IsExternalConstantNamespace(const FNiagaraVariable& InVar, const UNiagaraScript* InScript, const FGuid& VersionGuid);
-	static bool IsExternalConstantNamespace(const FNiagaraVariable& InVar, ENiagaraScriptUsage InUsage, int32 InUsageBitmask);
-	
-	/** Take a non-namespaced variable and move it to an appropriate external constant namespace for this script type.*/
-	static FNiagaraVariable MoveToExternalConstantNamespaceVariable(const FNiagaraVariable& InVar, const UNiagaraScript* InScript);
-	static FNiagaraVariable MoveToExternalConstantNamespaceVariable(const FNiagaraVariable& InVar, ENiagaraScriptUsage InUsage);	
-	/**
-	* Does this parameter start with the "Particles" namespace?
-	*/
-	static bool IsAttribute(const FNiagaraVariableBase& InVar);
-
-	/** Does this parameter contain the "Initial" namespace as one of its intermediate namespaces?*/
-	static bool IsInitialValue(const FNiagaraVariableBase& InVar);
-
-	/** Does this variable name contain the "Initial" namespace as one of its intermediate namespaces? */
-	static bool IsInitialName(const FName& InVariableName);
-
-	/** Does this parameter contain the "Previous" namespace as one of its intermediate namespaces?*/
-	static bool IsPreviousValue(const FNiagaraVariableBase& InVar);
-
-	/** Does this variable name contain the "Previous" namespace as one of its intermediate namespaces? */
-	static bool IsPreviousName(const FName& InVariableName);
 
 	/** Get the output node associated with this graph.*/
-	const UNiagaraNodeOutput* GetFinalOutputNode() const;
-
-	/** Does this parameter contain the "Initial" namespace as one of its intermediate namespaces? If so, remove the "Initial" namespace and return the original value.*/
-	static FNiagaraVariable GetSourceForInitialValue(const FNiagaraVariable& InVar);
-
-	/** Does this parameter contain the "Initial" namespace as one of its intermediate namespaces? If so, remove the "Initial" namespace and return the original value.*/
-	static FName GetSourceForInitialValue(const FName& InVariableName);
-
-	/** Does this parameter contain the "Previous" namespace as one of its intermediate namespaces? If so, remove the "Previous" namespace and return the original value.*/
-	static FNiagaraVariable GetSourceForPreviousValue(const FNiagaraVariable& InVar);
-
-	/** Does this parameter contain the "Previous" namespace as one of its intermediate namespaces? If so, remove the "Previous" namespace and return the original value.*/
-	static FName GetSourceForPreviousValue(const FName& InVariableName);
+	const FOutputNode* GetFinalOutputNode() const;
 
 	/**
 	* Helper to add a variable to the known list for a parameter map.
 	*/
-	int32 AddVariable(const FNiagaraVariable& InVar, const FNiagaraVariable& InAliasedVar, FName ModuleName, const UEdGraphPin* InPin, TOptional<FNiagaraVariableMetaData> InMetaData = TOptional<FNiagaraVariableMetaData>());
+	int32 AddVariable(const FNiagaraVariable& InVar, const FNiagaraVariable& InAliasedVar, FName ModuleName, const FPin* InPin, TOptional<FNiagaraVariableMetaData> InMetaData = TOptional<FNiagaraVariableMetaData>());
 
 	/** Get the default value for this variable.*/
-	const UEdGraphPin* GetDefaultValuePin(int32 VarIdx) const;
-
-	static FNiagaraVariable ConvertVariableToRapidIterationConstantName(FNiagaraVariable InVar, const TCHAR* InEmitterName, ENiagaraScriptUsage InUsage);
+	const FInputPin* GetDefaultValuePin(int32 VarIdx) const;
 
 	/**
 	If this is variable is a parameter in one of our tracked collections, return it.
 	@param InVar	Variable to test.
 	@param bMissingParameter	bool set to mark if this parameter was a collection parameter but is now missing from it's collection.
 	*/
-	UNiagaraParameterCollection* IsParameterCollectionParameter(FNiagaraVariable& InVar, bool& bMissingParameter);
+	FParameterCollection IsParameterCollectionParameter(FNiagaraVariable& InVar, bool& bMissingParameter);
 
 	bool ShouldIgnoreVariableDefault(const FNiagaraVariable& Var)const;
 };
 
-class FNiagaraParameterMapHistoryBuilder
+template<typename GraphBridge>
+class TNiagaraParameterMapHistoryBuilder : public GraphBridge::FBuilderExtraData
 {
 public:
+	using FHistory = TNiagaraParameterMapHistory<GraphBridge>;
+	using FOutputNode = typename GraphBridge::FOutputNode;
+	using FNode = typename GraphBridge::FNode;
+	using FPin = typename GraphBridge::FPin;
+	using FInputPin = typename GraphBridge::FInputPin;
+	using FOutputPin = typename GraphBridge::FOutputPin;
+	using FGraph = typename GraphBridge::FGraph;
+	using FEmitterNode = typename GraphBridge::FEmitterNode;
+	using FConstantResolver = typename GraphBridge::FConstantResolver;
+	using FParameterCollection = typename GraphBridge::FParameterCollection;
+
 	/** Collection of the build histories from the graph traversal.*/
-	TArray<FNiagaraParameterMapHistory> Histories;
+	TArray<FHistory> Histories;
 
 	/** Constructor*/
-	FNiagaraParameterMapHistoryBuilder();
+	TNiagaraParameterMapHistoryBuilder();
 
-	virtual ~FNiagaraParameterMapHistoryBuilder() {};
+	virtual ~TNiagaraParameterMapHistoryBuilder() = default;
 
 	/** Add a new parameter map to the array.*/
 	int32 CreateParameterMap();
 
 	/** Called in order to set up the correct initial context for an Output node and invokes the UNiagaraNodeOutput's BuildParameterMapHistory method.*/
-	void BuildParameterMaps(const UNiagaraNodeOutput* OutputNode, bool bRecursive = true);
+	void BuildParameterMaps(const FOutputNode* OutputNode, bool bRecursive = true);
 
 	/**
 	* Called first during a node's visitation during traversal to identify that a node has been visited.
 	*/
-	void RegisterNodeVisitation(const UEdGraphNode* Node);
+	void RegisterNodeVisitation(const FNode* Node);
 
 	/** Important. Must be called for each routing of the parameter map. This feeds the list used by TraceParameterMapOutputPin.*/
-	int32 RegisterParameterMapPin(int32 WhichParameterMap, const UEdGraphPin* Pin);
+	int32 RegisterParameterMapPin(int32 WhichParameterMap, const FPin* Pin);
 
-	int32 RegisterConstantPin(int32 WhichConstant, const UEdGraphPin* Pin);
+	int32 RegisterConstantPin(int32 WhichConstant, const FPin* Pin);
 	int32 RegisterConstantVariableWrite(int32 WhichParamMapIdx, int32 WhichConstant, int32 WhichVarIdx, bool bIsSettingDefault, bool bLinkNotValue);
 
-	int32  RegisterConstantFromInputPin(const UEdGraphPin* InputPin);
+	int32 RegisterConstantFromInputPin(const FInputPin* InputPin, const FString& PinDefaultValue);
 
-	int32 GetConstantFromInputPin(const UEdGraphPin* InputPin) const;
-	int32 GetConstantFromOutputPin(const UEdGraphPin* OutputPin) const;
+	int32 GetConstantFromInputPin(const FInputPin* InputPin) const;
+	int32 GetConstantFromOutputPin(const FOutputPin* OutputPin) const;
 	int32 GetConstantFromVariableRead(const FNiagaraVariableBase& InVar);
 	int32 AddOrGetConstantFromValue(const FString& Value);
 	FString GetConstant(int32 InIndex);
@@ -404,11 +424,11 @@ public:
 	/** Records a write to a DataSet in the appropriate parameter map history */
 	void RegisterDataSetWrite(int32 WhichParameterMap, const FNiagaraDataSetID& DataSet);
 
-	uint32 BeginNodeVisitation(int32 WhichParameterMap, const class UNiagaraNode* Node);
+	uint32 BeginNodeVisitation(int32 WhichParameterMap, const FNode* Node);
 	void EndNodeVisitation(int32 WhichParameterMap, uint32 IndexFromBeginNode);
 
 	/** Trace back a pin to whom it was connected to to find the current parameter map to use.*/
-	int32 TraceParameterMapOutputPin(const UEdGraphPin* OutputPin);
+	int32 TraceParameterMapOutputPin(const FOutputPin* OutputPin);
 
 	void BeginTranslation(const UNiagaraScript* Script);
 
@@ -430,22 +450,22 @@ public:
 	/**
 	* Record that we have entered a new function scope.
 	*/
-	void EnterFunction(const FString& InNodeName, const class UNiagaraGraph* InGraph, const class UNiagaraNode* Node);
+	void EnterFunction(const FString& InNodeName, const FGraph* InGraph, const FNode* Node);
 
 	/**
 	* Record that we have exited a function scope.
 	*/
-	void ExitFunction(const FString& InNodeName, const class UNiagaraNode* Node);
+	void ExitFunction(const FNode* Node);
 
 	/**
 	* Record that we have entered an emitter scope.
 	*/
-	void EnterEmitter(const FString& InEmitterName, const class UNiagaraGraph* InGraph, const class UNiagaraNode* Node);
+	void EnterEmitter(const FString& InEmitterName, const FGraph* InGraph, const FNode* Node);
 
 	/**
 	* Record that we have exited an emitter scope.
 	*/
-	void ExitEmitter(const FString& InEmitterName, const class UNiagaraNode* Node);
+	void ExitEmitter(const FString& InEmitterName, const FNode* Node);
 
 	/**
 	* Use the current alias map to resolve any aliases in this input variable name.
@@ -455,23 +475,23 @@ public:
 	/**
 	* Has RegisterNodeVisitation been called yet on the owning node of this pin?
 	*/
-	bool GetPinPreviouslyVisited(const UEdGraphPin* InPin) const;
+	bool GetPinPreviouslyVisited(const FPin* InPin) const;
 
 	/**
 	* Has RegisterNodeVisitation been called on the input node yet?
 	*/
-	bool GetNodePreviouslyVisited(const class UNiagaraNode* Node) const;
+	bool GetNodePreviouslyVisited(const FNode* Node) const;
 
 	/** If we haven't already visited the owning nodes, do so.*/
-	void VisitInputPins(const class UNiagaraNode*, bool bFilterForCompilation);
+	void VisitInputPins(const FNode*, bool bFilterForCompilation);
 	
 	/** If we haven't already visited the owning node, do so.*/
-	void VisitInputPin(const UEdGraphPin* Pin, const class UNiagaraNode*, bool bFilterForCompilation);
+	void VisitInputPin(const FInputPin* Pin, bool bFilterForCompilation);
 
 	/**
 	* Record that a pin writes to the parameter map. The pin name is expected to be the namespaced parameter map version of the name. If any aliases are in place, they are removed.
 	*/
-	int32 HandleVariableWrite(int32 ParameterMapIndex, const UEdGraphPin* InPin);
+	int32 HandleVariableWrite(int32 ParameterMapIndex, const FInputPin* InPin);
 	
 	/**
 	* Record that a variable write to the parameter map. The var name is expected to be the namespaced parameter map version of the name. If any aliases are in place, they are removed.
@@ -481,14 +501,17 @@ public:
 	/**
 	* Record that a pin reads from the parameter map. The pin name is expected to be the namespaced parameter map version of the name. If any aliases are in place, they are removed.
 	*/
-	int32 HandleVariableRead(int32 ParameterMapIndex, const UEdGraphPin* InPin, bool RegisterReadsAsVariables, const UEdGraphPin* InDefaultPin, bool bFilterForCompilation, bool& OutUsedDefault);
+	int32 HandleVariableRead(int32 ParameterMapIndex, const FPin* InPin, bool RegisterReadsAsVariables, const FInputPin* InDefaultPin, bool bFilterForCompilation, bool& OutUsedDefault);
 
 	int32 HandleExternalVariableRead(int32 ParamMapIdx, const FName& InVarName);
 
+	void HandleParameterCollectionRead(const FNiagaraVariableBase& InVar, int32 ParamMapIndex);
+
+
 	/**
 	* Virtual wrapper to call correct AddVariable method on the target ParameterMapHistory.
 	*/
-	virtual int32 AddVariableToHistory(FNiagaraParameterMapHistory& History, const FNiagaraVariable& InVar, const FNiagaraVariable& InAliasedVar, const UEdGraphPin* InPin);
+	virtual int32 AddVariableToHistory(FHistory& History, const FNiagaraVariable& InVar, const FNiagaraVariable& InAliasedVar, const FPin* InPin);
 	
 	/**
 	* Get the string that the "Module" namespace maps to currently (if it exists)
@@ -501,7 +524,7 @@ public:
 	const FString* GetEmitterAlias() const;
 
 	/** Get the node calling this sub-graph.*/
-	const UNiagaraNode* GetCallingContext() const;
+	const FNode* GetCallingContext() const;
 
 	/** Are we currently in a top-level function call context (ie a node in the main graph or an argument into a function for the main graph.)*/
 	bool InTopLevelFunctionCall(ENiagaraScriptUsage InFilterScriptType) const;
@@ -537,7 +560,7 @@ public:
 
 	void RegisterExternalStaticVariables(TConstArrayView<FNiagaraVariable> Variables);
 
-	FCompileConstantResolver ConstantResolver;
+	TPimplPtr<FConstantResolver, EPimplPtrMode::DeepCopy> ConstantResolver;
 
 	/** An optional emitter handle guid you can register. This will skip parameter map building of all emitters that don't match this one, if set.
 	 * Reason for this is to be able to build parameter map history for a system, while excluding certain emitters.
@@ -550,14 +573,14 @@ public:
 	int32 MaxGraphDepthTraversal = INDEX_NONE;
 	int32 CurrentGraphDepth = 0;
 
-	void SetConstantByStaticVariable(int32& OutValue, const UEdGraphPin* InDefaultPin);
+	void SetConstantByStaticVariable(int32& OutValue, const FInputPin* InDefaultPin);
 	void SetConstantByStaticVariable(int32& OutValue, const FNiagaraVariable& Var);
 	int32 FindStaticVariable(const FNiagaraVariable& Var) const;
 
 	TArray<FNiagaraVariable> StaticVariables;
 	TArray<bool> StaticVariableExportable; // Should we export out these static variables to calling context?
 
-	void GetContextuallyVisitedNodes(TArray<const class UNiagaraNode*>& OutVistedNodes)
+	void GetContextuallyVisitedNodes(TArray<const FNode*>& OutVistedNodes)
 	{
 		if (ContextuallyVisitedNodes.Num() > 0)
 			OutVistedNodes.Append(ContextuallyVisitedNodes.Last());
@@ -568,7 +591,7 @@ public:
 		StaticVariableSearchContext.bCollectingStaticVariablesOnly = true;
 	}
 
-	bool ShouldProcessDepthTraversal(const UNiagaraGraph* Graph);
+	bool ShouldProcessDepthTraversal(const FGraph* Graph);
 
 protected:
 	/** Internal helper function to decide whether or not we should use this static variable when merging with other
@@ -586,20 +609,22 @@ protected:
 	/** Helper method used to take in input script type and determine if the passed in namespaced variable is worth tracking.*/
 	bool IsNamespacedVariableRelevantToScriptType(const FNiagaraVariable& InVar, ENiagaraScriptUsage ScriptType);
 
+	bool AddParameterCollection(FParameterCollection Collection, int32 ParamMapIndex);
+
 	/** Contains the hierarchy of nodes leading to the current graph being processed. Usually made up of FunctionCall and Emitter nodes.*/
-	TArray<const UNiagaraNode*> CallingContext;
+	TArray<const FNode*> CallingContext;
 
 	/** Contains the hierarchy of graphs leading to and including the current graph being processed. Is not in sync with CallingContext as there is an additional 0th entry for the NodeGraphDeepCopy. */
-	TArray<const UNiagaraGraph*> CallingGraphContext;
+	TArray<const FGraph*> CallingGraphContext;
 
-	TArray< TMap<const UEdGraphPin*, int32> > PinToConstantIndices;
+	TArray< TMap<const FPin*, int32> > PinToConstantIndices;
 	TArray<FString> Constants;
 	TArray< TMap<FNiagaraVariableBase, int32> > VariableToConstantIndices;
 
 	/** Tracker for each context level of the parameter map index associated with a given pin. Used to trace parameter maps through the graph.*/
-	TArray<TMap<const UEdGraphPin*, int32> > PinToParameterMapIndices;
+	TArray<TMap<const FPin*, int32> > PinToParameterMapIndices;
 	/** List of previously visited nodes per context. Note that the same node may be visited multiple times across all graph traversals, but only one time per context level.*/
-	TArray<TArray<const class UNiagaraNode*> > ContextuallyVisitedNodes;
+	TArray<TArray<const FNode*> > ContextuallyVisitedNodes;
 	/** Contains the hierarchy of emitter node names leading to the current graph being processed.*/
 	TArray<FName> EmitterNameContextStack;
 	/** Contains the hierarchy of function call node names leading to the current graph being processed.*/
@@ -634,16 +659,23 @@ protected:
 	FNiagaraStaticVariableSearchContext StaticVariableSearchContext;
 };
 
-class FNiagaraParameterMapHistoryWithMetaDataBuilder : public FNiagaraParameterMapHistoryBuilder
+template<typename GraphBridge>
+class TNiagaraParameterMapHistoryWithMetaDataBuilder : public TNiagaraParameterMapHistoryBuilder<GraphBridge>
 {
 public:
+	using FGraph = typename GraphBridge::FGraph;
+	using FPin = typename GraphBridge::FPin;
+	using FHistory = typename TNiagaraParameterMapHistoryBuilder<GraphBridge>::FHistory;
+	using TNiagaraParameterMapHistoryBuilder<GraphBridge>::GetModuleAlias;
+	using TNiagaraParameterMapHistoryBuilder<GraphBridge>::CallingGraphContext;
+
 	//** Add a graph to the calling graph context stack. Generally used to prime the context stack with the node graph deep copy during precompilation. */
-	void AddGraphToCallingGraphContextStack(const UNiagaraGraph* InGraph) { CallingGraphContext.Add(InGraph); };
+	void AddGraphToCallingGraphContextStack(const FGraph* InGraph);
 
 	/**
 	* Virtual wrapper to call correct AddVariable method on the target ParameterMapHistory.
 	*/
-	virtual int32 AddVariableToHistory(FNiagaraParameterMapHistory& History, const FNiagaraVariable& InVar, const FNiagaraVariable& InAliasedVar, const UEdGraphPin* InPin);
+	virtual int32 AddVariableToHistory(FHistory& History, const FNiagaraVariable& InVar, const FNiagaraVariable& InAliasedVar, const FPin* InPin) override;
 };
 
 struct FNiagaraGraphCachedBuiltHistory : public FNiagaraGraphCachedDataBase
diff --git a/Engine/Plugins/FX/Niagara/Source/NiagaraEditor/Public/NiagaraParameterMapHistoryFwd.h b/Engine/Plugins/FX/Niagara/Source/NiagaraEditor/Public/NiagaraParameterMapHistoryFwd.h
new file mode 100644
--- /dev/null
+++ b/Engine/Plugins/FX/Niagara/Source/NiagaraEditor/Public/NiagaraParameterMapHistoryFwd.h
@@ -0,0 +1,19 @@
+// Copyright Epic Games, Inc. All Rights Reserved.
+
+#pragma once
+
+struct FNiagaraCompilationGraphBridge;
+struct FNiagaraCompilationDigestBridge;
+
+template<typename GraphBridge>
+struct TNiagaraParameterMapHistory;
+
+template<typename GraphBridge>
+class TNiagaraParameterMapHistoryBuilder;
+
+template<typename GraphBridge>
+class TNiagaraParameterMapHistoryWithMetaDataBuilder;
+
+using FNiagaraParameterMapHistory = TNiagaraParameterMapHistory<FNiagaraCompilationGraphBridge>;
+using FNiagaraParameterMapHistoryBuilder = TNiagaraParameterMapHistoryBuilder<FNiagaraCompilationGraphBridge>;
+using FNiagaraParameterMapHistoryWithMetaDataBuilder = TNiagaraParameterMapHistoryWithMetaDataBuilder<FNiagaraCompilationGraphBridge>;
diff --git a/Engine/Plugins/FX/Niagara/Source/NiagaraEditor/Public/ViewModels/Stack/NiagaraStackGraphUtilities.h b/Engine/Plugins/FX/Niagara/Source/NiagaraEditor/Public/ViewModels/Stack/NiagaraStackGraphUtilities.h
--- a/Engine/Plugins/FX/Niagara/Source/NiagaraEditor/Public/ViewModels/Stack/NiagaraStackGraphUtilities.h
+++ b/Engine/Plugins/FX/Niagara/Source/NiagaraEditor/Public/ViewModels/Stack/NiagaraStackGraphUtilities.h
@@ -4,6 +4,7 @@
 
 #include "NiagaraTypes.h"
 #include "NiagaraCommon.h"
+#include "NiagaraParameterMapHistoryFwd.h"
 #include "ViewModels/Stack/NiagaraParameterHandle.h"
 #include "ViewModels/Stack/NiagaraStackEntry.h"
 #include "AssetRegistry/AssetData.h"
@@ -28,7 +29,6 @@ class UNiagaraStackEditorData;
 class UNiagaraStackErrorItem;
 class FCompileConstantResolver;
 class INiagaraMessage;
-class FNiagaraParameterMapHistoryBuilder;
 class UNiagaraHierarchyModuleInput;
 class UNiagaraSimulationStageBase;
 struct FNiagaraEventScriptProperties;
@@ -93,7 +93,7 @@ namespace FNiagaraStackGraphUtilities
 	FString GenerateStackModuleEditorDataKey(UNiagaraNodeFunctionCall& ModuleNode);
 
 	TArray<FName> StackContextResolution(FVersionedNiagaraEmitter OwningEmitter, UNiagaraNodeOutput* OutputNodeInChain);
-	void BuildParameterMapHistoryWithStackContextResolution(FVersionedNiagaraEmitter OwningEmitter, UNiagaraNodeOutput* OutputNodeInChain, UNiagaraNode* NodeToVisit, FNiagaraParameterMapHistoryBuilder& Builder, bool bRecursive = true, bool bFilterForCompilation = true);
+	void BuildParameterMapHistoryWithStackContextResolution(FVersionedNiagaraEmitter OwningEmitter, UNiagaraNodeOutput* OutputNodeInChain, UNiagaraNode* NodeToVisit, TArray<FNiagaraParameterMapHistory>& OutHistories, bool bRecursive = true, bool bFilterForCompilation = true);
 
 	struct FInputDataCollection
 	{
