- Fix for SubUVs & multiple UV sets in ES3.1

diff --git a/Engine/Plugins/FX/Niagara/Shaders/Private/NiagaraMeshVertexFactory.ush b/Engine/Plugins/FX/Niagara/Shaders/Private/NiagaraMeshVertexFactory.ush
--- a/Engine/Plugins/FX/Niagara/Shaders/Private/NiagaraMeshVertexFactory.ush
+++ b/Engine/Plugins/FX/Niagara/Shaders/Private/NiagaraMeshVertexFactory.ush
@@ -615,7 +615,13 @@ FVertexFactoryIntermediates GetVertexFactoryIntermediates(FVertexFactoryInput In
 		const uint NumFetchSubUvTexCoords = NiagaraMeshVF.VertexFetch_Parameters[VF_NumTexcoords_Index];
 		const float2 SubUvTexCoord = NiagaraMeshVF.VertexFetch_TexCoordBuffer[NumFetchSubUvTexCoords * Input.VertexId];
 	#else
-		const float2 SubUvTexCoord = Input.TexCoords0.xy;
+		#if NUM_MATERIAL_TEXCOORDS_VERTEX == 1
+			const float2 SubUvTexCoord = Input.TexCoords0;
+		#elif NUM_MATERIAL_TEXCOORDS_VERTEX >= 2
+			const float2 SubUvTexCoord = Input.TexCoords01.xy;
+		#else
+			const float2 SubUvTexCoord = float2(0.0f, 0.0f);
+		#endif
 	#endif // MANUAL_VERTEX_FETCH
 
 	Intermediates.SubUVCoords.xy = (float2(SubImageAH, SubImageAV) + SubUvTexCoord) * NiagaraMeshVF.SubImageSize.zw;
