Crash fix for null deref in Niagara Baker UI.

diff --git a/Engine/Plugins/FX/Niagara/Source/NiagaraEditor/Private/ViewModels/NiagaraBakerViewModel.cpp b/Engine/Plugins/FX/Niagara/Source/NiagaraEditor/Private/ViewModels/NiagaraBakerViewModel.cpp
--- a/Engine/Plugins/FX/Niagara/Source/NiagaraEditor/Private/ViewModels/NiagaraBakerViewModel.cpp
+++ b/Engine/Plugins/FX/Niagara/Source/NiagaraEditor/Private/ViewModels/NiagaraBakerViewModel.cpp
@@ -459,7 +459,10 @@ FText FNiagaraBakerViewModel::GetOutputText(int32 OutputIndex) const
 	UNiagaraBakerSettings* BakerSettings = GetBakerSettings();
 	if ( BakerSettings && BakerSettings->Outputs.IsValidIndex(OutputIndex) )
 	{
-		return FText::FromString(BakerSettings->Outputs[OutputIndex]->OutputName);
+		if(BakerSettings->Outputs[OutputIndex])
+		{
+			return FText::FromString(BakerSettings->Outputs[OutputIndex]->OutputName);
+		}
 	}
 	return FText::GetEmpty();
 }
