Fix for parameter map histories missing some static constant values
#rnx

diff --git a/Engine/Plugins/FX/Niagara/Source/NiagaraEditor/Private/NiagaraParameterMapHistory.cpp b/Engine/Plugins/FX/Niagara/Source/NiagaraEditor/Private/NiagaraParameterMapHistory.cpp
--- a/Engine/Plugins/FX/Niagara/Source/NiagaraEditor/Private/NiagaraParameterMapHistory.cpp
+++ b/Engine/Plugins/FX/Niagara/Source/NiagaraEditor/Private/NiagaraParameterMapHistory.cpp
@@ -247,7 +247,7 @@ void TNiagaraParameterMapHistory<GraphBridge>::RegisterConstantPin(const FGraphT
 	{
 		FGraphTraversalHandle Handle = InTraversalPath;
 		Handle.PushPin(InPin);
-		
+
 		FString* FoundValue = PinToConstantValues.Find(Handle);
 		if (FoundValue != nullptr)
 		{
@@ -1175,7 +1175,7 @@ int32 TNiagaraParameterMapHistoryBuilder<GraphBridge>::RegisterConstantFromInput
 		else
 		{
 			ensure(InputPin->Direction == EEdGraphPinDirection::EGPD_Input);
-			if (InputPin->PinName.IsNone())
+			if (!PinDefaultValue.IsEmpty() && !InputPin->PinName.IsNone())
 			{
 				ConstantIdx = AddOrGetConstantFromValue(PinDefaultValue);
 				RegisterConstantPin(ConstantIdx, InputPin);
