Focusing the scratch pad manager tab when creating a new scratch pad so you can immediately rename it
Made the toolkitmodebase public so that other parts in the plugin can access the tab names properly

diff --git a/Engine/Plugins/FX/Niagara/Source/NiagaraEditor/Private/Toolkits/NiagaraSystemToolkit.cpp b/Engine/Plugins/FX/Niagara/Source/NiagaraEditor/Private/Toolkits/NiagaraSystemToolkit.cpp
--- a/Engine/Plugins/FX/Niagara/Source/NiagaraEditor/Private/Toolkits/NiagaraSystemToolkit.cpp
+++ b/Engine/Plugins/FX/Niagara/Source/NiagaraEditor/Private/Toolkits/NiagaraSystemToolkit.cpp
@@ -54,7 +54,7 @@
 #include "ViewModels/NiagaraMessageLogViewModel.h"
 #include "NiagaraVersionMetaData.h"
 #include "SNiagaraAssetPickerList.h"
-#include "SystemToolkitModes/NiagaraSystemToolkitModeBase.h"
+#include "Toolkits/SystemToolkitModes/NiagaraSystemToolkitModeBase.h"
 #include "SystemToolkitModes/NiagaraSystemToolkitMode_Default.h"
 #include "SystemToolkitModes/NiagaraSystemToolkitMode_Scalability.h"
 #include "ViewModels/NiagaraParameterDefinitionsPanelViewModel.h"
diff --git a/Engine/Plugins/FX/Niagara/Source/NiagaraEditor/Private/Toolkits/SystemToolkitModes/NiagaraSystemToolkitModeBase.cpp b/Engine/Plugins/FX/Niagara/Source/NiagaraEditor/Private/Toolkits/SystemToolkitModes/NiagaraSystemToolkitModeBase.cpp
--- a/Engine/Plugins/FX/Niagara/Source/NiagaraEditor/Private/Toolkits/SystemToolkitModes/NiagaraSystemToolkitModeBase.cpp
+++ b/Engine/Plugins/FX/Niagara/Source/NiagaraEditor/Private/Toolkits/SystemToolkitModes/NiagaraSystemToolkitModeBase.cpp
@@ -1,6 +1,6 @@
 // Copyright Epic Games, Inc. All Rights Reserved.
 
-#include "NiagaraSystemToolkitModeBase.h"
+#include "Toolkits/SystemToolkitModes/NiagaraSystemToolkitModeBase.h"
 
 #include "AdvancedPreviewSceneModule.h"
 #include "ClassIconFinder.h"
diff --git a/Engine/Plugins/FX/Niagara/Source/NiagaraEditor/Private/Toolkits/SystemToolkitModes/NiagaraSystemToolkitModeBase.h b/Engine/Plugins/FX/Niagara/Source/NiagaraEditor/Private/Toolkits/SystemToolkitModes/NiagaraSystemToolkitModeBase.h
deleted file mode 100644
diff --git a/Engine/Plugins/FX/Niagara/Source/NiagaraEditor/Private/Toolkits/SystemToolkitModes/NiagaraSystemToolkitMode_Default.h b/Engine/Plugins/FX/Niagara/Source/NiagaraEditor/Private/Toolkits/SystemToolkitModes/NiagaraSystemToolkitMode_Default.h
--- a/Engine/Plugins/FX/Niagara/Source/NiagaraEditor/Private/Toolkits/SystemToolkitModes/NiagaraSystemToolkitMode_Default.h
+++ b/Engine/Plugins/FX/Niagara/Source/NiagaraEditor/Private/Toolkits/SystemToolkitModes/NiagaraSystemToolkitMode_Default.h
@@ -2,7 +2,7 @@
 
 #pragma once
 
-#include "NiagaraSystemToolkitModeBase.h"
+#include "Toolkits/SystemToolkitModes/NiagaraSystemToolkitModeBase.h"
 
 class FNiagaraSystemToolkit;
 
diff --git a/Engine/Plugins/FX/Niagara/Source/NiagaraEditor/Private/Toolkits/SystemToolkitModes/NiagaraSystemToolkitMode_Scalability.h b/Engine/Plugins/FX/Niagara/Source/NiagaraEditor/Private/Toolkits/SystemToolkitModes/NiagaraSystemToolkitMode_Scalability.h
--- a/Engine/Plugins/FX/Niagara/Source/NiagaraEditor/Private/Toolkits/SystemToolkitModes/NiagaraSystemToolkitMode_Scalability.h
+++ b/Engine/Plugins/FX/Niagara/Source/NiagaraEditor/Private/Toolkits/SystemToolkitModes/NiagaraSystemToolkitMode_Scalability.h
@@ -2,7 +2,7 @@
 
 #pragma once
 
-#include "NiagaraSystemToolkitModeBase.h"
+#include "Toolkits/SystemToolkitModes/NiagaraSystemToolkitModeBase.h"
 #include "Customizations/NiagaraPlatformSetCustomization.h"
 
 class FNiagaraSystemToolkit;
diff --git a/Engine/Plugins/FX/Niagara/Source/NiagaraEditor/Private/ViewModels/Stack/NiagaraStackScriptItemGroup.cpp b/Engine/Plugins/FX/Niagara/Source/NiagaraEditor/Private/ViewModels/Stack/NiagaraStackScriptItemGroup.cpp
--- a/Engine/Plugins/FX/Niagara/Source/NiagaraEditor/Private/ViewModels/Stack/NiagaraStackScriptItemGroup.cpp
+++ b/Engine/Plugins/FX/Niagara/Source/NiagaraEditor/Private/ViewModels/Stack/NiagaraStackScriptItemGroup.cpp
@@ -315,6 +315,7 @@ public:
 			{
 				NewModuleNode = FNiagaraStackGraphUtilities::AddScriptModuleToStack(ScratchPadScriptViewModel->GetOriginalScript(), *OutputNode, TargetIndex);
 				SystemViewModel.Pin()->GetScriptScratchPadViewModel()->FocusScratchPadScriptViewModel(ScratchPadScriptViewModel.ToSharedRef());
+				SystemViewModel.Pin()->FocusTab(FNiagaraSystemToolkitModeBase::ScratchPadScriptsTabID);
 				ScratchPadScriptViewModel->SetIsPendingRename(true);
 			}
 		}
diff --git a/Engine/Plugins/FX/Niagara/Source/NiagaraEditor/Public/Toolkits/SystemToolkitModes/NiagaraSystemToolkitModeBase.h b/Engine/Plugins/FX/Niagara/Source/NiagaraEditor/Public/Toolkits/SystemToolkitModes/NiagaraSystemToolkitModeBase.h
new file mode 100644
--- /dev/null
+++ b/Engine/Plugins/FX/Niagara/Source/NiagaraEditor/Public/Toolkits/SystemToolkitModes/NiagaraSystemToolkitModeBase.h
@@ -0,0 +1,100 @@
+// Copyright Epic Games, Inc. All Rights Reserved.
+
+#pragma once
+
+#include "ViewModels/NiagaraBakerViewModel.h"
+#include "WorkflowOrientedApp/ApplicationMode.h"
+#include "Widgets/Docking/SDockTab.h"
+#include "GraphEditAction.h"
+class FDocumentTracker;
+
+class FNiagaraSystemToolkit;
+class FNiagaraObjectSelection;
+
+class FNiagaraSystemToolkitModeBase : public FApplicationMode
+{
+public:
+	FNiagaraSystemToolkitModeBase(FName InModeName, TWeakPtr<FNiagaraSystemToolkit> InSystemToolkit);
+	~FNiagaraSystemToolkitModeBase();
+	
+	virtual void RegisterTabFactories(TSharedPtr<FTabManager> InTabManager) override;
+
+	int GetActiveSelectionDetailsIndex() const;
+
+	virtual void OnActiveDocumentChanged(TSharedPtr<class SDockTab> NewActiveTab);
+	virtual void OnSystemSelectionChanged();
+
+	virtual void PostActivateMode() override;
+private:
+	TSharedRef<SDockTab> SpawnTab_Viewport(const FSpawnTabArgs& Args);
+	TSharedRef<SDockTab> SpawnTab_CurveEd(const FSpawnTabArgs& Args);
+	TSharedRef<SDockTab> SpawnTab_Sequencer(const FSpawnTabArgs& Args);
+	TSharedRef<SDockTab> SpawnTab_SystemScript(const FSpawnTabArgs& Args);
+	TSharedRef<SDockTab> SpawnTab_SystemParameters(const FSpawnTabArgs& Args);
+	TSharedRef<SDockTab> SpawnTab_SystemParameterDefinitions(const FSpawnTabArgs& Args);
+	TSharedRef<SDockTab> SpawnTab_Details(const FSpawnTabArgs& Args);
+	TSharedRef<SDockTab> SpawnTab_SelectedEmitterGraph(const FSpawnTabArgs& Args);
+	TSharedRef<SDockTab> SpawnTab_DebugSpreadsheet(const FSpawnTabArgs& Args);
+	TSharedRef<SDockTab> SpawnTab_DebugCacheSpreadsheet(const FSpawnTabArgs& Args);
+	TSharedRef<SDockTab> SpawnTab_PreviewSettings(const FSpawnTabArgs& Args);
+	TSharedRef<SDockTab> SpawnTab_GeneratedCode(const FSpawnTabArgs& Args);
+	TSharedRef<SDockTab> SpawnTab_MessageLog(const FSpawnTabArgs& Args);
+	TSharedRef<SDockTab> SpawnTab_SystemOverview(const FSpawnTabArgs& Args);
+	TSharedRef<SDockTab> SpawnTab_ScratchPadScripts(const FSpawnTabArgs& Args);
+	TSharedRef<SDockTab> SpawnTab_ScriptStats(const FSpawnTabArgs& Args);
+	TSharedRef<SDockTab> SpawnTab_Baker(const FSpawnTabArgs& Args);
+	TSharedRef<SDockTab> SpawnTab_Versioning(const FSpawnTabArgs& Args);
+	TSharedRef<SDockTab> SpawnTab_UserParameters(const FSpawnTabArgs& Args);
+	TSharedRef<SDockTab> SpawnTab_UserParametersHierarchyEditor(const FSpawnTabArgs& Args);
+	TSharedRef<SDockTab> SpawnTab_SummaryViewEditor(const FSpawnTabArgs& Args);
+
+protected:
+	TWeakPtr<FNiagaraSystemToolkit> SystemToolkit;
+	FDelegateHandle DocChangedHandle;
+	FDelegateHandle LastSelectionUpdateDelegate;
+	FDelegateHandle LastParamPanelSelectionUpdateDelegate;
+	FDelegateHandle LastSystemSelectionUpdateDelegate;
+	FDelegateHandle LastGraphEditDelegate;
+	FDelegateHandle UpdateSummaryViewHandle;
+	int32 SwitcherIdx;
+	TSharedPtr<FNiagaraObjectSelection> ObjectSelection;
+	FText ObjectSelectionSubHeaderText;
+
+	EVisibility GetObjectSelectionSubHeaderTextVisibility() const;
+	FText GetObjectSelectionSubHeaderText() const;
+	EVisibility GetObjectSelectionNoSelectionTextVisibility() const;
+	void UpdateSelectionForActiveDocument();
+	void OnParameterPanelViewModelExternalSelectionChanged();
+	void OnEditedScriptGraphChanged(const FEdGraphEditAction& InAction);
+	
+	TSharedRef<SWidget> CreateSummaryViewWidget() const;
+	void UpdateSummaryViewOnSelectionChanged() const;
+	
+	void OnSummaryViewEditorClosed(TSharedRef<SDockTab> DockTab) const;
+
+	TWeakPtr< class FNiagaraScratchPadScriptViewModel > LastActiveDocumentModel;
+	TSharedPtr<class SBox> SummaryViewContainer;
+public:
+	static const FName ViewportTabID;
+	static const FName CurveEditorTabID;
+	static const FName SequencerTabID;
+	static const FName SystemScriptTabID;
+	static const FName SystemParametersTabID;
+	static const FName SystemParameterDefinitionsTabID;
+	static const FName DetailsTabID;
+	static const FName SelectedEmitterGraphTabID;
+	static const FName DebugSpreadsheetTabID;
+	static const FName DebugCacheSpreadsheetTabID;
+	static const FName PreviewSettingsTabId;
+	static const FName GeneratedCodeTabID;
+	static const FName MessageLogTabID;
+	static const FName SystemOverviewTabID;
+	static const FName ScratchPadScriptsTabID;
+	static const FName ScriptStatsTabID;
+	static const FName BakerTabID;
+	static const FName VersioningTabID;
+	static const FName UserParametersTabID;
+	static const FName UserParametersHierarchyTabID;
+	static const FName EmitterSummaryViewEditorTabID;;
+};
+
