Fixed a bug where replacing a module with another module using the same name caused it to lose static switch inputs

diff --git a/Engine/Plugins/FX/Niagara/Source/NiagaraEditor/Private/ViewModels/Stack/NiagaraStackModuleItem.cpp b/Engine/Plugins/FX/Niagara/Source/NiagaraEditor/Private/ViewModels/Stack/NiagaraStackModuleItem.cpp
--- a/Engine/Plugins/FX/Niagara/Source/NiagaraEditor/Private/ViewModels/Stack/NiagaraStackModuleItem.cpp
+++ b/Engine/Plugins/FX/Niagara/Source/NiagaraEditor/Private/ViewModels/Stack/NiagaraStackModuleItem.cpp
@@ -1418,10 +1418,11 @@ void UNiagaraStackModuleItem::ReassignModuleScript(UNiagaraScript* ModuleScript)
 			UNiagaraSystem& System = GetSystemViewModel()->GetSystem();
 			FVersionedNiagaraEmitter Emitter = GetEmitterViewModel().IsValid() ? GetEmitterViewModel()->GetEmitter() : FVersionedNiagaraEmitter();
 			FNiagaraStackGraphUtilities::RenameReferencingParameters(&System, Emitter, *FunctionCallNode, OldName, NewName);
-			FunctionCallNode->RefreshFromExternalChanges();
-			FunctionCallNode->MarkNodeRequiresSynchronization(TEXT("Module script reassigned."), true);
-			RefreshChildren();
 		}
+		// we need to refresh the node here to update static switch inputs, as they are set on the pins
+		FunctionCallNode->RefreshFromExternalChanges();
+		FunctionCallNode->MarkNodeRequiresSynchronization(TEXT("Module script reassigned."), true);
+		RefreshChildren();
 		
 		if (NewScriptData->ConversionUtility != nullptr && OldClipboardContent != nullptr)
 		{
