#UE Trivial fix to allow an error message to be localized and not converted from FString

diff --git a/Engine/Plugins/FX/Niagara/Source/Niagara/Private/NiagaraSystem.cpp b/Engine/Plugins/FX/Niagara/Source/Niagara/Private/NiagaraSystem.cpp
--- a/Engine/Plugins/FX/Niagara/Source/Niagara/Private/NiagaraSystem.cpp
+++ b/Engine/Plugins/FX/Niagara/Source/Niagara/Private/NiagaraSystem.cpp
@@ -1761,19 +1761,19 @@ void UNiagaraSystem::ComputeEmittersExecutionOrder()
 				{
 #if WITH_EDITORONLY_DATA
 					FName EmitterName = EmitterHandles[EmitterIdx].GetName();
-					FString ErrorMessage = FString::Printf(TEXT("Found circular dependency involving emitter '%s' in system '%s'. The execution order will be undefined."), *EmitterName.ToString(), *GetName());
+					FText ErrorMessage = FText::Format(LOCTEXT("CircularDependencyError", "Found circular dependency involving emitter '{0}' in system '{1}'. The execution order will be undefined."), FText::FromName(EmitterName), FText::FromString(GetName()));
 					INiagaraModule& NiagaraModule = FModuleManager::GetModuleChecked<INiagaraModule>("Niagara");
 					bool bAllowDismissal = true;
 					UNiagaraMessageDataBase* ComputeEmitterExecutionOrderMessage = NiagaraModule.GetEditorOnlyDataUtilities().CreateWarningMessage(
 						this,
 						LOCTEXT("ComputeEmitterOrderError", "Error computing emitter execution order"),
-						FText::FromString(ErrorMessage),
+						ErrorMessage,
 						"Resolve Data Interfaces",
 						bAllowDismissal);
 					MessageStore.AddMessage(ComputeEmitterExecutionOrderMessageId, ComputeEmitterExecutionOrderMessage);
 					if (MessageStore.IsMessageDismissed(ComputeEmitterExecutionOrderMessageId) == false)
 					{
-						UE_LOG(LogNiagara, Warning, TEXT("%s"), *ErrorMessage);
+						UE_LOG(LogNiagara, Warning, TEXT("%s"), *ErrorMessage.ToString());
 					}
 #endif
 					break;
