Inline enum support for static switches added to Niagara
Effectively it's only integer + display override metadata.

Create a static switch, set it to integer mode. Next to the case labels, edit buttons will appear that let you edit case labels + tooltips.
When editing any of these, the widget type on the enum will be set to 'EnumStyle' for ease of use.
Can be changed back to default by explicitly changing this value on the script value

Changes:
- Added 'Enum Style' widget type for script variables
- Changed SNiagaraNumericDropdown to include an 'requireexactvalues' attribute. When set to true, the text block will not be editable, so it can be used like an enum
- Fixed a bug where if you changed a static switch parameter type and it would match an existing one, the existing metadata would get wiped
- Refactored the case logic to no longer construct the widgets based on the input pin names. Instead, there is dedicated separate logic for it. This is done because the pin names need to keep their previous names when a case label changes so they don't get orphaned
- Added customizability for case labels
-- Click the case label twice to edit it inline, or click the edit button next to it to edit name & tooltip
- Allocate default pins for static switches now uses thes enum style entries if the widget type is set to 'enum style' so that all static switch nodes that reference the same static switch will sync up. This relies on post edit change to be called as the graph will refresh the correct nodes

diff --git a/Engine/Plugins/FX/Niagara/Source/Niagara/Public/NiagaraTypes.h b/Engine/Plugins/FX/Niagara/Source/Niagara/Public/NiagaraTypes.h
--- a/Engine/Plugins/FX/Niagara/Source/Niagara/Public/NiagaraTypes.h
+++ b/Engine/Plugins/FX/Niagara/Source/Niagara/Public/NiagaraTypes.h
@@ -640,7 +640,10 @@ enum class ENiagaraInputWidgetType : uint8
 	Volume,
 
 	// a numeric input, but also has a dropdown with named values
-	NumericDropdown
+	NumericDropdown,
+
+	// A dropdown that behaves like an enum; only allows the exact pre-defined values, for integer inputs only
+	EnumStyle
 };
 
 USTRUCT()
@@ -658,6 +661,19 @@ struct NIAGARA_API FWidgetNamedInputValue
 	FText Tooltip;
 };
 
+/** A struct that serves as display metadata for integer type static switches. Is used in conjuction with the 'EnumStyle' widget customization. */
+USTRUCT()
+struct NIAGARA_API FNiagaraWidgetNamedIntegerInputValue
+{
+	GENERATED_BODY()
+	
+	UPROPERTY(EditAnywhere, Category="Customization")
+	FText DisplayName;
+
+	UPROPERTY(EditAnywhere, Category="Customization", meta = (MultiLine = "true"))
+	FText Tooltip;
+};
+
 USTRUCT()
 struct NIAGARA_API FNiagaraInputParameterCustomization
 {
@@ -691,6 +707,9 @@ struct NIAGARA_API FNiagaraInputParameterCustomization
 	UPROPERTY(EditAnywhere, Category="Customization", meta=(EditCondition="WidgetType == ENiagaraInputWidgetType::NumericDropdown", EditConditionHides))
 	TArray<FWidgetNamedInputValue> InputDropdownValues;
 	
+	UPROPERTY(EditAnywhere, Category="Customization", meta=(EditCondition="WidgetType == ENiagaraInputWidgetType::EnumStyle", EditConditionHides))
+	TArray<FNiagaraWidgetNamedIntegerInputValue> EnumStyleDropdownValues;
+	
 	/** If true then the input is also displayed and editable as a 3d widget in the viewport (vector and transform types only). */
 	//UPROPERTY(EditAnywhere, Category="Customization")
 	//bool bCreateViewPortEditWidget = false;
diff --git a/Engine/Plugins/FX/Niagara/Source/NiagaraEditor/Private/Customizations/NiagaraStaticSwitchNodeDetails.cpp b/Engine/Plugins/FX/Niagara/Source/NiagaraEditor/Private/Customizations/NiagaraStaticSwitchNodeDetails.cpp
--- a/Engine/Plugins/FX/Niagara/Source/NiagaraEditor/Private/Customizations/NiagaraStaticSwitchNodeDetails.cpp
+++ b/Engine/Plugins/FX/Niagara/Source/NiagaraEditor/Private/Customizations/NiagaraStaticSwitchNodeDetails.cpp
@@ -593,25 +593,6 @@ bool FNiagaraStaticSwitchNodeDetails::GetExposeAsPinEnabled() const
 
 }
 
-bool FNiagaraStaticSwitchNodeDetails::GetIntOptionEnabled() const
-{
-	return Node.IsValid() && Node->SwitchTypeData.SwitchType == ENiagaraStaticSwitchType::Integer && (!IsConstantSelection() || SelectedDropdownItem == DropdownOptions[3]);
-}
-
-TOptional<int32> FNiagaraStaticSwitchNodeDetails::GetIntOptionValue() const
-{
-	return Node.IsValid() ? TOptional<int32>(Node->NumOptionsPerVariable) : TOptional<int32>();
-}
-
-void FNiagaraStaticSwitchNodeDetails::IntOptionValueCommitted(int32 Value, ETextCommit::Type CommitInfo)
-{
-	if (Node.IsValid() && Value > 0)
-	{
-		Node->NumOptionsPerVariable = Value;
-		Node->RefreshFromExternalChanges();
-	}
-}
-
 FText FNiagaraStaticSwitchNodeDetails::GetParameterNameText() const
 {
 	return Node.IsValid() ? FText::FromName(Node->InputParameterName) : FText();
diff --git a/Engine/Plugins/FX/Niagara/Source/NiagaraEditor/Private/Customizations/NiagaraStaticSwitchNodeDetails.h b/Engine/Plugins/FX/Niagara/Source/NiagaraEditor/Private/Customizations/NiagaraStaticSwitchNodeDetails.h
--- a/Engine/Plugins/FX/Niagara/Source/NiagaraEditor/Private/Customizations/NiagaraStaticSwitchNodeDetails.h
+++ b/Engine/Plugins/FX/Niagara/Source/NiagaraEditor/Private/Customizations/NiagaraStaticSwitchNodeDetails.h
@@ -73,11 +73,6 @@ private:
 	void OnSelectionChanged(TSharedPtr<SwitchDropdownOption> NewValue, ESelectInfo::Type);
 	FText GetDropdownItemLabel() const;
 	void UpdateSelectionFromNode();
-	
-	// float type option functions
-	bool GetIntOptionEnabled() const;
-	TOptional<int32> GetIntOptionValue() const;
-	void IntOptionValueCommitted(int32 Value, ETextCommit::Type CommitInfo);
 
 	// parameter name option function
 	FText GetParameterNameText() const;
diff --git a/Engine/Plugins/FX/Niagara/Source/NiagaraEditor/Private/NiagaraActions.cpp b/Engine/Plugins/FX/Niagara/Source/NiagaraEditor/Private/NiagaraActions.cpp
--- a/Engine/Plugins/FX/Niagara/Source/NiagaraEditor/Private/NiagaraActions.cpp
+++ b/Engine/Plugins/FX/Niagara/Source/NiagaraEditor/Private/NiagaraActions.cpp
@@ -589,7 +589,10 @@ void FNiagaraParameterGraphDragOperation::MakeStaticSwitch(FNiagaraParameterNode
 	// copy metadata
 	if (UNiagaraGraph* NiagaraGraph = Cast<UNiagaraGraph>(InParams.Graph))
 	{
-		NiagaraGraph->AddParameter(ScriptVariable);
+		if(NiagaraGraph->GetScriptVariable(ScriptVariable->Variable) == nullptr)
+		{
+			NiagaraGraph->AddParameter(ScriptVariable);
+		}
 	}
 	
 	FGraphNodeCreator<UNiagaraNodeStaticSwitch> SetNodeCreator(*InParams.Graph);
diff --git a/Engine/Plugins/FX/Niagara/Source/NiagaraEditor/Private/NiagaraEditorStyle.cpp b/Engine/Plugins/FX/Niagara/Source/NiagaraEditor/Private/NiagaraEditorStyle.cpp
--- a/Engine/Plugins/FX/Niagara/Source/NiagaraEditor/Private/NiagaraEditorStyle.cpp
+++ b/Engine/Plugins/FX/Niagara/Source/NiagaraEditor/Private/NiagaraEditorStyle.cpp
@@ -499,6 +499,7 @@ void FNiagaraEditorStyle::InitStackWidgets()
 			.SetBackgroundImageReadOnly( FSlateNoResource() );
 	
 	Set("NiagaraEditor.Stack.NumericDropdownInput", EditableTextBoxStyle);
+	Set("NiagaraEditor.Stack.IntegerAsEnum", NormalText.SetFont(FAppStyle::Get().GetFontStyle("PropertyWindow.NormalFont")));
 }
 
 void FNiagaraEditorStyle::InitNiagaraSequence()
diff --git a/Engine/Plugins/FX/Niagara/Source/NiagaraEditor/Private/NiagaraEditorUtilities.cpp b/Engine/Plugins/FX/Niagara/Source/NiagaraEditor/Private/NiagaraEditorUtilities.cpp
--- a/Engine/Plugins/FX/Niagara/Source/NiagaraEditor/Private/NiagaraEditorUtilities.cpp
+++ b/Engine/Plugins/FX/Niagara/Source/NiagaraEditor/Private/NiagaraEditorUtilities.cpp
@@ -4335,7 +4335,7 @@ TSharedRef<SWidget> FNiagaraParameterUtilities::GetParameterWidget(FNiagaraVaria
 	return ResultWidget;
 }
 
-TSharedRef<SToolTip> FNiagaraParameterUtilities::GetTooltipWidget(FNiagaraVariable Variable, bool bShowValue, TSharedPtr<SWidget> AdditionalVerticalWidget,  TSharedPtr<SWidget> AdditionalHorizontalWidget)
+TSharedRef<SToolTip> FNiagaraParameterUtilities::GetTooltipWidget(FNiagaraVariable Variable, bool bShowValue, TSharedPtr<SWidget> AdditionalVerticalWidget)
 {
 	FNiagaraTypeDefinition Type = Variable.GetType();
 	const FLinearColor TypeColor = UEdGraphSchema_Niagara::GetTypeColor(Type);
@@ -4347,8 +4347,7 @@ TSharedRef<SToolTip> FNiagaraParameterUtilities::GetTooltipWidget(FNiagaraVariab
 	FString			   IconDocLink, IconDocExcerpt;
 	FSlateBrush const* SecondaryIconBrush = FAppStyle::GetBrush(TEXT("NoBrush"));
 	FSlateColor        SecondaryIconColor = IconColor;
-
-	TSharedPtr<SHorizontalBox> ParameterContainer;
+	
 	TSharedRef<SWidget> TooltipContent = GetParameterWidget(Variable, true, bShowValue);
 
 	TSharedPtr<SVerticalBox> InnerContainer;
@@ -4374,15 +4373,6 @@ TSharedRef<SToolTip> FNiagaraParameterUtilities::GetTooltipWidget(FNiagaraVariab
 		];
 	}
 
-	if(AdditionalHorizontalWidget.IsValid())
-	{
-		ParameterContainer->AddSlot()
-		.AutoWidth()
-		[
-			AdditionalHorizontalWidget.ToSharedRef()	
-		];
-	}
-
 	return TooltipWidget;
 }
 
diff --git a/Engine/Plugins/FX/Niagara/Source/NiagaraEditor/Private/NiagaraGraph.cpp b/Engine/Plugins/FX/Niagara/Source/NiagaraEditor/Private/NiagaraGraph.cpp
--- a/Engine/Plugins/FX/Niagara/Source/NiagaraEditor/Private/NiagaraGraph.cpp
+++ b/Engine/Plugins/FX/Niagara/Source/NiagaraEditor/Private/NiagaraGraph.cpp
@@ -2871,11 +2871,27 @@ void UNiagaraGraph::ScriptVariableChanged(FNiagaraVariable Variable)
 	check(!bIsForCompilationOnly);
 
 	UNiagaraScriptVariable* ScriptVariable = GetScriptVariable(Variable);
-	if (!ScriptVariable || ScriptVariable->GetIsStaticSwitch())
+	if (!ScriptVariable)
 	{
 		return;
 	}
 
+	if(ScriptVariable->GetIsStaticSwitch())
+	{
+		TArray<UNiagaraNodeStaticSwitch*> StaticSwitchNodes;
+		GetNodesOfClass<UNiagaraNodeStaticSwitch>(StaticSwitchNodes);
+
+		for(UNiagaraNodeStaticSwitch* StaticSwitchNode : StaticSwitchNodes)
+		{
+			if(StaticSwitchNode->InputParameterName == ScriptVariable->Variable.GetName())
+			{
+				StaticSwitchNode->AttemptUpdatePins();
+			}
+		}
+		
+		return;
+	}
+
 	TSet<UNiagaraNodeParameterMapGet*> MapGetNodes;
 	TArray<UEdGraphPin*> Pins = FindParameterMapDefaultValuePins(Variable.GetName());
 	for (UEdGraphPin* Pin : Pins)
diff --git a/Engine/Plugins/FX/Niagara/Source/NiagaraEditor/Private/NiagaraNodeSelect.cpp b/Engine/Plugins/FX/Niagara/Source/NiagaraEditor/Private/NiagaraNodeSelect.cpp
--- a/Engine/Plugins/FX/Niagara/Source/NiagaraEditor/Private/NiagaraNodeSelect.cpp
+++ b/Engine/Plugins/FX/Niagara/Source/NiagaraEditor/Private/NiagaraNodeSelect.cpp
@@ -263,10 +263,7 @@ void UNiagaraNodeSelect::PostLoad()
 	Super::PostLoad();
 	
 	// we assume something changed externally if the input pins are outdated; i.e. the assigned enum changed or value order changed
-	if (AreInputPinsOutdated())
-	{
-		RefreshFromExternalChanges();
-	}
+	AttemptUpdatePins();
 }
 
 void UNiagaraNodeSelect::GetNodeContextMenuActions(UToolMenu* Menu, UGraphNodeContextMenuContext* Context) const
diff --git a/Engine/Plugins/FX/Niagara/Source/NiagaraEditor/Private/NiagaraNodeStaticSwitch.cpp b/Engine/Plugins/FX/Niagara/Source/NiagaraEditor/Private/NiagaraNodeStaticSwitch.cpp
--- a/Engine/Plugins/FX/Niagara/Source/NiagaraEditor/Private/NiagaraNodeStaticSwitch.cpp
+++ b/Engine/Plugins/FX/Niagara/Source/NiagaraEditor/Private/NiagaraNodeStaticSwitch.cpp
@@ -18,7 +18,12 @@
 #include "Subsystems/AssetEditorSubsystem.h"
 #include "Editor/EditorEngine.h"
 #include "Editor.h"
+#include "ISinglePropertyView.h"
+#include "IStructureDetailsView.h"
 #include "NiagaraNodeFunctionCall.h"
+#include "Framework/Notifications/NotificationManager.h"
+#include "Widgets/Input/SMultiLineEditableTextBox.h"
+#include "Widgets/Notifications/SNotificationList.h"
 
 #include UE_INLINE_GENERATED_CPP_BY_NAME(NiagaraNodeStaticSwitch)
 
@@ -53,6 +58,114 @@ FNiagaraTypeDefinition UNiagaraNodeStaticSwitch::GetInputType() const
 	return FNiagaraTypeDefinition();
 }
 
+void UNiagaraNodeStaticSwitch::OnInputCaseLabelSubmitted(const FText& Text, ETextCommit::Type Arg, int32 InputCase) const
+{
+	if(UNiagaraScriptVariable* StaticSwitchVariable = GetStaticSwitchScriptVariable())
+	{
+		StaticSwitchVariable->PreEditChange(nullptr);
+		// if we have at least one case label that has been changed, we want to switch to enum style display.
+		StaticSwitchVariable->Metadata.WidgetCustomization.WidgetType = ENiagaraInputWidgetType::EnumStyle;
+		
+		if(ensure(StaticSwitchVariable->Metadata.WidgetCustomization.EnumStyleDropdownValues.IsValidIndex(InputCase)))
+		{
+			StaticSwitchVariable->Metadata.WidgetCustomization.EnumStyleDropdownValues[InputCase].DisplayName = Text;
+		}
+
+		FPropertyChangedEvent PropertyChangedEvent(nullptr);
+		StaticSwitchVariable->PostEditChangeProperty(PropertyChangedEvent);
+	}
+}
+
+bool UNiagaraNodeStaticSwitch::VerifyCaseLabelCandidate(const FText& InCandidate, FText& OutErrorMessage) const
+{
+	bool bResult = Super::VerifyCaseLabelCandidate(InCandidate, OutErrorMessage);
+	if(bResult == false)
+	{
+		return bResult;
+	}
+
+	if(UNiagaraScriptVariable* ScriptVariable = GetStaticSwitchScriptVariable())
+	{		
+		for(const FNiagaraWidgetNamedIntegerInputValue& Value : ScriptVariable->Metadata.WidgetCustomization.EnumStyleDropdownValues)
+		{
+			if(Value.DisplayName.CompareToCaseIgnored(InCandidate) == 0)
+			{
+				OutErrorMessage = LOCTEXT("CaseAlreadyExists", "This case already exists. Choose a different label!");
+				return false;
+			}
+		}
+	}
+
+	return true;
+}
+
+FText UNiagaraNodeStaticSwitch::GetInputCaseTooltip(int32 InputCase) const
+{
+	if(UNiagaraScriptVariable* ScriptVariable = GetStaticSwitchScriptVariable())
+	{
+		if(ScriptVariable->Metadata.WidgetCustomization.EnumStyleDropdownValues.IsValidIndex(InputCase))
+		{
+			return ScriptVariable->Metadata.WidgetCustomization.EnumStyleDropdownValues[InputCase].Tooltip;
+		}
+	}
+
+	return FText::GetEmpty();
+}
+
+FText UNiagaraNodeStaticSwitch::GetInputCaseButtonTooltip(int32 Case) const
+{
+	return FText::FormatOrdered(LOCTEXT("SummonCaseLabelButtonTooltip", "If input is {0}.\nEdit the case label & tooltip. Editing these will set the widget type to 'Enum Style'."), FText::FromString(FString::FromInt(Case))); 
+}
+
+void UNiagaraNodeStaticSwitch::OnInputCaseTooltipSubmitted(const FText& Text, ETextCommit::Type Arg, int32 InputCase) const
+{
+	if(UNiagaraScriptVariable* ScriptVariable = GetStaticSwitchScriptVariable())
+	{
+		ScriptVariable->PreEditChange(nullptr);
+		if(ScriptVariable->Metadata.WidgetCustomization.EnumStyleDropdownValues.IsValidIndex(InputCase))
+		{
+			ScriptVariable->Metadata.WidgetCustomization.EnumStyleDropdownValues[InputCase].Tooltip = Text;
+			FPropertyChangedEvent PropertyChangedEvent(nullptr);
+			ScriptVariable->PostEditChangeProperty(PropertyChangedEvent);
+		}
+	}
+}
+
+TSharedRef<SWidget> UNiagaraNodeStaticSwitch::GetInputCaseContextOptions(int32 Case)
+{
+	if(SwitchTypeData.SwitchType == ENiagaraStaticSwitchType::Integer)
+	{
+		if(UNiagaraScriptVariable* ScriptVariable = GetStaticSwitchScriptVariable())
+		{
+			// since we want to edit some case of a static switch 
+			if(Case >= ScriptVariable->Metadata.WidgetCustomization.EnumStyleDropdownValues.Num())
+			{
+				int32 PreviousElementCount = ScriptVariable->Metadata.WidgetCustomization.EnumStyleDropdownValues.Num();
+				ScriptVariable->Metadata.WidgetCustomization.EnumStyleDropdownValues.Reserve(Case + 1);
+
+				for(int32 MissingElementIndex = PreviousElementCount; MissingElementIndex <= Case; MissingElementIndex++)
+				{
+					FNiagaraWidgetNamedIntegerInputValue& NamedIntegerInputValue = ScriptVariable->Metadata.WidgetCustomization.EnumStyleDropdownValues.AddDefaulted_GetRef();
+					NamedIntegerInputValue.DisplayName = FText::FromString(FString::FromInt(Case));
+				}
+			}
+
+			TSharedRef<FStructOnScope> MenuStruct = MakeShared<FStructOnScope>(StaticStruct<FNiagaraWidgetNamedIntegerInputValue>(), (uint8*) &GetStaticSwitchScriptVariable()->Metadata.WidgetCustomization.EnumStyleDropdownValues[Case]);;
+			MenuStruct->SetPackage(GetPackage());
+			FPropertyEditorModule& PropertyEditorModule = FModuleManager::LoadModuleChecked<FPropertyEditorModule>("PropertyEditor");
+
+			FDetailsViewArgs DetailsViewArgs;
+			DetailsViewArgs.bAllowSearch = false;
+			DetailsViewArgs.NotifyHook = this;
+			TSharedRef<IStructureDetailsView> StructureView = PropertyEditorModule.CreateStructureDetailView(DetailsViewArgs, FStructureDetailsViewArgs(), MenuStruct);
+
+			return StructureView->GetWidget().ToSharedRef();
+		}
+	}
+	
+	return SNullWidget::NullWidget;	
+}
+
 FString UNiagaraNodeStaticSwitch::GetInputCaseName(int32 Case) const
 {
 	if (SwitchTypeData.SwitchType == ENiagaraStaticSwitchType::Bool)
@@ -68,6 +181,15 @@ FString UNiagaraNodeStaticSwitch::GetInputCaseName(int32 Case) const
 	}
 	else if (SwitchTypeData.SwitchType == ENiagaraStaticSwitchType::Integer)
 	{
+		UNiagaraScriptVariable* ScriptVariable = GetStaticSwitchScriptVariable();
+		if(ScriptVariable && ScriptVariable->Metadata.WidgetCustomization.WidgetType == ENiagaraInputWidgetType::EnumStyle)
+		{
+			if(ScriptVariable->Metadata.WidgetCustomization.EnumStyleDropdownValues.IsValidIndex(Case))
+			{
+				return ScriptVariable->Metadata.WidgetCustomization.EnumStyleDropdownValues[Case].DisplayName.ToString();
+			}
+		}
+		
 		return FString::FromInt(Case);
 	}
 	else if (SwitchTypeData.SwitchType == ENiagaraStaticSwitchType::Enum && SwitchTypeData.Enum)
@@ -91,6 +213,22 @@ void UNiagaraNodeStaticSwitch::PostEditChangeProperty(struct FPropertyChangedEve
 		ReallocatePins();
 	}
 }
+
+int32 UNiagaraNodeStaticSwitch::GetNumberOfCases() const
+{
+	if(SwitchTypeData.SwitchType == ENiagaraStaticSwitchType::Integer)
+	{
+		if(UNiagaraScriptVariable* ScriptVariable = GetStaticSwitchScriptVariable())
+		{
+			if(ScriptVariable->Metadata.WidgetCustomization.WidgetType == ENiagaraInputWidgetType::EnumStyle)
+			{
+				return ScriptVariable->Metadata.WidgetCustomization.EnumStyleDropdownValues.Num();
+			}
+		}
+	}
+
+	return Super::GetNumberOfCases();
+}
 #endif
 
 TArray<int32> UNiagaraNodeStaticSwitch::GetOptionValues() const
@@ -102,7 +240,17 @@ TArray<int32> UNiagaraNodeStaticSwitch::GetOptionValues() const
 	}
 	else if (SwitchTypeData.SwitchType == ENiagaraStaticSwitchType::Integer)
 	{
-		int32 NewOptionsCount = FMath::Max(2, NumOptionsPerVariable);
+		// we want complete sync of the enum entries, if selected
+		bool bEnsureMinimumCases = true;
+		if(UNiagaraScriptVariable* ScriptVariable = GetStaticSwitchScriptVariable())
+		{
+			if(ScriptVariable->Metadata.WidgetCustomization.WidgetType == ENiagaraInputWidgetType::EnumStyle)
+			{
+				bEnsureMinimumCases = false;
+			}
+		}
+		
+		int32 NewOptionsCount = bEnsureMinimumCases ? FMath::Max(2, GetNumberOfCases()) : GetNumberOfCases();
 		for(int32 Counter = 0; Counter < NewOptionsCount; Counter++)
 		{
 			OptionValues.Add(Counter);
@@ -151,6 +299,39 @@ void UNiagaraNodeStaticSwitch::PostChange(const UUserDefinedEnum* Changed, FEnum
 	}
 }
 
+void UNiagaraNodeStaticSwitch::NotifyPreChange(FProperty* PropertyAboutToChange)
+{
+	if(PropertyAboutToChange->GetFName() == GET_MEMBER_NAME_CHECKED(FNiagaraWidgetNamedIntegerInputValue, DisplayName) || PropertyAboutToChange->GetFName() == GET_MEMBER_NAME_CHECKED(FNiagaraWidgetNamedIntegerInputValue, Tooltip))
+	{
+		if(UNiagaraScriptVariable* ScriptVariable = GetStaticSwitchScriptVariable())
+		{
+			ScriptVariable->PreEditChange(PropertyAboutToChange);
+		}
+	}
+}
+
+void UNiagaraNodeStaticSwitch::NotifyPostChange(const FPropertyChangedEvent& PropertyChangedEvent, FProperty* PropertyThatChanged)
+{
+	if(PropertyThatChanged->GetFName() == GET_MEMBER_NAME_CHECKED(FNiagaraWidgetNamedIntegerInputValue, DisplayName) || PropertyThatChanged->GetFName() == GET_MEMBER_NAME_CHECKED(FNiagaraWidgetNamedIntegerInputValue, Tooltip))
+	{
+		if(UNiagaraScriptVariable* ScriptVariable = GetStaticSwitchScriptVariable())
+		{
+			// since we only use this function to edit enum style entries in a menu, it's safe to assume the user wants to display this entry as enum
+			if(ScriptVariable->Metadata.WidgetCustomization.WidgetType != ENiagaraInputWidgetType::EnumStyle)
+			{
+				ScriptVariable->Metadata.WidgetCustomization.WidgetType = ENiagaraInputWidgetType::EnumStyle;
+				
+				FText NotificationText = FText::FormatOrdered(LOCTEXT("WidgetTypeOfStaticSwitchSetToEnumStyle", "Variable {0}'s widget style has been set to 'Enum Style' to support display of custom names and tooltips."), FText::FromName(ScriptVariable->Variable.GetName()));
+				FNotificationInfo NotificationInfo(NotificationText);
+				NotificationInfo.ExpireDuration = 5.f;
+				FSlateNotificationManager::Get().AddNotification(NotificationInfo);				
+			}
+			FPropertyChangedEvent PropertyChangedEventNonConst = PropertyChangedEvent;
+			ScriptVariable->PostEditChangeProperty(PropertyChangedEventNonConst);
+		}
+	}
+}
+
 FName UNiagaraNodeStaticSwitch::GetOptionPinName(const FNiagaraVariable& Variable, int32 Value) const
 {
 	FString Suffix = TEXT("");
@@ -188,13 +369,29 @@ void UNiagaraNodeStaticSwitch::ChangeSwitchParameterName(const FName& NewName)
 
 void UNiagaraNodeStaticSwitch::OnSwitchParameterTypeChanged(const FNiagaraTypeDefinition& OldType)
 {
+	TOptional<FNiagaraVariableMetaData> AlreadyExistingMetaData = GetNiagaraGraph()->GetMetaData(FNiagaraVariable(GetInputType(), InputParameterName));
 	TOptional<FNiagaraVariableMetaData> OldMetaData = GetNiagaraGraph()->GetMetaData(FNiagaraVariable(OldType, InputParameterName));
-	RefreshFromExternalChanges(); // Magick happens here: The old pins are destroyed and new ones are created.
-	if (OldMetaData.IsSet())
+
+	// we make sure the new parameter exists at this point
+	GetNiagaraGraph()->AddParameter(FNiagaraVariable(GetInputType(), InputParameterName), true);
+	
+	// we overwrite the metadata only if it didn't exist already
+	if (AlreadyExistingMetaData.IsSet() == false && OldMetaData.IsSet())
 	{
 		GetNiagaraGraph()->SetMetaData(FNiagaraVariable(GetInputType(), InputParameterName), OldMetaData.GetValue());
 	}
 
+	// if we changed to an integer, we want enum style behavior by default
+	if(AlreadyExistingMetaData.IsSet() == false && GetInputType() == FNiagaraTypeDefinition::GetIntDef())
+	{
+		if(UNiagaraScriptVariable* ScriptVariable = GetStaticSwitchScriptVariable())
+		{
+			ScriptVariable->Metadata.WidgetCustomization.WidgetType = ENiagaraInputWidgetType::EnumStyle;
+		}
+	}
+
+	RefreshFromExternalChanges(); // Magick happens here: The old pins are destroyed and new ones are created.
+
 	VisualsChangedDelegate.Broadcast(this);
 	RemoveUnusedGraphParameter(FNiagaraVariable(OldType, InputParameterName));
 }
@@ -275,9 +472,23 @@ void UNiagaraNodeStaticSwitch::AllocateDefaultPins()
 {
 	const UEdGraphSchema_Niagara* Schema = GetDefault<UEdGraphSchema_Niagara>();
 
+	// we ensure we have at least 2 entries in the enum style drop downs
+	if(SwitchTypeData.SwitchType == ENiagaraStaticSwitchType::Integer)
+	{
+		if(UNiagaraScriptVariable* ScriptVariable = GetStaticSwitchScriptVariable())
+		{
+			int32 PreviousEntryCount = ScriptVariable->Metadata.WidgetCustomization.EnumStyleDropdownValues.Num();
+			for(int32 MissingValueIndex = PreviousEntryCount; MissingValueIndex < 2; MissingValueIndex++)
+			{
+				FNiagaraWidgetNamedIntegerInputValue& NamedIntegerInputValue = ScriptVariable->Metadata.WidgetCustomization.EnumStyleDropdownValues.AddDefaulted_GetRef();
+				NamedIntegerInputValue.DisplayName = FText::FromString(FString::FromInt(MissingValueIndex));
+			}
+		}
+	}
+	
 	TArray<int32> OptionValues = GetOptionValues();
 	NumOptionsPerVariable = OptionValues.Num();
-
+	
 	for (int32 i = 0; i < NumOptionsPerVariable; i++)
 	{
 		for (const FNiagaraVariable& Variable : OutputVars)
@@ -486,6 +697,17 @@ bool UNiagaraNodeStaticSwitch::GetVarIndex(FTranslator* Translator, int32 InputP
 	return Success;
 }
 
+UNiagaraScriptVariable* UNiagaraNodeStaticSwitch::GetStaticSwitchScriptVariable() const
+{
+	if(GetNiagaraGraph()->IsCompilationCopy() || IsSetByPin())
+	{
+		return nullptr;
+	}
+
+	FNiagaraVariable Variable(GetInputType(), InputParameterName);
+	return GetNiagaraGraph()->GetScriptVariable(Variable);
+}
+
 void UNiagaraNodeStaticSwitch::CheckForOutdatedEnum(FTranslator* Translator)
 {
 	// we check during compilation if the node is using outdated enum information and log it
@@ -691,6 +913,16 @@ void UNiagaraNodeStaticSwitch::AddIntegerInputPin()
 	this->Modify();
 
 	NumOptionsPerVariable++;
+	
+	if(UNiagaraScriptVariable* ScriptVariable = GetStaticSwitchScriptVariable())
+	{
+		ScriptVariable->PreEditChange(nullptr);
+		FNiagaraWidgetNamedIntegerInputValue& NewEntry = ScriptVariable->Metadata.WidgetCustomization.EnumStyleDropdownValues.AddDefaulted_GetRef();
+		NewEntry.DisplayName = FText::FromString(FString::FromInt(NumOptionsPerVariable - 1));
+		FPropertyChangedEvent PropertyChangedEvent(nullptr);
+		ScriptVariable->PostEditChangeProperty(PropertyChangedEvent);
+	}
+	
 	ReallocatePins();
 }
 
@@ -700,7 +932,16 @@ void UNiagaraNodeStaticSwitch::RemoveIntegerInputPin()
 
 	this->Modify();
 
-	NumOptionsPerVariable = FMath::Max(1, --NumOptionsPerVariable);
+	NumOptionsPerVariable = FMath::Max(2, --NumOptionsPerVariable);
+	
+	if(UNiagaraScriptVariable* ScriptVariable = GetStaticSwitchScriptVariable())
+	{
+		ScriptVariable->PreEditChange(nullptr);
+		ScriptVariable->Metadata.WidgetCustomization.EnumStyleDropdownValues.RemoveAt(ScriptVariable->Metadata.WidgetCustomization.EnumStyleDropdownValues.Num() - 1);
+		FPropertyChangedEvent PropertyChangedEvent(nullptr);
+		ScriptVariable->PostEditChangeProperty(PropertyChangedEvent);
+	}
+	
 	ReallocatePins();
 }
 
@@ -748,10 +989,7 @@ void UNiagaraNodeStaticSwitch::PostLoad()
 	if(SwitchTypeData.bAutoRefreshEnabled == true || NiagaraStaticSwitchCVars::bEnabledAutoRefreshOldStaticSwitches)
 	{
 		// we assume something changed externally if the input pins are outdated; i.e. the assigned enum changed or value order changed
-		if(AreInputPinsOutdated())
-		{
-			RefreshFromExternalChanges();
-		}
+		AttemptUpdatePins();
 	}	
 }
 
diff --git a/Engine/Plugins/FX/Niagara/Source/NiagaraEditor/Private/NiagaraNodeStaticSwitch.h b/Engine/Plugins/FX/Niagara/Source/NiagaraEditor/Private/NiagaraNodeStaticSwitch.h
--- a/Engine/Plugins/FX/Niagara/Source/NiagaraEditor/Private/NiagaraNodeStaticSwitch.h
+++ b/Engine/Plugins/FX/Niagara/Source/NiagaraEditor/Private/NiagaraNodeStaticSwitch.h
@@ -45,7 +45,7 @@ struct FStaticSwitchTypeData
 };
 
 UCLASS(MinimalAPI)
-class UNiagaraNodeStaticSwitch : public UNiagaraNodeUsageSelector, public FEnumEditorUtils::INotifyOnEnumChanged
+class UNiagaraNodeStaticSwitch : public UNiagaraNodeUsageSelector, public FEnumEditorUtils::INotifyOnEnumChanged, public FNotifyHook
 {
 	GENERATED_UCLASS_BODY()
 
@@ -72,7 +72,7 @@ public:
 
 	bool IsSetByPin() const;
 	UEdGraphPin* GetSelectorPin() const;
-
+	
 	/** This is a hack used in the translator to check for inconsistencies with old static switches before auto refresh was a thing */
 	void CheckForOutdatedEnum(FTranslator* Translator);
 
@@ -107,6 +107,18 @@ public:
 
 protected:
 	//~ Begin UNiagaraNodeUsageSelector Interface
+	virtual int32 GetNumberOfCases() const override;
+	
+	virtual bool AreInputCaseLabelsReadOnly() const override { return SwitchTypeData.SwitchType != ENiagaraStaticSwitchType::Integer || IsSetByPin(); }
+	virtual void OnInputCaseLabelSubmitted(const FText& Text, ETextCommit::Type Arg, int32 InputCase) const override;
+	virtual bool VerifyCaseLabelCandidate(const FText& InCandidate, FText& OutErrorMessage) const override;
+	
+	virtual FText GetInputCaseTooltip(int32 Case) const override;
+	virtual FText GetInputCaseButtonTooltip(int32 Case) const override;
+	virtual void OnInputCaseTooltipSubmitted(const FText& Text, ETextCommit::Type Arg, int32 InputCase) const override;
+
+	virtual TSharedRef<SWidget> GetInputCaseContextOptions(int32 Case) override;
+	
 	virtual FString GetInputCaseName(int32 Case) const override;
 	virtual FName GetOptionPinName(const FNiagaraVariable& Variable, int32 Value) const override;
 	virtual TArray<int32> GetOptionValues() const override;
@@ -117,12 +129,19 @@ private:
 	/** INotifyOnEnumChanged interface */
 	virtual void PreChange(const UUserDefinedEnum* Changed, FEnumEditorUtils::EEnumEditorChangeInfo ChangedType) override;
 	virtual void PostChange(const UUserDefinedEnum* Changed, FEnumEditorUtils::EEnumEditorChangeInfo ChangedType) override;
+
+	/** FNotifyHook interface. Used to propagate changes in structure details views. */
+	virtual void NotifyPreChange(FProperty* PropertyAboutToChange) override;
+	virtual void NotifyPostChange(const FPropertyChangedEvent& PropertyChangedEvent, FProperty* PropertyThatChanged) override;
 	
 	/** This finds the first valid input pin index for the current switch value, returns false if no value can be found */
 	bool GetVarIndex(FTranslator* Translator, int32 InputPinCount, int32& VarIndexOut) const;
 
 	bool GetVarIndex(FTranslator* Translator, int32 InputPinCount, int32 Value, int32& VarIndexOut) const;
 
+	/** Retrieve the script variable this switch currently represents. Can be nullptr in certain cases, such as if 'Expose as Pin' is active. */
+	UNiagaraScriptVariable* GetStaticSwitchScriptVariable() const;
+	
 	void RemoveUnusedGraphParameter(const FNiagaraVariable& OldParameter);
 	
 	void AddIntegerInputPin();
diff --git a/Engine/Plugins/FX/Niagara/Source/NiagaraEditor/Private/NiagaraNodeUsageSelector.cpp b/Engine/Plugins/FX/Niagara/Source/NiagaraEditor/Private/NiagaraNodeUsageSelector.cpp
--- a/Engine/Plugins/FX/Niagara/Source/NiagaraEditor/Private/NiagaraNodeUsageSelector.cpp
+++ b/Engine/Plugins/FX/Niagara/Source/NiagaraEditor/Private/NiagaraNodeUsageSelector.cpp
@@ -1,10 +1,16 @@
 // Copyright Epic Games, Inc. All Rights Reserved.
 
 #include "NiagaraNodeUsageSelector.h"
+
+#include "NiagaraCustomVersion.h"
+#include "NiagaraEditorStyle.h"
 #include "NiagaraEditorUtilities.h"
 #include "NiagaraHlslTranslator.h"
 #include "NiagaraScriptSource.h"
 #include "Widgets/Layout/SSeparator.h"
+#include "Widgets/Layout/SWidgetSwitcher.h"
+#include "Widgets/Input/SComboButton.h"
+#include "Widgets/Text/SInlineEditableTextBlock.h"
 
 #include UE_INLINE_GENERATED_CPP_BY_NAME(NiagaraNodeUsageSelector)
 
@@ -66,6 +72,14 @@ void UNiagaraNodeUsageSelector::AddOptionPin(const FNiagaraVariable& OutputVaria
 	NewPin->PinFriendlyName = GetOptionPinFriendlyName(OutputVariable);
 }
 
+void UNiagaraNodeUsageSelector::AttemptUpdatePins()
+{
+	if(AreInputPinsOutdated() && GetNiagaraGraph()->IsCompilationCopy() == false)
+	{
+		ReallocatePins();
+	}
+}
+
 TArray<int32> UNiagaraNodeUsageSelector::GetOptionValues() const
 {
 	TArray<int32> OptionValues;
@@ -113,6 +127,22 @@ FText UNiagaraNodeUsageSelector::GetOptionPinFriendlyName(const FNiagaraVariable
 	return FText::AsCultureInvariant(Variable.GetName().ToString());
 }
 
+bool UNiagaraNodeUsageSelector::VerifyCaseLabelCandidate(const FText& InCandidate, FText& OutErrorMessage) const
+{
+	if(InCandidate.IsEmptyOrWhitespace())
+	{
+		OutErrorMessage = LOCTEXT("CaseLabelCantBeEmpty", "Case can not be empty");
+		return false;
+	}
+
+	return true;
+}
+
+FText UNiagaraNodeUsageSelector::GetInputCaseButtonTooltip(int32 Case) const
+{
+	return FText::FormatOrdered(LOCTEXT("SummonCaseLabelButtonTooltip", "If input is {0}.\nEdit the case label & tooltip."), FText::FromString(FString::FromInt(Case))); 
+}
+
 FString UNiagaraNodeUsageSelector::GetInputCaseName(int32 Case) const
 {
 	UEnum* ENiagaraScriptGroupEnum = StaticEnum<ENiagaraScriptGroup>();
@@ -124,6 +154,11 @@ FString UNiagaraNodeUsageSelector::GetInputCaseName(int32 Case) const
 	return TEXT("");
 }
 
+FText UNiagaraNodeUsageSelector::GetInputCaseNameAsText(int32 Case) const
+{
+	return FText::FromString(GetInputCaseName(Case));
+}
+
 FName UNiagaraNodeUsageSelector::GetOptionPinName(const FNiagaraVariable& Variable, int32 Value) const
 {
 	return FName(Variable.GetName().ToString() + FString::Printf(TEXT(" if %s"), *GetInputCaseName(Value)));
@@ -235,41 +270,13 @@ void UNiagaraNodeUsageSelector::AddWidgetsToInputBox(TSharedPtr<SVerticalBox> In
 		return Pin->bOrphanedPin == true;
 	});
 	
-	for (int32 Idx = 0; Idx < OutputVars.Num() * NumOptionsPerVariable; Idx += OutputVars.Num())
-	{
-		TArray<FString> PinNameParts;
-		InputPins[Idx]->PinName.ToString().ParseIntoArray(PinNameParts, TEXT(" "), true);
-
-		int32 IfIndex = INDEX_NONE;
-		// we are looking for the last if as the variable name could contain a separate "if"
-		PinNameParts.FindLast(TEXT("if"), IfIndex);
-
-		// should never be index_none
-		if(IfIndex != INDEX_NONE)
-		{
-			// there should always be at least one additional part in the pin name that makes up the case label
-			ensure(PinNameParts.IsValidIndex(IfIndex + 1));
-			FString Case = TEXT("");
-			
-			for(int32 CasePartIndex = IfIndex + 1; CasePartIndex < PinNameParts.Num(); CasePartIndex++)
-			{
-				Case += PinNameParts[CasePartIndex] + TEXT(" ");
-			}
-
-			// remove the last unnecessary space
-			Case.RemoveFromEnd(TEXT(" "));
-			Cases.Add(Case);
-		}		
-	}
-	
-	int32 Offset = 0;
-	int32 NaturalIndex = 0;
+	int32 WidgetIndexOffset = 0;
 	// insert separators to make the grouping apparent.
 	for (int32 Idx = 0; Idx < OutputVars.Num() * NumOptionsPerVariable; Idx += OutputVars.Num())
 	{
-		const int32 OptionValueIndex = Idx / OutputVars.Num();
+		const int32 CaseValue = Idx / OutputVars.Num();
 		// adding an offset to account for the already added separators
-		InputBox->InsertSlot(Idx + Offset)
+		InputBox->InsertSlot(Idx + WidgetIndexOffset)
 		.Padding(5.f, 5.f, 0.f, 2.f)
 		.AutoHeight()
 		[
@@ -280,17 +287,46 @@ void UNiagaraNodeUsageSelector::AddWidgetsToInputBox(TSharedPtr<SVerticalBox> In
 			.AutoWidth()
 			[
 				SNew(STextBlock)
-				.Text(FText::FromString(TEXT("If ") + Cases[NaturalIndex++]))
+				.Text(FText::FromString(TEXT("If ")))
+			]
+			+ SHorizontalBox::Slot()
+			.VAlign(VAlign_Center)
+			.HAlign(HAlign_Left)
+			.AutoWidth()
+			[
+				SNew(SWidgetSwitcher)
+				.WidgetIndex(AreInputCaseLabelsReadOnly() ? 0 : 1)
+				+ SWidgetSwitcher::Slot()
+				[
+					SNew(STextBlock)
+					.Text_UObject(this, &UNiagaraNodeUsageSelector::GetInputCaseNameAsText, CaseValue)
+				]
+				+ SWidgetSwitcher::Slot()
+				[
+					SNew(SComboButton)
+					.OnGetMenuContent_UObject(this, &UNiagaraNodeUsageSelector::GetInputCaseContextOptions, CaseValue)
+					.ToolTipText_UObject(this, &UNiagaraNodeUsageSelector::GetInputCaseButtonTooltip, CaseValue)
+					.ButtonContent()
+					[
+						SNew(SHorizontalBox)
+						+ SHorizontalBox::Slot()
+						.AutoWidth()
+						[
+							SNew(STextBlock)
+							.Text_UObject(this, &UNiagaraNodeUsageSelector::GetInputCaseNameAsText, CaseValue)
+						]
+					]
+				]
 			]
 		];
-		Offset++;
+		WidgetIndexOffset++;
 		
-		InputBox->InsertSlot(Idx + Offset)
+		InputBox->InsertSlot(Idx + WidgetIndexOffset)
 		.VAlign(VAlign_Center)
 		[
 			SNew(SSeparator)
 		];
-		Offset++;
+		WidgetIndexOffset++;
 	}
 }
 
diff --git a/Engine/Plugins/FX/Niagara/Source/NiagaraEditor/Private/NiagaraNodeUsageSelector.h b/Engine/Plugins/FX/Niagara/Source/NiagaraEditor/Private/NiagaraNodeUsageSelector.h
--- a/Engine/Plugins/FX/Niagara/Source/NiagaraEditor/Private/NiagaraNodeUsageSelector.h
+++ b/Engine/Plugins/FX/Niagara/Source/NiagaraEditor/Private/NiagaraNodeUsageSelector.h
@@ -19,9 +19,6 @@ public:
 
 	UPROPERTY(meta = (SkipForCompileHash = "true"))
 	FGuid SelectorGuid;
-
-	UPROPERTY()
-	int32 NumOptionsPerVariable = 0;
 	
 	//~ Begin UObject Interface
 	virtual void PostEditChangeProperty(FPropertyChangedEvent& PropertyChangedEvent) override;
@@ -49,7 +46,9 @@ public:
 
 	/** Helper function to create a variable to add to the OutputVars and FGuid to add to OutputVarGuids. */
 	FGuid AddOutput(FNiagaraTypeDefinition Type, const FName& Name);
+	virtual int32 GetNumberOfCases() const { return NumOptionsPerVariable; }
 
+	void AttemptUpdatePins();
 	bool AreInputPinsOutdated() const;
 
 	/** Retrieves option values. Required to determine separator entries since we no longer have case information after pin creation. */
@@ -58,9 +57,20 @@ public:
 protected:
 	/** Use this function to add an option pin in allocation or insertion. Optionally at a specific slot. */
 	void AddOptionPin(const FNiagaraVariable& OutputVariable, int32 Value, int32 InsertionSlot = INDEX_NONE);
+
+	virtual bool AreInputCaseLabelsReadOnly() const { return true; }
+	virtual void OnInputCaseLabelSubmitted(const FText& Text, ETextCommit::Type Arg, int32 OptionValue) const {}
+	virtual bool VerifyCaseLabelCandidate(const FText& InCandidate, FText& OutErrorMessage) const;
+
+	virtual FText GetInputCaseTooltip(int32 Case) const { return FText::GetEmpty(); }
+	virtual FText GetInputCaseButtonTooltip(int32 Case) const;
+	virtual void OnInputCaseTooltipSubmitted(const FText& Text, ETextCommit::Type Arg, int32 OptionValue) const {}
 	
+	virtual TSharedRef<SWidget> GetInputCaseContextOptions(int32 Case) { return SNullWidget::NullWidget; }
+
 	/** Used to determine the text for the separators */
 	virtual FString GetInputCaseName(int32 Case) const;
+	FText GetInputCaseNameAsText(int32 Case) const;
 	
 	/** Ideally we wouldn't need this due to the separators but to maintain backwards compatibility pin names needs to stay consistent with what they were */
 	virtual FName GetOptionPinName(const FNiagaraVariable& Variable, int32 Value) const;
@@ -73,6 +83,9 @@ protected:
 	virtual bool AllowNiagaraTypeForAddPin(const FNiagaraTypeDefinition& InType) const override;
 	//~ End UNiagaraNodeWithDynamicPins Interface
 
+protected:
+	UPROPERTY()
+	int32 NumOptionsPerVariable = 0;
 private:
 	/** private UNiagaraNodeWithDynamicPins interface */
 	virtual void OnNewTypedPinAdded(UEdGraphPin*& NewPin) override;
diff --git a/Engine/Plugins/FX/Niagara/Source/NiagaraEditor/Private/TypeEditorUtilities/NiagaraIntegerTypeEditorUtilities.cpp b/Engine/Plugins/FX/Niagara/Source/NiagaraEditor/Private/TypeEditorUtilities/NiagaraIntegerTypeEditorUtilities.cpp
--- a/Engine/Plugins/FX/Niagara/Source/NiagaraEditor/Private/TypeEditorUtilities/NiagaraIntegerTypeEditorUtilities.cpp
+++ b/Engine/Plugins/FX/Niagara/Source/NiagaraEditor/Private/TypeEditorUtilities/NiagaraIntegerTypeEditorUtilities.cpp
@@ -108,6 +108,36 @@ public:
 				})
 			];
 		}
+		else if(WidgetCustomization.WidgetType == ENiagaraInputWidgetType::EnumStyle && WidgetCustomization.EnumStyleDropdownValues.Num() > 0)
+		{
+			TArray<SNiagaraNumericDropDown<int32>::FNamedValue> DropDownValues;
+			for(int32 EnumStyleValueIndex = 0; EnumStyleValueIndex < WidgetCustomization.EnumStyleDropdownValues.Num(); EnumStyleValueIndex++)
+			{
+				FText DisplayName = WidgetCustomization.EnumStyleDropdownValues[EnumStyleValueIndex].DisplayName;
+				FText Tooltip = WidgetCustomization.EnumStyleDropdownValues[EnumStyleValueIndex].Tooltip;
+						
+				DropDownValues.Add(SNiagaraNumericDropDown<int32>::FNamedValue(EnumStyleValueIndex, DisplayName, Tooltip));
+			}
+			
+			ChildSlot
+			[
+				SNew(SNiagaraNumericDropDown<int32>)
+				.DropDownValues(DropDownValues)
+				.bAllowTyping(false)
+				.bShowNamedValue(true)
+				.MinDesiredValueWidth(75)
+				.PillType(FNiagaraTypeDefinition::GetIntDef())
+				.Value_Lambda([this]() 
+				{ 
+					return IntValue;
+				})
+				.OnValueChanged_Lambda([this](int32 NewVal)
+				{
+					IntValue = NewVal;
+					ExecuteOnValueChanged();
+				})
+			];
+		}
 		else
 		{
 			TOptional<int32> MinValue;
diff --git a/Engine/Plugins/FX/Niagara/Source/NiagaraEditor/Private/ViewModels/Stack/NiagaraStackFunctionInput.cpp b/Engine/Plugins/FX/Niagara/Source/NiagaraEditor/Private/ViewModels/Stack/NiagaraStackFunctionInput.cpp
--- a/Engine/Plugins/FX/Niagara/Source/NiagaraEditor/Private/ViewModels/Stack/NiagaraStackFunctionInput.cpp
+++ b/Engine/Plugins/FX/Niagara/Source/NiagaraEditor/Private/ViewModels/Stack/NiagaraStackFunctionInput.cpp
@@ -666,6 +666,22 @@ void UNiagaraStackFunctionInput::RefreshChildrenInternal(const TArray<UNiagaraSt
 		}
 	}
 
+	// we check for local ints acting as enums whether the local value is allowed with the current metadata
+	// it can happen if a module version changes or a module in use is modified to lose options that were in use
+	if(InputMetaData.IsSet() && InputMetaData->WidgetCustomization.WidgetType == ENiagaraInputWidgetType::EnumStyle && GetInputType() == FNiagaraTypeDefinition::GetIntDef() && GetValueMode() == EValueMode::Local)
+	{
+		int32 LocalValue = *((int32*) GetLocalValueStruct()->GetStructMemory());
+		if(LocalValue < 0 || LocalValue >= InputMetaData->WidgetCustomization.EnumStyleDropdownValues.Num())
+		{
+			NewIssues.Add(FStackIssue(
+			EStackIssueSeverity::Error,
+			LOCTEXT("IntegerAsEnumOutsideOfRangeShort", "Invalid value"),
+			LOCTEXT("IntegerAsEnumOutsideOfRangeLong", "The chosen value is not valid. Please choose a new valid entry."),
+			GetStackEditorDataKey(),
+			false));
+		}
+	}
+	
 	if (InputValues.DynamicNode.IsValid())
 	{
 		FNiagaraStackGraphUtilities::CheckForDeprecatedScriptVersion(GetDynamicInputNode(), GetStackEditorDataKey(), GetUpgradeDynamicInputVersionFix(), NewIssues);
diff --git a/Engine/Plugins/FX/Niagara/Source/NiagaraEditor/Private/ViewModels/Stack/NiagaraStackFunctionInputCollection.cpp b/Engine/Plugins/FX/Niagara/Source/NiagaraEditor/Private/ViewModels/Stack/NiagaraStackFunctionInputCollection.cpp
--- a/Engine/Plugins/FX/Niagara/Source/NiagaraEditor/Private/ViewModels/Stack/NiagaraStackFunctionInputCollection.cpp
+++ b/Engine/Plugins/FX/Niagara/Source/NiagaraEditor/Private/ViewModels/Stack/NiagaraStackFunctionInputCollection.cpp
@@ -521,12 +521,18 @@ void UNiagaraStackFunctionInputCollection::AddInputToCategory(const FInputData&
 	// Try to find an existing category in the already processed children.
 	UNiagaraStackInputCategory* InputCategory = FindCurrentChildOfTypeByPredicate<UNiagaraStackInputCategory>(NewChildren,
 		[&](UNiagaraStackInputCategory* CurrentCategory) { return CurrentCategory->GetDisplayName().CompareTo(InputData.Category) == 0; });
-
+	
 	if (InputCategory == nullptr)
 	{
 		// If we haven't added any children to this category yet see if there is one that can be reused from the current children.
 		InputCategory = FindCurrentChildOfTypeByPredicate<UNiagaraStackInputCategory>(CurrentChildren,
 			[&](UNiagaraStackInputCategory* CurrentCategory) { return CurrentCategory->GetDisplayName().CompareTo(InputData.Category) == 0; });
+
+		if(bForceCompleteRebuild)
+		{
+			InputCategory = nullptr;
+		}
+		
 		if (InputCategory == nullptr)
 		{
 			// If we don't have a current child for this category make a new one.
@@ -800,7 +806,9 @@ void UNiagaraStackFunctionInputCollection::OnScriptApplied(UNiagaraScript* Niaga
 {
 	if(InputFunctionCallNode->FunctionScript == NiagaraScript)
 	{
+		bForceCompleteRebuild = true;
 		RefreshChildren();
+		bForceCompleteRebuild = false;
 	}
 }
 
diff --git a/Engine/Plugins/FX/Niagara/Source/NiagaraEditor/Private/Widgets/SNiagaraNumericDropDown.h b/Engine/Plugins/FX/Niagara/Source/NiagaraEditor/Private/Widgets/SNiagaraNumericDropDown.h
--- a/Engine/Plugins/FX/Niagara/Source/NiagaraEditor/Private/Widgets/SNiagaraNumericDropDown.h
+++ b/Engine/Plugins/FX/Niagara/Source/NiagaraEditor/Private/Widgets/SNiagaraNumericDropDown.h
@@ -69,12 +69,15 @@ public:
 public:
 	SLATE_BEGIN_ARGS( SNiagaraNumericDropDown<NumericType> )
 		: _DropDownValues()
+		, _bAllowTyping(true)
 		, _MinDesiredValueWidth( 40 )
 		, _bShowNamedValue(false)
 		{}
 
 		/** The values which are used to populate the drop down menu. */
-		SLATE_ARGUMENT( TArray<FNamedValue>, DropDownValues )
+		SLATE_ATTRIBUTE( TArray<FNamedValue>, DropDownValues )
+		/** If set to false, will disable the text box and only allow choosing a value from the dropdown. */
+		SLATE_ARGUMENT( bool, bAllowTyping )
 		/** Controls the minimum width for the text box portion of the control. */
 		SLATE_ATTRIBUTE( float, MinDesiredValueWidth )
 		/** Toggle to show the drop down text value if the value matches the numeric value. */
@@ -91,10 +94,32 @@ public:
 	void Construct( const FArguments& InArgs )
 	{
 		DropDownValues = InArgs._DropDownValues;
+		bAllowTyping = InArgs._bAllowTyping;
 		bShowNamedValue = InArgs._bShowNamedValue;
 		Value = InArgs._Value;
 		OnValueChanged = InArgs._OnValueChanged;
 
+		TSharedPtr<SWidget> ValueWidget;
+
+		if(bAllowTyping)
+		{
+			ValueWidget = SNew(SEditableTextBox)
+			.Padding(0)
+			.Style(FNiagaraEditorStyle::Get(), "NiagaraEditor.Stack.NumericDropdownInput")
+			.MinDesiredWidth( InArgs._MinDesiredValueWidth )
+			.RevertTextOnEscape(true)
+			.SelectAllTextWhenFocused(true)
+			.Text( this, &SNiagaraNumericDropDown<NumericType>::GetValueText )
+			.OnTextCommitted( this, &SNiagaraNumericDropDown<NumericType>::ValueTextComitted );
+		}
+		else
+		{
+			ValueWidget = SNew(STextBlock)
+			.MinDesiredWidth( InArgs._MinDesiredValueWidth )
+			.Text( this, &SNiagaraNumericDropDown<NumericType>::GetValueText )
+			.TextStyle(FNiagaraEditorStyle::Get(), "NiagaraEditor.Stack.IntegerAsEnum");
+		}
+		
 		ChildSlot
 		[
 			SNew( SHorizontalBox )
@@ -114,17 +139,12 @@ public:
 				.OnGetMenuContent( this, &SNiagaraNumericDropDown<NumericType>::BuildMenu )
 				.ButtonContent()
 				[
-					SNew(SEditableTextBox)
-					.Padding(0)
-					.Style(FNiagaraEditorStyle::Get(), "NiagaraEditor.Stack.NumericDropdownInput")
-					.MinDesiredWidth( InArgs._MinDesiredValueWidth )
-					.RevertTextOnEscape(true)
-					.SelectAllTextWhenFocused(true)
-					.Text( this, &SNiagaraNumericDropDown<NumericType>::GetValueText )
-					.OnTextCommitted( this, &SNiagaraNumericDropDown<NumericType>::ValueTextComitted )
+					ValueWidget.ToSharedRef()
 				]
 			]
 		];
+
+		SetToolTipText(TAttribute<FText>::CreateSP(this, &SNiagaraNumericDropDown::GetValueTooltip));
 	}
 
 private:
@@ -133,7 +153,7 @@ private:
 	{
 		if (bShowNamedValue.Get())
 		{
-			for ( FNamedValue DropDownValue : DropDownValues )
+			for ( FNamedValue DropDownValue : DropDownValues.Get() )
 			{
 				if (FMath::IsNearlyEqual(static_cast<float>(DropDownValue.GetValue()), static_cast<float>(Value.Get())))
 				{
@@ -154,11 +174,27 @@ private:
 		}
 	}
 
+	FText GetValueTooltip() const
+	{
+		if (bShowNamedValue.Get())
+		{
+			for ( FNamedValue DropDownValue : DropDownValues.Get() )
+			{
+				if (FMath::IsNearlyEqual(static_cast<float>(DropDownValue.GetValue()), static_cast<float>(Value.Get())))
+				{
+					return DropDownValue.GetDescription();
+				}
+			}
+		}
+		
+		return FText::GetEmpty();
+	}
+
 	TSharedRef<SWidget> BuildMenu()
 	{
 		FMenuBuilder MenuBuilder( true, NULL );
 
-		for ( FNamedValue DropDownValue : DropDownValues )
+		for ( FNamedValue DropDownValue : DropDownValues.Get() )
 		{
 			FUIAction MenuAction(FExecuteAction::CreateSP(this, &SNiagaraNumericDropDown::SetValue, DropDownValue.GetValue()));
 			MenuBuilder.AddMenuEntry(DropDownValue.GetName(), DropDownValue.GetDescription(), FSlateIcon(), MenuAction);
@@ -173,7 +209,8 @@ private:
 		OnValueChanged.ExecuteIfBound(InValue);
 	}
 
-	TArray<FNamedValue> DropDownValues;
+	TAttribute<TArray<FNamedValue>> DropDownValues;
+	bool bAllowTyping = true;
 	TAttribute<bool> bShowNamedValue;
 	TAttribute<NumericType> Value;
 	FOnValueChanged OnValueChanged;
diff --git a/Engine/Plugins/FX/Niagara/Source/NiagaraEditor/Public/NiagaraEditorUtilities.h b/Engine/Plugins/FX/Niagara/Source/NiagaraEditor/Public/NiagaraEditorUtilities.h
--- a/Engine/Plugins/FX/Niagara/Source/NiagaraEditor/Public/NiagaraEditorUtilities.h
+++ b/Engine/Plugins/FX/Niagara/Source/NiagaraEditor/Public/NiagaraEditorUtilities.h
@@ -502,7 +502,7 @@ namespace FNiagaraParameterUtilities
 	NIAGARAEDITOR_API TSharedRef<SWidget> GetParameterWidget(FNiagaraVariable Variable, FNiagaraVariableMetaData MetaData, FNiagaraParameterWidgetOptions Options);
 
 	/** Creates a tooltip based on a parameter. Also shows the value, if allocated and enabled. */
-	NIAGARAEDITOR_API TSharedRef<SToolTip> GetTooltipWidget(FNiagaraVariable Variable, bool bShowValue = true, TSharedPtr<SWidget> AdditionalVerticalWidget = nullptr,  TSharedPtr<SWidget> AdditionalHorizontalWidget = nullptr);
+	NIAGARAEDITOR_API TSharedRef<SToolTip> GetTooltipWidget(FNiagaraVariable Variable, bool bShowValue = true, TSharedPtr<SWidget> AdditionalVerticalWidget = nullptr);
 
 	NIAGARAEDITOR_API void FilterToRelevantStaticVariables(TConstArrayView<FNiagaraVariable> InVars, TArray<FNiagaraVariable>& OutVars, FName InOldEmitterAlias, FName InNewEmitterAlias, bool bFilterByEmitterAliasAndConvertToUnaliased);
 };
diff --git a/Engine/Plugins/FX/Niagara/Source/NiagaraEditor/Public/ViewModels/Stack/NiagaraStackFunctionInputCollection.h b/Engine/Plugins/FX/Niagara/Source/NiagaraEditor/Public/ViewModels/Stack/NiagaraStackFunctionInputCollection.h
--- a/Engine/Plugins/FX/Niagara/Source/NiagaraEditor/Public/ViewModels/Stack/NiagaraStackFunctionInputCollection.h
+++ b/Engine/Plugins/FX/Niagara/Source/NiagaraEditor/Public/ViewModels/Stack/NiagaraStackFunctionInputCollection.h
@@ -148,4 +148,7 @@ private:
 private:
 	UNiagaraNodeFunctionCall* ModuleNode;
 	UNiagaraNodeFunctionCall* InputFunctionCallNode;
+
+	/** If this is set to true, no children will be resued when RefreshChildren is called */
+	bool bForceCompleteRebuild = false;
 };
diff --git a/Engine/Plugins/FX/Niagara/Source/NiagaraEditorWidgets/Private/SNiagaraOverviewInlineParameterBox.cpp b/Engine/Plugins/FX/Niagara/Source/NiagaraEditorWidgets/Private/SNiagaraOverviewInlineParameterBox.cpp
--- a/Engine/Plugins/FX/Niagara/Source/NiagaraEditorWidgets/Private/SNiagaraOverviewInlineParameterBox.cpp
+++ b/Engine/Plugins/FX/Niagara/Source/NiagaraEditorWidgets/Private/SNiagaraOverviewInlineParameterBox.cpp
@@ -193,9 +193,11 @@ TSharedRef<SWidget> SNiagaraOverviewInlineParameterBox::GenerateParameterWidgetF
 	FText DisplayedText = ValueText;
 	const FLinearColor OriginalTypeColor = UEdGraphSchema_Niagara::GetTypeColor(Type);
 	FLinearColor ActualTypeColor = OriginalTypeColor;
+	
+	TSharedPtr<IToolTip> TooltipOverride;
 
-	// if we have an enum type, we want to override the type color to be gray instead
-	if(Type.GetEnum())
+	// if we have an enum type or integer acting as an enum, we want to override the type color to be gray instead
+	if(Type.GetEnum() || (Type == FNiagaraTypeDefinition::GetIntDef() && InputMetaData.IsSet() && InputMetaData->WidgetCustomization.WidgetType == ENiagaraInputWidgetType::EnumStyle))
 	{
 		ActualTypeColor = FLinearColor(0.02f, 0.02f, 0.02f, 1.f);
 	}
@@ -216,6 +218,48 @@ TSharedRef<SWidget> SNiagaraOverviewInlineParameterBox::GenerateParameterWidgetF
 			DisplayedColor = EnumParameterMetaData.bUseColorOverride ? EnumParameterMetaData.ColorOverride : DisplayedColor;
 		}
 	}
+	else if(Type == FNiagaraTypeDefinition::GetIntDef() && InputMetaData->WidgetCustomization.WidgetType == ENiagaraInputWidgetType::EnumStyle)
+	{
+		int32 Value = *(int32*)FunctionInput->GetLocalValueStruct()->GetStructMemory();
+		if(InputMetaData->WidgetCustomization.EnumStyleDropdownValues.IsValidIndex(Value))
+		{
+			FString ValueString = FString::FromInt(Value);
+			FText DisplayName = InputMetaData->WidgetCustomization.EnumStyleDropdownValues[Value].DisplayName;
+			FText TooltipActualValueText = FText::FormatOrdered(FText::FromString("({0})"), FText::FromString(ValueString));
+			DisplayedText = DisplayName.IsEmptyOrWhitespace() == false ? DisplayName : FText::FromString(ValueString);
+			
+			TSharedRef<SWidget> ValueOverrideTooltip = SNew(SHorizontalBox)
+			+ SHorizontalBox::Slot()
+			.AutoWidth()
+			.HAlign(HAlign_Center)
+			.VAlign(VAlign_Center)
+			.Padding(3.f)
+			[
+				SNew(STextBlock)
+				.Text(LOCTEXT("IntAsEnumParameterTooltipTextOverride", "Value: "))
+			]
+			+ SHorizontalBox::Slot()
+			.AutoWidth()
+			.HAlign(HAlign_Center)
+			.VAlign(VAlign_Center)
+			.Padding(3.f)
+			[
+				SNew(STextBlock)
+				.Text(DisplayedText)
+			]
+			+ SHorizontalBox::Slot()
+			.AutoWidth()
+			.HAlign(HAlign_Center)
+			.VAlign(VAlign_Center)
+			.Padding(3.f)
+			[
+				SNew(STextBlock)
+				.Text(TooltipActualValueText)
+			];
+			
+			TooltipOverride = FNiagaraParameterUtilities::GetTooltipWidget(Variable, false, ValueOverrideTooltip);
+		}
+	}
 	else if(Type == FNiagaraTypeDefinition::GetBoolDef() && InputMetaData->bEnableBoolOverride)
 	{
 		const FNiagaraBoolParameterMetaData& BoolParameterMetaData = InputMetaData->InlineParameterBoolOverride;
@@ -297,8 +341,15 @@ TSharedRef<SWidget> SNiagaraOverviewInlineParameterBox::GenerateParameterWidgetF
 	}
 	
 	// we construct a tooltip widget that shows the parameter the value is associated with
-	TSharedRef<SToolTip> TooltipWidget = FNiagaraParameterUtilities::GetTooltipWidget(Variable);
-	ParameterWidget->SetToolTip(TooltipWidget);
+	if(TooltipOverride != nullptr)
+	{
+		ParameterWidget->SetToolTip(TooltipOverride);
+	}
+	else
+	{
+		TSharedRef<SToolTip> TooltipWidget = FNiagaraParameterUtilities::GetTooltipWidget(Variable);
+		ParameterWidget->SetToolTip(TooltipWidget);
+	}
 	return ParameterWidget.ToSharedRef();
 }
 
