Fixed regression where case labels for nodes in bool mode list 'False, True' instead of 'True, False'

diff --git a/Engine/Plugins/FX/Niagara/Source/NiagaraEditor/Private/NiagaraNodeUsageSelector.cpp b/Engine/Plugins/FX/Niagara/Source/NiagaraEditor/Private/NiagaraNodeUsageSelector.cpp
--- a/Engine/Plugins/FX/Niagara/Source/NiagaraEditor/Private/NiagaraNodeUsageSelector.cpp
+++ b/Engine/Plugins/FX/Niagara/Source/NiagaraEditor/Private/NiagaraNodeUsageSelector.cpp
@@ -272,9 +272,12 @@ void UNiagaraNodeUsageSelector::AddWidgetsToInputBox(TSharedPtr<SVerticalBox> In
 	
 	int32 WidgetIndexOffset = 0;
 	// insert separators to make the grouping apparent.
-	for (int32 Idx = 0; Idx < OutputVars.Num() * NumOptionsPerVariable; Idx += OutputVars.Num())
+	TArray<int32> OptionValues = GetOptionValues();
+	ensure(OptionValues.Num() == GetNumberOfCases());
+	for (int32 Idx = 0; Idx < OutputVars.Num() * GetNumberOfCases(); Idx += OutputVars.Num())
 	{
-		const int32 CaseValue = Idx / OutputVars.Num();
+		const int32 OptionValueIndex = Idx / OutputVars.Num();
+		const int32 OptionValue = OptionValues[OptionValueIndex];
 		// adding an offset to account for the already added separators
 		InputBox->InsertSlot(Idx + WidgetIndexOffset)
 		.Padding(5.f, 5.f, 0.f, 2.f)
@@ -299,13 +302,13 @@ void UNiagaraNodeUsageSelector::AddWidgetsToInputBox(TSharedPtr<SVerticalBox> In
 				+ SWidgetSwitcher::Slot()
 				[
 					SNew(STextBlock)
-					.Text_UObject(this, &UNiagaraNodeUsageSelector::GetInputCaseNameAsText, CaseValue)
+					.Text_UObject(this, &UNiagaraNodeUsageSelector::GetInputCaseNameAsText, OptionValue)
 				]
 				+ SWidgetSwitcher::Slot()
 				[
 					SNew(SComboButton)
-					.OnGetMenuContent_UObject(this, &UNiagaraNodeUsageSelector::GetInputCaseContextOptions, CaseValue)
-					.ToolTipText_UObject(this, &UNiagaraNodeUsageSelector::GetInputCaseButtonTooltip, CaseValue)
+					.OnGetMenuContent_UObject(this, &UNiagaraNodeUsageSelector::GetInputCaseContextOptions, OptionValue)
+					.ToolTipText_UObject(this, &UNiagaraNodeUsageSelector::GetInputCaseButtonTooltip, OptionValue)
 					.ButtonContent()
 					[
 						SNew(SHorizontalBox)
@@ -313,7 +316,7 @@ void UNiagaraNodeUsageSelector::AddWidgetsToInputBox(TSharedPtr<SVerticalBox> In
 						.AutoWidth()
 						[
 							SNew(STextBlock)
-							.Text_UObject(this, &UNiagaraNodeUsageSelector::GetInputCaseNameAsText, CaseValue)
+							.Text_UObject(this, &UNiagaraNodeUsageSelector::GetInputCaseNameAsText, OptionValue)
 						]
 					]
 				]
