Allow Enum->Int conversions in niagara graphs.

diff --git a/Engine/Plugins/FX/Niagara/Source/Niagara/Private/NiagaraModule.cpp b/Engine/Plugins/FX/Niagara/Source/Niagara/Private/NiagaraModule.cpp
--- a/Engine/Plugins/FX/Niagara/Source/Niagara/Private/NiagaraModule.cpp
+++ b/Engine/Plugins/FX/Niagara/Source/Niagara/Private/NiagaraModule.cpp
@@ -1256,13 +1256,17 @@ bool FNiagaraTypeDefinition::TypesAreAssignable(const FNiagaraTypeDefinition& Ty
 		return false;
 	}
 
-	// Make sure that enums are not assignable to enums of different types or just plain ints
-	if (TypeInput.GetStruct() == TypeOutput.GetStruct() &&
-		TypeInput.GetEnum() != TypeOutput.GetEnum())
+	// Enums can only be assigned from enums of the same type.
+	if(TypeInput.IsEnum() && TypeInput.GetEnum() != TypeOutput.GetEnum())
 	{
 		return false;
 	}
 
+	//Enums can be assigned to enums of the same type or plain integers.
+	if(TypeOutput.IsEnum() && TypeInput.GetEnum() != TypeOutput.GetEnum() && TypeInput != FNiagaraTypeDefinition::GetIntDef())
+	{
+		return false;
+	}
 
 	// Static can go to static or static can go to nonstatic, but nonstatic can't go to static
 	if (((!TypeInput.IsStatic() && TypeOutput.IsStatic()) || (TypeInput.IsStatic() && TypeOutput.IsStatic())) && TypeInput.GetStruct() == TypeOutput.GetStruct() && TypeInput.GetEnum() == TypeOutput.GetEnum())
