- Generate a compile error if a UObject is written to more than once.

#rnx

diff --git a/Engine/Plugins/FX/Niagara/Source/NiagaraEditor/Private/NiagaraHlslTranslator.cpp b/Engine/Plugins/FX/Niagara/Source/NiagaraEditor/Private/NiagaraHlslTranslator.cpp
--- a/Engine/Plugins/FX/Niagara/Source/NiagaraEditor/Private/NiagaraHlslTranslator.cpp
+++ b/Engine/Plugins/FX/Niagara/Source/NiagaraEditor/Private/NiagaraHlslTranslator.cpp
@@ -5585,6 +5585,27 @@ void TNiagaraHlslTranslator<GraphBridge>::ParameterMapSet(const FParamMapSetNode
 						ResolvedName = Var.GetName();
 					}
 
+					// Ensure we don't have duplicate parameter map writes for the same parameter as it would be meaningless
+					for (int32 iObjectInfo=0; iObjectInfo < CompilationOutput.ScriptData.UObjectInfos.Num(); ++iObjectInfo)
+					{
+						if (iObjectInfo == Input)
+						{
+							continue;
+						}
+						const FNiagaraScriptUObjectCompileInfo& ObjectInfo = CompilationOutput.ScriptData.UObjectInfos[iObjectInfo];
+						if (ObjectInfo.RegisteredParameterMapWrites.Contains(ResolvedName))
+						{
+							Error(
+								FText::Format(
+									LOCTEXT("UObjectVariableWroteMoreThanOnce", "Object variable '{0}' has been wrote to more than once. This is not supported."),
+									FText::FromName(ResolvedName)
+								),
+								nullptr, nullptr
+							);
+							break;
+						}
+					}
+
 					FNiagaraScriptUObjectCompileInfo& Info = CompilationOutput.ScriptData.UObjectInfos[Input];
 					Info.RegisteredParameterMapWrites.AddUnique(ResolvedName);
 
