- Reduce implicit conversion of Base -> Var
- Add method to get an TOptional from the parameter store

diff --git a/Engine/Plugins/FX/Niagara/Source/Niagara/Public/NiagaraParameterStore.h b/Engine/Plugins/FX/Niagara/Source/Niagara/Public/NiagaraParameterStore.h
--- a/Engine/Plugins/FX/Niagara/Source/Niagara/Public/NiagaraParameterStore.h
+++ b/Engine/Plugins/FX/Niagara/Source/Niagara/Public/NiagaraParameterStore.h
@@ -381,7 +381,7 @@ public:
 	virtual void InitFromSource(const FNiagaraParameterStore* SrcStore, bool bNotifyAsDirty);
 
 	/** Gets the index of the passed parameter. If it is a data interface, this is an offset into the data interface table, otherwise a byte offset into the parameter data buffer. */
-	FORCEINLINE_DEBUGGABLE int32 IndexOf(const FNiagaraVariable& Parameter) const
+	FORCEINLINE_DEBUGGABLE int32 IndexOf(const FNiagaraVariableBase& Parameter) const
 	{
 		const int32* Off = FindParameterOffset(Parameter);
 		return Off ? *Off : (int32)INDEX_NONE;
@@ -391,7 +391,7 @@ public:
 
 	/** Gets the typed parameter data. */
 	template<typename T>
-	FORCEINLINE_DEBUGGABLE void GetParameterValue(T& OutValue, const FNiagaraVariable& Parameter)const
+	FORCEINLINE_DEBUGGABLE void GetParameterValue(T& OutValue, const FNiagaraVariableBase& Parameter)const
 	{
 		check(Parameter.GetSizeInBytes() == sizeof(T));
 		int32 Offset = IndexOf(Parameter);
@@ -402,7 +402,7 @@ public:
 	}
 
 	template<typename T>
-	FORCEINLINE_DEBUGGABLE T GetParameterValue(const FNiagaraVariable& Parameter)const
+	FORCEINLINE_DEBUGGABLE T GetParameterValue(const FNiagaraVariableBase& Parameter)const
 	{
 		check(Parameter.GetSizeInBytes() == sizeof(T));
 		int32 Offset = IndexOf(Parameter);
@@ -414,7 +414,20 @@ public:
 	}
 
 	template<typename T>
-	FORCEINLINE_DEBUGGABLE T GetParameterValueOrDefault(const FNiagaraVariable& Parameter, const T& DefaultValue) const
+	FORCEINLINE_DEBUGGABLE TOptional<T> GetParameterOptionalValue(const FNiagaraVariableBase& Parameter)const
+	{
+		check(Parameter.GetSizeInBytes() == sizeof(T));
+		TOptional<T> OptionalValue;
+		const int32 Offset = IndexOf(Parameter);
+		if (Offset != INDEX_NONE)
+		{
+			OptionalValue.Emplace(*(const T*)(GetParameterData(Offset)));
+		}
+		return OptionalValue;
+	}
+
+	template<typename T>
+	FORCEINLINE_DEBUGGABLE T GetParameterValueOrDefault(const FNiagaraVariableBase& Parameter, const T& DefaultValue) const
 	{
 		check(Parameter.GetSizeInBytes() == sizeof(T));
 		int32 Offset = IndexOf(Parameter);
@@ -427,7 +440,7 @@ public:
 	}
 
 	/** Returns the parameter data for the passed parameter if it exists in this store. Null if not. */
-	FORCEINLINE_DEBUGGABLE const uint8* GetParameterData(const FNiagaraVariable& Parameter)const
+	FORCEINLINE_DEBUGGABLE const uint8* GetParameterData(const FNiagaraVariableBase& Parameter)const
 	{
 		int32 Offset = IndexOf(Parameter);
 		if (Offset != INDEX_NONE)
