- Remove implicit conversions from FNiagaraVariableBase to FNiagaraVariable

diff --git a/Engine/Plugins/FX/Niagara/Source/Niagara/Classes/NiagaraDataSet.h b/Engine/Plugins/FX/Niagara/Source/Niagara/Classes/NiagaraDataSet.h
--- a/Engine/Plugins/FX/Niagara/Source/Niagara/Classes/NiagaraDataSet.h
+++ b/Engine/Plugins/FX/Niagara/Source/Niagara/Classes/NiagaraDataSet.h
@@ -290,15 +290,15 @@ public:
 
 	FORCEINLINE const TArray<FNiagaraVariableBase>& GetVariables() const { return CompiledData->Variables; }
 	FORCEINLINE uint32 GetNumVariables() const { return CompiledData->Variables.Num(); }
-	FORCEINLINE bool HasVariable(const FNiagaraVariable& Var) const { return CompiledData->Variables.Contains(Var); }
-	FORCEINLINE bool HasVariable(const FName& Name) const { return CompiledData->Variables.FindByPredicate([&](const FNiagaraVariable& VarInfo) { return VarInfo.GetName() == Name; }) != nullptr; }
+	FORCEINLINE bool HasVariable(const FNiagaraVariableBase& Var) const { return CompiledData->Variables.Contains(Var); }
+	FORCEINLINE bool HasVariable(const FName& Name) const { return CompiledData->Variables.FindByPredicate([&](const FNiagaraVariableBase& VarInfo) { return VarInfo.GetName() == Name; }) != nullptr; }
 	FORCEINLINE uint32 GetNumFloatComponents() const { return CompiledData->TotalFloatComponents; }
 	FORCEINLINE uint32 GetNumInt32Components() const { return CompiledData->TotalInt32Components; }
 	FORCEINLINE uint32 GetNumHalfComponents() const { return CompiledData->TotalHalfComponents; }
 
 	const TArray<FNiagaraVariableLayoutInfo>& GetVariableLayouts() const { return CompiledData->VariableLayouts; }
-	const FNiagaraVariableLayoutInfo* GetVariableLayout(const FNiagaraVariable& Var) const;
-	bool GetVariableComponentOffsets(const FNiagaraVariable& Var, int32 &FloatStart, int32 &IntStart, int32& HalfStart) const;
+	const FNiagaraVariableLayoutInfo* GetVariableLayout(const FNiagaraVariableBase& Var) const;
+	bool GetVariableComponentOffsets(const FNiagaraVariableBase& Var, int32 &FloatStart, int32 &IntStart, int32& HalfStart) const;
 
 	void CopyTo(FNiagaraDataSet& Other, int32 StartIdx = 0, int32 NumInstances = INDEX_NONE, bool bResetOther=true)const;
 
diff --git a/Engine/Plugins/FX/Niagara/Source/Niagara/Classes/NiagaraDataSetDebugAccessor.h b/Engine/Plugins/FX/Niagara/Source/Niagara/Classes/NiagaraDataSetDebugAccessor.h
--- a/Engine/Plugins/FX/Niagara/Source/Niagara/Classes/NiagaraDataSetDebugAccessor.h
+++ b/Engine/Plugins/FX/Niagara/Source/Niagara/Classes/NiagaraDataSetDebugAccessor.h
@@ -68,8 +68,8 @@ struct FNiagaraDataSetDebugAccessor
 	bool IsEnum() const { return bIsEnum; }
 	bool IsBool() const { return bIsBool; }
 
-	static bool ValidateDataBuffer(const FNiagaraDataSetCompiledData& CompiledData, const FNiagaraDataBuffer* DataBuffer, uint32 iInstance, TFunction<void(const FNiagaraVariable&, int32)> ErrorCallback);
-	static bool ValidateDataBuffer(const FNiagaraDataSetCompiledData& CompiledData, const FNiagaraDataBuffer* DataBuffer, TFunction<void(const FNiagaraVariable&, uint32, int32)> ErrorCallback);
+	static bool ValidateDataBuffer(const FNiagaraDataSetCompiledData& CompiledData, const FNiagaraDataBuffer* DataBuffer, uint32 iInstance, TFunction<void(const FNiagaraVariableBase&, int32)> ErrorCallback);
+	static bool ValidateDataBuffer(const FNiagaraDataSetCompiledData& CompiledData, const FNiagaraDataBuffer* DataBuffer, TFunction<void(const FNiagaraVariableBase&, uint32, int32)> ErrorCallback);
 
 private:
 	FName					VariableName;
diff --git a/Engine/Plugins/FX/Niagara/Source/Niagara/Classes/NiagaraSystem.h b/Engine/Plugins/FX/Niagara/Source/Niagara/Classes/NiagaraSystem.h
--- a/Engine/Plugins/FX/Niagara/Source/Niagara/Classes/NiagaraSystem.h
+++ b/Engine/Plugins/FX/Niagara/Source/Niagara/Classes/NiagaraSystem.h
@@ -787,7 +787,7 @@ private:
 	void InitEmitterVariableAliasNames(FNiagaraEmitterCompiledData& EmitterCompiledDataToInit, const UNiagaraEmitter* InAssociatedEmitter);
 
 	/** Helper for generating aliased FNiagaraVariable names for the Emitter they are associated with. */
-	const FName GetEmitterVariableAliasName(const FNiagaraVariable& InEmitterVar, const UNiagaraEmitter* InEmitter) const;
+	const FName GetEmitterVariableAliasName(const FNiagaraVariableBase& InEmitterVar, const UNiagaraEmitter* InEmitter) const;
 
 	/** Helper for filling in attribute datasets per emitter. */
 	void InitEmitterDataSetCompiledData(FNiagaraDataSetCompiledData& DataSetToInit, const FNiagaraEmitterHandle& InAssociatedEmitterHandle);
diff --git a/Engine/Plugins/FX/Niagara/Source/Niagara/Private/NiagaraDataSet.cpp b/Engine/Plugins/FX/Niagara/Source/Niagara/Private/NiagaraDataSet.cpp
--- a/Engine/Plugins/FX/Niagara/Source/Niagara/Private/NiagaraDataSet.cpp
+++ b/Engine/Plugins/FX/Niagara/Source/Niagara/Private/NiagaraDataSet.cpp
@@ -473,13 +473,13 @@ void FNiagaraDataSet::AllocateGPUFreeIDs(uint32 InNumInstances, FRHICommandList&
 #endif
 }
 
-const FNiagaraVariableLayoutInfo* FNiagaraDataSet::GetVariableLayout(const FNiagaraVariable& Var)const
+const FNiagaraVariableLayoutInfo* FNiagaraDataSet::GetVariableLayout(const FNiagaraVariableBase& Var)const
 {
 	int32 VarLayoutIndex = GetVariables().IndexOfByKey(Var);
 	return VarLayoutIndex != INDEX_NONE ? &CompiledData->VariableLayouts[VarLayoutIndex] : nullptr;
 }
 
-bool FNiagaraDataSet::GetVariableComponentOffsets(const FNiagaraVariable& Var, int32 &FloatStart, int32 &IntStart, int32& HalfStart) const
+bool FNiagaraDataSet::GetVariableComponentOffsets(const FNiagaraVariableBase& Var, int32 &FloatStart, int32 &IntStart, int32& HalfStart) const
 {
 	const FNiagaraVariableLayoutInfo *Info = GetVariableLayout(Var);
 	if (Info)
@@ -1050,7 +1050,7 @@ void FNiagaraDataBuffer::CopyToUnrelated(FNiagaraDataBuffer& DestBuffer, int32 S
 
 		for(int32 VarIdx = 0; VarIdx < DestCompiledData.Variables.Num(); ++VarIdx)
 		{
-			const FNiagaraVariable& DestVar = DestCompiledData.Variables[VarIdx];
+			const FNiagaraVariableBase& DestVar = DestCompiledData.Variables[VarIdx];
 			const FNiagaraVariableLayoutInfo& DestVarLayout = DestCompiledData.VariableLayouts[VarIdx];
 
 			const int32 SrcVarIdx = CompiledData.Variables.IndexOfByKey(DestVar);
diff --git a/Engine/Plugins/FX/Niagara/Source/Niagara/Private/NiagaraDataSetDebugAccessor.cpp b/Engine/Plugins/FX/Niagara/Source/Niagara/Private/NiagaraDataSetDebugAccessor.cpp
--- a/Engine/Plugins/FX/Niagara/Source/Niagara/Private/NiagaraDataSetDebugAccessor.cpp
+++ b/Engine/Plugins/FX/Niagara/Source/Niagara/Private/NiagaraDataSetDebugAccessor.cpp
@@ -16,7 +16,7 @@ bool FNiagaraDataSetDebugAccessor::Init(const FNiagaraDataSetCompiledData& Compi
 
 	for (int32 i = 0; i < CompiledData.Variables.Num(); ++i)
 	{
-		const FNiagaraVariable& Variable = CompiledData.Variables[i];
+		const FNiagaraVariableBase& Variable = CompiledData.Variables[i];
 		if (Variable.GetName() != VariableName)
 		{
 			continue;
@@ -68,7 +68,7 @@ int32 FNiagaraDataSetDebugAccessor::ReadInt(const FNiagaraDataBuffer* DataBuffer
 	return 0;
 }
 
-bool FNiagaraDataSetDebugAccessor::ValidateDataBuffer(const FNiagaraDataSetCompiledData& CompiledData, const FNiagaraDataBuffer* DataBuffer, uint32 iInstance, TFunction<void(const FNiagaraVariable&, int32)> ErrorCallback)
+bool FNiagaraDataSetDebugAccessor::ValidateDataBuffer(const FNiagaraDataSetCompiledData& CompiledData, const FNiagaraDataBuffer* DataBuffer, uint32 iInstance, TFunction<void(const FNiagaraVariableBase&, int32)> ErrorCallback)
 {
 	bool bIsValid = true;
 
@@ -122,7 +122,7 @@ bool FNiagaraDataSetDebugAccessor::ValidateDataBuffer(const FNiagaraDataSetCompi
 	return bIsValid;
 }
 
-bool FNiagaraDataSetDebugAccessor::ValidateDataBuffer(const FNiagaraDataSetCompiledData& CompiledData, const FNiagaraDataBuffer* DataBuffer, TFunction<void(const FNiagaraVariable&, uint32, int32)> ErrorCallback)
+bool FNiagaraDataSetDebugAccessor::ValidateDataBuffer(const FNiagaraDataSetCompiledData& CompiledData, const FNiagaraDataBuffer* DataBuffer, TFunction<void(const FNiagaraVariableBase&, uint32, int32)> ErrorCallback)
 {
 	bool bIsValid = true;
 
diff --git a/Engine/Plugins/FX/Niagara/Source/Niagara/Private/NiagaraDebugHud.cpp b/Engine/Plugins/FX/Niagara/Source/Niagara/Private/NiagaraDebugHud.cpp
--- a/Engine/Plugins/FX/Niagara/Source/Niagara/Private/NiagaraDebugHud.cpp
+++ b/Engine/Plugins/FX/Niagara/Source/Niagara/Private/NiagaraDebugHud.cpp
@@ -2475,7 +2475,7 @@ void FNiagaraDebugHud::DrawValidation(class FNiagaraWorldManager* WorldManager,
 				SystemSimulation->MainDataSet.GetCompiledData(),
 				SystemSimulation->MainDataSet.GetCurrentData(),
 				SystemInstance->GetSystemInstanceIndex(),
-				[&](const FNiagaraVariable& Variable, int32 ComponentIndex)
+				[&](const FNiagaraVariableBase& Variable, int32 ComponentIndex)
 				{
 					auto& ValidationError = GetValidationErrorInfo(NiagaraComponent);
 					ValidationError.SystemVariablesWithErrors.AddUnique(Variable.GetName());
@@ -2509,7 +2509,7 @@ void FNiagaraDebugHud::DrawValidation(class FNiagaraWorldManager* WorldManager,
 				FNiagaraDataSetDebugAccessor::ValidateDataBuffer(
 					ParticleDataSet->GetCompiledData(),
 					DataBuffer,
-					[&](const FNiagaraVariable& Variable, int32 InstanceIndex, int32 ComponentIndex)
+					[&](const FNiagaraVariableBase& Variable, int32 InstanceIndex, int32 ComponentIndex)
 					{
 						auto& ValidationError = GetValidationErrorInfo(NiagaraComponent);
 						const FName EmitterName(*EmitterInstance->GetCachedEmitter().Emitter->GetUniqueEmitterName());
diff --git a/Engine/Plugins/FX/Niagara/Source/Niagara/Private/NiagaraSimCacheHelper.h b/Engine/Plugins/FX/Niagara/Source/Niagara/Private/NiagaraSimCacheHelper.h
--- a/Engine/Plugins/FX/Niagara/Source/Niagara/Private/NiagaraSimCacheHelper.h
+++ b/Engine/Plugins/FX/Niagara/Source/Niagara/Private/NiagaraSimCacheHelper.h
@@ -136,7 +136,7 @@ struct FNiagaraSimCacheHelper
 		for (int32 i = 0; i < CompiledData.Variables.Num(); ++i)
 		{
 			const FNiagaraVariableLayoutInfo& DataSetVariableLayout = CompiledData.VariableLayouts[i];
-			const FNiagaraVariable& DataSetVariable = CompiledData.Variables[i];
+			const FNiagaraVariableBase& DataSetVariable = CompiledData.Variables[i];
 			if (ExplicitCaptureAttributes.Num() == 0 || ExplicitCaptureAttributes.Contains(DataSetVariable.GetName()))
 			{
 				CacheToDataSetVariables.Add(i);
diff --git a/Engine/Plugins/FX/Niagara/Source/Niagara/Private/NiagaraSystem.cpp b/Engine/Plugins/FX/Niagara/Source/Niagara/Private/NiagaraSystem.cpp
--- a/Engine/Plugins/FX/Niagara/Source/Niagara/Private/NiagaraSystem.cpp
+++ b/Engine/Plugins/FX/Niagara/Source/Niagara/Private/NiagaraSystem.cpp
@@ -1593,7 +1593,7 @@ void UNiagaraSystem::FindDataInterfaceDependencies(FVersionedNiagaraEmitterData*
 			if (StoreDataInterfaces.Num() > 0)
 			{
 				auto FindCachedDefaultDI =
-					[](UNiagaraScript* Script, const FNiagaraVariable& Variable) -> UNiagaraDataInterface*
+					[](UNiagaraScript* Script, const FNiagaraVariableBase& Variable) -> UNiagaraDataInterface*
 				{
 					if (Script)
 					{
@@ -2829,7 +2829,7 @@ void UNiagaraSystem::InitEmitterVariableAliasNames(FNiagaraEmitterCompiledData&
 	EmitterCompiledDataToInit.EmitterTotalSpawnedParticlesVar.SetName(GetEmitterVariableAliasName(SYS_PARAM_ENGINE_EMITTER_TOTAL_SPAWNED_PARTICLES, InAssociatedEmitter));
 }
 
-const FName UNiagaraSystem::GetEmitterVariableAliasName(const FNiagaraVariable& InEmitterVar, const UNiagaraEmitter* InEmitter) const
+const FName UNiagaraSystem::GetEmitterVariableAliasName(const FNiagaraVariableBase& InEmitterVar, const UNiagaraEmitter* InEmitter) const
 {
 	return FName(*InEmitterVar.GetName().ToString().Replace(TEXT("Emitter."), *(InEmitter->GetUniqueEmitterName() + TEXT("."))));
 }
diff --git a/Engine/Plugins/FX/Niagara/Source/Niagara/Public/NiagaraSystemSimulation.h b/Engine/Plugins/FX/Niagara/Source/Niagara/Public/NiagaraSystemSimulation.h
--- a/Engine/Plugins/FX/Niagara/Source/Niagara/Public/NiagaraSystemSimulation.h
+++ b/Engine/Plugins/FX/Niagara/Source/Niagara/Public/NiagaraSystemSimulation.h
@@ -80,7 +80,7 @@ struct FNiagaraParameterStoreToDataSetBinding
 		const auto& DataSetVariableLayouts = DataSet.GetVariableLayouts();
 		for (int i=0; i < DataSetVariables.Num(); ++i)
 		{
-			const FNiagaraVariable& Var = DataSetVariables[i];
+			const FNiagaraVariableBase& Var = DataSetVariables[i];
 			const int32* ParameterOffsetPtr = ParameterStore.FindParameterOffset(Var, true);
 			if (ParameterOffsetPtr == nullptr)
 			{
diff --git a/Engine/Plugins/FX/Niagara/Source/NiagaraEditor/Private/NiagaraBakerRenderer.cpp b/Engine/Plugins/FX/Niagara/Source/NiagaraEditor/Private/NiagaraBakerRenderer.cpp
--- a/Engine/Plugins/FX/Niagara/Source/NiagaraEditor/Private/NiagaraBakerRenderer.cpp
+++ b/Engine/Plugins/FX/Niagara/Source/NiagaraEditor/Private/NiagaraBakerRenderer.cpp
@@ -250,7 +250,7 @@ void FNiagaraBakerOutputBindingHelper::GetParticleAttributeBindings(TArray<FNiag
 			const FNiagaraDataSetCompiledData& ParticleDataSet = AllEmitterCompiledData[EmitterIndex]->DataSetCompiledData;
 			for (int32 iVariable = 0; iVariable < ParticleDataSet.VariableLayouts.Num(); ++iVariable)
 			{
-				const FNiagaraVariable& Variable = ParticleDataSet.Variables[iVariable];
+				const FNiagaraVariableBase& Variable = ParticleDataSet.Variables[iVariable];
 				const FNiagaraVariableLayoutInfo& VariableLayout = ParticleDataSet.VariableLayouts[iVariable];
 				if (VariableLayout.GetNumFloatComponents() > 0)
 				{
