- Add a cvar to bypass disabling partial reads if we have an particle read
- This is temporary until we can detect that the particle read was used in this stage vs another stage

diff --git a/Engine/Plugins/FX/Niagara/Source/Niagara/Private/NiagaraDataInterfaceParticleRead.cpp b/Engine/Plugins/FX/Niagara/Source/Niagara/Private/NiagaraDataInterfaceParticleRead.cpp
--- a/Engine/Plugins/FX/Niagara/Source/Niagara/Private/NiagaraDataInterfaceParticleRead.cpp
+++ b/Engine/Plugins/FX/Niagara/Source/Niagara/Private/NiagaraDataInterfaceParticleRead.cpp
@@ -102,6 +102,15 @@ enum class ENiagaraParticleDataValueType : uint8
 	Position
 };
 
+//-TODO: This is a workaround for when a read is safe because the dispatches can not overlap.  UE-187062
+static bool GNiagaraParticleReadIgnoreUnsafeReads = false;
+static FAutoConsoleVariableRef CVarNiagaraParticleReadIgnoreUnsafeReads(
+	TEXT("fx.Niagara.ParticleRead.IgnoreUnsafeReads"),
+	GNiagaraParticleReadIgnoreUnsafeReads,
+	TEXT("When enabled we will allow unsafe reads to compile with a warning, however the read result will be invalid."),
+	ECVF_Default
+);
+
 DECLARE_INTRINSIC_TYPE_LAYOUT(ENiagaraParticleDataValueType);
 
 struct FNiagaraDataInterfaceParametersCS_ParticleRead : public FNiagaraDataInterfaceParametersCS
@@ -2187,7 +2196,10 @@ void UNiagaraDataInterfaceParticleRead::SetShaderParameters(const FNiagaraDataIn
 		{
 			if (StageMetaData->bPartialParticleUpdate)
 			{
-				UE_LOG(LogNiagara, Error, TEXT("Particle read DI reading self '%s' on stage '%s' is unsafe, please disable partial writes on the stage."), *InstanceData->DebugSourceName, *StageMetaData->SimulationStageName.ToString());
+				if (GNiagaraParticleReadIgnoreUnsafeReads == false)
+				{
+					UE_LOG(LogNiagara, Error, TEXT("Particle read DI reading self '%s' on stage '%s' is unsafe, please disable partial writes on the stage."), *InstanceData->DebugSourceName, *StageMetaData->SimulationStageName.ToString());
+				}
 				SetErrorParams(true);
 				return;
 			}
diff --git a/Engine/Plugins/FX/Niagara/Source/NiagaraEditor/Private/NiagaraHlslTranslator.cpp b/Engine/Plugins/FX/Niagara/Source/NiagaraEditor/Private/NiagaraHlslTranslator.cpp
--- a/Engine/Plugins/FX/Niagara/Source/NiagaraEditor/Private/NiagaraHlslTranslator.cpp
+++ b/Engine/Plugins/FX/Niagara/Source/NiagaraEditor/Private/NiagaraHlslTranslator.cpp
@@ -1347,7 +1347,16 @@ FNiagaraTranslateResults TNiagaraHlslTranslator<GraphBridge>::Translate(const FN
 						{
 							if(DataInterfaceData.ReadsEmitterParticleData.Contains(CompileData->EmitterUniqueName))
 							{
-								SimStage.bPartialParticleUpdate = false;
+								IConsoleVariable* CVarIgnoreUnsafeReads = IConsoleManager::Get().FindConsoleVariable(TEXT("fx.Niagara.ParticleRead.IgnoreUnsafeReads"));
+								const bool bIgnoredUnsafeReads = CVarIgnoreUnsafeReads ? CVarIgnoreUnsafeReads->GetBool() : false;
+								if (bIgnoredUnsafeReads && CompileSimStageData.bParticleIterationStateEnabled)
+								{
+									Warning(LOCTEXT("ReadsSelfEmitterDataIterationStage", "Reading own emitter data using iteration state can cause a race."), nullptr, nullptr);
+								}
+								else
+								{
+									SimStage[Index].bPartialParticleUpdate = false;
+								}
 								break;
 							}
 						}
