- Remove nullptr outputs from baker on postload
- Unsure how as asset got into this state, but this will fix the issue

#rnx

diff --git a/Engine/Plugins/FX/Niagara/Source/Niagara/Private/NiagaraBakerSettings.cpp b/Engine/Plugins/FX/Niagara/Source/Niagara/Private/NiagaraBakerSettings.cpp
--- a/Engine/Plugins/FX/Niagara/Source/Niagara/Private/NiagaraBakerSettings.cpp
+++ b/Engine/Plugins/FX/Niagara/Source/Niagara/Private/NiagaraBakerSettings.cpp
@@ -202,7 +202,6 @@ void UNiagaraBakerSettings::PostLoad()
 		OutputTextures_DEPRECATED.Empty();
 	}
 
-
 	if (NiagaraVersion < FNiagaraCustomVersion::AddBakerCameraBookmarks)
 	{
 		for ( int32 i=0; i < int(ENiagaraBakerViewMode::Num); ++i )
@@ -218,6 +217,10 @@ void UNiagaraBakerSettings::PostLoad()
 		}
 		CurrentCameraIndex = int(CameraViewportMode_DEPRECATED);
 	}
+
+	// Sanitize outputs are not nullptr
+	// This should not occur but an asset has been seen to get into this state
+	Outputs.Remove(nullptr);
 }
 
 #if WITH_EDITORONLY_DATA
