Optimize FixupVMAssembly by replacing __OP__ in one pass rather than independent passes.  Saves about 5s of main thread time in the cook of engine test.

diff --git a/Engine/Plugins/FX/Niagara/Source/NiagaraEditor/Private/NiagaraCompiler.cpp b/Engine/Plugins/FX/Niagara/Source/NiagaraEditor/Private/NiagaraCompiler.cpp
--- a/Engine/Plugins/FX/Niagara/Source/NiagaraEditor/Private/NiagaraCompiler.cpp
+++ b/Engine/Plugins/FX/Niagara/Source/NiagaraEditor/Private/NiagaraCompiler.cpp
@@ -1902,12 +1902,76 @@ int32 FHlslNiagaraCompiler::CompileScript(const FStringView GroupName, const FNi
 
 void FHlslNiagaraCompiler::FixupVMAssembly(FString& Asm)
 {
-	for (int32 OpCode = 0; OpCode < VectorVM::GetNumOpCodes(); ++OpCode)
+	const TCHAR* OpTag = TEXT("__OP__");
+	const TCHAR* OpDelimTags[] = { TEXT("("), TEXT(";") };
+	const int32 OpTagLen = FCString::Strlen(OpTag);
+	const int32 OpCount = VectorVM::GetNumOpCodes();
+	constexpr int32 MaxDigitLength = 8;
+
+	const FString OriginalAsm = MoveTemp(Asm);
+	const FStringView OriginalAsmView(OriginalAsm);
+
+	Asm.Reserve(OriginalAsm.Len());
+
+	int32 SearchStartLocation = 0;
+	while (SearchStartLocation != INDEX_NONE)
 	{
-		//TODO: reduce string ops here by moving these out to a static list.
-		FString ToReplace = TEXT("__OP__") + LexToString(OpCode) + TEXT("(");
-		FString Replacement = VectorVM::GetOpName(EVectorVMOp(OpCode)) + TEXT("(");
-		Asm.ReplaceInline(*ToReplace, *Replacement);
+		const int32 OpStartLocation = OriginalAsmView.Find(OpTag, SearchStartLocation);
+
+		// append the string leading up to the OpCode
+		const int32 SubStringLength = OpStartLocation == INDEX_NONE ? MAX_int32 : OpStartLocation - SearchStartLocation;
+		if (SubStringLength)
+		{
+			Asm.Append(OriginalAsmView.Mid(SearchStartLocation, SubStringLength));
+			SearchStartLocation = OpStartLocation;
+		}
+
+		if (OpStartLocation == INDEX_NONE)
+		{
+			break;
+		}
+
+		// find the next delimiter
+		const int32 OpEndLocation = OpStartLocation + OpTagLen;
+		FStringView DelimSearchView = OriginalAsmView.Mid(OpEndLocation, MaxDigitLength + 1);
+		int32 DelimOffset = INDEX_NONE;
+
+		for (const TCHAR* DelimTag : OpDelimTags)
+		{
+			int32 NextLocation = DelimSearchView.Find(DelimTag);
+			if (NextLocation != INDEX_NONE && (DelimOffset == INDEX_NONE || NextLocation < DelimOffset))
+			{
+				DelimOffset = NextLocation;
+			}
+		}
+
+		// check if the intervening spaces are digits (op index)
+		int32 OpIndex = INDEX_NONE;
+
+		if (DelimOffset != INDEX_NONE)
+		{
+			FStringView OpIndexStrView = OriginalAsmView.Mid(OpEndLocation, DelimOffset);
+			if (OpIndexStrView.Len() <= MaxDigitLength)
+			{
+				TStringBuilder<MaxDigitLength + 1> OpIndexStr;
+				OpIndexStr.Append(OpIndexStrView);
+				if (!LexTryParseString<int32>(OpIndex, *OpIndexStr))
+				{
+					OpIndex = INDEX_NONE;
+				}
+			}
+		}
+
+		if (OpIndex != INDEX_NONE && OpIndex < OpCount)
+		{
+			Asm.Append(VectorVM::GetOpName(EVectorVMOp(OpIndex)));
+			SearchStartLocation = OpEndLocation + DelimOffset;
+		}
+		else
+		{
+			Asm.Append(OriginalAsmView.Mid(OpStartLocation, OpTagLen));
+			SearchStartLocation = OpEndLocation;
+		}
 	}
 }
 
