- Add a way to visualize emitter bounds inside the system editor

diff --git a/Engine/Plugins/FX/Niagara/Source/Niagara/Classes/NiagaraEmitterHandle.h b/Engine/Plugins/FX/Niagara/Source/Niagara/Classes/NiagaraEmitterHandle.h
--- a/Engine/Plugins/FX/Niagara/Source/Niagara/Classes/NiagaraEmitterHandle.h
+++ b/Engine/Plugins/FX/Niagara/Source/Niagara/Classes/NiagaraEmitterHandle.h
@@ -83,6 +83,8 @@ public:
 
 	void ClearEmitter();
 
+	bool GetDebugShowBounds() const { return bDebugShowBounds; }
+	void SetDebugShowBounds(bool bShowBounds) { bDebugShowBounds = bShowBounds; }
 #endif
 public:
 	/** A static const invalid handle. */
@@ -121,6 +123,8 @@ private:
 	/** The copied instance of the emitter this handle references. */
 	UPROPERTY()
 	TObjectPtr<UNiagaraEmitter> Instance_DEPRECATED;
+
+	bool bDebugShowBounds = false;
 #endif
 
 	/** The copied instance of the emitter this handle references. */
diff --git a/Engine/Plugins/FX/Niagara/Source/NiagaraEditor/Private/Toolkits/NiagaraSystemToolkit.cpp b/Engine/Plugins/FX/Niagara/Source/NiagaraEditor/Private/Toolkits/NiagaraSystemToolkit.cpp
--- a/Engine/Plugins/FX/Niagara/Source/NiagaraEditor/Private/Toolkits/NiagaraSystemToolkit.cpp
+++ b/Engine/Plugins/FX/Niagara/Source/NiagaraEditor/Private/Toolkits/NiagaraSystemToolkit.cpp
@@ -614,6 +614,24 @@ TSharedRef<SWidget> FNiagaraSystemToolkit::GenerateBoundsMenuContent(TSharedRef<
 	
 	MenuBuilder.AddMenuEntry(FNiagaraEditorCommands::Get().ToggleBounds_SetFixedBounds_SelectedEmitters);
 
+	MenuBuilder.AddMenuSeparator();
+
+	for (TSharedRef<FNiagaraEmitterHandleViewModel> EmitterViewModel : SystemViewModel->GetEmitterHandleViewModels())
+	{
+		MenuBuilder.AddMenuEntry(
+			FText::Format(LOCTEXT("ShowEmitterBounds", "Show {0} Emitter Bounds"), EmitterViewModel->GetNameText()),
+			LOCTEXT("ShowEmitterBoundsTooltip", "Shows the emitter bounds."),
+			FSlateIcon(),
+			FUIAction(
+				FExecuteAction::CreateSP(this, &FNiagaraSystemToolkit::ToggleShowEmitterBounds, EmitterViewModel),
+				FCanExecuteAction::CreateSP(this, &FNiagaraSystemToolkit::CanShowEmitterBounds),
+				FIsActionChecked::CreateSP(this, &FNiagaraSystemToolkit::IsShowEmitterBounds, EmitterViewModel)
+			),
+			NAME_None,
+			EUserInterfaceActionType::ToggleButton
+		);
+	}
+
 	return MenuBuilder.MakeWidget();
 }
 
@@ -873,6 +891,25 @@ bool FNiagaraSystemToolkit::IsToggleBoundsChecked() const
 	return IsDrawOptionEnabled(SNiagaraSystemViewport::Bounds);
 }
 
+bool FNiagaraSystemToolkit::CanShowEmitterBounds() const
+{
+	return SystemViewModel->GetSystem().bFixedBounds == false;
+}
+
+bool FNiagaraSystemToolkit::IsShowEmitterBounds(TSharedRef<FNiagaraEmitterHandleViewModel> EmitterViewModel) const
+{
+	FNiagaraEmitterHandle* EmitterHandle = EmitterViewModel->GetEmitterHandle();
+	return EmitterHandle ? EmitterHandle->GetDebugShowBounds() : false;
+}
+
+void FNiagaraSystemToolkit::ToggleShowEmitterBounds(TSharedRef<FNiagaraEmitterHandleViewModel> EmitterViewModel)
+{
+	if (FNiagaraEmitterHandle* EmitterHandle = EmitterViewModel->GetEmitterHandle())
+	{
+		EmitterHandle->SetDebugShowBounds(!EmitterHandle->GetDebugShowBounds());
+	}
+}
+
 void FNiagaraSystemToolkit::ToggleDrawOption(int32 Element)
 {
 	if (Viewport.IsValid() && Viewport->GetViewportClient().IsValid())
diff --git a/Engine/Plugins/FX/Niagara/Source/NiagaraEditor/Private/Toolkits/NiagaraSystemToolkit.h b/Engine/Plugins/FX/Niagara/Source/NiagaraEditor/Private/Toolkits/NiagaraSystemToolkit.h
--- a/Engine/Plugins/FX/Niagara/Source/NiagaraEditor/Private/Toolkits/NiagaraSystemToolkit.h
+++ b/Engine/Plugins/FX/Niagara/Source/NiagaraEditor/Private/Toolkits/NiagaraSystemToolkit.h
@@ -36,6 +36,7 @@ class UNiagaraSequence;
 struct FAssetData;
 class FMenuBuilder;
 class ISequencer;
+class FNiagaraEmitterHandleViewModel;
 class FNiagaraMessageLogViewModel;
 class FNiagaraSystemToolkitParameterPanelViewModel;
 class FNiagaraSystemToolkitParameterDefinitionsPanelViewModel;
@@ -157,6 +158,10 @@ protected:
 	void ToggleDrawOption(int32 Element);
 	bool IsDrawOptionEnabled(int32 Element) const;
 
+	bool CanShowEmitterBounds() const;
+	bool IsShowEmitterBounds(TSharedRef<FNiagaraEmitterHandleViewModel> EmitterViewModel) const;
+	void ToggleShowEmitterBounds(TSharedRef<FNiagaraEmitterHandleViewModel> EmitterViewModel);
+
 	void OpenDebugHUD();
 	void OpenDebugOutliner();
 	void OpenAttributeSpreadsheet();
diff --git a/Engine/Plugins/FX/Niagara/Source/NiagaraEditor/Private/Widgets/SNiagaraSystemViewport.cpp b/Engine/Plugins/FX/Niagara/Source/NiagaraEditor/Private/Widgets/SNiagaraSystemViewport.cpp
--- a/Engine/Plugins/FX/Niagara/Source/NiagaraEditor/Private/Widgets/SNiagaraSystemViewport.cpp
+++ b/Engine/Plugins/FX/Niagara/Source/NiagaraEditor/Private/Widgets/SNiagaraSystemViewport.cpp
@@ -141,6 +141,24 @@ void FNiagaraSystemViewportClient::Draw(FViewport* InViewport, FCanvas* Canvas)
 			ToggleOrbitCamera(bCachedShouldOrbitCamera);
 			SetUpdateViewportFocus(false);
 		}
+
+		UWorld* World = Component->GetWorld();
+		if (World && !ParticleSystem->bFixedBounds)
+		{
+			FNiagaraSystemInstance* SystemInstance = SystemInstanceController->GetSystemInstance_Unsafe();
+			const TArray<FNiagaraEmitterHandle>& EmitterHandles = ParticleSystem->GetEmitterHandles();
+			for (int32 i=0; i < EmitterHandles.Num(); ++i)
+			{
+				if (!EmitterHandles[i].GetIsEnabled() || !EmitterHandles[i].GetDebugShowBounds())
+				{
+					continue;
+				}
+				const FNiagaraEmitterInstance& EmitterInstance = SystemInstance->GetEmitters()[i].Get();
+				const FTransform& Transform = Component->GetComponentTransform();
+				const FBox EmitterBounds = EmitterInstance.GetBounds();
+				DrawDebugBox(World, EmitterBounds.GetCenter() + Transform.GetLocation(), EmitterBounds.GetExtent(), Transform.GetRotation(), FColor::Green);
+			}
+		}
 	}
 
 	if (NiagaraViewport.IsValid() && NiagaraViewport->GetDrawElement(SNiagaraSystemViewport::EDrawElements::Bounds))
