Ensuring NDC Write DIs complete their tick within the same tick group they were issues to ensure same frame read ordering is possible.

diff --git a/Engine/Plugins/FX/Niagara/Source/Niagara/Classes/NiagaraDataInterface.h b/Engine/Plugins/FX/Niagara/Source/Niagara/Classes/NiagaraDataInterface.h
--- a/Engine/Plugins/FX/Niagara/Source/Niagara/Classes/NiagaraDataInterface.h
+++ b/Engine/Plugins/FX/Niagara/Source/Niagara/Classes/NiagaraDataInterface.h
@@ -662,6 +662,12 @@ public:
 	*/
 	virtual bool PostSimulateCanOverlapFrames() const { return true; }
 
+	/**
+	When set to true simulation may not complete in the same tick group it started.
+	You must override and set to false if you require the data interface post stage tick to complete before the tick group ends.
+	*/
+	virtual bool PostStageCanOverlapTickGroups() const { return true; }
+
 	virtual bool RequiresDistanceFieldData() const { return false; }
 	virtual bool RequiresDepthBuffer() const { return false; }
 	virtual bool RequiresEarlyViewData() const { return false; }
diff --git a/Engine/Plugins/FX/Niagara/Source/Niagara/Classes/NiagaraSystem.h b/Engine/Plugins/FX/Niagara/Source/Niagara/Classes/NiagaraSystem.h
--- a/Engine/Plugins/FX/Niagara/Source/Niagara/Classes/NiagaraSystem.h
+++ b/Engine/Plugins/FX/Niagara/Source/Niagara/Classes/NiagaraSystem.h
@@ -696,6 +696,7 @@ public:
 	bool HasSystemScriptDIsWithPerInstanceData() const;
 	FORCEINLINE bool HasDIsWithPostSimulateTick() const { return bHasDIsWithPostSimulateTick; }
 	FORCEINLINE bool AllDIsPostSimulateCanOverlapFrames() const { return bAllDIsPostSimulateCanOverlapFrames; }
+	FORCEINLINE bool AsyncWorkCanOverlapTickGroups() const { return bAllDIsPostStageCanOverlapTickGroups; }
 	FORCEINLINE bool HasAnyGPUEmitters()const{ return bHasAnyGPUEmitters; }
 	FORCEINLINE bool NeedsGPUContextInitForDataInterfaces() const { return bNeedsGPUContextInitForDataInterfaces; }
 
@@ -990,6 +991,7 @@ protected:
 
 	uint32 bHasDIsWithPostSimulateTick : 1;
 	uint32 bAllDIsPostSimulateCanOverlapFrames : 1;
+	uint32 bAllDIsPostStageCanOverlapTickGroups : 1;
 	uint32 bHasAnyGPUEmitters : 1;
 	uint32 bNeedsSortedSignificanceCull : 1;
 
diff --git a/Engine/Plugins/FX/Niagara/Source/Niagara/Internal/DataInterface/NiagaraDataInterfaceDataChannelWrite.h b/Engine/Plugins/FX/Niagara/Source/Niagara/Internal/DataInterface/NiagaraDataInterfaceDataChannelWrite.h
--- a/Engine/Plugins/FX/Niagara/Source/Niagara/Internal/DataInterface/NiagaraDataInterfaceDataChannelWrite.h
+++ b/Engine/Plugins/FX/Niagara/Source/Niagara/Internal/DataInterface/NiagaraDataInterfaceDataChannelWrite.h
@@ -112,8 +112,10 @@ public:
 	virtual void PostStageTick(FNDICpuPostStageContext& Context) override;
 	virtual void ProvidePerInstanceDataForRenderThread(void* DataForRenderThread, void* PerInstanceData, const FNiagaraSystemInstanceID& SystemInstance) override;
 
-	//We cannot overlapt frames as we must correctly sync up with the data channel manager on Begin/End frame etc.
+	//We cannot overlap frames as we must correctly sync up with the data channel manager on Begin/End frame etc.
 	virtual bool PostSimulateCanOverlapFrames() const { return false; }
+	//We cannot have post stage overlap tick groups so that the write DI can publish it's contents to the data channel at the correct time to allow same frame reads.
+	virtual bool PostStageCanOverlapTickGroups() const { return false; }
 	//UNiagaraDataInterface Interface
 
 	void Num(FVectorVMExternalFunctionContext& Context);
diff --git a/Engine/Plugins/FX/Niagara/Source/Niagara/Private/NiagaraSystem.cpp b/Engine/Plugins/FX/Niagara/Source/Niagara/Private/NiagaraSystem.cpp
--- a/Engine/Plugins/FX/Niagara/Source/Niagara/Private/NiagaraSystem.cpp
+++ b/Engine/Plugins/FX/Niagara/Source/Niagara/Private/NiagaraSystem.cpp
@@ -169,6 +169,7 @@ UNiagaraSystem::UNiagaraSystem(const FObjectInitializer& ObjectInitializer)
 , CurrentScalabilitySettings(*new FNiagaraSystemScalabilitySettings())
 , bHasDIsWithPostSimulateTick(false)
 , bAllDIsPostSimulateCanOverlapFrames(true)
+, bAllDIsPostStageCanOverlapTickGroups(true)
 , bHasAnyGPUEmitters(false)
 , bNeedsSortedSignificanceCull(false)
 , ActiveInstances(0)
@@ -2114,6 +2115,7 @@ void UNiagaraSystem::UpdateDITickFlags()
 {
 	bHasDIsWithPostSimulateTick = false;
 	bAllDIsPostSimulateCanOverlapFrames = true;
+	bAllDIsPostStageCanOverlapTickGroups = true;
 	auto CheckPostSimTick = [&](UNiagaraScript* Script)
 	{
 		if (Script)
@@ -2126,6 +2128,8 @@ void UNiagaraSystem::UpdateDITickFlags()
 					bHasDIsWithPostSimulateTick |= true;
 					bAllDIsPostSimulateCanOverlapFrames &= DefaultDataInterface->PostSimulateCanOverlapFrames();
 				}
+				
+				bAllDIsPostStageCanOverlapTickGroups &= DefaultDataInterface->PostStageCanOverlapTickGroups();
 			}
 		}
 	};
diff --git a/Engine/Plugins/FX/Niagara/Source/Niagara/Private/NiagaraSystemSimulation.cpp b/Engine/Plugins/FX/Niagara/Source/Niagara/Private/NiagaraSystemSimulation.cpp
--- a/Engine/Plugins/FX/Niagara/Source/Niagara/Private/NiagaraSystemSimulation.cpp
+++ b/Engine/Plugins/FX/Niagara/Source/Niagara/Private/NiagaraSystemSimulation.cpp
@@ -1341,7 +1341,7 @@ void FNiagaraSystemSimulation::Tick_GameThread_Internal(float DeltaSeconds, cons
 			Instance->ConcurrentTickGraphEvent = ConcurrentTickGraphEvent;
 		}
 		ConcurrentTickTask->Unlock();
-		if (bIsSolo || GNiagaraSystemSimulationTickTaskShouldWait)
+		if (bIsSolo || GNiagaraSystemSimulationTickTaskShouldWait || !System->AsyncWorkCanOverlapTickGroups())
 		{
 			MyCompletionGraphEvent->DontCompleteUntil(AllWorkCompleteGraphEvent);
 		}
diff --git a/Engine/Plugins/FX/Niagara/Source/Niagara/Private/NiagaraWorldManager.cpp b/Engine/Plugins/FX/Niagara/Source/Niagara/Private/NiagaraWorldManager.cpp
--- a/Engine/Plugins/FX/Niagara/Source/Niagara/Private/NiagaraWorldManager.cpp
+++ b/Engine/Plugins/FX/Niagara/Source/Niagara/Private/NiagaraWorldManager.cpp
@@ -37,7 +37,7 @@ DECLARE_CYCLE_STAT(TEXT("Niagara Manager Wait On Render [GT]"), STAT_NiagaraWorl
 DECLARE_CYCLE_STAT(TEXT("Niagara Manager Wait Pre Garbage Collect [GT]"), STAT_NiagaraWorldManWaitPreGC, STATGROUP_Niagara); 
 DECLARE_CYCLE_STAT(TEXT("Niagara Manager Refresh Owner Allows Scalability"), STAT_NiagaraWorldManRefreshOwnerAllowsScalability, STATGROUP_Niagara);
 
-static int GNiagaraAllowAsyncWorkToEndOfFrame = 0;
+static int GNiagaraAllowAsyncWorkToEndOfFrame = 1;
 static FAutoConsoleVariableRef CVarNiagaraAllowAsyncWorkToEndOfFrame(
 	TEXT("fx.Niagara.AllowAsyncWorkToEndOfFrame"),
 	GNiagaraAllowAsyncWorkToEndOfFrame,
@@ -553,6 +553,15 @@ FNiagaraSystemSimulationRef FNiagaraWorldManager::GetSystemSimulation(ETickingGr
 	SystemSimulations[ActualTickGroup].Add(System, Sim);
 	Sim->Init(System, World, false, TickGroup);
 
+	//If this system doesn't allow async work to overlap TGs then we must ensure our Tick task is performed in it's correct TG.
+	//The async work for this system can then properly be chained to the tick task so that everything completes on time.
+	//Most systems most of the time will no need this and we'll clear this requirement as soon as all systems no longer need it.
+	//TODO: This should go away when we move over to the new task graph system and rejig our dependency tracking.
+	if(System->AsyncWorkCanOverlapTickGroups())
+	{
+		TickFunctions[ActualTickGroup].EndTickGroup = (ETickingGroup)ActualTickGroup;
+	}
+
 #if WITH_EDITOR
 	System->OnSystemPostEditChange().AddRaw(this, &FNiagaraWorldManager::OnSystemPostChange);
 #endif
@@ -742,6 +751,11 @@ void FNiagaraWorldManager::RemoveDataChannel(const UNiagaraDataChannel* Channel)
 	GetDataChannelManager().RemoveDataChannel(Channel);
 }
 
+UNiagaraDataChannelHandler* FNiagaraWorldManager::FindDataChannelHandler(const UNiagaraDataChannel* Channel)
+{
+	return GetDataChannelManager().FindDataChannelHandler(Channel);
+}
+
 void FNiagaraWorldManager::WaitForAsyncWork()
 {
 	for (int TG = 0; TG < NiagaraNumTickGroups; ++TG)
@@ -968,12 +982,6 @@ void FNiagaraWorldManager::PostActorTick(float DeltaSeconds)
 	//-TODO: This could be done after each system sim has run
 	DeferredDeletionQueue.Empty();
 
-	// Update tick groups
-	for (FNiagaraWorldManagerTickFunction& TickFunc : TickFunctions )
-	{
-		TickFunc.EndTickGroup = GNiagaraAllowAsyncWorkToEndOfFrame ? TG_LastDemotable : (ETickingGroup)TickFunc.TickGroup;
-	}
-
 #if WITH_NIAGARA_DEBUGGER
 	// Tick debug HUD for the world
 	if (NiagaraDebugHud != nullptr)
@@ -1244,6 +1252,7 @@ void FNiagaraWorldManager::Tick(ETickingGroup TickGroup, float DeltaSeconds, ELe
 
 	ActiveNiagaraTickGroup = ActualTickGroup;
 
+	bool bTickFunctionCanOverlapTickGroups = GNiagaraAllowAsyncWorkToEndOfFrame != 0;
 	for (auto SimIt=SystemSimulations[ActualTickGroup].CreateIterator(); SimIt; ++SimIt)
 	{
 		FNiagaraSystemSimulation* Sim = &SimIt.Value().Get();
@@ -1251,6 +1260,8 @@ void FNiagaraWorldManager::Tick(ETickingGroup TickGroup, float DeltaSeconds, ELe
 		if (Sim->IsValid())
 		{
 			Sim->Tick_GameThread(DeltaSeconds, MyCompletionGraphEvent);
+
+			bTickFunctionCanOverlapTickGroups &= Sim->GetSystem()->AsyncWorkCanOverlapTickGroups();
 		}
 		else
 		{
@@ -1259,6 +1270,15 @@ void FNiagaraWorldManager::Tick(ETickingGroup TickGroup, float DeltaSeconds, ELe
 		}
 	}
 
+	//If all systems in this TG can overlap their async work then we can allow the tick as a whole to run until last demotable.
+	//When adding new system sims to this tick group we can change this.
+	//Any systems needing their async work to complete must have this Tick function finish inside this TG so that we can properly complete our async work before the TG finishes.
+	if(bTickFunctionCanOverlapTickGroups)
+	{
+		//Set here but this wont take effect until next tick.
+		TickFunctions[ActualTickGroup].EndTickGroup = TG_LastDemotable;
+	}
+
 	ActiveNiagaraTickGroup = -1;
 
 	//We update scalability managers here so that any new systems can be culled or setup with other scalability based parameters correctly for their spawn.
diff --git a/Engine/Plugins/FX/Niagara/Source/Niagara/Public/NiagaraWorldManager.h b/Engine/Plugins/FX/Niagara/Source/Niagara/Public/NiagaraWorldManager.h
--- a/Engine/Plugins/FX/Niagara/Source/Niagara/Public/NiagaraWorldManager.h
+++ b/Engine/Plugins/FX/Niagara/Source/Niagara/Public/NiagaraWorldManager.h
@@ -30,6 +30,7 @@ class UNiagaraParameterCollection;
 class UNiagaraParameterCollectionInstance;
 class UNiagaraComponentPool;
 struct FNiagaraScalabilityState;
+class UNiagaraDataChannelHandler;
 
 USTRUCT()
 struct FNiagaraWorldManagerTickFunction : public FTickFunction
@@ -223,6 +224,7 @@ public:
 	FNiagaraDataChannelManager& GetDataChannelManager(){ return *DataChannelManager; }
 	NIAGARA_API void InitDataChannel(const UNiagaraDataChannel* Channel, bool bForce);
 	NIAGARA_API void RemoveDataChannel(const UNiagaraDataChannel* Channel);
+	NIAGARA_API UNiagaraDataChannelHandler* FindDataChannelHandler(const UNiagaraDataChannel* Channel);
 
 	/** Waits for all currently in flight async work to be completed. */
 	void WaitForAsyncWork();
