Added metadata to allow multline text entry for python script

diff --git a/Engine/Plugins/FX/Niagara/Source/Niagara/Classes/NiagaraScript.h b/Engine/Plugins/FX/Niagara/Source/Niagara/Classes/NiagaraScript.h
--- a/Engine/Plugins/FX/Niagara/Source/Niagara/Classes/NiagaraScript.h
+++ b/Engine/Plugins/FX/Niagara/Source/Niagara/Classes/NiagaraScript.h
@@ -669,7 +669,7 @@ public:
 	ENiagaraPythonUpdateScriptReference ConversionScriptExecution = ENiagaraPythonUpdateScriptReference::ScriptAsset;
 
 	/** Python script to run when converting this script to the recommended deprecation update script. */
-	UPROPERTY(EditAnywhere, Category = Script, meta=(EditCondition="bUsePythonScriptConversion && ConversionScriptExecution == ENiagaraPythonUpdateScriptReference::DirectTextEntry", EditConditionHides))
+	UPROPERTY(EditAnywhere, Category = Script, meta=(MultiLine = true, EditCondition="bUsePythonScriptConversion && ConversionScriptExecution == ENiagaraPythonUpdateScriptReference::DirectTextEntry", EditConditionHides))
 	FString PythonConversionScript;
 
 	/** Asset reference to a python script to run when converting this script to the recommended deprecation update script. */
