- Handle Variable aliasing for User UObjects
- Currently we only need to resolve the aliasing for render bindings stage
- in future we may want to pull objects into DI functions, or allow bindings from DI's to pull from Objects at levels other than User.
- This may involve us having to put the aliasing information elsewhere to resolve correctly

diff --git a/Engine/Plugins/FX/Niagara/Source/Niagara/Classes/NiagaraEmitter.h b/Engine/Plugins/FX/Niagara/Source/Niagara/Classes/NiagaraEmitter.h
--- a/Engine/Plugins/FX/Niagara/Source/Niagara/Classes/NiagaraEmitter.h
+++ b/Engine/Plugins/FX/Niagara/Source/Niagara/Classes/NiagaraEmitter.h
@@ -365,6 +365,9 @@ struct NIAGARA_API FVersionedNiagaraEmitterData
 	UPROPERTY()
 	FNiagaraParameterStore RendererBindings;
 
+	UPROPERTY()
+	TArray<FNiagaraExternalUObjectInfo> RendererBindingsExternalObjects;
+
 	UPROPERTY()
 	TMap<FNiagaraVariableBase, FNiagaraVariableBase> ResolvedDIBindings;
 
diff --git a/Engine/Plugins/FX/Niagara/Source/Niagara/Private/NiagaraEmitter.cpp b/Engine/Plugins/FX/Niagara/Source/Niagara/Private/NiagaraEmitter.cpp
--- a/Engine/Plugins/FX/Niagara/Source/Niagara/Private/NiagaraEmitter.cpp
+++ b/Engine/Plugins/FX/Niagara/Source/Niagara/Private/NiagaraEmitter.cpp
@@ -1796,6 +1796,9 @@ void FVersionedNiagaraEmitterData::RebuildRendererBindings(const UNiagaraEmitter
 	FNiagaraParameterStore TempStore;
 	const bool bAnyRendererBindingsAdded = BuildParameterStoreRendererBindings(TempStore);
 
+	RendererBindings.Empty();
+	RendererBindingsExternalObjects.Empty();
+
 	if (bAnyRendererBindingsAdded)
 	{
 		TArray<UNiagaraScript*> TargetScripts;
@@ -1807,7 +1810,6 @@ void FVersionedNiagaraEmitterData::RebuildRendererBindings(const UNiagaraEmitter
 		const int32 EmitterScriptIndex = TargetScripts.Num();
 		GetScripts(TargetScripts, false, true);
 
-		RendererBindings.Empty();
 		TArrayView<const FNiagaraVariableWithOffset> Vars = TempStore.ReadParameterVariables();
 		for (const FNiagaraVariableWithOffset& Var : Vars)
 		{
@@ -1822,16 +1824,15 @@ void FVersionedNiagaraEmitterData::RebuildRendererBindings(const UNiagaraEmitter
 
 				// Find Resolved UObjects
 				// Note: Most parameters are pushed in as part of the DataSet -> Parameters process
-				if (Var.GetType().IsUObject())
+				if (Var.GetType().IsUObject() && !Var.GetType().IsDataInterface())
 				{
 					const ENiagaraSimTarget ScriptSimTarget = iScript < EmitterScriptIndex ? ENiagaraSimTarget::CPUSim : SimTarget;
 					if (const FNiagaraScriptExecutionParameterStore* ScriptParameterStore = Script->GetExecutionReadyParameterStore(ScriptSimTarget))
 					{
-						if (UObject* FoundObject = ScriptParameterStore->GetUObject(Var))
+						if (ScriptParameterStore->IndexOf(Var) != INDEX_NONE)
 						{
 							int32 ParameterOffset = INDEX_NONE;
 							RendererBindings.AddParameter(Var, true, false, &ParameterOffset);
-							RendererBindings.SetUObject(FoundObject, ParameterOffset);
 							bVariableFound = true;
 						}
 					}
@@ -1857,6 +1858,32 @@ void FVersionedNiagaraEmitterData::RebuildRendererBindings(const UNiagaraEmitter
 			}
 		}
 
+		if (RendererBindings.GetUObjects().Num() > 0)
+		{
+			for (const UNiagaraScript* Script : TargetScripts)
+			{
+				for (const FNiagaraScriptUObjectCompileInfo& DefaultInfo : Script->GetCachedDefaultUObjects())
+				{
+					const FNiagaraVariableBase ReadVariable(DefaultInfo.Variable.GetType(), DefaultInfo.RegisteredParameterMapRead);
+					const bool bIsUserVariable = ReadVariable.IsInNameSpace(FNiagaraConstants::UserNamespaceString);
+					if (!bIsUserVariable)
+					{
+						continue;
+					}
+					for (const FName WriteName : DefaultInfo.RegisteredParameterMapWrites)
+					{
+						const FNiagaraVariableBase WriteVariable(DefaultInfo.Variable.GetType(), WriteName);
+						if (RendererBindings.IndexOf(WriteVariable) != INDEX_NONE)
+						{
+							FNiagaraExternalUObjectInfo& ExternalInfo = RendererBindingsExternalObjects.AddDefaulted_GetRef();
+							ExternalInfo.Variable = WriteVariable;
+							ExternalInfo.ExternalName = ReadVariable.GetName();
+						}
+					}
+				}
+			}
+		}
+
 		RendererBindings.TriggerOnLayoutChanged();
 	}
 #endif
diff --git a/Engine/Plugins/FX/Niagara/Source/Niagara/Private/NiagaraEmitterInstance.cpp b/Engine/Plugins/FX/Niagara/Source/Niagara/Private/NiagaraEmitterInstance.cpp
--- a/Engine/Plugins/FX/Niagara/Source/Niagara/Private/NiagaraEmitterInstance.cpp
+++ b/Engine/Plugins/FX/Niagara/Source/Niagara/Private/NiagaraEmitterInstance.cpp
@@ -477,13 +477,14 @@ void FNiagaraEmitterInstance::Init(int32 InEmitterIdx, FNiagaraSystemInstanceID
 		{
 			if (ParentSystemInstance)
 			{
-				ParentSystemInstance->GetInstanceParameters().Bind(&RendererBindings);
 				FNiagaraSystemSimulationPtr ParentSystemSimulation = ParentSystemInstance->GetSystemSimulation();
 				if (ParentSystemSimulation.IsValid())
 				{
 					ParentSystemSimulation->GetSpawnExecutionContext()->Parameters.Bind(&RendererBindings);
 					ParentSystemSimulation->GetUpdateExecutionContext()->Parameters.Bind(&RendererBindings);
 				}
+
+				ParentSystemInstance->GetInstanceParameters().Bind(&RendererBindings);
 			}
 
 			if (GPUExecContext && EmitterData->SimTarget == ENiagaraSimTarget::GPUComputeSim)
@@ -1240,6 +1241,26 @@ void FNiagaraEmitterInstance::PreTick()
 		}
 	}
 
+	if (ParentSystemInstance && RendererBindings.GetUObjectsDirty())
+	{
+		if ( const FNiagaraParameterStore* SrcStore = ParentSystemInstance->GetOverrideParameters() )
+		{
+			FNiagaraParameterStore* DstStore = &RendererBindings;
+
+			for ( const FNiagaraExternalUObjectInfo& ExternalObject : CachedEmitter.GetEmitterData()->RendererBindingsExternalObjects )
+			{
+				const int32 SrcIndex = SrcStore->IndexOf(ExternalObject.GetExternalVariable());
+				const int32 DstIndex = DstStore->IndexOf(ExternalObject.Variable);
+				if (SrcIndex != INDEX_NONE && DstIndex != INDEX_NONE)
+				{
+					DstStore->SetUObject(SrcStore->GetUObject(SrcIndex), DstIndex);
+				}
+			}
+		}
+		RendererBindings.Tick();
+	}
+
+
 	++TickCount;
 	ParticleDataSet->SetIDAcquireTag(TickCount);
 }
diff --git a/Engine/Plugins/FX/Niagara/Source/Niagara/Private/NiagaraSystem.cpp b/Engine/Plugins/FX/Niagara/Source/Niagara/Private/NiagaraSystem.cpp
--- a/Engine/Plugins/FX/Niagara/Source/Niagara/Private/NiagaraSystem.cpp
+++ b/Engine/Plugins/FX/Niagara/Source/Niagara/Private/NiagaraSystem.cpp
@@ -2020,23 +2020,6 @@ void UNiagaraSystem::ResolveParameterStoreBindings()
 
 	// Resolve Object bindings
 	{
-		TMap<FNiagaraVariableBase, UObject*> WroteUObjects;
-		ForEachScript(
-			[&](UNiagaraScript* Script)
-			{
-				for (const FNiagaraScriptUObjectCompileInfo& DefaultInfo : Script->GetCachedDefaultUObjects())
-				{
-					for (FName WriteName : DefaultInfo.RegisteredParameterMapWrites)
-					{
-						const FNiagaraVariableBase WriteVariable(DefaultInfo.Variable.GetType(), WriteName);
-						check(WroteUObjects.Contains(WriteVariable) == false);
-						WroteUObjects.Emplace(WriteVariable, DefaultInfo.Object);
-					}
-				}
-			}
-		);
-
-		// Create resolved data per script
 		ForEachScript(
 			[&](UNiagaraScript* Script)
 			{
@@ -2045,23 +2028,30 @@ void UNiagaraSystem::ResolveParameterStoreBindings()
 					return;
 				}
 
+				const FNiagaraParameterStore& InstanceParamStore = GetSystemCompiledData().InstanceParamStore;
+
 				TArray<FNiagaraResolvedUObjectInfo> ResolvedUObjects;
 				ResolvedUObjects.Reserve(Script->GetCachedDefaultUObjects().Num());
 
 				for (const FNiagaraScriptUObjectCompileInfo& DefaultInfo : Script->GetCachedDefaultUObjects())
 				{
 					const FNiagaraVariableBase ReadVariable(DefaultInfo.Variable.GetType(), DefaultInfo.RegisteredParameterMapRead);
+					const bool bIsUserVariable = ReadVariable.IsInNameSpace(FNiagaraConstants::UserNamespaceString);
 
-					//-TODO:UObject will need to handle User / NPC bindings here.  I.e. Variable=User.Variable in the stack set
-					//const bool bIsUserVariable = ReadVariable.IsInNameSpace(FNiagaraConstants::UserNamespaceString);
-
-					for (FName WriteName : DefaultInfo.RegisteredParameterMapWrites)
+					for (const FName WriteName : DefaultInfo.RegisteredParameterMapWrites)
 					{
 						FNiagaraResolvedUObjectInfo& ResolvedInfo = ResolvedUObjects.AddDefaulted_GetRef();
 						const FNiagaraVariableBase WriteVariable(DefaultInfo.Variable.GetType(), WriteName);
 						ResolvedInfo.ReadVariableName = ReadVariable.GetName();
 						ResolvedInfo.ResolvedVariable = WriteVariable;
-						ResolvedInfo.Object = WroteUObjects.FindRef(WriteVariable);
+						if (bIsUserVariable)
+						{
+							ResolvedInfo.Object = InstanceParamStore.GetUObject(ReadVariable);
+						}
+						else
+						{
+							ResolvedInfo.Object = DefaultInfo.Object;
+						}
 					}
 				}
 
diff --git a/Engine/Plugins/FX/Niagara/Source/Niagara/Public/NiagaraCommon.h b/Engine/Plugins/FX/Niagara/Source/Niagara/Public/NiagaraCommon.h
--- a/Engine/Plugins/FX/Niagara/Source/Niagara/Public/NiagaraCommon.h
+++ b/Engine/Plugins/FX/Niagara/Source/Niagara/Public/NiagaraCommon.h
@@ -705,6 +705,20 @@ struct NIAGARA_API FNiagaraResolvedUObjectInfo
 	//int32 UserPtrIdx = INDEX_NONE;
 };
 
+USTRUCT()
+struct NIAGARA_API FNiagaraExternalUObjectInfo
+{
+	GENERATED_BODY()
+
+	UPROPERTY()
+	FNiagaraVariableBase Variable;
+
+	UPROPERTY()
+	FName ExternalName;
+
+	[[nodiscard]] FNiagaraVariableBase GetExternalVariable() const { return FNiagaraVariableBase(Variable.GetType(), ExternalName); }
+};
+
 USTRUCT()
 struct NIAGARA_API FNiagaraScriptDataInterfaceInfo
 {
diff --git a/Engine/Plugins/FX/Niagara/Source/Niagara/Public/NiagaraParameterStore.h b/Engine/Plugins/FX/Niagara/Source/Niagara/Public/NiagaraParameterStore.h
--- a/Engine/Plugins/FX/Niagara/Source/Niagara/Public/NiagaraParameterStore.h
+++ b/Engine/Plugins/FX/Niagara/Source/Niagara/Public/NiagaraParameterStore.h
@@ -289,10 +289,10 @@ public:
 		return ResourceSize;
 	}
 
-	FORCEINLINE uint32 GetParametersDirty() const { return bParametersDirty; }
-	FORCEINLINE uint32 GetInterfacesDirty() const { return bInterfacesDirty; }
-	FORCEINLINE uint32 GetUObjectsDirty() const { return bUObjectsDirty; }
-	FORCEINLINE uint32 GetPositionDataDirty() const { return bPositionDataDirty; }
+	FORCEINLINE bool GetParametersDirty() const { return bParametersDirty; }
+	FORCEINLINE bool GetInterfacesDirty() const { return bInterfacesDirty; }
+	FORCEINLINE bool GetUObjectsDirty() const { return bUObjectsDirty; }
+	FORCEINLINE bool GetPositionDataDirty() const { return bPositionDataDirty; }
 
 	FORCEINLINE void MarkParametersDirty() { bParametersDirty = true; }
 	FORCEINLINE void MarkInterfacesDirty() { bInterfacesDirty = true; }
diff --git a/Engine/Plugins/FX/Niagara/Source/Niagara/Public/NiagaraTypes.h b/Engine/Plugins/FX/Niagara/Source/Niagara/Public/NiagaraTypes.h
--- a/Engine/Plugins/FX/Niagara/Source/Niagara/Public/NiagaraTypes.h
+++ b/Engine/Plugins/FX/Niagara/Source/Niagara/Public/NiagaraTypes.h
@@ -1833,11 +1833,11 @@ struct FNiagaraVariableBase
 	*/
 	NIAGARA_API bool ReplaceRootNamespace(const FStringView& ExpectedNamespace, const FStringView& NewNamespace);
 
-	FORCEINLINE bool IsInNameSpace(const FStringView& Namespace) const
+	static bool IsInNameSpace(const FStringView& Namespace, FName VariableName)
 	{
 		FNameBuilder NameString;
-		Name.ToString(NameString);
-		
+		VariableName.ToString(NameString);
+
 		FStringView NameStringView = NameString.ToView();
 		if (Namespace.Len() > 1 && Namespace[Namespace.Len() - 1] == '.') // Includes period in namespace
 		{
@@ -1849,6 +1849,11 @@ struct FNiagaraVariableBase
 		}
 	}
 
+	FORCEINLINE bool IsInNameSpace(const FStringView& Namespace) const
+	{
+		return IsInNameSpace(Namespace, Name);
+	}
+
 #if WITH_EDITORONLY_DATA
 	// This method should not be used at runtime as we have pre-defined strings in FNiagaraConstants for runtime cases
 	FORCEINLINE bool IsInNameSpace(const FName& Namespace) const
diff --git a/Engine/Plugins/FX/Niagara/Source/NiagaraEditor/Private/NiagaraHlslTranslator.cpp b/Engine/Plugins/FX/Niagara/Source/NiagaraEditor/Private/NiagaraHlslTranslator.cpp
--- a/Engine/Plugins/FX/Niagara/Source/NiagaraEditor/Private/NiagaraHlslTranslator.cpp
+++ b/Engine/Plugins/FX/Niagara/Source/NiagaraEditor/Private/NiagaraHlslTranslator.cpp
@@ -5618,12 +5618,12 @@ void TNiagaraHlslTranslator<GraphBridge>::ParameterMapSet(const FParamMapSetNode
 					FNiagaraScriptUObjectCompileInfo& Info = CompilationOutput.ScriptData.UObjectInfos[Input];
 					Info.RegisteredParameterMapWrites.AddUnique(ResolvedName);
 
-					//-TODO: We don't support Var=User.Var or Var=NPC.Var currently this requires some runtime remapping of the parameter stores
-					if (Info.Variable.IsInNameSpace(FNiagaraConstants::UserNamespaceString) || Info.Variable.IsInNameSpace(FNiagaraConstants::ParameterCollectionNamespaceString))
+					//-TODO: We don't support Var=NPC.Var currently this requires some runtime remapping of the parameter stores
+					if (Info.Variable.IsInNameSpace(FNiagaraConstants::ParameterCollectionNamespaceString))
 					{
 						Error(
 							FText::Format(
-								LOCTEXT("UObjectParameterInvalidNamespaceFormat", "Variable Write '{0}' can not read User / NPC variable '{1}'. This is not supported currently."),
+								LOCTEXT("UObjectParameterInvalidNamespaceFormat", "Variable Write '{0}' can not read NPC variable '{1}'. This is not supported currently."),
 								FText::FromName(ResolvedName),
 								FText::FromName(Info.Variable.GetName())
 							),
