- Only allow particle attribute bindings for ribbons as we don't support anything else

diff --git a/Engine/Plugins/FX/Niagara/Source/Niagara/Private/NiagaraRibbonRendererProperties.cpp b/Engine/Plugins/FX/Niagara/Source/Niagara/Private/NiagaraRibbonRendererProperties.cpp
--- a/Engine/Plugins/FX/Niagara/Source/Niagara/Private/NiagaraRibbonRendererProperties.cpp
+++ b/Engine/Plugins/FX/Niagara/Source/Niagara/Private/NiagaraRibbonRendererProperties.cpp
@@ -357,6 +357,21 @@ void UNiagaraRibbonRendererProperties::SetPreviousBindings(const FVersionedNiaga
 	PrevRibbonTwistBinding.SetAsPreviousValue(RibbonTwistBinding, SrcEmitter, ENiagaraRendererSourceDataMode::Particles);
 }
 
+#if WITH_EDITORONLY_DATA
+bool UNiagaraRibbonRendererProperties::IsSupportedVariableForBinding(const FNiagaraVariableBase& InSourceForBinding, const FName& InTargetBindingName) const
+{
+	if (InTargetBindingName == GET_MEMBER_NAME_CHECKED(UNiagaraRendererProperties, RendererEnabledBinding))
+	{
+		return
+			InSourceForBinding.IsInNameSpace(FNiagaraConstants::UserNamespace) ||
+			InSourceForBinding.IsInNameSpace(FNiagaraConstants::SystemNamespace) ||
+			InSourceForBinding.IsInNameSpace(FNiagaraConstants::EmitterNamespace);
+	}
+
+	return InSourceForBinding.IsInNameSpace(FNiagaraConstants::ParticleAttributeNamespaceString);
+}
+#endif
+
 void UNiagaraRibbonRendererProperties::CacheFromCompiledData(const FNiagaraDataSetCompiledData* CompiledData)
 {
 	UpdateMICs();
diff --git a/Engine/Plugins/FX/Niagara/Source/Niagara/Public/NiagaraRibbonRendererProperties.h b/Engine/Plugins/FX/Niagara/Source/Niagara/Public/NiagaraRibbonRendererProperties.h
--- a/Engine/Plugins/FX/Niagara/Source/Niagara/Public/NiagaraRibbonRendererProperties.h
+++ b/Engine/Plugins/FX/Niagara/Source/Niagara/Public/NiagaraRibbonRendererProperties.h
@@ -231,6 +231,9 @@ public:
 	virtual void GetRendererFeedback(const FVersionedNiagaraEmitter& InEmitter, TArray<FNiagaraRendererFeedback>& OutErrors, TArray<FNiagaraRendererFeedback>& OutWarnings, TArray<FNiagaraRendererFeedback>& OutInfo) const override;
 
 	virtual TArray<FNiagaraVariable> GetBoundAttributes() const override; 
+#endif
+#if WITH_EDITORONLY_DATA
+	virtual bool IsSupportedVariableForBinding(const FNiagaraVariableBase& InSourceForBinding, const FName& InTargetBindingName) const override;
 #endif
 	virtual void CacheFromCompiledData(const FNiagaraDataSetCompiledData* CompiledData) override;
 	//UNiagaraRendererProperties Interface END
