- Allow external MIDs to have parameter bindings set on them

diff --git a/Engine/Plugins/FX/Niagara/Source/Niagara/Private/NiagaraSystemInstanceController.cpp b/Engine/Plugins/FX/Niagara/Source/Niagara/Private/NiagaraSystemInstanceController.cpp
--- a/Engine/Plugins/FX/Niagara/Source/Niagara/Private/NiagaraSystemInstanceController.cpp
+++ b/Engine/Plugins/FX/Niagara/Source/Niagara/Private/NiagaraSystemInstanceController.cpp
@@ -314,7 +314,7 @@ void FNiagaraSystemInstanceController::UpdateEmitterMaterials()
 						{
 							if (ExistingMaterial)
 							{
-								if (auto MID = Cast<UMaterialInstanceDynamic>(ExistingMaterial))
+								if (UMaterialInstanceDynamic* MID = Cast<UMaterialInstanceDynamic>(ExistingMaterial))
 								{
 									if (EmitterMaterials.FindByPredicate([&](const FMaterialOverride& ExistingOverride) -> bool { return (ExistingOverride.Material == ExistingMaterial) && (ExistingOverride.EmitterRendererProperty == Properties) && (ExistingOverride.MaterialSubIndex == MaterialIndex); }))
 									{
@@ -322,6 +322,11 @@ void FNiagaraSystemInstanceController::UpdateEmitterMaterials()
 										// TODO: Are there cases where we don't always have to recreate it?
 										ExistingMaterial = MID->Parent;
 									}
+									else
+									{
+										// If we get here this is an external MID so do not create a new one
+										continue;
+									}
 								}
 
 								// Create a new MID
