- Add Niagara Matrix Array Support

diff --git a/Engine/Plugins/FX/Niagara/Source/Niagara/Classes/NiagaraDataInterface.h b/Engine/Plugins/FX/Niagara/Source/Niagara/Classes/NiagaraDataInterface.h
--- a/Engine/Plugins/FX/Niagara/Source/Niagara/Classes/NiagaraDataInterface.h
+++ b/Engine/Plugins/FX/Niagara/Source/Niagara/Classes/NiagaraDataInterface.h
@@ -1215,6 +1215,63 @@ struct FNDIInputParam<FQuat4f>
 	FORCEINLINE void Reset() { X.Reset(); Y.Reset(); Z.Reset(); W.Reset(); }
 };
 
+template<>
+struct FNDIInputParam<FMatrix44f>
+{
+	VectorVM::FExternalFuncInputHandler<float> M[16];
+
+	FORCEINLINE FNDIInputParam() { }
+	FORCEINLINE FNDIInputParam(FVectorVMExternalFunctionContext& Context) { Init(Context); }
+	FORCEINLINE void Init(FVectorVMExternalFunctionContext& Context)
+	{
+		for (int32 i=0; i < 16; ++i)
+		{
+			M[i].Init(Context);
+		}
+	}
+	FORCEINLINE FMatrix44f GetAndAdvance()
+	{
+		FMatrix44f Matrix;
+		for (int32 i=0; i < 16; ++i)
+		{
+			Matrix.M[i/4][i%4] = M[i].GetAndAdvance();
+		}
+		return Matrix;
+	}
+	FORCEINLINE FMatrix44f Get()
+	{
+		FMatrix44f Matrix;
+		for (int32 i = 0; i < 16; ++i)
+		{
+			Matrix.M[i / 4][i % 4] = M[i].Get();
+		}
+		return Matrix;
+	}
+	FORCEINLINE void Advance(int32 Count = 1)
+	{
+		for (int32 i = 0; i < 16; ++i)
+		{
+			M[i].Advance(Count);
+		}
+	}
+	FORCEINLINE bool IsConstant() const
+	{
+		bool bConstant = true;
+		for (int32 i = 0; i < 16; ++i)
+		{
+			bConstant &= M[i].IsConstant();
+		}
+		return bConstant;
+	}
+	FORCEINLINE void Reset()
+	{
+		for (int32 i = 0; i < 16; ++i)
+		{
+			M[i].Reset();
+		}
+	}
+};
+
 template<>
 struct FNDIInputParam<FLinearColor>
 {
diff --git a/Engine/Plugins/FX/Niagara/Source/Niagara/Classes/NiagaraDataInterfaceArrayFloat.h b/Engine/Plugins/FX/Niagara/Source/Niagara/Classes/NiagaraDataInterfaceArrayFloat.h
--- a/Engine/Plugins/FX/Niagara/Source/Niagara/Classes/NiagaraDataInterfaceArrayFloat.h
+++ b/Engine/Plugins/FX/Niagara/Source/Niagara/Classes/NiagaraDataInterfaceArrayFloat.h
@@ -108,3 +108,20 @@ public:
 
 	NDIARRAY_GENERATE_BODY_LWC(UNiagaraDataInterfaceArrayQuat, FQuat4f, QuatData)
 };
+
+UCLASS(EditInlineNew, Category = "Array", CollapseCategories, meta = (DisplayName = "Matrix Array"), Blueprintable, BlueprintType)
+class NIAGARA_API UNiagaraDataInterfaceArrayMatrix : public UNiagaraDataInterfaceArray
+{
+public:
+	GENERATED_BODY()
+
+#if WITH_EDITORONLY_DATA
+	UPROPERTY(EditAnywhere, BlueprintReadWrite, Category = "Array")
+	TArray<FMatrix> MatrixData;
+#endif
+
+	UPROPERTY()
+	TArray<FMatrix44f> InternalMatrixData;
+
+	NDIARRAY_GENERATE_BODY_LWC(UNiagaraDataInterfaceArrayMatrix, FMatrix44f, MatrixData)
+};
diff --git a/Engine/Plugins/FX/Niagara/Source/Niagara/Classes/NiagaraDataInterfaceArrayFunctionLibrary.h b/Engine/Plugins/FX/Niagara/Source/Niagara/Classes/NiagaraDataInterfaceArrayFunctionLibrary.h
--- a/Engine/Plugins/FX/Niagara/Source/Niagara/Classes/NiagaraDataInterfaceArrayFunctionLibrary.h
+++ b/Engine/Plugins/FX/Niagara/Source/Niagara/Classes/NiagaraDataInterfaceArrayFunctionLibrary.h
@@ -37,6 +37,12 @@ public:
 	/** Sets Niagara Array FQuat Data. */
 	UFUNCTION(BlueprintCallable, Category = Niagara, meta = (DisplayName = "Niagara Set Quaternion Array"))
 	static void SetNiagaraArrayQuat(UNiagaraComponent* NiagaraSystem, FName OverrideName, const TArray<FQuat>& ArrayData);
+	/**
+	 * Sets Niagara Array FMatrix Data.
+	 * @param bApplyLWCRebase When enabled the matrix translation will have the simulation tile offset subtracted from it
+	 */
+	UFUNCTION(BlueprintCallable, Category = Niagara, meta = (DisplayName = "Niagara Set Matrix Array"))
+	static void SetNiagaraArrayMatrix(UNiagaraComponent* NiagaraSystem, FName OverrideName, const TArray<FMatrix>& ArrayData, bool bApplyLWCRebase = true);
 	/** Sets Niagara Array Int32 Data. */
 	UFUNCTION(BlueprintCallable, Category = Niagara, meta = (DisplayName = "Niagara Set Int32 Array"))
 	static void SetNiagaraArrayInt32(UNiagaraComponent* NiagaraSystem, FName OverrideName, const TArray<int32>& ArrayData);
@@ -68,6 +74,12 @@ public:
 	/** Gets a copy of Niagara FQuat Data. */
 	UFUNCTION(BlueprintCallable, Category = Niagara, meta = (DisplayName = "Niagara Get Quaternion Array"))
 	static TArray<FQuat> GetNiagaraArrayQuat(UNiagaraComponent* NiagaraSystem, FName OverrideName);
+	/**
+	 * Gets a copy of Niagara FMatrix Data.
+	 * @param bApplyLWCRebase When enabled the matrix translation will have the simulation tile offset added to it
+	 */
+	UFUNCTION(BlueprintCallable, Category = Niagara, meta = (DisplayName = "Niagara Get Matrix Array"))
+	static TArray<FMatrix> GetNiagaraArrayMatrix(UNiagaraComponent* NiagaraSystem, FName OverrideName, bool bApplyLWCRebase = true);
 	/** Gets a copy of Niagara Int32 Data. */
 	UFUNCTION(BlueprintCallable, Category = Niagara, meta = (DisplayName = "Niagara Get Int32 Array"))
 	static TArray<int32> GetNiagaraArrayInt32(UNiagaraComponent* NiagaraSystem, FName OverrideName);
@@ -98,6 +110,12 @@ public:
 	/** Sets a single value within a Niagara Array FQuat. */
 	UFUNCTION(BlueprintCallable, Category = Niagara, meta = (DisplayName = "Niagara Set Quaternion Array Value"))
 	static void SetNiagaraArrayQuatValue(UNiagaraComponent* NiagaraSystem, FName OverrideName, int Index, const FQuat& Value, bool bSizeToFit);
+	/**
+	 * Sets a single value within a Niagara Array FMatrix.
+	 * @param bApplyLWCRebase When enabled the matrix translation will have the simulation tile offset subtracted from it
+	 */
+	UFUNCTION(BlueprintCallable, Category = Niagara, meta = (DisplayName = "Niagara Set Matrix Array Value"))
+	static void SetNiagaraArrayMatrixValue(UNiagaraComponent* NiagaraSystem, FName OverrideName, int Index, const FMatrix& Value, bool bSizeToFit, bool bApplyLWCRebase = true);
 	/** Sets a single value within a Niagara Array Int32. */
 	UFUNCTION(BlueprintCallable, Category = Niagara, meta = (DisplayName = "Niagara Set Int32 Array Value"))
 	static void SetNiagaraArrayInt32Value(UNiagaraComponent* NiagaraSystem, FName OverrideName, int Index, int32 Value, bool bSizeToFit);
@@ -129,6 +147,12 @@ public:
 	/** Gets a single value within a Niagara Array FQuat. */
 	UFUNCTION(BlueprintCallable, Category = Niagara, meta = (DisplayName = "Niagara Get Quaternion Array Value"))
 	static FQuat GetNiagaraArrayQuatValue(UNiagaraComponent* NiagaraSystem, FName OverrideName, int Index);
+	/**
+	 * Gets a single value within a Niagara Array FMatrix.
+	 * @param bApplyLWCRebase When enabled the matrix translation will have the simulation tile offset added to it
+	 */
+	UFUNCTION(BlueprintCallable, Category = Niagara, meta = (DisplayName = "Niagara Get Matrix Array Value"))
+	static FMatrix GetNiagaraArrayMatrixValue(UNiagaraComponent* NiagaraSystem, FName OverrideName, int Index, bool bApplyLWCRebase = true);
 	/** Gets a single value within a Niagara Array Int32. */
 	UFUNCTION(BlueprintCallable, Category = Niagara, meta = (DisplayName = "Niagara Get Int32 Array Value"))
 	static int32 GetNiagaraArrayInt32Value(UNiagaraComponent* NiagaraSystem, FName OverrideName, int Index);
@@ -145,5 +169,6 @@ public:
 	static void SetNiagaraArrayVector(UNiagaraComponent* NiagaraSystem, FName OverrideName, TConstArrayView<FVector3f> ArrayData);
 	static void SetNiagaraArrayVector4(UNiagaraComponent* NiagaraSystem, FName OverrideName, TConstArrayView<FVector4f> ArrayData);
 	static void SetNiagaraArrayQuat(UNiagaraComponent* NiagaraSystem, FName OverrideName, TConstArrayView<FQuat4f> ArrayData);
+	static void SetNiagaraArrayMatrix(UNiagaraComponent* NiagaraSystem, FName OverrideName, TConstArrayView<FMatrix44f> ArrayData);
 	static void SetNiagaraArrayUInt8(UNiagaraComponent* NiagaraSystem, FName OverrideName, TConstArrayView<uint8> ArrayData);
 };
diff --git a/Engine/Plugins/FX/Niagara/Source/Niagara/Private/NiagaraDataInterfaceArrayFunctionLibrary.cpp b/Engine/Plugins/FX/Niagara/Source/Niagara/Private/NiagaraDataInterfaceArrayFunctionLibrary.cpp
--- a/Engine/Plugins/FX/Niagara/Source/Niagara/Private/NiagaraDataInterfaceArrayFunctionLibrary.cpp
+++ b/Engine/Plugins/FX/Niagara/Source/Niagara/Private/NiagaraDataInterfaceArrayFunctionLibrary.cpp
@@ -23,6 +23,7 @@ NDIARRAY_GENERATE_IMPL(UNiagaraDataInterfaceArrayPosition, FNiagaraPosition, Pos
 NDIARRAY_GENERATE_IMPL_LWC(UNiagaraDataInterfaceArrayFloat4, FVector4f, FloatData)
 NDIARRAY_GENERATE_IMPL(UNiagaraDataInterfaceArrayColor, FLinearColor, ColorData)
 NDIARRAY_GENERATE_IMPL_LWC(UNiagaraDataInterfaceArrayQuat, FQuat4f, QuatData)
+NDIARRAY_GENERATE_IMPL_LWC(UNiagaraDataInterfaceArrayMatrix, FMatrix44f, MatrixData)
 NDIARRAY_GENERATE_IMPL(UNiagaraDataInterfaceArrayInt32, int32, IntData)
 NDIARRAY_GENERATE_IMPL_LWC(UNiagaraDataInterfaceArrayUInt8, uint8, IntData)
 NDIARRAY_GENERATE_IMPL(UNiagaraDataInterfaceArrayBool, bool, BoolData)
@@ -171,6 +172,26 @@ void UNiagaraDataInterfaceArrayFunctionLibrary::SetNiagaraArrayQuat(UNiagaraComp
 	SetNiagaraArray<FQuat, UNiagaraDataInterfaceArrayQuat>(NiagaraComponent, OverrideName, ArrayData);
 }
 
+void UNiagaraDataInterfaceArrayFunctionLibrary::SetNiagaraArrayMatrix(UNiagaraComponent* NiagaraComponent, FName OverrideName, const TArray<FMatrix>& ArrayData, bool bApplyLWCRebase)
+{
+	if (bApplyLWCRebase)
+	{
+		const FNiagaraLWCConverter LwcConverter = GetLWCConverter(NiagaraComponent);
+
+		TArray<FMatrix> RebasedArrayData;
+		RebasedArrayData.SetNumUninitialized(ArrayData.Num());
+		for (int32 i=0; i < ArrayData.Num(); ++i)
+		{
+			RebasedArrayData[i] = LwcConverter.ConvertWorldToSimulationMatrix(ArrayData[i]);
+		}
+		SetNiagaraArray<FMatrix, UNiagaraDataInterfaceArrayMatrix>(NiagaraComponent, OverrideName, RebasedArrayData);
+	}
+	else
+	{
+		SetNiagaraArray<FMatrix, UNiagaraDataInterfaceArrayMatrix>(NiagaraComponent, OverrideName, ArrayData);
+	}
+}
+
 void UNiagaraDataInterfaceArrayFunctionLibrary::SetNiagaraArrayInt32(UNiagaraComponent* NiagaraComponent, FName OverrideName, const TArray<int32>& ArrayData)
 {
 	SetNiagaraArray<int32, UNiagaraDataInterfaceArrayInt32>(NiagaraComponent, OverrideName, ArrayData);
@@ -233,6 +254,20 @@ TArray<FQuat> UNiagaraDataInterfaceArrayFunctionLibrary::GetNiagaraArrayQuat(UNi
 	return GetNiagaraArray<FQuat, UNiagaraDataInterfaceArrayQuat>(NiagaraComponent, OverrideName);
 }
 
+TArray<FMatrix> UNiagaraDataInterfaceArrayFunctionLibrary::GetNiagaraArrayMatrix(UNiagaraComponent* NiagaraComponent, FName OverrideName, bool bApplyLWCRebase)
+{
+	TArray<FMatrix> Results = GetNiagaraArray<FMatrix, UNiagaraDataInterfaceArrayMatrix>(NiagaraComponent, OverrideName);
+	if (bApplyLWCRebase)
+	{
+		const FNiagaraLWCConverter LwcConverter = GetLWCConverter(NiagaraComponent);
+		for (FMatrix& Result : Results)
+		{
+			Result = LwcConverter.ConvertSimulationToWorldMatrix(Result);
+		}
+	}
+	return Results;
+}
+
 TArray<int32> UNiagaraDataInterfaceArrayFunctionLibrary::GetNiagaraArrayInt32(UNiagaraComponent* NiagaraComponent, FName OverrideName)
 {
 	return GetNiagaraArray<int32, UNiagaraDataInterfaceArrayInt32>(NiagaraComponent, OverrideName);
@@ -267,8 +302,8 @@ void UNiagaraDataInterfaceArrayFunctionLibrary::SetNiagaraArrayVectorValue(UNiag
 
 void UNiagaraDataInterfaceArrayFunctionLibrary::SetNiagaraArrayPositionValue(UNiagaraComponent* NiagaraComponent, FName OverrideName, int Index, const FVector& Value, bool bSizeToFit)
 {
-	FNiagaraLWCConverter LwcConverter = GetLWCConverter(NiagaraComponent);
-	FNiagaraPosition SimulationPosition = LwcConverter.ConvertWorldToSimulationPosition(Value);
+	const FNiagaraLWCConverter LwcConverter = GetLWCConverter(NiagaraComponent);
+	const FNiagaraPosition SimulationPosition = LwcConverter.ConvertWorldToSimulationPosition(Value);
 	SetNiagaraArrayValue<FNiagaraPosition, UNiagaraDataInterfaceArrayPosition>(NiagaraComponent, OverrideName, Index, SimulationPosition, bSizeToFit);
 }
 
@@ -287,6 +322,19 @@ void UNiagaraDataInterfaceArrayFunctionLibrary::SetNiagaraArrayQuatValue(UNiagar
 	SetNiagaraArrayValue<FQuat, UNiagaraDataInterfaceArrayQuat>(NiagaraComponent, OverrideName, Index, Value, bSizeToFit);
 }
 
+void UNiagaraDataInterfaceArrayFunctionLibrary::SetNiagaraArrayMatrixValue(UNiagaraComponent* NiagaraComponent, FName OverrideName, int Index, const FMatrix& Value, bool bSizeToFit, bool bApplyLWCRebase)
+{
+	if (bApplyLWCRebase)
+	{
+		const FNiagaraLWCConverter LwcConverter = GetLWCConverter(NiagaraComponent);
+		SetNiagaraArrayValue<FMatrix, UNiagaraDataInterfaceArrayMatrix>(NiagaraComponent, OverrideName, Index, LwcConverter.ConvertWorldToSimulationMatrix(Value), bSizeToFit);
+	}
+	else
+	{
+		SetNiagaraArrayValue<FMatrix, UNiagaraDataInterfaceArrayMatrix>(NiagaraComponent, OverrideName, Index, Value, bSizeToFit);
+	}
+}
+
 void UNiagaraDataInterfaceArrayFunctionLibrary::SetNiagaraArrayInt32Value(UNiagaraComponent* NiagaraComponent, FName OverrideName, int Index, int32 Value, bool bSizeToFit)
 {
 	SetNiagaraArrayValue<int32, UNiagaraDataInterfaceArrayInt32>(NiagaraComponent, OverrideName, Index, Value, bSizeToFit);
@@ -341,6 +389,18 @@ FQuat UNiagaraDataInterfaceArrayFunctionLibrary::GetNiagaraArrayQuatValue(UNiaga
 	return GetNiagaraArrayValue<FQuat, UNiagaraDataInterfaceArrayQuat>(NiagaraComponent, OverrideName, Index);
 }
 
+FMatrix UNiagaraDataInterfaceArrayFunctionLibrary::GetNiagaraArrayMatrixValue(UNiagaraComponent* NiagaraComponent, FName OverrideName, int Index, bool bApplyLWCRebase)
+{
+	FMatrix Result = GetNiagaraArrayValue<FMatrix, UNiagaraDataInterfaceArrayMatrix>(NiagaraComponent, OverrideName, Index);
+	if (bApplyLWCRebase)
+	{
+		const FNiagaraLWCConverter LwcConverter = GetLWCConverter(NiagaraComponent);
+		Result = LwcConverter.ConvertSimulationToWorldMatrix(Result);
+	}
+
+	return Result;
+}
+
 int32 UNiagaraDataInterfaceArrayFunctionLibrary::GetNiagaraArrayInt32Value(UNiagaraComponent* NiagaraComponent, FName OverrideName, int Index)
 {
 	return GetNiagaraArrayValue<int32, UNiagaraDataInterfaceArrayInt32>(NiagaraComponent, OverrideName, Index);
@@ -378,6 +438,11 @@ void UNiagaraDataInterfaceArrayFunctionLibrary::SetNiagaraArrayQuat(UNiagaraComp
 	SetNiagaraArray<FQuat4f, UNiagaraDataInterfaceArrayQuat>(NiagaraSystem, OverrideName, ArrayData);
 }
 
+void UNiagaraDataInterfaceArrayFunctionLibrary::SetNiagaraArrayMatrix(UNiagaraComponent* NiagaraSystem, FName OverrideName, TConstArrayView<FMatrix44f> ArrayData)
+{
+	SetNiagaraArray<FMatrix44f, UNiagaraDataInterfaceArrayMatrix>(NiagaraSystem, OverrideName, ArrayData);
+}
+
 void UNiagaraDataInterfaceArrayFunctionLibrary::SetNiagaraArrayUInt8(UNiagaraComponent* NiagaraSystem, FName OverrideName, TConstArrayView<uint8> ArrayData)
 {
 	SetNiagaraArray<uint8, UNiagaraDataInterfaceArrayUInt8>(NiagaraSystem, OverrideName, ArrayData);
diff --git a/Engine/Plugins/FX/Niagara/Source/Niagara/Private/NiagaraDataInterfaceArrayImplHelpers.h b/Engine/Plugins/FX/Niagara/Source/Niagara/Private/NiagaraDataInterfaceArrayImplHelpers.h
--- a/Engine/Plugins/FX/Niagara/Source/Niagara/Private/NiagaraDataInterfaceArrayImplHelpers.h
+++ b/Engine/Plugins/FX/Niagara/Source/Niagara/Private/NiagaraDataInterfaceArrayImplHelpers.h
@@ -248,6 +248,72 @@ struct FNDIArrayImplHelper<FQuat4f> : public FNDIArrayImplHelperBase<FQuat4f>
 	}
 };
 
+template<>
+struct FNDIArrayImplHelper<FMatrix44f> : public FNDIArrayImplHelperBase<FMatrix44f>
+{
+	typedef FMatrix44f TVMArrayType;
+
+	static constexpr TCHAR const* HLSLVariableType = TEXT("float4x4");
+	static constexpr EPixelFormat ReadPixelFormat = PF_A32B32G32R32F;
+	static constexpr TCHAR const* ReadHLSLBufferType = TEXT("float4");
+	static constexpr TCHAR const* ReadHLSLBufferRead = TEXT("Value = float4x4(BUFFER_NAME[Index * 4 + 0], BUFFER_NAME[Index * 4 + 1], BUFFER_NAME[Index * 4 + 2], BUFFER_NAME[Index * 4 + 3]);");
+	static constexpr EPixelFormat RWPixelFormat = PF_R32_FLOAT;
+	static constexpr TCHAR const* RWHLSLBufferType = TEXT("float");
+	static constexpr TCHAR const* RWHLSLBufferRead = TEXT(R"(Value = float4x4(
+		BUFFER_NAME[Index * 16 +  0], BUFFER_NAME[Index * 16 +  1], BUFFER_NAME[Index * 16 +  2], BUFFER_NAME[Index * 16 +  3],
+		BUFFER_NAME[Index * 16 +  4], BUFFER_NAME[Index * 16 +  5], BUFFER_NAME[Index * 16 +  6], BUFFER_NAME[Index * 16 +  7],
+		BUFFER_NAME[Index * 16 +  8], BUFFER_NAME[Index * 16 +  9], BUFFER_NAME[Index * 16 + 10], BUFFER_NAME[Index * 16 + 11],
+		BUFFER_NAME[Index * 16 + 12], BUFFER_NAME[Index * 16 + 13], BUFFER_NAME[Index * 16 + 14], BUFFER_NAME[Index * 16 + 15]);
+	)");
+	static constexpr TCHAR const* RWHLSLBufferWrite =
+		TEXT("BUFFER_NAME[Index * 16 +  0] = Value[0].x, BUFFER_NAME[Index * 16 +  1] = Value[0].y, BUFFER_NAME[Index * 16 +  2] = Value[0].z, BUFFER_NAME[Index * 16 +  3] = Value[0].w,")
+		TEXT("BUFFER_NAME[Index * 16 +  4] = Value[1].x, BUFFER_NAME[Index * 16 +  5] = Value[1].y, BUFFER_NAME[Index * 16 +  6] = Value[1].z, BUFFER_NAME[Index * 16 +  7] = Value[1].w,")
+		TEXT("BUFFER_NAME[Index * 16 +  8] = Value[2].x, BUFFER_NAME[Index * 16 +  9] = Value[2].y, BUFFER_NAME[Index * 16 + 10] = Value[2].z, BUFFER_NAME[Index * 16 + 11] = Value[2].w,")
+		TEXT("BUFFER_NAME[Index * 16 + 12] = Value[3].x, BUFFER_NAME[Index * 16 + 13] = Value[3].y, BUFFER_NAME[Index * 16 + 14] = Value[3].z, BUFFER_NAME[Index * 16 + 15] = Value[3].w;");
+
+	static const FNiagaraTypeDefinition& GetTypeDefinition() { return FNiagaraTypeDefinition::GetMatrix4Def(); }
+	static const FMatrix44f GetDefaultValue() { return FMatrix44f::Identity; }
+
+	static void CopyCpuToCpuMemory(FMatrix44f* Dest, const FMatrix44f* Src, int32 NumElements)
+	{
+		FMemory::Memcpy(Dest, Src, NumElements * sizeof(FMatrix44f));
+	}
+
+	static void CopyCpuToCpuMemory(FMatrix44f* Dest, const FMatrix* Src, int32 NumElements)
+	{
+		for (int32 i = 0; i < NumElements; ++i)
+		{
+			Dest[i] = FMatrix44f(Src[i]);
+		}
+	}
+
+	static void CopyCpuToCpuMemory(FMatrix* Dest, const FMatrix44f* Src, int32 NumElements)
+	{
+		for (int32 i = 0; i < NumElements; ++i)
+		{
+			Dest[i] = FMatrix(Src[i]);
+		}
+	}
+
+	static void AppendValueToString(const FMatrix44f& Value, FString& OutString)
+	{
+		for (int32 r=0; r < 4; ++r)
+		{
+			for (int32 c=0; c < 4; ++c)
+			{
+				if (r + c == 0)
+				{
+					OutString.Appendf(TEXT("%f"), Value.M[r][c]);
+				}
+				else
+				{
+					OutString.Appendf(TEXT(", %f"), Value.M[r][c]);
+				}
+			}
+		}
+	}
+};
+
 template<>
 struct FNDIArrayImplHelper<int32> : public FNDIArrayImplHelperBase<int32>
 {
diff --git a/Engine/Plugins/FX/Niagara/Source/Niagara/Private/NiagaraTypes.cpp b/Engine/Plugins/FX/Niagara/Source/Niagara/Private/NiagaraTypes.cpp
--- a/Engine/Plugins/FX/Niagara/Source/Niagara/Private/NiagaraTypes.cpp
+++ b/Engine/Plugins/FX/Niagara/Source/Niagara/Private/NiagaraTypes.cpp
@@ -163,6 +163,16 @@ FVector FNiagaraLWCConverter::ConvertSimulationVectorToWorld(FVector3f Simulatio
 	return FVector(SimulationPosition) + SystemWorldPos;
 }
 
+FMatrix FNiagaraLWCConverter::ConvertWorldToSimulationMatrix(const FMatrix& Matrix) const
+{
+	return Matrix.ConcatTranslation(-SystemWorldPos);
+}
+
+FMatrix FNiagaraLWCConverter::ConvertSimulationToWorldMatrix(const FMatrix& Matrix) const
+{
+	return Matrix.ConcatTranslation(SystemWorldPos);
+}
+
 FNiagaraStructConversionStep::FNiagaraStructConversionStep()
 {
 }
diff --git a/Engine/Plugins/FX/Niagara/Source/Niagara/Public/NiagaraTypes.h b/Engine/Plugins/FX/Niagara/Source/Niagara/Public/NiagaraTypes.h
--- a/Engine/Plugins/FX/Niagara/Source/Niagara/Public/NiagaraTypes.h
+++ b/Engine/Plugins/FX/Niagara/Source/Niagara/Public/NiagaraTypes.h
@@ -227,13 +227,16 @@ struct NIAGARA_API FNiagaraAssetVersion
 
 struct NIAGARA_API FNiagaraLWCConverter
 {
-	FNiagaraLWCConverter(FVector InSystemWorldPos = FVector::ZeroVector);
+	explicit FNiagaraLWCConverter(FVector InSystemWorldPos = FVector::ZeroVector);
 
-	FVector3f ConvertWorldToSimulationVector(FVector WorldPosition) const;
-	FNiagaraPosition ConvertWorldToSimulationPosition(FVector WorldPosition) const;
+	[[nodiscard]] FVector3f ConvertWorldToSimulationVector(FVector WorldPosition) const;
+	[[nodiscard]] FNiagaraPosition ConvertWorldToSimulationPosition(FVector WorldPosition) const;
 	
-	FVector ConvertSimulationPositionToWorld(FNiagaraPosition SimulationPosition) const;
-	FVector ConvertSimulationVectorToWorld(FVector3f SimulationPosition) const;
+	[[nodiscard]] FVector ConvertSimulationPositionToWorld(FNiagaraPosition SimulationPosition) const;
+	[[nodiscard]] FVector ConvertSimulationVectorToWorld(FVector3f SimulationPosition) const;
+
+	[[nodiscard]] FMatrix ConvertWorldToSimulationMatrix(const FMatrix& Matrix) const;
+	[[nodiscard]] FMatrix ConvertSimulationToWorldMatrix(const FMatrix& Matrix) const;
 
 private:
 	FVector SystemWorldPos;
