Static Analysis Fix

diff --git a/Engine/Plugins/FX/Niagara/Source/Niagara/Private/NiagaraSystem.cpp b/Engine/Plugins/FX/Niagara/Source/Niagara/Private/NiagaraSystem.cpp
--- a/Engine/Plugins/FX/Niagara/Source/Niagara/Private/NiagaraSystem.cpp
+++ b/Engine/Plugins/FX/Niagara/Source/Niagara/Private/NiagaraSystem.cpp
@@ -2112,14 +2112,15 @@ void UNiagaraSystem::UpdateDITickFlags()
 		{
 			for (FNiagaraScriptDataInterfaceCompileInfo& Info : Script->GetVMExecutableData().DataInterfaceInfo)
 			{
-				UNiagaraDataInterface* DefaultDataInterface = Info.GetDefaultDataInterface();
-				if (DefaultDataInterface && DefaultDataInterface->HasPostSimulateTick())
+				if(UNiagaraDataInterface* DefaultDataInterface = Info.GetDefaultDataInterface())
 				{
-					bHasDIsWithPostSimulateTick |= true;
-					bAllDIsPostSimulateCanOverlapFrames &= DefaultDataInterface->PostSimulateCanOverlapFrames();
-				}
-				
-				bAllDIsPostStageCanOverlapTickGroups &= DefaultDataInterface->PostStageCanOverlapTickGroups();
+					if (DefaultDataInterface->HasPostSimulateTick())
+					{
+						bHasDIsWithPostSimulateTick |= true;
+						bAllDIsPostSimulateCanOverlapFrames &= DefaultDataInterface->PostSimulateCanOverlapFrames();
+					}
+					bAllDIsPostStageCanOverlapTickGroups &= DefaultDataInterface->PostStageCanOverlapTickGroups();
+				}				
 			}
 		}
 	};
