Limiting NDC writes to allocated size removed by recent optimizations.

diff --git a/Engine/Plugins/FX/Niagara/Source/Niagara/Private/DataInterface/NiagaraDataInterfaceDataChannelWrite.cpp b/Engine/Plugins/FX/Niagara/Source/Niagara/Private/DataInterface/NiagaraDataInterfaceDataChannelWrite.cpp
--- a/Engine/Plugins/FX/Niagara/Source/Niagara/Private/DataInterface/NiagaraDataInterfaceDataChannelWrite.cpp
+++ b/Engine/Plugins/FX/Niagara/Source/Niagara/Private/DataInterface/NiagaraDataInterfaceDataChannelWrite.cpp
@@ -269,8 +269,10 @@ struct FNDIDataChannelWriteInstanceData
 	{
 		if (Data && Data->GetDestinationData())
 		{
-			//TOOD: Move allocation/end to pre/post stage?
-			Data->GetDestinationData()->SetNumInstances(AtomicNumInstances.load(std::memory_order_seq_cst));
+			//The count here can overrun the num allocated but we should never actually write beyond the max allocated.
+			uint32 WrittenInstances = AtomicNumInstances.load(std::memory_order_seq_cst);
+			WrittenInstances = FMath::Min(WrittenInstances, Data->GetDestinationData()->GetNumInstancesAllocated());
+			Data->GetDestinationData()->SetNumInstances(WrittenInstances);
 			Data->EndSimulate();
 
 			if (GbDebugDumpWriter)
@@ -801,7 +803,7 @@ void UNiagaraDataInterfaceDataChannelWrite::Append(FVectorVMExternalFunctionCont
 
 			if (LocalNumToEmit > 0)
 			{
-				int32 NumAllocated = Data->GetNumInstancesAllocated();
+				uint32 NumAllocated = Data->GetNumInstancesAllocated();
 				InEmit.Reset();
 
 				//Update the shared atomic instance count and grab the current index at which we can write.
@@ -813,14 +815,18 @@ void UNiagaraDataInterfaceDataChannelWrite::Append(FVectorVMExternalFunctionCont
 
 				if(bEmitAll)
 				{
+					//limit the number to emit so we do not write over the end of the buffers.
+					uint32 MaxWriteCount = NumAllocated - FMath::Min(CurrNumInstances, NumAllocated);
+					LocalNumToEmit = FMath::Min(LocalNumToEmit, MaxWriteCount);
+
 					//If we're writing all instances then we can do a memcpy instead of slower loop copies.
 					bool bSuccess = false;
-					int32 Index = CurrNumInstances;
+					uint32 Index = CurrNumInstances;
 					auto FloatFunc = [Data, Index, LocalNumToEmit](const FNDIDataChannelRegisterBinding& VMBinding, FNDIInputParam<float>& FloatData)
 					{
 						if (VMBinding.DataSetRegisterIndex != INDEX_NONE)
 						{
-							float* Dest = Data->GetInstancePtrFloat(VMBinding.DataSetRegisterIndex, (uint32)Index);
+							float* Dest = Data->GetInstancePtrFloat(VMBinding.DataSetRegisterIndex, Index);
 							if (FloatData.IsConstant())
 							{
 								float Value = FloatData.GetAndAdvance();
@@ -837,7 +843,7 @@ void UNiagaraDataInterfaceDataChannelWrite::Append(FVectorVMExternalFunctionCont
 					{
 						if (VMBinding.DataSetRegisterIndex != INDEX_NONE)
 						{
-							int32* Dest = Data->GetInstancePtrInt32(VMBinding.DataSetRegisterIndex, (uint32)Index);
+							int32* Dest = Data->GetInstancePtrInt32(VMBinding.DataSetRegisterIndex, Index);
 							if (IntData.IsConstant())
 							{
 								int32 Value = IntData.GetAndAdvance();
@@ -854,7 +860,7 @@ void UNiagaraDataInterfaceDataChannelWrite::Append(FVectorVMExternalFunctionCont
 					{
 						if (VMBinding.DataSetRegisterIndex != INDEX_NONE)
 						{
-							FFloat16* Dest = Data->GetInstancePtrHalf(VMBinding.DataSetRegisterIndex, (uint32)Index);
+							FFloat16* Dest = Data->GetInstancePtrHalf(VMBinding.DataSetRegisterIndex, Index);
 							if (HalfData.IsConstant())
 							{
 								FFloat16 Value = HalfData.GetAndAdvance();
@@ -879,13 +885,13 @@ void UNiagaraDataInterfaceDataChannelWrite::Append(FVectorVMExternalFunctionCont
 				}
 				else
 				{
-					for (int32 i = 0; i < Context.GetNumInstances(); ++i)
-					{
+					for (int32 i = 0; i < Context.GetNumInstances() && CurrNumInstances < NumAllocated; ++i)
+					{						
+						uint32 Index = CurrNumInstances;
+
 						bool bEmit = InEmit.GetAndAdvance();
 						bool bSuccess = false;
 
-						int32 Index = CurrNumInstances;
-
 						if(bEmit)
 						{
 							++CurrNumInstances;
@@ -894,17 +900,17 @@ void UNiagaraDataInterfaceDataChannelWrite::Append(FVectorVMExternalFunctionCont
 						auto FloatFunc = [Data, Index](const FNDIDataChannelRegisterBinding& VMBinding, FNDIInputParam<float>& FloatData)
 						{
 							if (VMBinding.DataSetRegisterIndex != INDEX_NONE)
-								*Data->GetInstancePtrFloat(VMBinding.DataSetRegisterIndex, (uint32)Index) = FloatData.GetAndAdvance();
+								*Data->GetInstancePtrFloat(VMBinding.DataSetRegisterIndex, Index) = FloatData.GetAndAdvance();
 						};
 						auto IntFunc = [Data, Index](const FNDIDataChannelRegisterBinding& VMBinding, FNDIInputParam<int32>& IntData)
 						{
 							if (VMBinding.DataSetRegisterIndex != INDEX_NONE)
-								*Data->GetInstancePtrInt32(VMBinding.DataSetRegisterIndex, (uint32)Index) = IntData.GetAndAdvance();
+								*Data->GetInstancePtrInt32(VMBinding.DataSetRegisterIndex, Index) = IntData.GetAndAdvance();
 						};
 						auto HalfFunc = [Data, Index](const FNDIDataChannelRegisterBinding& VMBinding, FNDIInputParam<FFloat16>& HalfData)
 						{
 							if (VMBinding.DataSetRegisterIndex != INDEX_NONE)
-								*Data->GetInstancePtrHalf(VMBinding.DataSetRegisterIndex, (uint32)Index) = HalfData.GetAndAdvance();
+								*Data->GetInstancePtrHalf(VMBinding.DataSetRegisterIndex, Index) = HalfData.GetAndAdvance();
 						};
 						bSuccess = VariadicInputs.Process(bEmit, 1, BindingInfo, FloatFunc, IntFunc, HalfFunc);
 
