- Fix for persistent debug draw functions not working with GPU sims
- This is because the compile tag name would be appended onto the function name so are unique per call

#rnx

diff --git a/Engine/Plugins/FX/Niagara/Shaders/Private/NiagaraDataInterfaceDebugDraw.ush b/Engine/Plugins/FX/Niagara/Shaders/Private/NiagaraDataInterfaceDebugDraw.ush
--- a/Engine/Plugins/FX/Niagara/Shaders/Private/NiagaraDataInterfaceDebugDraw.ush
+++ b/Engine/Plugins/FX/Niagara/Shaders/Private/NiagaraDataInterfaceDebugDraw.ush
@@ -451,7 +451,19 @@ void NDIDebugDraw_Cone(bool bDraw, float3 Location, float3 Axis, float Height, f
 	}
 }
 
-
 void NDIDebugDraw_Torus(bool bDraw, float3 Location, float3 Axis, float MajorRadius, float MinorRadius, float MajorRadiusSegments, float MinorRadiusSegments, float4 Color)
 {
-}
\ No newline at end of file
+}
+
+
+#define MakeDrawLinePersistent(NAME)				void NAME(float3 StartLocation, int StartLocationCoordinateSpace, float3 EndLocation, int EndLocationCoordinateSpace, float4 Color) {}
+#define MakeDrawRectanglePersistent(NAME)			void NAME(float3 Center, int CenterCoordinateSpace, float3 Offset, int OffsetCoordinateSpace, float3 Extents, bool HalfExtents, float3 XAxis, float3 YAxis, int AxisCoordinateSpace, int NumXSegments, int NumYSegments, bool bUnboundedPlane, float4 Color) {}
+#define MakeDrawCirclePersistent(NAME)				void NAME(float3 Center, int CenterCoordinateSpace, float3 Offset, int OffsetCoordinateSpace, float Radius, float3 XAxis, float3 YAxis, int AxisCoordinateSpace, int NumSegments, float4 Color) {}
+#define MakeDrawBoxPersistent(NAME)					void NAME(float3 Center, int CenterCoordinateSpace, float3 Offset, int OffsetCoordinateSpace, float3 Extents, bool HalfExtents, float3 RotationAxis, float RotationAngle, int RotationCoordinateSpace, float4 Color) {}
+#define MakeDrawSpherePersistent(NAME)				void NAME(float3 Center, int CenterCoordinateSpace, float3 OffsetFromCenter, int OffsetCoordinateSpace, float Radius,  int RadiusCoordinateSpace, bool HemisphereX, bool HemisphereY, bool HemisphereZ, float3 OrientationAxis, int RotationCoordinateSpace, float4 AdditionalRotation, float3 NonUniformScale, int NumSegments, float4 Color) {}
+#define MakeDrawCylinderPersistent(NAME)			void NAME(float3 Center, int CenterCoordinateSpace, float3 OffsetFromCenter, int OffsetCoordinateSpace, float Radius, float Height, bool IsHalfHeight, bool HemisphereX, bool HemisphereY, float3 OrientationAxis, int OrientationAxisCoordinateSpace, float3 NonUniformScale, int NumRadiusSegments, int NumHeightSegments, float4 Color) {}
+#define MakeDrawConePersistent(NAME)				void NAME(float3 Center, int CenterCoordinateSpace, float3 OffsetFromCenter, int OffsetCoordinateSpace, float Angle, float Length, int OrientationAxis, int OrientationAxisCoordinateSpace, float3 NonUniformScale, int NumRadiusSegments, int NumHeightSegments, float4 Color) {}
+#define MakeDrawTorusPersistent(NAME)				void NAME(float3 Center, int CenterCoordinateSpace, float3 OffsetFromCenter, int OffsetCoordinateSpace, float MajorRadius, float MinorRadius, float3 OrientationAxis, int OrientationAxisCoordinateSpace, float3 NonUniformScale, int MajorRadiusSegments, int MinorRadiusSegments, float4 Color) {}
+#define MakeDrawCoordinateSystemPersistent(NAME)	void NAME(float3 Center, int CenterCoordinateSpace, float3 OffsetFromCenter, int OffsetCoordinateSpace, float4 Rotation, int RotationCoordinateSpace, float Scale) {}
+#define MakeDrawGrid2DPersistent(NAME)				void NAME(float3 Center, int CenterCoordinateSpace, float3 OffsetFromCenter, int OffsetCoordinateSpace, float4 Rotation, int RotationCoordinateSpace, float2 Extents, int NumCellsX, int NumCellsY, float4 Color) {}
+#define MakeDrawGrid3DPersistent(NAME)				void NAME(float3 Center, int CenterCoordinateSpace, float3 OffsetFromCenter, int OffsetCoordinateSpace, float4 Rotation, int RotationCoordinateSpace, float3 Extents, int NumCellsX, int NumCellsY, int NumCellsZ, float4 Color) {}
diff --git a/Engine/Plugins/FX/Niagara/Shaders/Private/NiagaraDataInterfaceDebugDrawTemplate.ush b/Engine/Plugins/FX/Niagara/Shaders/Private/NiagaraDataInterfaceDebugDrawTemplate.ush
--- a/Engine/Plugins/FX/Niagara/Shaders/Private/NiagaraDataInterfaceDebugDrawTemplate.ush
+++ b/Engine/Plugins/FX/Niagara/Shaders/Private/NiagaraDataInterfaceDebugDrawTemplate.ush
@@ -53,48 +53,3 @@ void DrawGrid3D_{ParameterName}_UEImpureCall(bool bExecute, float3 Center, float
 {
 	NDIDebugDraw_Grid3D(bExecute, Center, Rotation, Extents, int3(NumCellsX, NumCellsY, NumCellsZ), Color);
 }
-
-void DrawLinePersistent_{ParameterName}_UEImpureCall(float3 StartLocation, int StartLocationCoordinateSpace, float3 EndLocation, int EndLocationCoordinateSpace, float4 Color)
-{
-}
-
-void DrawRectanglePersistent_{ParameterName}_UEImpureCall(float3 Center, int CenterCoordinateSpace, float3 Offset, int OffsetCoordinateSpace, float3 Extents, bool HalfExtents, float3 XAxis, float3 YAxis, int AxisCoordinateSpace, int NumXSegments, int NumYSegments, bool bUnboundedPlane, float4 Color)
-{
-}
-
-void DrawCirclePersistent_{ParameterName}_UEImpureCall(float3 Center, int CenterCoordinateSpace, float3 Offset, int OffsetCoordinateSpace, float Radius, float3 XAxis, float3 YAxis, int AxisCoordinateSpace, int NumSegments, float4 Color)
-{
-}
-
-void DrawBoxPersistent_{ParameterName}_UEImpureCall(float3 Center, int CenterCoordinateSpace, float3 Offset, int OffsetCoordinateSpace, float3 Extents, bool HalfExtents, float3 RotationAxis, float RotationAngle, int RotationCoordinateSpace, float4 Color)
-{
-}
-
-void DrawSpherePersistent_{ParameterName}_UEImpureCall(float3 Center, int CenterCoordinateSpace, float3 OffsetFromCenter, int OffsetCoordinateSpace, float Radius,  int RadiusCoordinateSpace, bool HemisphereX, bool HemisphereY, bool HemisphereZ, float3 OrientationAxis, int RotationCoordinateSpace, float4 AdditionalRotation, float3 NonUniformScale, int NumSegments, float4 Color)
-{
-}
-
-void DrawCylinderPersistent_{ParameterName}_UEImpureCall(float3 Center, int CenterCoordinateSpace, float3 OffsetFromCenter, int OffsetCoordinateSpace, float Radius, float Height, bool IsHalfHeight, bool HemisphereX, bool HemisphereY, float3 OrientationAxis, int OrientationAxisCoordinateSpace, float3 NonUniformScale, int NumRadiusSegments, int NumHeightSegments, float4 Color)
-{
-}
-
-void DrawConePersistent_{ParameterName}_UEImpureCall(float3 Center, int CenterCoordinateSpace, float3 OffsetFromCenter, int OffsetCoordinateSpace, float Angle, float Length, int OrientationAxis, int OrientationAxisCoordinateSpace, float3 NonUniformScale, int NumRadiusSegments, int NumHeightSegments, float4 Color)
-{
-}
-
-void DrawTorusPersistent_{ParameterName}_UEImpureCall(float3 Center, int CenterCoordinateSpace, float3 OffsetFromCenter, int OffsetCoordinateSpace, float MajorRadius, float MinorRadius, float3 OrientationAxis, int OrientationAxisCoordinateSpace, float3 NonUniformScale, int MajorRadiusSegments, int MinorRadiusSegments, float4 Color)
-{
-}
-
-void DrawCoordinateSystemPersistent_{ParameterName}_UEImpureCall(float3 Center, int CenterCoordinateSpace, float3 OffsetFromCenter, int OffsetCoordinateSpace, float4 Rotation, int RotationCoordinateSpace, float Scale)
-{
-}
-
-void DrawGrid2DPersistent_{ParameterName}_UEImpureCall(float3 Center, int CenterCoordinateSpace, float3 OffsetFromCenter, int OffsetCoordinateSpace, float4 Rotation, int RotationCoordinateSpace, float2 Extents, int NumCellsX, int NumCellsY, float4 Color)
-{
-}
-
-void DrawGrid3DPersistent_{ParameterName}_UEImpureCall(float3 Center, int CenterCoordinateSpace, float3 OffsetFromCenter, int OffsetCoordinateSpace, float4 Rotation, int RotationCoordinateSpace, float3 Extents, int NumCellsX, int NumCellsY, int NumCellsZ, float4 Color)
-{
-}
-
diff --git a/Engine/Plugins/FX/Niagara/Source/Niagara/Private/DataInterface/NiagaraDataInterfaceDebugDraw.cpp b/Engine/Plugins/FX/Niagara/Source/Niagara/Private/DataInterface/NiagaraDataInterfaceDebugDraw.cpp
--- a/Engine/Plugins/FX/Niagara/Source/Niagara/Private/DataInterface/NiagaraDataInterfaceDebugDraw.cpp
+++ b/Engine/Plugins/FX/Niagara/Source/Niagara/Private/DataInterface/NiagaraDataInterfaceDebugDraw.cpp
@@ -2946,19 +2946,8 @@ void UNiagaraDataInterfaceDebugDraw::GetParameterDefinitionHLSL(const FNiagaraDa
 
 bool UNiagaraDataInterfaceDebugDraw::GetFunctionHLSL(const FNiagaraDataInterfaceGPUParamInfo& ParamInfo, const FNiagaraDataInterfaceGeneratedFunction& FunctionInfo, int FunctionInstanceIndex, FString& OutHLSL)
 {
-	static const TSet<FName> ValidGpuFunctions =
+	static const TSet<FName> PersistentFunctions =
 	{
-		NDIDebugDrawLocal::DrawLineName,
-		NDIDebugDrawLocal::DrawRectangleName,
-		NDIDebugDrawLocal::DrawCircleName,
-		NDIDebugDrawLocal::DrawBoxName,
-		NDIDebugDrawLocal::DrawSphereName,
-		NDIDebugDrawLocal::DrawCylinderName,
-		NDIDebugDrawLocal::DrawConeName,
-		NDIDebugDrawLocal::DrawTorusName,
-		NDIDebugDrawLocal::DrawCoordinateSystemName,
-		NDIDebugDrawLocal::DrawGrid2DName,
-		NDIDebugDrawLocal::DrawGrid3DName,
 		NDIDebugDrawLocal::DrawLinePersistentName,
 		NDIDebugDrawLocal::DrawRectanglePersistentName,
 		NDIDebugDrawLocal::DrawCirclePersistentName,
@@ -2971,7 +2960,26 @@ bool UNiagaraDataInterfaceDebugDraw::GetFunctionHLSL(const FNiagaraDataInterface
 		NDIDebugDrawLocal::DrawGrid2DPersistentName,
 		NDIDebugDrawLocal::DrawGrid3DPersistentName,
 	};
+	if (PersistentFunctions.Contains(FunctionInfo.DefinitionName))
+	{
+		OutHLSL.Appendf(TEXT("Make%s(%s)\r\n"), *FunctionInfo.DefinitionName.ToString(), *FunctionInfo.InstanceName);
+		return true;
+	}
 
+	static const TSet<FName> ValidGpuFunctions =
+	{
+		NDIDebugDrawLocal::DrawLineName,
+		NDIDebugDrawLocal::DrawRectangleName,
+		NDIDebugDrawLocal::DrawCircleName,
+		NDIDebugDrawLocal::DrawBoxName,
+		NDIDebugDrawLocal::DrawSphereName,
+		NDIDebugDrawLocal::DrawCylinderName,
+		NDIDebugDrawLocal::DrawConeName,
+		NDIDebugDrawLocal::DrawTorusName,
+		NDIDebugDrawLocal::DrawCoordinateSystemName,
+		NDIDebugDrawLocal::DrawGrid2DName,
+		NDIDebugDrawLocal::DrawGrid3DName,
+	};
 	return ValidGpuFunctions.Contains(FunctionInfo.DefinitionName);
 }
 
