Niagara - Fix resolve data interfaces so that it synchronizes internal data interfaces for gpu scripts.  Also update the test which verifies interoplated spawn and update synchroniztion to test the gpu case as well and also verify that the editor data interface can be resolved to the runtime data interface.

diff --git a/Engine/Plugins/FX/Niagara/Source/Niagara/Private/NiagaraResolveDIHelpers.cpp b/Engine/Plugins/FX/Niagara/Source/Niagara/Private/NiagaraResolveDIHelpers.cpp
--- a/Engine/Plugins/FX/Niagara/Source/Niagara/Private/NiagaraResolveDIHelpers.cpp
+++ b/Engine/Plugins/FX/Niagara/Source/Niagara/Private/NiagaraResolveDIHelpers.cpp
@@ -309,16 +309,17 @@ namespace FNiagaraResolveDIHelpers
 		}
 	}
 
-	/** Handles the special case where internal data interfaces defined in particle update need to be copied to particle spawn so that they're using the same instance. */
-	void ResolveInternalDIsForInterpolatedSpawning(UNiagaraScript* ParticleSpawnScript, UNiagaraScript* ParticleUpdateScript)
+	/** Handles the special case where internal data interfaces defined in certain scripts need to be synchronized with other scripts that contain them. */
+	void ResolveInternalDataInterfaces(TArray<UNiagaraScript*> SourceScripts, UNiagaraScript* TargetScript)
 	{
-		TArray<FNiagaraScriptResolvedDataInterfaceInfo> ResolvedSpawnDataInterfaces;
-		ResolvedSpawnDataInterfaces.Append(ParticleSpawnScript->GetResolvedDataInterfaces());
-
-		TArrayView<const FNiagaraScriptResolvedDataInterfaceInfo> ResolvedUpdateDataInterfaces = ParticleUpdateScript->GetResolvedDataInterfaces();
-		SynchronizeMatchingInternalResolvedDataInterfaces(ResolvedUpdateDataInterfaces, ResolvedSpawnDataInterfaces);
-
-		ParticleSpawnScript->SetResolvedDataInterfaces(ResolvedSpawnDataInterfaces);
+		TArray<FNiagaraScriptResolvedDataInterfaceInfo> ResolvedTargetDataInterfaces;
+		ResolvedTargetDataInterfaces.Append(TargetScript->GetResolvedDataInterfaces());
+		for (UNiagaraScript* SourceScript : SourceScripts)
+		{
+			TArrayView<const FNiagaraScriptResolvedDataInterfaceInfo> ResolvedSourceParticleDataInterfaces = SourceScript->GetResolvedDataInterfaces();
+			SynchronizeMatchingInternalResolvedDataInterfaces(ResolvedSourceParticleDataInterfaces, ResolvedTargetDataInterfaces);
+		}
+		TargetScript->SetResolvedDataInterfaces(ResolvedTargetDataInterfaces);
 	}
 
 	void ResolveDIsInternal(
@@ -357,7 +358,8 @@ namespace FNiagaraResolveDIHelpers
 
 			if (VersionedNiagaraEmitterData->bInterpolatedSpawning)
 			{
-				ResolveInternalDIsForInterpolatedSpawning(VersionedNiagaraEmitterData->SpawnScriptProps.Script, VersionedNiagaraEmitterData->UpdateScriptProps.Script);
+				TArray<UNiagaraScript*> SourceScripts = { VersionedNiagaraEmitterData->UpdateScriptProps.Script };
+				ResolveInternalDataInterfaces(SourceScripts, VersionedNiagaraEmitterData->SpawnScriptProps.Script);
 			}
 
 			for (const FNiagaraEventScriptProperties& EventHandler : VersionedNiagaraEmitterData->GetEventHandlers())
@@ -373,6 +375,21 @@ namespace FNiagaraResolveDIHelpers
 				}
 				ResolveDIsForScript(System, SimulationStage->Script, EmitterHandle.GetUniqueInstanceName(), EmitterVariableAssignmentMap, EmitterCompactedVariableBindingMap, OutErrorMessages);
 			}
+
+			if (VersionedNiagaraEmitterData->SimTarget == ENiagaraSimTarget::GPUComputeSim)
+			{
+				TArray<UNiagaraScript*> SourceScripts;
+				SourceScripts.Add(VersionedNiagaraEmitterData->SpawnScriptProps.Script);
+				SourceScripts.Add(VersionedNiagaraEmitterData->UpdateScriptProps.Script);
+				for (const UNiagaraSimulationStageBase* SimulationStage : VersionedNiagaraEmitterData->GetSimulationStages())
+				{
+					if (SimulationStage != nullptr && SimulationStage->bEnabled && SimulationStage->Script != nullptr)
+					{
+						SourceScripts.Add(SimulationStage->Script);
+					}
+				}
+				ResolveInternalDataInterfaces(SourceScripts, VersionedNiagaraEmitterData->GetGPUComputeScript());
+			}
 		}
 	}
 
