Allow python conversion scripts to reset values to their default + set hidden inputs on the new module

diff --git a/Engine/Plugins/FX/Niagara/Source/NiagaraEditor/Private/NiagaraClipboard.cpp b/Engine/Plugins/FX/Niagara/Source/NiagaraEditor/Private/NiagaraClipboard.cpp
--- a/Engine/Plugins/FX/Niagara/Source/NiagaraEditor/Private/NiagaraClipboard.cpp
+++ b/Engine/Plugins/FX/Niagara/Source/NiagaraEditor/Private/NiagaraClipboard.cpp
@@ -98,6 +98,13 @@ const UNiagaraClipboardFunctionInput* UNiagaraClipboardFunctionInput::CreateDyna
 	return NewInput;
 }
 
+const UNiagaraClipboardFunctionInput* UNiagaraClipboardFunctionInput::CreateDefaultInputValue(UObject* InOuter, FName InInputName, FNiagaraTypeDefinition InInputType)
+{
+	UNiagaraClipboardFunctionInput* NewInput = MakeNewInput(InOuter, InInputName, InInputType, TOptional<bool>(), ENiagaraClipboardFunctionInputValueMode::ResetToDefault);
+	NewInput->Local.SetNumZeroed(InInputType.GetSize());
+	return NewInput;
+}
+
 
 bool UNiagaraClipboardFunctionInput::CopyValuesFrom(const UNiagaraClipboardFunctionInput* InOther)
 {
diff --git a/Engine/Plugins/FX/Niagara/Source/NiagaraEditor/Private/UpgradeNiagaraScriptResults.cpp b/Engine/Plugins/FX/Niagara/Source/NiagaraEditor/Private/UpgradeNiagaraScriptResults.cpp
--- a/Engine/Plugins/FX/Niagara/Source/NiagaraEditor/Private/UpgradeNiagaraScriptResults.cpp
+++ b/Engine/Plugins/FX/Niagara/Source/NiagaraEditor/Private/UpgradeNiagaraScriptResults.cpp
@@ -319,6 +319,15 @@ void UUpgradeNiagaraScriptResults::SetNewInput(const FString& InputName, UNiagar
 	}
 }
 
+void UUpgradeNiagaraScriptResults::ResetToDefault(const FString& InputName)
+{
+	if (UNiagaraPythonScriptModuleInput* ModuleInput = GetNewInput(FName(InputName)))
+	{
+		const UNiagaraClipboardFunctionInput* CurrentInput = ModuleInput->Input;
+		ModuleInput->Input = UNiagaraClipboardFunctionInput::CreateDefaultInputValue(ModuleInput, CurrentInput->InputName, CurrentInput->InputType);
+	}
+}
+
 UNiagaraPythonScriptModuleInput* UUpgradeNiagaraScriptResults::GetNewInput(const FName& InputName) const
 {
 	for (UNiagaraPythonScriptModuleInput* ModuleInput : NewInputs)
diff --git a/Engine/Plugins/FX/Niagara/Source/NiagaraEditor/Private/ViewModels/Stack/NiagaraStackInputCategory.cpp b/Engine/Plugins/FX/Niagara/Source/NiagaraEditor/Private/ViewModels/Stack/NiagaraStackInputCategory.cpp
--- a/Engine/Plugins/FX/Niagara/Source/NiagaraEditor/Private/ViewModels/Stack/NiagaraStackInputCategory.cpp
+++ b/Engine/Plugins/FX/Niagara/Source/NiagaraEditor/Private/ViewModels/Stack/NiagaraStackInputCategory.cpp
@@ -188,7 +188,14 @@ void SetValuesFromFunctionInputsInternal(UNiagaraStackInputCategory* Category, c
 		{
 			if (InputMatchesFilter(ChildInput) && ChildInput->GetInputParameterHandle().GetName() == ClipboardFunctionInput->InputName && ChildInput->GetInputType() == ClipboardFunctionInput->InputType)
 			{
-				ChildInput->SetValueFromClipboardFunctionInput(*ClipboardFunctionInput);
+				if (ClipboardFunctionInput->ValueMode == ENiagaraClipboardFunctionInputValueMode::ResetToDefault)
+				{
+					ChildInput->Reset();
+				}
+				else
+				{
+					ChildInput->SetValueFromClipboardFunctionInput(*ClipboardFunctionInput);
+				}
 				if(OwningFunctionCollection)
 				{
 					OwningFunctionCollection->RefreshChildren();
diff --git a/Engine/Plugins/FX/Niagara/Source/NiagaraEditor/Private/ViewModels/Stack/NiagaraStackModuleItem.cpp b/Engine/Plugins/FX/Niagara/Source/NiagaraEditor/Private/ViewModels/Stack/NiagaraStackModuleItem.cpp
--- a/Engine/Plugins/FX/Niagara/Source/NiagaraEditor/Private/ViewModels/Stack/NiagaraStackModuleItem.cpp
+++ b/Engine/Plugins/FX/Niagara/Source/NiagaraEditor/Private/ViewModels/Stack/NiagaraStackModuleItem.cpp
@@ -1460,6 +1460,41 @@ void UNiagaraStackModuleItem::ReassignModuleScript(UNiagaraScript* ModuleScript)
 				NewClipboardContent->FunctionInputs = NewClipboardContent->Functions[0]->Inputs;
 				NewClipboardContent->Functions.Empty();
 			}
+
+			// There can be hidden inputs due to static switches, so we need to add all of them to the clipboard for the python script to see them.
+			TArray<FNiagaraVariable> StaticSwitchInputs = FunctionCallNode->GetCalledGraph()->FindStaticSwitchInputs();
+			TArray<FNiagaraVariable> InputVariables;
+			TSet<FNiagaraVariable> HiddenInputVariables;
+			FCompileConstantResolver Resolver = GetEmitterViewModel().IsValid() ?
+			FCompileConstantResolver(GetEmitterViewModel()->GetEmitter(), FNiagaraStackGraphUtilities::GetOutputNodeUsage(*FunctionCallNode)) :
+			FCompileConstantResolver(&GetSystemViewModel()->GetSystem(), FNiagaraStackGraphUtilities::GetOutputNodeUsage(*FunctionCallNode));
+			GetStackFunctionInputs(*FunctionCallNode, InputVariables, HiddenInputVariables, Resolver, FNiagaraStackGraphUtilities::ENiagaraGetStackFunctionInputPinsOptions::ModuleInputsOnly);
+
+			TSet<FNiagaraVariableBase> SeenVars;
+			for (const UNiagaraClipboardFunctionInput* Input : NewClipboardContent->FunctionInputs)
+			{
+				SeenVars.Add(FNiagaraVariable(Input->GetTypeDef(), Input->InputName));
+			}
+			for (const FNiagaraVariableBase& StaticVar : StaticSwitchInputs)
+			{
+				if (!SeenVars.Contains(StaticVar))
+				{
+					NewClipboardContent->FunctionInputs.Add(UNiagaraClipboardFunctionInput::CreateDefaultInputValue(NewClipboardContent, StaticVar.GetName(), StaticVar.GetType()));
+					SeenVars.Add(StaticVar);
+				}
+			}
+			for (const FNiagaraVariableBase& InputVar : InputVariables)
+			{
+				FString InputName = InputVar.GetName().ToString();
+				InputName.RemoveFromStart(PARAM_MAP_MODULE_STR);
+				FNiagaraVariableBase NewVar(InputVar.GetType(), FName(InputName));
+				if (!SeenVars.Contains(NewVar))
+				{
+					NewClipboardContent->FunctionInputs.Add(UNiagaraClipboardFunctionInput::CreateDefaultInputValue(NewClipboardContent, NewVar.GetName(), NewVar.GetType()));
+					SeenVars.Add(NewVar);
+				}
+			}
+			
 			FunctionCallNode->PythonUpgradeScriptWarnings.Empty();
 			FText Warnings;
 			if (UNiagaraClipboardContent* NewInputs = FNiagaraEditorUtilities::RunPythonConversionScript(*NewScriptData, NewClipboardContent, *OldScriptData, OldClipboardContent, Warnings))
diff --git a/Engine/Plugins/FX/Niagara/Source/NiagaraEditor/Public/NiagaraClipboard.h b/Engine/Plugins/FX/Niagara/Source/NiagaraEditor/Public/NiagaraClipboard.h
--- a/Engine/Plugins/FX/Niagara/Source/NiagaraEditor/Public/NiagaraClipboard.h
+++ b/Engine/Plugins/FX/Niagara/Source/NiagaraEditor/Public/NiagaraClipboard.h
@@ -24,7 +24,8 @@ enum class ENiagaraClipboardFunctionInputValueMode
 	Data,
 	ObjectAsset,
 	Expression,
-	Dynamic
+	Dynamic,
+	ResetToDefault // special paste mode where it resets to the default value
 };
 
 class UNiagaraClipboardFunction;
@@ -47,6 +48,8 @@ public:
 
 	static const UNiagaraClipboardFunctionInput* CreateDynamicValue(UObject* InOuter, FName InInputName, FNiagaraTypeDefinition InInputType, TOptional<bool> bInEditConditionValue, FString InDynamicValueName, UNiagaraScript* InDynamicValue, const FGuid& InScriptVersion);
 
+	static const UNiagaraClipboardFunctionInput* CreateDefaultInputValue(UObject* InOuter, FName InInputName, FNiagaraTypeDefinition InInputType);
+	
 	UPROPERTY()
 	FName InputName;
 
diff --git a/Engine/Plugins/FX/Niagara/Source/NiagaraEditor/Public/UpgradeNiagaraScriptResults.h b/Engine/Plugins/FX/Niagara/Source/NiagaraEditor/Public/UpgradeNiagaraScriptResults.h
--- a/Engine/Plugins/FX/Niagara/Source/NiagaraEditor/Public/UpgradeNiagaraScriptResults.h
+++ b/Engine/Plugins/FX/Niagara/Source/NiagaraEditor/Public/UpgradeNiagaraScriptResults.h
@@ -146,6 +146,9 @@ public:
 	UFUNCTION(BlueprintCallable, Category = "Scripting")
 	void SetNewInput(const FString& InputName, UNiagaraPythonScriptModuleInput* Value);
 
+	UFUNCTION(BlueprintCallable, Category = "Scripting")
+	void ResetToDefault(const FString& InputName);
+
 private:
 	UNiagaraPythonScriptModuleInput* GetNewInput(const FName& InputName) const;
 
