Crash fix for NDCWrite DI being used as a custom input for module defaults.

diff --git a/Engine/Plugins/FX/Niagara/Source/Niagara/Private/DataInterface/NiagaraDataInterfaceDataChannelWrite.cpp b/Engine/Plugins/FX/Niagara/Source/Niagara/Private/DataInterface/NiagaraDataInterfaceDataChannelWrite.cpp
--- a/Engine/Plugins/FX/Niagara/Source/Niagara/Private/DataInterface/NiagaraDataInterfaceDataChannelWrite.cpp
+++ b/Engine/Plugins/FX/Niagara/Source/Niagara/Private/DataInterface/NiagaraDataInterfaceDataChannelWrite.cpp
@@ -440,7 +440,7 @@ void UNiagaraDataInterfaceDataChannelWrite::GetFeedback(UNiagaraSystem* InAsset,
 {
 	INiagaraModule& NiagaraModule = FModuleManager::GetModuleChecked<INiagaraModule>("Niagara");
 	const INiagaraEditorOnlyDataUtilities& EditorOnlyDataUtilities = NiagaraModule.GetEditorOnlyDataUtilities();
-	UNiagaraDataInterface* RuntimeInstanceOfThis = EditorOnlyDataUtilities.IsEditorDataInterfaceInstance(this)
+	UNiagaraDataInterface* RuntimeInstanceOfThis = InAsset && EditorOnlyDataUtilities.IsEditorDataInterfaceInstance(this)
 		? EditorOnlyDataUtilities.GetResolvedRuntimeInstanceForEditorDataInterfaceInstance(*InAsset, *this)
 		: this;
 
