* Fixed lwc types for python script api
* Fixed a bug when reusing unset input

diff --git a/Engine/Plugins/FX/Niagara/Source/NiagaraEditor/Private/NiagaraClipboard.cpp b/Engine/Plugins/FX/Niagara/Source/NiagaraEditor/Private/NiagaraClipboard.cpp
--- a/Engine/Plugins/FX/Niagara/Source/NiagaraEditor/Private/NiagaraClipboard.cpp
+++ b/Engine/Plugins/FX/Niagara/Source/NiagaraEditor/Private/NiagaraClipboard.cpp
@@ -56,7 +56,7 @@ UNiagaraClipboardFunctionInput* MakeNewInput(UObject* InOuter, FName InInputName
 
 const UNiagaraClipboardFunctionInput* UNiagaraClipboardFunctionInput::CreateLocalValue(UObject* InOuter, FName InInputName, FNiagaraTypeDefinition InInputType, TOptional<bool> bInEditConditionValue, TArray<uint8>& InLocalValueData)
 {
-	checkf(InLocalValueData.Num() == InInputType.GetSize(), TEXT("Input data size didn't match type size."))
+	checkf(InLocalValueData.Num() == InInputType.GetSize(), TEXT("Input data size (%d) didn't match type size (%d)."), InLocalValueData.Num(), InInputType.GetSize());
 	UNiagaraClipboardFunctionInput* NewInput = MakeNewInput(InOuter, InInputName, InInputType, bInEditConditionValue, ENiagaraClipboardFunctionInputValueMode::Local);
 	NewInput->Local = InLocalValueData;
 	return NewInput;
diff --git a/Engine/Plugins/FX/Niagara/Source/NiagaraEditor/Private/UpgradeNiagaraScriptResults.cpp b/Engine/Plugins/FX/Niagara/Source/NiagaraEditor/Private/UpgradeNiagaraScriptResults.cpp
--- a/Engine/Plugins/FX/Niagara/Source/NiagaraEditor/Private/UpgradeNiagaraScriptResults.cpp
+++ b/Engine/Plugins/FX/Niagara/Source/NiagaraEditor/Private/UpgradeNiagaraScriptResults.cpp
@@ -241,7 +241,7 @@ void UUpgradeNiagaraScriptResults::SetVec2Input(const FString& InputName, FVecto
 	UNiagaraPythonScriptModuleInput* ModuleInput = GetNewInput(FName(InputName));
 	if (ModuleInput && ModuleInput->Input->InputType == FNiagaraTypeDefinition::GetVec2Def())
 	{
-		NiagaraScriptResults::SetValue(ModuleInput, Value);
+		NiagaraScriptResults::SetValue(ModuleInput, FVector2f(Value));
 	}
 }
 
@@ -250,7 +250,7 @@ void UUpgradeNiagaraScriptResults::SetVec3Input(const FString& InputName, FVecto
 	UNiagaraPythonScriptModuleInput* ModuleInput = GetNewInput(FName(InputName));
 	if (ModuleInput && ModuleInput->Input->InputType == FNiagaraTypeDefinition::GetVec3Def())
 	{
-		NiagaraScriptResults::SetValue(ModuleInput, Value);
+		NiagaraScriptResults::SetValue(ModuleInput, FVector3f(Value));
 	}
 }
 
@@ -259,7 +259,7 @@ void UUpgradeNiagaraScriptResults::SetVec4Input(const FString& InputName, FVecto
 	UNiagaraPythonScriptModuleInput* ModuleInput = GetNewInput(FName(InputName));
 	if (ModuleInput && ModuleInput->Input->InputType == FNiagaraTypeDefinition::GetVec4Def())
 	{
-		NiagaraScriptResults::SetValue(ModuleInput, Value);
+		NiagaraScriptResults::SetValue(ModuleInput, FVector4f(Value));
 	}
 }
 
@@ -277,7 +277,7 @@ void UUpgradeNiagaraScriptResults::SetQuatInput(const FString& InputName, FQuat
 	UNiagaraPythonScriptModuleInput* ModuleInput = GetNewInput(FName(InputName));
 	if (ModuleInput && ModuleInput->Input->InputType == FNiagaraTypeDefinition::GetQuatDef())
 	{
-		NiagaraScriptResults::SetValue(ModuleInput, Value);
+		NiagaraScriptResults::SetValue(ModuleInput, FQuat4f(Value));
 	}
 }
 
@@ -306,14 +306,21 @@ void UUpgradeNiagaraScriptResults::SetNewInput(const FString& InputName, UNiagar
 	{
 		UNiagaraPythonScriptModuleInput* ModuleInput = NewInputs[i];
 		UNiagaraClipboardFunctionInput* FunctionInput = const_cast<UNiagaraClipboardFunctionInput*>(ModuleInput->Input.Get());
-		if (Value && FunctionInput && FunctionInput->InputName == InputName && FunctionInput->InputType == Value->Input->InputType)
+		if (Value && FunctionInput && FunctionInput->InputName == InputName )
 		{
-			FunctionInput->Data = Value->Input->Data;
-			FunctionInput->Dynamic = Value->Input->Dynamic;
-			FunctionInput->Expression = Value->Input->Expression;
-			FunctionInput->Linked = Value->Input->Linked;
-			FunctionInput->Local = Value->Input->Local;
-			FunctionInput->ValueMode = Value->Input->ValueMode;
+			if (Value->IsSet() && FunctionInput->InputType == Value->Input->InputType)
+			{
+				FunctionInput->Data = Value->Input->Data;
+				FunctionInput->Dynamic = Value->Input->Dynamic;
+				FunctionInput->Expression = Value->Input->Expression;
+				FunctionInput->Linked = Value->Input->Linked;
+				FunctionInput->Local = Value->Input->Local;
+				FunctionInput->ValueMode = Value->Input->ValueMode;
+			}
+			else
+			{
+				ResetToDefault(InputName);
+			}
 			return;
 		}
 	}
