- Fix for renderer feedback fixes not having a scoped transaction / refreshing the system
- This made undoing a fix impossible & with some fixes it may not even be visible

diff --git a/Engine/Plugins/FX/Niagara/Source/NiagaraEditor/Private/ViewModels/Stack/NiagaraStackRendererItem.cpp b/Engine/Plugins/FX/Niagara/Source/NiagaraEditor/Private/ViewModels/Stack/NiagaraStackRendererItem.cpp
--- a/Engine/Plugins/FX/Niagara/Source/NiagaraEditor/Private/ViewModels/Stack/NiagaraStackRendererItem.cpp
+++ b/Engine/Plugins/FX/Niagara/Source/NiagaraEditor/Private/ViewModels/Stack/NiagaraStackRendererItem.cpp
@@ -521,7 +521,36 @@ void UNiagaraStackRendererItem::ProcessRendererIssues(const TArray<FNiagaraRende
 		TArray<FStackIssueFix> Fixes;
 		if (Item.IsFixable())
 		{
-			Fixes.Add(FStackIssueFix(Item.GetFixDescriptionText(), FStackIssueFixDelegate::CreateLambda([Item]() { Item.TryFix(); })));
+			Fixes.Emplace(
+				Item.GetFixDescriptionText(),
+				FStackIssueFixDelegate::CreateLambda(
+					[WeakStackItem=TWeakObjectPtr<UNiagaraStackRendererItem>(this), Item]()
+					{
+						if ( Item.IsFixable() )
+						{
+							const FScopedTransaction ScopedTransaction(LOCTEXT("RendererItemFixTransaction", "Apply renderer fix"));
+
+							UNiagaraStackRendererItem* StackItem = WeakStackItem.Get();
+							UNiagaraRendererProperties* RendererProperties = StackItem ? StackItem->GetRendererProperties() : nullptr;
+							if (RendererProperties)
+							{
+								RendererProperties->Modify();
+							}
+
+							Item.TryFix();
+
+							if (RendererProperties)
+							{
+								RendererProperties->PostEditChange();
+
+								TArray<UObject*> ChangedObjects;
+								ChangedObjects.Add(RendererProperties);
+								StackItem->OnDataObjectModified().Broadcast(ChangedObjects, ENiagaraDataObjectChange::Changed);
+							}
+						}
+					}
+				)
+			);
 		}
 		FStackIssue TargetSupportError(Severity, Item.GetSummaryText(), Item.GetDescriptionText(), GetStackEditorDataKey(), Item.IsDismissable(), Fixes);
 		OutIssues.Add(TargetSupportError);
