- Default fallback for ribbon renderers change from normalized age binding to unique ID
- Normalized Age will give unpredictable results with burst modules and can result in inconsistent ordering with varying particle lifetimes
- Unique ID works around these issues since it's a monotonically increasing number

diff --git a/Engine/Plugins/FX/Niagara/Shaders/Private/Ribbons/NiagaraRibbonCommon.ush b/Engine/Plugins/FX/Niagara/Shaders/Private/Ribbons/NiagaraRibbonCommon.ush
--- a/Engine/Plugins/FX/Niagara/Shaders/Private/Ribbons/NiagaraRibbonCommon.ush
+++ b/Engine/Plugins/FX/Niagara/Shaders/Private/Ribbons/NiagaraRibbonCommon.ush
@@ -11,8 +11,6 @@
 #define PACKED_PER_RIBBON_STARTPARTICLEOFFSET 6
 #define PACKED_PER_RIBBON_ENDPARTICLEOFFSET 7
 
-
-
 #if RIBBONS_WANTS_AUTOMATIC_TESSELLATION
 #define VERTEX_GEN_OUTPUT_DATA_STRIDE 15
 #else
@@ -225,21 +223,16 @@ static const FNiagaraID InvalidID = { INT_MIN, INT_MIN };
 
 #endif
 
-
-#if RIBBON_HAS_LINK_ORDER
-
+#if RIBBONLINK_IS_FLOAT
 	float GetRibbonLinkID(int InstanceID)
 	{
 		return GetFloat(RibbonLinkOrderDataOffset, InstanceID);
 	}
-
 #else
-
-	float GetRibbonLinkID(int InstanceID)
+	int GetRibbonLinkID(int InstanceID)
 	{
-		return GetFloat(NormalizedAgeDataOffset, InstanceID);
+		return GetInt(RibbonLinkOrderDataOffset, InstanceID);
 	}
-
 #endif
 
 // max absolute error 9.0x10^-3
diff --git a/Engine/Plugins/FX/Niagara/Shaders/Private/Ribbons/NiagaraRibbonSortParticles.usf b/Engine/Plugins/FX/Niagara/Shaders/Private/Ribbons/NiagaraRibbonSortParticles.usf
--- a/Engine/Plugins/FX/Niagara/Shaders/Private/Ribbons/NiagaraRibbonSortParticles.usf
+++ b/Engine/Plugins/FX/Niagara/Shaders/Private/Ribbons/NiagaraRibbonSortParticles.usf
@@ -11,10 +11,18 @@ uint CountOffset;
 int MergeSortSourceBlockSize;
 int MergeSortDestinationBlockSize;
 
-groupshared uint		SortParticleIDs[THREADGROUP_SIZE + 1];
-groupshared float		SortLinkIDs[THREADGROUP_SIZE + 1];
+#if RIBBONLINK_IS_FLOAT
+	#define FRibbonLink	float
+	bool CompareRibbonLink(float Lhs, float Rhs) { return Lhs > Rhs; }
+#else
+	#define FRibbonLink	int
+	bool CompareRibbonLink(int Lhs, int Rhs) { return Lhs < Rhs; }
+#endif
+
+groupshared uint			SortParticleIDs[THREADGROUP_SIZE + 1];
+groupshared FRibbonLink		SortLinkIDs[THREADGROUP_SIZE + 1];
 #if HAS_RIBBON_ID
-groupshared FRibbonID	SortRibbonIDs[THREADGROUP_SIZE + 1];
+groupshared FRibbonID		SortRibbonIDs[THREADGROUP_SIZE + 1];
 #endif
 
 [numthreads(THREADGROUP_SIZE, 1, 1)]
@@ -45,12 +53,12 @@ void InitialSortList(uint3 DispatchThreadId : SV_DispatchThreadID, uint InGroupI
 		bool bShouldSwap = false;
 		if ( ((LocalIndexOffset & 1) == (SortGroup & 1)) && (LocalIndexOffset + 1 < TotalNumParticlesInSet) )
 		{
-#if HAS_RIBBON_ID			
+		#if HAS_RIBBON_ID			
 			const int RibbonIDComparisonResult = CompareRibbonIDs(SortRibbonIDs[LocalIndexOffset + 0], SortRibbonIDs[LocalIndexOffset + 1]);
-			bShouldSwap = RibbonIDComparisonResult > 0 || (RibbonIDComparisonResult == 0 && SortLinkIDs[LocalIndexOffset + 0] > SortLinkIDs[LocalIndexOffset + 1]);
-#else
-			bShouldSwap = SortLinkIDs[LocalIndexOffset + 0] > SortLinkIDs[LocalIndexOffset + 1];
-#endif
+			bShouldSwap = RibbonIDComparisonResult > 0 || (RibbonIDComparisonResult == 0 && CompareRibbonLink(SortLinkIDs[LocalIndexOffset + 0], SortLinkIDs[LocalIndexOffset + 1]));
+		#else
+			bShouldSwap = CompareRibbonLink(SortLinkIDs[LocalIndexOffset + 0], SortLinkIDs[LocalIndexOffset + 1]);
+		#endif
 		}
 
 		GroupMemoryBarrierWithGroupSync();
@@ -60,7 +68,7 @@ void InitialSortList(uint3 DispatchThreadId : SV_DispatchThreadID, uint InGroupI
 			SortParticleIDs[LocalIndexOffset + 0] = SortParticleIDs[LocalIndexOffset + 1];
 			SortParticleIDs[LocalIndexOffset + 1] = TempParticleID;
 
-			float TempLinkID = SortLinkIDs[LocalIndexOffset + 0];
+			FRibbonLink TempLinkID = SortLinkIDs[LocalIndexOffset + 0];
 			SortLinkIDs[LocalIndexOffset + 0] = SortLinkIDs[LocalIndexOffset + 1];
 			SortLinkIDs[LocalIndexOffset + 1] = TempLinkID;
 
@@ -80,9 +88,9 @@ void InitialSortList(uint3 DispatchThreadId : SV_DispatchThreadID, uint InGroupI
 }
 
 
-int MergeSortCompare(float SearchLinkOrder, float TestLinkOrder)
+int MergeSortCompare(FRibbonLink SearchLinkOrder, FRibbonLink TestLinkOrder)
 {
-	if (SearchLinkOrder > TestLinkOrder)
+	if (CompareRibbonLink(SearchLinkOrder, TestLinkOrder))
 	{
 		return 1;
 	}
@@ -96,7 +104,7 @@ int MergeSortCompare(float SearchLinkOrder, float TestLinkOrder)
 }
 
 #if HAS_RIBBON_ID
-int MergeSortCompare(FRibbonID Left, float LeftLinkOrder, FRibbonID Right, float RightLinkOrder)
+int MergeSortCompare(FRibbonID Left, FRibbonLink LeftLinkOrder, FRibbonID Right, FRibbonLink RightLinkOrder)
 {
 	const int RibbonIDComparisonResult = CompareRibbonIDs(Left, Right);
 
@@ -131,7 +139,7 @@ void MergeSort(uint3 DispatchThreadId : SV_DispatchThreadID)
 		const uint IndexWithinSourceBlock = Index - SourceBlockStart;
 		
 		const uint ParticleID = SortedIndices[Index];
-		const float LinkID = GetRibbonLinkID(ParticleID);
+		const FRibbonLink LinkID = GetRibbonLinkID(ParticleID);
 #if HAS_RIBBON_ID
 		const FRibbonID RibbonID = GetRibbonID(ParticleID);
 #endif
@@ -149,7 +157,7 @@ void MergeSort(uint3 DispatchThreadId : SV_DispatchThreadID)
 			const int Mid = (SearchStart + SearchStop) / 2;
 			const int SearchIndex = OtherSourceBlockStart + Mid;
 			const uint OtherParticleID = SortedIndices[SearchIndex];
-			const float OtherLinkID = GetRibbonLinkID(OtherParticleID);
+			const FRibbonLink OtherLinkID = GetRibbonLinkID(OtherParticleID);
 #if HAS_RIBBON_ID
 			const FRibbonID OtherRibbonID = GetRibbonID(OtherParticleID);
 			const int Result = MergeSortCompare(RibbonID, LinkID, OtherRibbonID, OtherLinkID);
diff --git a/Engine/Plugins/FX/Niagara/Source/Niagara/Private/NiagaraRendererRibbons.cpp b/Engine/Plugins/FX/Niagara/Source/Niagara/Private/NiagaraRendererRibbons.cpp
--- a/Engine/Plugins/FX/Niagara/Source/Niagara/Private/NiagaraRendererRibbons.cpp
+++ b/Engine/Plugins/FX/Niagara/Source/Niagara/Private/NiagaraRendererRibbons.cpp
@@ -647,6 +647,49 @@ private:
 	void GenerateAllGPUData(FRHICommandListImmediate& RHICmdList, TArray<FNiagaraRibbonGPUInitParameters>& RenderersToGenerate);
 };
 
+struct FNiagaraGenerationInputDataCPUAccessors
+{
+	FNiagaraGenerationInputDataCPUAccessors(const UNiagaraRibbonRendererProperties * Properties, const FNiagaraDataSet & Data)
+		: TotalNumParticles(Data.GetCurrentDataChecked().GetNumInstances())
+		, RibbonLinkOrderFloatData(Properties->RibbonLinkOrderFloatAccessor.GetReader(Data))
+		, RibbonLinkOrderInt32Data(Properties->RibbonLinkOrderInt32Accessor.GetReader(Data))
+		, SimpleRibbonIDData(Properties->RibbonIdDataSetAccessor.GetReader(Data))
+		, FullRibbonIDData(Properties->RibbonFullIDDataSetAccessor.GetReader(Data))
+		, PosData(Properties->PositionDataSetAccessor.GetReader(Data))
+		, AgeData(Properties->NormalizedAgeAccessor.GetReader(Data))
+		, SizeData(Properties->SizeDataSetAccessor.GetReader(Data))
+		, TwistData(Properties->TwistDataSetAccessor.GetReader(Data))
+	{
+	}
+
+	template<typename TContainer>
+	void RibbonLinkOrderSort(TContainer& Container) const
+	{
+		if (RibbonLinkOrderFloatData.IsValid())
+		{
+			Container.Sort([this](const uint32& A, const uint32& B) { return RibbonLinkOrderFloatData[A] < RibbonLinkOrderFloatData[B]; });
+		}
+		else
+		{
+			Container.Sort([this](const uint32& A, const uint32& B) { return RibbonLinkOrderInt32Data[A] > RibbonLinkOrderInt32Data[B]; });
+		}
+	}
+
+	bool HasRibbonLinkOrder() const { return RibbonLinkOrderFloatData.IsValid() || RibbonLinkOrderInt32Data.IsValid(); }
+
+	const uint32 TotalNumParticles;
+
+	const FNiagaraDataSetReaderFloat<float>	RibbonLinkOrderFloatData;
+	const FNiagaraDataSetReaderInt32<int32>	RibbonLinkOrderInt32Data;
+
+	const FNiagaraDataSetReaderInt32<int> SimpleRibbonIDData;
+	const FNiagaraDataSetReaderStruct<FNiagaraID> FullRibbonIDData;
+
+	const FNiagaraDataSetReaderFloat<FNiagaraPosition> PosData;
+	const FNiagaraDataSetReaderFloat<float> AgeData;
+	const FNiagaraDataSetReaderFloat<float> SizeData;
+	const FNiagaraDataSetReaderFloat<float> TwistData;
+};
 
 FNiagaraRendererRibbons::FNiagaraRendererRibbons(ERHIFeatureLevel::Type FeatureLevel, const UNiagaraRendererProperties *InProps, const FNiagaraEmitterInstance* Emitter)
 	: FNiagaraRenderer(FeatureLevel, InProps, Emitter)
@@ -666,7 +709,9 @@ FNiagaraRendererRibbons::FNiagaraRendererRibbons(ERHIFeatureLevel::Type FeatureL
 	FacingMode = Properties->FacingMode;
 	DrawDirection = Properties->DrawDirection;
 	RendererLayout = &Properties->RendererLayout;
-	
+	bGpuRibbonLinkIsFloat = Properties->bGpuRibbonLinkIsFloat;
+	GpuRibbonLinkOrderOffset = Properties->GpuRibbonLinkOrderOffset;
+
 	InitializeShape(Properties);
 	InitializeTessellation(Properties);	
 }
@@ -886,7 +931,7 @@ FNiagaraDynamicDataBase* FNiagaraRendererRibbons::GenerateDynamicData(const FNia
 				
 				DynamicData->GenerationOutput = MakeShared<FNiagaraRibbonCPUGeneratedVertexData>();
 
-				if (CPUData.PosData.IsValid() && CPUData.SortKeyReader.IsValid() && CPUData.TotalNumParticles >= 2)
+				if (CPUData.PosData.IsValid() && CPUData.HasRibbonLinkOrder() && CPUData.TotalNumParticles >= 2)
 				{
 					GenerateVertexBufferCPU(CPUData, *DynamicData->GenerationOutput);
 				}
@@ -1577,11 +1622,10 @@ void FNiagaraRendererRibbons::GenerateVertexBufferForMultiRibbonInternal(const F
 	for (TPair<IDType, FIndexArray>& Pair : MultiRibbonSortedIndices)
 	{
 		FIndexArray& SortedIndices = Pair.Value;
-		const auto& SortKeyReader = CPUData.SortKeyReader;
-		SortedIndices.Sort([&SortKeyReader](const uint32& A, const uint32& B) { return (SortKeyReader[A] < SortKeyReader[B]); });
+		CPUData.RibbonLinkOrderSort(SortedIndices);
 		GenerateVertexBufferForRibbonPart<bWantsTessellation, bHasTwist, true>(CPUData, SortedIndices, RibbonIndex, OutputData);
 		RibbonIndex++;
-	};
+	}
 }
 
 template<typename IDType, typename ReaderType>
@@ -1608,7 +1652,7 @@ void FNiagaraRendererRibbons::GenerateVertexBufferCPU(const FNiagaraGenerationIn
 {
 	SCOPE_CYCLE_COUNTER(STAT_NiagaraRenderRibbonsGenVerticesCPU);
 	
-	check(CPUData.PosData.IsValid() && CPUData.SortKeyReader.IsValid());
+	check(CPUData.PosData.IsValid() && CPUData.HasRibbonLinkOrder());
 
 	// TODO: Move sorting to share code with sprite and mesh sorting and support the custom sorting key.
 	FMemMark Mark(FMemStack::Get());
@@ -1636,8 +1680,7 @@ void FNiagaraRendererRibbons::GenerateVertexBufferCPU(const FNiagaraGenerationIn
 		}
 		OutputData.RibbonInfoLookup.AddZeroed(1);
 
-		const auto& SortKeyReader = CPUData.SortKeyReader;
-		SortedIndices.Sort([&SortKeyReader](const uint32& A, const uint32& B) {	return (SortKeyReader[A] < SortKeyReader[B]); });
+		CPUData.RibbonLinkOrderSort(SortedIndices);
 
 		OutputData.SortedIndices.Reserve(OutputData.SortedIndices.Num() + CPUData.TotalNumParticles + 1);
 		OutputData.TangentAndDistances.Reserve(OutputData.TangentAndDistances.Num() + CPUData.TotalNumParticles + 1);
@@ -1989,8 +2032,6 @@ void FNiagaraRendererRibbons::SetupPerViewUniformBuffer(FNiagaraIndexGenerationI
 	PerViewUniformParameters.FacingDataOffset = bShouldDoFacing ? VFVariables[ENiagaraRibbonVFLayout::Facing].GetGPUOffset() : -1;
 	PerViewUniformParameters.PrevFacingDataOffset = bShouldDoFacing ? VFVariables[ENiagaraRibbonVFLayout::PrevRibbonFacing].GetGPUOffset() : -1;
 
-	PerViewUniformParameters.LinkOrderDataOffset = VFVariables[ENiagaraRibbonVFLayout::LinkOrder].GetGPUOffset();
-
 	PerViewUniformParameters.U0DistributionMode = static_cast<int32>(UV0Settings.DistributionMode);
 	PerViewUniformParameters.U1DistributionMode = static_cast<int32>(UV1Settings.DistributionMode);
 	PerViewUniformParameters.PackedVData.X = float(UV0Settings.Scale.Y);
@@ -2330,7 +2371,7 @@ FRibbonComputeUniformParameters FNiagaraRendererRibbons::SetupComputeVertexGenPa
 	CommonParams.FacingDataOffset = bShouldDoFacing ? VFVariables[ENiagaraRibbonVFLayout::Facing].GetGPUOffset() : -1;
 	CommonParams.PrevFacingDataOffset = bShouldDoFacing ? VFVariables[ENiagaraRibbonVFLayout::PrevRibbonFacing].GetGPUOffset() : -1;
 
-	CommonParams.RibbonLinkOrderDataOffset = GenerationConfig.HasCustomLinkOrder()? VFVariables[ENiagaraRibbonVFLayout::LinkOrder].GetGPUOffset() : -1;
+	CommonParams.RibbonLinkOrderDataOffset = GpuRibbonLinkOrderOffset;
 
 	CommonParams.U0DistributionMode = static_cast<int32>(UV0Settings.DistributionMode);
 	CommonParams.U1DistributionMode = static_cast<int32>(UV1Settings.DistributionMode);
@@ -2371,7 +2412,7 @@ void FNiagaraRendererRibbons::InitializeVertexBuffersGPU(FRHICommandListImmediat
 		FNiagaraRibbonSortPhase1CS::FPermutationDomain PermutationVector;
 		PermutationVector.Set<FRibbonHasFullRibbonID>(GenerationConfig.HasFullRibbonIDs());
 		PermutationVector.Set<FRibbonHasRibbonID>(GenerationConfig.HasSimpleRibbonIDs());
-		PermutationVector.Set<FRibbonHasCustomLinkOrder>(GenerationConfig.HasCustomLinkOrder());
+		PermutationVector.Set<FRibbonLinkIsFloat>(bGpuRibbonLinkIsFloat);
 		
 		TShaderMapRef<FNiagaraRibbonSortPhase1CS> BubbleSortShader(GetGlobalShaderMap(GMaxRHIFeatureLevel), PermutationVector);
 		TShaderMapRef<FNiagaraRibbonSortPhase2CS> MergeSortShader(GetGlobalShaderMap(GMaxRHIFeatureLevel), PermutationVector);
diff --git a/Engine/Plugins/FX/Niagara/Source/Niagara/Private/NiagaraRibbonRendererProperties.cpp b/Engine/Plugins/FX/Niagara/Source/Niagara/Private/NiagaraRibbonRendererProperties.cpp
--- a/Engine/Plugins/FX/Niagara/Source/Niagara/Private/NiagaraRibbonRendererProperties.cpp
+++ b/Engine/Plugins/FX/Niagara/Source/Niagara/Private/NiagaraRibbonRendererProperties.cpp
@@ -24,8 +24,13 @@
 
 #define LOCTEXT_NAMESPACE "UNiagaraRibbonRendererProperties"
 
-TArray<TWeakObjectPtr<UNiagaraRibbonRendererProperties>> UNiagaraRibbonRendererProperties::RibbonRendererPropertiesToDeferredInit;
+namespace NiagaraRibbonRendererPropertiesPrivate
+{
+	static const FName NAME_UniqueID("UniqueID");
+	static FNiagaraVariableBase GetUniqueIDVariable() { return FNiagaraVariableBase(FNiagaraTypeDefinition::GetIntDef(), NAME_UniqueID); }
 
+	TArray<TWeakObjectPtr<UNiagaraRibbonRendererProperties>> RibbonRendererPropertiesToDeferredInit;
+}
 
 FNiagaraRibbonShapeCustomVertex::FNiagaraRibbonShapeCustomVertex()
 	: Position(0.0f)
@@ -63,6 +68,7 @@ UNiagaraRibbonRendererProperties::UNiagaraRibbonRendererProperties()
 	, bUseGPUInit(false)
 	, bUseConstantFactor(false)
 	, bScreenSpaceTessellation(true)
+	, bLinkOrderUseUniqueID(true)
 	, WidthSegmentationCount(1)
 	, MultiPlaneCount(2)
 	, TubeSubdivisions(3)
@@ -71,32 +77,34 @@ UNiagaraRibbonRendererProperties::UNiagaraRibbonRendererProperties()
 	, TessellationFactor(16)
 	, TessellationAngle(15)
 {
-	AttributeBindings.Reserve(19);
-	AttributeBindings.Add(&PositionBinding);
-	AttributeBindings.Add(&ColorBinding);
-	AttributeBindings.Add(&VelocityBinding);
-	AttributeBindings.Add(&NormalizedAgeBinding);
-	AttributeBindings.Add(&RibbonTwistBinding);
-	AttributeBindings.Add(&RibbonWidthBinding);
-	AttributeBindings.Add(&RibbonFacingBinding);
-	AttributeBindings.Add(&RibbonIdBinding);
-	AttributeBindings.Add(&RibbonLinkOrderBinding);
+	AttributeBindings =
+	{
+		&PositionBinding,
+		&ColorBinding,
+		&VelocityBinding,
+		&NormalizedAgeBinding,
+		&RibbonTwistBinding,
+		&RibbonWidthBinding,
+		&RibbonFacingBinding,
+		&RibbonIdBinding,
+		&RibbonLinkOrderBinding,
 	
-	AttributeBindings.Add(&MaterialRandomBinding);
-	AttributeBindings.Add(&DynamicMaterialBinding);
-	AttributeBindings.Add(&DynamicMaterial1Binding);
-	AttributeBindings.Add(&DynamicMaterial2Binding);
-	AttributeBindings.Add(&DynamicMaterial3Binding);
-	AttributeBindings.Add(&RibbonUVDistance);
-	AttributeBindings.Add(&U0OverrideBinding);
-	AttributeBindings.Add(&V0RangeOverrideBinding);
-	AttributeBindings.Add(&U1OverrideBinding);
-	AttributeBindings.Add(&V1RangeOverrideBinding);
-
-	AttributeBindings.Add(&PrevPositionBinding);
-	AttributeBindings.Add(&PrevRibbonWidthBinding);
-	AttributeBindings.Add(&PrevRibbonFacingBinding);
-	AttributeBindings.Add(&PrevRibbonTwistBinding);
+		&MaterialRandomBinding,
+		&DynamicMaterialBinding,
+		&DynamicMaterial1Binding,
+		&DynamicMaterial2Binding,
+		&DynamicMaterial3Binding,
+		&RibbonUVDistance,
+		&U0OverrideBinding,
+		&V0RangeOverrideBinding,
+		&U1OverrideBinding,
+		&V1RangeOverrideBinding,
+
+		&PrevPositionBinding,
+		&PrevRibbonWidthBinding,
+		&PrevRibbonFacingBinding,
+		&PrevRibbonTwistBinding,
+	};
 }
 
 FNiagaraRenderer* UNiagaraRibbonRendererProperties::CreateEmitterRenderer(ERHIFeatureLevel::Type FeatureLevel, const FNiagaraEmitterInstance* Emitter, const FNiagaraSystemInstanceController& InController)
@@ -145,6 +153,11 @@ void UNiagaraRibbonRendererProperties::PostLoad()
 		UpgradeUVSettings(UV1Settings, UV1TilingDistance_DEPRECATED, UV1Offset_DEPRECATED, UV1Scale_DEPRECATED);
 	}
 
+	if (NiagaraVer < FNiagaraCustomVersion::RibbonRendererLinkOrderDefaultIsUniqueID)
+	{
+		bLinkOrderUseUniqueID = false;
+	}
+
 	ChangeToPositionBinding(PositionBinding);
 #endif
 	
@@ -219,6 +232,12 @@ bool UNiagaraRibbonRendererProperties::IsBackfaceCullingDisabled() const
 TArray<FNiagaraVariable> UNiagaraRibbonRendererProperties::GetBoundAttributes() const 
 {
 	TArray<FNiagaraVariable> BoundAttributes = Super::GetBoundAttributes();
+
+	if (bLinkOrderUseUniqueID)
+	{
+		BoundAttributes.AddUnique(NiagaraRibbonRendererPropertiesPrivate::GetUniqueIDVariable());
+	}
+
 	BoundAttributes.Reserve(BoundAttributes.Num() + MaterialParameters.AttributeBindings.Num());
 
 	for (const FNiagaraMaterialAttributeBinding& MaterialParamBinding : MaterialParameters.AttributeBindings)
@@ -274,7 +293,7 @@ void UNiagaraRibbonRendererProperties::PostInitProperties()
 		// We can end up hitting PostInitProperties before the Niagara Module has initialized bindings this needs, mark this object for deferred init and early out.
 		if (FModuleManager::Get().IsModuleLoaded("Niagara") == false)
 		{
-			RibbonRendererPropertiesToDeferredInit.Add(this);
+			NiagaraRibbonRendererPropertiesPrivate::RibbonRendererPropertiesToDeferredInit.Add(this);
 			return;
 		}
 		InitBindings();
@@ -312,7 +331,7 @@ void UNiagaraRibbonRendererProperties::InitCDOPropertiesAfterModuleStartup()
 	UNiagaraRibbonRendererProperties* CDO = CastChecked<UNiagaraRibbonRendererProperties>(UNiagaraRibbonRendererProperties::StaticClass()->GetDefaultObject());
 	CDO->InitBindings();
 
-	for (TWeakObjectPtr<UNiagaraRibbonRendererProperties>& WeakRibbonRendererProperties : RibbonRendererPropertiesToDeferredInit)
+	for (TWeakObjectPtr<UNiagaraRibbonRendererProperties>& WeakRibbonRendererProperties : NiagaraRibbonRendererPropertiesPrivate::RibbonRendererPropertiesToDeferredInit)
 	{
 		if (WeakRibbonRendererProperties.Get())
 		{
@@ -377,12 +396,18 @@ void UNiagaraRibbonRendererProperties::CacheFromCompiledData(const FNiagaraDataS
 	UpdateMICs();
 
 	// Initialize accessors
-	bSortKeyDataSetAccessorIsAge = false;
-	SortKeyDataSetAccessor.Init(CompiledData, RibbonLinkOrderBinding.GetDataSetBindableVariable().GetName());
-	if (!SortKeyDataSetAccessor.IsValid())
+	RibbonLinkOrderFloatAccessor.Init(CompiledData, RibbonLinkOrderBinding.GetDataSetBindableVariable().GetName());
+	RibbonLinkOrderInt32Accessor.Init(nullptr, NAME_None);
+	if (!RibbonLinkOrderFloatAccessor.IsValid())
 	{
-		bSortKeyDataSetAccessorIsAge = true;
-		SortKeyDataSetAccessor.Init(CompiledData, NormalizedAgeBinding.GetDataSetBindableVariable().GetName());
+		if (bLinkOrderUseUniqueID)
+		{
+			RibbonLinkOrderInt32Accessor.Init(CompiledData, NiagaraRibbonRendererPropertiesPrivate::NAME_UniqueID);
+		}
+		if (!RibbonLinkOrderInt32Accessor.IsValid())
+		{
+			RibbonLinkOrderFloatAccessor.Init(CompiledData, NormalizedAgeBinding.GetDataSetBindableVariable().GetName());
+		}
 	}
 
 	NormalizedAgeAccessor.Init(CompiledData, NormalizedAgeBinding.GetDataSetBindableVariable().GetName());
@@ -417,8 +442,6 @@ void UNiagaraRibbonRendererProperties::CacheFromCompiledData(const FNiagaraDataS
 		RibbonIdDataSetAccessor.Init(CompiledData, RibbonIdBinding.GetDataSetBindableVariable().GetName());
 	}
 
-	RibbonLinkOrderDataSetAccessor.Init(CompiledData, RibbonLinkOrderBinding.GetDataSetBindableVariable().GetName());
-
 	const bool bShouldDoFacing = FacingMode == ENiagaraRibbonFacingMode::Custom || FacingMode == ENiagaraRibbonFacingMode::CustomSideVector;
 
 	// Initialize layout
@@ -452,10 +475,36 @@ void UNiagaraRibbonRendererProperties::CacheFromCompiledData(const FNiagaraDataS
 		RendererLayout.SetVariableFromBinding(CompiledData, PrevRibbonTwistBinding, ENiagaraRibbonVFLayout::PrevRibbonTwist);
 	}
 
-	RendererLayout.SetVariableFromBinding(CompiledData, RibbonLinkOrderBinding, ENiagaraRibbonVFLayout::LinkOrder);
-	
 	RendererLayout.Finalize();
 
+	// Find Ribbon link order for GPU.  VF bindings don't support Int's at the moment, and there's no point any CPU emitter caring about this data.
+	bGpuRibbonLinkIsFloat = false;
+	GpuRibbonLinkOrderOffset = INDEX_NONE;
+	if (CompiledData)
+	{
+		if (const FNiagaraVariableLayoutInfo* BindingLinkOrderInfo = CompiledData->FindVariableLayoutInfo(RibbonLinkOrderBinding.GetDataSetBindableVariable()))
+		{
+			bGpuRibbonLinkIsFloat = true;
+			GpuRibbonLinkOrderOffset = BindingLinkOrderInfo->GetFloatComponentStart();
+		}
+		else if (bLinkOrderUseUniqueID)
+		{
+			if (const FNiagaraVariableLayoutInfo* UniqueIDLinkOrderInfo = CompiledData->FindVariableLayoutInfo(NiagaraRibbonRendererPropertiesPrivate::GetUniqueIDVariable()))
+			{
+				bGpuRibbonLinkIsFloat = false;
+				GpuRibbonLinkOrderOffset = UniqueIDLinkOrderInfo->GetInt32ComponentStart();
+			}
+		}
+		if (GpuRibbonLinkOrderOffset == INDEX_NONE)
+		{
+			if (const FNiagaraVariableLayoutInfo* NormAgeLinkOrderInfo = CompiledData->FindVariableLayoutInfo(NormalizedAgeBinding.GetDataSetBindableVariable()))
+			{
+				bGpuRibbonLinkIsFloat = true;
+				GpuRibbonLinkOrderOffset = NormAgeLinkOrderInfo->GetFloatComponentStart();
+			}
+		}
+	}
+
 #if WITH_EDITORONLY_DATA
 	// Build dynamic parameter mask
 	// Serialized in cooked builds
@@ -600,6 +649,21 @@ void UNiagaraRibbonRendererProperties::GetRendererFeedback(const FVersionedNiaga
 
 	const FVersionedNiagaraEmitterData* EmitterData = InEmitter.GetEmitterData();
 
+	// If this is a renderer made before we used UniqueID give the user an option to switch into the new mode
+	if (EmitterData && bLinkOrderUseUniqueID == false)
+	{
+		const FNiagaraVMExecutableData& ExecData = EmitterData->SpawnScriptProps.Script->GetVMExecutableData();
+		if (!ExecData.Attributes.Contains(RibbonLinkOrderBinding.GetDataSetBindableVariable()))
+		{
+			OutInfo.Emplace(
+				LOCTEXT("RibbonLinkOrderUsesNormalizedAgeSummary", "RibbonLinkOrder will use normalized age if the link order binding does not exist.  This can produce unpredictable results with burst modules, new renderers will use Particles.UniqueID to fix this issue."),
+				LOCTEXT("RibbonLinkOrderUsesNormalizedAgeDesc", "RibbonLinkOrder will fallback to normalized age"),
+				LOCTEXT("RibbonLinkOrderUsesNormalizedAgeFix", "Change fallback to use Particles.UniqueID"),
+				FNiagaraRendererFeedbackFix::CreateLambda([Renderer = const_cast<UNiagaraRibbonRendererProperties*>(this)]() { Renderer->bLinkOrderUseUniqueID = true; }),
+				true
+			);
+		}
+	}
 
 	// If we're in a gpu sim, then uv mode uniform by segment can cause some visual oddity due to non-existent
 	// culling of near particles like the cpu initialization pipeline runs
@@ -619,7 +683,6 @@ void UNiagaraRibbonRendererProperties::GetRendererFeedback(const FVersionedNiaga
 		CheckUVSettingsForChannel(UV1Settings, 1);
 	}
 
-
 	// If we're in multiplane shape, and multiplane count is even while we're in camera facing mode then one
 	// slice out of the set will be invisible because the camera will be coplanar to it
 	if (FacingMode == ENiagaraRibbonFacingMode::Screen && Shape == ENiagaraRibbonShapeMode::MultiPlane && MultiPlaneCount % 2 == 0)
@@ -631,7 +694,6 @@ void UNiagaraRibbonRendererProperties::GetRendererFeedback(const FVersionedNiaga
 		OutWarnings.Add(FNiagaraRendererFeedback(ErrorDescription, ErrorSummary, ErrorFix, MultiPlaneFix, true));	
 	}	
 
-
 	if (MaterialParameters.HasAnyBindings())
 	{
 		TArray<UMaterialInterface*> Materials;
diff --git a/Engine/Plugins/FX/Niagara/Source/Niagara/Public/NiagaraRendererRibbons.h b/Engine/Plugins/FX/Niagara/Source/Niagara/Public/NiagaraRendererRibbons.h
--- a/Engine/Plugins/FX/Niagara/Source/Niagara/Public/NiagaraRendererRibbons.h
+++ b/Engine/Plugins/FX/Niagara/Source/Niagara/Public/NiagaraRendererRibbons.h
@@ -25,36 +25,7 @@ struct FNiagaraRibbonRenderingFrameResources;
 struct FNiagaraRibbonRenderingFrameViewResources;
 struct FNiagaraRibbonGPUInitParameters;
 struct FNiagaraRibbonGPUInitComputeBuffers;
-
-struct FNiagaraGenerationInputDataCPUAccessors
-{
-	FNiagaraGenerationInputDataCPUAccessors(const UNiagaraRibbonRendererProperties* Properties, const FNiagaraDataSet& Data)
-		: TotalNumParticles(Data.GetCurrentDataChecked().GetNumInstances())
-		, SortKeyReader(Properties->SortKeyDataSetAccessor.GetReader(Data))
-		, RibbonLinkOrderData(Properties->RibbonLinkOrderDataSetAccessor.GetReader(Data))
-		, SimpleRibbonIDData(Properties->RibbonIdDataSetAccessor.GetReader(Data))
-		, FullRibbonIDData(Properties->RibbonFullIDDataSetAccessor.GetReader(Data))
-		, PosData(Properties->PositionDataSetAccessor.GetReader(Data))
-		, AgeData(Properties->NormalizedAgeAccessor.GetReader(Data))
-		, SizeData(Properties->SizeDataSetAccessor.GetReader(Data))
-		, TwistData(Properties->TwistDataSetAccessor.GetReader(Data))
-	{
-		
-	}
-	
-	const uint32 TotalNumParticles;
-	
-	const FNiagaraDataSetReaderFloat<float> SortKeyReader;
-	const FNiagaraDataSetReaderFloat<float> RibbonLinkOrderData;
-	
-	const FNiagaraDataSetReaderInt32<int> SimpleRibbonIDData;
-	const FNiagaraDataSetReaderStruct<FNiagaraID> FullRibbonIDData;
-	
-	const FNiagaraDataSetReaderFloat<FNiagaraPosition> PosData;
-	const FNiagaraDataSetReaderFloat<float> AgeData;
-	const FNiagaraDataSetReaderFloat<float> SizeData;
-	const FNiagaraDataSetReaderFloat<float> TwistData;
-};
+struct FNiagaraGenerationInputDataCPUAccessors;
 
 struct FNiagaraIndexGenerationInput
 {
@@ -87,7 +58,6 @@ public:
 		, MaxNumRibbons(Properties->MaxNumRibbons)
 		, bHasFullRibbonIDs(Properties->RibbonFullIDDataSetAccessor.IsValid())
 		, bHasSimpleRibbonIDs(Properties->RibbonIdDataSetAccessor.IsValid())
-		, bHasCustomLinkOrder(Properties->RibbonLinkOrderDataSetAccessor.IsValid())
 		, bHasTwist(Properties->TwistDataSetAccessor.IsValid() && Properties->SizeDataSetAccessor.IsValid())
 		, bHasCustomU0Data(Properties->UV0Settings.bEnablePerParticleUOverride && Properties->U0OverrideIsBound)
 		, bHasCustomU1Data(Properties->UV1Settings.bEnablePerParticleUOverride && Properties->U1OverrideIsBound)
@@ -106,7 +76,6 @@ public:
 	bool HasSimpleRibbonIDs() const { return bHasSimpleRibbonIDs; }
 	bool HasRibbonIDs() const { return HasFullRibbonIDs() || HasSimpleRibbonIDs(); };
 	
-	bool HasCustomLinkOrder() const { return bHasCustomLinkOrder; }
 	bool HasTwist() const { return bHasTwist; }
 	
 	bool HasCustomU0Data() const { return bHasCustomU0Data; }
@@ -124,7 +93,6 @@ private:
 	
 	const uint32 bHasFullRibbonIDs : 1;
 	const uint32 bHasSimpleRibbonIDs : 1;
-	const uint32 bHasCustomLinkOrder : 1;
 	const uint32 bHasTwist : 1;
 	
 	const uint32 bHasCustomU0Data : 1;
@@ -423,6 +391,8 @@ protected:
 	ENiagaraRibbonDrawDirection DrawDirection;
 	
 	const FNiagaraRendererLayout* RendererLayout;
+	bool bGpuRibbonLinkIsFloat = false;
+	uint32 GpuRibbonLinkOrderOffset = INDEX_NONE;
 
 	mutable FNiagaraRibbonVertexBuffers VertexBuffers;
 	
diff --git a/Engine/Plugins/FX/Niagara/Source/Niagara/Public/NiagaraRibbonRendererProperties.h b/Engine/Plugins/FX/Niagara/Source/Niagara/Public/NiagaraRibbonRendererProperties.h
--- a/Engine/Plugins/FX/Niagara/Source/Niagara/Public/NiagaraRibbonRendererProperties.h
+++ b/Engine/Plugins/FX/Niagara/Source/Niagara/Public/NiagaraRibbonRendererProperties.h
@@ -185,7 +185,6 @@ namespace ENiagaraRibbonVFLayout
 		PrevRibbonWidth,
 		PrevRibbonFacing,
 		PrevRibbonTwist,
-		LinkOrder,
 		Num,
 	};
 };
@@ -322,6 +321,9 @@ public:
 	UPROPERTY(EditAnywhere, Category = "Ribbon Tessellation", meta = (DisplayName = "Screen Space", EditCondition = "TessellationMode == ENiagaraRibbonTessellationMode::Custom && !bUseConstantFactor", EditConditionHides, DisplayAfter="TessellationAngle"))
 	uint8 bScreenSpaceTessellation : 1;
 
+	UPROPERTY()
+	uint8 bLinkOrderUseUniqueID : 1;
+
 	/** Tessellation factor to apply to the width of the ribbon.
 	* Ranges from 1 to 16. Greater values increase amount of tessellation.
 	*/
@@ -467,8 +469,8 @@ public:
 	UPROPERTY()
 	uint32 MaterialParamValidMask = 0;
 
-	bool								bSortKeyDataSetAccessorIsAge = false;
-	FNiagaraDataSetAccessor<float>		SortKeyDataSetAccessor;
+	FNiagaraDataSetAccessor<float>		RibbonLinkOrderFloatAccessor;
+	FNiagaraDataSetAccessor<int32>		RibbonLinkOrderInt32Accessor;
 	FNiagaraDataSetAccessor<FNiagaraPosition>	PositionDataSetAccessor;
 	FNiagaraDataSetAccessor<float>		NormalizedAgeAccessor;
 	FNiagaraDataSetAccessor<float>		SizeDataSetAccessor;
@@ -485,7 +487,8 @@ public:
 	FNiagaraDataSetAccessor<int32>		RibbonIdDataSetAccessor;
 	FNiagaraDataSetAccessor<FNiagaraID>	RibbonFullIDDataSetAccessor;
 	
-	FNiagaraDataSetAccessor<float>		RibbonLinkOrderDataSetAccessor;
+	bool								bGpuRibbonLinkIsFloat = false;
+	uint32								GpuRibbonLinkOrderOffset = INDEX_NONE;
 
 	FNiagaraRendererLayout RendererLayout;
 
@@ -498,6 +501,4 @@ protected:
 	void UpdateMICs();
 
 	virtual bool NeedsMIDsForMaterials() const { return MaterialParameters.HasAnyBindings(); }
-private: 
-	static TArray<TWeakObjectPtr<UNiagaraRibbonRendererProperties>> RibbonRendererPropertiesToDeferredInit;
 };
diff --git a/Engine/Plugins/FX/Niagara/Source/NiagaraCore/Public/NiagaraCustomVersion.h b/Engine/Plugins/FX/Niagara/Source/NiagaraCore/Public/NiagaraCustomVersion.h
--- a/Engine/Plugins/FX/Niagara/Source/NiagaraCore/Public/NiagaraCustomVersion.h
+++ b/Engine/Plugins/FX/Niagara/Source/NiagaraCore/Public/NiagaraCustomVersion.h
@@ -187,6 +187,8 @@ struct FNiagaraCustomVersion
 
 		DynamicPinNodeFixup, // Some data fixup for NiagaraNodeWithDynamicPins.
 
+		RibbonRendererLinkOrderDefaultIsUniqueID,	// Ribbon renderer will default to unique ID rather than normalized age to make more things 'just work'
+
 		// DO NOT ADD A NEW VERSION UNLESS YOU HAVE TALKED TO THE NIAGARA LEAD. Mismanagement of these versions can lead to data loss if it is adjusted in multiple streams simultaneously.
 		// -----<new versions can be added above this line>  -------------------------------------------------
 		VersionPlusOne,
diff --git a/Engine/Plugins/FX/Niagara/Source/NiagaraShader/Public/NiagaraRibbonCompute.h b/Engine/Plugins/FX/Niagara/Source/NiagaraShader/Public/NiagaraRibbonCompute.h
--- a/Engine/Plugins/FX/Niagara/Source/NiagaraShader/Public/NiagaraRibbonCompute.h
+++ b/Engine/Plugins/FX/Niagara/Source/NiagaraShader/Public/NiagaraRibbonCompute.h
@@ -59,7 +59,7 @@ END_SHADER_PARAMETER_STRUCT()
 
 class FRibbonHasFullRibbonID : SHADER_PERMUTATION_BOOL("RIBBONID_IS_NIAGARAID");
 class FRibbonHasRibbonID : SHADER_PERMUTATION_BOOL("RIBBONID_IS_INT");
-class FRibbonHasCustomLinkOrder : SHADER_PERMUTATION_BOOL("RIBBON_HAS_LINK_ORDER");
+class FRibbonLinkIsFloat : SHADER_PERMUTATION_BOOL("RIBBONLINK_IS_FLOAT");
 
 class FRibbonWantsConstantTessellation : SHADER_PERMUTATION_BOOL("RIBBONS_WANTS_CONSTANT_TESSELLATION");
 class FRibbonWantsAutomaticTessellation : SHADER_PERMUTATION_BOOL("RIBBONS_WANTS_AUTOMATIC_TESSELLATION");
@@ -101,7 +101,7 @@ class FNiagaraRibbonSortPhase1CS : public FGlobalShader
 	static constexpr uint32 BubbleSortGroupWidth = 32;
 
 	using FParameters = FRibbonOrderSortParameters;
-	using FPermutationDomain = TShaderPermutationDomain<FRibbonHasFullRibbonID, FRibbonHasRibbonID, FRibbonHasCustomLinkOrder>;
+	using FPermutationDomain = TShaderPermutationDomain<FRibbonHasFullRibbonID, FRibbonHasRibbonID, FRibbonLinkIsFloat>;
 
 	static void ModifyCompilationEnvironment(const FGlobalShaderPermutationParameters& Parameters, FShaderCompilerEnvironment& OutEnvironment);
 };
@@ -117,7 +117,7 @@ class FNiagaraRibbonSortPhase2CS : public FGlobalShader
 	static constexpr uint32 ThreadGroupSize = 64;
 
 	using FParameters = FRibbonOrderSortParameters;
-	using FPermutationDomain = TShaderPermutationDomain<FRibbonHasFullRibbonID, FRibbonHasRibbonID, FRibbonHasCustomLinkOrder>;
+	using FPermutationDomain = TShaderPermutationDomain<FRibbonHasFullRibbonID, FRibbonHasRibbonID, FRibbonLinkIsFloat>;
 
 	static void ModifyCompilationEnvironment(const FGlobalShaderPermutationParameters& Parameters, FShaderCompilerEnvironment& OutEnvironment);
 };
diff --git a/Engine/Plugins/FX/Niagara/Source/NiagaraVertexFactories/Public/NiagaraRibbonVertexFactory.h b/Engine/Plugins/FX/Niagara/Source/NiagaraVertexFactories/Public/NiagaraRibbonVertexFactory.h
--- a/Engine/Plugins/FX/Niagara/Source/NiagaraVertexFactories/Public/NiagaraRibbonVertexFactory.h
+++ b/Engine/Plugins/FX/Niagara/Source/NiagaraVertexFactories/Public/NiagaraRibbonVertexFactory.h
@@ -39,7 +39,6 @@ BEGIN_GLOBAL_SHADER_PARAMETER_STRUCT(FNiagaraRibbonUniformParameters, NIAGARAVER
 	SHADER_PARAMETER(int, ColorDataOffset)
 	SHADER_PARAMETER(int, FacingDataOffset)
 	SHADER_PARAMETER(int, PrevFacingDataOffset)
-	SHADER_PARAMETER(int, LinkOrderDataOffset)
 	SHADER_PARAMETER(int, NormalizedAgeDataOffset)
 	SHADER_PARAMETER(int, MaterialRandomDataOffset)
 	SHADER_PARAMETER(uint32, MaterialParamValidMask)
