- Fix for CIS shadowing

#rnx

diff --git a/Engine/Plugins/FX/Niagara/Source/NiagaraEditor/Private/ViewModels/Stack/NiagaraStackRendererItem.cpp b/Engine/Plugins/FX/Niagara/Source/NiagaraEditor/Private/ViewModels/Stack/NiagaraStackRendererItem.cpp
--- a/Engine/Plugins/FX/Niagara/Source/NiagaraEditor/Private/ViewModels/Stack/NiagaraStackRendererItem.cpp
+++ b/Engine/Plugins/FX/Niagara/Source/NiagaraEditor/Private/ViewModels/Stack/NiagaraStackRendererItem.cpp
@@ -531,20 +531,20 @@ void UNiagaraStackRendererItem::ProcessRendererIssues(const TArray<FNiagaraRende
 							const FScopedTransaction ScopedTransaction(LOCTEXT("RendererItemFixTransaction", "Apply renderer fix"));
 
 							UNiagaraStackRendererItem* StackItem = WeakStackItem.Get();
-							UNiagaraRendererProperties* RendererProperties = StackItem ? StackItem->GetRendererProperties() : nullptr;
-							if (RendererProperties)
+							UNiagaraRendererProperties* RendererPropertiesLocal = StackItem ? StackItem->GetRendererProperties() : nullptr;
+							if (RendererPropertiesLocal)
 							{
-								RendererProperties->Modify();
+								RendererPropertiesLocal->Modify();
 							}
 
 							Item.TryFix();
 
-							if (RendererProperties)
+							if (RendererPropertiesLocal)
 							{
-								RendererProperties->PostEditChange();
+								RendererPropertiesLocal->PostEditChange();
 
 								TArray<UObject*> ChangedObjects;
-								ChangedObjects.Add(RendererProperties);
+								ChangedObjects.Add(RendererPropertiesLocal);
 								StackItem->OnDataObjectModified().Broadcast(ChangedObjects, ENiagaraDataObjectChange::Changed);
 							}
 						}
