Removed forward declaration for non-existant class UNiagaraSystemComponent

Looks like this class was renamed at some point.

diff --git a/Engine/Plugins/FX/Niagara/Source/NiagaraAnimNotifies/Public/AnimNotify_PlayNiagaraEffect.h b/Engine/Plugins/FX/Niagara/Source/NiagaraAnimNotifies/Public/AnimNotify_PlayNiagaraEffect.h
--- a/Engine/Plugins/FX/Niagara/Source/NiagaraAnimNotifies/Public/AnimNotify_PlayNiagaraEffect.h
+++ b/Engine/Plugins/FX/Niagara/Source/NiagaraAnimNotifies/Public/AnimNotify_PlayNiagaraEffect.h
@@ -12,7 +12,6 @@
 class UAnimSequenceBase;
 class UNiagaraSystem;
 class USkeletalMeshComponent;
-class UNiagaraSystemComponent;
 class UNiagaraSystem;
 class UFXSystemComponent;
 
