- Use initializer list for AttributeBindings as not all have correct reserve sizes

diff --git a/Engine/Plugins/FX/Niagara/Source/Niagara/Private/NiagaraComponentRendererProperties.cpp b/Engine/Plugins/FX/Niagara/Source/Niagara/Private/NiagaraComponentRendererProperties.cpp
--- a/Engine/Plugins/FX/Niagara/Source/Niagara/Private/NiagaraComponentRendererProperties.cpp
+++ b/Engine/Plugins/FX/Niagara/Source/Niagara/Private/NiagaraComponentRendererProperties.cpp
@@ -199,9 +199,11 @@ UNiagaraComponentRendererProperties::UNiagaraComponentRendererProperties()
 	FCoreUObjectDelegates::OnObjectsReplaced.AddUObject(this, &UNiagaraComponentRendererProperties::OnObjectsReplacedCallback);
 #endif
 
-	AttributeBindings.Reserve(2);
-	AttributeBindings.Add(&EnabledBinding);
-	AttributeBindings.Add(&RendererVisibilityTagBinding);
+	AttributeBindings =
+	{
+		&EnabledBinding,
+		&RendererVisibilityTagBinding,
+	};
 	IsSetterMappingDirty.store(true);
 }
 
diff --git a/Engine/Plugins/FX/Niagara/Source/Niagara/Private/NiagaraDecalRendererProperties.cpp b/Engine/Plugins/FX/Niagara/Source/Niagara/Private/NiagaraDecalRendererProperties.cpp
--- a/Engine/Plugins/FX/Niagara/Source/Niagara/Private/NiagaraDecalRendererProperties.cpp
+++ b/Engine/Plugins/FX/Niagara/Source/Niagara/Private/NiagaraDecalRendererProperties.cpp
@@ -81,15 +81,17 @@ namespace NiagaraDecalRendererPropertiesLocal
 
 UNiagaraDecalRendererProperties::UNiagaraDecalRendererProperties()
 {
-	AttributeBindings.Reserve(6);
-	AttributeBindings.Add(&PositionBinding);
-	AttributeBindings.Add(&DecalOrientationBinding);
-	AttributeBindings.Add(&DecalSizeBinding);
-	AttributeBindings.Add(&DecalFadeBinding);
-	AttributeBindings.Add(&DecalSortOrderBinding);
-	AttributeBindings.Add(&DecalColorBinding);
-	AttributeBindings.Add(&DecalVisibleBinding);
-	AttributeBindings.Add(&RendererVisibilityTagBinding);
+	AttributeBindings =
+	{
+		&PositionBinding,
+		&DecalOrientationBinding,
+		&DecalSizeBinding,
+		&DecalFadeBinding,
+		&DecalSortOrderBinding,
+		&DecalColorBinding,
+		&DecalVisibleBinding,
+		&RendererVisibilityTagBinding,
+	};
 }
 
 void UNiagaraDecalRendererProperties::PostLoad()
diff --git a/Engine/Plugins/FX/Niagara/Source/Niagara/Private/NiagaraLightRendererProperties.cpp b/Engine/Plugins/FX/Niagara/Source/Niagara/Private/NiagaraLightRendererProperties.cpp
--- a/Engine/Plugins/FX/Niagara/Source/Niagara/Private/NiagaraLightRendererProperties.cpp
+++ b/Engine/Plugins/FX/Niagara/Source/Niagara/Private/NiagaraLightRendererProperties.cpp
@@ -31,14 +31,16 @@ UNiagaraLightRendererProperties::UNiagaraLightRendererProperties()
 	, DefaultExponent(1.0f)
 	, ColorAdd(FVector(0.0f, 0.0f, 0.0f))
 {
-	AttributeBindings.Reserve(7);
-	AttributeBindings.Add(&LightRenderingEnabledBinding);
-	AttributeBindings.Add(&LightExponentBinding);
-	AttributeBindings.Add(&PositionBinding);
-	AttributeBindings.Add(&ColorBinding);
-	AttributeBindings.Add(&RadiusBinding);
-	AttributeBindings.Add(&VolumetricScatteringBinding);
-	AttributeBindings.Add(&RendererVisibilityTagBinding);
+	AttributeBindings =
+	{
+		&LightRenderingEnabledBinding,
+		&LightExponentBinding,
+		&PositionBinding,
+		&ColorBinding,
+		&RadiusBinding,
+		&VolumetricScatteringBinding,
+		&RendererVisibilityTagBinding,
+	};
 }
 
 void UNiagaraLightRendererProperties::PostLoad()
diff --git a/Engine/Plugins/FX/Niagara/Source/Niagara/Private/NiagaraMeshRendererProperties.cpp b/Engine/Plugins/FX/Niagara/Source/Niagara/Private/NiagaraMeshRendererProperties.cpp
--- a/Engine/Plugins/FX/Niagara/Source/Niagara/Private/NiagaraMeshRendererProperties.cpp
+++ b/Engine/Plugins/FX/Niagara/Source/Niagara/Private/NiagaraMeshRendererProperties.cpp
@@ -257,32 +257,34 @@ UNiagaraMeshRendererProperties::UNiagaraMeshRendererProperties()
 	NumFlipbookFrames = 1;
 #endif
 
-	AttributeBindings.Reserve(21);
-	AttributeBindings.Add(&PositionBinding);
-	AttributeBindings.Add(&VelocityBinding);
-	AttributeBindings.Add(&ColorBinding);
-	AttributeBindings.Add(&ScaleBinding);
-	AttributeBindings.Add(&MeshOrientationBinding);
-	AttributeBindings.Add(&MaterialRandomBinding);
-	AttributeBindings.Add(&NormalizedAgeBinding);
-	AttributeBindings.Add(&CustomSortingBinding);
-	AttributeBindings.Add(&SubImageIndexBinding);
-	AttributeBindings.Add(&DynamicMaterialBinding);
-	AttributeBindings.Add(&DynamicMaterial1Binding);
-	AttributeBindings.Add(&DynamicMaterial2Binding);
-	AttributeBindings.Add(&DynamicMaterial3Binding);
-	AttributeBindings.Add(&CameraOffsetBinding);
-
-	// These are associated with attributes in the VF layout only if bGenerateAccurateMotionVectors is true
-	AttributeBindings.Add(&PrevPositionBinding);
-	AttributeBindings.Add(&PrevScaleBinding);
-	AttributeBindings.Add(&PrevMeshOrientationBinding);
-	AttributeBindings.Add(&PrevCameraOffsetBinding);
-	AttributeBindings.Add(&PrevVelocityBinding);
-
-	// The remaining bindings are not associated with attributes in the VF layout
-	AttributeBindings.Add(&RendererVisibilityTagBinding);
-	AttributeBindings.Add(&MeshIndexBinding);
+	AttributeBindings =
+	{
+		&PositionBinding,
+		&VelocityBinding,
+		&ColorBinding,
+		&ScaleBinding,
+		&MeshOrientationBinding,
+		&MaterialRandomBinding,
+		&NormalizedAgeBinding,
+		&CustomSortingBinding,
+		&SubImageIndexBinding,
+		&DynamicMaterialBinding,
+		&DynamicMaterial1Binding,
+		&DynamicMaterial2Binding,
+		&DynamicMaterial3Binding,
+		&CameraOffsetBinding,
+
+		// These are associated with attributes in the VF layout only if bGenerateAccurateMotionVectors is true
+		&PrevPositionBinding,
+		&PrevScaleBinding,
+		&PrevMeshOrientationBinding,
+		&PrevCameraOffsetBinding,
+		&PrevVelocityBinding,
+
+		// The remaining bindings are not associated with attributes in the VF layout
+		&RendererVisibilityTagBinding,
+		&MeshIndexBinding,
+	};
 }
 
 FNiagaraRenderer* UNiagaraMeshRendererProperties::CreateEmitterRenderer(ERHIFeatureLevel::Type FeatureLevel, const FNiagaraEmitterInstance* Emitter, const FNiagaraSystemInstanceController& InController)
diff --git a/Engine/Plugins/FX/Niagara/Source/Niagara/Private/NiagaraSpriteRendererProperties.cpp b/Engine/Plugins/FX/Niagara/Source/Niagara/Private/NiagaraSpriteRendererProperties.cpp
--- a/Engine/Plugins/FX/Niagara/Source/Niagara/Private/NiagaraSpriteRendererProperties.cpp
+++ b/Engine/Plugins/FX/Niagara/Source/Niagara/Private/NiagaraSpriteRendererProperties.cpp
@@ -62,40 +62,41 @@ UNiagaraSpriteRendererProperties::UNiagaraSpriteRendererProperties()
 	, AlphaThreshold(0.1f)
 #endif // WITH_EDITORONLY_DATA
 {
-	AttributeBindings.Reserve(27);
-
-	// NOTE: These bindings' indices have to align to their counterpart in ENiagaraSpriteVFLayout
-	AttributeBindings.Add(&PositionBinding);
-	AttributeBindings.Add(&ColorBinding);
-	AttributeBindings.Add(&VelocityBinding);
-	AttributeBindings.Add(&SpriteRotationBinding);
-	AttributeBindings.Add(&SpriteSizeBinding);
-	AttributeBindings.Add(&SpriteFacingBinding);
-	AttributeBindings.Add(&SpriteAlignmentBinding);
-	AttributeBindings.Add(&SubImageIndexBinding);
-	AttributeBindings.Add(&DynamicMaterialBinding);
-	AttributeBindings.Add(&DynamicMaterial1Binding);
-	AttributeBindings.Add(&DynamicMaterial2Binding);
-	AttributeBindings.Add(&DynamicMaterial3Binding);
-	AttributeBindings.Add(&CameraOffsetBinding);
-	AttributeBindings.Add(&UVScaleBinding);
-	AttributeBindings.Add(&PivotOffsetBinding);
-	AttributeBindings.Add(&MaterialRandomBinding);
-	AttributeBindings.Add(&CustomSortingBinding);
-	AttributeBindings.Add(&NormalizedAgeBinding);
-
-	// These bindings are only actually used with accurate motion vectors (indices still need to align)
-	AttributeBindings.Add(&PrevPositionBinding);
-	AttributeBindings.Add(&PrevVelocityBinding);
-	AttributeBindings.Add(&PrevSpriteRotationBinding);
-	AttributeBindings.Add(&PrevSpriteSizeBinding);
-	AttributeBindings.Add(&PrevSpriteFacingBinding);
-	AttributeBindings.Add(&PrevSpriteAlignmentBinding);
-	AttributeBindings.Add(&PrevCameraOffsetBinding);
-	AttributeBindings.Add(&PrevPivotOffsetBinding);
-
-	// The remaining bindings are not associated with attributes in the VF layout
-	AttributeBindings.Add(&RendererVisibilityTagBinding);
+	AttributeBindings =
+	{
+		// NOTE: These bindings' indices have to align to their counterpart in ENiagaraSpriteVFLayout
+		&PositionBinding,
+		&ColorBinding,
+		&VelocityBinding,
+		&SpriteRotationBinding,
+		&SpriteSizeBinding,
+		&SpriteFacingBinding,
+		&SpriteAlignmentBinding,
+		&SubImageIndexBinding,
+		&DynamicMaterialBinding,
+		&DynamicMaterial1Binding,
+		&DynamicMaterial2Binding,
+		&DynamicMaterial3Binding,
+		&CameraOffsetBinding,
+		&UVScaleBinding,
+		&PivotOffsetBinding,
+		&MaterialRandomBinding,
+		&CustomSortingBinding,
+		&NormalizedAgeBinding,
+
+		// These bindings are only actually used with accurate motion vectors (indices still need to align)
+		&PrevPositionBinding,
+		&PrevVelocityBinding,
+		&PrevSpriteRotationBinding,
+		&PrevSpriteSizeBinding,
+		&PrevSpriteFacingBinding,
+		&PrevSpriteAlignmentBinding,
+		&PrevCameraOffsetBinding,
+		&PrevPivotOffsetBinding,
+
+		// The remaining bindings are not associated with attributes in the VF layout
+		&RendererVisibilityTagBinding,
+	};
 }
 
 FNiagaraRenderer* UNiagaraSpriteRendererProperties::CreateEmitterRenderer(ERHIFeatureLevel::Type FeatureLevel, const FNiagaraEmitterInstance* Emitter, const FNiagaraSystemInstanceController& InController)
diff --git a/Engine/Plugins/FX/Niagara/Source/Niagara/Private/NiagaraVolumeRendererProperties.cpp b/Engine/Plugins/FX/Niagara/Source/Niagara/Private/NiagaraVolumeRendererProperties.cpp
--- a/Engine/Plugins/FX/Niagara/Source/Niagara/Private/NiagaraVolumeRendererProperties.cpp
+++ b/Engine/Plugins/FX/Niagara/Source/Niagara/Private/NiagaraVolumeRendererProperties.cpp
@@ -79,13 +79,15 @@ namespace NiagaraVolumeRendererPropertiesLocal
 
 UNiagaraVolumeRendererProperties::UNiagaraVolumeRendererProperties()
 {
-	AttributeBindings.Reserve(ENiagaraVolumeVFLayout::Num);
-	AttributeBindings.Add(&PositionBinding);
-	AttributeBindings.Add(&RotationBinding);
-	AttributeBindings.Add(&ScaleBinding);
-	AttributeBindings.Add(&RendererVisibilityTagBinding);
-	AttributeBindings.Add(&VolumeResolutionMaxAxisBinding);
-	AttributeBindings.Add(&VolumeWorldSpaceSizeBinding);
+	AttributeBindings =
+	{
+		&PositionBinding,
+		&RotationBinding,
+		&ScaleBinding,
+		&RendererVisibilityTagBinding,
+		&VolumeResolutionMaxAxisBinding,
+		&VolumeWorldSpaceSizeBinding,
+	};
 }
 
 void UNiagaraVolumeRendererProperties::PostLoad()
diff --git a/Engine/Plugins/Runtime/GeometryCache/Source/GeometryCache/Private/NiagaraGeometryCacheRendererProperties.cpp b/Engine/Plugins/Runtime/GeometryCache/Source/GeometryCache/Private/NiagaraGeometryCacheRendererProperties.cpp
--- a/Engine/Plugins/Runtime/GeometryCache/Source/GeometryCache/Private/NiagaraGeometryCacheRendererProperties.cpp
+++ b/Engine/Plugins/Runtime/GeometryCache/Source/GeometryCache/Private/NiagaraGeometryCacheRendererProperties.cpp
@@ -56,14 +56,16 @@ UMaterialInterface* FNiagaraGeometryCacheReference::ResolveMaterial(UGeometryCac
 
 UNiagaraGeometryCacheRendererProperties::UNiagaraGeometryCacheRendererProperties()
 {
-	AttributeBindings.Reserve(6);
-	AttributeBindings.Add(&PositionBinding);
-	AttributeBindings.Add(&RotationBinding);
-	AttributeBindings.Add(&ScaleBinding);
-	AttributeBindings.Add(&ElapsedTimeBinding);
-	AttributeBindings.Add(&EnabledBinding);
-	AttributeBindings.Add(&RendererVisibilityTagBinding);
-	AttributeBindings.Add(&ArrayIndexBinding);
+	AttributeBindings =
+	{
+		&PositionBinding,
+		&RotationBinding,
+		&ScaleBinding,
+		&ElapsedTimeBinding,
+		&EnabledBinding,
+		&RendererVisibilityTagBinding,
+		&ArrayIndexBinding,
+	};
 
 	if (GeometryCaches.Num() == 0)
 	{
