- Data Channel Sim Stage Iteration Source Support

#rnx

diff --git a/Engine/Plugins/FX/Niagara/Source/Niagara/Internal/DataInterface/NiagaraDataInterfaceDataChannelRead.h b/Engine/Plugins/FX/Niagara/Source/Niagara/Internal/DataInterface/NiagaraDataInterfaceDataChannelRead.h
--- a/Engine/Plugins/FX/Niagara/Source/Niagara/Internal/DataInterface/NiagaraDataInterfaceDataChannelRead.h
+++ b/Engine/Plugins/FX/Niagara/Source/Niagara/Internal/DataInterface/NiagaraDataInterfaceDataChannelRead.h
@@ -23,7 +23,7 @@ There is working code for this already but the concept/API needs more fleshing o
 
 
 #include "NiagaraDataInterfaceDataChannelCommon.h"
-#include "NiagaraDataInterface.h"
+#include "NiagaraDataInterfaceRW.h"
 #include "NiagaraDataInterfaceEmitterBinding.h"
 #include "NiagaraDataSetAccessor.h"
 #include "NiagaraDataInterfaceDataChannelRead.generated.h"
@@ -52,7 +52,7 @@ enum class ENDIDataChannelSpawnMode
 };
 
 UCLASS(Experimental, EditInlineNew, Category = "Data Channels", CollapseCategories, meta = (DisplayName = "Data Channel Reader"))
-class NIAGARA_API UNiagaraDataInterfaceDataChannelRead : public UNiagaraDataInterface
+class NIAGARA_API UNiagaraDataInterfaceDataChannelRead : public UNiagaraDataInterfaceRWBase
 {
 	GENERATED_UCLASS_BODY()
 protected:
@@ -221,10 +221,11 @@ struct FNDIDataChannelReadInstanceData
 	bool PostTick(UNiagaraDataInterfaceDataChannelRead* Interface, FNiagaraSystemInstance* Instance);
 };
 
-struct FNiagaraDataInterfaceProxy_DataChannelRead : public FNiagaraDataInterfaceProxy
+struct FNiagaraDataInterfaceProxy_DataChannelRead : public FNiagaraDataInterfaceProxyRW
 {
 	virtual void ConsumePerInstanceDataFromGameThread(void* PerInstanceData, const FNiagaraSystemInstanceID& Instance) override;
 	virtual int32 PerInstanceDataPassedToRenderThreadSize() const override;
+	virtual void GetDispatchArgs(const FNDIGpuComputeDispatchArgsGenContext& Context) override;
 
 	/** Persistent per instance data on the RT. Constructed when consuming data passed from GT->RT. */
 	struct FInstanceData
diff --git a/Engine/Plugins/FX/Niagara/Source/Niagara/Private/DataInterface/NiagaraDataInterfaceDataChannelRead.cpp b/Engine/Plugins/FX/Niagara/Source/Niagara/Private/DataInterface/NiagaraDataInterfaceDataChannelRead.cpp
--- a/Engine/Plugins/FX/Niagara/Source/Niagara/Private/DataInterface/NiagaraDataInterfaceDataChannelRead.cpp
+++ b/Engine/Plugins/FX/Niagara/Source/Niagara/Private/DataInterface/NiagaraDataInterfaceDataChannelRead.cpp
@@ -1444,6 +1444,7 @@ void UNiagaraDataInterfaceDataChannelRead::GetParameterDefinitionHLSL(FNiagaraDa
 
 			//Finally generate the final code for this function and add it to the final hlsl.
 			OutHLSL += FString::Format(*FunctionTemplateFile, HlslTemplateArgs);
+			OutHLSL += TEXT("\n");
 		}
 	}
 }
@@ -1551,4 +1552,17 @@ int32 FNiagaraDataInterfaceProxy_DataChannelRead::PerInstanceDataPassedToRenderT
 	return sizeof(FNDIDataChannelReadInstanceData_RT);
 }
 
-#undef LOCTEXT_NAMESPACE
\ No newline at end of file
+void FNiagaraDataInterfaceProxy_DataChannelRead::GetDispatchArgs(const FNDIGpuComputeDispatchArgsGenContext& Context)
+{
+	const FNiagaraDataInterfaceProxy_DataChannelRead::FInstanceData* InstanceData = SystemInstancesToProxyData_RT.Find(Context.GetSystemInstanceID());
+	if (InstanceData && InstanceData->ChannelDataRTProxy)
+	{
+		FNiagaraDataBuffer* Data = InstanceData->bReadPrevFrame ? InstanceData->ChannelDataRTProxy->PrevFrameData.GetReference() : InstanceData->ChannelDataRTProxy->GPUDataSet->GetCurrentData();
+		if (Data)
+		{
+			Context.SetDirect(Data->GetNumInstances());
+		}
+	}
+}
+
+#undef LOCTEXT_NAMESPACE
