Niagara - Update the data interface multi assignment message generation in the resolve data interfaces utility so that it uses the cached data interface name, rather than the name of the data interface object.  The cached name is a better identifier and is also deterministic across compiles.

diff --git a/Engine/Plugins/FX/Niagara/Source/Niagara/Private/NiagaraResolveDIHelpers.cpp b/Engine/Plugins/FX/Niagara/Source/Niagara/Private/NiagaraResolveDIHelpers.cpp
--- a/Engine/Plugins/FX/Niagara/Source/Niagara/Private/NiagaraResolveDIHelpers.cpp
+++ b/Engine/Plugins/FX/Niagara/Source/Niagara/Private/NiagaraResolveDIHelpers.cpp
@@ -19,6 +19,7 @@ namespace FNiagaraResolveDIHelpers
 	struct FDataInterfaceSourceEmitterNamePair
 	{
 		UNiagaraDataInterface* DataInterface;
+		FName SourceName;
 		FString SourceEmitterName;
 	};
 
@@ -86,14 +87,15 @@ namespace FNiagaraResolveDIHelpers
 						OutErrorMessages.Add(FText::Format(
 							LOCTEXT("MultipleAssignmentsFormat", "A data interface parameter was the target of an assignment multiple times in a single system.  The data interface used in the simulation may be incorrect.  Target Parameter: {0} First Assignment: {1} Current Assignment: {2}"),
 							FText::FromName(WriteVariable.GetName()),
-							FText::FromString(CurrentAssignment->DataInterface->GetName()),
-							FText::FromString(CachedDefaultDataInterface.DataInterface->GetName())));
+							FText::FromName(CurrentAssignment->SourceName),
+							FText::FromName(CachedDefaultDataInterface.Name)));
 					}
 				}
 				else
 				{
 					FDataInterfaceSourceEmitterNamePair& NewAssignment = OutVariableAssignmentMap.Add(WriteVariable);
 					NewAssignment.DataInterface = CachedDefaultDataInterface.DataInterface;
+					NewAssignment.SourceName = CachedDefaultDataInterface.Name;
 					NewAssignment.SourceEmitterName = CachedDefaultDataInterface.SourceEmitterName;
 				}
 			}
