Fix for mesh component tags on Skeletal Mesh DI not working. Added assignment to CopyToInternal and comparison to Equals.

diff --git a/Engine/Plugins/FX/Niagara/Source/Niagara/Private/NiagaraDataInterfaceSkeletalMesh.cpp b/Engine/Plugins/FX/Niagara/Source/Niagara/Private/NiagaraDataInterfaceSkeletalMesh.cpp
--- a/Engine/Plugins/FX/Niagara/Source/Niagara/Private/NiagaraDataInterfaceSkeletalMesh.cpp
+++ b/Engine/Plugins/FX/Niagara/Source/Niagara/Private/NiagaraDataInterfaceSkeletalMesh.cpp
@@ -2373,6 +2373,7 @@ bool UNiagaraDataInterfaceSkeletalMesh::CopyToInternal(UNiagaraDataInterface* De
 	OtherTyped->SourceMode = SourceMode;
 	OtherTyped->SoftSourceActor = SoftSourceActor;
 	OtherTyped->MeshUserParameter = MeshUserParameter;
+	OtherTyped->ComponentTags = ComponentTags;
 	OtherTyped->SourceComponent = SourceComponent;
 	OtherTyped->SkinningMode = SkinningMode;
 	OtherTyped->SamplingRegions = SamplingRegions;
@@ -2404,6 +2405,7 @@ bool UNiagaraDataInterfaceSkeletalMesh::Equals(const UNiagaraDataInterface* Othe
 #endif
 		OtherTyped->SoftSourceActor == SoftSourceActor &&
 		OtherTyped->MeshUserParameter == MeshUserParameter &&
+		OtherTyped->ComponentTags == ComponentTags &&
 		OtherTyped->SourceComponent == SourceComponent &&
 		OtherTyped->SkinningMode == SkinningMode &&
 		OtherTyped->SamplingRegions == SamplingRegions &&
