- Remove log as it can spam during development

#rnx

diff --git a/Engine/Plugins/FX/Niagara/Source/Niagara/Private/DataInterface/NiagaraDataInterfaceDataChannelRead.cpp b/Engine/Plugins/FX/Niagara/Source/Niagara/Private/DataInterface/NiagaraDataInterfaceDataChannelRead.cpp
--- a/Engine/Plugins/FX/Niagara/Source/Niagara/Private/DataInterface/NiagaraDataInterfaceDataChannelRead.cpp
+++ b/Engine/Plugins/FX/Niagara/Source/Niagara/Private/DataInterface/NiagaraDataInterfaceDataChannelRead.cpp
@@ -144,7 +144,7 @@ bool FNDIDataChannelReadInstanceData::Tick(UNiagaraDataInterfaceDataChannelRead*
 	//Interface is never used so do not create any instance data.
 	if (CompiledData.UsedByCPU() == false && CompiledData.UsedByGPU() == false)
 	{
-		UE_LOG(LogNiagara, Warning, TEXT("Data Channel Interface is being initialized but it is never used.\nSystem: %s\nInterface: %s"), *Instance->GetSystem()->GetFullName(), *Interface->GetFullName());
+		//UE_LOG(LogNiagara, Warning, TEXT("Data Channel Interface is being initialized but it is never used.\nSystem: %s\nInterface: %s"), *Instance->GetSystem()->GetFullName(), *Interface->GetFullName());
 		return true;
 	}
 	
