- Fix data channel crash with null RHI

#rnx

diff --git a/Engine/Plugins/FX/Niagara/Source/Niagara/Private/NiagaraDataChannel.cpp b/Engine/Plugins/FX/Niagara/Source/Niagara/Private/NiagaraDataChannel.cpp
--- a/Engine/Plugins/FX/Niagara/Source/Niagara/Private/NiagaraDataChannel.cpp
+++ b/Engine/Plugins/FX/Niagara/Source/Niagara/Private/NiagaraDataChannel.cpp
@@ -537,12 +537,14 @@ void FNiagaraDataChannelData::BeginFrame(UNiagaraDataChannelHandler* Owner)
 	CPUSimData->EndSimulate();
 	
 	check(RTProxy);
-	FNiagaraGpuComputeDispatchInterface* DispathInterface = FNiagaraGpuComputeDispatchInterface::Get(Owner->GetWorld());
-	ENQUEUE_RENDER_COMMAND(FDataChannelProxyBeginFrame) (
-		[RT_Proxy = RTProxy.Get(), bRequirePreviousData](FRHICommandListImmediate& CmdList)
+	if (FNiagaraGpuComputeDispatchInterface* DispathInterface = FNiagaraGpuComputeDispatchInterface::Get(Owner->GetWorld()))
 	{
-		RT_Proxy->BeginFrame(bRequirePreviousData);
-	});
+		ENQUEUE_RENDER_COMMAND(FDataChannelProxyBeginFrame) (
+			[RT_Proxy = RTProxy.Get(), bRequirePreviousData](FRHICommandListImmediate& CmdList)
+		{
+			RT_Proxy->BeginFrame(bRequirePreviousData);
+		});
+	}
 }
 
 void FNiagaraDataChannelData::EndFrame(UNiagaraDataChannelHandler* Owner)
@@ -551,12 +553,14 @@ void FNiagaraDataChannelData::EndFrame(UNiagaraDataChannelHandler* Owner)
 	ConsumePublishRequests(Owner);
 
 	check(RTProxy);
-	FNiagaraGpuComputeDispatchInterface* DispathInterface = FNiagaraGpuComputeDispatchInterface::Get(Owner->GetWorld());
-	ENQUEUE_RENDER_COMMAND(FDataChannelProxyEndFrame) (
-		[DispathInterface, RT_Proxy = RTProxy.Get(), RT_BuffersForGPU = MoveTemp(BuffersForGPU)](FRHICommandListImmediate& CmdList)
+	if (FNiagaraGpuComputeDispatchInterface* DispathInterface = FNiagaraGpuComputeDispatchInterface::Get(Owner->GetWorld()))
 	{
-		RT_Proxy->EndFrame(DispathInterface, CmdList, RT_BuffersForGPU);
-	});
+		ENQUEUE_RENDER_COMMAND(FDataChannelProxyEndFrame) (
+			[DispathInterface, RT_Proxy = RTProxy.Get(), RT_BuffersForGPU = MoveTemp(BuffersForGPU)](FRHICommandListImmediate& CmdList)
+		{
+			RT_Proxy->EndFrame(DispathInterface, CmdList, RT_BuffersForGPU);
+		});
+	}
 	BuffersForGPU.Reset();
 }
 
