Remove updating NPC change id being updated when adding new parameters, results in non deterministic CompileIds for any associated niagara scripts which results in consistent recompilation.
#rnx

diff --git a/Engine/Plugins/FX/Niagara/Source/Niagara/Private/NiagaraParameterCollection.cpp b/Engine/Plugins/FX/Niagara/Source/Niagara/Private/NiagaraParameterCollection.cpp
--- a/Engine/Plugins/FX/Niagara/Source/Niagara/Private/NiagaraParameterCollection.cpp
+++ b/Engine/Plugins/FX/Niagara/Source/Niagara/Private/NiagaraParameterCollection.cpp
@@ -457,7 +457,6 @@ int32 UNiagaraParameterCollection::AddParameter(const FNiagaraVariable& Paramete
 	if (Idx == INDEX_NONE)
 	{
 		Modify();
-		CompileId = FGuid::NewGuid(); // Any scripts could now be referencing this new parameter and so would need to be changed.
 
 		Idx = Parameters.Add(Parameter);
 		DefaultInstance->AddParameter(Parameter);
