make svt the default caching for Render Target Volume sim caching

diff --git a/Engine/Plugins/FX/Niagara/Source/Niagara/Private/NiagaraDataInterfaceRenderTargetVolume.cpp b/Engine/Plugins/FX/Niagara/Source/Niagara/Private/NiagaraDataInterfaceRenderTargetVolume.cpp
--- a/Engine/Plugins/FX/Niagara/Source/Niagara/Private/NiagaraDataInterfaceRenderTargetVolume.cpp
+++ b/Engine/Plugins/FX/Niagara/Source/Niagara/Private/NiagaraDataInterfaceRenderTargetVolume.cpp
@@ -98,7 +98,7 @@ namespace NDIRenderTargetVolumeLocal
 		SVT = 2
 	};
 
-	int32 GSimCacheDataStorageMode = SimCacheStorageMode::OpenVDB;
+	int32 GSimCacheDataStorageMode = SimCacheStorageMode::SVT;
 	static FAutoConsoleVariableRef CVarSimCacheDataStorageMode(
 		TEXT("fx.Niagara.RenderTargetVolume.SimCacheDataStorageMode"),
 		GSimCacheDataStorageMode,
