- Update Niagara Component Settings when we spin up a world
- Allows us to be more dynamic with changes

diff --git a/Engine/Plugins/FX/Niagara/Source/Niagara/Private/NiagaraComponentSettings.cpp b/Engine/Plugins/FX/Niagara/Source/Niagara/Private/NiagaraComponentSettings.cpp
--- a/Engine/Plugins/FX/Niagara/Source/Niagara/Private/NiagaraComponentSettings.cpp
+++ b/Engine/Plugins/FX/Niagara/Source/Niagara/Private/NiagaraComponentSettings.cpp
@@ -171,7 +171,7 @@ namespace FNiagaraComponentSettings
 		ECVF_Scalability | ECVF_Default
 	);
 
-	void OnPostEngineInit()
+	void UpdateSettings()
 	{
 		if (GDynamicRHI == nullptr)
 		{
@@ -204,8 +204,6 @@ namespace FNiagaraComponentSettings
 
 		if (bShouldAllowGpuEmitters && GpuDenyListString.Len() > 0)
 		{
-			UE_LOG(LogNiagara, Warning, TEXT("GpuDenyListString = %s"), *GpuDenyListString);
-
 			TArray<FString> BanList;
 			GpuDenyListString.ParseIntoArray(BanList, TEXT("+"));
 
@@ -224,11 +222,6 @@ namespace FNiagaraComponentSettings
 					for (int32 i=0; i < UE_ARRAY_COUNT(CategoryData); ++i)
 					{
 						bMatches = BanCategoryStrings[i].IsEmpty() || CategoryData[i].MatchesWildcard(BanCategoryStrings[i]);
-						if (!BanCategoryStrings[i].IsEmpty())
-						{
-							UE_LOG(LogNiagara, Warning, TEXT("Matching %s=%s == %d"), *CategoryData[i], *BanCategoryStrings[i], bMatches);
-						}
-
 						if (bMatches == false)
 						{
 							break;
@@ -243,17 +236,21 @@ namespace FNiagaraComponentSettings
 			}
 		}
 
-		if (bShouldAllowGpuEmitters != bAllowGpuEmitters)
+		if (bAllowGpuEmitters != bShouldAllowGpuEmitters)
 		{
 			bAllowGpuEmitters = bShouldAllowGpuEmitters;
 			if (bAllowGpuEmitters == false)
 			{
 				UE_LOG(LogNiagara, Log,
-					TEXT("GPU emitters will not run for RHI(%s) Adapter(%s) as they match deny list RHIDeny(%s) AdapterDeny(%s)."),
+					TEXT("GPU emitters are disabled for RHI(%s) Adapter(%s).  RHIDeny(%s) AdapterDeny(%s) GpuDenyList(%s)."),
 					GDynamicRHI->GetName(), *GRHIAdapterName,
-					*GpuRHIDenyListString, *GpuRHIAdapterDenyListString
+					*GpuRHIDenyListString, *GpuRHIAdapterDenyListString, *GpuDenyListString
 				);
 			}
+			else
+			{
+				UE_LOG(LogNiagara, Log, TEXT("GPU emitters are enabled for RHI(%s) Adapter(%s)."), GDynamicRHI->GetName(), *GRHIAdapterName);
+			}
 		}
 	}
 
diff --git a/Engine/Plugins/FX/Niagara/Source/Niagara/Private/NiagaraModule.cpp b/Engine/Plugins/FX/Niagara/Source/Niagara/Private/NiagaraModule.cpp
--- a/Engine/Plugins/FX/Niagara/Source/Niagara/Private/NiagaraModule.cpp
+++ b/Engine/Plugins/FX/Niagara/Source/Niagara/Private/NiagaraModule.cpp
@@ -479,7 +479,7 @@ void INiagaraModule::OnPostEngineInit()
 		}
 	);
 
-	FNiagaraComponentSettings::OnPostEngineInit();
+	FNiagaraComponentSettings::UpdateSettings();
 }
 
 void INiagaraModule::OnPreExit()
diff --git a/Engine/Plugins/FX/Niagara/Source/Niagara/Private/NiagaraWorldManager.cpp b/Engine/Plugins/FX/Niagara/Source/Niagara/Private/NiagaraWorldManager.cpp
--- a/Engine/Plugins/FX/Niagara/Source/Niagara/Private/NiagaraWorldManager.cpp
+++ b/Engine/Plugins/FX/Niagara/Source/Niagara/Private/NiagaraWorldManager.cpp
@@ -4,6 +4,7 @@
 #include "Engine/Level.h"
 #include "GameFramework/Pawn.h"
 #include "NiagaraModule.h"
+#include "NiagaraComponentSettings.h"
 #include "NiagaraSystemSimulation.h"
 #include "NiagaraTypes.h"
 #include "NiagaraParameterCollection.h"
@@ -339,6 +340,9 @@ void FNiagaraWorldManager::Init(UWorld* InWorld)
 #if WITH_NIAGARA_DEBUGGER
 	NiagaraDebugHud.Reset(new FNiagaraDebugHud(World));
 #endif
+
+	// Make sure we update our component settings, this includes ban lists, etc
+	FNiagaraComponentSettings::UpdateSettings();
 }
 
 FNiagaraWorldManager::~FNiagaraWorldManager()
diff --git a/Engine/Plugins/FX/Niagara/Source/Niagara/Public/NiagaraComponentSettings.h b/Engine/Plugins/FX/Niagara/Source/Niagara/Public/NiagaraComponentSettings.h
--- a/Engine/Plugins/FX/Niagara/Source/Niagara/Public/NiagaraComponentSettings.h
+++ b/Engine/Plugins/FX/Niagara/Source/Niagara/Public/NiagaraComponentSettings.h
@@ -9,7 +9,7 @@ class FNiagaraEmitterInstance;
 
 namespace FNiagaraComponentSettings
 {
-	extern void OnPostEngineInit();
+	extern void UpdateSettings();
 	extern NIAGARA_API bool IsSystemAllowedToRun(const UNiagaraSystem* System);
 	extern NIAGARA_API bool IsEmitterAllowedToRun(const FNiagaraEmitterInstance* EmitterInstance);
 };
