Fix Niagara Mesh GPU LocalBounds BV being Meshscaled twice, once on the BV itself and once inside the LocalToWorld matrix on the GPU.

#rnx

diff --git a/Engine/Plugins/FX/Niagara/Source/Niagara/Private/NiagaraRendererMeshes.cpp b/Engine/Plugins/FX/Niagara/Source/Niagara/Private/NiagaraRendererMeshes.cpp
--- a/Engine/Plugins/FX/Niagara/Source/Niagara/Private/NiagaraRendererMeshes.cpp
+++ b/Engine/Plugins/FX/Niagara/Source/Niagara/Private/NiagaraRendererMeshes.cpp
@@ -1087,9 +1087,6 @@ void FNiagaraRendererMeshes::CreateMeshBatchForSection(
 	MeshBatch.DepthPriorityGroup = (ESceneDepthPriorityGroup)SceneProxy.GetDepthPriorityGroup(&View);
 
 	FBox LocalBounds = MeshData.RenderableMesh->GetLocalBounds();
-	LocalBounds.Min *= FVector(MeshData.Scale);
-	LocalBounds.Max *= FVector(MeshData.Scale);
-
 	FMeshBatchElement& BatchElement = MeshBatch.Elements[0];
 	if (ParticleMeshRenderData.bUseGPUScene)
 	{
