Simplify emitter properties:
* Remove "combine event spawn" option
* Remove "script category" property for emitters in systems
* Remove "asset options" properties for emitters in systems
* Move "attributes to preserve" to advanced view
* Hide "pre allocation count" if mode is automatic estimate
* move "max delta time per tick" to the plugin settings

diff --git a/Engine/Plugins/FX/Niagara/Source/Niagara/Classes/NiagaraEmitter.h b/Engine/Plugins/FX/Niagara/Source/Niagara/Classes/NiagaraEmitter.h
--- a/Engine/Plugins/FX/Niagara/Source/Niagara/Classes/NiagaraEmitter.h
+++ b/Engine/Plugins/FX/Niagara/Source/Niagara/Classes/NiagaraEmitter.h
@@ -316,10 +316,6 @@ struct NIAGARA_API FVersionedNiagaraEmitterData
 	UPROPERTY(EditAnywhere, BlueprintReadWrite, Category = "Emitter")
 	uint32 bRequiresPersistentIDs : 1;
 
-	/** Performance option to allow event based spawning to be combined into a single spawn.  This will result in a single exec from 0 to number of particles rather than several, when using ExecIndex() it is recommended not to do this. */
-	UPROPERTY(EditAnywhere, BlueprintReadWrite, Category = "Emitter")
-	uint32 bCombineEventSpawn : 1;
-
 	UPROPERTY(meta=(NiagaraNoMerge))
 	TArray<FNiagaraEventScriptProperties> EventHandlerScriptProps;
 	
@@ -329,18 +325,10 @@ struct NIAGARA_API FVersionedNiagaraEmitterData
 	UPROPERTY(EditAnywhere, Category = "Scalability", meta=(DisplayInScalabilityContext))
 	FNiagaraEmitterScalabilityOverrides ScalabilityOverrides;
 
-	/** Whether to limit the max tick delta time or not. */
-	UPROPERTY(EditAnywhere, BlueprintReadWrite, AdvancedDisplay, Category = "Emitter", meta = (InlineEditConditionToggle))
-	uint32 bLimitDeltaTime : 1;
-
 	/** An override on the max number of GPU particles we expect to spawn in a single frame. A value of 0 means it'll use fx.MaxNiagaraGPUParticlesSpawnPerFrame.*/
-	UPROPERTY(EditAnywhere, BlueprintReadWrite, AdvancedDisplay, Category = "Emitter", meta = (EditCondition = "SimTarget == ENiagaraSimTarget::GPUComputeSim", DisplayName = "Max GPU Particles Spawn per Frame"))
+	UPROPERTY(EditAnywhere, BlueprintReadWrite, AdvancedDisplay, Category = "Emitter", meta = (EditCondition = "SimTarget == ENiagaraSimTarget::GPUComputeSim", DisplayName = "Max GPU Particles Spawn per Frame", EditConditionHides))
 	int32 MaxGPUParticlesSpawnPerFrame;
 
-	/** Limits the delta time per tick to prevent simulation spikes due to frame lags. */
-	UPROPERTY(EditAnywhere, BlueprintReadWrite, AdvancedDisplay, Category = "Emitter", meta = (EditCondition = "bLimitDeltaTime", ForceUnits=s))
-	float MaxDeltaTimePerTick = 0.125;
-
 	/**
 	The emitter needs to allocate memory for the particles each tick.
 	To prevent reallocations, the emitter should allocate as much memory as is needed for the max particle count.
@@ -353,7 +341,7 @@ struct NIAGARA_API FVersionedNiagaraEmitterData
 	The emitter will allocate at least this many particles on it's first tick.
 	This can aid performance by avoiding many allocations as an emitter ramps up to it's max size.
 	*/
-	UPROPERTY(EditAnywhere, AdvancedDisplay, Category = "Emitter", meta = (EditCondition = "AllocationMode != EParticleAllocationMode::AutomaticEstimate"))
+	UPROPERTY(EditAnywhere, AdvancedDisplay, Category = "Emitter", meta = (EditCondition = "AllocationMode != EParticleAllocationMode::AutomaticEstimate", EditConditionHides))
 	int32 PreAllocationCount = 0;
 
 	UPROPERTY()
@@ -825,18 +813,6 @@ public:
 	UPROPERTY(meta = (DeprecatedProperty))
 	uint32 bRequiresPersistentIDs_DEPRECATED : 1;
 
-	/** Use property in struct returned from GetEmitterData() instead */ 
-	UPROPERTY(meta = (DeprecatedProperty))
-	uint32 bCombineEventSpawn_DEPRECATED : 1;
-
-	/** Use property in struct returned from GetEmitterData() instead */ 
-	UPROPERTY(meta = (DeprecatedProperty))
-	float MaxDeltaTimePerTick_DEPRECATED;
-
-	/** Use property in struct returned from GetEmitterData() instead */ 
-	UPROPERTY(meta = (DeprecatedProperty))
-	uint32 bLimitDeltaTime_DEPRECATED : 1;
-
 	/** Use property in struct returned from GetEmitterData() instead */ 
 	UPROPERTY(meta = (DeprecatedProperty))
 	uint32 MaxGPUParticlesSpawnPerFrame_DEPRECATED;
diff --git a/Engine/Plugins/FX/Niagara/Source/Niagara/Classes/NiagaraEmitterInstance.h b/Engine/Plugins/FX/Niagara/Source/Niagara/Classes/NiagaraEmitterInstance.h
--- a/Engine/Plugins/FX/Niagara/Source/Niagara/Classes/NiagaraEmitterInstance.h
+++ b/Engine/Plugins/FX/Niagara/Source/Niagara/Classes/NiagaraEmitterInstance.h
@@ -255,8 +255,6 @@ private:
 
 	/** Typical resets must be deferred until the tick as the RT could still be using the current buffer. */
 	uint32 bResetPending : 1;
-	/** Allows event spawn to be combined into a single spawn.  This is only safe when not using things like ExecIndex(). */
-	uint32 bCombineEventSpawn : 1;
 
 	// This is used to keep track which particles have spawned a component. This is needed when the bOnlyCreateComponentsOnParticleSpawn flag is set in the renderer.
 	// Without this bookkeeping, the particles would lose their components when the render state is recreated or the visibility tag flips them off and on again.
diff --git a/Engine/Plugins/FX/Niagara/Source/Niagara/Private/NiagaraEmitter.cpp b/Engine/Plugins/FX/Niagara/Source/Niagara/Private/NiagaraEmitter.cpp
--- a/Engine/Plugins/FX/Niagara/Source/Niagara/Private/NiagaraEmitter.cpp
+++ b/Engine/Plugins/FX/Niagara/Source/Niagara/Private/NiagaraEmitter.cpp
@@ -177,9 +177,6 @@ UNiagaraEmitter::UNiagaraEmitter(const FObjectInitializer& Initializer)
 , bUseMinDetailLevel_DEPRECATED(false)
 , bUseMaxDetailLevel_DEPRECATED(false)
 , bRequiresPersistentIDs_DEPRECATED(false)
-, bCombineEventSpawn_DEPRECATED(false)
-, MaxDeltaTimePerTick_DEPRECATED(0.125)
-, bLimitDeltaTime_DEPRECATED(true)
 , MaxGPUParticlesSpawnPerFrame_DEPRECATED(0)
 #endif
 {
@@ -1288,14 +1285,6 @@ void UNiagaraEmitter::PostEditChangeVersionedProperty(FPropertyChangedEvent& Pro
 		}
 		bNeedsRecompile = true;
 	}
-	else if (PropertyName == GET_MEMBER_NAME_CHECKED(FVersionedNiagaraEmitterData, MaxDeltaTimePerTick))
-	{
-		if (EmitterData->GraphSource != nullptr)
-		{
-			EmitterData->GraphSource->MarkNotSynchronized(TEXT("MaxDeltaTimePerTick changed."));
-		}
-		bNeedsRecompile = true;
-	}
 
 	ResolveScalabilitySettings();
 
@@ -2757,16 +2746,13 @@ void UNiagaraEmitter::CheckVersionDataAvailable()
 	Data.GraphSource = GraphSource_DEPRECATED;
 	Data.bDeterminism = bDeterminism_DEPRECATED;
 	Data.AllocationMode = AllocationMode_DEPRECATED;
-	Data.bLimitDeltaTime = bLimitDeltaTime_DEPRECATED;
 	Data.SpawnScriptProps = SpawnScriptProps_DEPRECATED;
 	Data.SimulationStages = SimulationStages_DEPRECATED;
 	Data.EditorParameters = EditorParameters_DEPRECATED;
 	Data.GPUComputeScript = GPUComputeScript_DEPRECATED;
 	Data.UpdateScriptProps = UpdateScriptProps_DEPRECATED;
 	Data.PreAllocationCount = PreAllocationCount_DEPRECATED;
-	Data.bCombineEventSpawn = bCombineEventSpawn_DEPRECATED;
 	Data.RendererProperties = RendererProperties_DEPRECATED;
-	Data.MaxDeltaTimePerTick = MaxDeltaTimePerTick_DEPRECATED;
 	Data.ScalabilityOverrides = ScalabilityOverrides_DEPRECATED;
 	Data.AttributesToPreserve = AttributesToPreserve_DEPRECATED;
 	Data.bInterpolatedSpawning = bInterpolatedSpawning_DEPRECATED;
@@ -2822,8 +2808,6 @@ FVersionedNiagaraEmitterData::FVersionedNiagaraEmitterData()
 	, bGpuAlwaysRunParticleUpdateScript(false)
 #endif
 	, bRequiresPersistentIDs(false)
-	, bCombineEventSpawn(false)
-	, bLimitDeltaTime(true)
 	, MaxGPUParticlesSpawnPerFrame(0)
 	, bRequiresViewUniformBuffer(false)
 #if WITH_EDITORONLY_DATA
diff --git a/Engine/Plugins/FX/Niagara/Source/Niagara/Private/NiagaraEmitterInstance.cpp b/Engine/Plugins/FX/Niagara/Source/Niagara/Private/NiagaraEmitterInstance.cpp
--- a/Engine/Plugins/FX/Niagara/Source/Niagara/Private/NiagaraEmitterInstance.cpp
+++ b/Engine/Plugins/FX/Niagara/Source/Niagara/Private/NiagaraEmitterInstance.cpp
@@ -27,14 +27,6 @@ DECLARE_CYCLE_STAT(TEXT("Emitter Event CopyBuffer [CNC]"), STAT_NiagaraEvent_Cop
 DECLARE_CYCLE_STAT(TEXT("Emitter Error Check [CNC]"), STAT_NiagaraEmitterErrorCheck, STATGROUP_Niagara);
 DECLARE_CYCLE_STAT(TEXT("Init Emitters [GT]"), STAT_NiagaraEmitterInit, STATGROUP_Niagara);
 
-static int32 GbNiagaraAllowEventSpawnCombine = 1;
-static FAutoConsoleVariableRef CVarNiagaraAllowEventSpawnCombine(
-	TEXT("fx.Niagara.AllowEventSpawnCombine"),
-	GbNiagaraAllowEventSpawnCombine,
-	TEXT("Allows events spawning to be combined, 0=Disabled, 1=Allowed Based On Emitter, 2=Force On."),
-	ECVF_Default
-);
-
 static int32 GbDumpParticleData = 0;
 static FAutoConsoleVariableRef CVarNiagaraDumpParticleData(
 	TEXT("fx.DumpParticleData"),
@@ -127,7 +119,6 @@ FNiagaraEmitterInstance::FNiagaraEmitterInstance(FNiagaraSystemInstance* InParen
 	, CachedSystemFixedBounds(ForceInit)
 	, ParentSystemInstance(InParentSystemInstance)
 	, bResetPending(false)
-	, bCombineEventSpawn(false)
 {
 	ParticleDataSet = new FNiagaraDataSet();
 
@@ -497,8 +488,6 @@ void FNiagaraEmitterInstance::Init(int32 InEmitterIdx, FNiagaraSystemInstanceID
 	MaxInstanceCount = EmitterData->GetMaxInstanceCount();
 	ParticleDataSet->SetMaxInstanceCount(MaxInstanceCount);
 	ParticleDataSet->SetMaxAllocationCount(EmitterData->GetMaxAllocationCount());
-
-	bCombineEventSpawn = GbNiagaraAllowEventSpawnCombine && (EmitterData->bCombineEventSpawn || (GbNiagaraAllowEventSpawnCombine == 2));
 }
 
 void FNiagaraEmitterInstance::InitDITickLists()
@@ -1825,60 +1814,32 @@ void FNiagaraEmitterInstance::Tick(float DeltaSeconds)
 		EventSpawnStart = Data.GetDestinationDataChecked().GetNumInstances();
 		if (EventInstanceData.IsValid())
 		{
-			if (bCombineEventSpawn)
+			for (int32 EventScriptIdx = 0; EventScriptIdx < EmitterData->GetEventHandlers().Num(); EventScriptIdx++)
 			{
-				int32 EventParticlesToSpawn = 0;
+				FNiagaraEventHandlingInfo& Info = EventInstanceData->EventHandlingInfo[EventScriptIdx];
 
-				for (int32 EventScriptIdx = 0; EventScriptIdx < EmitterData->GetEventHandlers().Num(); EventScriptIdx++)
+				for (int32 i = 0; i < Info.SpawnCounts.Num(); i++)
 				{
-					FNiagaraEventHandlingInfo& Info = EventInstanceData->EventHandlingInfo[EventScriptIdx];
-					for (int32 i = 0; i < Info.SpawnCounts.Num(); i++)
+					const int32 EventNumToSpawn = FMath::Min(Info.SpawnCounts[i], SpawnCountRemaining);
+					if (EventNumToSpawn > 0)
 					{
-						const int32 EventNumToSpawn = FMath::Min(Info.SpawnCounts[i], SpawnCountRemaining);
-						EventParticlesToSpawn += EventNumToSpawn;
-						SpawnCountRemaining -= EventNumToSpawn;
-						TotalActualEventSpawns += EventNumToSpawn;
-						Info.SpawnCounts[i] = EventNumToSpawn;
+						const int32 CurrNumParticles = Data.GetDestinationDataChecked().GetNumInstances();
+
+						//Event spawns are instantaneous at the middle of the frame?
+						SpawnIntervalBinding.SetValue(0.0f);
+						InterpSpawnStartBinding.SetValue(DeltaSeconds * 0.5f);
+						SpawnGroupBinding.SetValue(0);
+						SpawnParticles(EventNumToSpawn, TEXT("Event Spawn"));
+
+						//Update EventSpawnCounts to the number actually spawned.
+						const int32 NumActuallySpawned = Data.GetDestinationDataChecked().GetNumInstances() - CurrNumParticles;
+						TotalActualEventSpawns += NumActuallySpawned;
+						Info.SpawnCounts[i] = NumActuallySpawned;
+						SpawnCountRemaining -= NumActuallySpawned;
 					}
-				}
-
-				if (EventParticlesToSpawn > 0)
-				{
-					SpawnIntervalBinding.SetValue(0.0f);
-					InterpSpawnStartBinding.SetValue(DeltaSeconds * 0.5f);
-					SpawnGroupBinding.SetValue(0);
-					SpawnParticles(EventParticlesToSpawn, TEXT("Event Spawn"));
-				}
-			}
-			else
-			{
-				for (int32 EventScriptIdx = 0; EventScriptIdx < EmitterData->GetEventHandlers().Num(); EventScriptIdx++)
-				{
-					FNiagaraEventHandlingInfo& Info = EventInstanceData->EventHandlingInfo[EventScriptIdx];
-
-					for (int32 i = 0; i < Info.SpawnCounts.Num(); i++)
+					else
 					{
-						const int32 EventNumToSpawn = FMath::Min(Info.SpawnCounts[i], SpawnCountRemaining);
-						if (EventNumToSpawn > 0)
-						{
-							const int32 CurrNumParticles = Data.GetDestinationDataChecked().GetNumInstances();
-
-							//Event spawns are instantaneous at the middle of the frame?
-							SpawnIntervalBinding.SetValue(0.0f);
-							InterpSpawnStartBinding.SetValue(DeltaSeconds * 0.5f);
-							SpawnGroupBinding.SetValue(0);
-							SpawnParticles(EventNumToSpawn, TEXT("Event Spawn"));
-
-							//Update EventSpawnCounts to the number actually spawned.
-							const int32 NumActuallySpawned = Data.GetDestinationDataChecked().GetNumInstances() - CurrNumParticles;
-							TotalActualEventSpawns += NumActuallySpawned;
-							Info.SpawnCounts[i] = NumActuallySpawned;
-							SpawnCountRemaining -= NumActuallySpawned;
-						}
-						else
-						{
-							Info.SpawnCounts[i] = 0;
-						}
+						Info.SpawnCounts[i] = 0;
 					}
 				}
 			}
diff --git a/Engine/Plugins/FX/Niagara/Source/Niagara/Private/NiagaraSystem.cpp b/Engine/Plugins/FX/Niagara/Source/Niagara/Private/NiagaraSystem.cpp
--- a/Engine/Plugins/FX/Niagara/Source/Niagara/Private/NiagaraSystem.cpp
+++ b/Engine/Plugins/FX/Niagara/Source/Niagara/Private/NiagaraSystem.cpp
@@ -1849,6 +1849,7 @@ void UNiagaraSystem::CacheFromCompiledData()
 {
 	LLM_SCOPE(ELLMTag::Niagara);
 	const FNiagaraDataSetCompiledData& SystemDataSet = SystemCompiledData.DataSetCompiledData;
+	const UNiagaraSettings* NiagaraSettings = GetDefault<UNiagaraSettings>();
 
 	bNeedsAsyncOptimize = true;
 
@@ -1896,9 +1897,9 @@ void UNiagaraSystem::CacheFromCompiledData()
 
 				DataSetCompiledData = &EmitterCompiledData[i]->DataSetCompiledData;
 
-				if (EmitterData->bLimitDeltaTime)
+				if (NiagaraSettings->bLimitDeltaTime)
 				{
-					MaxDeltaTime = MaxDeltaTime.IsSet() ? FMath::Min(MaxDeltaTime.GetValue(), EmitterData->MaxDeltaTimePerTick) : EmitterData->MaxDeltaTimePerTick;
+					MaxDeltaTime = MaxDeltaTime.IsSet() ? FMath::Min(MaxDeltaTime.GetValue(), NiagaraSettings->MaxDeltaTimePerTick) : NiagaraSettings->MaxDeltaTimePerTick;
 				}
 			}
 			Handle.GetInstance().Emitter->ConditionalPostLoad();
diff --git a/Engine/Plugins/FX/Niagara/Source/Niagara/Public/NiagaraSettings.h b/Engine/Plugins/FX/Niagara/Source/Niagara/Public/NiagaraSettings.h
--- a/Engine/Plugins/FX/Niagara/Source/Niagara/Public/NiagaraSettings.h
+++ b/Engine/Plugins/FX/Niagara/Source/Niagara/Public/NiagaraSettings.h
@@ -140,6 +140,14 @@ class NIAGARA_API UNiagaraSettings : public UDeveloperSettings
 	UPROPERTY(config, EditAnywhere, Category = Niagara, meta = (DisplayName = "Enable building data for Experimental VM"))
 	bool bExperimentalVMEnabled = false;
 
+	/** Whether to limit the max tick delta time or not. */
+	UPROPERTY(config, EditAnywhere, BlueprintReadWrite, AdvancedDisplay, Category = "Niagara", meta = (InlineEditConditionToggle))
+	bool bLimitDeltaTime = true;
+
+	/** Limits the delta time per tick for emitters to prevent simulation spikes due to frame lags. */
+	UPROPERTY(config, EditAnywhere, BlueprintReadWrite, AdvancedDisplay, Category = "Niagara", meta = (EditCondition = "bLimitDeltaTime", ForceUnits=s))
+	float MaxDeltaTimePerTick = 0.125;
+	
 	/** Default effect type to use for effects that don't define their own. Can be null. */
 	UPROPERTY(config, EditAnywhere, Category = Niagara, meta = (AllowedClasses = "/Script/Niagara.NiagaraEffectType"))
 	FSoftObjectPath DefaultEffectType;
diff --git a/Engine/Plugins/FX/Niagara/Source/NiagaraEditor/Private/NiagaraEmitterDetailsCustomization.cpp b/Engine/Plugins/FX/Niagara/Source/NiagaraEditor/Private/NiagaraEmitterDetailsCustomization.cpp
--- a/Engine/Plugins/FX/Niagara/Source/NiagaraEditor/Private/NiagaraEmitterDetailsCustomization.cpp
+++ b/Engine/Plugins/FX/Niagara/Source/NiagaraEditor/Private/NiagaraEmitterDetailsCustomization.cpp
@@ -58,6 +58,16 @@ void CustomizeEmitterData(IDetailLayoutBuilder& InDetailLayout, UNiagaraEmitter*
 			PropertyRow->GetPropertyHandle()->SetOnPropertyValueChanged(FSimpleDelegate::CreateLambda(PostEditChangeLambda));
 		}
 	}
+
+	// only show some categories for standalone emitter assets 
+	if (NiagaraSystem && NiagaraSystem->GetOuter() != GetTransientPackage())
+	{
+		IDetailCategoryBuilder& ScriptCategory = InDetailLayout.EditCategory(FName("Script"));
+		ScriptCategory.SetCategoryVisibility(false);
+		IDetailCategoryBuilder& AssetCategory = InDetailLayout.EditCategory(FName("Asset Options"));
+		AssetCategory.SetCategoryVisibility(false);
+		
+	}
 }
 
 void FNiagaraEmitterDetails::CustomizeDetails(IDetailLayoutBuilder& InDetailLayout)
diff --git a/Engine/Plugins/FX/Niagara/Source/NiagaraEditor/Private/ViewModels/Stack/NiagaraStackEmitterSettingsGroup.cpp b/Engine/Plugins/FX/Niagara/Source/NiagaraEditor/Private/ViewModels/Stack/NiagaraStackEmitterSettingsGroup.cpp
--- a/Engine/Plugins/FX/Niagara/Source/NiagaraEditor/Private/ViewModels/Stack/NiagaraStackEmitterSettingsGroup.cpp
+++ b/Engine/Plugins/FX/Niagara/Source/NiagaraEditor/Private/ViewModels/Stack/NiagaraStackEmitterSettingsGroup.cpp
@@ -20,6 +20,7 @@
 #include "Styling/AppStyle.h"
 #include "IDetailTreeNode.h"
 #include "NiagaraNodeFunctionCall.h"
+#include "NiagaraSettings.h"
 #include "Toolkits/SystemToolkitModes/NiagaraSystemToolkitModeBase.h"
 #include "Widgets/SNiagaraHierarchy.h"
 
@@ -202,12 +203,13 @@ void UNiagaraStackEmitterPropertiesItem::RefreshIssues(TArray<FStackIssue>& NewI
 		TWeakPtr<FNiagaraSystemViewModel> WeakSysViewModel = GetSystemViewModel();
 		if (System.NeedsWarmup())
 		{
+			const UNiagaraSettings* NiagaraSettings = GetDefault<UNiagaraSettings>();
 			float WarmupDelta = System.GetWarmupTickDelta();
-			if (ActualEmitterData->bLimitDeltaTime && ActualEmitterData->MaxDeltaTimePerTick < WarmupDelta)
+			if (NiagaraSettings->bLimitDeltaTime && NiagaraSettings->MaxDeltaTimePerTick < WarmupDelta)
 			{
 				TArray<FStackIssueFix> Fixes;
 
-				float MaxEmitterDt = ActualEmitterData->MaxDeltaTimePerTick;
+				float MaxEmitterDt = NiagaraSettings->MaxDeltaTimePerTick;
 				//This emitter does not allow ticks with a delta time so large.
 				FText FixDescriptionReduceWarmupDt = LOCTEXT("FixWarmupDeltaTime", "Reduce System Warmup Delta Time");
 				Fixes.Emplace(
@@ -223,28 +225,10 @@ void UNiagaraStackEmitterPropertiesItem::RefreshIssues(TArray<FStackIssue>& NewI
 							}
 						}));
 
-				FVersionedNiagaraEmitterWeakPtr WeakEmitter = GetEmitterViewModel()->GetEmitter().ToWeakPtr();
-				FText FixDescriptionReduceIncreaseEmitterDt = LOCTEXT("FixEmitterDeltaTime", "Increase Max Emitter Delta Time");
-				Fixes.Emplace(
-					FixDescriptionReduceIncreaseEmitterDt,
-					FStackIssueFixDelegate::CreateLambda([=]()
-						{
-							auto PinnedSysViewModel = WeakSysViewModel.Pin();
-							FVersionedNiagaraEmitter PinnedEmitter = WeakEmitter.ResolveWeakPtr();
-							if (PinnedEmitter.Emitter && PinnedSysViewModel)
-							{
-								FScopedTransaction ScopedTransaction(FixDescriptionReduceIncreaseEmitterDt);
-
-								PinnedEmitter.Emitter->Modify();
-								PinnedEmitter.GetEmitterData()->MaxDeltaTimePerTick = WarmupDelta;
-								PinnedSysViewModel->RefreshAll();
-							}
-						}));
-
 				FStackIssue WarmupDeltaTimeExceedsEmitterDeltaTimeWarning(
 					EStackIssueSeverity::Warning,
 					LOCTEXT("WarmupDeltaTimeExceedsEmitterDeltaTimeWarningSummary", "System Warmup Delta Time Exceeds Emitter Max Delta Time."),
-					LOCTEXT("WarmupDeltaTimeExceedsEmitterDeltaTimeWarningText", "Max Tick Delta Time is smaller than the System's Warmup Delta Time. This could cause unintended results during warmup for this emitter."),
+					LOCTEXT("WarmupDeltaTimeExceedsEmitterDeltaTimeWarningText", "Max Tick Delta Time is smaller than the System's Warmup Delta Time. This could cause unintended results during warmup for this emitter.\nThe max tick delta time can be changed in the Niagara settings."),
 					GetStackEditorDataKey(),
 					false,
 					Fixes);
diff --git a/Engine/Plugins/FX/Niagara/Source/NiagaraEditor/Private/ViewModels/Stack/NiagaraStackSystemPropertiesItem.cpp b/Engine/Plugins/FX/Niagara/Source/NiagaraEditor/Private/ViewModels/Stack/NiagaraStackSystemPropertiesItem.cpp
--- a/Engine/Plugins/FX/Niagara/Source/NiagaraEditor/Private/ViewModels/Stack/NiagaraStackSystemPropertiesItem.cpp
+++ b/Engine/Plugins/FX/Niagara/Source/NiagaraEditor/Private/ViewModels/Stack/NiagaraStackSystemPropertiesItem.cpp
@@ -78,17 +78,14 @@ void UNiagaraStackSystemPropertiesItem::RefreshChildrenInternal(const TArray<UNi
 	if (System->NeedsWarmup())
 	{
 		float WarmupDelta = System->GetWarmupTickDelta();
-
+		const UNiagaraSettings* NiagaraSettings = GetDefault<UNiagaraSettings>();
 		for (FNiagaraEmitterHandle& EmitterHandle : System->GetEmitterHandles())
 		{
-			FVersionedNiagaraEmitter& Emitter = EmitterHandle.GetInstance();
-			const FVersionedNiagaraEmitterData* EmitterData = Emitter.GetEmitterData();			
-			check(EmitterData);
-			if (EmitterData->bLimitDeltaTime && EmitterData->MaxDeltaTimePerTick < WarmupDelta)
+			if (NiagaraSettings->bLimitDeltaTime && NiagaraSettings->MaxDeltaTimePerTick < WarmupDelta)
 			{
 				TArray<FStackIssueFix> Fixes;
 
-				float MaxEmitterDt = EmitterData->MaxDeltaTimePerTick;
+				float MaxEmitterDt = NiagaraSettings->MaxDeltaTimePerTick;
 				//This emitter does not allow ticks with a delta time so large.
 				FText FixDescriptionReduceWarmupDt = LOCTEXT("FixWarmupDeltaTime", "Reduce System Warmup Delta Time");
 				Fixes.Emplace(
@@ -104,28 +101,10 @@ void UNiagaraStackSystemPropertiesItem::RefreshChildrenInternal(const TArray<UNi
 							}
 						}));
 
-				FVersionedNiagaraEmitterWeakPtr WeakEmitter = Emitter.ToWeakPtr();
-				FText FixDescriptionReduceIncreaseEmitterDt = LOCTEXT("FixEmitterDeltaTime", "Increase Max Emitter Delta Time");
-				Fixes.Emplace(
-					FixDescriptionReduceIncreaseEmitterDt,
-					FStackIssueFixDelegate::CreateLambda([=]()
-						{
-							auto PinnedSysViewModel = WeakSysViewModel.Pin();
-							FVersionedNiagaraEmitter PinnedEmitter = WeakEmitter.ResolveWeakPtr();
-							if (PinnedEmitter.Emitter && PinnedSysViewModel)
-							{
-								FScopedTransaction ScopedTransaction(FixDescriptionReduceIncreaseEmitterDt);
-
-								PinnedEmitter.Emitter->Modify();
-								PinnedEmitter.GetEmitterData()->MaxDeltaTimePerTick = WarmupDelta;
-								PinnedSysViewModel->RefreshAll();
-							}
-						}));
-
 				FStackIssue WarmupDeltaTimeExceedsEmitterDeltaTimeWarning(
 					EStackIssueSeverity::Warning,
 					LOCTEXT("WarmupDeltaTimeExceedsEmitterDeltaTimeWarningSummary", "System Warmup Delta Time Exceeds Emitter Max Delta Time."),
-					FText::Format(LOCTEXT("WarmupDeltaTimeExceedsEmitterDeltaTimeWarningText", "Emitter {0} has a Max Tick Delta Time smaller than the System's Warmup Delta Time. This could cause unintended results during warmup for this emitter."), FText::FromString(EmitterHandle.GetUniqueInstanceName())),
+					FText::Format(LOCTEXT("WarmupDeltaTimeExceedsEmitterDeltaTimeWarningText", "Emitter {0} has a Max Tick Delta Time smaller than the System's Warmup Delta Time. This could cause unintended results during warmup for this emitter.\nThe max tick delta time can be changed in the Niagara setting."), FText::FromString(EmitterHandle.GetUniqueInstanceName())),
 					GetStackEditorDataKey(),
 					false,
 					Fixes);
