Reorder system properties

diff --git a/Engine/Plugins/FX/Niagara/Source/Niagara/Classes/NiagaraSystem.h b/Engine/Plugins/FX/Niagara/Source/Niagara/Classes/NiagaraSystem.h
--- a/Engine/Plugins/FX/Niagara/Source/Niagara/Classes/NiagaraSystem.h
+++ b/Engine/Plugins/FX/Niagara/Source/Niagara/Classes/NiagaraSystem.h
@@ -680,10 +680,10 @@ public:
 		return Found ? *Found : nullptr;
 	}
 	
-	UPROPERTY(EditAnywhere, Category = "Debug", Transient)
+	UPROPERTY(EditAnywhere, Category = "Debug", Transient, AdvancedDisplay)
 	bool bDumpDebugSystemInfo = false;
 
-	UPROPERTY(EditAnywhere, Category = "Debug", Transient)
+	UPROPERTY(EditAnywhere, Category = "Debug", Transient, AdvancedDisplay)
 	bool bDumpDebugEmitterInfo = false;
 
 	bool bFullyLoaded = false;
@@ -802,6 +802,9 @@ private:
 	void UpdateHasGPUEmitters();
 
 protected:
+	/** An effect types defines settings shared between systems, for example scalability and validation rules.
+	 * Things like environment fx usually have a different effect type than gameplay relevant fx such as weapon impacts.
+	 * This way whole classes of effects can be adjusted at once. */
 	UPROPERTY(EditAnywhere, Category = "System")
 	TObjectPtr<UNiagaraEffectType> EffectType;
 
@@ -823,7 +826,7 @@ protected:
 	UPROPERTY()
 	TArray<FNiagaraEmitterHandle> EmitterHandles;
 
-	UPROPERTY(EditAnywhere, Category="System")
+	UPROPERTY(EditAnywhere, Category="System", AdvancedDisplay)
 	TArray<TObjectPtr<UNiagaraParameterCollectionInstance>> ParameterCollectionOverrides;
 
 #if WITH_EDITORONLY_DATA
@@ -879,7 +882,7 @@ protected:
 	FOnScalabilityChanged OnScalabilityChangedDelegate;
 #endif
 
-	/** The fixed bounding box value. bFixedBounds is the condition whether the fixed bounds can be edited. */
+	/** The fixed bounding box value for the whole system. When placed in the level and the bounding box is not visible to the camera, the effect is culled from rendering. */
 	UPROPERTY(EditAnywhere, Category = "System", meta = (SkipSystemResetOnChange = "true", EditCondition = "bFixedBounds"))
 	FBox FixedBounds;
 
diff --git a/Engine/Plugins/FX/Niagara/Source/NiagaraEditor/Private/NiagaraSystemDetailsCustomization.cpp b/Engine/Plugins/FX/Niagara/Source/NiagaraEditor/Private/NiagaraSystemDetailsCustomization.cpp
--- a/Engine/Plugins/FX/Niagara/Source/NiagaraEditor/Private/NiagaraSystemDetailsCustomization.cpp
+++ b/Engine/Plugins/FX/Niagara/Source/NiagaraEditor/Private/NiagaraSystemDetailsCustomization.cpp
@@ -16,4 +16,40 @@ void FNiagaraSystemDetails::CustomizeDetails(IDetailLayoutBuilder& InDetailLayou
 {
 	// we display the scalability category within scalability mode, which is why we hide it here
 	InDetailLayout.HideCategory("Scalability");
+
+	InDetailLayout.SortCategories([](const TMap<FName, IDetailCategoryBuilder*>& CategoryMap)
+	{
+		for (const TPair<FName, IDetailCategoryBuilder*>& Pair : CategoryMap)
+		{
+			int32 SortOrder = Pair.Value->GetSortOrder();
+			const FName& CategoryName = Pair.Key;
+
+			if (CategoryName == "Random")
+			{
+				SortOrder = 1;
+			}
+			else if (CategoryName == "Rendering")
+			{
+				SortOrder = 2;
+			}
+			else if (CategoryName == "System")
+			{
+				SortOrder = 3;
+			}
+			else if (CategoryName == "Warmup")
+			{
+				SortOrder = 4;
+			}
+			else if (CategoryName == "Performance")
+			{
+				SortOrder = 5;
+			}
+			else
+			{
+				SortOrder += 6;
+			}
+
+			Pair.Value->SetSortOrder(SortOrder);
+		}
+	});
 }
diff --git a/Engine/Source/Runtime/Engine/Classes/Particles/ParticleSystem.h b/Engine/Source/Runtime/Engine/Classes/Particles/ParticleSystem.h
--- a/Engine/Source/Runtime/Engine/Classes/Particles/ParticleSystem.h
+++ b/Engine/Source/Runtime/Engine/Classes/Particles/ParticleSystem.h
@@ -125,7 +125,7 @@ public:
 	ENGINE_API virtual bool IsReadyForFinishDestroy() override;
 
 	/** Max number of components of this system to keep resident in the world component pool. */
-	UPROPERTY(EditAnywhere, Category = Performance)
+	UPROPERTY(EditAnywhere, Category = Performance, AdvancedDisplay)
 	uint32 MaxPoolSize;
 	//TODO: Allow pool size overriding per world and possibly implement some preallocation too.
 
@@ -134,7 +134,7 @@ public:
 	* This can amortize runtime activation cost by moving it to load time.
 	* Use with care as this could cause large hitches for systems loaded/unloaded during play rather than at level load.
 	*/
-	UPROPERTY(EditAnywhere, Category = Performance)
+	UPROPERTY(EditAnywhere, Category = Performance, AdvancedDisplay)
 	uint32 PoolPrimeSize = 0;
 
 #if WITH_PER_SYSTEM_PARTICLE_PERF_STATS
