- Landscape data interface optimizations

diff --git a/Engine/Plugins/FX/Niagara/Source/Niagara/Classes/NiagaraDataInterfaceLandscape.h b/Engine/Plugins/FX/Niagara/Source/Niagara/Classes/NiagaraDataInterfaceLandscape.h
--- a/Engine/Plugins/FX/Niagara/Source/Niagara/Classes/NiagaraDataInterfaceLandscape.h
+++ b/Engine/Plugins/FX/Niagara/Source/Niagara/Classes/NiagaraDataInterfaceLandscape.h
@@ -82,7 +82,7 @@ public:
 protected:
 	virtual bool CopyToInternal(UNiagaraDataInterface* Destination) const override;
 	void ApplyLandscape(const FNiagaraSystemInstance& SystemInstance, FNDILandscapeData_GameThread& InstanceData) const;
-	ALandscape* GetLandscape(const FNiagaraSystemInstance& SystemInstance, ALandscape* Hint) const;
+	ALandscape* GetLandscape(const FNiagaraSystemInstance& SystemInstance, FNDILandscapeData_GameThread& Hint) const;
 
 	static const FName GetBaseColorName;
 	static const FName GetHeightName;
diff --git a/Engine/Plugins/FX/Niagara/Source/Niagara/Private/NiagaraDataInterfaceLandscape.cpp b/Engine/Plugins/FX/Niagara/Source/Niagara/Private/NiagaraDataInterfaceLandscape.cpp
--- a/Engine/Plugins/FX/Niagara/Source/Niagara/Private/NiagaraDataInterfaceLandscape.cpp
+++ b/Engine/Plugins/FX/Niagara/Source/Niagara/Private/NiagaraDataInterfaceLandscape.cpp
@@ -208,6 +208,7 @@ using FNDI_Landscape_SharedResourceHandle = FNDI_SharedResourceHandle<FNDI_Lands
 struct FNDILandscapeData_GameThread
 {
 	TWeakObjectPtr<ALandscape> Landscape;
+	TWeakObjectPtr<ULandscapeHeightfieldCollisionComponent> CollisionComponent;
 	FNDI_Landscape_SharedResourceHandle SharedResourceHandle;
 	bool BaseColorVirtualTextureSRGB = false;
 	int32 BaseColorVirtualTextureIndex = INDEX_NONE;
@@ -217,6 +218,9 @@ struct FNDILandscapeData_GameThread
 	bool RequiresCollisionCacheCpu = false;
 	bool RequiresCollisionCacheGpu = false;
 	bool RequiresPhysMatCacheGpu = false;
+	bool SystemRequiresBaseColorGpu = false;
+	bool SystemRequiresHeightsGpu = false;
+	bool SystemRequiresNormalsGpu = false;
 
 	void Reset()
 	{
@@ -230,6 +234,9 @@ struct FNDILandscapeData_GameThread
 		RequiresCollisionCacheCpu = false;
 		RequiresCollisionCacheGpu = false;
 		RequiresPhysMatCacheGpu = false;
+		SystemRequiresBaseColorGpu = false;
+		SystemRequiresHeightsGpu = false;
+		SystemRequiresNormalsGpu = false;
 	}
 };
 
@@ -1370,6 +1377,27 @@ bool UNiagaraDataInterfaceLandscape::InitPerInstanceData(void* PerInstanceData,
 	FNDILandscapeData_GameThread* InstanceData = new(PerInstanceData) FNDILandscapeData_GameThread();
 	ApplyLandscape(*SystemInstance, *InstanceData);
 
+	bool SystemRequiresHeightsCpu = false;
+	bool SystemRequiresHeightsGpu = false;
+	SystemInstance->EvaluateBoundFunction(GetHeightName, SystemRequiresHeightsCpu, SystemRequiresHeightsGpu);
+
+	bool SystemRequiresNormalsCpu = false;
+	bool SystemRequiresNormalsGpu = false;
+	SystemInstance->EvaluateBoundFunction(GetWorldNormalName, SystemRequiresNormalsCpu, SystemRequiresNormalsGpu);
+
+	bool SystemRequiresPhysMatCpu = false;
+	bool SystemRequiresPhysMatGpu = false;
+	SystemInstance->EvaluateBoundFunction(GetPhysicalMaterialIndexName, SystemRequiresPhysMatCpu, SystemRequiresPhysMatGpu);
+
+	bool SystemRequiresBaseColorCpu = false;
+	bool SystemRequiresBaseColorGpu = false;
+	SystemInstance->EvaluateBoundFunction(GetBaseColorName, SystemRequiresBaseColorCpu, SystemRequiresBaseColorGpu);
+
+	InstanceData->SystemRequiresBaseColorGpu = SystemRequiresBaseColorGpu;
+	InstanceData->SystemRequiresHeightsGpu = SystemRequiresHeightsGpu;
+	InstanceData->SystemRequiresNormalsGpu = SystemRequiresNormalsGpu;
+	InstanceData->RequiresPhysMatCacheGpu = SystemRequiresPhysMatGpu;
+
 	FNiagaraDataInterfaceProxyLandscape* RT_Proxy = GetProxyAs<FNiagaraDataInterfaceProxyLandscape>();
 	ENQUEUE_RENDER_COMMAND(FNiagaraDICreateProxy) (
 		[RT_Proxy, InstanceID = SystemInstance->GetId()](FRHICommandListImmediate& CmdList)
@@ -1427,7 +1455,7 @@ bool UNiagaraDataInterfaceLandscape::PerInstanceTick(void* PerInstanceData, FNia
 
 void UNiagaraDataInterfaceLandscape::ApplyLandscape(const FNiagaraSystemInstance& SystemInstance, FNDILandscapeData_GameThread& InstanceData) const
 {
-	ALandscape* Landscape = GetLandscape(SystemInstance, InstanceData.Landscape.Get());
+	ALandscape* Landscape = GetLandscape(SystemInstance, InstanceData);
 
 	// when in editor the contents of the Landscape are volatile and so we'll make sure to
 	// refresh our instance properties any time we apply
@@ -1496,42 +1524,22 @@ void UNiagaraDataInterfaceLandscape::ApplyLandscape(const FNiagaraSystemInstance
 		}
 	}
 
-	bool SystemRequiresHeightsCpu = false;
-	bool SystemRequiresHeightsGpu = false;
-	SystemInstance.EvaluateBoundFunction(GetHeightName, SystemRequiresHeightsCpu, SystemRequiresHeightsGpu);
-
-	bool SystemRequiresNormalsCpu = false;
-	bool SystemRequiresNormalsGpu = false;
-	SystemInstance.EvaluateBoundFunction(GetWorldNormalName, SystemRequiresNormalsCpu, SystemRequiresNormalsGpu);
-
-	bool SystemRequiresPhysMatCpu = false;
-	bool SystemRequiresPhysMatGpu = false;
-	SystemInstance.EvaluateBoundFunction(GetPhysicalMaterialIndexName, SystemRequiresPhysMatCpu, SystemRequiresPhysMatGpu);
-
-	bool SystemRequiresBaseColorCpu = false;
-	bool SystemRequiresBaseColorGpu = false;
-	SystemInstance.EvaluateBoundFunction(GetBaseColorName, SystemRequiresBaseColorCpu, SystemRequiresBaseColorGpu);
-
 	// we need to create our own copy of the collision geometry if either the heights are needed, and they're not
 	// provided by a virtual texture or if the normals are needed and they're not provided by a virtual texture
 	InstanceData.RequiresCollisionCacheGpu =
-		(SystemRequiresBaseColorGpu && InstanceData.BaseColorVirtualTextureIndex == INDEX_NONE) ||
-		(SystemRequiresHeightsGpu && InstanceData.HeightVirtualTextureIndex == INDEX_NONE) ||
-		(SystemRequiresNormalsGpu && InstanceData.NormalVirtualTextureIndex == INDEX_NONE);
-
-	InstanceData.RequiresPhysMatCacheGpu = SystemRequiresPhysMatGpu;
+		(InstanceData.SystemRequiresBaseColorGpu && InstanceData.BaseColorVirtualTextureIndex == INDEX_NONE) ||
+		(InstanceData.SystemRequiresHeightsGpu && InstanceData.HeightVirtualTextureIndex == INDEX_NONE) ||
+		(InstanceData.SystemRequiresNormalsGpu && InstanceData.NormalVirtualTextureIndex == INDEX_NONE);
 }
 
 
 // Users can supply a ALandscape actor
 // if none is provided, then we use the World's LandscapeInfoMap to find an appropriate ALandscape actor
-ALandscape* UNiagaraDataInterfaceLandscape::GetLandscape(const FNiagaraSystemInstance& SystemInstance, ALandscape* Hint) const
+ALandscape* UNiagaraDataInterfaceLandscape::GetLandscape(const FNiagaraSystemInstance& SystemInstance, FNDILandscapeData_GameThread& InstanceData) const
 {
 	if (ALandscape* Landscape = Cast<ALandscape>(SourceLandscape))
 	{
-		const bool HasValidSource = Landscape != nullptr;
-
-		if (SourceMode == ENDILandscape_SourceMode::Source || (SourceMode == ENDILandscape_SourceMode::Default && HasValidSource))
+		if (SourceMode == ENDILandscape_SourceMode::Source || SourceMode == ENDILandscape_SourceMode::Default)
 		{
 			return Landscape;
 		}
@@ -1541,6 +1549,14 @@ ALandscape* UNiagaraDataInterfaceLandscape::GetLandscape(const FNiagaraSystemIns
 
 	auto TestLandscape = [&](const ALandscape* InLandscape)
 	{
+		if (ULandscapeHeightfieldCollisionComponent* CollisionComponent = InstanceData.CollisionComponent.Get())
+		{
+			if (WorldBounds.IntersectXY(CollisionComponent->Bounds.GetBox()))
+			{
+				return true;
+			}
+		}
+
 		if (InLandscape->GetWorld() == SystemInstance.GetWorld())
 		{
 			if (const ULandscapeInfo* LandscapeInfo = InLandscape->GetLandscapeInfo())
@@ -1549,6 +1565,7 @@ ALandscape* UNiagaraDataInterfaceLandscape::GetLandscape(const FNiagaraSystemIns
 				{
 					if (WorldBounds.IntersectXY(ComponentIt.Value->Bounds.GetBox()))
 					{
+						InstanceData.CollisionComponent = ComponentIt.Value;
 						return true;
 					}
 				}
@@ -1558,6 +1575,7 @@ ALandscape* UNiagaraDataInterfaceLandscape::GetLandscape(const FNiagaraSystemIns
 		return false;
 	};
 
+	ALandscape* Hint = InstanceData.Landscape.Get();
 	if (Hint && TestLandscape(Hint))
 	{
 		return Hint;
