- Fix scalability deactivating components that are in flight from the new only update

diff --git a/Engine/Plugins/FX/Niagara/Source/Niagara/Private/NiagaraScalabilityManager.cpp b/Engine/Plugins/FX/Niagara/Source/Niagara/Private/NiagaraScalabilityManager.cpp
--- a/Engine/Plugins/FX/Niagara/Source/Niagara/Private/NiagaraScalabilityManager.cpp
+++ b/Engine/Plugins/FX/Niagara/Source/Niagara/Private/NiagaraScalabilityManager.cpp
@@ -224,6 +224,7 @@ bool FNiagaraScalabilityManager::EvaluateCullState(FNiagaraWorldManager* WorldMa
 			Unregister(Component);
 			return false;
 		}
+		CompState.bNewlyRegisteredDirty = CompState.bNewlyRegistered;
 		CompState.bNewlyRegistered = false;
 
 		const FNiagaraSystemScalabilitySettings& Scalability = System->GetScalabilitySettings();
@@ -336,7 +337,7 @@ void FNiagaraScalabilityManager::ProcessSignificance(FNiagaraWorldManager* World
 	}
 }
 
-bool FNiagaraScalabilityManager::ApplyScalabilityState(int32 ComponentIndex, ENiagaraCullReaction CullReaction)
+bool FNiagaraScalabilityManager::ApplyScalabilityState(int32 ComponentIndex, ENiagaraCullReaction CullReaction, bool bNewOnly)
 {
 	FNiagaraScalabilityState& CompState = State[ComponentIndex];
 
@@ -345,6 +346,12 @@ bool FNiagaraScalabilityManager::ApplyScalabilityState(int32 ComponentIndex, ENi
 		return true;
 	}
 
+	// If we are only processing new only do not update the state for existing component as it will cause a stall on the game thread
+	if (bNewOnly && !CompState.bNewlyRegisteredDirty)
+	{
+		return true;
+	}
+
 	bool bContinueIteration = true;
 	if (UNiagaraComponent* Component = ManagedComponents[ComponentIndex])
 	{
@@ -447,7 +454,7 @@ void FNiagaraScalabilityManager::UpdateInternal(FNiagaraWorldManager* WorldMan,
 			//As we'll be activating and deactivating here, this must be done on the game thread.
 			while (CompIdx < ManagedComponents.Num())
 			{
-				if (ApplyScalabilityState(CompIdx, CullReaction))
+				if (ApplyScalabilityState(CompIdx, CullReaction, Context.bNewOnly))
 				{
 					++CompIdx;
 				}
diff --git a/Engine/Plugins/FX/Niagara/Source/Niagara/Public/NiagaraScalabilityManager.h b/Engine/Plugins/FX/Niagara/Source/Niagara/Public/NiagaraScalabilityManager.h
--- a/Engine/Plugins/FX/Niagara/Source/Niagara/Public/NiagaraScalabilityManager.h
+++ b/Engine/Plugins/FX/Niagara/Source/Niagara/Public/NiagaraScalabilityManager.h
@@ -98,7 +98,7 @@ private:
 	void UpdateInternal(FNiagaraWorldManager* WorldMan, FComponentIterationContext& Context);
 	bool EvaluateCullState(FNiagaraWorldManager* WorldMan, FComponentIterationContext& Context, int32 ComponentIndex, int32& UpdateCounter);
 	void ProcessSignificance(FNiagaraWorldManager* WorldMan, UNiagaraSignificanceHandler* SignificanceHandler, FComponentIterationContext& Context);
-	bool ApplyScalabilityState(int32 ComponentIndex, ENiagaraCullReaction CullReaction);
+	bool ApplyScalabilityState(int32 ComponentIndex, ENiagaraCullReaction CullReaction, bool bNewOnly);
 
 	FNiagaraScalabilitySystemData& GetSystemData(int32 ComponentIndex, bool bForceRefresh=false);
 
diff --git a/Engine/Plugins/FX/Niagara/Source/Niagara/Public/NiagaraScalabilityState.h b/Engine/Plugins/FX/Niagara/Source/Niagara/Public/NiagaraScalabilityState.h
--- a/Engine/Plugins/FX/Niagara/Source/Niagara/Public/NiagaraScalabilityState.h
+++ b/Engine/Plugins/FX/Niagara/Source/Niagara/Public/NiagaraScalabilityState.h
@@ -51,6 +51,9 @@ struct FNiagaraScalabilityState
 	UPROPERTY(VisibleAnywhere, Category = "Scalability")
 	uint8 bNewlyRegistered : 1;
 
+	UPROPERTY(VisibleAnywhere, Category = "Scalability")
+	uint8 bNewlyRegisteredDirty : 1;
+
 	UPROPERTY(VisibleAnywhere, Category = "Scalability")
 	uint8 bCulled : 1;
 
