- Defer scalability invalidation until we run a world manager tick
- This fixes a race between the loading & game thread

diff --git a/Engine/Plugins/FX/Niagara/Source/Niagara/Private/NiagaraSystem.cpp b/Engine/Plugins/FX/Niagara/Source/Niagara/Private/NiagaraSystem.cpp
--- a/Engine/Plugins/FX/Niagara/Source/Niagara/Private/NiagaraSystem.cpp
+++ b/Engine/Plugins/FX/Niagara/Source/Niagara/Private/NiagaraSystem.cpp
@@ -3496,7 +3496,7 @@ void UNiagaraSystem::ResolveScalabilitySettings()
 		ForEachScript(ScriptUsesSigIndex);
 	}
 
-	FNiagaraWorldManager::InvalidateCachedSystemScalabilityDataForAllWorlds();
+	FNiagaraWorldManager::RequestInvalidateCachedSystemScalabilityDataForAllWorlds();
 }
 
 void UNiagaraSystem::UpdateScalability()
diff --git a/Engine/Plugins/FX/Niagara/Source/Niagara/Private/NiagaraWorldManager.cpp b/Engine/Plugins/FX/Niagara/Source/Niagara/Private/NiagaraWorldManager.cpp
--- a/Engine/Plugins/FX/Niagara/Source/Niagara/Private/NiagaraWorldManager.cpp
+++ b/Engine/Plugins/FX/Niagara/Source/Niagara/Private/NiagaraWorldManager.cpp
@@ -242,6 +242,8 @@ FDelegateHandle FNiagaraWorldManager::PostReachabilityAnalysisHandle;
 FDelegateHandle FNiagaraWorldManager::PostGCHandle;
 FDelegateHandle FNiagaraWorldManager::PreGCBeginDestroyHandle;
 FDelegateHandle FNiagaraWorldManager::ViewTargetChangedHandle;
+
+bool FNiagaraWorldManager::bInvalidateCachedSystemScalabilityData = false;
 TMap<class UWorld*, class FNiagaraWorldManager*> FNiagaraWorldManager::WorldManagers;
 
 namespace FNiagaraUtilities
@@ -1166,6 +1168,12 @@ void FNiagaraWorldManager::Tick(ETickingGroup TickGroup, float DeltaSeconds, ELe
 			WorldLoopTime -= DeltaSeconds;
 		}
 
+		if (bInvalidateCachedSystemScalabilityData)
+		{
+			bInvalidateCachedSystemScalabilityData = false;
+			InvalidateCachedSystemScalabilityDataForAllWorlds();
+		}
+
 		//Ensure the pools have been primed.
 		//WorldInit is too early.
 		if(!bPoolIsPrimed)
@@ -1846,6 +1854,11 @@ bool FNiagaraWorldManager::GetScalabilityState(UNiagaraComponent* Component, FNi
 	return false;
 }
 
+void FNiagaraWorldManager::RequestInvalidateCachedSystemScalabilityDataForAllWorlds()
+{
+	bInvalidateCachedSystemScalabilityData = true;
+}
+
 void FNiagaraWorldManager::InvalidateCachedSystemScalabilityDataForAllWorlds()
 {
 	for (auto& Pair : WorldManagers)
diff --git a/Engine/Plugins/FX/Niagara/Source/Niagara/Public/NiagaraWorldManager.h b/Engine/Plugins/FX/Niagara/Source/Niagara/Public/NiagaraWorldManager.h
--- a/Engine/Plugins/FX/Niagara/Source/Niagara/Public/NiagaraWorldManager.h
+++ b/Engine/Plugins/FX/Niagara/Source/Niagara/Public/NiagaraWorldManager.h
@@ -198,9 +198,12 @@ public:
 	void PrimePoolForAllSystems();
 	void PrimePool(UNiagaraSystem* System);
 
+	static void RequestInvalidateCachedSystemScalabilityDataForAllWorlds();
+private:
 	static void InvalidateCachedSystemScalabilityDataForAllWorlds();
 	void InvalidateCachedSystemScalabilityData();
 
+public:
 	void SetDebugPlaybackMode(ENiagaraDebugPlaybackMode Mode) { RequestedDebugPlaybackMode = Mode; }
 	ENiagaraDebugPlaybackMode GetDebugPlaybackMode() const { return DebugPlaybackMode; }
 
@@ -293,6 +296,8 @@ private:
 	static FDelegateHandle PreGCBeginDestroyHandle;
 	static FDelegateHandle ViewTargetChangedHandle;
 
+	static bool bInvalidateCachedSystemScalabilityData;
+
 	static NIAGARA_API TMap<class UWorld*, class FNiagaraWorldManager*> WorldManagers;
 
 	UWorld* World = nullptr;
