Set play audio sounds to play sound from owner instead of world
- only for fire and forget sounds

#test PIE, used sound audio param to have audio update based on input from blueprint.

diff --git a/Engine/Plugins/FX/Niagara/Source/Niagara/Private/NiagaraDataInterfaceAudioPlayer.cpp b/Engine/Plugins/FX/Niagara/Source/Niagara/Private/NiagaraDataInterfaceAudioPlayer.cpp
--- a/Engine/Plugins/FX/Niagara/Source/Niagara/Private/NiagaraDataInterfaceAudioPlayer.cpp
+++ b/Engine/Plugins/FX/Niagara/Source/Niagara/Private/NiagaraDataInterfaceAudioPlayer.cpp
@@ -195,10 +195,18 @@ bool UNiagaraDataInterfaceAudioPlayer::PerInstanceTickPostSimulate(void* PerInst
 		// play the fire-and-forget sounds
 		if (World && PIData->SoundToPlay.IsValid())
 		{
+			AActor* SoundOwner = nullptr;
+			if (const USceneComponent* AttachComp = SystemInstance->GetAttachComponent())
+			{
+				if (const USceneComponent* AttachParent = AttachComp->GetAttachParent())
+				{
+					SoundOwner = AttachParent->GetOwner();
+				}
+			}
 			for (const FAudioParticleData& ParticleData : Data)
 			{
 				UGameplayStatics::PlaySoundAtLocation(World, PIData->SoundToPlay.Get(), ParticleData.Position, ParticleData.Rotation, ParticleData.Volume,
-					ParticleData.Pitch, ParticleData.StartTime, PIData->Attenuation.Get(), PIData->Concurrency.Get());
+					ParticleData.Pitch, ParticleData.StartTime, PIData->Attenuation.Get(), PIData->Concurrency.Get(), SoundOwner);
 			}
 		}
 	}
