- Change bInvalidateCachedSystemScalabilityData into an atomic bool for safety

#rnx

diff --git a/Engine/Plugins/FX/Niagara/Source/Niagara/Private/NiagaraWorldManager.cpp b/Engine/Plugins/FX/Niagara/Source/Niagara/Private/NiagaraWorldManager.cpp
--- a/Engine/Plugins/FX/Niagara/Source/Niagara/Private/NiagaraWorldManager.cpp
+++ b/Engine/Plugins/FX/Niagara/Source/Niagara/Private/NiagaraWorldManager.cpp
@@ -243,7 +243,7 @@ FDelegateHandle FNiagaraWorldManager::PostGCHandle;
 FDelegateHandle FNiagaraWorldManager::PreGCBeginDestroyHandle;
 FDelegateHandle FNiagaraWorldManager::ViewTargetChangedHandle;
 
-bool FNiagaraWorldManager::bInvalidateCachedSystemScalabilityData = false;
+std::atomic<bool> FNiagaraWorldManager::bInvalidateCachedSystemScalabilityData(false);
 TMap<class UWorld*, class FNiagaraWorldManager*> FNiagaraWorldManager::WorldManagers;
 
 namespace FNiagaraUtilities
@@ -1168,9 +1168,8 @@ void FNiagaraWorldManager::Tick(ETickingGroup TickGroup, float DeltaSeconds, ELe
 			WorldLoopTime -= DeltaSeconds;
 		}
 
-		if (bInvalidateCachedSystemScalabilityData)
+		if (bInvalidateCachedSystemScalabilityData.exchange(false))
 		{
-			bInvalidateCachedSystemScalabilityData = false;
 			InvalidateCachedSystemScalabilityDataForAllWorlds();
 		}
 
diff --git a/Engine/Plugins/FX/Niagara/Source/Niagara/Public/NiagaraWorldManager.h b/Engine/Plugins/FX/Niagara/Source/Niagara/Public/NiagaraWorldManager.h
--- a/Engine/Plugins/FX/Niagara/Source/Niagara/Public/NiagaraWorldManager.h
+++ b/Engine/Plugins/FX/Niagara/Source/Niagara/Public/NiagaraWorldManager.h
@@ -296,7 +296,7 @@ private:
 	static FDelegateHandle PreGCBeginDestroyHandle;
 	static FDelegateHandle ViewTargetChangedHandle;
 
-	static bool bInvalidateCachedSystemScalabilityData;
+	static std::atomic<bool> bInvalidateCachedSystemScalabilityData;
 
 	static NIAGARA_API TMap<class UWorld*, class FNiagaraWorldManager*> WorldManagers;
 
