- FIx uninitialized variable

#rnx

diff --git a/Engine/Plugins/FX/Niagara/Source/Niagara/Public/NiagaraScalabilityState.h b/Engine/Plugins/FX/Niagara/Source/Niagara/Public/NiagaraScalabilityState.h
--- a/Engine/Plugins/FX/Niagara/Source/Niagara/Public/NiagaraScalabilityState.h
+++ b/Engine/Plugins/FX/Niagara/Source/Niagara/Public/NiagaraScalabilityState.h
@@ -28,6 +28,7 @@ struct FNiagaraScalabilityState
 		, LastVisibleTime(0.0f)
 		, SystemDataIndex(INDEX_NONE)
 		, bNewlyRegistered(1)
+		, bNewlyRegisteredDirty(0)
 		, bCulled(InCulled)
 		, bPreviousCulled(InPreviousCulled)
 		, bCulledByDistance(0)
