Added a PostLoad fixup for Assignment nodes to handle LWC types (Particles.Position).  Prior to this the graphs could have been corrected only during the NiagaraGraph PostLoad when the graph was being duplicated for compilation.

diff --git a/Engine/Plugins/FX/Niagara/Source/NiagaraEditor/Private/NiagaraNodeAssignment.cpp b/Engine/Plugins/FX/Niagara/Source/NiagaraEditor/Private/NiagaraNodeAssignment.cpp
--- a/Engine/Plugins/FX/Niagara/Source/NiagaraEditor/Private/NiagaraNodeAssignment.cpp
+++ b/Engine/Plugins/FX/Niagara/Source/NiagaraEditor/Private/NiagaraNodeAssignment.cpp
@@ -94,6 +94,8 @@ void UNiagaraNodeAssignment::PostLoad()
 {
 	Super::PostLoad();
 
+	bool NeedsScriptRefresh = false;
+
 	// Handle the case where we moved towards an array of assignment targets...
 	if (AssignmentTarget_DEPRECATED.IsValid() && AssignmentTargets.Num() == 0)
 	{
@@ -220,8 +222,7 @@ void UNiagaraNodeAssignment::PostLoad()
 		const int32 NiagaraVer = GetLinkerCustomVersion(FNiagaraCustomVersion::GUID);
 		if (NiagaraVer < FNiagaraCustomVersion::AssignmentNodeUsesBeginDefaults)
 		{
-			FunctionScript = nullptr;
-			GenerateScript();
+			NeedsScriptRefresh = true;
 		}
 		if (NiagaraVer < FNiagaraCustomVersion::AssignmentNodeHasCorrectUsageBitmask)
 		{
@@ -241,12 +242,77 @@ void UNiagaraNodeAssignment::PostLoad()
 			FName NewName = UNiagaraGraph::StandardizeName(CurrentName, ENiagaraScriptUsage::Module, false, true);
 			AssignmentTarget.SetName(NewName);
 		}
-		RefreshFromExternalChanges();
+		NeedsScriptRefresh = true;
+	}
+
+	// check for outdated LWC types
+	if (PostLoad_LWCFixup(NiagaraVer))
+	{
+		NeedsScriptRefresh = true;
+	}
+
+	if (NeedsScriptRefresh)
+	{
+		FunctionScript = nullptr;
+		GenerateScript();
 	}
 
 	RefreshTitle();
 }
 
+bool UNiagaraNodeAssignment::PostLoad_LWCFixup(int32 NiagaraVersion)
+{
+	bool NeedsScriptRefresh = false;
+
+	const TArray<FNiagaraVariable>& OldTypes = FNiagaraConstants::GetOldPositionTypeVariables();
+	const FEdGraphPinType PositionPinType = UEdGraphSchema_Niagara::TypeDefinitionToPinType(FNiagaraTypeDefinition::GetPositionDef());
+
+	for (FNiagaraVariable& AssignmentTarget : AssignmentTargets)
+	{
+		if (OldTypes.Contains(AssignmentTarget))
+		{
+			AssignmentTarget.SetType(FNiagaraTypeDefinition::GetPositionDef());
+
+			FNameBuilder ExternalVariableNameBuilder;
+			ExternalVariableNameBuilder.Append(FunctionDisplayName);
+			ExternalVariableNameBuilder.Append(TEXT("."));
+			AssignmentTarget.GetName().AppendString(ExternalVariableNameBuilder);
+
+			FName ExternalVariableName(ExternalVariableNameBuilder);
+
+			// additionally we need to look at the input to the assignment node, if it's a set node
+			// then we'll need to change the type of the pin there too so that when we run a compile
+			// we'll be able to connect that to the internal input node that we've created in our graph
+			if (UEdGraphPin* InputPin = GetInputPin(0))
+			{
+				if (UEdGraphPin* ConnectedPin = InputPin->LinkedTo.IsEmpty() ? nullptr : InputPin->LinkedTo[0])
+				{
+					if (UNiagaraNodeParameterMapSet* SetNode = Cast<UNiagaraNodeParameterMapSet>(ConnectedPin->GetOwningNode()))
+					{
+						SetNode->ConditionalPostLoad();
+
+						// find any appropriate input pins that will need to have their types changed
+						FPinCollectorArray SetNodeInputPins;
+						SetNode->GetInputPins(SetNodeInputPins);
+						for (UEdGraphPin* SetNodeInputPin : SetNodeInputPins)
+						{
+							FNiagaraVariable SetNodeInputVariable = UEdGraphSchema_Niagara::PinToNiagaraVariable(SetNodeInputPin);
+							if (SetNodeInputVariable.IsValid() && SetNodeInputVariable.GetName() == ExternalVariableName)
+							{
+								SetNodeInputPin->PinType = PositionPinType;
+							}
+						}
+					}
+				}
+			}
+
+			NeedsScriptRefresh = true;
+		}
+	}
+
+	return NeedsScriptRefresh;
+}
+
 #if WITH_EDITORONLY_DATA
 void UNiagaraNodeAssignment::DeclareConstructClasses(TArray<FTopLevelAssetPath>& OutConstructClasses, const UClass* SpecificSubclass)
 {
diff --git a/Engine/Plugins/FX/Niagara/Source/NiagaraEditor/Public/NiagaraNodeAssignment.h b/Engine/Plugins/FX/Niagara/Source/NiagaraEditor/Public/NiagaraNodeAssignment.h
--- a/Engine/Plugins/FX/Niagara/Source/NiagaraEditor/Public/NiagaraNodeAssignment.h
+++ b/Engine/Plugins/FX/Niagara/Source/NiagaraEditor/Public/NiagaraNodeAssignment.h
@@ -90,5 +90,7 @@ private:
 	int32 CalculateScriptUsageBitmask();
 
 	void RefreshTitle();
+
+	bool PostLoad_LWCFixup(int32 NiagaraVersion);
 };
 
