Adds two new compilation modes for Niagara systems - AsyncTasks & Verify.

AsyncTasks or Digest mode is where we digest the graphs (and other objects) that are required for compilation into FStructs, which we piece together into immutable data from which we can go wide with the precompilation, instantiation and translation stages of the compilation of niagara scripts.

This change also introduces the NiagaraSystemCompilingManager which allows us to ensure that compiling Niagara scripts will play well with the rest of the engine (most noticeably cleaning up the cook pattern where we can now effectively asynchronously do work off the game thread).

Work through the translator and parameter map history builders is done through the GraphBridge interface implemented earlier to keep a unique set of code.  This is not true for node implementations of Compile() & BUildParameterMapHistory() which are forked.

The Verifiy compilation mode will run both methods of compilation and will compare the resulting ExeData.

A number of outstanding issues remain before we can make the AsyncTask the default:
-generating good comparible metrics for the different stages of the compilation process
-databasae of digested graphs/NPC is currently unbounded for the editor session
-aborting AsyncTasks isn't fully implemented
-there are some false positive differences between the two compliation modes (to my knowledge no current oustanding behavior differences exist)

diff --git a/Engine/Plugins/FX/Niagara/Source/Niagara/Classes/NiagaraScript.h b/Engine/Plugins/FX/Niagara/Source/Niagara/Classes/NiagaraScript.h
--- a/Engine/Plugins/FX/Niagara/Source/Niagara/Classes/NiagaraScript.h
+++ b/Engine/Plugins/FX/Niagara/Source/Niagara/Classes/NiagaraScript.h
@@ -941,7 +941,7 @@ public:
 	NIAGARA_API static const FName NiagaraCustomVersionTagName;
 #endif
 
-	NIAGARA_API void ComputeVMCompilationId(FNiagaraVMExecutableDataId& Id, FGuid VersionGuid) const;
+	NIAGARA_API void ComputeVMCompilationId(FNiagaraVMExecutableDataId& Id, const FGuid& VersionGuid) const;
 	NIAGARA_API const FNiagaraVMExecutableDataId& GetComputedVMCompilationId() const
 	{
 #if WITH_EDITORONLY_DATA
@@ -1194,7 +1194,11 @@ public:
 	NIAGARA_API FOnPropertyChanged& OnPropertyChanged();
 
 	/** External call used to identify the values for a successful VM script compilation. OnVMScriptCompiled will be issued in this case.*/
-	void SetVMCompilationResults(const FNiagaraVMExecutableDataId& InCompileId, FNiagaraVMExecutableData& InScriptVM, FString EmitterUniqueName, const TMap<FName, UNiagaraDataInterface*>& ObjectNameMap, bool ApplyRapidIterationParameters);
+	void SetVMCompilationResults(const FNiagaraVMExecutableDataId& InCompileId, FNiagaraVMExecutableData& InScriptVM, const FString& EmitterUniqueName, const TMap<FName, UNiagaraDataInterface*>& ObjectNameMap, bool ApplyRapidIterationParameters);
+
+	/** Updates the RI parameter store based on the provided variables (missing entries will be
+	    added and stale entries will be removed.  Returns true if the parameter store was modified. */
+	bool ApplyRapidIterationParameters(TConstArrayView<FNiagaraVariable> InParameters, bool bAllowRemoval);
 
 	/** In the event where we "merge" we duplicate the changes of the source script onto the newly cloned copy. This function will synchronize the compiled script 
 		results assuming that the scripts themselves are bound to the same key. This saves looking things up in the DDC. It returns true if successfully synchronized and 
diff --git a/Engine/Plugins/FX/Niagara/Source/Niagara/Private/NiagaraAsyncCompile.cpp b/Engine/Plugins/FX/Niagara/Source/Niagara/Private/NiagaraAsyncCompile.cpp
--- a/Engine/Plugins/FX/Niagara/Source/Niagara/Private/NiagaraAsyncCompile.cpp
+++ b/Engine/Plugins/FX/Niagara/Source/Niagara/Private/NiagaraAsyncCompile.cpp
@@ -215,6 +215,7 @@ void FNiagaraAsyncCompileTask::MoveToState(ENiagaraCompilationState NewState)
 	}
 	if (NewState == ENiagaraCompilationState::Finished)
 	{
+		CompileMetrics.TaskWallTime = (float) (FPlatformTime::Seconds() - StartTaskTime);
 		check(CurrentState == ENiagaraCompilationState::CheckDDC || CurrentState == ENiagaraCompilationState::PutToDDC);
 	}
 
@@ -322,7 +323,7 @@ bool FNiagaraAsyncCompileTask::AwaitResult()
 
 	if (!ExeData.IsValid())
 	{
-		ExeData = NiagaraModule.GetCompileJobResult(ScriptPair.PendingJobID, bWaitForCompileJob);
+		ExeData = NiagaraModule.GetCompileJobResult(ScriptPair.PendingJobID, bWaitForCompileJob, CompileMetrics);
 	}
 
 	return ExeData.IsValid();
@@ -342,6 +343,8 @@ void FNiagaraAsyncCompileTask::OptimizeByteCode()
 #if VECTORVM_SUPPORTS_EXPERIMENTAL
 	if (const uint8* ByteCode = ExeData->ByteCode.GetDataPtr())
 	{
+		const double OptimizeStartTime = FPlatformTime::Seconds();
+
 		//this is just necessary because VectorVM doesn't know about FVMExternalFunctionBindingInfo
 		const int32 NumExtFns = ExeData->CalledVMExternalFunctions.Num();
 
@@ -361,6 +364,8 @@ void FNiagaraAsyncCompileTask::OptimizeByteCode()
 		FreezeVectorVMOptimizeContext(OptimizeContext, ExeData->ExperimentalContextData);
 
 		FreeVectorVMOptimizeContext(&OptimizeContext);
+
+		CompileMetrics.ByteCodeOptimizeTime = (float) (FPlatformTime::Seconds() - OptimizeStartTime);
 	}
 #endif // VECTORVM_SUPPORTS_EXPERIMENTAL
 }
@@ -427,8 +432,8 @@ void FNiagaraAsyncCompileTask::CheckDDCResult()
 			MoveToState(ENiagaraCompilationState::Precompile);
 			UE_LOG(LogNiagara, Verbose, TEXT("No compilation data for %s found in the ddc."), *AssetPath);
 		}
-		DDCFetchTime = FPlatformTime::Seconds() - StartTaskTime;
-		UE_LOG(LogNiagara, Verbose, TEXT("Compile task for %s took %f seconds to fetch data from ddc."), *AssetPath, DDCFetchTime);
+		CompileMetrics.DDCFetchTime = (float) (FPlatformTime::Seconds() - StartTaskTime);
+		UE_LOG(LogNiagara, VeryVerbose, TEXT("Compile task for %s took %f seconds to fetch data from ddc."), *AssetPath, CompileMetrics.DDCFetchTime);
 	}
 	else if ((FPlatformTime::Seconds() - StartTaskTime) > GNiagaraCompileDDCWaitTimeout)
 	{
@@ -496,7 +501,7 @@ void FNiagaraActiveCompilationDefault::Abort()
 	}
 }
 
-bool FNiagaraActiveCompilationDefault::Validate(const FNiagaraQueryCompilationOptions& Options) const
+bool FNiagaraActiveCompilationDefault::ValidateConsistentResults(const FNiagaraQueryCompilationOptions& Options) const
 {
 	// for now the only thing we're concerned about is if we've got results for SystemSpawn and SystemUpdate scripts
 	// then we need to make sure that they agree in terms of the dataset attributes
@@ -987,8 +992,31 @@ void FNiagaraActiveCompilationDefault::ReportResults(const FNiagaraQueryCompilat
 	{
 		UE_LOG(LogNiagara, Verbose, TEXT("Retrieving %s from DDC took %f sec."), *Options.System->GetFullName(), ElapsedWallTime);
 	}
+
+	if (Options.bGenerateTimingsFile)
+	{
+		FNiagaraSystemCompileMetrics SystemMetrics;
+		SystemMetrics.SystemCompileWallTime = ElapsedWallTime;
+
+		for (const FAsyncTaskPtr& AsyncTask : Tasks)
+		{
+			if (AsyncTask->ScriptPair.bResultsReady)
+			{
+				SystemMetrics.ScriptMetrics.Add(AsyncTask->ScriptPair.CompiledScript, AsyncTask->CompileMetrics);
+			}
+		}
+
+		WriteTimingsEntry(TEXT("Default Compilation"), Options, SystemMetrics);
+	}
 }
 
+bool FNiagaraActiveCompilationDefault::BlocksBeginCacheForCooked() const
+{
+	// With the current setup we will timeout DDC fill jobs while polling for the work to be complete.
+	// Not entirely sure why, but my guess is that we get too many systems in the process of being
+	// compiled and we iteratively block the gamethread advancing each of them starving progress
+	return true;
+}
 
 #endif // WITH_EDITORONLY_DATA
 
diff --git a/Engine/Plugins/FX/Niagara/Source/Niagara/Private/NiagaraAsyncCompile.h b/Engine/Plugins/FX/Niagara/Source/Niagara/Private/NiagaraAsyncCompile.h
--- a/Engine/Plugins/FX/Niagara/Source/Niagara/Private/NiagaraAsyncCompile.h
+++ b/Engine/Plugins/FX/Niagara/Source/Niagara/Private/NiagaraAsyncCompile.h
@@ -65,7 +65,6 @@ public:
 	uint32 TaskHandle = 0;
 
 	double StartTaskTime = 0;
-	double DDCFetchTime = 0;
 	double StartCompileTime = 0;
 	bool bWaitForCompileJob = false;
 	bool bUsedShaderCompilerWorker = false;
@@ -79,6 +78,7 @@ public:
 	FEmitterCompiledScriptPair ScriptPair;
 	TArray<FNiagaraVariable> EncounteredExposedVars;
 	TArray<FNiagaraVariable> BakedRapidIterationParameters;
+	FNiagaraScriptCompileMetrics CompileMetrics;
 
 	ENiagaraCompilationState CurrentState;
 
@@ -126,9 +126,10 @@ public:
 	virtual bool Launch(const FNiagaraCompilationOptions& Options) override;
 	virtual void Abort() override;
 	virtual bool QueryCompileComplete(const FNiagaraQueryCompilationOptions& Options) override;
-	virtual bool Validate(const FNiagaraQueryCompilationOptions& Options) const override;
+	virtual bool ValidateConsistentResults(const FNiagaraQueryCompilationOptions& Options) const override;
 	virtual void Apply(const FNiagaraQueryCompilationOptions& Options) override;
 	virtual void ReportResults(const FNiagaraQueryCompilationOptions& Options) const override;
+	virtual bool BlocksBeginCacheForCooked() const override;
 
 	double StartTime = 0.0;
 
diff --git a/Engine/Plugins/FX/Niagara/Source/Niagara/Private/NiagaraCompilationTypes.cpp b/Engine/Plugins/FX/Niagara/Source/Niagara/Private/NiagaraCompilationTypes.cpp
--- a/Engine/Plugins/FX/Niagara/Source/Niagara/Private/NiagaraCompilationTypes.cpp
+++ b/Engine/Plugins/FX/Niagara/Source/Niagara/Private/NiagaraCompilationTypes.cpp
@@ -2,20 +2,857 @@
 
 #include "NiagaraCompilationTypes.h"
 
+#include "Algo/RemoveIf.h"
+#include "Internationalization/Regex.h"
+#include "Modules/ModuleManager.h"
 #include "NiagaraAsyncCompile.h"
+#include "NiagaraModule.h"
+#include "NiagaraScriptSourceBase.h"
+#include "NiagaraSystem.h"
+
+#define LOCTEXT_NAMESPACE "NiagaraSystem"
 
 #if WITH_EDITORONLY_DATA
 
-TUniquePtr<FNiagaraActiveCompilation> FNiagaraActiveCompilation::CreateCompilation()
+enum class ENiagaraCompilationMode : int32
 {
-	return MakeUnique<FNiagaraActiveCompilationDefault>();
+	Original,
+	AsyncTasks,
+	Verify
+};
+
+enum class ENiagaraCompilationValidateMode : int32
+{
+	Ignore,
+	Warn,
+	Ensure,
+	Assert
+};
+
+static int32 GNiagaraSystemCompileMode = (int32) ENiagaraCompilationMode::Original;
+static FAutoConsoleVariableRef CVarNiagaraSystemCompileMode(
+	TEXT("fx.Niagara.SystemCompileMode"),
+	GNiagaraSystemCompileMode,
+	TEXT("Defines how NiagaraSystem will be compiled"),
+	ECVF_Default
+);
+
+static bool GNiagaraCompileDumpTimings = false;
+static FAutoConsoleVariableRef CVarNiagaraCompileDumpTimings(
+	TEXT("fx.Niagara.CompileDumpTimings"),
+	GNiagaraCompileDumpTimings,
+	TEXT("If enabled a file containing compile metrics for the different compiled scripts will be dumped to the log folder"),
+	ECVF_ReadOnly
+);
+
+static int32 GNiagaraCompileValidateMode = (int32) ENiagaraCompilationValidateMode::Warn;
+static FAutoConsoleVariableRef CVarGNiagaraCompileValidateMode(
+	TEXT("fx.Niagara.CompileValidateMode"),
+	GNiagaraCompileValidateMode,
+	TEXT("Controls how the validate compile mode will report differences it encounters when comparing default/async compiles"),
+	ECVF_Default
+);
+
+// Collection of operators to help validate data between two ExeData (see CompareExeData below)
+bool operator==(const FNiagaraVMExecutableByteCode& Lhs, const FNiagaraVMExecutableByteCode& Rhs)
+{
+	return Lhs.GetData() == Rhs.GetData();
 }
 
-#else // WITH_EDITORONLY_DATA
+bool operator==(const FNiagaraParameters& Lhs, const FNiagaraParameters& Rhs)
+{
+	return Lhs.Parameters == Rhs.Parameters;
+}
+
+bool operator==(const FNiagaraScriptDataUsageInfo& Lhs, const FNiagaraScriptDataUsageInfo& Rhs)
+{
+	return Lhs.bReadsAttributeData == Rhs.bReadsAttributeData;
+}
+
+
+bool operator==(const FShaderParametersMetadata& Lhs, const FShaderParametersMetadata& Rhs)
+{
+	return true;
+}
+bool operator==(const FNiagaraDataInterfaceGPUParamInfo& Lhs, const FNiagaraDataInterfaceGPUParamInfo& Rhs)
+{
+	return Lhs.DataInterfaceHLSLSymbol == Rhs.DataInterfaceHLSLSymbol
+		&& Lhs.DIClassName == Rhs.DIClassName
+		&& Lhs.ShaderParametersOffset == Rhs.ShaderParametersOffset
+		&& Lhs.GeneratedFunctions == Rhs.GeneratedFunctions;
+}
+
+bool operator==(const FVMFunctionSpecifier& Lhs, const FVMFunctionSpecifier& Rhs)
+{
+	return Lhs.Key == Rhs.Key
+		&& Lhs.Value == Rhs.Value;
+}
+
+bool operator==(const FVMExternalFunctionBindingInfo& Lhs, const FVMExternalFunctionBindingInfo& Rhs)
+{
+	return Lhs.Name == Rhs.Name
+		&& Lhs.OwnerName == Rhs.OwnerName
+		&& Lhs.InputParamLocations == Rhs.InputParamLocations
+		&& Lhs.NumOutputs == Rhs.NumOutputs
+		&& Lhs.FunctionSpecifiers == Rhs.FunctionSpecifiers;
+}
+
+bool operator==(const FNiagaraScriptDataInterfaceCompileInfo& Lhs, const FNiagaraScriptDataInterfaceCompileInfo& Rhs)
+{
+	return Lhs.Name == Rhs.Name
+		&& Lhs.UserPtrIdx == Rhs.UserPtrIdx
+		&& Lhs.Type == Rhs.Type
+		&& Lhs.RegisteredFunctions == Rhs.RegisteredFunctions
+		&& Lhs.RegisteredParameterMapRead == Rhs.RegisteredParameterMapRead
+		&& Lhs.RegisteredParameterMapWrite == Rhs.RegisteredParameterMapWrite
+		&& Lhs.bIsPlaceholder == Rhs.bIsPlaceholder;
+}
+
+bool operator==(const FNiagaraShaderScriptParametersMetadata& Lhs, const FNiagaraShaderScriptParametersMetadata& Rhs)
+{
+	return Lhs.DataInterfaceParamInfo == Rhs.DataInterfaceParamInfo
+		&& Lhs.LooseMetadataNames == Rhs.LooseMetadataNames
+		&& Lhs.bExternalConstantsInterpolated == Rhs.bExternalConstantsInterpolated
+		&& Lhs.ExternalConstants == Rhs.ExternalConstants
+		&& Lhs.StructIncludeInfos == Rhs.StructIncludeInfos
+		&& Lhs.ShaderParametersMetadata == Rhs.ShaderParametersMetadata;
+}
+
+bool operator==(const FNiagaraCompileEvent& Lhs, const FNiagaraCompileEvent& Rhs)
+{
+	return Lhs.Severity ==  Rhs.Severity
+		&& Lhs.Message == Rhs.Message
+		&& Lhs.ShortDescription == Rhs.ShortDescription
+		&& Lhs.NodeGuid == Rhs.NodeGuid
+		&& Lhs.PinGuid == Rhs.PinGuid
+		&& Lhs.StackGuids == Rhs.StackGuids
+		&& Lhs.Source == Rhs.Source;
+}
+
+bool operator==(const FSimulationStageMetaData& Lhs, const FSimulationStageMetaData& Rhs)
+{
+	return Lhs.SimulationStageName == Rhs.SimulationStageName
+		&& Lhs.EnabledBinding == Rhs.EnabledBinding
+		&& Lhs.ElementCount == Rhs.ElementCount
+		&& Lhs.ElementCountXBinding == Rhs.ElementCountXBinding
+		&& Lhs.ElementCountYBinding == Rhs.ElementCountYBinding
+		&& Lhs.ElementCountZBinding == Rhs.ElementCountZBinding
+		&& Lhs.IterationSourceType == Rhs.IterationSourceType
+		&& Lhs.IterationDataInterface == Rhs.IterationDataInterface
+		&& Lhs.IterationDirectBinding == Rhs.IterationDirectBinding
+		&& Lhs.ExecuteBehavior == Rhs.ExecuteBehavior
+		&& Lhs.bWritesParticles == Rhs.bWritesParticles
+		&& Lhs.bPartialParticleUpdate == Rhs.bPartialParticleUpdate
+		&& Lhs.bParticleIterationStateEnabled == Rhs.bParticleIterationStateEnabled
+		&& Lhs.bGpuIndirectDispatch == Rhs.bGpuIndirectDispatch
+		&& Lhs.ParticleIterationStateBinding == Rhs.ParticleIterationStateBinding
+		&& Lhs.ParticleIterationStateComponentIndex == Rhs.ParticleIterationStateComponentIndex
+		&& Lhs.ParticleIterationStateRange == Rhs.ParticleIterationStateRange
+		&& Lhs.OutputDestinations == Rhs.OutputDestinations
+		&& Lhs.InputDataInterfaces == Rhs.InputDataInterfaces
+		&& Lhs.NumIterations == Rhs.NumIterations
+		&& Lhs.NumIterationsBinding == Rhs.NumIterationsBinding
+		&& Lhs.GpuDispatchType == Rhs.GpuDispatchType
+		&& Lhs.GpuDirectDispatchElementType == Rhs.GpuDirectDispatchElementType
+		&& Lhs.GpuDispatchNumThreads == Rhs.GpuDispatchNumThreads;
+}
+
+bool operator==(const FNiagaraDataSetProperties& Lhs, const FNiagaraDataSetProperties& Rhs)
+{
+	return Lhs.ID == Rhs.ID
+		&& Lhs.Variables == Rhs.Variables;
+}
+
+bool operator==(const FVectorVMExternalFunctionContextProxy& Lhs, const FVectorVMExternalFunctionContextProxy& Rhs)
+{
+	return true;
+}
+
+bool operator==(const FNiagaraCompilerTag& Lhs, const FNiagaraCompilerTag& Rhs)
+{
+	return Lhs.Variable == Rhs.Variable
+		&& Lhs.StringValue == Rhs.StringValue;
+}
+
+bool operator==(const FNiagaraDataInterfaceStructIncludeInfo& Lhs, const FNiagaraDataInterfaceStructIncludeInfo& Rhs)
+{
+	return Lhs.StructMetadata == Rhs.StructMetadata
+		&& Lhs.ParamterOffset == Rhs.ParamterOffset;
+}
+
+bool operator==(const FNiagaraShaderScriptExternalConstant& Lhs, const FNiagaraShaderScriptExternalConstant& Rhs)
+{
+	return Lhs.Type == Rhs.Type
+		&& Lhs.Name == Rhs.Name;
+}
+
+class FNiagaraActiveCompilationAsyncTask : public FNiagaraActiveCompilation
+{
+public:
+	virtual void AddReferencedObjects(FReferenceCollector& Collector) override
+	{
+	}
+	virtual FString GetReferencerName() const override
+	{
+		return TEXT("FNiagaraActiveCompilationAsyncTask");
+	}
+
+	virtual bool Launch(const FNiagaraCompilationOptions& Options) override
+	{
+		TaskStartTime = FPlatformTime::Seconds();
+		bForced = Options.bForced;
+
+		FNiagaraSystemAsyncCompileResults NewCompileRequest;
+		INiagaraModule& NiagaraModule = FModuleManager::GetModuleChecked<INiagaraModule>("Niagara");
+
+		CompileRequestHandle = NiagaraModule.RequestCompileSystem(Options.System, bForced);
+
+		return CompileRequestHandle != INDEX_NONE;
+	}
+
+	virtual void Abort() override
+	{
+		INiagaraModule& NiagaraModule = FModuleManager::Get().LoadModuleChecked<INiagaraModule>(TEXT("Niagara"));
+		NiagaraModule.AbortSystemCompile(CompileRequestHandle);
+
+		bShouldApply = false;
+	}
+
+	virtual bool ValidateConsistentResults(const FNiagaraQueryCompilationOptions& Options) const override
+	{
+		const UNiagaraScript* SystemSpawnScript = Options.System->GetSystemSpawnScript();
+		const UNiagaraScript* SystemUpdateScript = Options.System->GetSystemUpdateScript();
+
+		// for now the only thing we're concerned about is if we've got results for SystemSpawn and SystemUpdate scripts
+		// then we need to make sure that they agree in terms of the dataset attributes
+		const FNiagaraScriptAsyncCompileData* SpawnScriptRequest = CompileResults.CompileResultMap.Find(SystemSpawnScript);
+		const FNiagaraScriptAsyncCompileData* UpdateScriptRequest = CompileResults.CompileResultMap.Find(SystemUpdateScript);
+
+		const bool SpawnScriptValid = SpawnScriptRequest
+			&& SpawnScriptRequest->ExeData.IsValid()
+			&& SpawnScriptRequest->ExeData->LastCompileStatus != ENiagaraScriptCompileStatus::NCS_Error;
+
+		const bool UpdateScriptValid = UpdateScriptRequest
+			&& UpdateScriptRequest->ExeData.IsValid()
+			&& UpdateScriptRequest->ExeData->LastCompileStatus != ENiagaraScriptCompileStatus::NCS_Error;
+
+		if (SpawnScriptValid && UpdateScriptValid)
+		{
+			if (SpawnScriptRequest->ExeData->Attributes != UpdateScriptRequest->ExeData->Attributes)
+			{
+				// if we had requested a full rebuild, then we've got a case where the generated scripts are not compatible.  This indicates
+				// a significant issue where we're allowing graphs to generate invalid collections of scripts.  One known example is using
+				// the Script.Context static switch that isn't fully processed in all scripts, leading to attributes differing between the
+				// SystemSpawnScript and the SystemUpdateScript
+				if (bForced)
+				{
+					FString MissingAttributes;
+					FString AdditionalAttributes;
+
+					for (const auto& SpawnAttrib : SpawnScriptRequest->ExeData->Attributes)
+					{
+						if (!UpdateScriptRequest->ExeData->Attributes.Contains(SpawnAttrib))
+						{
+							MissingAttributes.Appendf(TEXT("%s%s"), MissingAttributes.Len() ? TEXT(", ") : TEXT(""), *SpawnAttrib.GetName().ToString());
+						}
+					}
+
+					for (const auto& UpdateAttrib : UpdateScriptRequest->ExeData->Attributes)
+					{
+						if (!SpawnScriptRequest->ExeData->Attributes.Contains(UpdateAttrib))
+						{
+							AdditionalAttributes.Appendf(TEXT("%s%s"), AdditionalAttributes.Len() ? TEXT(", ") : TEXT(""), *UpdateAttrib.GetName().ToString());
+						}
+					}
+
+					FNiagaraCompileEvent AttributeMismatchEvent(
+						FNiagaraCompileEventSeverity::Error,
+						FText::Format(LOCTEXT("SystemScriptAttributeMismatchError", "System Spawn/Update scripts have attributes which don't match!\n\tMissing update attributes: {0}\n\tAdditional update attributes: {1}"),
+							FText::FromString(MissingAttributes),
+							FText::FromString(AdditionalAttributes))
+						.ToString());
+
+					SpawnScriptRequest->ExeData->LastCompileStatus = ENiagaraScriptCompileStatus::NCS_Error;
+					SpawnScriptRequest->ExeData->LastCompileEvents.Add(AttributeMismatchEvent);
+				}
+				else
+				{
+					UE_LOG(LogNiagara, Log, TEXT("Failed to generate consistent results for System spawn and update scripts for system %s."), *Options.System->GetFullName());
+				}
+
+				return false;
+			}
+		}
+
+		return true;
+	}
+
+	virtual bool QueryCompileComplete(const FNiagaraQueryCompilationOptions& Options) override
+	{
+		if (CompileRequestHandle == INDEX_NONE)
+		{
+			return true;
+		}
+
+		check(IsInGameThread());
+		INiagaraModule& NiagaraModule = FModuleManager::Get().LoadModuleChecked<INiagaraModule>(TEXT("Niagara"));
+
+		constexpr bool bPeek = false;
+		if (NiagaraModule.PollSystemCompile(CompileRequestHandle, CompileResults, Options.bWait, bPeek))
+		{
+			CompileRequestHandle = INDEX_NONE;
+			return true;
+		}
+		return false;
+	}
+
+	virtual void Apply(const FNiagaraQueryCompilationOptions& Options) override
+	{
+		bool bAnyRIParametersDirtied = false;
+
+		// we need to do this apply in multiple passes.  Because the script's VMCompilationId has dependencies
+		// on the rapid iteration parameters of all the related scripts then we need to be sure to update the
+		// RI parameters with the data that we've collected before we actually try to set the results.
+		// Some additional caveats with this path:
+		// -our digested graph will not include disabled modules which means that their inputs, which could
+		//	resolve to rapid iteration parameters, and so we'll be missing those in our list.  We can't remove
+		//	parameters just because things are disabled because the editor will forget the settings.
+		constexpr bool bAllowParameterRemoval = false;
+
+		for (FNiagaraSystemAsyncCompileResults::FCompileResultMap::TConstIterator ResultIt(CompileResults.CompileResultMap);
+			ResultIt;
+			++ResultIt)
+		{
+			UNiagaraScript* TargetScript = ResultIt->Key;
+			const FNiagaraScriptAsyncCompileData& ScriptCompileData = ResultIt->Value;
+
+			if (TargetScript->ApplyRapidIterationParameters(ScriptCompileData.RapidIterationParameters, bAllowParameterRemoval))
+			{
+				bAnyRIParametersDirtied = true;
+			}
+		}
+
+		// Now that the above code says they are all complete, go ahead and resolve them all at once.
+		for (FNiagaraSystemAsyncCompileResults::FCompileResultMap::TConstIterator ResultIt(CompileResults.CompileResultMap);
+			ResultIt;
+			++ResultIt)
+		{
+			const FNiagaraScriptAsyncCompileData& ScriptCompileData = ResultIt->Value;
+			if (ensure(ScriptCompileData.ExeData.IsValid()))
+			{
+				UNiagaraScript* TargetScript = ResultIt->Key;
+				TMap<FName, UNiagaraDataInterface*> ObjectNameMap;
+
+				if (ScriptCompileData.bFromDerivedDataCache)
+				{
+					// if the data was pulled from the DDC then we'll need to generate the map from the source
+					const UNiagaraScriptSourceBase* ScriptSource = TargetScript->GetLatestSource();
+					ObjectNameMap = ScriptSource->ComputeObjectNameMap(*Options.System, TargetScript->GetUsage(), TargetScript->GetUsageId(), ScriptCompileData.UniqueEmitterName);
+				}
+				else
+				{
+					// if we actually generated our data will include a name map that we can use
+					Algo::Transform(ScriptCompileData.NamedDataInterfaces, ObjectNameMap, [](const TMap<FName, TObjectPtr<UNiagaraDataInterface>>::ElementType& Element)
+					{
+						return TMap<FName, UNiagaraDataInterface*>::ElementType(Element.Key, Element.Value);
+					});
+				}
+
+				// if we dirtied any RI parameters then we need to regenerate our CompilationId
+				FNiagaraVMExecutableDataId UpdatedCompileId;
+				const FNiagaraVMExecutableDataId* CompileIdToUse = &ScriptCompileData.CompileId;
+
+				if (bAnyRIParametersDirtied)
+				{
+					TargetScript->ComputeVMCompilationId(UpdatedCompileId, FGuid());
+					CompileIdToUse = &UpdatedCompileId;
+				}
+
+				constexpr bool bApplyRapidIterationParameters = false;
+				TargetScript->SetVMCompilationResults(
+					*CompileIdToUse,
+					*ScriptCompileData.ExeData,
+					ScriptCompileData.UniqueEmitterName,
+					ObjectNameMap,
+					bApplyRapidIterationParameters);
+			}
+		}
+
+		// Synchronize the variables that we actually encountered during precompile so that we can expose them to the end user.
+		FNiagaraUserRedirectionParameterStore& ExposedParameters = Options.System->GetExposedParameters();
+
+		for (const FNiagaraVariable& ExposedVariable : CompileResults.ExposedVariables)
+		{
+			if (!ExposedParameters.FindParameterOffset(ExposedVariable))
+			{
+				ExposedParameters.AddParameter(ExposedVariable, true, false);
+			}
+		}
+	}
+
+	virtual void ReportResults(const FNiagaraQueryCompilationOptions& Options) const override
+	{
+		const float ElapsedWallTime = (float)(FPlatformTime::Seconds() - TaskStartTime);
+
+		UE_LOG(LogNiagara, Log, TEXT("Compiling System %s took %f sec (time since issued)."),
+			*Options.System->GetFullName(), ElapsedWallTime);
+
+		if (Options.bGenerateTimingsFile)
+		{
+			FNiagaraSystemCompileMetrics SystemMetrics;
+			SystemMetrics.SystemCompileWallTime = ElapsedWallTime;
+			SystemMetrics.ScriptMetrics.Reserve(CompileResults.CompileResultMap.Num());
+
+			for (FNiagaraSystemAsyncCompileResults::FCompileResultMap::TConstIterator ResultIt(CompileResults.CompileResultMap);
+				ResultIt;
+				++ResultIt)
+			{
+				SystemMetrics.ScriptMetrics.Add(ResultIt.Key(), ResultIt.Value().CompileMetrics);
+			}
+
+			WriteTimingsEntry(TEXT("AsyncTask Compilation"), Options, SystemMetrics);
+		}
+	}
+
+	FNiagaraCompilationTaskHandle CompileRequestHandle;
+	FNiagaraSystemAsyncCompileResults CompileResults;
+};
+
+class FNiagaraActiveCompilationVerify : public FNiagaraActiveCompilation
+{
+public:
+	FNiagaraActiveCompilationVerify()
+	{
+		AsyncTaskRequest = MakeUnique<FNiagaraActiveCompilationAsyncTask>();
+		DefaultRequest = MakeUnique<FNiagaraActiveCompilationDefault>();
+	}
+
+	virtual void AddReferencedObjects(FReferenceCollector& Collector) override
+	{
+		AsyncTaskRequest->AddReferencedObjects(Collector);
+		DefaultRequest->AddReferencedObjects(Collector);
+	}
+	virtual FString GetReferencerName() const override
+	{
+		return TEXT("FNiagaraActiveCompilationVerify");
+	}
+
+	virtual bool Launch(const FNiagaraCompilationOptions& Options) override
+	{
+		bForced = Options.bForced;
+
+		const bool AsyncLaunched = AsyncTaskRequest->Launch(Options);
+		const bool DefaultLaunched = DefaultRequest->Launch(Options);
+
+		check(AsyncLaunched == DefaultLaunched);
+
+		return AsyncLaunched || DefaultLaunched;
+	}
+
+	virtual void Abort() override
+	{
+		AsyncTaskRequest->Abort();
+		DefaultRequest->Abort();
+	}
+
+	virtual bool QueryCompileComplete(const FNiagaraQueryCompilationOptions& Options) override
+	{
+		const bool AsyncComplete = AsyncTaskRequest->QueryCompileComplete(Options);
+		const bool DefaultComplete = DefaultRequest->QueryCompileComplete(Options);
+
+		return AsyncComplete && DefaultComplete;
+	}
+
+	virtual bool ValidateConsistentResults(const FNiagaraQueryCompilationOptions& Options) const override
+	{
+		bool DefaultValid = DefaultRequest->ValidateConsistentResults(Options);
+		bool AsyncTaskValid = AsyncTaskRequest->ValidateConsistentResults(Options);
+
+		// for this mode to be valid, both child modes must be valid
+		return DefaultValid && AsyncTaskValid;
+	}
+
+	virtual void Apply(const FNiagaraQueryCompilationOptions& Options) override
+	{
+		// for now we'll just apply the default results
+		DefaultRequest->Apply(Options);
+	}
+
+	virtual void ReportResults(const FNiagaraQueryCompilationOptions& Options) const override
+	{
+		auto VariableLess = [](const FNiagaraVariableBase& Lhs, const FNiagaraVariableBase& Rhs) -> bool
+		{
+			return Lhs.GetName().LexicalLess(Rhs.GetName());
+		};
+
+		// exposed parameters
+		{
+			TArray<FNiagaraVariable> DefaultExposedVariables;
+			{
+				TSet<FNiagaraVariable> UniqueDefaultExposedVariables;
+				for (FNiagaraActiveCompilationDefault::FAsyncTaskPtr TaskPtr : DefaultRequest->Tasks)
+				{
+					UniqueDefaultExposedVariables.Append(TaskPtr->EncounteredExposedVars);
+				}
+
+				DefaultExposedVariables = UniqueDefaultExposedVariables.Array();
+				DefaultExposedVariables.Sort(VariableLess);
+			}
+
+			TArray<FNiagaraVariable> AsyncTaskExposedVariables;
+			{
+				TSet<FNiagaraVariable> UniqueAsyncTaskExposedVariables(AsyncTaskRequest->CompileResults.ExposedVariables);
+
+				AsyncTaskExposedVariables = UniqueAsyncTaskExposedVariables.Array();
+				AsyncTaskExposedVariables.Sort(VariableLess);
+			}
+
+			EvaluateTest(Options, nullptr, DefaultExposedVariables == AsyncTaskExposedVariables, TEXT("Comparing ExposedVariables"));
+		}
+
+		// FNiagaraVMExecutableData
+		{
+			using FExeDataPair = TTuple<TSharedPtr<FNiagaraVMExecutableData>, TSharedPtr<FNiagaraVMExecutableData>>;
+			TMap<UNiagaraScript*, FExeDataPair> MappedExeData;
+
+			for (FNiagaraActiveCompilationDefault::FAsyncTaskPtr TaskPtr : DefaultRequest->Tasks)
+			{
+				FExeDataPair& DataPair = MappedExeData.FindOrAdd(TaskPtr->ScriptPair.CompiledScript);
+				DataPair.Key = TaskPtr->ExeData;
+			}
+
+			for (const auto& ResultIt : AsyncTaskRequest->CompileResults.CompileResultMap)
+			{
+				FExeDataPair& DataPair = MappedExeData.FindOrAdd(ResultIt.Key);
+				DataPair.Value = ResultIt.Value.ExeData;
+			}
+
+			for (const auto& ResultIt : MappedExeData)
+			{
+				CompareExeData(Options, ResultIt.Key, ResultIt.Value.Key, ResultIt.Value.Value);
+			}
+		}
+	}
+
+	virtual bool BlocksBeginCacheForCooked() const override
+	{
+		return DefaultRequest->BlocksBeginCacheForCooked() || AsyncTaskRequest->BlocksBeginCacheForCooked();
+	}
+
+private:
+	bool EvaluateTest(const FNiagaraQueryCompilationOptions& Options, UNiagaraScript* Script, bool bTestSuccess, const TCHAR* TestDescription) const
+	{
+		if (bTestSuccess || GNiagaraCompileValidateMode == (int32) ENiagaraCompilationValidateMode::Ignore)
+		{
+			return bTestSuccess;
+		}
+
+		auto GetSystemScriptName = [](UNiagaraSystem* InSystem, UNiagaraScript* InScript) -> FString
+		{
+			FString SystemScriptName = GetPathNameSafe(InSystem);
+			if (InScript && InSystem)
+			{
+				SystemScriptName.Append(TEXT(":"));
+				SystemScriptName.Append(InScript->GetPathName(InSystem));
+			}
+
+			return SystemScriptName;
+		};
+
+		switch ((ENiagaraCompilationValidateMode)GNiagaraCompileValidateMode)
+		{
+			case ENiagaraCompilationValidateMode::Warn:
+				UE_LOG(LogNiagara, Warning, TEXT("%s - Failed NiagaraCompile validation %s"), *GetSystemScriptName(Options.System, Script), TestDescription);
+				break;
+
+			case ENiagaraCompilationValidateMode::Ensure:
+				ensureMsgf(false, TEXT("%s - Failed NiagaraCompile validation %s"), *GetSystemScriptName(Options.System, Script), TestDescription);
+				break;
+
+			case ENiagaraCompilationValidateMode::Assert:
+				checkf(false, TEXT("%s - Failed NiagaraCompile validation %s"), *GetSystemScriptName(Options.System, Script), TestDescription);
+				break;
+		}
+
+		return bTestSuccess;
+	}
+
+	void CompareExeData(const FNiagaraQueryCompilationOptions& Options, UNiagaraScript* Script, TSharedPtr<FNiagaraVMExecutableData> DefaultExeData, TSharedPtr<FNiagaraVMExecutableData> AsyncTaskExeData) const
+	{
+		const bool HasDefaultData = DefaultExeData.IsValid();
+		const bool HasAsyncTaskData = AsyncTaskExeData.IsValid();
+
+		if (!EvaluateTest(Options, Script, HasDefaultData == HasAsyncTaskData, TEXT("Missing script compilation result")))
+		{
+			return;
+		}
+		if (!HasDefaultData && !HasAsyncTaskData)
+		{
+			return;
+		}
+
+		const bool ByteCodeMatches = DefaultExeData->ByteCode == AsyncTaskExeData->ByteCode;
+		if (!EvaluateTest(Options, Script, ByteCodeMatches, TEXT("Comparing ByteCode")))
+		{
+			bool bSignificantSourceDiff = false;
+
+			if (DefaultExeData->LastHlslTranslation != AsyncTaskExeData->LastHlslTranslation)
+			{
+				// see if it's just the two common lines that are different
+				TArray<FString> ExeDataTextLines;
+				DefaultExeData->LastHlslTranslation.ParseIntoArrayLines(ExeDataTextLines);
+				TArray<FString> ExeData2TextLines;
+				AsyncTaskExeData->LastHlslTranslation.ParseIntoArrayLines(ExeData2TextLines);
+
+				auto IsAComment = [](const FString& TextLine) -> bool
+				{
+					return TextLine.TrimStart().StartsWith(TEXT("//"));
+				};
+
+				ExeDataTextLines.SetNum(Algo::StableRemoveIf(ExeDataTextLines, IsAComment));
+				ExeData2TextLines.SetNum(Algo::StableRemoveIf(ExeData2TextLines, IsAComment));
+
+				const int32 LineCount = ExeDataTextLines.Num();
+				if (LineCount == ExeData2TextLines.Num())
+				{
+					for (int32 LineIt = 0; LineIt < LineCount; ++LineIt)
+					{
+						const FString& ExeDataLine = ExeDataTextLines[LineIt];
+						const FString& ExeData2Line = ExeData2TextLines[LineIt];
+
+						if (ExeDataLine == ExeData2Line)
+						{
+							continue;
+						}
+						else
+						{
+							// the following are various hacks put in place to let some specific code
+							// differences be ignored.  In particular bCompleteOnInactive0, when talking about
+							// the system spawn script is a common false positive error - the flag is initialized
+							// to true in the system update script, but is basically undefined in the system
+							// spawn script.  With the new digest pattern the default value from the update
+							// script is identified and used rather than leaving it undefined (and getting a
+							// default value of false).
+							{
+								FRegexPattern Pattern(TEXT("\\s*bool System_bCompleteOnInactive0 = (.*?);"));
+								FRegexMatcher ExeDataMatcher(Pattern, ExeDataLine);
+								FRegexMatcher ExeData2Matcher(Pattern, ExeData2Line);
+								if (ExeDataMatcher.FindNext() && ExeData2Matcher.FindNext())
+								{
+									if (ExeDataMatcher.GetCaptureGroup(1) == TEXT("false")
+										&& ExeData2Matcher.GetCaptureGroup(1) == TEXT("true"))
+									{
+										continue;
+									}
+								}
+							}
+
+							// this follows the above false positive with bCompleteOnInactive, where we will
+							// sometimes get the code of the following pattern:
+							// bool Constant4 = true;
+							// System_bCompleteOnInactive = Constant4;
+							// This pattern will evaluate the Constant declaration and see if it is immediately
+							// being applied to bCompleteOnInactive, if so we can safely ignore the difference
+							{
+								FRegexPattern Pattern(TEXT("\\s*bool Constant(\\d+) = (.*?);"));
+								FRegexMatcher ExeDataMatcher(Pattern, ExeDataLine);
+								FRegexMatcher ExeData2Matcher(Pattern, ExeData2Line);
+								if (ExeDataMatcher.FindNext() && ExeData2Matcher.FindNext())
+								{
+									const FString& ExeDataConstantIndex = ExeDataMatcher.GetCaptureGroup(1);
+									const FString& ExeData2ConstantIndex = ExeData2Matcher.GetCaptureGroup(1);
+									if (ExeDataConstantIndex == ExeData2ConstantIndex)
+									{
+										// look ahead to the next line to see if this is in relation to bCompleteOnInactive
+										// if it is then we'll ignore this difference
+										const FString& NextExeDataLine = ExeDataTextLines[LineIt + 1];
+										const FString& NextExeData2Line = ExeData2TextLines[LineIt + 1];
+										if (NextExeDataLine == NextExeData2Line && NextExeDataLine.Contains(TEXT("bCompleteOnInactive")))
+										{
+											continue;
+										}
+									}
+								}
+							}
+
+							// another common false positive in comparing the two versions of the hlsl source
+							// is the difference between:
+							// float MyVariable = 0.0f;
+							// and
+							// float MyVariable = (0);
+							{
+								FRegexPattern Pattern(TEXT("\t*float .*? = \\(?(.*?)\\)?;"));
+								FRegexMatcher ExeDataMatcher(Pattern, ExeDataLine);
+								FRegexMatcher ExeData2Matcher(Pattern, ExeData2Line);
+								if (ExeDataMatcher.FindNext() && ExeData2Matcher.FindNext())
+								{
+									const FString& ExeDataValue = ExeDataMatcher.GetCaptureGroup(1);
+									const FString& ExeData2Value = ExeData2Matcher.GetCaptureGroup(1);
+									if (ExeDataValue == ExeData2Value)
+									{
+										continue;
+									}
+									float ExeDataFloat = -1.0f;
+									float ExeData2Float = -1.0f;
+									LexFromString(ExeDataFloat, *ExeDataValue);
+									LexFromString(ExeData2Float, *ExeData2Value);
+
+									if (ExeDataFloat == ExeData2Float)
+									{
+										continue;
+									}
+								}
+							}
+
+							bSignificantSourceDiff = true;
+							break;
+						}
+					}
+				}
+				else
+				{
+					bSignificantSourceDiff = true;
+				}
+			}
+
+			EvaluateTest(Options, Script, !bSignificantSourceDiff, TEXT("Comparing HLSL source code"));
+		}
+
+		EvaluateTest(Options, Script, DefaultExeData->NumTempRegisters == AsyncTaskExeData->NumTempRegisters, TEXT("NumTempRegisters"));
+		EvaluateTest(Options, Script, DefaultExeData->NumUserPtrs == AsyncTaskExeData->NumUserPtrs, TEXT("NumUserPtrs"));
+		EvaluateTest(Options, Script, DefaultExeData->Parameters == AsyncTaskExeData->Parameters, TEXT("Parameters"));
+		EvaluateTest(Options, Script, DefaultExeData->InternalParameters == AsyncTaskExeData->InternalParameters, TEXT("InternalParameters"));
+		EvaluateTest(Options, Script, DefaultExeData->ExternalDependencies == AsyncTaskExeData->ExternalDependencies, TEXT("ExternalDependencies"));
+		EvaluateTest(Options, Script, DefaultExeData->CompileTags == AsyncTaskExeData->CompileTags, TEXT("CompileTags"));
+		EvaluateTest(Options, Script, DefaultExeData->ScriptLiterals == AsyncTaskExeData->ScriptLiterals, TEXT("ScriptLiterals"));
+		EvaluateTest(Options, Script, DefaultExeData->Attributes == AsyncTaskExeData->Attributes, TEXT("Attributes"));
+		EvaluateTest(Options, Script, DefaultExeData->DataUsage == AsyncTaskExeData->DataUsage, TEXT("DataUsage"));
+
+		// known false positive for DataSetToParameters - when compiling the new graph results from static switches
+		// can be propagated forward as we cull the static switch on instantiating the graph.  This can lead to
+		// reduced nodes visited during compilation, which can have the side effect of not encountering constants;
+		// Emitter.RandomSeed is an example of a constant that is skipped with CalculateRangedInt.  The new
+		// traversal is more correct, as the constant is not actually referenced by the system.
+		EvaluateTest(Options, Script, DefaultExeData->DataSetToParameters.OrderIndependentCompareEqual(AsyncTaskExeData->DataSetToParameters), TEXT("DataSetToParameters"));
+		EvaluateTest(Options, Script, DefaultExeData->AdditionalExternalFunctions == AsyncTaskExeData->AdditionalExternalFunctions, TEXT("AdditionalExternalFunctions"));
+		EvaluateTest(Options, Script, DefaultExeData->DataInterfaceInfo == AsyncTaskExeData->DataInterfaceInfo, TEXT("DataInterfaceInfo"));
+		EvaluateTest(Options, Script, DefaultExeData->CalledVMExternalFunctions == AsyncTaskExeData->CalledVMExternalFunctions, TEXT("CalledVMExternalFunctions"));
+		EvaluateTest(Options, Script, DefaultExeData->ReadDataSets == AsyncTaskExeData->ReadDataSets, TEXT("ReadDataSets"));
+		EvaluateTest(Options, Script, DefaultExeData->WriteDataSets == AsyncTaskExeData->WriteDataSets, TEXT("WriteDataSets"));
+		EvaluateTest(Options, Script, DefaultExeData->StatScopes == AsyncTaskExeData->StatScopes, TEXT("StatScopes"));
+		EvaluateTest(Options, Script, DefaultExeData->LastOpCount == AsyncTaskExeData->LastOpCount, TEXT("LastOpCount"));
+		EvaluateTest(Options, Script, DefaultExeData->ShaderScriptParametersMetadata == AsyncTaskExeData->ShaderScriptParametersMetadata, TEXT("ShaderScirptParameterMetaData"));
+		EvaluateTest(Options, Script, DefaultExeData->ParameterCollectionPaths == AsyncTaskExeData->ParameterCollectionPaths, TEXT("ParameterCollectionPaths"));
+		EvaluateTest(Options, Script, DefaultExeData->SimulationStageMetaData == AsyncTaskExeData->SimulationStageMetaData, TEXT("SimulationStageMetaData"));
+		EvaluateTest(Options, Script, DefaultExeData->bReadsAttributeData == AsyncTaskExeData->bReadsAttributeData, TEXT("bReadsAttributeData"));
+		EvaluateTest(Options, Script, DefaultExeData->AttributesWritten == AsyncTaskExeData->AttributesWritten, TEXT("AttributesWritten"));
+		EvaluateTest(Options, Script, DefaultExeData->StaticVariablesWritten == AsyncTaskExeData->StaticVariablesWritten, TEXT("StaticVaraiblesWritten"));
+		EvaluateTest(Options, Script, DefaultExeData->ErrorMsg == AsyncTaskExeData->ErrorMsg, TEXT("ErrorMsg"));
+		EvaluateTest(Options, Script, DefaultExeData->ExperimentalContextData == AsyncTaskExeData->ExperimentalContextData, TEXT("ExperimentalContextData"));
+		EvaluateTest(Options, Script, DefaultExeData->bReadsSignificanceIndex == AsyncTaskExeData->bReadsSignificanceIndex, TEXT("bReadsSignificanceIndex"));
+		EvaluateTest(Options, Script, DefaultExeData->bNeedsGPUContextInit == AsyncTaskExeData->bNeedsGPUContextInit, TEXT("bNeedsGPUContextInit"));
+
+		// because we can't fully trust the rapid iteration parameters we generate with the asyncTask method (because we ignore the
+		// entries from disabled modules).  We can't just do a straight compare.  Instead we'll see if any entries from the default
+		// are missing from the asynctask results
+		for (const FNiagaraVariable& DefaultParameter : DefaultExeData->BakedRapidIterationParameters)
+		{
+			const int32 AsyncTaskIndex = AsyncTaskExeData->BakedRapidIterationParameters.IndexOfByKey(DefaultParameter);
+			EvaluateTest(Options, Script, AsyncTaskIndex != INDEX_NONE, TEXT("Missing BakedRapidIterationParameter"));
+			if (AsyncTaskIndex != INDEX_NONE)
+			{
+				const FNiagaraVariable& AsyncTaskParameter = AsyncTaskExeData->BakedRapidIterationParameters[AsyncTaskIndex];
+				EvaluateTest(Options, Script, AsyncTaskParameter.HoldsSameData(DefaultParameter), TEXT("Mismatched BakedRapidIterationParmaeter data value"));
+			}
+		}
+	}
+
+	TUniquePtr<FNiagaraActiveCompilationDefault> DefaultRequest;
+	TUniquePtr<FNiagaraActiveCompilationAsyncTask> AsyncTaskRequest;
+};
+
+#endif // WITH_EDITORONLY_DATA
 
 TUniquePtr<FNiagaraActiveCompilation> FNiagaraActiveCompilation::CreateCompilation()
 {
+#if WITH_EDITORONLY_DATA
+	if (GNiagaraSystemCompileMode == (int32)ENiagaraCompilationMode::AsyncTasks)
+	{
+		return MakeUnique<FNiagaraActiveCompilationAsyncTask>();
+	}
+	else if (GNiagaraSystemCompileMode == (int32)ENiagaraCompilationMode::Verify)
+	{
+		return MakeUnique<FNiagaraActiveCompilationVerify>();
+	}
+
+	return MakeUnique<FNiagaraActiveCompilationDefault>();
+#else // WITH_EDITORONLY_DATA
 	return nullptr;
+#endif
 }
 
+void FNiagaraActiveCompilation::WriteTimingsEntry(const TCHAR* TimingsSource, const FNiagaraQueryCompilationOptions& Options, const FNiagaraSystemCompileMetrics& SystemMetrics) const
+{
+#if WITH_EDITORONLY_DATA && ALLOW_DEBUG_FILES
+	static bool HasWrittenFile = false;
+	static FCriticalSection TimingLock;
+	static FString TimingFile;
+
+	FScopeLock Lock(&TimingLock);
+
+	uint32 WriteFlags = FILEWRITE_Append;
+	if (!HasWrittenFile)
+	{
+		WriteFlags = FILEWRITE_None;
+		HasWrittenFile = true;
+		TimingFile = FPaths::ProjectSavedDir() / TEXT("Logs") / TEXT("NiagaraTimings_") + FDateTime::Now().ToString() + TEXT(".log");
+	}
+
+	TUniquePtr<FArchive> FileArchive(IFileManager::Get().CreateDebugFileWriter(*TimingFile, WriteFlags));
+
+	const int32 ScriptCount = SystemMetrics.ScriptMetrics.Num();
+	TArray<TObjectKey<UNiagaraScript>> SortedScripts;
+	SystemMetrics.ScriptMetrics.GenerateKeyArray(SortedScripts);
+
+	SortedScripts.Sort([](const TObjectKey<UNiagaraScript>& LhsScriptKey, const TObjectKey<UNiagaraScript>& RhsScriptKey) -> bool
+	{
+		return GetNameSafe(LhsScriptKey.ResolveObjectPtr()) < GetNameSafe(RhsScriptKey.ResolveObjectPtr());
+	});
+
+	FileArchive->Logf(TEXT("============================================================================"));
+	FileArchive->Logf(TEXT("==  Compiling %s using %s"), *Options.System->GetFullName(), TimingsSource);
+	FileArchive->Logf(TEXT("============================================================================"));
+	FileArchive->Logf(TEXT("== Task count: %d -> %fs for full task"), SortedScripts.Num(), SystemMetrics.SystemCompileWallTime);
+
+	if (!SortedScripts.IsEmpty())
+	{
+		FileArchive->Logf(TEXT("Script, Task Wall Time, DDC Fetch, Translate, Compile (Wall Time), Compile (Preprocess), Compile (Worker), Byte Code Optimize"));
+		for (const TObjectKey<UNiagaraScript> ScriptKey : SortedScripts)
+		{
+			const FNiagaraScriptCompileMetrics& ScriptMetric = SystemMetrics.ScriptMetrics.FindRef(ScriptKey);
+
+			FileArchive->Logf(TEXT("%s, %f, %f, %f, %f, %f, %f, %f"),
+				*GetPathNameSafe(ScriptKey.ResolveObjectPtr()),
+				ScriptMetric.TaskWallTime,
+				ScriptMetric.DDCFetchTime,
+				ScriptMetric.TranslateTime,
+				ScriptMetric.CompilerWallTime,
+				ScriptMetric.CompilerPreprocessTime,
+				ScriptMetric.CompilerWorkerTime,
+				ScriptMetric.ByteCodeOptimizeTime);
+		}
+	}
 #endif // WITH_EDITORONLY_DATA
+}
+
+FNiagaraQueryCompilationOptions::FNiagaraQueryCompilationOptions()
+#if WITH_EDITORONLY_DATA && ALLOW_DEBUG_FILES
+	: bGenerateTimingsFile(GNiagaraCompileDumpTimings)
+#else
+	: bGenerateTimingsFile(false)
+#endif
+{
+
+}
+
+#undef LOCTEXT_NAMESPACE // NiagaraSystem
diff --git a/Engine/Plugins/FX/Niagara/Source/Niagara/Private/NiagaraFunctionLibrary.cpp b/Engine/Plugins/FX/Niagara/Source/Niagara/Private/NiagaraFunctionLibrary.cpp
--- a/Engine/Plugins/FX/Niagara/Source/Niagara/Private/NiagaraFunctionLibrary.cpp
+++ b/Engine/Plugins/FX/Niagara/Source/Niagara/Private/NiagaraFunctionLibrary.cpp
@@ -1863,24 +1863,44 @@ bool UNiagaraFunctionLibrary::GetVectorVMFastPathExternalFunction(const FVMExter
 
 void UNiagaraFunctionLibrary::InitVectorVMFastPathOps()
 {
-	if (VectorVMOps.Num() > 0)
+	static FRWLock FastPathBuildLock;
+
+	{
+		FReadScopeLock Lock(FastPathBuildLock);
+
+		if (!VectorVMOps.IsEmpty())
+		{
+			return;
+		}
+	}
+
+	FWriteScopeLock Lock(FastPathBuildLock);
+	if (!VectorVMOps.IsEmpty())
+	{
 		return;
+	}
 
-	VectorVMOps.Emplace(FVectorKernelFastDot4::GetFunctionSignature());
-	VectorVMOps.Emplace(FVectorKernelFastTransformPosition::GetFunctionSignature());
-	VectorVMOps.Emplace(FVectorKernelFastMatrixToQuaternion::GetFunctionSignature());
-	VectorVMOps.Emplace(FVectorKernel_EmitterLifeCycle::GetFunctionSignature());
-	VectorVMOps.Emplace(FVectorKernel_SpawnRate::GetFunctionSignature());
-	VectorVMOps.Emplace(FVectorKernel_SpawnBurstInstantaneous::GetFunctionSignature());
-	VectorVMOps.Emplace(FVectorKernel_SolveVelocitiesAndForces::GetFunctionSignature());
-
-	VectorVMOpsHLSL.Emplace(FVectorKernelFastDot4::GetFunctionHLSL());
-	VectorVMOpsHLSL.Emplace(FVectorKernelFastTransformPosition::GetFunctionHLSL());
-	VectorVMOpsHLSL.Emplace(FVectorKernelFastMatrixToQuaternion::GetFunctionHLSL());
-	VectorVMOpsHLSL.Emplace(FVectorKernel_EmitterLifeCycle::GetFunctionHLSL());
-	VectorVMOpsHLSL.Emplace(FVectorKernel_SpawnRate::GetFunctionHLSL());
-	VectorVMOpsHLSL.Emplace(FVectorKernel_SpawnBurstInstantaneous::GetFunctionHLSL());
-	VectorVMOpsHLSL.Emplace(FVectorKernel_SolveVelocitiesAndForces::GetFunctionHLSL());
+	VectorVMOps =
+	{
+		FVectorKernelFastDot4::GetFunctionSignature(),
+		FVectorKernelFastTransformPosition::GetFunctionSignature(),
+		FVectorKernelFastMatrixToQuaternion::GetFunctionSignature(),
+		FVectorKernel_EmitterLifeCycle::GetFunctionSignature(),
+		FVectorKernel_SpawnRate::GetFunctionSignature(),
+		FVectorKernel_SpawnBurstInstantaneous::GetFunctionSignature(),
+		FVectorKernel_SolveVelocitiesAndForces::GetFunctionSignature()
+	};
+
+	VectorVMOpsHLSL =
+	{
+		FVectorKernelFastDot4::GetFunctionHLSL(),
+		FVectorKernelFastTransformPosition::GetFunctionHLSL(),
+		FVectorKernelFastMatrixToQuaternion::GetFunctionHLSL(),
+		FVectorKernel_EmitterLifeCycle::GetFunctionHLSL(),
+		FVectorKernel_SpawnRate::GetFunctionHLSL(),
+		FVectorKernel_SpawnBurstInstantaneous::GetFunctionHLSL(),
+		FVectorKernel_SolveVelocitiesAndForces::GetFunctionHLSL()
+	};
 
 	check(VectorVMOps.Num() == VectorVMOpsHLSL.Num());
 }
diff --git a/Engine/Plugins/FX/Niagara/Source/Niagara/Private/NiagaraModule.cpp b/Engine/Plugins/FX/Niagara/Source/Niagara/Private/NiagaraModule.cpp
--- a/Engine/Plugins/FX/Niagara/Source/Niagara/Private/NiagaraModule.cpp
+++ b/Engine/Plugins/FX/Niagara/Source/Niagara/Private/NiagaraModule.cpp
@@ -630,10 +630,10 @@ int32 INiagaraModule::StartScriptCompileJob(const FNiagaraCompileRequestDataBase
 	return ScriptCompilerDelegate.Execute(InCompileData, InCompileDuplicateData, InCompileOptions);
 }
 
-TSharedPtr<FNiagaraVMExecutableData> INiagaraModule::GetCompileJobResult(int32 JobID, bool bWait)
+TSharedPtr<FNiagaraVMExecutableData> INiagaraModule::GetCompileJobResult(int32 JobID, bool bWait, FNiagaraScriptCompileMetrics& Metrics)
 {
 	checkf(ScriptCompilerDelegate.IsBound(), TEXT("Script compilation result delegate not bound."));
-	return CompilationResultDelegate.Execute(JobID, bWait);
+	return CompilationResultDelegate.Execute(JobID, bWait, Metrics);
 }
 
 FDelegateHandle INiagaraModule::RegisterScriptCompiler(FScriptCompiler ScriptCompiler)
@@ -733,7 +733,65 @@ void INiagaraModule::UnregisterGraphTraversalCacher(FDelegateHandle DelegateHand
 	GraphTraversalCacheDelegate.Unbind();
 }
 
+FNiagaraCompilationTaskHandle INiagaraModule::RequestCompileSystem(UNiagaraSystem* System, bool bForce)
+{
+	checkf(RequestCompileSystemDelegate.IsBound(), TEXT("RequestCompileSystemDelegate delegate not bound."));
+	return RequestCompileSystemDelegate.Execute(System, bForce);
+}
 
+FDelegateHandle INiagaraModule::RegisterRequestCompileSystem(FOnRequestCompileSystem RequestCompileSystemCallback)
+{
+	checkf(RequestCompileSystemDelegate.IsBound() == false, TEXT("Only one handler is allowed for the RequestCompileSystem delegate"));
+	RequestCompileSystemDelegate = RequestCompileSystemCallback;
+	return RequestCompileSystemDelegate.GetHandle();
+}
+
+void INiagaraModule::UnregisterRequestCompileSystem(FDelegateHandle DelegateHandle)
+{
+	checkf(RequestCompileSystemDelegate.IsBound(), TEXT("RequestCompileSystemDelegate is not registered"));
+	checkf(RequestCompileSystemDelegate.GetHandle() == DelegateHandle, TEXT("Can only unregister the RequestCompileSystem delegate with the handle it was registered with."));
+	RequestCompileSystemDelegate.Unbind();
+}
+
+bool INiagaraModule::PollSystemCompile(FNiagaraCompilationTaskHandle TaskHandle, FNiagaraSystemAsyncCompileResults& Results, bool bWait, bool bPeek)
+{
+	checkf(PollSystemCompileDelegate.IsBound(), TEXT("PollSystemCompileDelegate delegate not bound."));
+	return PollSystemCompileDelegate.Execute(TaskHandle, Results, bWait, bPeek);
+}
+
+FDelegateHandle INiagaraModule::RegisterPollSystemCompile(FOnPollSystemCompile PollSystemCompileCallback)
+{
+	checkf(PollSystemCompileDelegate.IsBound() == false, TEXT("Only one handler is allowed for the PollSystemCompile delegate"));
+	PollSystemCompileDelegate = PollSystemCompileCallback;
+	return PollSystemCompileDelegate.GetHandle();
+}
+
+void INiagaraModule::UnregisterPollSystemCompile(FDelegateHandle DelegateHandle)
+{
+	checkf(PollSystemCompileDelegate.IsBound(), TEXT("PollSystemCompileDelegate is not registered"));
+	checkf(PollSystemCompileDelegate.GetHandle() == DelegateHandle, TEXT("Can only unregister the PollSystemCompile delegate with the handle it was registered with."));
+	PollSystemCompileDelegate.Unbind();
+}
+
+void INiagaraModule::AbortSystemCompile(FNiagaraCompilationTaskHandle TaskHandle)
+{
+	checkf(AbortSystemCompileDelegate.IsBound(), TEXT("AbortSystemCompileDelegate delegate not bound."));
+	AbortSystemCompileDelegate.Execute(TaskHandle);
+}
+
+FDelegateHandle INiagaraModule::RegisterAbortSystemCompile(FOnAbortSystemCompile AbortSystemCompileCallback)
+{
+	checkf(AbortSystemCompileDelegate.IsBound() == false, TEXT("Only one handler is allowed for the AbortSystemCompile delegate"));
+	AbortSystemCompileDelegate = AbortSystemCompileCallback;
+	return AbortSystemCompileDelegate.GetHandle();
+}
+
+void INiagaraModule::UnregisterAbortSystemCompile(FDelegateHandle DelegateHandle)
+{
+	checkf(AbortSystemCompileDelegate.IsBound(), TEXT("AbortSystemCompileDelegate is not registered"));
+	checkf(AbortSystemCompileDelegate.GetHandle() == DelegateHandle, TEXT("Can only unregister the AbortSystemCompile delegate with the handle it was registered with."));
+	AbortSystemCompileDelegate.Unbind();
+}
 
 void INiagaraModule::OnAssetLoaded(UObject* Asset)
 {
diff --git a/Engine/Plugins/FX/Niagara/Source/Niagara/Private/NiagaraScript.cpp b/Engine/Plugins/FX/Niagara/Source/Niagara/Private/NiagaraScript.cpp
--- a/Engine/Plugins/FX/Niagara/Source/Niagara/Private/NiagaraScript.cpp
+++ b/Engine/Plugins/FX/Niagara/Source/Niagara/Private/NiagaraScript.cpp
@@ -1011,7 +1011,7 @@ FString UNiagaraScript::GetNiagaraDDCKeyString(const FGuid& ScriptVersion, const
 	return BuildNiagaraDDCKeyString(GetLastGeneratedVMId(ScriptVersion), ScriptPath);
 }
 
-void UNiagaraScript::ComputeVMCompilationId(FNiagaraVMExecutableDataId& Id, FGuid VersionGuid) const
+void UNiagaraScript::ComputeVMCompilationId(FNiagaraVMExecutableDataId& Id, const FGuid& VersionGuid) const
 {
 	Id = FNiagaraVMExecutableDataId();
 
@@ -1460,7 +1460,6 @@ void UNiagaraScript::ComputeVMCompilationId(FNiagaraVMExecutableDataId& Id, FGui
 	LastGeneratedVMId = Id;
 }
 
-
 void UNiagaraScript::ComputeVMCompilationId_EmitterShared(FNiagaraVMExecutableDataId& Id, const FVersionedNiagaraEmitter& Emitter, ENiagaraRendererSourceDataMode InSourceMode) const
 {
 	// Gather additional variables from renderers
@@ -2847,7 +2846,7 @@ FName ResolveEmitterAlias(const FName& InName, const FString& InAlias)
 	return Var.GetName();
 }
 
-void UNiagaraScript::SetVMCompilationResults(const FNiagaraVMExecutableDataId& InCompileId, FNiagaraVMExecutableData& InScriptVM, FString EmitterUniqueName, const TMap<FName, UNiagaraDataInterface*>& ObjectNameMap, bool bApplyRapidIterationParameters)
+void UNiagaraScript::SetVMCompilationResults(const FNiagaraVMExecutableDataId& InCompileId, FNiagaraVMExecutableData& InScriptVM, const FString& EmitterUniqueName, const TMap<FName, UNiagaraDataInterface*>& ObjectNameMap, bool bApplyRapidIterationParameters)
 {
 	CachedScriptVMId = InCompileId;
 	CachedScriptVM = InScriptVM;
@@ -2908,6 +2907,15 @@ void UNiagaraScript::SetVMCompilationResults(const FNiagaraVMExecutableDataId& I
 
 	CachedDefaultUObjects = CachedScriptVM.UObjectInfos;
 
+	// if the compile process generated soft object paths instead of hard references than resolve them here
+	for (FNiagaraScriptUObjectCompileInfo& DefaultUObjectInfo : CachedDefaultUObjects)
+	{
+		if (DefaultUObjectInfo.ObjectPath.IsValid() && !DefaultUObjectInfo.Object)
+		{
+			DefaultUObjectInfo.Object = DefaultUObjectInfo.ObjectPath.ResolveObject();
+		}
+	}
+
 	if (bApplyRapidIterationParameters)
 	{
 		const bool bClearBindings = false;
@@ -2949,6 +2957,60 @@ void UNiagaraScript::SetVMCompilationResults(const FNiagaraVMExecutableDataId& I
 	OnVMScriptCompiled().Broadcast(this, InCompileId.ScriptVersionID);
 }
 
+bool UNiagaraScript::ApplyRapidIterationParameters(TConstArrayView<FNiagaraVariable> InParameters, bool bAllowRemoval)
+{
+	bool bStoreDirty = false;
+
+	{
+		FNiagaraParameterStore::FScopedSuppressOnChanged SuppressScope(RapidIterationParameters);
+
+		// go through and remove any parameters that aren't in the supplied parameter list
+		if (bAllowRemoval)
+		{
+			TConstArrayView<FNiagaraVariableWithOffset> ExistingParameters = RapidIterationParameters.ReadParameterVariables();
+			TArray<FNiagaraVariableBase> ParametersToRemove;
+	
+			for (const FNiagaraVariableWithOffset& ExistingParameter : ExistingParameters)
+			{
+				const bool bStaleParameter = !InParameters.ContainsByPredicate([&ExistingParameter](const FNiagaraVariable& Variable) -> bool
+				{
+					return Variable == ExistingParameter;
+				});
+
+				if (bStaleParameter)
+				{
+					ParametersToRemove.Add(ExistingParameter);
+					bStoreDirty = true;
+				}
+			}
+
+			for (const FNiagaraVariableBase& ParameterToRemove : ParametersToRemove)
+			{
+				RapidIterationParameters.RemoveParameter(ParameterToRemove);
+			}
+		}
+
+		// try to add 
+		for (const FNiagaraVariable& Parameter : InParameters)
+		{
+			constexpr bool bInitInterfaces = false;
+			constexpr bool bTriggerRebind = false;
+
+			if (RapidIterationParameters.AddParameter(Parameter, bInitInterfaces, bTriggerRebind, nullptr))
+			{
+				bStoreDirty = true;
+			}
+		}
+	}
+
+	if (bStoreDirty)
+	{
+		RapidIterationParameters.Rebind();
+	}
+
+	return bStoreDirty;
+}
+
 void UNiagaraScript::InvalidateExecutionReadyParameterStores()
 {
 #if WITH_EDITORONLY_DATA
@@ -3019,7 +3081,8 @@ void UNiagaraScript::RequestCompile(const FGuid& ScriptVersion, bool bForceCompi
 
 		FNiagaraCompileOptions Options(GetUsage(), GetUsageId(), ScriptData->ModuleUsageBitmask, ScriptPathName, GetFullName(), GetName());
 		int32 JobHandle = NiagaraModule.StartScriptCompileJob(RequestData.Get(), RequestDuplicateData.Get(), Options);
-		TSharedPtr<FNiagaraVMExecutableData> ExeData = NiagaraModule.GetCompileJobResult(JobHandle, true);
+		FNiagaraScriptCompileMetrics ScriptMetrics;
+		TSharedPtr<FNiagaraVMExecutableData> ExeData = NiagaraModule.GetCompileJobResult(JobHandle, true, ScriptMetrics);
 		if (ExeData)
 		{
 			SetVMCompilationResults(LastGeneratedVMId, *ExeData, FString(), RequestDuplicateData->GetObjectNameMap(), false);
diff --git a/Engine/Plugins/FX/Niagara/Source/Niagara/Private/NiagaraSystem.cpp b/Engine/Plugins/FX/Niagara/Source/Niagara/Private/NiagaraSystem.cpp
--- a/Engine/Plugins/FX/Niagara/Source/Niagara/Private/NiagaraSystem.cpp
+++ b/Engine/Plugins/FX/Niagara/Source/Niagara/Private/NiagaraSystem.cpp
@@ -254,19 +254,43 @@ void UNiagaraSystem::BeginCacheForCookedPlatformData(const ITargetPlatform *Targ
 	Super::BeginCacheForCookedPlatformData(TargetPlatform);
 	
 	EnsureFullyLoaded();
-#if WITH_EDITORONLY_DATA
-	// todo - when we can get compilation work to be properly asynchronous we need to move this over to the intended behavior where
-	// this function kicks off the work and we poll for results in IsCachedCookedPlatformDataLoaded().  With the current setup we
-	// will timeout DDC fill jobs while polling for the work to be complete.  Not entirely sure why, but my guess is that we get too
-	// many systems in the process of being compiled and we iteratively block the gamethread advancing each of them starving progress
-	WaitForCompilationComplete();
-#endif
+
+	if (bNeedsRequestCompile)
+	{
+		RequestCompile(false);
+	}
+
+	// check if any of the active compilations requires waiting - note that none should but the original
+	// implementation does as it will stall the DDC fill jobs.  Unclear exactly why, but likely because of
+	// the need to tick some of the managers
+	bool bWaitForCompilationComplete = false;
+	for (const TUniquePtr<FNiagaraActiveCompilation>& ActiveCompilation : ActiveCompilations)
+	{
+		if (ActiveCompilation->BlocksBeginCacheForCooked())
+		{
+			bWaitForCompilationComplete = true;
+		}
+	}
+
+	if (bWaitForCompilationComplete)
+	{
+		WaitForCompilationComplete();
+	}
 }
 
 bool UNiagaraSystem::IsCachedCookedPlatformDataLoaded(const ITargetPlatform* TargetPlatform)
 {
-	// todo - see note above in BeginCacheForCookedPlatformData()
-	return Super::IsCachedCookedPlatformDataLoaded(TargetPlatform);
+	if (!Super::IsCachedCookedPlatformDataLoaded(TargetPlatform))
+	{
+		return false;
+	}
+
+	if (!ActiveCompilations.IsEmpty())
+	{
+		return PollForCompilationComplete(false);
+	}
+
+	return true;
 }
 
 void UNiagaraSystem::HandleVariableRenamed(const FNiagaraVariable& InOldVariable, const FNiagaraVariable& InNewVariable, bool bUpdateContexts)
@@ -1871,7 +1895,7 @@ void UNiagaraSystem::CacheFromCompiledData()
 	FGraphEventArray PSOPrecacheEvents;
 
 	TSet<FName> DataInterfaceGpuUsage;
-	TStringBuilder<128> ExecutionStateNameBuilder;
+	FNameBuilder ExecutionStateNameBuilder;
 	for (int32 i=0; i < EmitterHandles.Num(); ++i)
 	{
 		FNiagaraEmitterHandle& Handle = EmitterHandles[i];
@@ -1882,7 +1906,7 @@ void UNiagaraSystem::CacheFromCompiledData()
 			ExecutionStateNameBuilder.Reset();
 			ExecutionStateNameBuilder << Handle.GetInstance().Emitter->GetUniqueEmitterName();
 			ExecutionStateNameBuilder << TEXT(".ExecutionState");
-			const FName ExecutionStateName(ExecutionStateNameBuilder.ToString());
+			const FName ExecutionStateName(ExecutionStateNameBuilder);
 
 			EmitterExecutionStateAccessors.AddDefaulted_GetRef().Init(SystemDataSet, ExecutionStateName);
 
@@ -2739,21 +2763,21 @@ bool UNiagaraSystem::QueryCompileComplete(bool bWait)
 	TUniquePtr<FNiagaraActiveCompilation> CurrentCompilation = MoveTemp(ActiveCompilations[0]);
 	ActiveCompilations.RemoveAt(0);
 
+	// In the world of do not apply, we're exiting the system completely so let's just kill any active compilations altogether.
+	if (!CurrentCompilation->ShouldApply())
+	{
+		return true;
+	}
+
 	// if we've gotten all the results, run a quick check to see if the data is valid, if it's not then that indicates that
 	// we've run into a compatibility issue and so we should see if we should issue a full rebuild
-	const bool ResultsValid = CurrentCompilation->Validate(Options);
-	if (!ResultsValid && !CurrentCompilation->WasForced())
+	const bool ResultsConsistent = CurrentCompilation->ValidateConsistentResults(Options);
+	if (!ResultsConsistent && !CurrentCompilation->WasForced())
 	{
 		RequestCompile(true, nullptr);
 		return false;
 	}
 
-	// In the world of do not apply, we're exiting the system completely so let's just kill any active compilations altogether.
-	if (!CurrentCompilation->ShouldApply())
-	{
-		return true;
-	}
-
 	CurrentCompilation->Apply(Options);
 
 	// Once compile results have been set, check dependencies found during compile and evaluate whether dependency compile events should be added.
diff --git a/Engine/Plugins/FX/Niagara/Source/Niagara/Public/NiagaraCommon.h b/Engine/Plugins/FX/Niagara/Source/Niagara/Public/NiagaraCommon.h
--- a/Engine/Plugins/FX/Niagara/Source/Niagara/Public/NiagaraCommon.h
+++ b/Engine/Plugins/FX/Niagara/Source/Niagara/Public/NiagaraCommon.h
@@ -680,6 +680,9 @@ struct NIAGARA_API FNiagaraScriptUObjectCompileInfo
 	UPROPERTY()
 	TObjectPtr<class UObject> Object;
 
+	UPROPERTY()
+	FSoftObjectPath ObjectPath;
+
 	UPROPERTY()
 	FName RegisteredParameterMapRead;
 
diff --git a/Engine/Plugins/FX/Niagara/Source/Niagara/Public/NiagaraCompilationTypes.h b/Engine/Plugins/FX/Niagara/Source/Niagara/Public/NiagaraCompilationTypes.h
--- a/Engine/Plugins/FX/Niagara/Source/Niagara/Public/NiagaraCompilationTypes.h
+++ b/Engine/Plugins/FX/Niagara/Source/Niagara/Public/NiagaraCompilationTypes.h
@@ -13,6 +13,25 @@ class UNiagaraSystem;
 
 using FNiagaraCompilationTaskHandle = int32;
 
+struct FNiagaraScriptCompileMetrics
+{
+	float TaskWallTime = 0.0f;
+	float DDCFetchTime = 0.0f;
+	float CompilerWallTime = 0.0f;
+	float CompilerWorkerTime = 0.0f;
+	float CompilerPreprocessTime = 0.0f;
+	float TranslateTime = 0.0f;
+	float ByteCodeOptimizeTime = 0.0f;
+};
+
+struct FNiagaraSystemCompileMetrics
+{
+	float SystemCompileWallTime = 0.0f;
+	float SystemPrecompileTime = 0.0f;
+	float SystemCompilationCopyTime = 0.0f;
+	TMap<TObjectKey<UNiagaraScript>, FNiagaraScriptCompileMetrics> ScriptMetrics;
+};
+
 USTRUCT()
 struct FNiagaraScriptAsyncCompileData
 {
@@ -21,8 +40,13 @@ struct FNiagaraScriptAsyncCompileData
 	FNiagaraVMExecutableDataId CompileId;
 	TSharedPtr<struct FNiagaraVMExecutableData> ExeData;
 	FString UniqueEmitterName;
+	FNiagaraScriptCompileMetrics CompileMetrics;
+
 	bool bFromDerivedDataCache = false;
 
+	UPROPERTY()
+	TArray<FNiagaraVariable> RapidIterationParameters;
+
 	UPROPERTY()
 	TMap<FName, TObjectPtr<UNiagaraDataInterface>> NamedDataInterfaces;
 };
@@ -57,9 +81,12 @@ struct FNiagaraCompilationOptions
 
 struct FNiagaraQueryCompilationOptions
 {
+	FNiagaraQueryCompilationOptions();
+
 	UNiagaraSystem* System = nullptr;
 	double MaxWaitDuration = 0.125;
 	bool bWait = false;
+	bool bGenerateTimingsFile = false;
 };
 
 class FNiagaraActiveCompilation : public FGCObject
@@ -70,9 +97,10 @@ public:
 	virtual bool Launch(const FNiagaraCompilationOptions& Options) = 0;
 	virtual void Abort() = 0;
 	virtual bool QueryCompileComplete(const FNiagaraQueryCompilationOptions& Options) = 0;
-	virtual bool Validate(const FNiagaraQueryCompilationOptions& Options) const = 0;
+	virtual bool ValidateConsistentResults(const FNiagaraQueryCompilationOptions& Options) const = 0;
 	virtual void Apply(const FNiagaraQueryCompilationOptions& Options) = 0;
 	virtual void ReportResults(const FNiagaraQueryCompilationOptions& Options) const = 0;
+	virtual bool BlocksBeginCacheForCooked() const { return false; }
 
 	void Invalidate()
 	{
@@ -90,8 +118,9 @@ public:
 	}
 
 protected:
-	bool bForced = false;
+	void WriteTimingsEntry(const TCHAR* TimingsSource, const FNiagaraQueryCompilationOptions& Options, const FNiagaraSystemCompileMetrics& SystemMetrics) const;
 
-private:
+	double TaskStartTime = 0;
+	bool bForced = false;
 	bool bShouldApply = true;
 };
diff --git a/Engine/Plugins/FX/Niagara/Source/Niagara/Public/NiagaraModule.h b/Engine/Plugins/FX/Niagara/Source/Niagara/Public/NiagaraModule.h
--- a/Engine/Plugins/FX/Niagara/Source/Niagara/Public/NiagaraModule.h
+++ b/Engine/Plugins/FX/Niagara/Source/Niagara/Public/NiagaraModule.h
@@ -6,6 +6,7 @@
 #include "Modules/ModuleInterface.h"
 #include "NiagaraTypes.h"
 #include "Templates/SharedPointer.h"
+#include "NiagaraCompilationTypes.h"
 #include "NiagaraPerfBaseline.h"
 #include "NiagaraDebuggerCommon.h"
 #include "NiagaraScript.h"
@@ -23,6 +24,7 @@ class INiagaraEditorOnlyDataUtilities;
 struct FNiagaraParameterStore;
 class FCommonViewportClient;
 class FNiagaraDebuggerClient;
+struct FNiagaraSystemAsyncCompileResults;
 
 extern NIAGARA_API int32 GEnableVerboseNiagaraChangeIdLogging;
 
@@ -38,10 +40,15 @@ public:
 	typedef TSharedPtr<FNiagaraGraphCachedDataBase, ESPMode::ThreadSafe> GraphCachedDataPtr;
 
 	DECLARE_DELEGATE_RetVal_ThreeParams(int32, FScriptCompiler, const FNiagaraCompileRequestDataBase*, const FNiagaraCompileRequestDuplicateDataBase*, const FNiagaraCompileOptions&);
-	DECLARE_DELEGATE_RetVal_TwoParams(TSharedPtr<FNiagaraVMExecutableData>, FCheckCompilationResult, int32, bool);
+	DECLARE_DELEGATE_RetVal_ThreeParams(TSharedPtr<FNiagaraVMExecutableData>, FCheckCompilationResult, int32, bool, FNiagaraScriptCompileMetrics&);
 	DECLARE_DELEGATE_RetVal_TwoParams(CompileRequestPtr, FOnPrecompile, UObject*, FGuid);
 	DECLARE_DELEGATE_RetVal_FiveParams(CompileRequestDuplicatePtr, FOnPrecompileDuplicate, const FNiagaraCompileRequestDataBase* /*OwningSystemRequestData*/, UNiagaraSystem* /*OwningSystem*/, UNiagaraEmitter* /*OwningEmitter*/, UNiagaraScript* /*TargetScript*/, FGuid /*Version*/);
 	DECLARE_DELEGATE_RetVal_TwoParams(GraphCachedDataPtr, FOnCacheGraphTraversal, const UObject*, FGuid);
+
+	DECLARE_DELEGATE_RetVal_TwoParams(FNiagaraCompilationTaskHandle, FOnRequestCompileSystem, UNiagaraSystem*, bool);
+	DECLARE_DELEGATE_RetVal_FourParams(bool, FOnPollSystemCompile, FNiagaraCompilationTaskHandle, FNiagaraSystemAsyncCompileResults&, bool /*bWait*/, bool /*bPeek*/);
+	DECLARE_DELEGATE_OneParam(FOnAbortSystemCompile, FNiagaraCompilationTaskHandle);
+
 #endif
 	DECLARE_DELEGATE_RetVal(void, FOnProcessQueue);
 
@@ -80,7 +87,7 @@ public:
 	void UnregisterEditorOnlyDataUtilities(TSharedRef<INiagaraEditorOnlyDataUtilities> InEditorOnlyDataUtilities);
 
 	int32 StartScriptCompileJob(const FNiagaraCompileRequestDataBase* InCompileData, const FNiagaraCompileRequestDuplicateDataBase* InCompileDuplicateData, const FNiagaraCompileOptions& InCompileOptions);
-	TSharedPtr<FNiagaraVMExecutableData> GetCompileJobResult(int32 JobID, bool bWait);
+	TSharedPtr<FNiagaraVMExecutableData> GetCompileJobResult(int32 JobID, bool bWait, FNiagaraScriptCompileMetrics& Metrics);
 
 	FDelegateHandle RegisterScriptCompiler(FScriptCompiler ScriptCompiler);
 	void UnregisterScriptCompiler(FDelegateHandle DelegateHandle);
@@ -104,6 +111,18 @@ public:
 	FDelegateHandle RegisterGraphTraversalCacher(FOnCacheGraphTraversal PreCompiler);
 	void UnregisterGraphTraversalCacher(FDelegateHandle DelegateHandle);
 
+	FNiagaraCompilationTaskHandle RequestCompileSystem(UNiagaraSystem* System, bool bForce);
+	FDelegateHandle RegisterRequestCompileSystem(FOnRequestCompileSystem RequestCompileSystemCallback);
+	void UnregisterRequestCompileSystem(FDelegateHandle DelegateHandle);
+
+	bool PollSystemCompile(FNiagaraCompilationTaskHandle, FNiagaraSystemAsyncCompileResults&, bool /*bWait*/, bool /*bPeek*/);
+	FDelegateHandle RegisterPollSystemCompile(FOnPollSystemCompile PollSystemCompileCallback);
+	void UnregisterPollSystemCompile(FDelegateHandle DelegateHandle);
+
+	void AbortSystemCompile(FNiagaraCompilationTaskHandle);
+	FDelegateHandle RegisterAbortSystemCompile(FOnAbortSystemCompile AbortSystemCompileCallback);
+	void UnregisterAbortSystemCompile(FDelegateHandle DelegateHandle);
+
 	void OnAssetLoaded(UObject* Asset);
 
 #endif
@@ -249,6 +268,9 @@ public:
 	FOnPrecompile PrecompileDelegate;
 	FOnPrecompileDuplicate PrecompileDuplicateDelegate;
 	FOnCacheGraphTraversal GraphTraversalCacheDelegate;
+	FOnRequestCompileSystem RequestCompileSystemDelegate;
+	FOnPollSystemCompile PollSystemCompileDelegate;
+	FOnAbortSystemCompile AbortSystemCompileDelegate;
 #endif
 
 	static int32 EngineEffectsQuality;
diff --git a/Engine/Plugins/FX/Niagara/Source/NiagaraEditor/NiagaraEditor.Build.cs b/Engine/Plugins/FX/Niagara/Source/NiagaraEditor/NiagaraEditor.Build.cs
--- a/Engine/Plugins/FX/Niagara/Source/NiagaraEditor/NiagaraEditor.Build.cs
+++ b/Engine/Plugins/FX/Niagara/Source/NiagaraEditor/NiagaraEditor.Build.cs
@@ -13,6 +13,7 @@ public class NiagaraEditor : ModuleRules
                 "Core", 
 				"CoreUObject", 
 				"CurveEditor",
+				"DerivedDataCache",
 				"ApplicationCore",
                 "InputCore",
 				"RenderCore",
diff --git a/Engine/Plugins/FX/Niagara/Source/NiagaraEditor/Private/NiagaraAttributeTrimmer.cpp b/Engine/Plugins/FX/Niagara/Source/NiagaraEditor/Private/NiagaraAttributeTrimmer.cpp
--- a/Engine/Plugins/FX/Niagara/Source/NiagaraEditor/Private/NiagaraAttributeTrimmer.cpp
+++ b/Engine/Plugins/FX/Niagara/Source/NiagaraEditor/Private/NiagaraAttributeTrimmer.cpp
@@ -7,6 +7,7 @@
 #include "EdGraphSchema_Niagara.h"
 #include "NiagaraCompilationBridge.h"
 #include "NiagaraConstants.h"
+#include "NiagaraGraphDigest.h"
 #include "NiagaraModule.h"
 #include "NiagaraNodeCustomHlsl.h"
 #include "NiagaraNodeFunctionCall.h"
@@ -782,4 +783,4 @@ void FNiagaraAttributeTrimmerHelper<GraphBridge>::TrimAttributes_Aggressive(cons
 }
 
 template class FNiagaraAttributeTrimmerHelper<FNiagaraCompilationGraphBridge>;
-
+template class FNiagaraAttributeTrimmerHelper<FNiagaraCompilationDigestBridge>;
diff --git a/Engine/Plugins/FX/Niagara/Source/NiagaraEditor/Private/NiagaraCompilationPrivate.h b/Engine/Plugins/FX/Niagara/Source/NiagaraEditor/Private/NiagaraCompilationPrivate.h
--- a/Engine/Plugins/FX/Niagara/Source/NiagaraEditor/Private/NiagaraCompilationPrivate.h
+++ b/Engine/Plugins/FX/Niagara/Source/NiagaraEditor/Private/NiagaraCompilationPrivate.h
@@ -4,6 +4,8 @@
 
 #include "CoreMinimal.h"
 #include "NiagaraCommon.h"
+#include "NiagaraDigestDatabase.h"
+#include "NiagaraGraphDigestTypes.h"
 #include "NiagaraHlslTranslator.h"
 #include "NiagaraParameterMapHistory.h"
 #include "NiagaraScriptSourceBase.h"
@@ -12,11 +14,59 @@
 #include "Templates/SharedPointer.h"
 
 class FCompileConstantResolver;
+class FNiagaraCompilationGraph;
+class FNiagaraCompilationNodeEmitter;
 class FNiagaraCompileRequestData;
 class FNiagaraCompileRequestDuplicateData;
+struct FNiagaraCompilationDigestBridge;
+class FNiagaraFixedConstantResolver;
+enum class ENiagaraFunctionDebugState : uint8;
 class UNiagaraGraph;
 struct FNiagaraGraphCachedBuiltHistory;
+template<typename GraphInterface> class TNiagaraHlslTranslator;
 class UNiagaraScriptSourceBase;
+struct FNiagaraSystemCompilationTask;
+
+// a version of the FCompileConstantResolver that is immutable (the constants are evaluated at the time of
+// construction rather than on demand).  Useful for digested graphs where their state is also immutable and
+// means we don't need to hold and reference UObjects (like UNiagaraSystem & UNiagaraEmitter)
+class FNiagaraFixedConstantResolver
+{
+public:
+	using FTranslator = TNiagaraHlslTranslator<FNiagaraCompilationDigestBridge>;
+
+	FNiagaraFixedConstantResolver();
+	FNiagaraFixedConstantResolver(const FTranslator* InTranslator, ENiagaraScriptUsage InUsage = ENiagaraScriptUsage::Function, ENiagaraFunctionDebugState InDebugState = ENiagaraFunctionDebugState::NoDebug);
+	FNiagaraFixedConstantResolver(const FCompileConstantResolver& SrcConstantResolver);
+
+	bool ResolveConstant(FNiagaraVariable& OutConstant) const;
+
+	FNiagaraFixedConstantResolver WithDebugState(ENiagaraFunctionDebugState InDebugState) const;
+	FNiagaraFixedConstantResolver WithUsage(ENiagaraScriptUsage ScriptUsage) const;
+
+	ENiagaraFunctionDebugState GetDebugState() const;
+
+private:
+	void InitConstants();
+	void SetScriptUsage(ENiagaraScriptUsage ScriptUsage);
+	void SetDebugState(ENiagaraFunctionDebugState DebugState);
+
+	const FTranslator* Translator = nullptr;
+
+	enum class EResolvedConstant : uint8
+	{
+		FunctionDebugState = 0,
+		ScriptUsage,
+		ScriptContext,
+		EmitterLocalspace,
+		EmitterDeterminism,
+		EmitterInterpolatedSpawn,
+		EmitterSimulationTarget,
+		Count
+	};
+
+	TArray<FNiagaraVariable, TFixedAllocator<(uint8)EResolvedConstant::Count>> ResolvedConstants;
+};
 
 struct FNiagaraSimulationStageInfo
 {
@@ -29,6 +79,129 @@ struct FNiagaraSimulationStageInfo
 	FNiagaraSimulationStageCompilationData CompilationData;
 };
 
+//////////////////////////////////////////////////////////////////////////
+// Helper structures for the FNiagaraActiveCompilationAsyncTask
+//////////////////////////////////////////////////////////////////////////
+
+// holds the precompile data for the compilation of a NiagaraSystem
+class FNiagaraPrecompileData
+{
+public:
+	using FSharedPrecompileData = TSharedPtr<FNiagaraPrecompileData, ESPMode::ThreadSafe>;
+
+	void GatherPreCompiledVariables(const FString& InNamespaceFilter, TArray<FNiagaraVariable>& OutVars) const;
+	FName ResolveEmitterAlias(FName VariableName) const;
+
+	const FString& GetUniqueEmitterName() const { return EmitterUniqueName; }
+	void FinishPrecompile(
+		const FNiagaraSystemCompilationTask& CompilationTask,
+		TConstArrayView<FNiagaraVariable> EncounterableVariables,
+		TConstArrayView<FNiagaraVariable> InStaticVariables,
+		const FNiagaraFixedConstantResolver& ConstantResolver,
+		TConstArrayView<ENiagaraScriptUsage> UsagesToProcess,
+		TConstArrayView<FNiagaraSimulationStageInfo> SimStages,
+		TConstArrayView<FString> EmitterNames);
+
+	void CollectBakedRapidIterationParameters(const FNiagaraSystemCompilationTask& CompilationTask, TConstArrayView<TObjectKey<UNiagaraScript>> OwnedScriptKeys);
+
+	int32 GetDependentRequestCount() const
+	{
+		return EmitterData.Num();
+	};
+	FSharedPrecompileData GetDependentRequest(int32 Index)
+	{
+		return EmitterData[Index];
+	}
+	const FNiagaraPrecompileData* GetDependentRequest(int32 Index) const
+	{
+		return EmitterData[Index].Get();
+	}
+
+	bool GetUseRapidIterationParams() const { return bUseRapidIterationParams; }
+
+	bool GetDisableDebugSwitches() const { return bDisableDebugSwitches; }
+
+	bool Matches(const FNiagaraCompileRequestData& Other) const;
+
+	// Simulation Stage Variables. Sim stage of 0 is always Spawn/Update
+	TArray<FNiagaraSimulationStageCompilationData> CompileSimStageData;
+
+	struct FCompileDataInterfaceData
+	{
+		FString EmitterName;
+		ENiagaraScriptUsage Usage;
+		FGuid UsageId;
+		FNiagaraVariable Variable;
+		TArray<FString> ReadsEmitterParticleData;
+	};
+	TSharedPtr<TArray<FCompileDataInterfaceData>> SharedCompileDataInterfaceData;
+
+	TArray<FNiagaraVariable> EncounteredVariables;
+	FString EmitterUniqueName;
+	TArray<FSharedPrecompileData> EmitterData;
+	FNiagaraCompilationGraphHandle SourceGraphHandle;
+	FString SourceName;
+	bool bUseRapidIterationParams = true;
+	bool bDisableDebugSwitches = false;
+
+	UEnum* ENiagaraScriptCompileStatusEnum = nullptr;
+	UEnum* ENiagaraScriptUsageEnum = nullptr;
+
+	TMap<FGraphTraversalHandle, FString> PinToConstantValues;
+	TArray<FNiagaraVariable> StaticVariables;
+	TArray<FNiagaraVariable> StaticVariablesWithMultipleWrites;
+
+	TArray<FNiagaraVariable> BakedRapidIterationParameters;
+
+	static void SortOutputNodesByDependencies(TArray<const FNiagaraCompilationNodeOutput*>& NodesToSort, TConstArrayView<FNiagaraSimulationStageInfo> SimStages);
+
+	template<typename T> T GetStaticVariableValue(const FNiagaraVariableBase Variable, const T DefaultValue) const
+	{
+		const int32 Index = StaticVariables.Find(Variable);
+		return Index != INDEX_NONE && StaticVariables[Index].IsDataAllocated() ? StaticVariables[Index].GetValue<T>() : DefaultValue;
+	}
+};
+
+
+class FNiagaraCompilationCopyData
+{
+public:
+	using FSharedCompilationCopy = TSharedPtr<FNiagaraCompilationCopyData, ESPMode::ThreadSafe>;
+	using FParameterMapHistory = TNiagaraParameterMapHistory<FNiagaraCompilationDigestBridge>;
+	using FParameterMapHistoryWithMetaDataBuilder = TNiagaraParameterMapHistoryWithMetaDataBuilder<FNiagaraCompilationDigestBridge>;
+
+	const TMap<FName, UNiagaraDataInterface*>& GetObjectNameMap();
+	UNiagaraDataInterface* GetDuplicatedDataInterfaceCDOForClass(UClass* Class) const;
+
+	TArray<FParameterMapHistory>& GetPrecomputedHistories() { return PrecompiledHistories; }
+	const TArray<FParameterMapHistory>& GetPrecomputedHistories() const { return PrecompiledHistories; }
+
+	void InstantiateCompilationCopy(FNiagaraCompilationGraphHandle GraphHandle, const FNiagaraPrecompileData* PrecompileData, ENiagaraScriptUsage InUsage, FNiagaraFixedConstantResolver ConstantResolver);
+	void CreateParameterMapHistory(const FNiagaraSystemCompilationTask& CompilationTask, const TArray<FNiagaraVariable>& EncounterableVariables, const TArray<FNiagaraVariable>& InStaticVariables, FNiagaraFixedConstantResolver ConstantResolver, TConstArrayView<FNiagaraSimulationStageInfo> SimStages);
+
+	int32 GetDependentRequestCount() const { return EmitterData.Num(); }
+	FSharedCompilationCopy GetDependentRequest(int32 Index) { return EmitterData[Index]; }
+
+	TArray<ENiagaraScriptUsage> ValidUsages;
+
+	TSharedPtr<FNiagaraCompilationGraph, ESPMode::ThreadSafe> InstantiatedGraph;
+
+	TArray<FParameterMapHistory> PrecompiledHistories;
+
+	TArray<FNiagaraVariable> ChangedFromNumericVars;
+	FString EmitterUniqueName;
+
+	TArray<FSharedCompilationCopy> EmitterData;
+
+	TMap<UClass*, UNiagaraDataInterface*> AggregatedDataInterfaceCDODuplicates;
+};
+
+//////////////////////////////////////////////////////////////////////////
+// Helper structures for the FNiagaraActiveCompilationDefault
+//////////////////////////////////////////////////////////////////////////
+
+// Implements the FNiagaraCompileRequestDataBase interface that is used to represent the precompile data
+// in order to perform the compilation of a Niagara asset.
 class FNiagaraCompileRequestData : public FNiagaraCompileRequestDataBase
 {
 public:
@@ -98,6 +271,8 @@ public:
 	static void SortOutputNodesByDependencies(TArray<class UNiagaraNodeOutput*>& NodesToSort, const TArray<class UNiagaraSimulationStageBase*>* SimStages);
 };
 
+// Implements the FNiagaraCompileRequestDuplicateDataBase interface that is used to represent the data that
+// is duplicated (for example graphs) in order to perform the compilation of a Niagara asset.
 class FNiagaraCompileRequestDuplicateData : public FNiagaraCompileRequestDuplicateDataBase
 {
 public:
diff --git a/Engine/Plugins/FX/Niagara/Source/NiagaraEditor/Private/NiagaraCompiler.cpp b/Engine/Plugins/FX/Niagara/Source/NiagaraEditor/Private/NiagaraCompiler.cpp
--- a/Engine/Plugins/FX/Niagara/Source/NiagaraEditor/Private/NiagaraCompiler.cpp
+++ b/Engine/Plugins/FX/Niagara/Source/NiagaraEditor/Private/NiagaraCompiler.cpp
@@ -1429,7 +1429,8 @@ int32 FNiagaraEditorModule::CompileScript(const FNiagaraCompileRequestDataBase*
 	UE_LOG(LogNiagaraEditor, Verbose, TEXT("Compiling System %s ..................................................................."), *InCompileOptions.FullName);
 
 	FNiagaraCompileResults Results;
-	TSharedPtr<FHlslNiagaraCompiler> Compiler = MakeShared<FHlslNiagaraCompiler>();
+	FActiveCompilation ActiveCompilation;
+	ActiveCompilation.Compiler = MakeShared<FHlslNiagaraCompiler>();
 	FNiagaraTranslateResults TranslateResults;
 	TUniquePtr<INiagaraHlslTranslator> Translator = INiagaraHlslTranslator::CreateTranslator(CompileRequest, CompileRequestDuplicate);
 
@@ -1458,6 +1459,8 @@ int32 FNiagaraEditorModule::CompileScript(const FNiagaraCompileRequestDataBase*
 	{
 		TranslateResults = Translator->Translate(InCompileOptions, TranslateOptions);
 	}
+
+	ActiveCompilation.TranslationTime = (float) (FPlatformTime::Seconds() - TranslationStartTime);
 	UE_LOG(LogNiagaraEditor, Verbose, TEXT("Translating System %s took %f sec."), *InCompileOptions.FullName, (float)(FPlatformTime::Seconds() - TranslationStartTime));
 
 	if (GbForceNiagaraTranslatorDump != 0)
@@ -1472,8 +1475,8 @@ int32 FNiagaraEditorModule::CompileScript(const FNiagaraCompileRequestDataBase*
 	NiagaraCompileRequestHelper::FDebugGroupNameBuilder DebugGroupName;
 	NiagaraCompileRequestHelper::BuildScriptDebugGroupName(CompileRequest, InCompileOptions, DebugGroupName);
 
-	int32 JobID = Compiler->CompileScript(DebugGroupName, InCompileOptions, TranslateResults, Translator->GetTranslateOutput(), Translator->GetTranslatedHLSL());
-	ActiveCompilations.Add(JobID, Compiler);
+	int32 JobID = ActiveCompilation.Compiler->CompileScript(DebugGroupName, InCompileOptions, TranslateResults, Translator->GetTranslateOutput(), Translator->GetTranslatedHLSL());
+	ActiveCompilations.Add(JobID, ActiveCompilation);
 	return JobID;
 }
 
@@ -1626,19 +1629,22 @@ TSharedPtr<FNiagaraGraphCachedDataBase, ESPMode::ThreadSafe> FNiagaraEditorModul
 	return CachedPtr;
 }
 
-TSharedPtr<FNiagaraVMExecutableData> FNiagaraEditorModule::GetCompilationResult(int32 JobID, bool bWait)
+TSharedPtr<FNiagaraVMExecutableData> FNiagaraEditorModule::GetCompilationResult(int32 JobID, bool bWait, FNiagaraScriptCompileMetrics& ScriptMetrics)
 {
 	TRACE_CPUPROFILER_EVENT_SCOPE(FNiagaraEditorModule::GetCompilationResult);
 
-	TSharedPtr<FHlslNiagaraCompiler>* MapEntry = ActiveCompilations.Find(JobID);
-	check(MapEntry && *MapEntry);
+	FActiveCompilation* MapEntry = ActiveCompilations.Find(JobID);
+	check(MapEntry && MapEntry->Compiler.IsValid());
 
-	TSharedPtr<FHlslNiagaraCompiler> Compiler = *MapEntry;
+	TSharedPtr<FHlslNiagaraCompiler> Compiler = MapEntry->Compiler;
 	TOptional<FNiagaraCompileResults> CompileResult = Compiler->GetCompileResult(JobID, bWait);
 	if (!CompileResult)
 	{
 		return TSharedPtr<FNiagaraVMExecutableData>();
 	}
+
+	ScriptMetrics.TranslateTime = MapEntry->TranslationTime;
+
 	ActiveCompilations.Remove(JobID);
 	FNiagaraCompileResults& Results = CompileResult.GetValue();
 
@@ -1666,6 +1672,11 @@ TSharedPtr<FNiagaraVMExecutableData> FNiagaraEditorModule::GetCompilationResult(
 		// When there are no errors the compile events get emptied, so add them back here.
 		Results.Data->LastCompileEvents.Append(Results.CompileEvents);
 	}
+
+	ScriptMetrics.CompilerWallTime = CompileResult->CompilerWallTime;
+	ScriptMetrics.CompilerPreprocessTime = CompileResult->CompilerPreprocessTime;
+	ScriptMetrics.CompilerWorkerTime = CompileResult->CompilerWorkerTime;
+
 	return CompileResult->Data;
 }
 
@@ -2138,7 +2149,10 @@ TOptional<FNiagaraCompileResults> FHlslNiagaraCompiler::GetCompileResult(int32 J
 			break;
 			}
 		}
-		Results.CompileTime = (float)(FPlatformTime::Seconds() - CompilationJob->StartTime);
+
+		Results.CompilerWallTime = (float)(FPlatformTime::Seconds() - CompilationJob->StartTime);
+		Results.CompilerWorkerTime = (float)CompilationJob->ShaderCompileJob->Output.CompileTime;
+		Results.CompilerPreprocessTime = (float)CompilationJob->ShaderCompileJob->Output.PreprocessTime;
 
 		Results.Data->CalledVMExternalFunctions.Empty(CompilationOutput.CalledVMFunctionTable.Num());
 		for (FCalledVMFunction& FuncInfo : CompilationOutput.CalledVMFunctionTable)
diff --git a/Engine/Plugins/FX/Niagara/Source/NiagaraEditor/Private/NiagaraDigestDatabase.cpp b/Engine/Plugins/FX/Niagara/Source/NiagaraEditor/Private/NiagaraDigestDatabase.cpp
new file mode 100644
--- /dev/null
+++ b/Engine/Plugins/FX/Niagara/Source/NiagaraEditor/Private/NiagaraDigestDatabase.cpp
@@ -0,0 +1,258 @@
+// Copyright Epic Games, Inc. All Rights Reserved.
+
+#include "NiagaraDigestDatabase.h"
+
+#include "Misc/LazySingleton.h"
+#include "NiagaraEditorModule.h"
+#include "NiagaraGraph.h"
+#include "NiagaraGraphDigest.h"
+#include "NiagaraNodeFunctionCall.h"
+#include "NiagaraParameterCollection.h"
+#include "NiagaraScript.h"
+#include "NiagaraScriptSource.h"
+
+namespace NiagaraGraphDigestDatabaseImpl
+{
+
+// todo - two tiered cache for digested objects (see NiagaraDigestDatabase.h)
+//static int GDigestGraphCacheSize = 256;
+//static FAutoConsoleVariableRef CVarDigestGraphCacheSize(
+//	TEXT("fx.Niagara.DigestGraphCacheSize"),
+//	GDigestGraphCacheSize,
+//	TEXT("Defines the size of the cache for digested Niagara graphs."),
+//	ECVF_ReadOnly
+//);
+//
+//static int GDigestCollectionCacheSize = 32;
+//static FAutoConsoleVariableRef CVarDigestCollectionCacheSize(
+//	TEXT("fx.Niagara.DigestCollectionCacheSize"),
+//	GDigestCollectionCacheSize,
+//	TEXT("Defines the size of the cache for digested Niagara parameter collections."),
+//	ECVF_ReadOnly
+//);
+
+FNiagaraDigestDatabase* GDigestDatabase = nullptr;
+
+} // NiagaraGraphDigestDatabaseImpl::
+
+FNiagaraDigestDatabase& FNiagaraDigestDatabase::Get()
+{
+	if (!NiagaraGraphDigestDatabaseImpl::GDigestDatabase)
+	{
+		NiagaraGraphDigestDatabaseImpl::GDigestDatabase = &TLazySingleton<FNiagaraDigestDatabase>::Get();
+	}
+
+	return *NiagaraGraphDigestDatabaseImpl::GDigestDatabase;
+}
+
+void FNiagaraDigestDatabase::Shutdown()
+{
+	Get().ReleaseDatabase();
+}
+
+FNiagaraDigestDatabase::FNiagaraDigestDatabase()
+	// todo - two tiered cache for digested objects (see NiagaraDigestDatabase.h)
+	//: CompilationGraphCache(NiagaraGraphDigestDatabaseImpl::GDigestGraphCacheSize)
+	//, CompilationNPCCache(NiagaraGraphDigestDatabaseImpl::GDigestCollectionCacheSize)
+{
+
+}
+
+FNiagaraDigestDatabase::~FNiagaraDigestDatabase()
+{
+	ReleaseDatabase();
+}
+
+void FNiagaraDigestDatabase::ReleaseDatabase()
+{
+	FWriteScopeLock WriteScope(DigestCacheLock);
+
+	CompilationGraphCache.Empty();
+	CompilationNPCCache.Empty();
+}
+
+void FNiagaraDigestDatabase::AddReferencedObjects(FReferenceCollector& Collector)
+{
+	FReadScopeLock ReadScope(DigestCacheLock);
+
+	for (FCompilationGraphCache::TIterator It(CompilationGraphCache); It; ++It)
+	{
+		It.Value()->AddReferencedObjects(Collector);
+	}
+
+	for (FCompilationNPCCache::TIterator It(CompilationNPCCache); It; ++It)
+	{
+		It.Value()->AddReferencedObjects(Collector);
+	}
+}
+
+FString FNiagaraDigestDatabase::GetReferencerName() const
+{
+	static const FString ReferencerName = TEXT("NiagaraDigestDatabsae");
+	return ReferencerName;
+}
+
+//////////////////////////////////////////////////////////////////////////
+/// Digested graphs
+
+FNiagaraCompilationGraphHandle FNiagaraDigestDatabase::CreateCompilationCopy(const UNiagaraGraph* Graph)
+{
+	check(IsInGameThread());
+	FNiagaraCompilationGraphHandle GraphHash(Graph);
+	FNiagaraCompilationGraph* PendingGraph = nullptr;
+
+	{
+		FWriteScopeLock WriteScope(DigestCacheLock);
+
+		if (const FNiagaraCompilationGraphHandle::FGraphPtr* CompilationGraph = CompilationGraphCache.Find(GraphHash))
+		{
+			++GraphCacheHits;
+			return GraphHash;
+		}
+
+		++GraphCacheMisses;
+		FNiagaraCompilationGraphHandle::FGraphPtr CompilationGraph = MakeShared<FNiagaraCompilationGraphHandle::FGraphPtr::ElementType, FNiagaraCompilationGraphHandle::FGraphPtr::Mode>();
+		PendingGraph = CompilationGraph.Get();
+
+		CompilationGraphCache.Add(GraphHash, CompilationGraph);
+	}
+
+	if (PendingGraph)
+	{
+		PendingGraph->Create(Graph);
+	}
+
+	return GraphHash;
+}
+
+FNiagaraCompilationGraphHandle::FGraphPtr FNiagaraDigestDatabase::Resolve(const FNiagaraCompilationGraphHandle& Handle) const
+{
+	FNiagaraCompilationGraphHandle::FGraphPtr GraphRef;
+
+	{
+		FReadScopeLock ReadScope(DigestCacheLock);
+		GraphRef = CompilationGraphCache.FindRef(Handle);
+	}
+
+	check(GraphRef.IsValid());
+
+	return GraphRef;
+}
+
+FNiagaraCompilationGraphHandle::FNiagaraCompilationGraphHandle(const UNiagaraScriptSourceBase* ScriptSourceBase)
+{
+	if (const UNiagaraScriptSource* ScriptSource = CastChecked<const UNiagaraScriptSource>(ScriptSourceBase))
+	{
+		if (const UNiagaraGraph* Graph = ScriptSource->NodeGraph)
+		{
+			AssetKey = Graph;
+			Hash = Graph->GetChangeID();
+		}
+	}
+}
+
+FNiagaraCompilationGraphHandle::FNiagaraCompilationGraphHandle(const UNiagaraGraph* Graph)
+{
+	AssetKey = Graph;
+	Hash = Graph ? Graph->GetChangeID() : FGuid();
+}
+
+FNiagaraCompilationGraphHandle::FGraphPtr FNiagaraCompilationGraphHandle::Resolve() const
+{
+	return FNiagaraDigestDatabase::Get().Resolve(*this);
+}
+
+//////////////////////////////////////////////////////////////////////////
+/// Digested parameter collection
+
+FNiagaraCompilationNPCHandle FNiagaraDigestDatabase::CreateCompilationCopy(const UNiagaraParameterCollection* Collection)
+{
+	check(IsInGameThread());
+	FNiagaraCompilationNPCHandle CollectionHash(Collection);
+	FNiagaraCompilationNPC* PendingCollection = nullptr;
+
+	{
+		FWriteScopeLock WriteScope(DigestCacheLock);
+
+		if (const FNiagaraCompilationNPCHandle::FDigestPtr* CompilationCollection = CompilationNPCCache.Find(CollectionHash))
+		{
+			++CollectionCacheHits;
+			return CollectionHash;
+		}
+
+		++CollectionCacheMisses;
+		FNiagaraCompilationNPCHandle::FDigestPtr CompilationCollection = MakeShared<FNiagaraCompilationNPCHandle::FDigestPtr::ElementType, FNiagaraCompilationNPCHandle::FDigestPtr::Mode>();
+		PendingCollection = CompilationCollection.Get();
+
+		CompilationNPCCache.Add(CollectionHash, CompilationCollection);
+	}
+
+	if (PendingCollection)
+	{
+		PendingCollection->Create(Collection);
+	}
+
+	return CollectionHash;
+}
+
+FNiagaraCompilationNPCHandle::FDigestPtr FNiagaraDigestDatabase::Resolve(const FNiagaraCompilationNPCHandle& Handle) const
+{
+	FReadScopeLock ReadScope(DigestCacheLock);
+
+	return CompilationNPCCache.FindRef(Handle);
+}
+
+void FNiagaraCompilationNPC::Create(const UNiagaraParameterCollection* Collection)
+{
+	SourceCollection = Collection;
+	Namespace = Collection->GetNamespace();
+	CollectionPath = FSoftObjectPath(Collection).ToString();
+	CollectionFullName = GetFullNameSafe(Collection);
+
+	Variables.Reserve(Collection->GetParameters().Num());
+	for (const FNiagaraVariable& Parameter : Collection->GetParameters())
+	{
+		Variables.Emplace(Parameter);
+	}
+
+	{ // we also need to deal with any data interfaces that might be stored in the NPC
+		UPackage* TransientPackage = GetTransientPackage();
+
+		const FNiagaraParameterStore& DefaultParamStore = Collection->GetDefaultInstance()->GetParameterStore();
+		for (const FNiagaraVariableBase& Variable : Variables)
+		{
+			const int32 VariableOffset = DefaultParamStore.IndexOf(Variable);
+			if (VariableOffset != INDEX_NONE)
+			{
+				if (UNiagaraDataInterface* DefaultDataInterface = DefaultParamStore.GetDataInterface(VariableOffset))
+				{
+					UNiagaraDataInterface* DuplicateDataInterface = DuplicateObject<UNiagaraDataInterface>(DefaultDataInterface, TransientPackage);
+
+					DefaultDataInterfaces.Add(Variable, DuplicateDataInterface);
+				}
+			}
+		}
+	}
+}
+
+void FNiagaraCompilationNPC::AddReferencedObjects(FReferenceCollector& Collector)
+{
+	Collector.AddStableReferenceMap(DefaultDataInterfaces);
+}
+
+UNiagaraDataInterface* FNiagaraCompilationNPC::GetDataInterface(const FNiagaraVariableBase& Variable) const
+{
+	return DefaultDataInterfaces.FindRef(Variable);
+}
+
+FNiagaraCompilationNPCHandle::FNiagaraCompilationNPCHandle(const UNiagaraParameterCollection* Connection)
+	: Namespace(Connection ? Connection->GetNamespace() : NAME_None)
+{
+	AssetKey = Connection;
+	Hash = Connection->GetCompileHash();
+}
+
+FNiagaraCompilationNPCHandle::FDigestPtr FNiagaraCompilationNPCHandle::Resolve() const
+{
+	return FNiagaraDigestDatabase::Get().Resolve(*this);
+}
\ No newline at end of file
diff --git a/Engine/Plugins/FX/Niagara/Source/NiagaraEditor/Private/NiagaraDigestDatabase.h b/Engine/Plugins/FX/Niagara/Source/NiagaraEditor/Private/NiagaraDigestDatabase.h
new file mode 100644
--- /dev/null
+++ b/Engine/Plugins/FX/Niagara/Source/NiagaraEditor/Private/NiagaraDigestDatabase.h
@@ -0,0 +1,142 @@
+// Copyright Epic Games, Inc. All Rights Reserved.
+
+#pragma once
+
+#include "CoreMinimal.h"
+#include "NiagaraCompileHash.h"
+#include "NiagaraGraphDigest.h"
+#include "UObject/GCObject.h"
+
+class UNiagaraCompilationGraph;
+class UNiagaraGraph;
+class UNiagaraParameterCollection;
+class UNiagaraScriptSourceBase;
+
+// base class used to reference a digested version of a UObject.  
+template<typename AssetType, typename DigestType, typename HashType>
+class FNiagaraDigestedAssetHandle
+{
+public:
+	using FDigestPtr = TSharedPtr<DigestType, ESPMode::ThreadSafe>;
+	FNiagaraDigestedAssetHandle() = default;
+	virtual ~FNiagaraDigestedAssetHandle() = default;
+	FNiagaraDigestedAssetHandle(const FNiagaraDigestedAssetHandle& Handle) = default;
+
+	bool IsValid() const
+	{
+		return Resolve().IsValid();
+	}
+	virtual FDigestPtr Resolve() const = 0;
+
+protected:
+	TObjectKey<AssetType> AssetKey;
+	HashType Hash;
+
+	friend FORCEINLINE uint32 GetTypeHash(const FNiagaraDigestedAssetHandle& CompilationCopy)
+	{
+		return HashCombine(GetTypeHash(CompilationCopy.AssetKey), GetTypeHash(CompilationCopy.Hash));
+	}
+
+	friend FORCEINLINE bool operator==(const FNiagaraDigestedAssetHandle& Lhs, const FNiagaraDigestedAssetHandle& Rhs)
+	{
+		return Lhs.AssetKey == Rhs.AssetKey && Lhs.Hash == Rhs.Hash;
+	}
+};
+
+class FNiagaraCompilationGraphHandle : public FNiagaraDigestedAssetHandle<UNiagaraGraph, FNiagaraCompilationGraph, FGuid>
+{
+public:
+	using Super = FNiagaraDigestedAssetHandle<UNiagaraGraph, FNiagaraCompilationGraph, FGuid>;
+	using FGraphPtr = Super::FDigestPtr;
+
+	FNiagaraCompilationGraphHandle() = default;
+	FNiagaraCompilationGraphHandle(const FNiagaraCompilationGraphHandle& Handle) = default;
+	FNiagaraCompilationGraphHandle(const UNiagaraScriptSourceBase* ScriptSource);
+	FNiagaraCompilationGraphHandle(const UNiagaraGraph* Graph);
+
+	virtual FGraphPtr Resolve() const override;
+};
+
+class FNiagaraCompilationNPC
+{
+public:
+	void Create(const UNiagaraParameterCollection* Collection);
+	void AddReferencedObjects(FReferenceCollector& Collector);
+
+	UNiagaraDataInterface* GetDataInterface(const FNiagaraVariableBase& Variable) const;
+
+	TArray<FNiagaraVariableBase> Variables;
+	FName Namespace;
+	TWeakObjectPtr<const UNiagaraParameterCollection> SourceCollection;
+	FString CollectionPath;
+	FString CollectionFullName;
+	TMap<FNiagaraVariableBase, UNiagaraDataInterface*> DefaultDataInterfaces;
+};
+
+class FNiagaraCompilationNPCHandle : public FNiagaraDigestedAssetHandle<UNiagaraParameterCollection, FNiagaraCompilationNPC, FNiagaraCompileHash>
+{
+public:
+	using Super = FNiagaraDigestedAssetHandle<UNiagaraParameterCollection, FNiagaraCompilationNPC, FNiagaraCompileHash>;
+
+	FNiagaraCompilationNPCHandle() = default;
+	FNiagaraCompilationNPCHandle(const FNiagaraCompilationNPCHandle& Handle) = default;
+	FNiagaraCompilationNPCHandle(const UNiagaraParameterCollection* Collection);
+
+	virtual FDigestPtr Resolve() const override;
+
+	FName Namespace = NAME_None;
+};
+
+class FNiagaraDigestedParameterCollections
+{
+public:
+	TArray<FNiagaraCompilationNPCHandle>& EditCollections() { return Collections; };
+
+	FNiagaraCompilationNPCHandle FindMatchingCollection(FName VariableName, bool bAllowPartialMatch, FNiagaraVariable& OutVar) const;
+	FNiagaraCompilationNPCHandle FindCollection(const FNiagaraVariable& Variable) const;
+
+protected:
+	FNiagaraCompilationNPCHandle FindCollectionByName(FName VariableName) const;
+
+	TArray<FNiagaraCompilationNPCHandle> Collections;
+};
+
+class FNiagaraDigestDatabase : public FGCObject
+{
+public:
+	FNiagaraDigestDatabase();
+	virtual ~FNiagaraDigestDatabase();
+
+	static FNiagaraDigestDatabase& Get();
+	static void Shutdown();
+
+	FNiagaraCompilationGraphHandle CreateCompilationCopy(const UNiagaraGraph* Graph);
+	FNiagaraCompilationGraphHandle::FGraphPtr Resolve(const FNiagaraCompilationGraphHandle& Handle) const;
+
+	FNiagaraCompilationNPCHandle CreateCompilationCopy(const UNiagaraParameterCollection* Collection);
+	FNiagaraCompilationNPCHandle::FDigestPtr Resolve(const FNiagaraCompilationNPCHandle& Handle) const;
+
+	void ReleaseDatabase();
+
+	virtual void AddReferencedObjects(FReferenceCollector& Collector) override;
+	
+	virtual FString GetReferencerName() const override;
+
+protected:
+	uint32 GraphCacheHits = 0;
+	uint32 GraphCacheMisses = 0;
+	uint32 CollectionCacheHits = 0;
+	uint32 CollectionCacheMisses = 0;
+
+	// todo - replace these maps with some kind of two tiered system where we've got an TLruCache backing
+	// those objects we've digested recently and then another layer of digested objects that we've got pinned
+	// for the lifetime of a compilation task.  For now we'll just have an ever growing map
+
+	using FCompilationGraphCache = TMap<FNiagaraCompilationGraphHandle, FNiagaraCompilationGraphHandle::FGraphPtr>;
+	FCompilationGraphCache CompilationGraphCache;
+
+	using FCompilationNPCCache = TMap<FNiagaraCompilationNPCHandle, FNiagaraCompilationNPCHandle::FDigestPtr>;
+	FCompilationNPCCache CompilationNPCCache;
+
+	mutable FRWLock DigestCacheLock;
+};
\ No newline at end of file
