Adds two new compilation modes for Niagara systems - AsyncTasks & Verify.

diff --git a/Engine/Plugins/FX/Niagara/Source/NiagaraEditor/Private/NiagaraEditorModule.cpp b/Engine/Plugins/FX/Niagara/Source/NiagaraEditor/Private/NiagaraEditorModule.cpp
--- a/Engine/Plugins/FX/Niagara/Source/NiagaraEditor/Private/NiagaraEditorModule.cpp
+++ b/Engine/Plugins/FX/Niagara/Source/NiagaraEditor/Private/NiagaraEditorModule.cpp
@@ -40,6 +40,7 @@
 #include "TypeEditorUtilities/NiagaraMatrixTypeEditorUtilities.h"
 #include "TypeEditorUtilities/NiagaraDataInterfaceCurveTypeEditorUtilities.h"
 
+#include "NiagaraSystemCompilingManager.h"
 #include "NiagaraEditorStyle.h"
 #include "NiagaraEditorCommands.h"
 #include "PropertyEditorModule.h"
@@ -62,6 +63,7 @@
 #include "TNiagaraGraphPinEditableName.h"
 #include "UObject/Class.h"
 #include "NiagaraScriptMergeManager.h"
+#include "NiagaraDigestDatabase.h"
 #include "NiagaraEmitter.h"
 #include "NiagaraScriptSource.h"
 #include "NiagaraTypes.h"
@@ -947,6 +949,9 @@ void FNiagaraEditorModule::OnPreExit()
 			Sys->KillAllActiveCompilations();
 		}
 	}
+
+	FNiagaraDigestDatabase::Shutdown();
+
 	ClearObjectPool();
 	
 	INiagaraDataInterfaceNodeActionProvider::Unregister<UNiagaraDataInterfaceDataChannelWrite>();
@@ -1311,9 +1316,9 @@ void FNiagaraEditorModule::StartupModule()
 		return CompileScript(CompileRequest, CompileRequestDuplicate, Options);
 	}));
 
-	CompileResultHandle = NiagaraModule.RegisterCompileResultDelegate(INiagaraModule::FCheckCompilationResult::CreateLambda([this](int32 JobID, bool bWait)
+	CompileResultHandle = NiagaraModule.RegisterCompileResultDelegate(INiagaraModule::FCheckCompilationResult::CreateLambda([this](int32 JobID, bool bWait, FNiagaraScriptCompileMetrics& ScriptMetrics)
 	{
-		return GetCompilationResult(JobID, bWait);
+		return GetCompilationResult(JobID, bWait, ScriptMetrics);
 	}));
 
 	PrecompilerHandle = NiagaraModule.RegisterPrecompiler(INiagaraModule::FOnPrecompile::CreateLambda([this](UObject* InObj, FGuid Version)
@@ -1327,10 +1332,24 @@ void FNiagaraEditorModule::StartupModule()
 	}));
 
 	GraphCacheTraversalHandle = NiagaraModule.RegisterGraphTraversalCacher(INiagaraModule::FOnCacheGraphTraversal::CreateLambda([this](const UObject* InObj, FGuid Version)
-		{
-			return CacheGraphTraversal(InObj, Version);
-		}));
+	{
+		return CacheGraphTraversal(InObj, Version);
+	}));
 
+	RequestCompileSystemHandle = NiagaraModule.RegisterRequestCompileSystem(INiagaraModule::FOnRequestCompileSystem::CreateLambda([this](UNiagaraSystem* System, bool bForced)
+	{
+		return RequestCompileSystem(System, bForced);
+	}));
+
+	PollSystemCompileHandle = NiagaraModule.RegisterPollSystemCompile(INiagaraModule::FOnPollSystemCompile::CreateLambda([this](FNiagaraCompilationTaskHandle TaskHandle, FNiagaraSystemAsyncCompileResults& Results, bool bWait, bool bPeek)
+	{
+		return PollSystemCompile(TaskHandle, Results, bWait, bPeek);
+	}));
+
+	AbortSystemCompileHandle = NiagaraModule.RegisterAbortSystemCompile(INiagaraModule::FOnAbortSystemCompile::CreateLambda([this](FNiagaraCompilationTaskHandle TaskHandle)
+	{
+		AbortSystemCompile(TaskHandle);
+	}));
 
 	TestCompileScriptCommand = IConsoleManager::Get().RegisterConsoleCommand(
 		TEXT("fx.TestCompileNiagaraScript"),
@@ -1395,6 +1414,8 @@ void FNiagaraEditorModule::StartupModule()
 	//Register node action providers for data interface functions.
 	INiagaraDataInterfaceNodeActionProvider::Register<UNiagaraDataInterfaceDataChannelWrite, FNiagaraDataInterfaceNodeActionProvider_DataChannelWrite>();
 	INiagaraDataInterfaceNodeActionProvider::Register<UNiagaraDataInterfaceDataChannelRead, FNiagaraDataInterfaceNodeActionProvider_DataChannelRead>();
+
+	FAssetCompilingManager::Get().RegisterManager(&FNiagaraSystemCompilingManager::Get());
 }
 
 void FNiagaraEditorModule::ShutdownModule()
@@ -1458,6 +1479,9 @@ void FNiagaraEditorModule::ShutdownModule()
 		NiagaraModule->UnregisterPrecompiler(PrecompilerHandle);
 		NiagaraModule->UnregisterPrecompileDuplicator(PrecompileDuplicatorHandle);
 		NiagaraModule->UnregisterGraphTraversalCacher(GraphCacheTraversalHandle);
+		NiagaraModule->UnregisterRequestCompileSystem(RequestCompileSystemHandle);
+		NiagaraModule->UnregisterPollSystemCompile(PollSystemCompileHandle);
+		NiagaraModule->UnregisterAbortSystemCompile(AbortSystemCompileHandle);
 	}
 
 	// Verify that we've cleaned up all the view models in the world.
@@ -1511,6 +1535,8 @@ void FNiagaraEditorModule::ShutdownModule()
 #endif
 
 	GraphDataCache.Reset();
+
+	FAssetCompilingManager::Get().UnregisterManager(&FNiagaraSystemCompilingManager::Get());
 }
 
 void FNiagaraEditorModule::OnPostEngineInit()
diff --git a/Engine/Plugins/FX/Niagara/Source/NiagaraEditor/Private/NiagaraHlslTranslator.cpp b/Engine/Plugins/FX/Niagara/Source/NiagaraEditor/Private/NiagaraHlslTranslator.cpp
--- a/Engine/Plugins/FX/Niagara/Source/NiagaraEditor/Private/NiagaraHlslTranslator.cpp
+++ b/Engine/Plugins/FX/Niagara/Source/NiagaraEditor/Private/NiagaraHlslTranslator.cpp
@@ -18,6 +18,7 @@
 #include "NiagaraFunctionLibrary.h"
 #include "NiagaraNodeEmitter.h"
 #include "NiagaraGraph.h"
+#include "NiagaraGraphDigest.h"
 #include "NiagaraNode.h"
 #include "NiagaraNodeConvert.h"
 #include "NiagaraNodeCustomHlsl.h"
@@ -37,11 +38,13 @@
 #include "NiagaraShared.h"
 #include "NiagaraSimulationStageBase.h"
 #include "NiagaraTrace.h"
+#include "NiagaraTraversalStateContext.h"
 #include "ShaderCore.h"
 #include "Misc/FileHelper.h"
 #include "Misc/CoreMiscDefines.h"
 #include "NiagaraDataInterfaceUtilities.h"
 #include "NiagaraShader.h"
+#include "HAL/Platform.h"
 
 #define LOCTEXT_NAMESPACE "NiagaraCompiler"
 
@@ -398,6 +401,42 @@ FString FNiagaraHlslTranslator::GetCodeAsSource(int32 ChunkIdx)
 }
 
 // specialization because validation may be different between graph implementations
+template<>
+bool TNiagaraHlslTranslator<FNiagaraCompilationDigestBridge>::ValidateTypePins(const FNode* NodeToValidate)
+{
+	bool bPinsAreValid = true;
+	for (const FPin* Pin : GetPins(NodeToValidate))
+	{
+		if (Pin->PinType.PinCategory.IsNone())
+		{
+			Error(LOCTEXT("InvalidPinTypeError", "Node pin has an undefined type."), NodeToValidate, Pin);
+			bPinsAreValid = false;
+		}
+		else if (Pin->PinType.PinCategory == UEdGraphSchema_Niagara::PinCategoryType || Pin->PinType.PinCategory == UEdGraphSchema_Niagara::PinCategoryStaticType)
+		{
+			FNiagaraTypeDefinition Type = Pin->Variable.GetType();
+			if (Type.IsValid() == false)
+			{
+				Error(LOCTEXT("InvalidPinTypeError", "Node pin has an undefined type."), NodeToValidate, Pin);
+				bPinsAreValid = false;
+			}
+			else if (Type == FNiagaraTypeDefinition::GetGenericNumericDef())
+			{
+				Error(LOCTEXT("NumericPinError", "A numeric pin was not resolved to a known type.  Numeric pins must be connected or must be converted to an explicitly typed pin in order to compile."), NodeToValidate, Pin);
+				bPinsAreValid = false;
+			}
+		}
+
+		// disabling this check because orphaned pins won't make it this far for the digested graphs
+		//if (Pin->bOrphanedPin)
+		//{
+		//	Warning(LOCTEXT("OrphanedPinError", "Node pin is no longer valid.  This pin must be disconnected or reset to default so it can be removed."), NodeToValidate, Pin);
+		//}
+	}
+	return bPinsAreValid;
+}
+
+
 template<>
 bool TNiagaraHlslTranslator<FNiagaraCompilationGraphBridge>::ValidateTypePins(const FNode* NodeToValidate)
 {
@@ -4586,6 +4625,11 @@ bool FNiagaraHlslTranslator::ShouldInterpolateParameter(const FNiagaraVariable&
 }
 
 // specialization handling the case where the translation process treats the graph as mutable
+template<>
+void TNiagaraHlslTranslator<FNiagaraCompilationDigestBridge>::UpdateStaticSwitchConstants(const FPin* Pin)
+{
+}
+
 template<>
 void TNiagaraHlslTranslator<FNiagaraCompilationGraphBridge>::UpdateStaticSwitchConstants(const FPin* ChildPin)
 {
@@ -6284,6 +6328,29 @@ private:
 	typename GraphBridge::FParamMapHistoryBuilder& Builder;
 };
 
+template<>
+class FScopedBuilderEmitter<FNiagaraCompilationDigestBridge>
+{
+public:
+	FScopedBuilderEmitter(TNiagaraParameterMapHistoryBuilder<FNiagaraCompilationDigestBridge>& InBuilder, const FNiagaraCompilationNodeEmitter* InEmitterNode)
+		: EmitterNode(InEmitterNode)
+		, Builder(InBuilder)
+	{
+		Builder.TraversalStateContext->PushEmitter(EmitterNode);
+		Builder.EnterEmitter(EmitterNode->EmitterUniqueName, EmitterNode->CalledGraph.Get(), EmitterNode);
+	}
+
+	~FScopedBuilderEmitter()
+	{
+		Builder.ExitEmitter(EmitterNode->EmitterUniqueName, EmitterNode);
+		Builder.TraversalStateContext->PopEmitter(EmitterNode);
+	}
+
+private:
+	const FNiagaraCompilationNodeEmitter* EmitterNode;
+	TNiagaraParameterMapHistoryBuilder<FNiagaraCompilationDigestBridge>& Builder;
+};
+
 template<typename GraphBridge>
 void TNiagaraHlslTranslator<GraphBridge>::Emitter(const FEmitterNode* EmitterNode, TArray<int32>& Inputs, TArray<int32>& Outputs)
 {
@@ -7103,7 +7170,7 @@ void TNiagaraHlslTranslator<GraphBridge>::WriteDataSet(const FNiagaraDataSetID D
 }
 
 template<typename GraphBridge>
-int32 TNiagaraHlslTranslator<GraphBridge>::RegisterUObject(FNiagaraVariable Variable, UObject* Object, bool bAddParameterMapRead)
+int32 TNiagaraHlslTranslator<GraphBridge>::RegisterUObject(const FNiagaraVariable& Variable, UObject* Object, bool bAddParameterMapRead)
 {
 	int32 ObjectIndex = INDEX_NONE;
 	if (bAddParameterMapRead)
@@ -7148,6 +7215,52 @@ int32 TNiagaraHlslTranslator<GraphBridge>::RegisterUObject(FNiagaraVariable Vari
 	return ObjectIndex;
 }
 
+template<typename GraphBridge>
+int32 TNiagaraHlslTranslator<GraphBridge>::RegisterUObjectPath(const FNiagaraVariable& Variable, const FSoftObjectPath& ObjectPath, bool bAddParameterMapRead)
+{
+	int32 ObjectIndex = INDEX_NONE;
+	if (bAddParameterMapRead)
+	{
+		// If we are registering a read then look for an existing one by name
+		ObjectIndex = CompilationOutput.ScriptData.UObjectInfos.IndexOfByPredicate([&](const FNiagaraScriptUObjectCompileInfo& ExistingInfo) { return ExistingInfo.Variable == Variable; });
+	}
+	// Assume all writes are unique
+	//else if ( Object != nullptr )
+	//{
+	//	// If we are registering a write then look for an existing one by value
+	//	ObjectIndex = CompilationOutput.ScriptData.UObjectInfos.IndexOfByPredicate([&](const FNiagaraScriptUObjectCompileInfo& ExistingInfo) { return ExistingInfo.Object == Object; });
+	//}
+
+	FNiagaraScriptUObjectCompileInfo* CompileInfo = nullptr;
+	if (ObjectIndex == INDEX_NONE)
+	{
+		ObjectIndex = CompilationOutput.ScriptData.UObjectInfos.AddDefaulted();
+		CompileInfo = &CompilationOutput.ScriptData.UObjectInfos[ObjectIndex];
+		CompileInfo->Variable = Variable;
+		CompileInfo->ObjectPath = ObjectPath;
+	}
+	else
+	{
+		CompileInfo = &CompilationOutput.ScriptData.UObjectInfos[ObjectIndex];
+	}
+
+	if (bAddParameterMapRead)
+	{
+		FName ResolvedName;
+		if (FNiagaraParameterUtilities::IsAliasedEmitterParameter(Variable.GetName().ToString()))
+		{
+			ResolvedName = ActiveHistoryForFunctionCalls.ResolveAliases(Variable).GetName();
+		}
+		else
+		{
+			ResolvedName = Variable.GetName();
+		}
+		CompileInfo->RegisteredParameterMapRead = ResolvedName;
+	}
+
+	return ObjectIndex;
+}
+
 template<typename GraphBridge>
 int32 TNiagaraHlslTranslator<GraphBridge>::RegisterDataInterface(const FNiagaraVariable& Var, const UNiagaraDataInterface* DataInterface, bool bPlaceholder, bool bAddParameterMapRead)
 {
@@ -7368,6 +7481,29 @@ private:
 	FNiagaraCompilationGraphBridge::FParamMapHistoryBuilder& Builder;
 };
 
+template<>
+class FScopedBuilderFunctionCall<FNiagaraCompilationDigestBridge>
+{
+public:
+	FScopedBuilderFunctionCall(TNiagaraHlslTranslator<FNiagaraCompilationDigestBridge>* InTranslator, FNiagaraCompilationDigestBridge::FParamMapHistoryBuilder& InBuilder, const FNiagaraCompilationDigestBridge::FFunctionCallNode* InFunctionCallNode)
+		: FunctionCallNode(InFunctionCallNode)
+		, Builder(InBuilder)
+	{
+		Builder.TraversalStateContext->PushFunction(FunctionCallNode, FNiagaraFixedConstantResolver(InTranslator));
+		Builder.EnterFunction(FunctionCallNode->FunctionName, FunctionCallNode->CalledGraph.Get(), FunctionCallNode);
+	}
+
+	~FScopedBuilderFunctionCall()
+	{
+		Builder.ExitFunction(FunctionCallNode);
+		Builder.TraversalStateContext->PopFunction(FunctionCallNode);
+	}
+
+private:
+	const FNiagaraCompilationDigestBridge::FFunctionCallNode* FunctionCallNode;
+	FNiagaraCompilationDigestBridge::FParamMapHistoryBuilder& Builder;
+};
+
 template<typename GraphBridge>
 void TNiagaraHlslTranslator<GraphBridge>::FunctionCall(const FFunctionCallNode* FunctionNode, TArray<int32>& Inputs, TArray<int32>& Outputs)
 {
@@ -7799,6 +7935,44 @@ bool TNiagaraHlslTranslator<GraphBridge>::ParseDIFunctionSpecifiers(const FNode*
 	return true;
 }
 
+template<typename GraphBridge>
+bool PartialParticleUpdateHelper(TConstArrayView<FString> Tokens, const typename GraphBridge::FPrecompileData& PrecompileData)
+{
+	static const FString UseParticleReadTokens[] =
+	{
+		TEXT("InputDataFloat"),
+		TEXT("InputDataInt"),
+		TEXT("InputDataBool"),
+		TEXT("InputDataHalf"),
+	};
+
+	for (const FString& Token : Tokens)
+	{
+		for (const FString& BannedToken : UseParticleReadTokens)
+		{
+			if (Token == BannedToken)
+			{
+				// Clear out the ability to use partial particle writes as we can't be sure how InputData is being used
+				for (const auto& CompileStageData : PrecompileData.CompileSimStageData)
+				{
+					CompileStageData.PartialParticleUpdate = false;
+				}
+
+				return true;
+			}
+		}
+	}
+
+	return false;
+}
+
+template<>
+bool PartialParticleUpdateHelper<FNiagaraCompilationDigestBridge>(TConstArrayView<FString> Tokens, const FNiagaraPrecompileData& PrecompileData)
+{
+	// does nothing; this is taken care of in a pre-transation step
+	return false;
+}
+
 template<typename GraphBridge>
 void TNiagaraHlslTranslator<GraphBridge>::ProcessCustomHlsl(const FString& InCustomHlsl, ENiagaraScriptUsage InUsage, const FNiagaraFunctionSignature& InSignature, const TArray<int32>& Inputs, const FNode* InNodeForErrorReporting, FString& OutCustomHlsl, FNiagaraFunctionSignature& OutSignature)
 {
@@ -7857,37 +8031,9 @@ void TNiagaraHlslTranslator<GraphBridge>::ProcessCustomHlsl(const FString& InCus
 	// Look for tokens that should be replaced with a data interface or not used directly
 	if (CompilationTarget != ENiagaraSimTarget::GPUComputeSim)
 	{
-		static const FString UseParticleReadTokens[] =
+		if (PartialParticleUpdateHelper<GraphBridge>(Tokens, *CompileData))
 		{
-			TEXT("InputDataFloat"),
-			TEXT("InputDataInt"),
-			TEXT("InputDataBool"),
-			TEXT("InputDataHalf"),
-		};
-
-		for (const FString& Token : Tokens)
-		{
-			bool bUsesInputData = false;
-			for (const FString& BannedToken : UseParticleReadTokens)
-			{
-				if (Token == BannedToken)
-				{
-					Warning(LOCTEXT("UseParticleReadsNotInputData", "Please convert usage of InputData methods to particle reads to avoid compatability issues."), InNodeForErrorReporting, nullptr);
-
-					bUsesInputData = true;
-					break;
-				}
-			}
-
-			if (bUsesInputData)
-			{
-				// Clear out the ability to use partial particle writes as we can't be sure how InputData is being used
-				for (const auto& CompileStageData : CompileData->CompileSimStageData)
-				{
-					CompileStageData.PartialParticleUpdate = false;
-				}
-				break;
-			}
+			Warning(LOCTEXT("UseParticleReadsNotInputData", "Please convert usage of InputData methods to particle reads to avoid compatibility issues."), InNodeForErrorReporting, nullptr);
 		}
 	}
 
@@ -10015,6 +10161,11 @@ TUniquePtr<INiagaraHlslTranslator> INiagaraHlslTranslator::CreateTranslator(cons
 	return MakeUnique<TNiagaraHlslTranslator<FNiagaraCompilationGraphBridge>>(CompileRequest, CompileRequestDuplicate);
 }
 
+TUniquePtr<INiagaraHlslTranslator> INiagaraHlslTranslator::CreateTranslator(const FNiagaraPrecompileData* InCompileData, const FNiagaraCompilationCopyData* InDuplicateData)
+{
+	return MakeUnique<TNiagaraHlslTranslator<FNiagaraCompilationDigestBridge>>(InCompileData, InDuplicateData);
+}
+
 template<typename GraphBridge>
 const FString& TNiagaraHlslTranslator<GraphBridge>::GetEmitterUniqueName() const
 {
@@ -10386,6 +10537,8 @@ TArray<FName> FNiagaraHlslTranslator::ConditionPropertyPath(const FNiagaraTypeDe
 
 // explicit template instantiation of the translators available based on the graph interfaces
 template class TNiagaraHlslTranslator<FNiagaraCompilationGraphBridge>;
+template class TNiagaraHlslTranslator<FNiagaraCompilationDigestBridge>;
+
 
 #undef NIAGARA_SCOPE_CYCLE_COUNTER
 #undef LOCTEXT_NAMESPACE
diff --git a/Engine/Plugins/FX/Niagara/Source/NiagaraEditor/Private/NiagaraHlslTranslator.h b/Engine/Plugins/FX/Niagara/Source/NiagaraEditor/Private/NiagaraHlslTranslator.h
--- a/Engine/Plugins/FX/Niagara/Source/NiagaraEditor/Private/NiagaraHlslTranslator.h
+++ b/Engine/Plugins/FX/Niagara/Source/NiagaraEditor/Private/NiagaraHlslTranslator.h
@@ -634,7 +634,8 @@ public:
 
 	int32 CompileInputPin(const FInputPin* Pin);
 
-	int32 RegisterUObject(FNiagaraVariable Variable, UObject* Object, bool bAddParameterMapRead);
+	int32 RegisterUObject(const FNiagaraVariable& Variable, UObject* Object, bool bAddParameterMapRead);
+	int32 RegisterUObjectPath(const FNiagaraVariable& Variable, const FSoftObjectPath& ObjectPath, bool bAddParmeterMapRead);
 	int32 RegisterDataInterface(const FNiagaraVariable& Var, const UNiagaraDataInterface* DataInterface, bool bPlaceholder, bool bAddParameterMapRead);
 
 	void Operation(const FOpNode* Operation, TArray<int32>& Inputs, TArray<int32>& Outputs);
diff --git a/Engine/Plugins/FX/Niagara/Source/NiagaraEditor/Private/NiagaraNodeStaticSwitch.h b/Engine/Plugins/FX/Niagara/Source/NiagaraEditor/Private/NiagaraNodeStaticSwitch.h
--- a/Engine/Plugins/FX/Niagara/Source/NiagaraEditor/Private/NiagaraNodeStaticSwitch.h
+++ b/Engine/Plugins/FX/Niagara/Source/NiagaraEditor/Private/NiagaraNodeStaticSwitch.h
@@ -101,14 +101,16 @@ public:
 	virtual ENiagaraNumericOutputTypeSelectionMode GetNumericOutputTypeSelectionMode() const override;
 	//~ End UNiagaraNode Interface
 
+	//~ Begin UNiagaraNodeUsageSelector Interface
+	virtual int32 GetNumberOfCases() const override;
+	// ~End UNiagaraNodeUsageSelector Interface
+
 #if WITH_EDITOR
 	virtual void PostEditChangeProperty(struct FPropertyChangedEvent& PropertyChangedEvent) override;
 #endif // WITH_EDITOR
 
 protected:
 	//~ Begin UNiagaraNodeUsageSelector Interface
-	virtual int32 GetNumberOfCases() const override;
-	
 	virtual bool AreInputCaseLabelsReadOnly() const override { return SwitchTypeData.SwitchType != ENiagaraStaticSwitchType::Integer || IsSetByPin(); }
 	virtual void OnInputCaseLabelSubmitted(const FText& Text, ETextCommit::Type Arg, int32 InputCase) const override;
 	virtual bool VerifyCaseLabelCandidate(const FText& InCandidate, FText& OutErrorMessage) const override;
diff --git a/Engine/Plugins/FX/Niagara/Source/NiagaraEditor/Private/NiagaraParameterMapHistory.cpp b/Engine/Plugins/FX/Niagara/Source/NiagaraEditor/Private/NiagaraParameterMapHistory.cpp
--- a/Engine/Plugins/FX/Niagara/Source/NiagaraEditor/Private/NiagaraParameterMapHistory.cpp
+++ b/Engine/Plugins/FX/Niagara/Source/NiagaraEditor/Private/NiagaraParameterMapHistory.cpp
@@ -6,12 +6,14 @@
 #include "INiagaraEditorTypeUtilities.h"
 #include "Modules/ModuleManager.h"
 #include "NiagaraCompilationBridge.h"
+#include "NiagaraCompilationPrivate.h"
 #include "NiagaraCompiler.h"
 #include "NiagaraConstants.h"
 #include "NiagaraEditorCommon.h"
 #include "NiagaraEditorModule.h"
 #include "NiagaraEditorUtilities.h"
 #include "NiagaraGraph.h"
+#include "NiagaraGraphDigest.h"
 #include "NiagaraHlslTranslator.h"
 #include "NiagaraNode.h"
 #include "NiagaraNodeEmitter.h"
@@ -91,6 +93,21 @@ void FGraphTraversalHandle::PushNode(const UEdGraphNode* Node)
 	}
 }
 
+void FGraphTraversalHandle::PushNode(const FNiagaraCompilationNode* Node)
+{
+	if (ensure(Node))
+	{
+		ensure(Node->NodeGuid.IsValid());
+		Path.Push(Node->NodeGuid);
+		FriendlyPath.Push(Node->FullTitle);
+	}
+	else
+	{
+		Path.Push(FGuid());
+		FriendlyPath.Push(FString());
+	}
+}
+
 void FGraphTraversalHandle::PushPin(const UEdGraphPin* Pin)
 {
 	if (ensure(Pin && Pin->GetOwningNode()))
@@ -120,6 +137,34 @@ void FGraphTraversalHandle::PushPin(const UEdGraphPin* Pin)
 	}
 }
 
+void FGraphTraversalHandle::PushPin(const FNiagaraCompilationPin* Pin)
+{
+	if (ensure(Pin && Pin->OwningNode))
+	{
+		PushNode(Pin->OwningNode);
+
+		if (Pin->PersistentGuid.IsValid())
+		{
+			Path.Push(Pin->PersistentGuid);
+		}
+		else
+		{
+			FGuid PinIndexGuid(Pin->SourcePinIndex, 0, 0, 0);
+			Path.Push(PinIndexGuid);
+		}
+
+		FriendlyPath.Push(Pin->OwningNode->FullTitle + TEXT("->") + Pin->PinName.ToString());
+	}
+	else
+	{
+		Path.Push(FGuid());
+		FriendlyPath.Push(FString());
+
+		Path.Push(FGuid());
+		FriendlyPath.Push(FString());
+	}
+}
+
 void FGraphTraversalHandle::PopNode()
 {
 	Path.Pop();
@@ -868,6 +913,30 @@ FNiagaraCompilationGraphBridge::FParameterCollection TNiagaraParameterMapHistory
 	return nullptr;
 }
 
+template<>
+FNiagaraCompilationDigestBridge::FParameterCollection TNiagaraParameterMapHistory<FNiagaraCompilationDigestBridge>::IsParameterCollectionParameter(FNiagaraVariable& InVar, bool& bMissingParameter)
+{
+	if (bParameterCollectionsSkipped)
+	{
+		UE_LOG(LogNiagaraEditor, Error, TEXT("NiagaraParameterCollection was skipped during history building.  History which require NPC data can not be generated during PostLoad()."));
+		return nullptr;
+	}
+
+	bMissingParameter = false;
+	FNameBuilder VarName(InVar.GetName());
+	for (const FNiagaraCompilationNPCHandle& CollectionHandle : EncounteredParameterCollections.Handles)
+	{
+		FNameBuilder CollectionName(CollectionHandle.Namespace);
+		if (VarName.ToView().StartsWith(CollectionName.ToView()))
+		{
+			bMissingParameter = CollectionHandle.Resolve()->Variables.Contains(InVar);
+			return CollectionHandle;
+		}
+	}
+
+	return FNiagaraCompilationNPCHandle();
+}
+
 template<typename GraphBridge>
 bool TNiagaraParameterMapHistory<GraphBridge>::ShouldIgnoreVariableDefault(const FNiagaraVariable& Var)const
 {
@@ -1589,6 +1658,11 @@ FNiagaraVariable TNiagaraParameterMapHistoryBuilder<GraphBridge>::ResolveAliases
 template<typename GraphBridge>
 void TNiagaraParameterMapHistoryBuilder<GraphBridge>::RegisterNodeVisitation(const FNode* Node)
 {
+	if (OnNodeVisitedDelegate.IsBound())
+	{
+		OnNodeVisitedDelegate.Broadcast(Node);
+	}
+
 	ContextuallyVisitedNodes.Last().AddUnique(Node);
 }
 
@@ -2088,6 +2162,19 @@ bool TNiagaraParameterMapHistoryBuilder<FNiagaraCompilationGraphBridge>::AddPara
 	return true;
 }
 
+
+template<>
+bool TNiagaraParameterMapHistoryBuilder<FNiagaraCompilationDigestBridge>::AddParameterCollection(FParameterCollection Collection, int32 ParamMapIndex)
+{
+	if (!Collection.IsValid())
+	{
+		return false;
+	}
+
+	Histories[ParamMapIndex].EncounteredParameterCollections.Handles.AddUnique(Collection);
+	return true;
+}
+
 template<typename GraphBridge>
 void TNiagaraParameterMapHistoryBuilder<GraphBridge>::SetConstantByStaticVariable(int32& OutValue, const FInputPin* InDefaultPin)
 {
@@ -2584,9 +2671,12 @@ void TNiagaraParameterMapHistoryWithMetaDataBuilder<GraphBridge>::AddGraphToCall
 //////////////////////////////////////////////////////////////////////////
 
 template struct TNiagaraParameterMapHistory<FNiagaraCompilationGraphBridge>;
+template struct TNiagaraParameterMapHistory<FNiagaraCompilationDigestBridge>;
 
 template class TNiagaraParameterMapHistoryBuilder<FNiagaraCompilationGraphBridge>;
+template class TNiagaraParameterMapHistoryBuilder<FNiagaraCompilationDigestBridge>;
 
 template class TNiagaraParameterMapHistoryWithMetaDataBuilder<FNiagaraCompilationGraphBridge>;
+template class TNiagaraParameterMapHistoryWithMetaDataBuilder<FNiagaraCompilationDigestBridge>;
 
 #undef LOCTEXT_NAMESPACE
diff --git a/Engine/Plugins/FX/Niagara/Source/NiagaraEditor/Public/INiagaraCompiler.h b/Engine/Plugins/FX/Niagara/Source/NiagaraEditor/Public/INiagaraCompiler.h
--- a/Engine/Plugins/FX/Niagara/Source/NiagaraEditor/Public/INiagaraCompiler.h
+++ b/Engine/Plugins/FX/Niagara/Source/NiagaraEditor/Public/INiagaraCompiler.h
@@ -11,21 +11,19 @@
 #include "NiagaraShared.h"
 
 class Error;
-class UEdGraphPin;
-class UNiagaraNode;
-class UNiagaraGraph;
-class UEdGraphPin;
 class FCompilerResultsLog;
-class UNiagaraDataInterface;
-struct FNiagaraTranslateResults;
-struct FNiagaraTranslatorOutput;
-struct FNiagaraVMExecutableData;
+class FNiagaraCompilationCopyData;
 class FNiagaraCompileOptions;
 class FNiagaraCompileRequestData;
 class FNiagaraCompileRequestDuplicateData;
 class FNiagaraPrecompileData;
-class FNiagaraCompilationCopyData;
-
+class UEdGraphPin;
+class UNiagaraDataInterface;
+class UNiagaraGraph;
+class UNiagaraNode;
+struct FNiagaraTranslateResults;
+struct FNiagaraTranslatorOutput;
+struct FNiagaraVMExecutableData;
 
 /** Defines information about the results of a Niagara script compile. */
 struct FNiagaraCompileResults
@@ -37,17 +35,16 @@ struct FNiagaraCompileResults
 	
 	/** The actual final compiled data.*/
 	TSharedPtr<FNiagaraVMExecutableData> Data;
-	float CompileTime = 0.0f;
+	float CompilerWallTime = 0.0f;
+	float CompilerPreprocessTime = 0.0f;
+	float CompilerWorkerTime = 0.0f;
+
 	/** Tracking any compilation warnings or errors that occur.*/
 	TArray<FNiagaraCompileEvent> CompileEvents;
 	uint32 NumErrors = 0;
 	uint32 NumWarnings = 0;
 	FString DumpDebugInfoPath;
  
- 
- 
- 
- 
 	static ENiagaraScriptCompileStatus CompileResultsToSummary(const FNiagaraCompileResults* CompileResults);
 	void AppendCompileEvents(TArrayView<const FNiagaraCompileEvent> InCompileEvents)
 	{
@@ -134,6 +131,7 @@ public:
 	virtual const FString& GetTranslatedHLSL() const = 0;
 
 	static TUniquePtr<INiagaraHlslTranslator> CreateTranslator(const FNiagaraCompileRequestDataBase* InCompileData, const FNiagaraCompileRequestDuplicateDataBase* InDuplicateData);
+	static TUniquePtr<INiagaraHlslTranslator> CreateTranslator(const FNiagaraPrecompileData* InCompileData, const FNiagaraCompilationCopyData* InDuplicateData);
 
 	static FString BuildHLSLStructDecl(const FNiagaraTypeDefinition& Type, FText& OutErrorMessage, bool bGpuScript);
 	static bool IsBuiltInHlslType(const FNiagaraTypeDefinition& Type);
diff --git a/Engine/Plugins/FX/Niagara/Source/NiagaraEditor/Public/NiagaraCompilationBridge.h b/Engine/Plugins/FX/Niagara/Source/NiagaraEditor/Public/NiagaraCompilationBridge.h
--- a/Engine/Plugins/FX/Niagara/Source/NiagaraEditor/Public/NiagaraCompilationBridge.h
+++ b/Engine/Plugins/FX/Niagara/Source/NiagaraEditor/Public/NiagaraCompilationBridge.h
@@ -4,36 +4,79 @@
 
 #include "NiagaraCommon.h"
 
-class FCompileConstantResolver;
-class FNiagaraCompileRequestData;
-class FNiagaraCompileRequestDuplicateData;
-class UEdGraphPin;
+// forward declarations for the GraphBridge
 class UNiagaraGraph;
-class UNiagaraNode;
-class UNiagaraNodeConvert;
+class UEdGraphPin;
 class UNiagaraNodeCustomHlsl;
-class UNiagaraNodeEmitter;
-class UNiagaraNodeFunctionCall;
-class UNiagaraNodeIf;
+class UNiagaraNode;
+class UNiagaraNodeOutput;
 class UNiagaraNodeInput;
 class UNiagaraNodeOp;
-class UNiagaraNodeOutput;
-class UNiagaraNodeParameterMapFor;
+class UNiagaraNodeEmitter;
+class UNiagaraNodeIf;
+class UNiagaraNodeConvert;
+class UNiagaraNodeSelect;
+class UNiagaraNodeFunctionCall;
 class UNiagaraNodeParameterMapGet;
 class UNiagaraNodeParameterMapSet;
-class UNiagaraNodeSelect;
+class UNiagaraNodeParameterMapFor;
 class UNiagaraNodeStaticSwitch;
-class UNiagaraParameterCollection;
+template<typename GraphBridge> struct TNiagaraParameterMapHistory;
+template<typename GraphBridge> class TNiagaraParameterMapHistoryBuilder;
+class FNiagaraCompileRequestData;
+class FNiagaraCompileRequestDuplicateData;
 struct FGraphTraversalHandle;
+struct FNiagaraGraphFunctionAliasContext;
 struct FNiagaraConvertConnection;
+class UNiagaraParameterCollection;
+template<typename PinType> struct TModuleScopedPin;
 struct FNiagaraCustomHlslInclude;
 struct FNiagaraFindInputNodeOptions;
-struct FNiagaraGraphFunctionAliasContext;
+class FCompileConstantResolver;
 struct FNiagaraStaticVariableSearchContext;
-template<typename GraphBridge> class TNiagaraParameterMapHistoryBuilder;
-template<typename GraphBridge> struct TNiagaraParameterMapHistory;
-template<typename PinType> struct TModuleScopedPin;
 
+// forward declarations for the DigestBridge
+class FNiagaraCompilationGraph;
+class FNiagaraCompilationPin;
+class FNiagaraCompilationNodeCustomHlsl;
+class FNiagaraCompilationNode;
+class FNiagaraCompilationNodeOutput;
+class FNiagaraCompilationNodeInput;
+class FNiagaraCompilationNodeOp;
+class FNiagaraCompilationNodeEmitter;
+class FNiagaraCompilationNodeIf;
+class FNiagaraCompilationNodeConvert;
+class FNiagaraCompilationNodeSelect;
+class FNiagaraCompilationNodeFunctionCall;
+class FNiagaraCompilationNodeParameterMapGet;
+class FNiagaraCompilationNodeParameterMapSet;
+class FNiagaraCompilationNodeParameterMapFor;
+class FNiagaraCompilationNodeStaticSwitch;
+class FNiagaraCompilationInputPin;
+class FNiagaraCompilationOutputPin;
+class FNiagaraPrecompileData;
+class FNiagaraCompilationCopyData;
+struct FNiagaraDigestFunctionAliasContext;
+struct FNiagaraCompilationCachedConnection;
+class FNiagaraCompilationNPCHandle;
+class FNiagaraFixedConstantResolver;
+struct FNiagaraTraversalStateContext;
+class FNiagaraDigestedParameterCollections;
+
+// these two structures represent a way to abstract the implementation details of the graph structure
+// that is being used for the different pieces of code related to compiling Niagara systems.  The GraphBridge
+// handles the standard editor objects for graphs & nodes (UNiagaraGraph and UNiagaraNode respectively) with
+// the remaining types following from that.  The DigestBridge represents the digested graphs and nodes that
+// we use to be able to compile Niagara systems async where the graphs & nodes are FNiagaraCompilationGraph
+// FNiagaraCompilationNode respectively.
+//
+// The bridge class is used as a template argument for the following classes (and supplementary functions/classes):
+//	-FNiagaraParameterMapHistory
+//	-FNiagaraParameterMapHistoryBuilder
+//	-FNiagaraHlslTranslator
+//
+// This allows us to have the logic for these classes unified between the different graph representations
+// without having to write two versions of the code.
 struct FNiagaraCompilationGraphBridge
 {
 	// base types
@@ -85,6 +128,8 @@ struct FNiagaraCompilationGraphBridge
 		UNiagaraParameterCollection* FindMatchingCollection(FName VariableName, bool bAllowPartialMatch, FNiagaraVariable& OutVar) const;
 	};
 
+	// Used as a base class for the FNiagaraParameterMapHistoryBuilder to store additional data that is specific
+	// to the graph representation
 	class FBuilderExtraData
 	{
 	public:
@@ -171,3 +216,130 @@ struct FNiagaraCompilationGraphBridge
 	static bool GetGraphReferencesStaticVariables(const FGraph* Graph, FNiagaraStaticVariableSearchContext& StaticVariableContext);
 	static const FEmitterNode* GetNodeAsEmitter(const FNode* Node);
 };
+
+struct FNiagaraCompilationDigestBridge
+{
+	using FGraph = FNiagaraCompilationGraph;
+	using FPin = FNiagaraCompilationPin;
+	using FCustomHlslNode = FNiagaraCompilationNodeCustomHlsl;
+	using FNode = FNiagaraCompilationNode;
+	using FOutputNode = FNiagaraCompilationNodeOutput;
+	using FInputNode = FNiagaraCompilationNodeInput;
+	using FOpNode = FNiagaraCompilationNodeOp;
+	using FEmitterNode = FNiagaraCompilationNodeEmitter;
+	using FIfNode = FNiagaraCompilationNodeIf;
+	using FConvertNode = FNiagaraCompilationNodeConvert;
+	using FSelectNode = FNiagaraCompilationNodeSelect;
+	using FFunctionCallNode = FNiagaraCompilationNodeFunctionCall;
+	using FParamMapGetNode = FNiagaraCompilationNodeParameterMapGet;
+	using FParamMapSetNode = FNiagaraCompilationNodeParameterMapSet;
+	using FParamMapForNode = FNiagaraCompilationNodeParameterMapFor;
+	using FStaticSwitchNode = FNiagaraCompilationNodeStaticSwitch;
+	using FInputPin = FNiagaraCompilationInputPin;
+	using FOutputPin = FNiagaraCompilationOutputPin;
+	using FParamMapHistory = TNiagaraParameterMapHistory<FNiagaraCompilationDigestBridge>;
+	using FParamMapHistoryBuilder = TNiagaraParameterMapHistoryBuilder<FNiagaraCompilationDigestBridge>;
+	using FPrecompileData = FNiagaraPrecompileData;
+	using FCompilationCopy = FNiagaraCompilationCopyData;
+	using FModuleScopedPin = TModuleScopedPin<FPin>;
+	using FGraphFunctionAliasContext = FNiagaraDigestFunctionAliasContext;
+	using FConvertConnection = FNiagaraCompilationCachedConnection;
+	using FParameterCollection = FNiagaraCompilationNPCHandle;
+	using FConstantResolver = FNiagaraFixedConstantResolver;
+
+	struct FParameterCollectionStore
+	{
+		void Append(const FParameterCollectionStore& Other);
+		TArray<FNiagaraCompilationNPCHandle> Handles;
+	};
+
+	// Used as a base class for the FNiagaraParameterMapHistoryBuilder to store additional data that is specific
+	// to the graph representation
+	class FBuilderExtraData
+	{
+	public:
+		FBuilderExtraData();
+		TPimplPtr<FNiagaraTraversalStateContext> TraversalStateContext;
+		TUniquePtr<FNiagaraDigestedParameterCollections> AvailableCollections;
+	};
+
+	static const FGraph* GetGraph(const FCompilationCopy* CompilationCopy);
+	static const FNode* GetOwningNode(const FPin* Pin);
+	static FNode* GetMutableOwningNode(const FPin* Pin);
+	static const FGraph* GetOwningGraph(const FNode* Node);
+	static bool CustomHlslReferencesTokens(const FCustomHlslNode* CustomNode, TConstArrayView<FStringView> TokenStrings);
+	static ENiagaraScriptUsage GetCustomHlslUsage(const FCustomHlslNode* CustomNode);
+	static FString GetCustomHlslString(const FCustomHlslNode* CustomNode);
+	static void GetCustomHlslIncludePaths(const FCustomHlslNode* CustomNode, TArray<FNiagaraCustomHlslInclude>& Includes);
+	static const TArray<FConvertConnection>& GetConvertConnections(const FConvertNode* ConvertNode);
+
+	// various cast functions
+	static const FFunctionCallNode* AsFunctionCallNode(const FNode* Node);
+	static const FInputNode* AsInputNode(const FNode* Node);
+	static const FParamMapGetNode* AsParamMapGetNode(const FNode* Node);
+	static const FCustomHlslNode* AsCustomHlslNode(const FNode* Node);
+	static const FParamMapSetNode* AsParamMapSetNode(const FNode* Node);
+
+	static bool GraphHasParametersOfType(const FGraph* Graph, const FNiagaraTypeDefinition& TypeDef);
+	static TArray<FNiagaraVariableBase> GraphGetStaticSwitchInputs(const FGraph* Graph);
+	static void FindOutputNodes(const FGraph* Graph, ENiagaraScriptUsage ScriptUsage, TArray<const FOutputNode*>& OutputNodes);
+	static void FindOutputNodes(const FGraph* Graph, TArray<const FOutputNode*>& OutputNodes);
+	static void BuildTraversal(const FGraph* Graph, const FNode* OutputNode, TArray<const FNode*>& TraversedNodes);
+	static const FGraph* GetEmitterGraph(const FEmitterNode* EmitterNode);
+	static FString GetEmitterUniqueName(const FEmitterNode* EmitterNode);
+	static ENiagaraScriptUsage GetEmitterUsage(const FEmitterNode* EmitterNode);
+	static FString GetEmitterName(const FEmitterNode* EmitterNode);
+	static FString GetEmitterPathName(const FEmitterNode* EmitterNode);
+	static FString GetEmitterHandleIdString(const FEmitterNode* EmitterNode);
+	static const FGraph* GetFunctionNodeGraph(const FFunctionCallNode* FunctionCall);
+	static FString GetFunctionFullName(const FFunctionCallNode* FunctionCall);
+	static FString GetFunctionScriptName(const FFunctionCallNode* FunctionCall);
+	static FString GetFunctionName(const FFunctionCallNode* FunctionCall);
+	static ENiagaraScriptUsage GetFunctionUsage(const FFunctionCallNode* FunctionCall);
+	static TOptional<ENiagaraDefaultMode> GetGraphDefaultMode(const FGraph* Graph, const FNiagaraVariable& Variable, FNiagaraScriptVariableBinding& Binding);
+	static const FInputPin* GetDefaultPin(const FParamMapGetNode* GetNode, const FOutputPin* OutputPin);
+	static bool IsStaticPin(const FPin* Pin);
+	// retrieves all input pins (excluding any add pins that may be present)
+	static TArray<const FInputPin*> GetInputPins(const FNode* Node);
+	// retrieves all output pins (excluding both orphaned pins and add pins)
+	static TArray<const FOutputPin*> GetOutputPins(const FNode* Node);
+	// gets all pins assoicated with the node
+	static TArray<const FPin*> GetPins(const FNode* Node);
+	static FNiagaraTypeDefinition GetPinType(const FPin* Pin, ENiagaraStructConversion Conversion);
+	static FText GetPinFriendlyName(const FPin* Pin);
+	static FText GetPinDisplayName(const FPin* Pin);
+	static FNiagaraVariable GetPinVariable(const FPin* Pin, bool bNeedsValue, ENiagaraStructConversion Conversion);
+	static const FInputPin* GetPinAsInput(const FPin* Pin);
+	static FNiagaraVariable GetInputVariable(const FInputNode* InputNode);
+	static const TArray<FNiagaraVariable>& GetOutputVariables(const FOutputNode* OutputNode);
+	static TArray<FNiagaraVariable> GetGraphOutputNodeVariables(const FGraph* Graph, ENiagaraScriptUsage Usage);
+	static TArray<const FInputNode*> GetGraphInputNodes(const FGraph* Graph, const FNiagaraFindInputNodeOptions& Options);
+	static const FOutputPin* GetLinkedOutputPin(const FInputPin* InputPin);
+	static bool CanCreateConnection(const FOutputPin* OutputPin, const FInputPin* InputPin, FText& FailureMessage);
+	static ENiagaraScriptUsage GetOutputNodeUsage(const FOutputNode* OutputNode);
+	static FGuid GetOutputNodeUsageId(const FOutputNode* OutputNode);
+	static ENiagaraScriptUsage GetOutputNodeScriptType(const FOutputNode* OutputNode);
+	static FGuid GetOutputNodeScriptTypeId(const FOutputNode* OutputNode);
+	static bool IsGraphEmpty(const FGraph* Graph);
+	static void AddCollectionPaths(const FParamMapHistory& History, TArray<FString>& Paths);
+	static bool NodeIsEnabled(const FNode* Node);
+	static TOptional<ENiagaraDefaultMode> GraphGetDefaultMode(const FGraph* Graph, const FNiagaraVariableBase& Variable, FNiagaraScriptVariableBinding& Binding);
+	static const FOutputPin* GetSelectOutputPin(const FSelectNode* SelectNode, const FNiagaraVariableBase& Variable);
+	static FString GetNodeName(const FNode* Node);
+	static FString GetNodeTitle(const FNode* Node);
+	static const FInputPin* GetInputPin(const FNode* Node, int32 PinIndex);
+	static int32 GetPinIndexById(TConstArrayView<const FInputPin*> Pins, const FGuid& PinId);
+	static int32 GetPinIndexById(TConstArrayView<const FOutputPin*> Pins, const FGuid& PinId);
+	static FString GetCollectionFullName(FParameterCollection Collection);
+	static bool IsCollectionValid(FParameterCollection Collection);
+	static UNiagaraDataInterface* GetCollectionDataInterface(FParameterCollection Collection, const FNiagaraVariable& Variable);
+	static UObject* GetCollectionUObject(FParameterCollection Collection, const FNiagaraVariable& Variable);
+
+	static const FOutputNode* AsOutputNode(const FNode* Node);
+	static bool IsParameterMapGet(const FNode* Node);
+	static TOptional<FName> GetOutputNodeStackContextOverride(const FOutputNode* OutputNode);
+	static FString GetNodeClassName(const FNode* Node);
+	static bool IsParameterMapPin(const FPin* Pin);
+	static bool GetGraphReferencesStaticVariables(const FGraph* Graph, FNiagaraStaticVariableSearchContext& StaticVariableContext);
+	static const FEmitterNode* GetNodeAsEmitter(const FNode* Node);
+};
diff --git a/Engine/Plugins/FX/Niagara/Source/NiagaraEditor/Public/NiagaraEditorModule.h b/Engine/Plugins/FX/Niagara/Source/NiagaraEditor/Public/NiagaraEditorModule.h
--- a/Engine/Plugins/FX/Niagara/Source/NiagaraEditor/Public/NiagaraEditorModule.h
+++ b/Engine/Plugins/FX/Niagara/Source/NiagaraEditor/Public/NiagaraEditorModule.h
@@ -9,6 +9,7 @@
 #include "NiagaraTypes.h"
 #include "INiagaraCompiler.h"
 #include "AssetTypeCategories.h"
+#include "NiagaraCompilationTypes.h"
 #include "NiagaraPerfBaseline.h"
 #include "NiagaraDebuggerCommon.h"
 #include "NiagaraRendererProperties.h"
@@ -48,6 +49,7 @@ class UNiagaraParameterDefinitions;
 class UNiagaraReservedParametersManager;
 class FNiagaraGraphDataCache;
 class UNiagaraParameterCollection;
+struct FNiagaraSystemAsyncCompileResults;
 
 DECLARE_STATS_GROUP(TEXT("Niagara Editor"), STATGROUP_NiagaraEditor, STATCAT_Advanced);
 
@@ -146,7 +148,7 @@ public:
 
 	/** Start the compilation of the specified script. */
 	virtual int32 CompileScript(const FNiagaraCompileRequestDataBase* InCompileRequest, const FNiagaraCompileRequestDuplicateDataBase* InCompileRequestDuplicate, const FNiagaraCompileOptions& InCompileOptions);
-	virtual TSharedPtr<FNiagaraVMExecutableData> GetCompilationResult(int32 JobID, bool bWait);
+	virtual TSharedPtr<FNiagaraVMExecutableData> GetCompilationResult(int32 JobID, bool bWait, FNiagaraScriptCompileMetrics& ScriptMetrics);
 
 	TSharedPtr<FNiagaraCompileRequestDataBase, ESPMode::ThreadSafe> Precompile(UObject* InObj, FGuid Version);
 	TSharedPtr<FNiagaraCompileRequestDuplicateDataBase, ESPMode::ThreadSafe> PrecompileDuplicate(
@@ -157,6 +159,9 @@ public:
 		FGuid TargetVersion);
 	TSharedPtr<FNiagaraGraphCachedDataBase, ESPMode::ThreadSafe> CacheGraphTraversal(const UObject* Obj, FGuid Version);
 
+	FNiagaraCompilationTaskHandle RequestCompileSystem(UNiagaraSystem* System, bool bForce);
+	bool PollSystemCompile(FNiagaraCompilationTaskHandle, FNiagaraSystemAsyncCompileResults&, bool /*bWait*/, bool /*bPeek*/);
+	void AbortSystemCompile(FNiagaraCompilationTaskHandle);
 
 	/** Gets the extensibility managers for outside entities to extend static mesh editor's menus and toolbars */
 	virtual TSharedPtr<FExtensibilityManager> GetMenuExtensibilityManager() override {return MenuExtensibilityManager;}
@@ -369,6 +374,9 @@ private:
 	FDelegateHandle PrecompilerHandle;
 	FDelegateHandle PrecompileDuplicatorHandle;
 	FDelegateHandle GraphCacheTraversalHandle;
+	FDelegateHandle RequestCompileSystemHandle;
+	FDelegateHandle PollSystemCompileHandle;
+	FDelegateHandle AbortSystemCompileHandle;
 
 	FDelegateHandle DeviceProfileManagerUpdatedHandle;
 
@@ -413,7 +421,13 @@ private:
 
 	IConsoleCommand* ReinitializeStyleCommand;
 
-	TMap<int32, TSharedPtr<FHlslNiagaraCompiler>> ActiveCompilations;
+	struct FActiveCompilation
+	{
+		TSharedPtr<FHlslNiagaraCompiler> Compiler;
+		float TranslationTime = 0.0f;
+	};
+
+	TMap<int32, FActiveCompilation> ActiveCompilations;
 
 	TArray<TSharedRef<const FDeferredDestructionContainerBase>> EnqueuedForDeferredDestruction;
 
diff --git a/Engine/Plugins/FX/Niagara/Source/NiagaraEditor/Public/NiagaraParameterMapHistory.h b/Engine/Plugins/FX/Niagara/Source/NiagaraEditor/Public/NiagaraParameterMapHistory.h
--- a/Engine/Plugins/FX/Niagara/Source/NiagaraEditor/Public/NiagaraParameterMapHistory.h
+++ b/Engine/Plugins/FX/Niagara/Source/NiagaraEditor/Public/NiagaraParameterMapHistory.h
@@ -13,6 +13,8 @@
 
 class UEdGraphPin;
 class UEdGraphNode;
+class FNiagaraCompilationNode;
+class FNiagaraCompilationPin;
 class UNiagaraGraph;
 class FNiagaraHlslTranslator;
 class UNiagaraNodeEmitter;
@@ -88,9 +90,11 @@ public:
 	}
 
 	void PushNode(const UEdGraphNode* Node);
+	void PushNode(const FNiagaraCompilationNode* Node);
 	void PopNode();
 
 	void PushPin(const UEdGraphPin* Pin);
+	void PushPin(const FNiagaraCompilationPin* Pin);
 	void PopPin();
 
 	bool operator==(const FGraphTraversalHandle& Rhs) const
@@ -388,6 +392,8 @@ public:
 	using FConstantResolver = typename GraphBridge::FConstantResolver;
 	using FParameterCollection = typename GraphBridge::FParameterCollection;
 
+	DECLARE_MULTICAST_DELEGATE_OneParam(FOnNodeVisited, const FNode* Node);
+
 	/** Collection of the build histories from the graph traversal.*/
 	TArray<FHistory> Histories;
 
@@ -593,6 +599,8 @@ public:
 
 	bool ShouldProcessDepthTraversal(const FGraph* Graph);
 
+	FOnNodeVisited OnNodeVisitedDelegate;
+
 protected:
 	/** Internal helper function to decide whether or not we should use this static variable when merging with other
 		graph traversal static variables, like Emitter graphs merging with System graphs. */
diff --git a/Engine/Plugins/FX/Niagara/Source/NiagaraShader/Public/NiagaraShared.h b/Engine/Plugins/FX/Niagara/Source/NiagaraShader/Public/NiagaraShared.h
--- a/Engine/Plugins/FX/Niagara/Source/NiagaraShader/Public/NiagaraShared.h
+++ b/Engine/Plugins/FX/Niagara/Source/NiagaraShader/Public/NiagaraShared.h
@@ -147,7 +147,7 @@ struct NIAGARASHADER_API FNiagaraDataInterfaceGeneratedFunction
 	using FunctionSpecifier = TTuple<FName, FName>;
 	TArray<FunctionSpecifier> Specifiers;
 
-	uint16 MiscUsageBitMask;
+	uint16 MiscUsageBitMask = 0;
 
 	UPROPERTY()
 	TArray<FNiagaraVariableCommonReference> VariadicInputs;
