- Fix for light / decal renderer in LWC space

diff --git a/Engine/Plugins/FX/Niagara/Source/Niagara/Private/NiagaraRendererDecals.cpp b/Engine/Plugins/FX/Niagara/Source/Niagara/Private/NiagaraRendererDecals.cpp
--- a/Engine/Plugins/FX/Niagara/Source/Niagara/Private/NiagaraRendererDecals.cpp
+++ b/Engine/Plugins/FX/Niagara/Source/Niagara/Private/NiagaraRendererDecals.cpp
@@ -152,9 +152,8 @@ FNiagaraDynamicDataBase* FNiagaraRendererDecals::GenerateDynamicData(const FNiag
 	const UNiagaraDecalRendererProperties* RendererProperties = CastChecked<const UNiagaraDecalRendererProperties>(InProperties);
 
 	const bool bUseLocalSpace = UseLocalSpace(Proxy);
-	const FTransform LocalToWorld = SystemInstance->GetWorldTransform();
-	const FNiagaraLWCConverter LwcConverter = SystemInstance->GetLWCConverter(bUseLocalSpace);
-	const FNiagaraPosition DefaultPos = bUseLocalSpace ? FVector::ZeroVector : LocalToWorld.GetLocation();
+	const FTransform SimToWorld = SystemInstance->GetLWCSimToWorld(bUseLocalSpace);
+	const FNiagaraPosition DefaultPos = bUseLocalSpace ? FVector::ZeroVector : SystemInstance->GetWorldTransform().GetLocation();
 	const FQuat4f DefaultRot = UNiagaraDecalRendererProperties::GetDefaultOrientation();
 	const FVector3f DefaultSize = UNiagaraDecalRendererProperties::GetDefaultDecalSize();
 	const FNiagaraBool DefaultVisible = UNiagaraDecalRendererProperties::GetDefaultDecalVisible();
@@ -217,8 +216,8 @@ FNiagaraDynamicDataBase* FNiagaraRendererDecals::GenerateDynamicData(const FNiag
 			// Grab Decal Attributes
 			const FVector3f SimPos = PositionReader.GetSafe(ParticleIndex, LocalDefaultSimPos);
 			const FQuat4f SimRot = RotationReader.GetSafe(ParticleIndex, LocalDefaultSimRot).GetNormalized();
-			const FVector Position = bUseLocalSpace ? LocalToWorld.TransformPosition(FVector(SimPos)) : LwcConverter.ConvertSimulationPositionToWorld(SimPos);
-			const FQuat Rotation = bUseLocalSpace ? LocalToWorld.TransformRotation(FQuat(SimRot)) : FQuat(SimRot);
+			const FVector Position = SimToWorld.TransformPosition(FVector(SimPos));
+			const FQuat Rotation = SimToWorld.TransformRotation(FQuat(SimRot));
 			const FVector Size = FVector(SizeReader.GetSafe(ParticleIndex, LocalDefaultSize) * 0.5f);
 			const FLinearColor DecalColor = ColorReader.GetSafe(ParticleIndex, LocalDefaultDecalColor);
 			const float Fade = FadeReader.GetSafe(ParticleIndex, LocalDefaultFade);
@@ -257,8 +256,8 @@ FNiagaraDynamicDataBase* FNiagaraRendererDecals::GenerateDynamicData(const FNiag
 		{
 			const FVector3f SimPos = LocalDefaultSimPos;
 			const FQuat4f SimRot = LocalDefaultSimRot;
-			const FVector Position = bUseLocalSpace ? LocalToWorld.TransformPosition(FVector(SimPos)) : LwcConverter.ConvertSimulationPositionToWorld(SimPos);
-			const FQuat Rotation = bUseLocalSpace ? LocalToWorld.TransformRotation(FQuat(SimRot)) : FQuat(SimRot);
+			const FVector Position = SimToWorld.TransformPosition(FVector(SimPos));
+			const FQuat Rotation = SimToWorld.TransformRotation(FQuat(SimRot));
 			const FVector Size = FVector(LocalDefaultSize * 0.5f);
 			const FLinearColor DecalColor = LocalDefaultDecalColor;
 			const float Fade = LocalDefaultFade;
diff --git a/Engine/Plugins/FX/Niagara/Source/Niagara/Private/NiagaraRendererLights.cpp b/Engine/Plugins/FX/Niagara/Source/Niagara/Private/NiagaraRendererLights.cpp
--- a/Engine/Plugins/FX/Niagara/Source/Niagara/Private/NiagaraRendererLights.cpp
+++ b/Engine/Plugins/FX/Niagara/Source/Niagara/Private/NiagaraRendererLights.cpp
@@ -16,7 +16,6 @@
 #include "PrimitiveViewRelevance.h"
 #include "SceneInterface.h"
 
-
 DECLARE_CYCLE_STAT(TEXT("Generate Particle Lights"), STAT_NiagaraGenLights, STATGROUP_Niagara);
 
 struct FNiagaraDynamicDataLights : public FNiagaraDynamicDataBase
@@ -98,13 +97,12 @@ FNiagaraDynamicDataBase* FNiagaraRendererLights::GenerateDynamicData(const FNiag
 
 	// This used to use Proxy->GetLocalToWorld(), but that's a bad thing to do here, because the proxy gets updated on the render thread,
 	// and this function happens during EndOfFrame updates. So instead, use the most up-to-date transform here (fixes local-space frame-behind issues)
-	const FTransform LocalToWorld = SystemInstance->GetWorldTransform();
-
 	const bool bUseLocalSpace = UseLocalSpace(Proxy);
-	FNiagaraLWCConverter LwcConverter = SystemInstance->GetLWCConverter(bUseLocalSpace);
+	const FTransform SimToWorld = SystemInstance->GetLWCSimToWorld(bUseLocalSpace);
+	const FVector3f DefaultSimPos = bUseLocalSpace ? FVector3f::ZeroVector : FVector3f(SystemInstance->GetWorldTransform().GetLocation());
 
 	const FNiagaraParameterStore& ParameterStore = Emitter->GetRendererBoundVariables();
-	const FVector3f DefaultPos = ParameterStore.GetParameterValueOrDefault(Properties->PositionBinding.GetParamMapBindableVariable(), FVector3f::ZeroVector);
+	const FVector3f DefaultPos = ParameterStore.GetParameterValueOrDefault(Properties->PositionBinding.GetParamMapBindableVariable(), DefaultSimPos);
 	const FLinearColor DefaultColor = ParameterStore.GetParameterValueOrDefault(Properties->ColorBinding.GetParamMapBindableVariable(), Properties->ColorBinding.GetDefaultValue<FLinearColor>());
 	const float DefaultRadius = ParameterStore.GetParameterValueOrDefault(Properties->RadiusBinding.GetParamMapBindableVariable(), Properties->RadiusBinding.GetDefaultValue<float>());
 	const float DefaultScattering = ParameterStore.GetParameterValueOrDefault(Properties->VolumetricScatteringBinding.GetParamMapBindableVariable(), Properties->VolumetricScatteringBinding.GetDefaultValue<float>());
@@ -146,7 +144,7 @@ FNiagaraDynamicDataBase* FNiagaraRendererLights::GenerateDynamicData(const FNiag
 				LightData.LightEntry.InverseExposureBlend = InverseExposureBlend;
 				LightData.LightEntry.bAffectTranslucency = Properties->bAffectsTranslucency;
 				LightData.LightEntry.VolumetricScatteringIntensity = ScatteringReader.GetSafe(ParticleIndex, DefaultScattering);
-				LightData.PerViewEntry.Position = bUseLocalSpace ? LocalToWorld.TransformPosition(FVector(SimPos)) : LwcConverter.ConvertSimulationPositionToWorld(SimPos);
+				LightData.PerViewEntry.Position = SimToWorld.TransformPosition(FVector(SimPos));
 			}
 		}
 	}
@@ -170,7 +168,7 @@ FNiagaraDynamicDataBase* FNiagaraRendererLights::GenerateDynamicData(const FNiag
 			LightData.LightEntry.InverseExposureBlend = InverseExposureBlend;
 			LightData.LightEntry.bAffectTranslucency = Properties->bAffectsTranslucency;
 			LightData.LightEntry.VolumetricScatteringIntensity = LightScattering;
-			LightData.PerViewEntry.Position = bUseLocalSpace ? LocalToWorld.TransformPosition(FVector(SimPos)) : LwcConverter.ConvertSimulationPositionToWorld(SimPos);
+			LightData.PerViewEntry.Position = SimToWorld.TransformPosition(FVector(SimPos));
 		}
 	}
 
diff --git a/Engine/Plugins/FX/Niagara/Source/Niagara/Private/NiagaraSystemInstance.cpp b/Engine/Plugins/FX/Niagara/Source/Niagara/Private/NiagaraSystemInstance.cpp
--- a/Engine/Plugins/FX/Niagara/Source/Niagara/Private/NiagaraSystemInstance.cpp
+++ b/Engine/Plugins/FX/Niagara/Source/Niagara/Private/NiagaraSystemInstance.cpp
@@ -1247,6 +1247,17 @@ FNiagaraLWCConverter FNiagaraSystemInstance::GetLWCConverter(bool bLocalSpaceEmi
 	return FNiagaraLWCConverter(FVector(LWCTile) * FLargeWorldRenderScalar::GetTileSize());
 }
 
+FTransform FNiagaraSystemInstance::GetLWCSimToWorld(bool bLocalSpaceEmitter) const
+{
+	FTransform SimToWorld;
+	if (bLocalSpaceEmitter)
+	{
+		SimToWorld = WorldTransform;
+	}
+	SimToWorld.AddToTranslation(FVector(LWCTile) * FLargeWorldRenderScalar::GetTileSize());
+	return SimToWorld;
+}
+
 FNiagaraWorldManager* FNiagaraSystemInstance::GetWorldManager() const
 {
 	check(World);
diff --git a/Engine/Plugins/FX/Niagara/Source/Niagara/Public/NiagaraSystemInstance.h b/Engine/Plugins/FX/Niagara/Source/Niagara/Public/NiagaraSystemInstance.h
--- a/Engine/Plugins/FX/Niagara/Source/Niagara/Public/NiagaraSystemInstance.h
+++ b/Engine/Plugins/FX/Niagara/Source/Niagara/Public/NiagaraSystemInstance.h
@@ -137,6 +137,7 @@ public:
 
 	FORCEINLINE FNiagaraParameterStore& GetInstanceParameters() { return InstanceParameters; }
 	FNiagaraLWCConverter GetLWCConverter(bool bLocalSpaceEmitter = false) const;
+	FTransform GetLWCSimToWorld(bool bLocalSpaceEmitter = false) const;
 
 	FORCEINLINE uint32 GetParameterIndex(bool PreviousFrame = false) const
 	{
