Fix for numeric resolution through a static branch during graph instantiation with the new niagara system compile mode
#rnx

diff --git a/Engine/Plugins/FX/Niagara/Source/NiagaraEditor/Private/NiagaraGraphDigest.cpp b/Engine/Plugins/FX/Niagara/Source/NiagaraEditor/Private/NiagaraGraphDigest.cpp
--- a/Engine/Plugins/FX/Niagara/Source/NiagaraEditor/Private/NiagaraGraphDigest.cpp
+++ b/Engine/Plugins/FX/Niagara/Source/NiagaraEditor/Private/NiagaraGraphDigest.cpp
@@ -1482,8 +1482,20 @@ FNiagaraCompilationInputPin::FNiagaraCompilationInputPin(const FNiagaraCompilati
 	// if we're not connected to anything we need to also check to see if that's because our link has
 	// been terminated by a static branch.  If so we may need to update our default value based on the contents
 	// of the static switch
-	else if (SourceInputPin.LinkedTo)
+	else if (TracedInputPin != &SourceInputPin)
 	{
+		const FNiagaraTypeDefinition& GenericNumericDef = FNiagaraTypeDefinition::GetGenericNumericDef();
+
+		if (Variable.GetType() == GenericNumericDef)
+		{
+			const FNiagaraTypeDefinition& SourceTypeDef = TracedInputPin->Variable.GetType();
+			if (SourceTypeDef != GenericNumericDef)
+			{
+				PinType = TracedInputPin->PinType;
+				Variable.SetType(SourceTypeDef);
+			}
+		}
+
 		if (ensure(Variable.GetType() == TracedInputPin->Variable.GetType()))
 		{
 			if (TracedInputPin->Variable.IsDataAllocated())
