warning for zero attributes fixed

diff --git a/Engine/Plugins/FX/Niagara/Source/Niagara/Private/NiagaraDataInterfaceGrid3DCollection.cpp b/Engine/Plugins/FX/Niagara/Source/Niagara/Private/NiagaraDataInterfaceGrid3DCollection.cpp
--- a/Engine/Plugins/FX/Niagara/Source/Niagara/Private/NiagaraDataInterfaceGrid3DCollection.cpp
+++ b/Engine/Plugins/FX/Niagara/Source/Niagara/Private/NiagaraDataInterfaceGrid3DCollection.cpp
@@ -367,9 +367,22 @@ void UNiagaraDataInterfaceGrid3DCollection::GetFeedback(UNiagaraSystem* Asset, U
 	TArray<FNiagaraDataInterfaceFeedback>& OutWarnings, TArray<FNiagaraDataInterfaceFeedback>& OutInfo)
 {
 	Super::GetFeedback(Asset, Component, OutErrors, OutWarnings, OutInfo);
-	// Put in placeholder for now.
+		
+	int32 NumAttribChannelsFound = 0;
+	int32 NumNamedAttribChannelsFound = 0;
+	TArray<FNiagaraVariableBase> Vars;
+	TArray<uint32> Offsets;
+	FindAttributes(Vars, Offsets, NumNamedAttribChannelsFound);
 
+	// Ensure we never allocate 0 attributes as that would fail to create the texture
+	NumAttribChannelsFound = NumAttributes + NumNamedAttribChannelsFound;
 
+	if (NumAttribChannelsFound == 0)
+	{
+		OutWarnings.Emplace(LOCTEXT("Grid3DZeroAttrs", "Zero attributes defined on grid"),
+			LOCTEXT("Grid3DZeroAttrsFmt", "Zero attributes defined on grid"),
+			FNiagaraDataInterfaceFix());		
+	}
 }
 #endif
 #if WITH_EDITORONLY_DATA
@@ -3343,7 +3356,7 @@ bool UNiagaraDataInterfaceGrid3DCollection::InitPerInstanceData(void* PerInstanc
 	{
 		if (!Proxy->SourceDIName.IsNone())
 		{
-			UE_LOG(LogNiagara, Warning, TEXT("Zero attributes defined on %s"), *Proxy->SourceDIName.ToString());
+			UE_LOG(LogNiagara, Log, TEXT("Zero attributes defined on %s"), *Proxy->SourceDIName.ToString());
 		}
 
 		// Push Updates to Proxy.
