Fix for static analysis warnings with NiagaraDigest changes
#rnx

diff --git a/Engine/Plugins/FX/Niagara/Source/NiagaraEditor/Private/NiagaraCompilationTasks.cpp b/Engine/Plugins/FX/Niagara/Source/NiagaraEditor/Private/NiagaraCompilationTasks.cpp
--- a/Engine/Plugins/FX/Niagara/Source/NiagaraEditor/Private/NiagaraCompilationTasks.cpp
+++ b/Engine/Plugins/FX/Niagara/Source/NiagaraEditor/Private/NiagaraCompilationTasks.cpp
@@ -54,35 +54,42 @@ namespace NiagaraCompilationCopyImpl
 		TArray<UNiagaraScript*> DependentScripts;
 
 		const ENiagaraScriptUsage ScriptUsage = Script->GetUsage();
-		const bool bIsGpuEmitter = EmitterData && EmitterData->SimTarget == ENiagaraSimTarget::GPUComputeSim;
 
-		if (UNiagaraScript::IsEquivalentUsage(ScriptUsage, ENiagaraScriptUsage::EmitterSpawnScript))
+		if (System)
 		{
-			DependentScripts.Add(System->GetSystemSpawnScript());
-		}
-
-		if (UNiagaraScript::IsEquivalentUsage(ScriptUsage, ENiagaraScriptUsage::EmitterUpdateScript))
-		{
-			DependentScripts.Add(System->GetSystemUpdateScript());
-		}
+			if (UNiagaraScript::IsEquivalentUsage(ScriptUsage, ENiagaraScriptUsage::EmitterSpawnScript))
+			{
+				DependentScripts.Add(System->GetSystemSpawnScript());
+			}
 
-		if (UNiagaraScript::IsEquivalentUsage(ScriptUsage, ENiagaraScriptUsage::ParticleSpawnScript))
-		{
-			if (bIsGpuEmitter)
+			if (UNiagaraScript::IsEquivalentUsage(ScriptUsage, ENiagaraScriptUsage::EmitterUpdateScript))
 			{
-				DependentScripts.Add(EmitterData->GetGPUComputeScript());
+				DependentScripts.Add(System->GetSystemUpdateScript());
 			}
 		}
 
-		if (UNiagaraScript::IsEquivalentUsage(ScriptUsage, ENiagaraScriptUsage::ParticleUpdateScript))
+		if (EmitterData)
 		{
-			if (bIsGpuEmitter)
+			const bool bIsGpuEmitter = EmitterData->SimTarget == ENiagaraSimTarget::GPUComputeSim;
+
+			if (UNiagaraScript::IsEquivalentUsage(ScriptUsage, ENiagaraScriptUsage::ParticleSpawnScript))
 			{
-				DependentScripts.Add(EmitterData->GetGPUComputeScript());
+				if (bIsGpuEmitter)
+				{
+					DependentScripts.Add(EmitterData->GetGPUComputeScript());
+				}
 			}
-			else if (EmitterData->bInterpolatedSpawning)
+
+			if (UNiagaraScript::IsEquivalentUsage(ScriptUsage, ENiagaraScriptUsage::ParticleUpdateScript))
 			{
-				DependentScripts.Add(EmitterData->SpawnScriptProps.Script);
+				if (bIsGpuEmitter)
+				{
+					DependentScripts.Add(EmitterData->GetGPUComputeScript());
+				}
+				else if (EmitterData->bInterpolatedSpawning)
+				{
+					DependentScripts.Add(EmitterData->SpawnScriptProps.Script);
+				}
 			}
 		}
 
diff --git a/Engine/Plugins/FX/Niagara/Source/NiagaraEditor/Private/NiagaraDigestDatabase.cpp b/Engine/Plugins/FX/Niagara/Source/NiagaraEditor/Private/NiagaraDigestDatabase.cpp
--- a/Engine/Plugins/FX/Niagara/Source/NiagaraEditor/Private/NiagaraDigestDatabase.cpp
+++ b/Engine/Plugins/FX/Niagara/Source/NiagaraEditor/Private/NiagaraDigestDatabase.cpp
@@ -246,10 +246,13 @@ UNiagaraDataInterface* FNiagaraCompilationNPC::GetDataInterface(const FNiagaraVa
 }
 
 FNiagaraCompilationNPCHandle::FNiagaraCompilationNPCHandle(const UNiagaraParameterCollection* Connection)
-	: Namespace(Connection ? Connection->GetNamespace() : NAME_None)
 {
-	AssetKey = Connection;
-	Hash = Connection->GetCompileHash();
+	if (Connection)
+	{
+		Namespace = Connection->GetNamespace();
+		AssetKey = Connection;
+		Hash = Connection->GetCompileHash();
+	}
 }
 
 FNiagaraCompilationNPCHandle::FDigestPtr FNiagaraCompilationNPCHandle::Resolve() const
diff --git a/Engine/Plugins/FX/Niagara/Source/NiagaraEditor/Private/NiagaraGraphDigest.cpp b/Engine/Plugins/FX/Niagara/Source/NiagaraEditor/Private/NiagaraGraphDigest.cpp
--- a/Engine/Plugins/FX/Niagara/Source/NiagaraEditor/Private/NiagaraGraphDigest.cpp
+++ b/Engine/Plugins/FX/Niagara/Source/NiagaraEditor/Private/NiagaraGraphDigest.cpp
@@ -140,7 +140,12 @@ FNodeHelper CreateNodeHelper(FNiagaraCompilationNode::ENodeType NodeType)
 	(
 		[](const UEdGraphNode* SourceNode, FNiagaraCompilationGraphCreateContext& NodeContext) -> TUniquePtr<FNiagaraCompilationNode>
 		{
-			return MakeUnique<CompilationNodeClass>(static_cast<const NodeClass*>(SourceNode), NodeContext);
+			if (const NodeClass* TypedSourceNode = static_cast<const NodeClass*>(SourceNode))
+			{
+				return MakeUnique<CompilationNodeClass>(TypedSourceNode, NodeContext);
+			}
+
+			return TUniquePtr<FNiagaraCompilationNode>();
 		}
 	);
 
@@ -148,7 +153,11 @@ FNodeHelper CreateNodeHelper(FNiagaraCompilationNode::ENodeType NodeType)
 	(
 		[](const FNiagaraCompilationNode* SourceNode, FNiagaraCompilationGraphDuplicateContext& NodeContext) -> TUniquePtr<FNiagaraCompilationNode>
 		{
-			return MakeUnique<CompilationNodeClass>(*static_cast<const CompilationNodeClass*>(SourceNode), NodeContext);
+			if (const CompilationNodeClass* TypedSourceNode = static_cast<const CompilationNodeClass*>(SourceNode))
+			{
+				return MakeUnique<CompilationNodeClass>(*TypedSourceNode, NodeContext);
+			}
+			return TUniquePtr<FNiagaraCompilationNode>();
 		}
 	);
 
@@ -1466,17 +1475,18 @@ FNiagaraCompilationInputPin::FNiagaraCompilationInputPin(const FNiagaraCompilati
 	if (const FNiagaraCompilationOutputPin* OutputPin = TracedInputPin->LinkedTo)
 	{
 		const FNiagaraCompilationNode* SourceNode = OutputPin->OwningNode;
-		ensure(SourceNode);
-
-		FNiagaraCompilationNode* ConnectedNode = Context.DuplicatedNodeMap.FindRef(SourceNode);
-		if (ensure(ConnectedNode))
+		if (ensure(SourceNode))
 		{
-			// find the pin index in the array
-			const int32 OutputPinIndex = UE_PTRDIFF_TO_INT32(OutputPin - SourceNode->OutputPins.GetData());
-			ensure(SourceNode->OutputPins.IsValidIndex(OutputPinIndex));
+			FNiagaraCompilationNode* ConnectedNode = Context.DuplicatedNodeMap.FindRef(SourceNode);
+			if (ensure(ConnectedNode))
+			{
+				// find the pin index in the array
+				const int32 OutputPinIndex = UE_PTRDIFF_TO_INT32(OutputPin - SourceNode->OutputPins.GetData());
+				ensure(SourceNode->OutputPins.IsValidIndex(OutputPinIndex));
 
-			LinkedTo = &ConnectedNode->OutputPins[OutputPinIndex];
-			ConnectedNode->OutputPins[OutputPinIndex].LinkedTo.Add(this);
+				LinkedTo = &ConnectedNode->OutputPins[OutputPinIndex];
+				ConnectedNode->OutputPins[OutputPinIndex].LinkedTo.Add(this);
+			}
 		}
 	}
 	// if we're not connected to anything we need to also check to see if that's because our link has
