- Temporarily remove ensure until we figure out why it's being triggered

#rnx

diff --git a/Engine/Plugins/FX/Niagara/Source/NiagaraEditor/Private/NiagaraScriptMergeManager.cpp b/Engine/Plugins/FX/Niagara/Source/NiagaraEditor/Private/NiagaraScriptMergeManager.cpp
--- a/Engine/Plugins/FX/Niagara/Source/NiagaraEditor/Private/NiagaraScriptMergeManager.cpp
+++ b/Engine/Plugins/FX/Niagara/Source/NiagaraEditor/Private/NiagaraScriptMergeManager.cpp
@@ -3763,7 +3763,8 @@ FNiagaraScriptMergeManager::FApplyDiffResults FNiagaraScriptMergeManager::ApplyE
 				{
 					UNiagaraHierarchyItemBase* NewChild = BaseRoot->CopyAndAddItemUnderParentIdentity(*AddedInput, ParentIdentity);
 					
-					if(!ensure(NewChild != nullptr))
+					//if(!ensure(NewChild != nullptr))
+					if(NewChild != nullptr)
 					{
 						UE_LOG(LogNiagaraEditor, Log, TEXT("Item %s could not be added during merge process"), *AddedInput->ToString());
 					}
