- Fix crash during cook

#rnx

diff --git a/Engine/Plugins/FX/Niagara/Source/NiagaraEditor/Private/NiagaraScriptMergeManager.cpp b/Engine/Plugins/FX/Niagara/Source/NiagaraEditor/Private/NiagaraScriptMergeManager.cpp
--- a/Engine/Plugins/FX/Niagara/Source/NiagaraEditor/Private/NiagaraScriptMergeManager.cpp
+++ b/Engine/Plugins/FX/Niagara/Source/NiagaraEditor/Private/NiagaraScriptMergeManager.cpp
@@ -3764,7 +3764,7 @@ FNiagaraScriptMergeManager::FApplyDiffResults FNiagaraScriptMergeManager::ApplyE
 					UNiagaraHierarchyItemBase* NewChild = BaseRoot->CopyAndAddItemUnderParentIdentity(*AddedInput, ParentIdentity);
 					
 					//if(!ensure(NewChild != nullptr))
-					if(NewChild != nullptr)
+					if(NewChild == nullptr)
 					{
 						UE_LOG(LogNiagaraEditor, Log, TEXT("Item %s could not be added during merge process"), *AddedInput->ToString());
 					}
