Fix localization duplication errors with NiagaraDigestGraph
#rnx

diff --git a/Engine/Plugins/FX/Niagara/Source/NiagaraEditor/Private/NiagaraGraphDigest.cpp b/Engine/Plugins/FX/Niagara/Source/NiagaraEditor/Private/NiagaraGraphDigest.cpp
--- a/Engine/Plugins/FX/Niagara/Source/NiagaraEditor/Private/NiagaraGraphDigest.cpp
+++ b/Engine/Plugins/FX/Niagara/Source/NiagaraEditor/Private/NiagaraGraphDigest.cpp
@@ -34,7 +34,7 @@
 
 #include "Algo/RemoveIf.h"
 
-#define LOCTEXT_NAMESPACE "NiagaraCompiler"
+#define LOCTEXT_NAMESPACE "NiagaraDigestCompiler"
 
 #define NIAGARA_GRAPH_DIGEST_NODE_TYPE(name) NIAGARA_GRAPH_DIGEST_NODE_IMPLEMENT_BODY(name)
 
