Added age update mode dropdown to life cycle track

diff --git a/Engine/Plugins/FX/Niagara/Source/Niagara/Private/MovieScene/MovieSceneNiagaraSystemSpawnSection.cpp b/Engine/Plugins/FX/Niagara/Source/Niagara/Private/MovieScene/MovieSceneNiagaraSystemSpawnSection.cpp
--- a/Engine/Plugins/FX/Niagara/Source/Niagara/Private/MovieScene/MovieSceneNiagaraSystemSpawnSection.cpp
+++ b/Engine/Plugins/FX/Niagara/Source/Niagara/Private/MovieScene/MovieSceneNiagaraSystemSpawnSection.cpp
@@ -35,6 +35,11 @@ ENiagaraAgeUpdateMode UMovieSceneNiagaraSystemSpawnSection::GetAgeUpdateMode() c
 	return AgeUpdateMode;
 }
 
+void UMovieSceneNiagaraSystemSpawnSection::SetAgeUpdateMode(ENiagaraAgeUpdateMode NewMode)
+{
+	AgeUpdateMode = NewMode;
+}
+
 bool UMovieSceneNiagaraSystemSpawnSection::GetAllowScalability()const
 {
 	return bAllowScalability;
diff --git a/Engine/Plugins/FX/Niagara/Source/Niagara/Public/MovieScene/MovieSceneNiagaraSystemSpawnSection.h b/Engine/Plugins/FX/Niagara/Source/Niagara/Public/MovieScene/MovieSceneNiagaraSystemSpawnSection.h
--- a/Engine/Plugins/FX/Niagara/Source/Niagara/Public/MovieScene/MovieSceneNiagaraSystemSpawnSection.h
+++ b/Engine/Plugins/FX/Niagara/Source/Niagara/Public/MovieScene/MovieSceneNiagaraSystemSpawnSection.h
@@ -52,8 +52,9 @@ public:
 	ENiagaraSystemSpawnSectionEndBehavior GetSectionEndBehavior() const;
 
 	NIAGARA_API ENiagaraAgeUpdateMode GetAgeUpdateMode() const;
+	NIAGARA_API void SetAgeUpdateMode(ENiagaraAgeUpdateMode NewMode);
 
-	bool GetAllowScalability()const;
+	bool GetAllowScalability() const;
 
 private:
 	/** Specifies what should happen to the niagara system from before the section evaluates up until the first frame of the section. */
diff --git a/Engine/Plugins/FX/Niagara/Source/NiagaraEditor/Private/Sequencer/LevelSequence/NiagaraSystemTrackEditor.cpp b/Engine/Plugins/FX/Niagara/Source/NiagaraEditor/Private/Sequencer/LevelSequence/NiagaraSystemTrackEditor.cpp
--- a/Engine/Plugins/FX/Niagara/Source/NiagaraEditor/Private/Sequencer/LevelSequence/NiagaraSystemTrackEditor.cpp
+++ b/Engine/Plugins/FX/Niagara/Source/NiagaraEditor/Private/Sequencer/LevelSequence/NiagaraSystemTrackEditor.cpp
@@ -1,18 +1,20 @@
 // Copyright Epic Games, Inc. All Rights Reserved.
 
 #include "NiagaraSystemTrackEditor.h"
+
+#include "NiagaraActor.h"
 #include "MovieScene/MovieSceneNiagaraSystemTrack.h"
 #include "MovieScene/MovieSceneNiagaraSystemSpawnSection.h"
 #include "MovieScene/Parameters/MovieSceneNiagaraParameterTrack.h"
 #include "NiagaraSystemSpawnSection.h"
 #include "NiagaraComponent.h"
 #include "NiagaraSystem.h"
-#include "NiagaraEmitterHandle.h"
 #include "NiagaraEditorUtilities.h"
 #include "Framework/MultiBox/MultiBoxBuilder.h"
 #include "Modules/ModuleManager.h"
 #include "NiagaraEditorModule.h"
 #include "NiagaraSystemEditorData.h"
+#include "SEnumCombo.h"
 
 #define LOCTEXT_NAMESPACE "NiagaraSystemTrackEditor"
 
@@ -137,6 +139,45 @@ void FNiagaraSystemTrackEditor::BuildObjectBindingTrackMenu(FMenuBuilder& MenuBu
 	}
 }
 
+TSharedPtr<SWidget> FNiagaraSystemTrackEditor::BuildOutlinerEditWidget(const FGuid& ObjectBinding, UMovieSceneTrack* Track, const FBuildEditWidgetParams& Params)
+{
+	if (UMovieSceneNiagaraSystemTrack* NiagaraSystemTrack = Cast<UMovieSceneNiagaraSystemTrack>(Track))
+	{
+		return SNew(SHorizontalBox)
+			+ SHorizontalBox::Slot()
+			.Padding(5, 0, 5, 0)
+			.VAlign(VAlign_Center)
+			.AutoWidth()
+			[
+				SNew(SEnumComboBox, StaticEnum<ENiagaraAgeUpdateMode>())
+				.Font(FAppStyle::GetFontStyle(TEXT("PropertyWindow.NormalFont")))
+				.CurrentValue_Lambda([NiagaraSystemTrack]() -> int32
+				{
+					UMovieSceneNiagaraSystemSpawnSection* SpawnSection = CastChecked<UMovieSceneNiagaraSystemSpawnSection>(NiagaraSystemTrack->GetAllSections()[0]);
+					return static_cast<int32>(SpawnSection->GetAgeUpdateMode());
+				})
+				.OnEnumSelectionChanged_Lambda([NiagaraSystemTrack](int32 NewSelection, ESelectInfo::Type)
+				{
+					if (NewSelection == INDEX_NONE)
+					{
+						return;
+					}
+
+					FScopedTransaction Transaction(LOCTEXT("ChangeValue", "Change Value"));
+					UMovieSceneNiagaraSystemSpawnSection* SpawnSection = CastChecked<UMovieSceneNiagaraSystemSpawnSection>(NiagaraSystemTrack->GetAllSections()[0]);
+					SpawnSection->Modify();
+					SpawnSection->SetAgeUpdateMode( static_cast<ENiagaraAgeUpdateMode>(NewSelection));
+				})
+				.IsEnabled_Lambda([NiagaraSystemTrack]() -> bool
+				{
+					TArray<UMovieSceneSection*> Sections = NiagaraSystemTrack->GetAllSections();
+					return Sections.Num() == 1 && Cast<UMovieSceneNiagaraSystemSpawnSection>(Sections[0]);
+				})
+			];
+	}
+	return FMovieSceneTrackEditor::BuildOutlinerEditWidget(ObjectBinding, Track, Params);
+}
+
 void FNiagaraSystemTrackEditor::AddNiagaraSystemTrack(TArray<FGuid> ObjectBindings)
 {
 	UMovieScene* MovieScene = GetSequencer()->GetFocusedMovieSceneSequence()->GetMovieScene();
diff --git a/Engine/Plugins/FX/Niagara/Source/NiagaraEditor/Private/Sequencer/LevelSequence/NiagaraSystemTrackEditor.h b/Engine/Plugins/FX/Niagara/Source/NiagaraEditor/Private/Sequencer/LevelSequence/NiagaraSystemTrackEditor.h
--- a/Engine/Plugins/FX/Niagara/Source/NiagaraEditor/Private/Sequencer/LevelSequence/NiagaraSystemTrackEditor.h
+++ b/Engine/Plugins/FX/Niagara/Source/NiagaraEditor/Private/Sequencer/LevelSequence/NiagaraSystemTrackEditor.h
@@ -40,6 +40,7 @@ public:
 	virtual TSharedRef<ISequencerSection> MakeSectionInterface(UMovieSceneSection& SectionObject, UMovieSceneTrack& Track, FGuid ObjectBinding) override;
 	virtual bool SupportsType(TSubclassOf<UMovieSceneTrack> Type) const override;
 	virtual void BuildObjectBindingTrackMenu(FMenuBuilder& MenuBuilder, const TArray<FGuid>& ObjectBindings, const UClass* ObjectClass) override;
+	virtual TSharedPtr<SWidget> BuildOutlinerEditWidget(const FGuid& ObjectBinding, UMovieSceneTrack* Track, const FBuildEditWidgetParams& Params) override;
 
 private:
 	void AddNiagaraSystemTrack(TArray<FGuid> ObjectBindings);
