PR #10416: Exposed Timed Niagara Effect Properties to Blueprint

diff --git a/Engine/Plugins/FX/Niagara/Source/NiagaraAnimNotifies/Public/AnimNotifyState_TimedNiagaraEffect.h b/Engine/Plugins/FX/Niagara/Source/NiagaraAnimNotifies/Public/AnimNotifyState_TimedNiagaraEffect.h
--- a/Engine/Plugins/FX/Niagara/Source/NiagaraAnimNotifies/Public/AnimNotifyState_TimedNiagaraEffect.h
+++ b/Engine/Plugins/FX/Niagara/Source/NiagaraAnimNotifies/Public/AnimNotifyState_TimedNiagaraEffect.h
@@ -22,24 +22,24 @@ class NIAGARAANIMNOTIFIES_API UAnimNotifyState_TimedNiagaraEffect : public UAnim
 	GENERATED_UCLASS_BODY()
 
 	// The niagara system template to use when spawning the niagara component
-	UPROPERTY(EditAnywhere, Category = NiagaraSystem, meta = (ToolTip = "The niagara system to spawn for the notify state"))
+	UPROPERTY(EditAnywhere, BlueprintReadWrite, Category = NiagaraSystem, meta = (DisplayName = "Niagara System", ToolTip = "The niagara system to spawn for the notify state"))
 	TObjectPtr<UNiagaraSystem> Template;
 
 	// The socket within our mesh component to attach to when we spawn the Niagara component
-	UPROPERTY(EditAnywhere, Category = NiagaraSystem, meta = (ToolTip = "The socket or bone to attach the system to", AnimNotifyBoneName = "true"))
+	UPROPERTY(EditAnywhere, BlueprintReadWrite, Category = NiagaraSystem, meta = (ToolTip = "The socket or bone to attach the system to", AnimNotifyBoneName = "true"))
 	FName SocketName;
 
 	// Offset from the socket / bone location
-	UPROPERTY(EditAnywhere, Category = NiagaraSystem, meta = (ToolTip = "Offset from the socket or bone to place the Niagara system"))
+	UPROPERTY(EditAnywhere, BlueprintReadWrite, Category = NiagaraSystem, meta = (ToolTip = "Offset from the socket or bone to place the Niagara system"))
 	FVector LocationOffset;
 
 	// Offset from the socket / bone rotation
-	UPROPERTY(EditAnywhere, Category = NiagaraSystem, meta = (ToolTip = "Rotation offset from the socket or bone for the Niagara system"))
+	UPROPERTY(EditAnywhere, BlueprintReadWrite, Category = NiagaraSystem, meta = (ToolTip = "Rotation offset from the socket or bone for the Niagara system"))
 	FRotator RotationOffset;
 
 	// Whether or not we destroy the component at the end of the notify or instead just stop
 	// the emitters.
-	UPROPERTY(EditAnywhere, Category = NiagaraSystem, meta = (DisplayName = "Destroy Immediately", ToolTip = "Whether the Niagara system should be immediately destroyed at the end of the notify state or be allowed to finish"))
+	UPROPERTY(EditAnywhere, BlueprintReadWrite, Category = NiagaraSystem, meta = (DisplayName = "Destroy Immediately", ToolTip = "Whether the Niagara system should be immediately destroyed at the end of the notify state or be allowed to finish"))
 	bool bDestroyAtEnd;
 
 	UE_DEPRECATED(5.0, "Please use the other NotifyBegin function instead")
