Fix for deprecated implicit copy assignment operator for FNiagaraCompilationGraphHandle
#rnx

diff --git a/Engine/Plugins/FX/Niagara/Source/NiagaraEditor/Private/NiagaraDigestDatabase.h b/Engine/Plugins/FX/Niagara/Source/NiagaraEditor/Private/NiagaraDigestDatabase.h
--- a/Engine/Plugins/FX/Niagara/Source/NiagaraEditor/Private/NiagaraDigestDatabase.h
+++ b/Engine/Plugins/FX/Niagara/Source/NiagaraEditor/Private/NiagaraDigestDatabase.h
@@ -50,7 +50,6 @@ public:
 	using FGraphPtr = Super::FDigestPtr;
 
 	FNiagaraCompilationGraphHandle() = default;
-	FNiagaraCompilationGraphHandle(const FNiagaraCompilationGraphHandle& Handle) = default;
 	FNiagaraCompilationGraphHandle(const UNiagaraScriptSourceBase* ScriptSource);
 	FNiagaraCompilationGraphHandle(const UNiagaraGraph* Graph);
 
