- Remove debug code I checked in that emptied the UBs
- Change Crc function to use none deprecated one

#rnx

diff --git a/Engine/Plugins/FX/Niagara/Source/Niagara/Private/NiagaraComponent.cpp b/Engine/Plugins/FX/Niagara/Source/Niagara/Private/NiagaraComponent.cpp
--- a/Engine/Plugins/FX/Niagara/Source/Niagara/Private/NiagaraComponent.cpp
+++ b/Engine/Plugins/FX/Niagara/Source/Niagara/Private/NiagaraComponent.cpp
@@ -357,11 +357,10 @@ TUniformBuffer<FPrimitiveUniformShaderParameters>* FNiagaraSceneProxy::GetCustom
 	const FCustomPrimitiveData* LocalCustomPrimitiveData = GetCustomPrimitiveData();
 	if (LocalCustomPrimitiveData && LocalCustomPrimitiveData->Data.Num())
 	{
-		const uint32 CustomFloatHash = FCrc::MemCrc_DEPRECATED(LocalCustomPrimitiveData->Data.GetData(), LocalCustomPrimitiveData->Data.Num() * LocalCustomPrimitiveData->Data.GetTypeSize());
+		const uint32 CustomFloatHash = FCrc::MemCrc32(LocalCustomPrimitiveData->Data.GetData(), LocalCustomPrimitiveData->Data.Num() * LocalCustomPrimitiveData->Data.GetTypeSize());
 		KeyHash = HashCombine(KeyHash, CustomFloatHash);
 	}
 
-	CustomUniformBuffers.Empty();
 	TUniformBuffer<FPrimitiveUniformShaderParameters>*& CustomUBRef = CustomUniformBuffers.FindOrAdd(KeyHash);
 	if (CustomUBRef == nullptr)
 	{
