- Move partial depth discovery to cache from compiled data to avoid doing it for every renderer

#rnx

diff --git a/Engine/Plugins/FX/Niagara/Source/Niagara/Classes/NiagaraEmitter.h b/Engine/Plugins/FX/Niagara/Source/Niagara/Classes/NiagaraEmitter.h
--- a/Engine/Plugins/FX/Niagara/Source/Niagara/Classes/NiagaraEmitter.h
+++ b/Engine/Plugins/FX/Niagara/Source/Niagara/Classes/NiagaraEmitter.h
@@ -379,6 +379,7 @@ struct NIAGARA_API FVersionedNiagaraEmitterData
 	bool DidPSOPrecacheFail() const { return PSOPrecacheResult == EPSOPrecacheResult::NotSupported; }
 
 	bool RequiresViewUniformBuffer() const { return bRequiresViewUniformBuffer; }
+	bool NeedsPartialDepthTexture() const { return bNeedsPartialDepthTexture; }
 	uint32 GetMaxInstanceCount() const { return MaxInstanceCount; }
 	uint32 GetMaxAllocationCount() const { return MaxAllocationCount; }
 	TConstArrayView<TSharedPtr<FNiagaraBoundsCalculator>> GetBoundsCalculators() const { return MakeArrayView(BoundsCalculators); }
@@ -502,6 +503,9 @@ private:
 	/** Indicates that the GPU script requires the view uniform buffer. */
 	uint32 bRequiresViewUniformBuffer : 1;
 
+	/** Indicates we use the partial depth textures. */
+	uint32 bNeedsPartialDepthTexture : 1;
+
 	/** Maximum number of instances we can create for this emitter. */
 	uint32 MaxInstanceCount = 0;
 
diff --git a/Engine/Plugins/FX/Niagara/Source/Niagara/Private/NiagaraEmitter.cpp b/Engine/Plugins/FX/Niagara/Source/Niagara/Private/NiagaraEmitter.cpp
--- a/Engine/Plugins/FX/Niagara/Source/Niagara/Private/NiagaraEmitter.cpp
+++ b/Engine/Plugins/FX/Niagara/Source/Niagara/Private/NiagaraEmitter.cpp
@@ -1589,6 +1589,7 @@ UNiagaraScript* FVersionedNiagaraEmitterData::GetScript(ENiagaraScriptUsage Usag
 void FVersionedNiagaraEmitterData::CacheFromCompiledData(const FNiagaraDataSetCompiledData* CompiledData, const UNiagaraEmitter& Emitter)
 {
 	bRequiresViewUniformBuffer = false;
+	bNeedsPartialDepthTexture = false;
 
 	MaxInstanceCount = 0;
 	MaxAllocationCount = 0;
@@ -1881,6 +1882,7 @@ void FVersionedNiagaraEmitterData::RebuildRendererBindings(const UNiagaraEmitter
 void FVersionedNiagaraEmitterData::CacheFromShaderCompiled()
 {
 	bRequiresViewUniformBuffer = false;
+	bNeedsPartialDepthTexture = false;
 	if (GPUComputeScript && (SimTarget == ENiagaraSimTarget::GPUComputeSim))
 	{
 		if (const FNiagaraShaderScript* NiagaraShaderScript = GPUComputeScript->GetRenderThreadScript())
@@ -1888,10 +1890,10 @@ void FVersionedNiagaraEmitterData::CacheFromShaderCompiled()
 			for (int i=0; i < NiagaraShaderScript->GetNumPermutations(); ++i)
 			{
 				FNiagaraShaderRef NiagaraShaderRef = NiagaraShaderScript->GetShaderGameThread(i);
-				if (NiagaraShaderRef.IsValid() && NiagaraShaderRef->bNeedsViewUniformBuffer)
+				if (NiagaraShaderRef.IsValid())
 				{
-					bRequiresViewUniformBuffer = true;
-					break;
+					bRequiresViewUniformBuffer |= NiagaraShaderRef->bNeedsViewUniformBuffer;
+					bNeedsPartialDepthTexture |= (NiagaraShaderRef->MiscUsageBitMask & uint16(ENiagaraScriptMiscUsageMask::UsesPartialDepthCollisionQuery)) != 0;
 				}
 			}
 		}
@@ -2810,6 +2812,7 @@ FVersionedNiagaraEmitterData::FVersionedNiagaraEmitterData()
 	, bRequiresPersistentIDs(false)
 	, MaxGPUParticlesSpawnPerFrame(0)
 	, bRequiresViewUniformBuffer(false)
+	, bNeedsPartialDepthTexture(false)
 #if WITH_EDITORONLY_DATA
 	, EditorData(nullptr)
 	, EditorParameters(nullptr)
diff --git a/Engine/Plugins/FX/Niagara/Source/Niagara/Private/NiagaraRenderer.cpp b/Engine/Plugins/FX/Niagara/Source/Niagara/Private/NiagaraRenderer.cpp
--- a/Engine/Plugins/FX/Niagara/Source/Niagara/Private/NiagaraRenderer.cpp
+++ b/Engine/Plugins/FX/Niagara/Source/Niagara/Private/NiagaraRenderer.cpp
@@ -381,27 +381,7 @@ void FNiagaraRenderer::Initialize(const UNiagaraRendererProperties* InProps, con
 	bRendersInSecondaryDepthPass = false;
 
 	// Let check if the GPU simulation shader script needed to use the partial depth texture for depth queries.
-	bool bNeedsPartialDepthTexture = false;
-	if (Emitter)
-	{
-		FVersionedNiagaraEmitterData* CachedEmitterData = Emitter->GetCachedEmitterData();
-		if (CachedEmitterData)
-		{
-			UNiagaraScript* NiagaraScriptGPU = CachedEmitterData->GetGPUComputeScript();
-			if (NiagaraScriptGPU)
-			{
-				FNiagaraShaderScript* NiagaraShaderScript = NiagaraScriptGPU->GetRenderThreadScript();
-				if (NiagaraShaderScript)
-				{
-					FNiagaraShaderRef NiagaraShader = NiagaraShaderScript->GetShaderGameThread(0);
-					if (NiagaraShader.IsValid())
-					{
-						bNeedsPartialDepthTexture = (NiagaraShader->MiscUsageBitMask & uint16(ENiagaraScriptMiscUsageMask::UsesPartialDepthCollisionQuery)) > 0;
-					}
-				}
-			}
-		}
-	}
+	const bool bNeedsPartialDepthTexture = Emitter->GetCachedEmitterData()->NeedsPartialDepthTexture();
 
 	uint32 Index = 0;
 	for (UMaterialInterface*& Mat : BaseMaterials_GT)
