Niagara - Add experimental inline display modes for dynamic inputs in the selection stack.

diff --git a/Engine/Plugins/FX/Niagara/Source/Niagara/Classes/NiagaraScript.h b/Engine/Plugins/FX/Niagara/Source/Niagara/Classes/NiagaraScript.h
--- a/Engine/Plugins/FX/Niagara/Source/Niagara/Classes/NiagaraScript.h
+++ b/Engine/Plugins/FX/Niagara/Source/Niagara/Classes/NiagaraScript.h
@@ -611,6 +611,48 @@ struct NIAGARA_API FNiagaraGraphCachedDataBase
 	virtual void SetSourceEmitter(const FVersionedNiagaraEmitterData* InEmitterData) {};
 };
 
+/** Defines usages for inline dynamic input format tokens. */
+UENUM()
+enum class ENiagaraInlineDynamicInputFormatTokenUsage
+{
+	/** The format token refers to a dynamic input function input. */
+	Input,
+	/** The format token represents text which is displayed along with inputs. */
+	Decorator,
+	/** The format token represents a line break which will force the start a new row when sorting format tokens into groups. */
+	LineBreak
+};
+
+/** Defines a single token in a format for displaying dynamic input trees inline in the UI.  These tokens can represent the
+	inputs to the dynamic input or additional text and layout options. */
+USTRUCT()
+struct NIAGARA_API FNiagaraInlineDynamicInputFormatToken
+{
+	GENERATED_BODY();
+#if WITH_EDITORONLY_DATA
+
+	FNiagaraInlineDynamicInputFormatToken()
+		: Usage(ENiagaraInlineDynamicInputFormatTokenUsage::Decorator)
+	{
+	}
+
+	FNiagaraInlineDynamicInputFormatToken(ENiagaraInlineDynamicInputFormatTokenUsage InUsage, const FString& InValue)
+		: Usage(InUsage)
+		, Value(InValue)
+	{
+	}
+
+	/** Defines how the value of this token should be used when formatting the dynamic input tree. */
+	UPROPERTY(EditAnywhere, Category="Format")
+	ENiagaraInlineDynamicInputFormatTokenUsage Usage;
+
+	/** The value of this token which is used for formatting an inline dynamic input tree.  The purpose of this value is different
+		depending on the value of the Usage property. */
+	UPROPERTY(EditAnywhere, Category ="Format")
+	FString Value;
+#endif
+};
+
 /** Struct containing all of the data that can be different between different script versions.*/
 USTRUCT()
 struct NIAGARA_API FVersionedNiagaraScriptData
@@ -712,6 +754,12 @@ public:
 	UPROPERTY(EditAnywhere, Category = Script)
 	FText CollapsedViewFormat;
 
+	UPROPERTY(EditAnywhere, Category = "Script")
+	TArray<FNiagaraInlineDynamicInputFormatToken> InlineExpressionFormat;
+
+	UPROPERTY(EditAnywhere, Category = "Script")
+	TArray<FNiagaraInlineDynamicInputFormatToken> InlineGraphFormat;
+
 	/* If used as a dynamic input with exactly one input and output of different types, setting this to true will auto-insert this script to convert from one type to another when dragging and dropping parameters in the stack.
 	   For example, a script with a bool as input and a float as output will be automatically inserted in the stack when dropping a bool parameter into the float input of a module in the stack. */
 	UPROPERTY(EditAnywhere, Category = Script)
diff --git a/Engine/Plugins/FX/Niagara/Source/NiagaraEditor/Private/NiagaraEditorModule.cpp b/Engine/Plugins/FX/Niagara/Source/NiagaraEditor/Private/NiagaraEditorModule.cpp
--- a/Engine/Plugins/FX/Niagara/Source/NiagaraEditor/Private/NiagaraEditorModule.cpp
+++ b/Engine/Plugins/FX/Niagara/Source/NiagaraEditor/Private/NiagaraEditorModule.cpp
@@ -193,6 +193,20 @@ static FAutoConsoleVariableRef CVarPreloadSelectablePluginAssetsOnDemand(
 	TEXT("If > 0 then niagara system, emitter, and script assets provided by the niagara plugin will be preloaded when a dialog is opened to select them. This is a temoporary workaround for asset registry issues in cooked editor builds.\n"),
 	ECVF_Default);
 
+int32 GbEnableExperimentalInlineDynamicInputs = 0;
+static FAutoConsoleVariableRef CVarEnableExperimentalInlineDynamicInputs(
+	TEXT("fx.Niagara.EnableExperimentalInlineDynamicInputs"),
+	GbEnableExperimentalInlineDynamicInputs,
+	TEXT("If > 0 experimental inline editors for dynamic input trees will be available via right click menu in the stack.\n"),
+	ECVF_Default);
+
+int32 GbEnableCustomInlineDynamicInputFormats = 1;
+static FAutoConsoleVariableRef CVarEnableCustomInlineDynamicInputFormats(
+	TEXT("fx.Niagara.EnableCustomInlineDynamicInputFormats"),
+	GbEnableCustomInlineDynamicInputFormats,
+	TEXT("If > 0 and experimental inline editors for dynamic input trees are enabled, custom formats which are defined on scripts will be applied.\n"),
+	ECVF_Default);
+
 // this is required for gpu script compilation ticks
 static FNiagaraShaderQueueTickable NiagaraShaderQueueProcessor;
 
diff --git a/Engine/Plugins/FX/Niagara/Source/NiagaraEditor/Private/NiagaraStackEditorData.cpp b/Engine/Plugins/FX/Niagara/Source/NiagaraEditor/Private/NiagaraStackEditorData.cpp
--- a/Engine/Plugins/FX/Niagara/Source/NiagaraEditor/Private/NiagaraStackEditorData.cpp
+++ b/Engine/Plugins/FX/Niagara/Source/NiagaraEditor/Private/NiagaraStackEditorData.cpp
@@ -73,6 +73,27 @@ void UNiagaraStackEditorData::SetStackEntryWasExpandedPreSearch(const FString& S
 	}
 }
 
+ENiagaraStackEntryInlineDisplayMode UNiagaraStackEditorData::GetStackEntryInlineDisplayMode(const FString& StackEntryKey) const
+{
+	const ENiagaraStackEntryInlineDisplayMode* InlineDisplayModePtr = StackEntryKeyToInlineDisplayModeMap.Find(StackEntryKey);
+	return InlineDisplayModePtr != nullptr ? *InlineDisplayModePtr : ENiagaraStackEntryInlineDisplayMode::None;
+}
+
+void UNiagaraStackEditorData::SetStackEntryInlineDisplayMode(const FString& StackEntryKey, ENiagaraStackEntryInlineDisplayMode InlineDisplayMode)
+{
+	bool bBroadcast = false;
+	if (ensureMsgf(StackEntryKey.IsEmpty() == false, TEXT("Can not set the inline display mode with an empty key")))
+	{
+		bBroadcast = GetStackEntryInlineDisplayMode(StackEntryKey) != InlineDisplayMode;
+		StackEntryKeyToInlineDisplayModeMap.FindOrAdd(StackEntryKey) = InlineDisplayMode;
+	}
+
+	if (bBroadcast)
+	{
+		OnPersistentDataChanged().Broadcast();
+	}
+}
+
 bool UNiagaraStackEditorData::GetStackItemShowAdvanced(const FString& StackEntryKey, bool bShowAdvancedDefault) const
 {
 	const bool* bShowAdvancedPtr = StackItemKeyToShowAdvancedMap.Find(StackEntryKey);
diff --git a/Engine/Plugins/FX/Niagara/Source/NiagaraEditor/Private/SNiagaraParameterEditor.cpp b/Engine/Plugins/FX/Niagara/Source/NiagaraEditor/Private/SNiagaraParameterEditor.cpp
--- a/Engine/Plugins/FX/Niagara/Source/NiagaraEditor/Private/SNiagaraParameterEditor.cpp
+++ b/Engine/Plugins/FX/Niagara/Source/NiagaraEditor/Private/SNiagaraParameterEditor.cpp
@@ -33,11 +33,21 @@ const TOptional<float>& SNiagaraParameterEditor::GetMinimumDesiredWidth() const
 	return MinimumDesiredWidth;
 }
 
+void SNiagaraParameterEditor::SetMinimumDesiredWidth(TOptional<float> InMinimumDesiredWidth)
+{
+	MinimumDesiredWidth = InMinimumDesiredWidth;
+}
+
 const TOptional<float>& SNiagaraParameterEditor::GetMaximumDesiredWidth() const
 {
 	return MaximumDesiredWidth;
 }
 
+void SNiagaraParameterEditor::SetMaximumDesiredWidth(TOptional<float> InMaximumDesiredWidth)
+{
+	MaximumDesiredWidth = InMaximumDesiredWidth;
+}
+
 EHorizontalAlignment SNiagaraParameterEditor::GetHorizontalAlignment() const
 {
 	return HorizontalAlignment;
diff --git a/Engine/Plugins/FX/Niagara/Source/NiagaraEditor/Private/ViewModels/Stack/NiagaraStackFunctionInput.cpp b/Engine/Plugins/FX/Niagara/Source/NiagaraEditor/Private/ViewModels/Stack/NiagaraStackFunctionInput.cpp
--- a/Engine/Plugins/FX/Niagara/Source/NiagaraEditor/Private/ViewModels/Stack/NiagaraStackFunctionInput.cpp
+++ b/Engine/Plugins/FX/Niagara/Source/NiagaraEditor/Private/ViewModels/Stack/NiagaraStackFunctionInput.cpp
@@ -178,6 +178,8 @@ void UNiagaraStackFunctionInput::Initialize(
 	VisibleCondition.Initialize(SourceScript.Get(), AffectedScriptsNotWeak, ConstantResolver, UniqueEmitterName, OwningFunctionCallNode.Get());
 
 	MessageLogGuid = GetSystemViewModel()->GetMessageLogGuid();
+
+	AddChildFilter(FOnFilterChild::CreateUObject(this, &UNiagaraStackFunctionInput::FilterInlineChildren));
 }
 
 void UNiagaraStackFunctionInput::FinalizeInternal()
@@ -431,12 +433,41 @@ TOptional<FNiagaraVariableMetaData> UNiagaraStackFunctionInput::GetInputMetaData
 	return InputMetaData;
 }
 
+void UNiagaraStackFunctionInput::GetFilteredChildInputs(TArray<UNiagaraStackFunctionInput*>& OutFilteredChildInputs) const
+{
+	TArray<UNiagaraStackFunctionInputCollection*> DynamicInputCollections;
+	GetUnfilteredChildrenOfType(DynamicInputCollections);
+	for (UNiagaraStackFunctionInputCollection* DynamicInputCollection : DynamicInputCollections)
+	{
+		DynamicInputCollection->GetFilteredChildInputs(OutFilteredChildInputs);
+	}
+}
+
+UNiagaraStackObject* UNiagaraStackFunctionInput::GetChildDataObject() const
+{
+	TArray<UNiagaraStackObject*> ChildDataObjects;
+	GetUnfilteredChildrenOfType(ChildDataObjects);
+	if (ensureMsgf(ChildDataObjects.Num() <= 1, TEXT("Function input should have at most one child object.  Function Path: %s Function Name: %s Input Name: %s"),
+		OwningFunctionCallNode.IsValid() ? *OwningFunctionCallNode->GetPathName() : TEXT("[Unknown]"), 
+		OwningFunctionCallNode.IsValid() ? *OwningFunctionCallNode->GetFunctionName() : TEXT("[Unknown]"),
+		*InputParameterHandle.GetName().ToString()))
+	{
+		return ChildDataObjects.Num() == 1 ? ChildDataObjects[0] : nullptr;
+	}
+	return nullptr;
+}
+
 void UNiagaraStackFunctionInput::GetCurrentChangeIds(FGuid& OutOwningGraphChangeId, FGuid& OutFunctionGraphChangeId) const
 {
 	OutOwningGraphChangeId = OwningFunctionCallNode->GetNiagaraGraph()->GetChangeID();
 	OutFunctionGraphChangeId = OwningFunctionCallNode->GetCalledGraph() != nullptr ? OwningFunctionCallNode->GetCalledGraph()->GetChangeID() : FGuid();
 }
 
+bool UNiagaraStackFunctionInput::FilterInlineChildren(const UNiagaraStackEntry& Child) const
+{
+	return GbEnableExperimentalInlineDynamicInputs == 0 || GetInlineDisplayMode() == ENiagaraStackEntryInlineDisplayMode::None;
+}
+
 FText UNiagaraStackFunctionInput::GetCollapsedStateText() const
 {
 	if (IsFinalized())
@@ -1971,8 +2002,13 @@ void UNiagaraStackFunctionInput::SetLocalValue(TSharedRef<FStructOnScope> InLoca
 	TSharedRef<FNiagaraSystemViewModel> CachedSysViewModel = GetSystemViewModel();
 	if (CachedSysViewModel->GetSystemStackViewModel())
 		CachedSysViewModel->GetSystemStackViewModel()->InvalidateCachedParameterUsage();
-	RefreshChildren();
+	EValueMode PreviousValueMode = InputValues.Mode;
 	RefreshValues();
+
+	if (InputValues.Mode != PreviousValueMode)
+	{
+		RefreshChildren();
+	}
 }
 
 bool UNiagaraStackFunctionInput::CanReset() const
@@ -2975,6 +3011,17 @@ const UNiagaraStackFunctionInput::FCollectedUsageData& UNiagaraStackFunctionInpu
 	return CachedCollectedUsageData.GetValue();
 }
 
+ENiagaraStackEntryInlineDisplayMode UNiagaraStackFunctionInput::GetInlineDisplayMode() const
+{
+	return GetStackEditorData().GetStackEntryInlineDisplayMode(GetStackEditorDataKey());
+}
+
+void UNiagaraStackFunctionInput::SetInlineDisplayMode(ENiagaraStackEntryInlineDisplayMode InlineDisplayMode)
+{
+	GetStackEditorData().SetStackEntryInlineDisplayMode(GetStackEditorDataKey(), InlineDisplayMode);
+	RefreshFilteredChildren();
+}
+
 void UNiagaraStackFunctionInput::OnGraphChanged(const struct FEdGraphEditAction& InAction)
 {
 	if (bUpdatingGraphDirectly == false)
diff --git a/Engine/Plugins/FX/Niagara/Source/NiagaraEditor/Private/ViewModels/Stack/NiagaraStackFunctionInputCollection.cpp b/Engine/Plugins/FX/Niagara/Source/NiagaraEditor/Private/ViewModels/Stack/NiagaraStackFunctionInputCollection.cpp
--- a/Engine/Plugins/FX/Niagara/Source/NiagaraEditor/Private/ViewModels/Stack/NiagaraStackFunctionInputCollection.cpp
+++ b/Engine/Plugins/FX/Niagara/Source/NiagaraEditor/Private/ViewModels/Stack/NiagaraStackFunctionInputCollection.cpp
@@ -772,6 +772,16 @@ void UNiagaraStackFunctionInputCollection::GetChildInputs(TArray<UNiagaraStackFu
 	}
 }
 
+void UNiagaraStackFunctionInputCollection::GetFilteredChildInputs(TArray<UNiagaraStackFunctionInput*>& OutFilteredChildInputs) const
+{
+	TArray<UNiagaraStackInputCategory*> ChildCategories;
+	GetUnfilteredChildrenOfType(ChildCategories);
+	for (UNiagaraStackInputCategory* ChildCategory : ChildCategories)
+	{
+		ChildCategory->GetFilteredChildInputs(OutFilteredChildInputs);
+	}
+}
+
 void UNiagaraStackFunctionInputCollection::GetCustomFilteredChildInputs(TArray<UNiagaraStackFunctionInput*>& OutResult, const TArray<FOnFilterChild>& CustomFilters) const
 {
 	TArray<UNiagaraStackInputCategory*> ChildCategories;
diff --git a/Engine/Plugins/FX/Niagara/Source/NiagaraEditor/Private/ViewModels/Stack/NiagaraStackInputCategory.cpp b/Engine/Plugins/FX/Niagara/Source/NiagaraEditor/Private/ViewModels/Stack/NiagaraStackInputCategory.cpp
--- a/Engine/Plugins/FX/Niagara/Source/NiagaraEditor/Private/ViewModels/Stack/NiagaraStackInputCategory.cpp
+++ b/Engine/Plugins/FX/Niagara/Source/NiagaraEditor/Private/ViewModels/Stack/NiagaraStackInputCategory.cpp
@@ -216,6 +216,11 @@ void  UNiagaraStackInputCategory::SetStandardValuesFromClipboardFunctionInputs(c
 	SetValuesFromFunctionInputsInternal(this, ClipboardFunctionInputs, [](UNiagaraStackFunctionInput* ChildInput) { return ChildInput->IsStaticParameter() == false; }, nullptr);
 }
 
+void UNiagaraStackInputCategory::GetFilteredChildInputs(TArray<UNiagaraStackFunctionInput*>& OutFilteredChildInputs) const
+{
+	GetFilteredChildrenOfType(OutFilteredChildInputs);
+}
+
 void UNiagaraStackSummaryCategory::Initialize(FRequiredEntryData InRequiredEntryData, TSharedPtr<FNiagaraHierarchyCategoryViewModel> InCategoryViewModel, FString InOwnerStackItemEditorDataKey)
 {
 	CategoryViewModelWeakPtr = InCategoryViewModel;
diff --git a/Engine/Plugins/FX/Niagara/Source/NiagaraEditor/Public/NiagaraEditorModule.h b/Engine/Plugins/FX/Niagara/Source/NiagaraEditor/Public/NiagaraEditorModule.h
--- a/Engine/Plugins/FX/Niagara/Source/NiagaraEditor/Public/NiagaraEditorModule.h
+++ b/Engine/Plugins/FX/Niagara/Source/NiagaraEditor/Public/NiagaraEditorModule.h
@@ -55,6 +55,8 @@ DECLARE_STATS_GROUP(TEXT("Niagara Editor"), STATGROUP_NiagaraEditor, STATCAT_Adv
 
 extern NIAGARAEDITOR_API int32 GbShowNiagaraDeveloperWindows;
 extern NIAGARAEDITOR_API int32 GbPreloadSelectablePluginAssetsOnDemand;
+extern NIAGARAEDITOR_API int32 GbEnableExperimentalInlineDynamicInputs;
+extern NIAGARAEDITOR_API int32 GbEnableCustomInlineDynamicInputFormats;
 
 /* Defines methods for allowing external modules to supply widgets to the core editor module. */
 class NIAGARAEDITOR_API INiagaraEditorWidgetProvider
diff --git a/Engine/Plugins/FX/Niagara/Source/NiagaraEditor/Public/NiagaraStackEditorData.h b/Engine/Plugins/FX/Niagara/Source/NiagaraEditor/Public/NiagaraStackEditorData.h
--- a/Engine/Plugins/FX/Niagara/Source/NiagaraEditor/Public/NiagaraStackEditorData.h
+++ b/Engine/Plugins/FX/Niagara/Source/NiagaraEditor/Public/NiagaraStackEditorData.h
@@ -7,6 +7,17 @@
 #include "NiagaraStackEditorData.generated.h"
 struct FStackIssue;
 
+/* Defines different modes for inline display in the stack. */
+UENUM()
+enum class ENiagaraStackEntryInlineDisplayMode
+{
+	Expression,
+	GraphHorizontal,
+	GraphVertical,
+	GraphHybrid,
+	None
+};
+
 /** Editor only UI data for emitters. */
 UCLASS()
 class UNiagaraStackEditorData : public UNiagaraEditorDataBase
@@ -69,6 +80,16 @@ public:
 	 */
 	void SetStackEntryWasExpandedPreSearch(const FString& StackEntryKey, bool bWasExpandedPreSearch);
 
+	/*
+	 * Gets whether or not this entry should display all of its content inline. 
+	 */
+	ENiagaraStackEntryInlineDisplayMode GetStackEntryInlineDisplayMode(const FString& StackEntryKey) const;
+
+	/*
+	 * Sets whether or not this entry should display all of its content inline.
+	 */
+	void SetStackEntryInlineDisplayMode(const FString& StackEntryKey, ENiagaraStackEntryInlineDisplayMode InlineDisplayMode);
+
 	/*
 	* Gets whether or not a stack item is showing advanced items.
 	* @param StackItemKey A unique key for the entry.
@@ -167,6 +188,9 @@ private:
 
 	TMap<FString, bool> StackEntryKeyToPreSearchExpandedMap;
 
+	UPROPERTY()
+	TMap<FString, ENiagaraStackEntryInlineDisplayMode> StackEntryKeyToInlineDisplayModeMap;
+
 	TMap<FString, bool> StackItemKeyToShowAdvancedMap;
 
 	UPROPERTY()
diff --git a/Engine/Plugins/FX/Niagara/Source/NiagaraEditor/Public/SNiagaraParameterEditor.h b/Engine/Plugins/FX/Niagara/Source/NiagaraEditor/Public/SNiagaraParameterEditor.h
--- a/Engine/Plugins/FX/Niagara/Source/NiagaraEditor/Public/SNiagaraParameterEditor.h
+++ b/Engine/Plugins/FX/Niagara/Source/NiagaraEditor/Public/SNiagaraParameterEditor.h
@@ -50,9 +50,15 @@ public:
 	/** Gets an optional minimum desired width for this parameter editor. */
 	const TOptional<float>& GetMinimumDesiredWidth() const;
 
+	/** Sets the minimum desired width for this parameter editor.  If the option parameter is not set it clears the current value. */
+	void SetMinimumDesiredWidth(TOptional<float> InMinimumDesiredWidth);
+
 	/** Gets an optional maximum desired width for this parameter editor. */
 	const TOptional<float>& GetMaximumDesiredWidth() const;
 
+	/** Sets the maximum desired width for this parameter editor.  If the option parameter is not set it clears the current value. */
+	void SetMaximumDesiredWidth(TOptional<float> InMaximumDesiredWidth);
+
 	/** Gets the desired horizontal alignment of this parameter editor for it's parent container. */
 	EHorizontalAlignment GetHorizontalAlignment() const;
 
diff --git a/Engine/Plugins/FX/Niagara/Source/NiagaraEditor/Public/ViewModels/Stack/NiagaraStackFunctionInput.h b/Engine/Plugins/FX/Niagara/Source/NiagaraEditor/Public/ViewModels/Stack/NiagaraStackFunctionInput.h
--- a/Engine/Plugins/FX/Niagara/Source/NiagaraEditor/Public/ViewModels/Stack/NiagaraStackFunctionInput.h
+++ b/Engine/Plugins/FX/Niagara/Source/NiagaraEditor/Public/ViewModels/Stack/NiagaraStackFunctionInput.h
@@ -6,6 +6,7 @@
 #include "NiagaraTypes.h"
 #include "ViewModels/Stack/NiagaraParameterHandle.h"
 #include "ViewModels/Stack/NiagaraStackFunctionInputCondition.h"
+#include "NiagaraStackEditorData.h"
 #include "NiagaraStackFunctionInput.generated.h"
 
 class UNiagaraNodeFunctionCall;
@@ -288,6 +289,10 @@ public:
 
 	TOptional<FNiagaraVariableMetaData> GetInputMetaData() const;
 	
+	void GetFilteredChildInputs(TArray<UNiagaraStackFunctionInput*>& OutFilteredChildInputs) const;
+
+	UNiagaraStackObject* GetChildDataObject() const;
+
 	virtual bool IsSemanticChild() const override;
 	void SetSemanticChild(bool IsSemanticChild);
 
@@ -296,6 +301,10 @@ public:
 
 	virtual const FCollectedUsageData& GetCollectedUsageData() const override;
 
+	ENiagaraStackEntryInlineDisplayMode GetInlineDisplayMode() const;
+
+	void SetInlineDisplayMode(ENiagaraStackEntryInlineDisplayMode InlineDisplayMode);
+
 	bool OpenSourceAsset() const;
 
 	bool SupportsCustomExpressions() const;
@@ -393,6 +402,8 @@ private:
 
 	UNiagaraScript* FindConversionScript(const FNiagaraTypeDefinition& FromType, TMap<FNiagaraTypeDefinition, UNiagaraScript*>& ConversionScriptCache, bool bIncludeConversionScripts) const;
 
+	bool FilterInlineChildren(const UNiagaraStackEntry& Child) const;
+
 private:
 	/** The module function call which owns this input entry. NOTE: This input might not be an input to the module function
 		call, it may be an input to a dynamic input function call which is owned by the module. */
diff --git a/Engine/Plugins/FX/Niagara/Source/NiagaraEditor/Public/ViewModels/Stack/NiagaraStackFunctionInputCollection.h b/Engine/Plugins/FX/Niagara/Source/NiagaraEditor/Public/ViewModels/Stack/NiagaraStackFunctionInputCollection.h
--- a/Engine/Plugins/FX/Niagara/Source/NiagaraEditor/Public/ViewModels/Stack/NiagaraStackFunctionInputCollection.h
+++ b/Engine/Plugins/FX/Niagara/Source/NiagaraEditor/Public/ViewModels/Stack/NiagaraStackFunctionInputCollection.h
@@ -83,6 +83,8 @@ public:
 
 	void GetChildInputs(TArray<UNiagaraStackFunctionInput*>& OutResult) const;
 
+	void GetFilteredChildInputs(TArray<UNiagaraStackFunctionInput*>& OutFilteredChildInputs) const;
+
 	void GetCustomFilteredChildInputs(TArray<UNiagaraStackFunctionInput*>& OutResult, const TArray<FOnFilterChild>& CustomFilters) const;
 
 	TArray<UNiagaraStackFunctionInput*> GetInlineParameterInputs() const;
diff --git a/Engine/Plugins/FX/Niagara/Source/NiagaraEditor/Public/ViewModels/Stack/NiagaraStackInputCategory.h b/Engine/Plugins/FX/Niagara/Source/NiagaraEditor/Public/ViewModels/Stack/NiagaraStackInputCategory.h
--- a/Engine/Plugins/FX/Niagara/Source/NiagaraEditor/Public/ViewModels/Stack/NiagaraStackInputCategory.h
+++ b/Engine/Plugins/FX/Niagara/Source/NiagaraEditor/Public/ViewModels/Stack/NiagaraStackInputCategory.h
@@ -79,6 +79,8 @@ public:
 
 	void SetStandardValuesFromClipboardFunctionInputs(const TArray<const UNiagaraClipboardFunctionInput*>& ClipboardFunctionInputs);
 
+	void GetFilteredChildInputs(TArray<UNiagaraStackFunctionInput*>& OutFilteredChildInputs) const;
+
 protected:
 	//~ UNiagaraStackEntry interface
 	virtual void RefreshChildrenInternal(const TArray<UNiagaraStackEntry*>& CurrentChildren, TArray<UNiagaraStackEntry*>& NewChildren, TArray<FStackIssue>& NewIssues) override;
diff --git a/Engine/Plugins/FX/Niagara/Source/NiagaraEditorWidgets/Private/NiagaraEditorWidgetsStyle.cpp b/Engine/Plugins/FX/Niagara/Source/NiagaraEditorWidgets/Private/NiagaraEditorWidgetsStyle.cpp
--- a/Engine/Plugins/FX/Niagara/Source/NiagaraEditorWidgets/Private/NiagaraEditorWidgetsStyle.cpp
+++ b/Engine/Plugins/FX/Niagara/Source/NiagaraEditorWidgets/Private/NiagaraEditorWidgetsStyle.cpp
@@ -133,7 +133,32 @@ FNiagaraEditorWidgetsStyle::FNiagaraEditorWidgetsStyle() : FSlateStyleSet("Niaga
 	FSlateFontInfo PerfWidgetEvalTypeFont = DEFAULT_FONT("Regular", 7);
 	Set("NiagaraEditor.Stack.Stats.EvalTypeFont", PerfWidgetEvalTypeFont);
 	
-
+	FSlateFontInfo StackInlineDynamicInputFormatFont = DEFAULT_FONT("Bold", 12);
+	FTextBlockStyle StackInlineDynamicInputFormatText = FTextBlockStyle(NormalText)
+		.SetFont(StackInlineDynamicInputFormatFont);
+	Set("NiagaraEditor.Stack.InlineDynamicInputFormatText", StackInlineDynamicInputFormatText);
+
+	FSlateFontInfo StackInlineDynamicInputFunctionFont = DEFAULT_FONT("Bold", 10);
+	FTextBlockStyle StackInlineDynamicInputFunctionText = FTextBlockStyle(NormalText)
+		.SetFont(StackInlineDynamicInputFunctionFont);
+	Set("NiagaraEditor.Stack.InlineDynamicInputFunctionText", StackInlineDynamicInputFunctionText);
+
+	Set("NiagaraEditor.Stack.InlineDynamicInput.GraphNode.Row.Top",				new BOX_BRUSH("Icons/MultiPartBorderRowTop", FMargin(1.0f / 4.0f)));
+	Set("NiagaraEditor.Stack.InlineDynamicInput.GraphNode.Row.Middle",			new BOX_BRUSH("Icons/MultiPartBorderRowMiddle", FMargin(1.0f / 4.0f)));
+	Set("NiagaraEditor.Stack.InlineDynamicInput.GraphNode.Row.Bottom",			new BOX_BRUSH("Icons/MultiPartBorderRowBottom", FMargin(1.0f / 4.0f)));
+	Set("NiagaraEditor.Stack.InlineDynamicInput.GraphNode.Column.Left",			new BOX_BRUSH("Icons/MultiPartBorderColumnLeft", FMargin(1.0f / 4.0f)));
+	Set("NiagaraEditor.Stack.InlineDynamicInput.GraphNode.Column.Middle",		new BOX_BRUSH("Icons/MultiPartBorderColumnMiddle", FMargin(1.0f / 4.0f)));
+	Set("NiagaraEditor.Stack.InlineDynamicInput.GraphNode.Column.Right",		new BOX_BRUSH("Icons/MultiPartBorderColumnRight", FMargin(1.0f / 4.0f)));
+	Set("NiagaraEditor.Stack.InlineDynamicInput.GraphNode.Corner.BottomLeft",	new BOX_BRUSH("Icons/MultiPartBorderCornerBottomLeft", FMargin(1.0f / 4.0f)));
+	Set("NiagaraEditor.Stack.InlineDynamicInput.GraphNode.Corner.BottomRight",	new BOX_BRUSH("Icons/MultiPartBorderCornerBottomRight", FMargin(1.0f / 4.0f)));
+	Set("NiagaraEditor.Stack.InlineDynamicInput.GraphNode.Edge.Bottom",			new BOX_BRUSH("Icons/MultiPartBorderEdgeBottom", FMargin(1.0f / 4.0f)));
+	Set("NiagaraEditor.Stack.InlineDynamicInput.GraphNode.Single",				new BOX_BRUSH("Icons/MultiPartBorderSingle", FMargin(1.0f / 4.0f)));
+	Set("NiagaraEditor.Stack.InlineDynamicInput.GraphNode.Pointer",				new IMAGE_BRUSH("Icons/MultiPartBorderPointer", Icon16x16));
+
+	FComboButtonStyle CompactComboButton = FAppStyle::Get().GetWidgetStyle<FComboButtonStyle>("ComboButton");
+	CompactComboButton.ButtonStyle = FAppStyle::Get().GetWidgetStyle<FButtonStyle>("HoverHintOnly");
+	CompactComboButton.DownArrowPadding = FMargin(1.0f, 0.0f);
+	Set("NiagaraEditor.Stack.CompactComboButton", CompactComboButton);
 
 	FSlateFontInfo StackSubduedItemFont = DEFAULT_FONT("Regular", 9);
 	FTextBlockStyle StackSubduedItemText = FTextBlockStyle(NormalText)
diff --git a/Engine/Plugins/FX/Niagara/Source/NiagaraEditorWidgets/Private/SNiagaraStack.cpp b/Engine/Plugins/FX/Niagara/Source/NiagaraEditorWidgets/Private/SNiagaraStack.cpp
--- a/Engine/Plugins/FX/Niagara/Source/NiagaraEditorWidgets/Private/SNiagaraStack.cpp
+++ b/Engine/Plugins/FX/Niagara/Source/NiagaraEditorWidgets/Private/SNiagaraStack.cpp
@@ -35,8 +35,9 @@
 #include "IDetailTreeNode.h"
 #include "Stack/NiagaraStackPropertyRowUtilities.h"
 #include "Stack/SNiagaraStackFunctionInputName.h"
-#include "Stack/SNiagaraStackFunctionInputValue.h"
 #include "Stack/SNiagaraStackFunctionInputCollection.h"
+#include "Stack/SNiagaraStackFunctionInputValue.h"
+#include "Stack/SNiagaraStackInlineDynamicInput.h"
 #include "Stack/SNiagaraStackItem.h"
 #include "Stack/SNiagaraStackItemFooter.h"
 #include "Stack/SNiagaraStackItemGroup.h"
@@ -895,9 +896,46 @@ SNiagaraStack::FRowWidgets SNiagaraStack::ConstructNameAndValueWidgetsForItem(UN
 		TSharedRef<SNiagaraStackFunctionInputName> FunctionInputNameWidget =
 			SNew(SNiagaraStackFunctionInputName, FunctionInput, StackViewModel)
 			.IsSelected(Container, &SNiagaraStackTableRow::IsSelected);
-		Container->AddFillRowContextMenuHandler(SNiagaraStackTableRow::FOnFillRowContextMenu::CreateSP(FunctionInputNameWidget, &SNiagaraStackFunctionInputName::FillRowContextMenu));
 
-		return FRowWidgets(FunctionInputNameWidget,	SNew(SNiagaraStackFunctionInputValue, FunctionInput));
+		FunctionInputNameWidget->OnRequestReconstructRow().AddSP(this, &SNiagaraStack::ReconstructNameAndValueWidgetsForItem, Item);
+
+		bool bShowInlineDynamicInput = 
+			GbEnableExperimentalInlineDynamicInputs &&
+			FunctionInput->GetValueMode() == UNiagaraStackFunctionInput::EValueMode::Dynamic &&
+			FunctionInput->GetInlineDisplayMode() != ENiagaraStackEntryInlineDisplayMode::None;
+		TSharedPtr<SWidget> FunctionInputValueWidget = nullptr;
+		if (bShowInlineDynamicInput)
+		{
+			FunctionInputValueWidget = SNew(SNiagaraStackInlineDynamicInput, FunctionInput);
+		}
+		else
+		{
+			FunctionInputValueWidget = SNew(SNiagaraStackFunctionInputValue, FunctionInput);
+		}
+
+		Container->AddFillRowContextMenuHandler(SNiagaraStackTableRow::FOnFillRowContextMenu::CreateSP(FunctionInputNameWidget, &SNiagaraStackFunctionInputName::FillRowContextMenu));
+		if (bShowInlineDynamicInput && 
+			(FunctionInput->GetInlineDisplayMode() == ENiagaraStackEntryInlineDisplayMode::GraphHorizontal ||
+			FunctionInput->GetInlineDisplayMode() == ENiagaraStackEntryInlineDisplayMode::GraphVertical ||
+			FunctionInput->GetInlineDisplayMode() == ENiagaraStackEntryInlineDisplayMode::GraphHybrid))
+		{
+			return FRowWidgets(
+				SNew(SVerticalBox)
+				+ SVerticalBox::Slot()
+				.AutoHeight()
+				[
+					FunctionInputNameWidget
+				]
+				+ SVerticalBox::Slot()
+				.AutoHeight()
+				[
+					FunctionInputValueWidget.ToSharedRef()
+				]);
+		}
+		else
+		{
+			return FRowWidgets(FunctionInputNameWidget, FunctionInputValueWidget.ToSharedRef());
+		}
 	}
 	else if (Item->IsA<UNiagaraStackErrorItem>())
 	{
@@ -1099,6 +1137,18 @@ SNiagaraStack::FRowWidgets SNiagaraStack::ConstructNameAndValueWidgetsForItem(UN
 	}
 }
 
+void SNiagaraStack::ReconstructNameAndValueWidgetsForItem(UNiagaraStackEntry* Item)
+{
+	TSharedPtr<ITableRow> TableRow = StackTree->WidgetFromItem(Item);
+	if (TableRow.IsValid())
+	{
+		TSharedRef<SNiagaraStackTableRow> Container = StaticCastSharedRef<SNiagaraStackTableRow>(TableRow.ToSharedRef());
+		Container->Reset();
+		FRowWidgets RowWidgets = ConstructNameAndValueWidgetsForItem(Item, Container);
+		Container->SetNameAndValueContent(RowWidgets.NameWidget, RowWidgets.ValueWidget);
+	}
+}
+
 void SNiagaraStack::OnGetChildren(UNiagaraStackEntry* Item, TArray<UNiagaraStackEntry*>& Children)
 {
 	Item->GetFilteredChildren(Children);
diff --git a/Engine/Plugins/FX/Niagara/Source/NiagaraEditorWidgets/Private/SNiagaraStack.h b/Engine/Plugins/FX/Niagara/Source/NiagaraEditorWidgets/Private/SNiagaraStack.h
--- a/Engine/Plugins/FX/Niagara/Source/NiagaraEditorWidgets/Private/SNiagaraStack.h
+++ b/Engine/Plugins/FX/Niagara/Source/NiagaraEditorWidgets/Private/SNiagaraStack.h
@@ -59,6 +59,8 @@ private:
 	TSharedRef<SNiagaraStackTableRow> ConstructContainerForItem(UNiagaraStackEntry* Item);
 
 	FRowWidgets ConstructNameAndValueWidgetsForItem(UNiagaraStackEntry* Item, TSharedRef<SNiagaraStackTableRow> Container);
+
+	void ReconstructNameAndValueWidgetsForItem(UNiagaraStackEntry* Item);
 	
 	void OnGetChildren(UNiagaraStackEntry* Item, TArray<UNiagaraStackEntry*>& Children);
 
diff --git a/Engine/Plugins/FX/Niagara/Source/NiagaraEditorWidgets/Private/Stack/SNiagaraStackFunctionInputName.cpp b/Engine/Plugins/FX/Niagara/Source/NiagaraEditorWidgets/Private/Stack/SNiagaraStackFunctionInputName.cpp
--- a/Engine/Plugins/FX/Niagara/Source/NiagaraEditorWidgets/Private/Stack/SNiagaraStackFunctionInputName.cpp
+++ b/Engine/Plugins/FX/Niagara/Source/NiagaraEditorWidgets/Private/Stack/SNiagaraStackFunctionInputName.cpp
@@ -26,6 +26,8 @@ void SNiagaraStackFunctionInputName::Construct(const FArguments& InArgs, UNiagar
 	StackEntryItem = InFunctionInput;
 	IsSelected = InArgs._IsSelected;
 
+	FunctionInput->OnStructureChanged().AddSP(this, &SNiagaraStackFunctionInputName::InputStructureChanged);
+
 	TSharedPtr<SHorizontalBox> NameBox;
 	ChildSlot
 	[
@@ -131,11 +133,88 @@ void SNiagaraStackFunctionInputName::FillRowContextMenu(class FMenuBuilder& Menu
 				LOCTEXT("ChangeNamespaceModifier", "Namespace Modifier"),
 				LOCTEXT("ChangeNamespaceModifierToolTip", "Edit the namespace modifier for the selected parameter."),
 				FNewMenuDelegate::CreateSP(this, &SNiagaraStackFunctionInputName::GetChangeNamespaceModifierSubMenu));
+
+			if (FunctionInput->CanDeleteInput())
+			{
+				MenuBuilder.AddMenuEntry(
+					LOCTEXT("DeleteInput", "Remove this parameter"),
+					LOCTEXT("DeleteInputToolTip", "Remove this parameter from set parameters node."),
+					FSlateIcon(),
+					FUIAction(
+						FNiagaraMenuAction_Generic::FOnExecuteAction::CreateUObject(FunctionInput, &UNiagaraStackFunctionInput::DeleteInput),
+						FNiagaraMenuAction_Generic::FCanExecuteAction::CreateUObject(FunctionInput, &UNiagaraStackFunctionInput::CanDeleteInput)));
+			}
+		}
+		MenuBuilder.EndSection();
+	}
+	if (FunctionInput->GetValueMode() == UNiagaraStackFunctionInput::EValueMode::Dynamic && GbEnableExperimentalInlineDynamicInputs)
+	{
+		MenuBuilder.BeginSection("View", LOCTEXT("ViewHeader", "View"));
+		{
+			MenuBuilder.AddMenuEntry(
+				LOCTEXT("UseDefaultDisplay", "Use Default Display"),
+				LOCTEXT("UseDefaultDisplayToolTip", "Use the default multi-row tree view display for this input."),
+				FSlateIcon(),
+				FUIAction(
+					FExecuteAction::CreateSP(this, &SNiagaraStackFunctionInputName::OnSetUseDefaultDisplay),
+					FCanExecuteAction(),
+					FGetActionCheckState::CreateSP(this, &SNiagaraStackFunctionInputName::GetUseDefaultDisplayCheckboxState)),
+				NAME_None,
+				EUserInterfaceActionType::RadioButton);
+
+			MenuBuilder.AddMenuEntry(
+				LOCTEXT("UseInlineExpressionDisplay", "Use Inline Expression Display"),
+				LOCTEXT("UseInlineExpressionDisplayToolTip", "Change the format of this input value to display inputs as an editable inline expression."),
+				FSlateIcon(),
+				FUIAction(
+					FExecuteAction::CreateSP(this, &SNiagaraStackFunctionInputName::OnSetUseInlineExpressionDisplay),
+					FCanExecuteAction(),
+					FGetActionCheckState::CreateSP(this, &SNiagaraStackFunctionInputName::GetUseInlineExpressionDisplayCheckboxState)),
+				NAME_None,
+				EUserInterfaceActionType::RadioButton);
+
+			MenuBuilder.AddMenuEntry(
+				LOCTEXT("UseInlineHorizontalGraphDisplay", "Use Inline Horizontal Graph Display"),
+				LOCTEXT("UseInlineHorizontalGraphDisplayToolTip", "Change the format of this input value to display dynamic inputs as a horizontal inline graph."),
+				FSlateIcon(),
+				FUIAction(
+					FExecuteAction::CreateSP(this, &SNiagaraStackFunctionInputName::OnSetUseInlineHorizontalGraphDisplay),
+					FCanExecuteAction(),
+					FGetActionCheckState::CreateSP(this, &SNiagaraStackFunctionInputName::GetUseInlineHorizontalGraphDisplayCheckboxState)),
+				NAME_None,
+				EUserInterfaceActionType::RadioButton);
+
+			MenuBuilder.AddMenuEntry(
+				LOCTEXT("UseInlineVerticalGraphDisplay", "Use Inline Vertical Graph Display"),
+				LOCTEXT("UseInlineVerticalGraphDisplayToolTip", "Change the format of this input value to display dynamic inputs as a vertical inline graph."),
+				FSlateIcon(),
+				FUIAction(
+					FExecuteAction::CreateSP(this, &SNiagaraStackFunctionInputName::OnSetUseInlineVerticalGraphDisplay),
+					FCanExecuteAction(),
+					FGetActionCheckState::CreateSP(this, &SNiagaraStackFunctionInputName::GetUseInlineVerticalGraphDisplayCheckboxState)),
+				NAME_None,
+				EUserInterfaceActionType::RadioButton);
+
+			MenuBuilder.AddMenuEntry(
+				LOCTEXT("UseInlineHybridGraphDisplay", "Use Inline Hybrid Graph Display"),
+				LOCTEXT("UseInlineHybridGraphDisplayToolTip", "Change the format of this input value to display dynamic inputs as a hybrid inline graph."),
+				FSlateIcon(),
+				FUIAction(
+					FExecuteAction::CreateSP(this, &SNiagaraStackFunctionInputName::OnSetUseInlineHybridGraphDisplay),
+					FCanExecuteAction(),
+					FGetActionCheckState::CreateSP(this, &SNiagaraStackFunctionInputName::GetUseInlineHybridGraphDisplayCheckboxState)),
+				NAME_None,
+				EUserInterfaceActionType::RadioButton);
 		}
 		MenuBuilder.EndSection();
 	}
 }
 
+FSimpleMulticastDelegate& SNiagaraStackFunctionInputName::OnRequestReconstructRow()
+{
+	return OnRequestReconstructRowDelegate;
+}
+
 EVisibility SNiagaraStackFunctionInputName::GetEditConditionCheckBoxVisibility() const
 {
 	return FunctionInput->GetHasEditCondition() && FunctionInput->GetShowEditConditionInline() ? EVisibility::Visible : EVisibility::Collapsed;
@@ -330,4 +409,67 @@ void SNiagaraStackFunctionInputName::OnCopyParameterReference()
 	FPlatformApplicationMisc::ClipboardCopy(*FunctionInput->GetDisplayName().ToString());
 }
 
+ECheckBoxState SNiagaraStackFunctionInputName::GetUseDefaultDisplayCheckboxState() const
+{
+	return FunctionInput->GetInlineDisplayMode() == ENiagaraStackEntryInlineDisplayMode::None ? ECheckBoxState::Checked : ECheckBoxState::Unchecked;
+}
+
+void SNiagaraStackFunctionInputName::OnSetUseDefaultDisplay()
+{
+	FunctionInput->SetInlineDisplayMode(ENiagaraStackEntryInlineDisplayMode::None);
+	OnRequestReconstructRowDelegate.Broadcast();
+}
+
+ECheckBoxState SNiagaraStackFunctionInputName::GetUseInlineExpressionDisplayCheckboxState() const
+{
+	return FunctionInput->GetInlineDisplayMode() == ENiagaraStackEntryInlineDisplayMode::Expression ? ECheckBoxState::Checked : ECheckBoxState::Unchecked;
+}
+
+void SNiagaraStackFunctionInputName::OnSetUseInlineExpressionDisplay()
+{
+	FunctionInput->SetInlineDisplayMode(ENiagaraStackEntryInlineDisplayMode::Expression);
+	OnRequestReconstructRowDelegate.Broadcast();
+}
+
+ECheckBoxState SNiagaraStackFunctionInputName::GetUseInlineHorizontalGraphDisplayCheckboxState() const
+{
+	return FunctionInput->GetInlineDisplayMode() == ENiagaraStackEntryInlineDisplayMode::GraphHorizontal ? ECheckBoxState::Checked : ECheckBoxState::Unchecked;
+}
+
+void SNiagaraStackFunctionInputName::OnSetUseInlineHorizontalGraphDisplay()
+{
+	FunctionInput->SetInlineDisplayMode(ENiagaraStackEntryInlineDisplayMode::GraphHorizontal);
+	OnRequestReconstructRowDelegate.Broadcast();
+}
+
+ECheckBoxState SNiagaraStackFunctionInputName::GetUseInlineVerticalGraphDisplayCheckboxState() const
+{
+	return FunctionInput->GetInlineDisplayMode() == ENiagaraStackEntryInlineDisplayMode::GraphVertical ? ECheckBoxState::Checked : ECheckBoxState::Unchecked;
+}
+
+void SNiagaraStackFunctionInputName::OnSetUseInlineVerticalGraphDisplay()
+{
+	FunctionInput->SetInlineDisplayMode(ENiagaraStackEntryInlineDisplayMode::GraphVertical);
+	OnRequestReconstructRowDelegate.Broadcast();
+}
+
+ECheckBoxState SNiagaraStackFunctionInputName::GetUseInlineHybridGraphDisplayCheckboxState() const
+{
+	return FunctionInput->GetInlineDisplayMode() == ENiagaraStackEntryInlineDisplayMode::GraphHybrid ? ECheckBoxState::Checked : ECheckBoxState::Unchecked;
+}
+
+void SNiagaraStackFunctionInputName::OnSetUseInlineHybridGraphDisplay()
+{
+	FunctionInput->SetInlineDisplayMode(ENiagaraStackEntryInlineDisplayMode::GraphHybrid);
+	OnRequestReconstructRowDelegate.Broadcast();
+}
+
+void SNiagaraStackFunctionInputName::InputStructureChanged(ENiagaraStructureChangedFlags Flags)
+{
+	if (FunctionInput != nullptr && FunctionInput->IsFinalized() == false && FunctionInput->GetInlineDisplayMode() != ENiagaraStackEntryInlineDisplayMode::None)
+	{
+		OnRequestReconstructRowDelegate.Broadcast();
+	}
+}
+
 #undef LOCTEXT_NAMESPACE
diff --git a/Engine/Plugins/FX/Niagara/Source/NiagaraEditorWidgets/Private/Stack/SNiagaraStackFunctionInputName.h b/Engine/Plugins/FX/Niagara/Source/NiagaraEditorWidgets/Private/Stack/SNiagaraStackFunctionInputName.h
--- a/Engine/Plugins/FX/Niagara/Source/NiagaraEditorWidgets/Private/Stack/SNiagaraStackFunctionInputName.h
+++ b/Engine/Plugins/FX/Niagara/Source/NiagaraEditorWidgets/Private/Stack/SNiagaraStackFunctionInputName.h
@@ -13,9 +13,6 @@ class SNiagaraParameterNameTextBlock;
 
 class SNiagaraStackFunctionInputName: public SNiagaraStackEntryWidget
 {
-public:
-	DECLARE_DELEGATE_OneParam(FOnColumnWidthChanged, float);
-
 public:
 	SLATE_BEGIN_ARGS(SNiagaraStackFunctionInputName) { }
 		SLATE_ATTRIBUTE(bool, IsSelected);
@@ -27,6 +24,8 @@ public:
 
 	void FillRowContextMenu(class FMenuBuilder& MenuBuilder);
 
+	FSimpleMulticastDelegate& OnRequestReconstructRow();
+
 private:
 	EVisibility GetEditConditionCheckBoxVisibility() const;
 
@@ -61,6 +60,23 @@ private:
 	bool CanCopyParameterReference() const;
 	void OnCopyParameterReference();
 
+	ECheckBoxState GetUseDefaultDisplayCheckboxState() const;
+	void OnSetUseDefaultDisplay();
+
+	ECheckBoxState GetUseInlineExpressionDisplayCheckboxState() const;
+	void OnSetUseInlineExpressionDisplay();
+
+	ECheckBoxState GetUseInlineHorizontalGraphDisplayCheckboxState() const;
+	void OnSetUseInlineHorizontalGraphDisplay();
+
+	ECheckBoxState GetUseInlineVerticalGraphDisplayCheckboxState() const;
+	void OnSetUseInlineVerticalGraphDisplay();
+
+	ECheckBoxState GetUseInlineHybridGraphDisplayCheckboxState() const;
+	void OnSetUseInlineHybridGraphDisplay();
+
+	void InputStructureChanged(ENiagaraStructureChangedFlags Flags);
+
 private:
 	UNiagaraStackFunctionInput* FunctionInput;
 
@@ -69,4 +85,6 @@ private:
 	TSharedPtr<SNiagaraParameterNameTextBlock> ParameterTextBlock;
 
 	TAttribute<bool> IsSelected;
+
+	FSimpleMulticastDelegate OnRequestReconstructRowDelegate;
 };
\ No newline at end of file
diff --git a/Engine/Plugins/FX/Niagara/Source/NiagaraEditorWidgets/Private/Stack/SNiagaraStackFunctionInputValue.cpp b/Engine/Plugins/FX/Niagara/Source/NiagaraEditorWidgets/Private/Stack/SNiagaraStackFunctionInputValue.cpp
--- a/Engine/Plugins/FX/Niagara/Source/NiagaraEditorWidgets/Private/Stack/SNiagaraStackFunctionInputValue.cpp
+++ b/Engine/Plugins/FX/Niagara/Source/NiagaraEditorWidgets/Private/Stack/SNiagaraStackFunctionInputValue.cpp
@@ -20,6 +20,7 @@
 #include "NiagaraNodeAssignment.h"
 #include "NiagaraNodeFunctionCall.h"
 #include "NiagaraParameterCollection.h"
+#include "NiagaraStackCommandContext.h"
 #include "NiagaraScriptVariable.h"
 #include "NiagaraSettings.h"
 #include "NiagaraSystem.h"
@@ -65,13 +66,16 @@ bool SNiagaraStackFunctionInputValue::bLibraryOnly = true;
 void SNiagaraStackFunctionInputValue::Construct(const FArguments& InArgs, UNiagaraStackFunctionInput* InFunctionInput)
 {
 	FunctionInput = InFunctionInput;
+	LayoutMode = InArgs._LayoutMode;
+	CompactActionMenuButtonVisibilityAttribute = InArgs._CompactActionMenuButtonVisibility;
 	FunctionInput->OnValueChanged().AddSP(this, &SNiagaraStackFunctionInputValue::OnInputValueChanged);
 	SyntaxHighlighter = FNiagaraHLSLSyntaxHighlighter::Create();
 
+	TSharedPtr<SHorizontalBox> OuterChildrenBox;
 	TSharedPtr<SHorizontalBox> ChildrenBox;
 	ChildSlot
 	[
-		SAssignNew(ChildrenBox, SHorizontalBox)
+		SAssignNew(OuterChildrenBox, SHorizontalBox)
 		+ SHorizontalBox::Slot()
 		[
 			SNew(SNiagaraParameterDropTarget)
@@ -87,16 +91,24 @@ void SNiagaraStackFunctionInputValue::Construct(const FArguments& InArgs, UNiaga
 			.bUseAllowDropCache(true)
 			.Content()
 			[
-				// Values
-				SNew(SHorizontalBox)
+				SAssignNew(ChildrenBox, SHorizontalBox)
 				.IsEnabled(this, &SNiagaraStackFunctionInputValue::GetInputEnabled)
-				+ SHorizontalBox::Slot()
+			])
+		]
+	];
+			// Values
+	if (LayoutMode == ELayoutMode::FullRow)
+	{
+			ChildrenBox->AddSlot()
 				.AutoWidth()
 				.Padding(0, 0, 3, 0)
 				[
 					SNew(SNiagaraStackIndent, FunctionInput, ENiagaraStackIndentMode::Value)
-				]
-				+ SHorizontalBox::Slot()
+				];
+	}
+	if (LayoutMode == ELayoutMode::FullRow || LayoutMode == ELayoutMode::CompactInline)
+	{
+			ChildrenBox->AddSlot()
 				.VAlign(VAlign_Center)
 				.AutoWidth()
 				.Padding(0, 0, 3, 0)
@@ -123,7 +135,9 @@ void SNiagaraStackFunctionInputValue::Construct(const FArguments& InArgs, UNiaga
 							.ColorAndOpacity(this, &SNiagaraStackFunctionInputValue::GetInputIconColor)
 						]
 					]
-					+ SHorizontalBox::Slot()
+				];
+
+				ChildrenBox->AddSlot()
 					.AutoWidth()
 					[
 						// icons type info pill
@@ -160,9 +174,9 @@ void SNiagaraStackFunctionInputValue::Construct(const FArguments& InArgs, UNiaga
 								.Image(FNiagaraEditorStyle::Get().GetBrush("NiagaraEditor.Module.DynamicInput"))
 							]
 						]
-					]
-				]
-				+ SHorizontalBox::Slot()
+					];
+
+			ChildrenBox->AddSlot()
 				.VAlign(VAlign_Center)
 				.AutoWidth()
 				.Padding(0, 0, 3, 0)
@@ -177,8 +191,9 @@ void SNiagaraStackFunctionInputValue::Construct(const FArguments& InArgs, UNiaga
 						.Image(GetTypeModifierIcon())
 						.ToolTipText(GetTypeModifierIconToolTip())
 					]
-				]
-				+ SHorizontalBox::Slot()
+				];
+
+			ChildrenBox->AddSlot()
 				.VAlign(VAlign_Center)
 				[
 					// Value container and widgets.
@@ -187,10 +202,13 @@ void SNiagaraStackFunctionInputValue::Construct(const FArguments& InArgs, UNiaga
 					[
 						ConstructValueWidgets()
 					]
-				]
+				];
+	}
 
+	if (LayoutMode == ELayoutMode::FullRow)
+	{ 
 				// Handle drop-down button
-				+ SHorizontalBox::Slot()
+				ChildrenBox->AddSlot()
 				.AutoWidth()
 				.VAlign(VAlign_Center)
 				.Padding(3, 0, 0, 0)
@@ -205,10 +223,10 @@ void SNiagaraStackFunctionInputValue::Construct(const FArguments& InArgs, UNiaga
 					.MenuPlacement(MenuPlacement_BelowRightAnchor)
 					.HAlign(HAlign_Center)
 					.VAlign(VAlign_Center)
-				]
+				];
 
 				// Reset Button
-				+ SHorizontalBox::Slot()
+				ChildrenBox->AddSlot()
 				.AutoWidth()
 				.VAlign(VAlign_Center)
 				.Padding(3, 0, 0, 0)
@@ -225,10 +243,10 @@ void SNiagaraStackFunctionInputValue::Construct(const FArguments& InArgs, UNiaga
 						SNew(SImage)
 						.Image(FAppStyle::GetBrush("PropertyWindow.DiffersFromDefault"))
 					]
-				]
+				];
 
 				// Reset to base Button
-				+ SHorizontalBox::Slot()
+				ChildrenBox->AddSlot()
 				.AutoWidth()
 				.VAlign(VAlign_Center)
 				.Padding(3, 0, 0, 0)
@@ -246,10 +264,26 @@ void SNiagaraStackFunctionInputValue::Construct(const FArguments& InArgs, UNiaga
 						.Image(FAppStyle::GetBrush("PropertyWindow.DiffersFromDefault"))
 						.ColorAndOpacity(FSlateColor(FLinearColor::Green))
 					]
-				]			
-			])
-		]
-	];
+				];
+	}
+
+	if (LayoutMode == ELayoutMode::CompactInline || LayoutMode == ELayoutMode::EditDropDownOnly)
+	{
+		ChildrenBox->AddSlot()
+		.AutoWidth()
+		.VAlign(VAlign_Center)
+		[
+			SNew(SComboButton)
+			.ComboButtonStyle(FNiagaraEditorWidgetsStyle::Get(), "NiagaraEditor.Stack.CompactComboButton")
+			.IsFocusable(false)
+			.ForegroundColor(FSlateColor::UseForeground())
+			.OnGetMenuContent(this, &SNiagaraStackFunctionInputValue::OnGetCompactActionMenu)
+			.Visibility(CompactActionMenuButtonVisibilityAttribute)
+			.MenuPlacement(MenuPlacement_BelowRightAnchor)
+			.HAlign(HAlign_Center)
+			.VAlign(VAlign_Center)
+		];
+	}
 
 	ValueModeForGeneratedWidgets = FunctionInput->GetValueMode();
 }
@@ -523,6 +557,10 @@ TSharedRef<SWidget> SNiagaraStackFunctionInputValue::ConstructLocalValueStructWi
 		if (TypeEditorUtilities.IsValid() && TypeEditorUtilities->CanCreateParameterEditor())
 		{
 			TSharedPtr<SNiagaraParameterEditor> ParameterEditor = TypeEditorUtilities->CreateParameterEditor(FunctionInput->GetInputType(), FunctionInput->GetInputDisplayUnit(), FunctionInput->GetInputWidgetCustomization());
+			if (LayoutMode == ELayoutMode::CompactInline && ParameterEditor->GetMinimumDesiredWidth().IsSet())
+			{
+				ParameterEditor->SetMinimumDesiredWidth(ParameterEditor->GetMinimumDesiredWidth().GetValue() / 2.0f);
+			}
 			ParameterEditor->UpdateInternalValueFromStruct(DisplayedLocalValueStruct.ToSharedRef());
 			ParameterEditor->SetOnBeginValueChange(SNiagaraParameterEditor::FOnValueChange::CreateSP(
 				this, &SNiagaraStackFunctionInputValue::ParameterBeginValueChange));
@@ -572,6 +610,11 @@ TSharedRef<SWidget> SNiagaraStackFunctionInputValue::ConstructLocalValueStructWi
 
 void SNiagaraStackFunctionInputValue::OnInputValueChanged()
 {
+	if (LayoutMode == ELayoutMode::EditDropDownOnly)
+	{
+		return;
+	}
+
 	if (ValueModeForGeneratedWidgets != FunctionInput->GetValueMode())
 	{
 		ValueContainer->SetContent(ConstructValueWidgets());
@@ -579,7 +622,7 @@ void SNiagaraStackFunctionInputValue::OnInputValueChanged()
 	}
 	else
 	{
-		if (ValueModeForGeneratedWidgets == UNiagaraStackFunctionInput::EValueMode::Local)
+		if (ValueModeForGeneratedWidgets == UNiagaraStackFunctionInput::EValueMode::Local && DisplayedLocalValueStruct.IsValid())
 		{
 			if (DisplayedLocalValueStruct->GetStruct() == FunctionInput->GetLocalValueStruct()->GetStruct())
 			{
@@ -809,10 +852,60 @@ TSharedRef<SWidget> SNiagaraStackFunctionInputValue::OnGetAvailableHandleMenu()
 		]
 	];
 
-	SetFunctionInputButton->SetMenuContentWidgetToFocus(ActionSelector->GetSearchBox());
+	if(SetFunctionInputButton.IsValid())
+	{ 
+		SetFunctionInputButton->SetMenuContentWidgetToFocus(ActionSelector->GetSearchBox());
+	}
 	return MenuWidget;
 }
 
+TSharedRef<SWidget> SNiagaraStackFunctionInputValue::OnGetCompactActionMenu()
+{
+	if (StackCommandContext.IsValid() == false)
+	{
+		StackCommandContext = MakeShared<FNiagaraStackCommandContext>();
+		TArray<UNiagaraStackEntry*> SelectedEntries;
+		SelectedEntries.Add(FunctionInput);
+		StackCommandContext->SetSelectedEntries(SelectedEntries);
+	}
+
+	FMenuBuilder MenuBuilder(true, StackCommandContext->GetCommands());
+	MenuBuilder.BeginSection("Value", LOCTEXT("ValueHeader", "Value"));
+	{
+		MenuBuilder.AddSubMenu(
+			LOCTEXT("AssignSubMenu", "Assign..."),
+			LOCTEXT("AssignSubMenuToolTip", "Assign this input a new value..."),
+			FNewMenuDelegate::CreateSP(this, &SNiagaraStackFunctionInputValue::OnFillAssignSubMenu));
+		MenuBuilder.AddMenuEntry(
+			LOCTEXT("ResetToDefaultMenuEntry", "Reset to Default"),
+			LOCTEXT("ResetToDefaultMenuEntryToolTip", "Reset this input to the value defined in the script."),
+			FSlateIcon(),
+			FUIAction(
+				FExecuteAction::CreateUObject(FunctionInput, &UNiagaraStackFunctionInput::Reset),
+				FCanExecuteAction::CreateUObject(FunctionInput, &UNiagaraStackFunctionInput::CanReset)));
+		if (FunctionInput->HasBaseEmitter())
+		{
+			MenuBuilder.AddMenuEntry(
+				LOCTEXT("ResetToBaseMenuEntry", "Reset to Base"),
+				LOCTEXT("ResetToBaseMenuEntryToolTip", "Reset this input to the value defined in the base emitter."),
+				FSlateIcon(),
+				FUIAction(
+					FExecuteAction::CreateUObject(FunctionInput, &UNiagaraStackFunctionInput::ResetToBase),
+					FCanExecuteAction::CreateUObject(FunctionInput, &UNiagaraStackFunctionInput::CanResetToBase)));
+		}
+	}
+	MenuBuilder.EndSection();
+
+	StackCommandContext->AddEditMenuItems(MenuBuilder);
+
+	return MenuBuilder.MakeWidget();
+}
+
+void SNiagaraStackFunctionInputValue::OnFillAssignSubMenu(FMenuBuilder& MenuBuilder)
+{
+	MenuBuilder.AddWidget(OnGetAvailableHandleMenu(), FText());
+}
+
 void SNiagaraStackFunctionInputValue::DynamicInputScriptSelected(UNiagaraScript* DynamicInputScript)
 {
 	FunctionInput->SetDynamicInput(DynamicInputScript);
diff --git a/Engine/Plugins/FX/Niagara/Source/NiagaraEditorWidgets/Private/Stack/SNiagaraStackFunctionInputValue.h b/Engine/Plugins/FX/Niagara/Source/NiagaraEditorWidgets/Private/Stack/SNiagaraStackFunctionInputValue.h
--- a/Engine/Plugins/FX/Niagara/Source/NiagaraEditorWidgets/Private/Stack/SNiagaraStackFunctionInputValue.h
+++ b/Engine/Plugins/FX/Niagara/Source/NiagaraEditorWidgets/Private/Stack/SNiagaraStackFunctionInputValue.h
@@ -18,6 +18,8 @@ class SBox;
 class IStructureDetailsView;
 class SComboButton;
 struct FGraphActionListBuilderBase;
+class FNiagaraStackCommandContext;
+class FMenuBuilder;
 
 typedef SItemSelector<FString, TSharedPtr<FNiagaraMenuAction_Generic>, ENiagaraMenuSections> SNiagaraMenuActionSelector;
 
@@ -25,8 +27,21 @@ class SNiagaraStackFunctionInputValue: public SCompoundWidget
 {
 public:
 	DECLARE_DELEGATE_OneParam(FOnColumnWidthChanged, float)
+
+	enum class ELayoutMode
+	{
+		FullRow,
+		CompactInline,
+		EditDropDownOnly
+	};
+
 public:
-	SLATE_BEGIN_ARGS(SNiagaraStackFunctionInputValue) { }
+	SLATE_BEGIN_ARGS(SNiagaraStackFunctionInputValue)
+		: _LayoutMode(ELayoutMode::FullRow)
+		, _CompactActionMenuButtonVisibility(EVisibility::Visible)
+		{ }
+		SLATE_ARGUMENT(ELayoutMode, LayoutMode)
+		SLATE_ATTRIBUTE(EVisibility, CompactActionMenuButtonVisibility)
 	SLATE_END_ARGS();
 
 	void Construct(const FArguments& InArgs, UNiagaraStackFunctionInput* InFunctionInput);
@@ -104,6 +119,10 @@ private:
 
 	TSharedRef<SWidget> OnGetAvailableHandleMenu();
 
+	TSharedRef<SWidget> OnGetCompactActionMenu();
+
+	void OnFillAssignSubMenu(FMenuBuilder& MenuBuilder);
+
 	TSharedRef<SWidget> GetVersionSelectorDropdownMenu();
 	void SwitchToVersion(FNiagaraAssetVersion Version);
 	FSlateColor GetVersionSelectorColor() const;
@@ -162,6 +181,8 @@ private:
 private:
 	UNiagaraStackFunctionInput* FunctionInput;
 
+	ELayoutMode LayoutMode;
+
 	TSharedPtr<SBox> ValueContainer;
 	UNiagaraStackFunctionInput::EValueMode ValueModeForGeneratedWidgets;
 
@@ -174,6 +195,9 @@ private:
 	TSharedPtr<SComboButton> SetFunctionInputButton;
 	TSharedPtr<FNiagaraHLSLSyntaxHighlighter> SyntaxHighlighter;
 
+	TAttribute<EVisibility> CompactActionMenuButtonVisibilityAttribute;
+	TSharedPtr<FNiagaraStackCommandContext> StackCommandContext;
+
 	static bool bLibraryOnly;
 
 private:
diff --git a/Engine/Plugins/FX/Niagara/Source/NiagaraEditorWidgets/Private/Stack/SNiagaraStackInlineDynamicInput.cpp b/Engine/Plugins/FX/Niagara/Source/NiagaraEditorWidgets/Private/Stack/SNiagaraStackInlineDynamicInput.cpp
new file mode 100644
--- /dev/null
+++ b/Engine/Plugins/FX/Niagara/Source/NiagaraEditorWidgets/Private/Stack/SNiagaraStackInlineDynamicInput.cpp
@@ -0,0 +1,1553 @@
+// Copyright Epic Games, Inc. All Rights Reserved.
+
+#include "Stack/SNiagaraStackInlineDynamicInput.h"
+
+#include "EditorFontGlyphs.h"
+#include "NiagaraEditorWidgetsStyle.h"
+#include "NiagaraEditorWidgetsUtilities.h"
+#include "NiagaraNodeFunctionCall.h"
+#include "NiagaraObjectSelection.h"
+#include "NiagaraScript.h"
+#include "Stack/SNiagaraStackFunctionInputValue.h"
+#include "Styling/StyleColors.h"
+#include "ViewModels/Stack/NiagaraStackEntry.h"
+#include "ViewModels/Stack/NiagaraStackFunctionInput.h"
+#include "ViewModels/Stack/NiagaraStackObject.h"
+#include "Widgets/Input/SCheckBox.h"
+#include "Widgets/Layout/SGridPanel.h"
+#include "Widgets/Layout/SWrapBox.h"
+#include "Widgets/SBoxPanel.h"
+#include "Widgets/SNiagaraSelectedObjectsDetails.h"
+
+
+#define LOCTEXT_NAMESPACE "NiagaraStackInlineDynamicInput"
+
+const float SNiagaraStackInlineDynamicInput::TextIconSize = 16;
+
+class SNiagaraStackInlineDynamicInputExpressionTokenBorder : public SBorder
+{
+public:
+	DECLARE_DELEGATE_TwoParams(FHoveredChanged, FGuid /*InputId*/, bool /*bIsHovered*/);
+
+public:
+	SLATE_BEGIN_ARGS(SNiagaraStackInlineDynamicInputExpressionTokenBorder)
+		: _HoveredInputId(FGuid())
+		{}
+		SLATE_ATTRIBUTE(FGuid, HoveredInputId)
+		SLATE_ATTRIBUTE(FMargin, Padding)
+		SLATE_EVENT(FHoveredChanged, OnHoveredChanged)
+		SLATE_DEFAULT_SLOT(FArguments, Content);
+	SLATE_END_ARGS();
+
+	void Construct(const FArguments& InArgs, UNiagaraStackFunctionInput* FunctionInput, FGuid InInputId)
+	{
+		FunctionInputWeak = FunctionInput;
+		InputId = InInputId;
+		HoveredInputIdAttribute = InArgs._HoveredInputId;
+		OnHoveredChangedDelegate = InArgs._OnHoveredChanged;
+		SBorder::Construct(SBorder::FArguments()
+			.BorderImage(FAppStyle::GetBrush("WhiteBrush"))
+			.BorderBackgroundColor(this, &SNiagaraStackInlineDynamicInputExpressionTokenBorder::GetBorderColor)
+			.VAlign(VAlign_Center)
+			.Padding(InArgs._Padding)
+			[
+				InArgs._Content.Widget
+			]);
+	}
+
+	virtual void OnMouseEnter(const FGeometry& MyGeometry, const FPointerEvent& MouseEvent) override
+	{
+		const bool bWasHovered = IsHovered();
+		SBorder::OnMouseEnter(MyGeometry, MouseEvent);
+		if (bWasHovered == false && IsHovered())
+		{
+			OnHoveredChangedDelegate.ExecuteIfBound(InputId, true);
+		}
+	}
+
+	virtual void OnMouseLeave(const FPointerEvent& MouseEvent) override
+	{
+		const bool bWasHovered = IsHovered();
+		SBorder::OnMouseLeave(MouseEvent);
+		if (bWasHovered && IsHovered() == false)
+		{
+			OnHoveredChangedDelegate.ExecuteIfBound(InputId, false);
+		}
+	}
+
+	virtual FReply OnMouseButtonDoubleClick(const FGeometry& InMyGeometry, const FPointerEvent& InMouseEvent) override
+	{
+		bool bHandled = false;
+		if (FunctionInputWeak.IsValid())
+		{
+			bHandled = FunctionInputWeak->OpenSourceAsset();
+		}
+		return bHandled 
+			? FReply::Handled()
+			: SBorder::OnMouseButtonDoubleClick(InMyGeometry, InMouseEvent);
+	}
+
+private:
+	FSlateColor GetBorderColor() const
+	{
+		return InputId == HoveredInputIdAttribute.Get() ? FLinearColor(1.0f, 1.0f, 1.0f, 0.25f) : FLinearColor::Transparent;
+	}
+
+private:
+	TWeakObjectPtr<UNiagaraStackFunctionInput> FunctionInputWeak;
+	FGuid InputId;
+	TAttribute<FGuid> HoveredInputIdAttribute;
+	FHoveredChanged OnHoveredChangedDelegate;
+};
+
+void SNiagaraStackInlineDynamicInput::Construct(const FArguments& InArgs, UNiagaraStackFunctionInput* InFunctionInput)
+{
+	RootFunctionInput = InFunctionInput;
+	RootFunctionInput->OnStructureChanged().AddSP(this, &SNiagaraStackInlineDynamicInput::RootFunctionInputStructureChanged);
+	RootFunctionInput->OnValueChanged().AddSP(this, &SNiagaraStackInlineDynamicInput::RootFunctionInputValueChanged);
+	bSingleValueMode = false;
+	ConstructChildren();
+}
+
+SNiagaraStackInlineDynamicInput::FInputDisplayEntries SNiagaraStackInlineDynamicInput::CollectInputDisplayEntriesRecursive(UNiagaraStackFunctionInput* Input)
+{
+	FInputDisplayEntries InputDisplayEntries;
+	InputDisplayEntries.InputEntry = Input;
+	if (Input->GetValueMode() == UNiagaraStackFunctionInput::EValueMode::Dynamic)
+	{
+		TArray<UNiagaraStackFunctionInput*> ChildInputs;
+		Input->GetFilteredChildInputs(ChildInputs);
+		for (UNiagaraStackFunctionInput* ChildInput : ChildInputs)
+		{
+			InputDisplayEntries.ChildInputDisplayEntries.Add(CollectInputDisplayEntriesRecursive(ChildInput));
+		}
+	}
+	else if (Input->GetValueMode() == UNiagaraStackFunctionInput::EValueMode::Data && Input->GetChildDataObject() != nullptr)
+	{
+		InputDisplayEntries.ChildObjectEntry = Input->GetChildDataObject();
+	}
+	return InputDisplayEntries;
+}
+
+void SNiagaraStackInlineDynamicInput::RootFunctionInputStructureChanged(ENiagaraStructureChangedFlags Flags)
+{
+	FInputDisplayEntries NewInputDisplayEntries = CollectInputDisplayEntriesRecursive(RootFunctionInput);
+	if(RootInputDisplayEntries != NewInputDisplayEntries)
+	{
+		ConstructChildren();
+	}
+}
+
+void SNiagaraStackInlineDynamicInput::RootFunctionInputValueChanged()
+{
+	FInputDisplayEntries NewInputDisplayEntries = CollectInputDisplayEntriesRecursive(RootFunctionInput);
+	if (RootInputDisplayEntries != NewInputDisplayEntries)
+	{
+		ConstructChildren();
+	}
+}
+
+bool SNiagaraStackInlineDynamicInput::TryMatchFormatInputs(
+	const TArray<FInputDisplayEntries>& ChildInputDisplayEntries,
+	const TArray<FNiagaraInlineDynamicInputFormatToken>& InlineFormat,
+	TArray<TPair<FNiagaraInlineDynamicInputFormatToken, FInputDisplayEntries>>& OutMatchedFormatInputs)
+{
+	if (InlineFormat.Num() == 0 || GbEnableCustomInlineDynamicInputFormats == false)
+	{
+		return false;
+	}
+
+	TArray<FString> MatchedInputNames;
+	for (const FNiagaraInlineDynamicInputFormatToken& InlineFormatToken : InlineFormat)
+	{
+		if (InlineFormatToken.Usage == ENiagaraInlineDynamicInputFormatTokenUsage::Input)
+		{
+			const FInputDisplayEntries* MatchingChildInputDisplayEntry = ChildInputDisplayEntries.FindByPredicate([InlineFormatToken](const FInputDisplayEntries& ChildInputDisplayEntry)
+				{ return ChildInputDisplayEntry.InputEntry->GetInputParameterHandle().GetName().ToString() == InlineFormatToken.Value; });
+			if (MatchingChildInputDisplayEntry == nullptr)
+			{
+				OutMatchedFormatInputs.Empty();
+				return false;
+			}
+			else
+			{
+				OutMatchedFormatInputs.Add(TPair<FNiagaraInlineDynamicInputFormatToken, FInputDisplayEntries>(InlineFormatToken, *MatchingChildInputDisplayEntry));
+				MatchedInputNames.AddUnique(InlineFormatToken.Value);
+			}
+		}
+		else if (InlineFormatToken.Usage == ENiagaraInlineDynamicInputFormatTokenUsage::Decorator || InlineFormatToken.Usage == ENiagaraInlineDynamicInputFormatTokenUsage::LineBreak)
+		{
+			OutMatchedFormatInputs.Add(TPair<FNiagaraInlineDynamicInputFormatToken, FInputDisplayEntries>(InlineFormatToken, FInputDisplayEntries()));
+		}
+	}
+
+	if (MatchedInputNames.Num() != ChildInputDisplayEntries.Num())
+	{
+		OutMatchedFormatInputs.Empty();
+		return false;
+	}
+	else
+	{
+		return true;
+	}
+}
+
+void SNiagaraStackInlineDynamicInput::ConstructChildren()
+{
+	WidthSynchronizers.Empty();
+	if (RootFunctionInput->GetInlineDisplayMode() == ENiagaraStackEntryInlineDisplayMode::Expression)
+	{
+		TSharedPtr<SWrapBox> ExpressionWrapBox;
+		TSharedPtr<SVerticalBox> DataInterfaceVerticalBox;
+
+		ChildSlot
+		[
+			SNew(SVerticalBox)
+			+ SVerticalBox::Slot()
+			.AutoHeight()
+			[
+				SAssignNew(ExpressionWrapBox, SWrapBox)
+				.UseAllottedSize(true)
+			]
+			+ SVerticalBox::Slot()
+			.AutoHeight()
+			[
+				SAssignNew(DataInterfaceVerticalBox, SVerticalBox)
+			]
+		];
+
+		RootInputDisplayEntries = CollectInputDisplayEntriesRecursive(RootFunctionInput);
+		ConstructWrappedExpressionInputWidgets(ExpressionWrapBox.ToSharedRef(), DataInterfaceVerticalBox.ToSharedRef(), RootInputDisplayEntries, FGuid::NewGuid(), true);
+	}
+	else if (RootFunctionInput->GetInlineDisplayMode() == ENiagaraStackEntryInlineDisplayMode::GraphHorizontal || 
+		RootFunctionInput->GetInlineDisplayMode() == ENiagaraStackEntryInlineDisplayMode::GraphVertical ||
+		RootFunctionInput->GetInlineDisplayMode() == ENiagaraStackEntryInlineDisplayMode::GraphHybrid)
+	{
+		TSharedPtr<SGridPanel> GraphPanel;
+		TSharedRef<SVerticalBox> DataInterfaceVerticalBox = SNew(SVerticalBox);
+		RootInputDisplayEntries = CollectInputDisplayEntriesRecursive(RootFunctionInput);
+		TArray<FInputDisplayEntries> InputDisplayEntries;
+		InputDisplayEntries.Add(RootInputDisplayEntries);
+		ChildSlot
+		[
+			SNew(SVerticalBox)
+			+ SVerticalBox::Slot()
+			.Padding(0, 5, 0, 5)
+			[
+				RootFunctionInput->GetInlineDisplayMode() == ENiagaraStackEntryInlineDisplayMode::GraphHorizontal
+					? ConstructHorizontalGraphInputWidgets(InputDisplayEntries, nullptr, DataInterfaceVerticalBox, TSharedPtr<FWidthSynchronize>())
+					: RootFunctionInput->GetInlineDisplayMode() == ENiagaraStackEntryInlineDisplayMode::GraphVertical
+						? ConstructVerticalGraphInputWidgets(InputDisplayEntries, nullptr, DataInterfaceVerticalBox)
+						: ConstructHybridGraphInputWidgets(InputDisplayEntries, nullptr, DataInterfaceVerticalBox)
+			]
+			+ SVerticalBox::Slot()
+			.AutoHeight()
+			[
+				DataInterfaceVerticalBox
+			]
+		];
+	}
+	else
+	{
+		ChildSlot[SNullWidget::NullWidget];
+	}
+}
+
+TSharedRef<SWidget> ConstructFunctionNameWithDynamicInputIcon(FText FunctionDisplayName, FNiagaraTypeDefinition InputType)
+{
+	return SNew(SHorizontalBox)
+	+ SHorizontalBox::Slot()
+	.AutoWidth()
+	.Padding(0, 0, 5, 0)
+	[
+		SNew(SBorder)
+		.BorderImage(FNiagaraEditorStyle::Get().GetBrush("NiagaraEditor.Module.InputTypeBorder"))
+		.BorderBackgroundColor(FLinearColor::White)
+		.Padding(FMargin(2, 2, 2, 2))
+		[
+			SNew(SHorizontalBox)
+			+ SHorizontalBox::Slot()
+			.VAlign(VAlign_Center)
+			.HAlign(HAlign_Center)
+			.AutoWidth()
+			[
+				SNew(SImage)
+				.ToolTipText(FNiagaraStackEditorWidgetsUtilities::GetIconToolTipForInputMode(UNiagaraStackFunctionInput::EValueMode::Dynamic))
+				.ColorAndOpacity(UEdGraphSchema_Niagara::GetTypeColor(InputType))
+				.Image(FNiagaraEditorStyle::Get().GetBrush("NiagaraEditor.Module.TypeIconPill"))
+			]
+			+ SHorizontalBox::Slot()
+			.VAlign(VAlign_Center)
+			.HAlign(HAlign_Center)
+			.AutoWidth()
+			[
+				SNew(SImage)
+				.ToolTipText(FNiagaraStackEditorWidgetsUtilities::GetIconToolTipForInputMode(UNiagaraStackFunctionInput::EValueMode::Dynamic))
+				.ColorAndOpacity(FLinearColor::White)
+				.Image(FNiagaraEditorStyle::Get().GetBrush("NiagaraEditor.Module.DynamicInput"))
+			]
+		]
+	]
+	+ SHorizontalBox::Slot()
+	[
+		SNew(STextBlock)
+		.Text(FunctionDisplayName)
+	];
+}
+
+FText GetDynamicInputFunctionDisplayName(const UNiagaraNodeFunctionCall* FunctionCallNode)
+{
+	return FunctionCallNode != nullptr && FunctionCallNode->FunctionScript != nullptr
+		? FText::FromString(FName::NameToDisplayString(FunctionCallNode->FunctionScript->GetName(), false))
+		: LOCTEXT("UnknownFunctionDisplayName", "Unknown Function");
+}
+
+FText GetDynamicInputToolTip(const UNiagaraStackFunctionInput* FunctionInput)
+{
+	return FText::Format(LOCTEXT("DynamicInputNameAndToolTipFormat", "{0}\n\n{1}"), GetDynamicInputFunctionDisplayName(FunctionInput->GetDynamicInputNode()), FunctionInput->GetValueToolTip());
+}
+
+void SNiagaraStackInlineDynamicInput::ConstructWrappedExpressionInputWidgets(TSharedRef<SWrapBox> ExpressionWrapBox, TSharedRef<SVerticalBox> DataInterfaceVerticalBox, const FInputDisplayEntries& InputDisplayEntries, FGuid InputId, bool bIsRootInput)
+{
+	UNiagaraStackFunctionInput* FunctionInput = InputDisplayEntries.InputEntry;
+	if (FunctionInput->GetValueMode() == UNiagaraStackFunctionInput::EValueMode::Dynamic)
+	{
+		ConstructWrappedExpressionDynamicInputWidgets(ExpressionWrapBox, DataInterfaceVerticalBox, InputDisplayEntries, InputId, bIsRootInput);
+	}
+	else
+	{
+		TSharedPtr<SBorder> InputValueBorder;
+		ExpressionWrapBox->AddSlot()
+			.VAlign(VAlign_Center)
+			[
+				SAssignNew(InputValueBorder, SNiagaraStackInlineDynamicInputExpressionTokenBorder, FunctionInput, InputId)
+				.HoveredInputId(this, &SNiagaraStackInlineDynamicInput::GetHoveredInputId)
+				.OnHoveredChanged(this, &SNiagaraStackInlineDynamicInput::OnFormatTokenOutlineHoveredChanged)
+				.ToolTipText_UObject(FunctionInput, &UNiagaraStackFunctionInput::GetDisplayName)
+				.Padding(FMargin(4, 1, 0, 1))
+			];
+
+		TSharedRef<SWidget> InputValueWidget = SNew(SNiagaraStackFunctionInputValue, FunctionInput)
+			.LayoutMode(SNiagaraStackFunctionInputValue::ELayoutMode::CompactInline)
+			.CompactActionMenuButtonVisibility(this, &SNiagaraStackInlineDynamicInput::GetCompactActionMenuButtonVisibility, InputId);
+
+		if (FunctionInput->GetHasEditCondition())
+		{
+			InputValueBorder->SetContent(SNew(SHorizontalBox)
+				+ SHorizontalBox::Slot()
+				.AutoWidth()
+				[
+					SNew(SCheckBox)
+					.Visibility(this, &SNiagaraStackInlineDynamicInput::GetEditConditionCheckBoxVisibility, FunctionInput)
+					.IsChecked(this, &SNiagaraStackInlineDynamicInput::GetEditConditionCheckState, FunctionInput)
+					.OnCheckStateChanged(this, &SNiagaraStackInlineDynamicInput::OnEditConditionCheckStateChanged, FunctionInput)
+				]
+				+ SHorizontalBox::Slot()
+				[
+					InputValueWidget
+				]);
+		}
+		else
+		{
+			InputValueBorder->SetContent(InputValueWidget);
+		}
+
+		if (FunctionInput->GetValueMode() == UNiagaraStackFunctionInput::EValueMode::Data)
+		{
+			UNiagaraStackObject* ChildDataObject = FunctionInput->GetChildDataObject();
+			if (ChildDataObject != nullptr)
+			{
+				ConstructDataInterfaceWidgets(DataInterfaceVerticalBox, InputDisplayEntries, InputId);
+			}
+		}
+	}
+}
+
+TSharedRef<SWidget> ConstructDecoratorText(const FString& DisplayString)
+{
+	return SNew(STextBlock)
+		.Text(FText::FromString(DisplayString))
+		.TextStyle(FNiagaraEditorWidgetsStyle::Get(), "NiagaraEditor.Stack.InlineDynamicInputFormatText")
+		.ColorAndOpacity(FNiagaraEditorWidgetsStyle::Get().GetColor(FNiagaraStackEditorWidgetsUtilities::GetIconColorNameForInputMode(UNiagaraStackFunctionInput::EValueMode::Dynamic)));
+}
+
+TSharedRef<SWidget> ConstructDynamicInputLabel(const UNiagaraStackFunctionInput* FunctionInput)
+{
+	return SNew(SHorizontalBox)
+	.ToolTipText(FunctionInput->GetValueToolTip())
+	+ SHorizontalBox::Slot()
+	.VAlign(VAlign_Center)
+	.AutoWidth()
+	.Padding(0, 0, 3, 0)
+	[
+		SNew(SBorder)
+		.BorderImage(FNiagaraEditorStyle::Get().GetBrush("NiagaraEditor.Module.InputTypeBorder"))
+		.BorderBackgroundColor(FLinearColor::White)
+		.Padding(FMargin(2, 2, 2, 2))
+		[
+			SNew(SHorizontalBox)
+			+ SHorizontalBox::Slot()
+			.VAlign(VAlign_Center)
+			.HAlign(HAlign_Center)
+			.AutoWidth()
+			[
+				SNew(SImage)
+				.ToolTipText(FNiagaraStackEditorWidgetsUtilities::GetIconToolTipForInputMode(UNiagaraStackFunctionInput::EValueMode::Dynamic))
+				.ColorAndOpacity(UEdGraphSchema_Niagara::GetTypeColor(FunctionInput->GetInputType()))
+				.Image(FNiagaraEditorStyle::Get().GetBrush("NiagaraEditor.Module.TypeIconPill"))
+			]
+			+ SHorizontalBox::Slot()
+			.VAlign(VAlign_Center)
+			.HAlign(HAlign_Center)
+			.AutoWidth()
+			[
+				SNew(SImage)
+				.ToolTipText(FNiagaraStackEditorWidgetsUtilities::GetIconToolTipForInputMode(UNiagaraStackFunctionInput::EValueMode::Dynamic))
+				.ColorAndOpacity(FLinearColor::White)
+				.Image(FNiagaraEditorStyle::Get().GetBrush("NiagaraEditor.Module.DynamicInput"))
+			]
+		]
+	]
+	+ SHorizontalBox::Slot()
+	.VAlign(VAlign_Center)
+	[
+		SNew(STextBlock)
+		.Text(GetDynamicInputFunctionDisplayName(FunctionInput->GetDynamicInputNode()))
+	];
+}
+
+void SNiagaraStackInlineDynamicInput::ConstructWrappedExpressionDynamicInputWidgets(TSharedRef<SWrapBox> ExpressionWrapBox, TSharedRef<SVerticalBox> DataInterfaceVerticalBox, const FInputDisplayEntries& InputDisplayEntries, FGuid InputId, bool bIsRootInput)
+{
+	UNiagaraStackFunctionInput* FunctionInput = InputDisplayEntries.InputEntry;
+
+	TArray<TPair<FNiagaraInlineDynamicInputFormatToken, FInputDisplayEntries>> OutMatchedFormatInputs;
+	if (FunctionInput->GetDynamicInputNode() != nullptr && TryMatchFormatInputs(InputDisplayEntries.ChildInputDisplayEntries, FunctionInput->GetDynamicInputNode()->GetScriptData()->InlineExpressionFormat, OutMatchedFormatInputs))
+	{
+		// If this isn't the root input and the matched format doesn't start and end with decorators, we add parentheses before and after for readability.
+		bool bAddParentheses = 
+			bIsRootInput == false &&
+			OutMatchedFormatInputs.Num() > 2 &&
+			(OutMatchedFormatInputs[0].Key.Usage != ENiagaraInlineDynamicInputFormatTokenUsage::Decorator || OutMatchedFormatInputs.Last(0).Key.Usage != ENiagaraInlineDynamicInputFormatTokenUsage::Decorator);
+
+		if (bAddParentheses)
+		{
+			ExpressionWrapBox->AddSlot()
+			[
+				SNew(SNiagaraStackInlineDynamicInputExpressionTokenBorder, FunctionInput, InputId)
+				.HoveredInputId(this, &SNiagaraStackInlineDynamicInput::GetHoveredInputId)
+				.OnHoveredChanged(this, &SNiagaraStackInlineDynamicInput::OnFormatTokenOutlineHoveredChanged)
+				.ToolTipText(GetDynamicInputToolTip(FunctionInput))
+				.Padding(FMargin(3, 0))
+				[
+					ConstructDecoratorText("(")
+				]
+			];
+		}
+
+		for (const TPair<FNiagaraInlineDynamicInputFormatToken, FInputDisplayEntries> FormatTokenInputPair : OutMatchedFormatInputs)
+		{
+			if (FormatTokenInputPair.Key.Usage == ENiagaraInlineDynamicInputFormatTokenUsage::Input)
+			{
+				ConstructWrappedExpressionInputWidgets(ExpressionWrapBox, DataInterfaceVerticalBox, FormatTokenInputPair.Value, FGuid::NewGuid(), false);
+			}
+			else if (FormatTokenInputPair.Key.Usage == ENiagaraInlineDynamicInputFormatTokenUsage::Decorator)
+			{
+				ExpressionWrapBox->AddSlot()
+				[
+					SNew(SNiagaraStackInlineDynamicInputExpressionTokenBorder, FunctionInput, InputId)
+					.HoveredInputId(this, &SNiagaraStackInlineDynamicInput::GetHoveredInputId)
+					.OnHoveredChanged(this, &SNiagaraStackInlineDynamicInput::OnFormatTokenOutlineHoveredChanged)
+					.ToolTipText(GetDynamicInputToolTip(FunctionInput))
+					.Padding(FMargin(3, 0, 6, 0))
+					[
+						ConstructDecoratorText(FormatTokenInputPair.Key.Value)
+					]
+				];
+			}
+		}
+
+		if (bAddParentheses)
+		{
+			ExpressionWrapBox->AddSlot()
+			[
+				SNew(SNiagaraStackInlineDynamicInputExpressionTokenBorder, FunctionInput, InputId)
+				.HoveredInputId(this, &SNiagaraStackInlineDynamicInput::GetHoveredInputId)
+				.OnHoveredChanged(this, &SNiagaraStackInlineDynamicInput::OnFormatTokenOutlineHoveredChanged)
+				.ToolTipText(GetDynamicInputToolTip(FunctionInput))
+				.Padding(FMargin(3, 0))
+				[
+					ConstructDecoratorText(")")
+				]
+			];
+		}
+	}
+	else
+	{
+		ExpressionWrapBox->AddSlot()
+		[
+			SNew(SNiagaraStackInlineDynamicInputExpressionTokenBorder, FunctionInput, InputId)
+			.HoveredInputId(this, &SNiagaraStackInlineDynamicInput::GetHoveredInputId)
+			.OnHoveredChanged(this, &SNiagaraStackInlineDynamicInput::OnFormatTokenOutlineHoveredChanged)
+			.Padding(FMargin(3, 0))
+			[
+				ConstructDynamicInputLabel(FunctionInput)
+			]
+		];
+
+		ExpressionWrapBox->AddSlot()
+		[
+			SNew(SNiagaraStackInlineDynamicInputExpressionTokenBorder, FunctionInput, InputId)
+			.HoveredInputId(this, &SNiagaraStackInlineDynamicInput::GetHoveredInputId)
+			.OnHoveredChanged(this, &SNiagaraStackInlineDynamicInput::OnFormatTokenOutlineHoveredChanged)
+			.ToolTipText(GetDynamicInputToolTip(FunctionInput))
+			.Padding(FMargin(3, 0))
+			[
+				ConstructDecoratorText("(")
+			]
+		];
+
+		bool bFirst = true;
+		for (const FInputDisplayEntries& ChildInputDisplayEntries : InputDisplayEntries.ChildInputDisplayEntries)
+		{
+			if (bFirst)
+			{
+				bFirst = false;
+			}
+			else
+			{
+				ExpressionWrapBox->AddSlot()
+				[
+					SNew(SNiagaraStackInlineDynamicInputExpressionTokenBorder, FunctionInput, InputId)
+					.HoveredInputId(this, &SNiagaraStackInlineDynamicInput::GetHoveredInputId)
+					.OnHoveredChanged(this, &SNiagaraStackInlineDynamicInput::OnFormatTokenOutlineHoveredChanged)
+					.ToolTipText(GetDynamicInputToolTip(FunctionInput))
+					.Padding(FMargin(3, 0, 6, 0))
+					[
+						ConstructDecoratorText(",")
+					]
+				];
+			}
+			ConstructWrappedExpressionInputWidgets(ExpressionWrapBox, DataInterfaceVerticalBox, ChildInputDisplayEntries, FGuid::NewGuid(), false);
+		}
+
+		ExpressionWrapBox->AddSlot()
+		[
+			SNew(SNiagaraStackInlineDynamicInputExpressionTokenBorder, FunctionInput, InputId)
+			.HoveredInputId(this, &SNiagaraStackInlineDynamicInput::GetHoveredInputId)
+			.OnHoveredChanged(this, &SNiagaraStackInlineDynamicInput::OnFormatTokenOutlineHoveredChanged)
+			.ToolTipText(GetDynamicInputToolTip(FunctionInput))
+			.Padding(FMargin(3, 0))
+			[
+				ConstructDecoratorText(")")
+			]
+		];
+	}
+
+	ExpressionWrapBox->AddSlot()
+	[
+		SNew(SNiagaraStackInlineDynamicInputExpressionTokenBorder, FunctionInput, InputId)
+		.HoveredInputId(this, &SNiagaraStackInlineDynamicInput::GetHoveredInputId)
+		.OnHoveredChanged(this, &SNiagaraStackInlineDynamicInput::OnFormatTokenOutlineHoveredChanged)
+		.ToolTipText(GetDynamicInputToolTip(FunctionInput))
+		[
+			SNew(SNiagaraStackFunctionInputValue, FunctionInput)
+			.LayoutMode(SNiagaraStackFunctionInputValue::ELayoutMode::EditDropDownOnly)
+			.CompactActionMenuButtonVisibility(this, &SNiagaraStackInlineDynamicInput::GetCompactActionMenuButtonVisibility, InputId)
+		]
+	];
+}
+
+const FSlateBrush* GetHorizontalGraphNodeBorderBrush(int32 RowIndex, int32 TotalRows)
+{
+	if (TotalRows == 1)
+	{
+		// Single
+		return FNiagaraEditorWidgetsStyle::Get().GetBrush("NiagaraEditor.Stack.InlineDynamicInput.GraphNode.Single");
+	}
+	else
+	{
+		if (RowIndex == 0)
+		{
+			// Top
+			return FNiagaraEditorWidgetsStyle::Get().GetBrush("NiagaraEditor.Stack.InlineDynamicInput.GraphNode.Row.Top");
+		}
+		else if (RowIndex < TotalRows - 1)
+		{
+			// Middle
+			return FNiagaraEditorWidgetsStyle::Get().GetBrush("NiagaraEditor.Stack.InlineDynamicInput.GraphNode.Row.Middle");
+		}
+		else
+		{
+			// Bottom
+			return FNiagaraEditorWidgetsStyle::Get().GetBrush("NiagaraEditor.Stack.InlineDynamicInput.GraphNode.Row.Bottom");
+		}
+	}
+}
+
+const FSlateBrush* GetVerticalGraphNodeBorderBrush(int32 ColumnIndex, int32 TotalColumns, bool bHasHeaderRow)
+{
+	if (TotalColumns == 1)
+	{
+		// Single
+		return bHasHeaderRow 
+			? FNiagaraEditorWidgetsStyle::Get().GetBrush("NiagaraEditor.Stack.InlineDynamicInput.GraphNode.Row.Bottom")
+			: FNiagaraEditorWidgetsStyle::Get().GetBrush("NiagaraEditor.Stack.InlineDynamicInput.GraphNode.Single");
+	}
+	else
+	{
+		if (ColumnIndex == 0)
+		{
+			// Left
+			return bHasHeaderRow
+				? FNiagaraEditorWidgetsStyle::Get().GetBrush("NiagaraEditor.Stack.InlineDynamicInput.GraphNode.Corner.BottomLeft")
+				: FNiagaraEditorWidgetsStyle::Get().GetBrush("NiagaraEditor.Stack.InlineDynamicInput.GraphNode.Column.Left");
+		}
+		else if (ColumnIndex < TotalColumns - 1)
+		{
+			// Middle
+			return bHasHeaderRow
+				? FNiagaraEditorWidgetsStyle::Get().GetBrush("NiagaraEditor.Stack.InlineDynamicInput.GraphNode.Edge.Bottom")
+				: FNiagaraEditorWidgetsStyle::Get().GetBrush("NiagaraEditor.Stack.InlineDynamicInput.GraphNode.Column.Middle");
+		}
+		else
+		{
+			// Right
+			return bHasHeaderRow
+				? FNiagaraEditorWidgetsStyle::Get().GetBrush("NiagaraEditor.Stack.InlineDynamicInput.GraphNode.Corner.BottomRight")
+				: FNiagaraEditorWidgetsStyle::Get().GetBrush("NiagaraEditor.Stack.InlineDynamicInput.GraphNode.Column.Right");
+		}
+	}
+}
+
+const FSlateBrush* GetHybridGraphNodeBorderBrush(int32 NodeColumnIndex, int32 TotalNodeColumns, int32 ColumnColumnIndex, int32 TotalColumnColumns)
+{
+	if (TotalNodeColumns == 1 && TotalColumnColumns == 1)
+	{
+		// Single
+		return FNiagaraEditorWidgetsStyle::Get().GetBrush("NiagaraEditor.Stack.InlineDynamicInput.GraphNode.Single");
+	}
+	else
+	{
+		if (NodeColumnIndex == 0 && ColumnColumnIndex == 0)
+		{
+			// Left
+			return FNiagaraEditorWidgetsStyle::Get().GetBrush("NiagaraEditor.Stack.InlineDynamicInput.GraphNode.Column.Left");
+		}
+		else if (NodeColumnIndex < TotalNodeColumns - 1 || ColumnColumnIndex < TotalColumnColumns - 1)
+		{
+			// Middle
+			return FNiagaraEditorWidgetsStyle::Get().GetBrush("NiagaraEditor.Stack.InlineDynamicInput.GraphNode.Column.Middle");
+		}
+		else
+		{
+			// Right
+			return FNiagaraEditorWidgetsStyle::Get().GetBrush("NiagaraEditor.Stack.InlineDynamicInput.GraphNode.Column.Right");
+		}
+	}
+}
+
+FMargin GetHorizontalGraphNodeBorderPadding(int32 RowIndex, int32 TotalRows, bool bHasChildren)
+{
+	if (TotalRows == 1 || bHasChildren)
+	{
+		// Single
+		return FMargin(8, 4);
+	}
+	else
+	{
+		if (RowIndex == 0)
+		{
+			// Top
+			return FMargin(8, 4, 8, 2);
+		}
+		else if (RowIndex < TotalRows - 1)
+		{
+			// Middle
+			return FMargin(8, 2);
+		}
+		else
+		{
+			// Bottom
+			return FMargin(8, 2, 8, 4);
+		}
+	}
+}
+
+FMargin GetVerticalGraphNodeBorderPadding(int32 ColumnIndex, int32 TotalColumns)
+{
+	if (TotalColumns == 1)
+	{
+		// Single
+		return FMargin(6, 4);
+	}
+	else
+	{
+		if (ColumnIndex == 0)
+		{
+			// Left
+			return FMargin(6, 4, 2, 4);
+		}
+		else if (ColumnIndex < TotalColumns - 1)
+		{
+			// Middle
+			return FMargin(2, 4);
+		}
+		else
+		{
+			// Right
+			return FMargin(2, 4, 6, 4);
+		}
+	}
+}
+
+FMargin GetHybridGraphNodeBorderPadding(int32 NodeColumnIndex, int32 TotalNodeColumns, int32 ColumnColumnIndex, int32 TotalColumnColumns)
+{
+	if (TotalNodeColumns == 1 && TotalColumnColumns == 1)
+	{
+		// Single
+		return FMargin(6, 4);
+	}
+	else
+	{
+		if (NodeColumnIndex == 0 && ColumnColumnIndex == 0)
+		{
+			// Left
+			return FMargin(6, 4, 2, 4);
+		}
+		else if (NodeColumnIndex < TotalNodeColumns - 1 || ColumnColumnIndex < TotalColumnColumns - 1)
+		{
+			// Middle
+			return FMargin(2, 4);
+		}
+		else
+		{
+			// Right
+			return FMargin(2, 4, 6, 4);
+		}
+	}
+}
+
+EVerticalAlignment GetHorizontalGraphNodeGridVerticalAlignment(int32 RowIndex, int32 TotalRows)
+{
+	if (TotalRows == 1)
+	{
+		return VAlign_Top;
+	}
+
+	if (RowIndex == TotalRows - 1)
+	{
+		return VAlign_Top;
+	}
+
+	return VAlign_Fill;
+}
+
+EHorizontalAlignment GetVerticalGraphNodeGridHorizontalAlignment(int32 ColumnIndex, int32 TotalColumns)
+{
+	if (TotalColumns == 1)
+	{
+		return HAlign_Left;
+	}
+
+	if (ColumnIndex == TotalColumns - 1)
+	{
+		return HAlign_Left;
+	}
+
+	return HAlign_Fill;
+}
+
+EHorizontalAlignment GetHybridGraphNodeGridHorizontalAlignment(int32 NodeColumnIndex, int32 TotalNodeColumns, int32 ColumnColumnIndex, int32 TotalColumnColumns)
+{
+	if (NodeColumnIndex < TotalNodeColumns - 1 && ColumnColumnIndex == TotalColumnColumns - 1)
+	{
+		return HAlign_Fill;
+	}
+	return HAlign_Left;
+}
+
+FMargin GetHorizontalGraphNodeGridPadding(int32 RowIndex, int32 TotalRows)
+{
+	if (TotalRows == 1)
+	{
+		// Single
+		return FMargin(0, 2, 0, 2);
+	}
+	else
+	{
+		if (RowIndex == 0)
+		{
+			// Left
+			return FMargin(0, 2, 0, 0);
+		}
+		else if (RowIndex < TotalRows - 1)
+		{
+			// Middle
+			return FMargin(0);
+		}
+		else
+		{
+			// Right
+			return FMargin(0, 0, 0, 2);
+		}
+	}
+}
+
+FMargin GetVerticalGraphNodeGridPadding(int32 ColumnIndex, int32 TotalColumns)
+{
+	if (TotalColumns == 1)
+	{
+		// Single
+		return FMargin(2, 0, 2, 0);
+	}
+	else
+	{
+		if (ColumnIndex == 0)
+		{
+			// Left
+			return FMargin(2, 0, 0, 0);
+		}
+		else if (ColumnIndex < TotalColumns - 1)
+		{
+			// Middle
+			return FMargin(0);
+		}
+		else
+		{
+			// Right
+			return FMargin(0, 0, 2, 0);
+		}
+	}
+}
+
+void SNiagaraStackInlineDynamicInput::FWidthSynchronize::AddWidget(TSharedRef<SWidget> InWidget)
+{
+	WeakWidgets.Add(TWeakPtr<SWidget>(InWidget));
+}
+
+FOptionalSize SNiagaraStackInlineDynamicInput::FWidthSynchronize::GetMaxDesiredWidth() const
+{
+	float MaxDesiredWidth = TNumericLimits<float>::Lowest();
+	for (const TWeakPtr<SWidget>& WeakWidget : WeakWidgets)
+	{
+		TSharedPtr<SWidget> Widget = WeakWidget.Pin();
+		if (Widget.IsValid())
+		{
+			MaxDesiredWidth = FMath::Max(MaxDesiredWidth, Widget->GetDesiredSize().X);
+		}
+	}
+	return MaxDesiredWidth > 0 
+		? FOptionalSize(MaxDesiredWidth)
+		: FOptionalSize();
+}
+
+TSharedRef<SWidget> SNiagaraStackInlineDynamicInput::ConstructGraphInputWidgetForFormatTokenInputPair(
+	const TPair<FNiagaraInlineDynamicInputFormatToken, FInputDisplayEntries>& FormatTokenInputPair,
+	const UNiagaraStackFunctionInput* OwningInput,
+	FGuid InputId)
+{
+	if (FormatTokenInputPair.Key.Usage == ENiagaraInlineDynamicInputFormatTokenUsage::Input)
+	{
+		UNiagaraStackFunctionInput* FunctionInput = FormatTokenInputPair.Value.InputEntry;
+		TSharedPtr<SWidget> InputWidget;
+		if (FunctionInput->GetValueMode() == UNiagaraStackFunctionInput::EValueMode::Dynamic)
+		{
+			// Inputs which reference another dynamic input just show a label for the input name and then a drop-down arrow.
+			return 
+				SNew(SHorizontalBox)
+				+ SHorizontalBox::Slot()
+				.VAlign(VAlign_Center)
+				.AutoWidth()
+				[
+					SNew(STextBlock)
+					.Text_UObject(FunctionInput, &UNiagaraStackEntry::GetDisplayName)
+				]
+				//+ SHorizontalBox::Slot()
+				//[
+				//	SNew
+				//]
+				+ SHorizontalBox::Slot()
+				.VAlign(VAlign_Center)
+				.AutoWidth()
+				[
+					SNew(SNiagaraStackFunctionInputValue, FunctionInput)
+					.LayoutMode(SNiagaraStackFunctionInputValue::ELayoutMode::EditDropDownOnly)
+				];
+		}
+		else
+		{
+			// Other inputs display a label and display the standard value control.
+			return
+				SNew(SHorizontalBox)
+				+ SHorizontalBox::Slot()
+				.VAlign(VAlign_Center)
+				.AutoWidth()
+				.Padding(0, 0, 5, 0)
+				[
+					SNew(STextBlock)
+					.Text_UObject(FunctionInput, &UNiagaraStackEntry::GetDisplayName)
+				]
+				+ SHorizontalBox::Slot()
+				.VAlign(VAlign_Center)
+				.AutoWidth()
+				.Padding(0, 2, 0, 2)
+				[
+					SNew(SNiagaraStackInlineDynamicInputExpressionTokenBorder, FunctionInput, InputId)
+					.HoveredInputId(this, &SNiagaraStackInlineDynamicInput::GetHoveredInputId)
+					.OnHoveredChanged(this, &SNiagaraStackInlineDynamicInput::OnFormatTokenOutlineHoveredChanged)
+					.ToolTipText_UObject(FunctionInput, &UNiagaraStackFunctionInput::GetDisplayName)
+					[
+						SNew(SNiagaraStackFunctionInputValue, FunctionInput)
+						.LayoutMode(SNiagaraStackFunctionInputValue::ELayoutMode::CompactInline)
+					]
+				];
+		}
+	}
+	else if (FormatTokenInputPair.Key.Usage == ENiagaraInlineDynamicInputFormatTokenUsage::Decorator)
+	{
+		return 
+			SNew(SBox)
+			.ToolTipText(GetDynamicInputToolTip(OwningInput))
+			.Padding(2)
+			[
+				ConstructDecoratorText(FormatTokenInputPair.Key.Value)
+			];
+	}
+	return SNullWidget::NullWidget;
+}
+
+void SNiagaraStackInlineDynamicInput::GenerateFormattedGroups(
+	const TArray<FInputDisplayEntries>& InputDisplayEntries,
+	const TArray<FNiagaraInlineDynamicInputFormatToken>& InlineFormat,
+	bool bIgnoreLinebreaks,
+	bool bTreatInputsWithoutChildrenLikeDecorators,
+	bool& bOutFormatMatched,
+	TArray<TArray<TPair<FNiagaraInlineDynamicInputFormatToken, FInputDisplayEntries>>>& OutFormattedGroups)
+{
+	// Try to match the inputs to the format.  If they can not be matched then use a generic format.
+	TArray<TPair<FNiagaraInlineDynamicInputFormatToken, FInputDisplayEntries>> FormattedInputs;
+	bOutFormatMatched = TryMatchFormatInputs(InputDisplayEntries, InlineFormat, FormattedInputs);
+	if(bOutFormatMatched == false)
+	{
+		for (const FInputDisplayEntries& InputDisplayEntry : InputDisplayEntries)
+		{
+			FormattedInputs.Add(TPair<FNiagaraInlineDynamicInputFormatToken, FInputDisplayEntries>(
+				FNiagaraInlineDynamicInputFormatToken(ENiagaraInlineDynamicInputFormatTokenUsage::Input, FString()),
+				InputDisplayEntry));
+		}
+	}
+
+	// Sort the formatters and inputs into groups for rows and columns
+	TArray<TPair<FNiagaraInlineDynamicInputFormatToken, FInputDisplayEntries>> CurrentFormattedGroup;
+	bool bCurrentGroupHasInput = false;
+	for (const TPair<FNiagaraInlineDynamicInputFormatToken, FInputDisplayEntries>& FormatTokenInputPair : FormattedInputs)
+	{
+		ENiagaraInlineDynamicInputFormatTokenUsage PairMode;
+		if (bTreatInputsWithoutChildrenLikeDecorators &&
+			FormatTokenInputPair.Key.Usage == ENiagaraInlineDynamicInputFormatTokenUsage::Input &&
+			FormatTokenInputPair.Value.ChildInputDisplayEntries.Num() == 0)
+		{
+			PairMode = ENiagaraInlineDynamicInputFormatTokenUsage::Decorator;
+		}
+		else
+		{
+			PairMode = FormatTokenInputPair.Key.Usage;
+		}
+
+		switch (PairMode)
+		{
+		case ENiagaraInlineDynamicInputFormatTokenUsage::Input:
+			if (bCurrentGroupHasInput)
+			{
+				OutFormattedGroups.Add(CurrentFormattedGroup);
+				CurrentFormattedGroup.Empty();
+			}
+			CurrentFormattedGroup.Add(FormatTokenInputPair);
+			bCurrentGroupHasInput = true;
+			break;
+		case ENiagaraInlineDynamicInputFormatTokenUsage::Decorator:
+			CurrentFormattedGroup.Add(FormatTokenInputPair);
+			break;
+		case ENiagaraInlineDynamicInputFormatTokenUsage::LineBreak:
+			if (bIgnoreLinebreaks == false && CurrentFormattedGroup.Num() > 0)
+			{
+				OutFormattedGroups.Add(CurrentFormattedGroup);
+				CurrentFormattedGroup.Empty();
+				bCurrentGroupHasInput = false;
+			}
+			break;
+		}
+	}
+	if (CurrentFormattedGroup.Num() > 0)
+	{
+		OutFormattedGroups.Add(CurrentFormattedGroup);
+		CurrentFormattedGroup.Empty();
+	}
+}
+
+TSharedRef<SWidget> SNiagaraStackInlineDynamicInput::ConstructHorizontalGraphInputWidgets(
+	const TArray<FInputDisplayEntries>& InputDisplayEntries,
+	const UNiagaraStackFunctionInput* OwningInput,
+	TSharedRef<SVerticalBox> DataInterfaceVerticalBox,
+	TSharedPtr<FWidthSynchronize> InputWidthSynchronizer)
+{
+	const UNiagaraNodeFunctionCall* DynamicInputFunctionCall = 
+		OwningInput != nullptr && OwningInput->GetValueMode() == UNiagaraStackFunctionInput::EValueMode::Dynamic
+			? OwningInput->GetDynamicInputNode() : nullptr;
+	bool bFormatMatched;
+	TArray<TArray<TPair<FNiagaraInlineDynamicInputFormatToken, FInputDisplayEntries>>> FormattedRows;
+	TArray<FNiagaraInlineDynamicInputFormatToken> InlineFormat = DynamicInputFunctionCall != nullptr 
+		? DynamicInputFunctionCall->GetScriptData()->InlineGraphFormat 
+		: TArray<FNiagaraInlineDynamicInputFormatToken>();
+
+	bool bIgnoreLinebreaks = false;
+	bool bTreatInputsWithoutChildrenLikeDecorators = false;
+	GenerateFormattedGroups(InputDisplayEntries, InlineFormat, bIgnoreLinebreaks, bTreatInputsWithoutChildrenLikeDecorators, bFormatMatched, FormattedRows);
+
+	// Generate the rows.
+	TSharedRef<SGridPanel> InputGrid = SNew(SGridPanel);
+	int32 GraphNodeRowIndex = 0;
+	int32 GraphNodeHeaderRowCount = bFormatMatched == false && DynamicInputFunctionCall != nullptr ? 1 : 0;
+	int32 TotalGraphNodeRowCount = FormattedRows.Num() + GraphNodeHeaderRowCount;
+	int32 ChildInputsRowIndex = 0;
+
+	auto AddGraphNodeRow = [InputGrid, TotalGraphNodeRowCount, InputWidthSynchronizer](int32 GraphNodeRowIndex, bool bHasChildren, TSharedRef<SWidget> Widget)
+	{
+		TSharedPtr<SWidget> RowWidget;
+		if (InputWidthSynchronizer.IsValid())
+		{
+			InputWidthSynchronizer->AddWidget(Widget);
+			RowWidget = SNew(SBox)
+				.WidthOverride(InputWidthSynchronizer.ToSharedRef(), &FWidthSynchronize::GetMaxDesiredWidth)
+				.Content()
+				[
+					Widget
+				];
+		}
+		else
+		{
+			RowWidget = Widget;
+		}
+		InputGrid->AddSlot(2, GraphNodeRowIndex)
+		.VAlign(GetHorizontalGraphNodeGridVerticalAlignment(GraphNodeRowIndex, TotalGraphNodeRowCount))
+		.Padding(GetHorizontalGraphNodeGridPadding(GraphNodeRowIndex, TotalGraphNodeRowCount))
+		[
+			SNew(SBorder)
+			.BorderImage(GetHorizontalGraphNodeBorderBrush(GraphNodeRowIndex, TotalGraphNodeRowCount))
+			.BorderBackgroundColor(FStyleColors::DropdownOutline)
+			.HAlign(HAlign_Left)
+			.VAlign(VAlign_Top)
+			.Padding(GetHorizontalGraphNodeBorderPadding(GraphNodeRowIndex, TotalGraphNodeRowCount, bHasChildren))
+			[
+				RowWidget.ToSharedRef()
+			]
+		];
+	};
+
+	if (bFormatMatched == false && DynamicInputFunctionCall != nullptr)
+	{
+		// For unformatted inputs add a header with the dynamic input icon and the function name.
+		AddGraphNodeRow(GraphNodeRowIndex, false, ConstructDynamicInputLabel(OwningInput));
+		GraphNodeRowIndex++;
+	}
+
+	// Generate row widgets.
+	TSharedRef<FWidthSynchronize> ChildInputWidthSynchronizer = MakeShared<FWidthSynchronize>();
+	WidthSynchronizers.Add(ChildInputWidthSynchronizer);
+	for (const TArray<TPair<FNiagaraInlineDynamicInputFormatToken, FInputDisplayEntries>>& FormattedRow : FormattedRows)
+	{
+		FGuid InputId = FGuid::NewGuid();
+		TOptional<FInputDisplayEntries> RowDisplayEntries;
+		TSharedPtr<SWidget> RowWidget;
+		if(FormattedRow.Num() == 1)
+		{
+			RowWidget = ConstructGraphInputWidgetForFormatTokenInputPair(FormattedRow[0], OwningInput, InputId);
+			if(FormattedRow[0].Key.Usage == ENiagaraInlineDynamicInputFormatTokenUsage::Input)
+			{
+				RowDisplayEntries = FormattedRow[0].Value;
+			}
+		}
+		else
+		{
+			TSharedRef<SHorizontalBox> RowBox = SNew(SHorizontalBox);
+			RowWidget = RowBox;
+			for (const TPair<FNiagaraInlineDynamicInputFormatToken, FInputDisplayEntries>& FormatTokenInputPair : FormattedRow)
+			{
+				RowBox->AddSlot()
+				.AutoWidth()
+				[
+					ConstructGraphInputWidgetForFormatTokenInputPair(FormatTokenInputPair, OwningInput, InputId)
+				];
+				if (FormatTokenInputPair.Key.Usage == ENiagaraInlineDynamicInputFormatTokenUsage::Input)
+				{
+					RowDisplayEntries = FormatTokenInputPair.Value;
+				}
+			}
+		}
+
+		if (RowDisplayEntries.IsSet() && RowDisplayEntries->InputEntry->GetValueMode() == UNiagaraStackFunctionInput::EValueMode::Data)
+		{
+			UNiagaraStackObject* ChildDataObject = RowDisplayEntries->InputEntry->GetChildDataObject();
+			if (ChildDataObject != nullptr)
+			{
+				ConstructDataInterfaceWidgets(DataInterfaceVerticalBox, RowDisplayEntries.GetValue(), InputId);
+			}
+		}
+
+		bool bIsDynamicInput = RowDisplayEntries.IsSet() && RowDisplayEntries->InputEntry->GetValueMode() == UNiagaraStackFunctionInput::EValueMode::Dynamic;
+		AddGraphNodeRow(GraphNodeRowIndex, bIsDynamicInput, RowWidget.ToSharedRef());
+
+		// Add additional widgets for dynamic inputs.
+		if (RowDisplayEntries.IsSet() && bIsDynamicInput)
+		{
+			// Connection arrow.
+			FMargin BoxPadding = GetHorizontalGraphNodeBorderPadding(GraphNodeRowIndex, TotalGraphNodeRowCount, true);
+			BoxPadding.Left = 2;
+			BoxPadding.Right = 2;
+			InputGrid->AddSlot(1, GraphNodeRowIndex)
+			.VAlign(GetHorizontalGraphNodeGridVerticalAlignment(GraphNodeRowIndex, TotalGraphNodeRowCount))
+			.Padding(GetHorizontalGraphNodeGridPadding(GraphNodeRowIndex, TotalGraphNodeRowCount))
+			[
+				SNew(SBox)
+				.VAlign(VAlign_Top)
+				.Padding(BoxPadding)
+				[
+					SNew(STextBlock)
+					.Font(FAppStyle::Get().GetFontStyle("FontAwesome.14"))
+					.Text(FEditorFontGlyphs::Long_Arrow_Right)
+				]
+			];
+
+			// Child inputs.
+			auto FindNextNodeRowWithChildren = [FormattedRows](int32 StartRow) -> int32
+			{
+				for (int32 NodeRow = StartRow; NodeRow < FormattedRows.Num(); NodeRow++)
+				{
+					for (const TPair<FNiagaraInlineDynamicInputFormatToken, FInputDisplayEntries>& RowItem : FormattedRows[NodeRow])
+					{
+						if (RowItem.Key.Usage == ENiagaraInlineDynamicInputFormatTokenUsage::Input && RowItem.Value.ChildInputDisplayEntries.Num() > 0)
+						{
+							return NodeRow;
+						}
+					}
+				}
+				return INDEX_NONE;
+			};
+
+			int32 StartRow = ChildInputsRowIndex;
+			int32 NextRowWithChildren = FindNextNodeRowWithChildren(GraphNodeRowIndex + 1);
+			int32 EndRow = 
+				GraphNodeHeaderRowCount + 
+				(NextRowWithChildren == INDEX_NONE ? FormattedRows.Num() - 1 : NextRowWithChildren - 1);
+			int32 RowSpan = EndRow - StartRow + 1;
+			InputGrid->AddSlot(0, StartRow)
+			.RowSpan(RowSpan)
+			.HAlign(HAlign_Right)
+			[
+				ConstructHorizontalGraphInputWidgets(RowDisplayEntries->ChildInputDisplayEntries, RowDisplayEntries->InputEntry, DataInterfaceVerticalBox, ChildInputWidthSynchronizer)
+			];
+			ChildInputsRowIndex += RowSpan;
+		}
+		GraphNodeRowIndex++;
+	}
+
+	return InputGrid;
+}
+
+TSharedRef<SWidget> SNiagaraStackInlineDynamicInput::ConstructVerticalGraphInputWidgets(
+	const TArray<FInputDisplayEntries>& InputDisplayEntries,
+	const UNiagaraStackFunctionInput* OwningInput,
+	TSharedRef<SVerticalBox> DataInterfaceVerticalBox)
+{
+	const UNiagaraNodeFunctionCall* DynamicInputFunctionCall =
+		OwningInput != nullptr && OwningInput->GetValueMode() == UNiagaraStackFunctionInput::EValueMode::Dynamic
+			? OwningInput->GetDynamicInputNode() : nullptr;
+	bool bFormatMatched;
+	TArray<TArray<TPair<FNiagaraInlineDynamicInputFormatToken, FInputDisplayEntries>>> FormattedColumns;
+	TArray<FNiagaraInlineDynamicInputFormatToken> InlineFormat = DynamicInputFunctionCall != nullptr
+		? DynamicInputFunctionCall->GetScriptData()->InlineGraphFormat
+		: TArray<FNiagaraInlineDynamicInputFormatToken>();
+
+	bool bIgnoreLinebreaks = true;
+	bool bTreatInputsWithoutChildrenLikeDecorators = false;
+	GenerateFormattedGroups(InputDisplayEntries, InlineFormat, bIgnoreLinebreaks, bTreatInputsWithoutChildrenLikeDecorators, bFormatMatched, FormattedColumns);
+
+	// Generate the columns.
+	TSharedRef<SGridPanel> InputGrid = SNew(SGridPanel);
+	int32 GraphNodeColumnIndex = 0;
+	int32 ChildInputsColumnIndex = 0;
+	int32 TotalGraphNodeColumnCount = FormattedColumns.Num();
+	bool bHasHeaderRow;
+
+	if (bFormatMatched == false && DynamicInputFunctionCall != nullptr)
+	{
+		// For unformatted inputs add a header with the dynamic input icon and the function name.
+		bHasHeaderRow = true;
+		bool bHasInputs = FormattedColumns.Num() > 0;
+		InputGrid->AddSlot(0, 0)
+		.ColumnSpan(TotalGraphNodeColumnCount)
+		.Padding(GetVerticalGraphNodeGridPadding(0, 1))
+		[
+			SNew(SBorder)
+			.BorderImage(bHasInputs 
+				? FNiagaraEditorWidgetsStyle::Get().GetBrush("NiagaraEditor.Stack.InlineDynamicInput.GraphNode.Row.Top")
+				: FNiagaraEditorWidgetsStyle::Get().GetBrush("NiagaraEditor.Stack.InlineDynamicInput.GraphNode.Single"))
+			.BorderBackgroundColor(FStyleColors::DropdownOutline)
+			.HAlign(HAlign_Center)
+			.VAlign(VAlign_Center)
+			.Padding(bHasInputs ? FMargin(6, 4, 6, 0) : FMargin(6, 4))
+			[
+				ConstructDynamicInputLabel(OwningInput)
+			]
+		];
+	}
+	else
+	{
+		bHasHeaderRow = false;
+	}
+
+	auto AddGraphNodeColumn = [InputGrid, TotalGraphNodeColumnCount, bHasHeaderRow](int32 GraphNodeColumnIndex, bool bHasChildren, TSharedRef<SWidget> Widget)
+	{
+		InputGrid->AddSlot(GraphNodeColumnIndex, 1)
+		.Padding(GetVerticalGraphNodeGridPadding(GraphNodeColumnIndex, TotalGraphNodeColumnCount))
+		[
+			SNew(SBorder)
+			.BorderImage(GetVerticalGraphNodeBorderBrush(GraphNodeColumnIndex, TotalGraphNodeColumnCount, bHasHeaderRow))
+			.BorderBackgroundColor(FStyleColors::DropdownOutline)
+			.HAlign(HAlign_Center)
+			.VAlign(VAlign_Center)
+			.Padding(GetVerticalGraphNodeBorderPadding(GraphNodeColumnIndex, TotalGraphNodeColumnCount))
+			[
+				Widget
+			]
+		];
+	};
+
+	// Figure out which dynamic input row is last so that it can spawn any additional non dynamic input rows.
+	int32 LastDynamicInputColumnIndex = INDEX_NONE;
+	for (int32 FormattedColumnIndex = 0; FormattedColumnIndex < FormattedColumns.Num(); FormattedColumnIndex++)
+	{
+		TPair<FNiagaraInlineDynamicInputFormatToken, FInputDisplayEntries>* DynamicInputPair = FormattedColumns[FormattedColumnIndex].FindByPredicate([](const TPair<FNiagaraInlineDynamicInputFormatToken, FInputDisplayEntries>& Pair)
+			{ return Pair.Key.Usage == ENiagaraInlineDynamicInputFormatTokenUsage::Input && Pair.Value.InputEntry->GetValueMode() == UNiagaraStackFunctionInput::EValueMode::Dynamic; });
+		if (DynamicInputPair != nullptr)
+		{
+			LastDynamicInputColumnIndex = FormattedColumnIndex;
+		}
+	}
+
+	// Generate column widgets.
+	TSharedRef<FWidthSynchronize> ChildInputWidthSynchronizer = MakeShared<FWidthSynchronize>();
+	WidthSynchronizers.Add(ChildInputWidthSynchronizer);
+	for (const TArray<TPair<FNiagaraInlineDynamicInputFormatToken, FInputDisplayEntries>>& FormattedColumn : FormattedColumns)
+	{
+		FGuid InputId = FGuid::NewGuid();
+		TOptional<FInputDisplayEntries> ColumnDisplayEntries;
+		TSharedPtr<SWidget> ColumnWidget;
+		if(FormattedColumn.Num() == 1)
+		{
+			ColumnWidget = ConstructGraphInputWidgetForFormatTokenInputPair(FormattedColumn[0], OwningInput, InputId);
+			if(FormattedColumn[0].Key.Usage == ENiagaraInlineDynamicInputFormatTokenUsage::Input)
+			{
+				ColumnDisplayEntries = FormattedColumn[0].Value;
+			}
+		}
+		else
+		{
+			TSharedRef<SHorizontalBox> ColumnBox = SNew(SHorizontalBox);
+			ColumnWidget = ColumnBox;
+			for (const TPair<FNiagaraInlineDynamicInputFormatToken, FInputDisplayEntries>& FormatTokenInputPair : FormattedColumn)
+			{
+				ColumnBox->AddSlot()
+				.AutoWidth()
+				[
+					ConstructGraphInputWidgetForFormatTokenInputPair(FormatTokenInputPair, OwningInput, InputId)
+				];
+				if (FormatTokenInputPair.Key.Usage == ENiagaraInlineDynamicInputFormatTokenUsage::Input)
+				{
+					ColumnDisplayEntries = FormatTokenInputPair.Value;
+				}
+			}
+		}
+
+		if (ColumnDisplayEntries.IsSet() && ColumnDisplayEntries->InputEntry->GetValueMode() == UNiagaraStackFunctionInput::EValueMode::Data)
+		{
+			UNiagaraStackObject* ChildDataObject = ColumnDisplayEntries->InputEntry->GetChildDataObject();
+			if (ChildDataObject != nullptr)
+			{
+				ConstructDataInterfaceWidgets(DataInterfaceVerticalBox, ColumnDisplayEntries.GetValue(), InputId);
+			}
+		}
+
+		bool bIsDynamicInput = ColumnDisplayEntries.IsSet() && ColumnDisplayEntries->InputEntry->GetValueMode() == UNiagaraStackFunctionInput::EValueMode::Dynamic;
+		AddGraphNodeColumn(GraphNodeColumnIndex, bIsDynamicInput, ColumnWidget.ToSharedRef());
+
+		// Add additional widgets for dynamic inputs.
+		if (ColumnDisplayEntries.IsSet() && bIsDynamicInput)
+		{
+			// Child inputs.
+			InputGrid->AddSlot(GraphNodeColumnIndex, 3)
+			.VAlign(VAlign_Top)
+			[
+				ConstructVerticalGraphInputWidgets(ColumnDisplayEntries->ChildInputDisplayEntries, ColumnDisplayEntries->InputEntry, DataInterfaceVerticalBox)
+			];
+
+			// Connection arrow.
+			FMargin BoxPadding = GetVerticalGraphNodeBorderPadding(GraphNodeColumnIndex, TotalGraphNodeColumnCount);
+			BoxPadding.Top = 2;
+			BoxPadding.Bottom = 2;
+			InputGrid->AddSlot(GraphNodeColumnIndex, 2)
+			.HAlign(HAlign_Center)
+			.Padding(FMargin(0, -2))
+			[
+				SNew(SBox)
+				.WidthOverride(16.0f)
+				.HeightOverride(12.0f)
+				[
+					SNew(SImage)
+					.Image(FNiagaraEditorWidgetsStyle::Get().GetBrush("NiagaraEditor.Stack.InlineDynamicInput.GraphNode.Pointer"))
+					.ColorAndOpacity(FStyleColors::DropdownOutline)
+				]
+			];
+
+			ChildInputsColumnIndex++;
+		}
+		GraphNodeColumnIndex++;
+	}
+
+	return InputGrid;
+}
+
+TSharedRef<SWidget> SNiagaraStackInlineDynamicInput::ConstructHybridGraphInputWidgets(
+	const TArray<FInputDisplayEntries>& InputDisplayEntries,
+	const UNiagaraStackFunctionInput* OwningInput,
+	TSharedRef<SVerticalBox> DataInterfaceVerticalBox)
+{
+	const UNiagaraNodeFunctionCall* DynamicInputFunctionCall =
+		OwningInput != nullptr && OwningInput->GetValueMode() == UNiagaraStackFunctionInput::EValueMode::Dynamic
+			? OwningInput->GetDynamicInputNode() : nullptr;
+	bool bFormatMatched;
+	TArray<TArray<TPair<FNiagaraInlineDynamicInputFormatToken, FInputDisplayEntries>>> FormattedColumns;
+	TArray<FNiagaraInlineDynamicInputFormatToken> InlineFormat = DynamicInputFunctionCall != nullptr 
+		? DynamicInputFunctionCall->GetScriptData()->InlineExpressionFormat 
+		: TArray<FNiagaraInlineDynamicInputFormatToken>();
+	
+	bool bIgnoreLinebreaks = true;
+	bool bTreatInputsWithoutChildrenLikeDecorators = true;
+	GenerateFormattedGroups(InputDisplayEntries, InlineFormat, bIgnoreLinebreaks, bTreatInputsWithoutChildrenLikeDecorators, bFormatMatched, FormattedColumns);
+
+	// Generate the columns
+	TSharedRef<SGridPanel> InputGrid = SNew(SGridPanel);
+	int32 GraphNodeColumnIndex = 0;
+	bool bNeedsDynamicInputLabel = bFormatMatched == false && DynamicInputFunctionCall != nullptr ? 1 : 0;
+	int32 TotalGraphNodeColumnCount = FormattedColumns.Num();
+
+	if (bNeedsDynamicInputLabel && FormattedColumns.Num() == 0)
+	{
+		InputGrid->AddSlot(0, 0)
+		[
+			SNew(SBorder)
+			.BorderImage(GetHybridGraphNodeBorderBrush(0, 1, 0, 1))
+			.BorderBackgroundColor(FStyleColors::DropdownOutline)
+			.HAlign(HAlign_Left)
+			.VAlign(VAlign_Center)
+			.Padding(GetHybridGraphNodeBorderPadding(0, 1, 0, 1))
+			[
+				ConstructDynamicInputLabel(OwningInput)
+			]
+		];
+	}
+
+	// Generate column widgets.
+	for (const TArray<TPair<FNiagaraInlineDynamicInputFormatToken, FInputDisplayEntries>>& FormattedColumn : FormattedColumns)
+	{
+		FGuid InputId = FGuid::NewGuid();
+		TOptional<FInputDisplayEntries> ColumnDisplayEntries;
+		TSharedRef<SGridPanel> ColumnGrid = SNew(SGridPanel);
+		ColumnGrid->SetRowFill(0, 1);
+		int32 ColumnGridColumnIndex = 0;
+		int32 TotalColumnGridColumns = FormattedColumn.Num();
+		if (bNeedsDynamicInputLabel && GraphNodeColumnIndex == 0)
+		{
+			TotalColumnGridColumns++;
+			ColumnGrid->AddSlot(0, 0)
+			[
+				SNew(SBorder)
+				.BorderImage(GetHybridGraphNodeBorderBrush(GraphNodeColumnIndex, TotalGraphNodeColumnCount, ColumnGridColumnIndex, TotalColumnGridColumns))
+				.BorderBackgroundColor(FStyleColors::DropdownOutline)
+				.HAlign(HAlign_Left)
+				.VAlign(VAlign_Center)
+				.Padding(GetHybridGraphNodeBorderPadding(GraphNodeColumnIndex, TotalGraphNodeColumnCount, ColumnGridColumnIndex, TotalColumnGridColumns) + FMargin(0, 0, 5, 0))
+				[
+					ConstructDynamicInputLabel(OwningInput)
+				]
+			];
+			ColumnGridColumnIndex++;
+		}
+
+		for (const TPair<FNiagaraInlineDynamicInputFormatToken, FInputDisplayEntries>& FormatTokenInputPair : FormattedColumn)
+		{
+			ColumnGrid->AddSlot(ColumnGridColumnIndex, 0)
+			.HAlign(GetHybridGraphNodeGridHorizontalAlignment(GraphNodeColumnIndex, TotalGraphNodeColumnCount, ColumnGridColumnIndex, TotalColumnGridColumns))
+			[
+				SNew(SBorder)
+				.BorderImage(GetHybridGraphNodeBorderBrush(GraphNodeColumnIndex, TotalGraphNodeColumnCount, ColumnGridColumnIndex, TotalColumnGridColumns))
+				.BorderBackgroundColor(FStyleColors::DropdownOutline)
+				.HAlign(HAlign_Left)
+				.VAlign(VAlign_Fill)
+				.Padding(GetHybridGraphNodeBorderPadding(GraphNodeColumnIndex, TotalGraphNodeColumnCount, ColumnGridColumnIndex, TotalColumnGridColumns))
+				[
+					ConstructGraphInputWidgetForFormatTokenInputPair(FormatTokenInputPair, OwningInput, InputId)
+				]
+			];
+			if (FormatTokenInputPair.Key.Usage == ENiagaraInlineDynamicInputFormatTokenUsage::Input && FormatTokenInputPair.Value.InputEntry->GetValueMode() == UNiagaraStackFunctionInput::EValueMode::Dynamic)
+			{
+				// Connection arrow.
+				ColumnGrid->AddSlot(ColumnGridColumnIndex, 1, SGridPanel::Layer(10))
+				.Padding(FMargin(0.0f, -2.0f))
+				.HAlign(HAlign_Center)
+				.VAlign(VAlign_Bottom)
+				[
+					SNew(SBox)
+					.WidthOverride(16.0f)
+					.HeightOverride(12.0f)
+					[
+						SNew(SImage)
+						.Image(FNiagaraEditorWidgetsStyle::Get().GetBrush("NiagaraEditor.Stack.InlineDynamicInput.GraphNode.Pointer"))
+						.ColorAndOpacity(FStyleColors::DropdownOutline)
+					]
+				];
+				ColumnDisplayEntries = FormatTokenInputPair.Value;
+			}
+			if (GraphNodeColumnIndex < TotalGraphNodeColumnCount - 1 && ColumnGridColumnIndex == TotalColumnGridColumns - 1)
+			{
+				ColumnGrid->SetColumnFill(ColumnGridColumnIndex, 1);
+			}
+			ColumnGridColumnIndex++;
+		}
+
+		if (ColumnDisplayEntries.IsSet() && ColumnDisplayEntries->InputEntry->GetValueMode() == UNiagaraStackFunctionInput::EValueMode::Data)
+		{
+			UNiagaraStackObject* ChildDataObject = ColumnDisplayEntries->InputEntry->GetChildDataObject();
+			if (ChildDataObject != nullptr)
+			{
+				ConstructDataInterfaceWidgets(DataInterfaceVerticalBox, ColumnDisplayEntries.GetValue(), InputId);
+			}
+		}
+
+		InputGrid->AddSlot(GraphNodeColumnIndex, 0)
+		.VAlign(VAlign_Fill)
+		[
+			ColumnGrid
+		];
+
+		// Add additional widgets for dynamic inputs.
+		if (ColumnDisplayEntries.IsSet())
+		{
+			InputGrid->AddSlot(GraphNodeColumnIndex, 1)
+			.HAlign(HAlign_Left)
+			[
+				ConstructHybridGraphInputWidgets(ColumnDisplayEntries->ChildInputDisplayEntries, ColumnDisplayEntries->InputEntry, DataInterfaceVerticalBox)
+			];
+		}
+		GraphNodeColumnIndex++;
+	}
+
+	return InputGrid;
+}
+
+void SNiagaraStackInlineDynamicInput::ConstructDataInterfaceWidgets(TSharedRef<SVerticalBox> DataInterfaceVerticalBox, const FInputDisplayEntries& InputDisplayEntries, FGuid InputId)
+{
+	UNiagaraStackFunctionInput* FunctionInput = InputDisplayEntries.InputEntry;
+	UNiagaraStackObject* DataInterfaceObject = InputDisplayEntries.ChildObjectEntry;
+	if (DataInterfaceObject->GetObject() == nullptr)
+	{
+		return;
+	}
+
+	TSharedRef<FNiagaraObjectSelection> SelectedObject = MakeShared<FNiagaraObjectSelection>();
+	SelectedObject->SetSelectedObject(DataInterfaceObject->GetObject());
+	DataInterfaceVerticalBox->AddSlot()
+	[
+		SNew(SVerticalBox)
+		+ SVerticalBox::Slot()
+		.AutoHeight()
+		.HAlign(HAlign_Left)
+		.Padding(0, 5, 0, 0)
+		[
+			SNew(SNiagaraStackInlineDynamicInputExpressionTokenBorder, FunctionInput, InputId)
+			.HoveredInputId(this, &SNiagaraStackInlineDynamicInput::GetHoveredInputId)
+			.OnHoveredChanged(this, &SNiagaraStackInlineDynamicInput::OnFormatTokenOutlineHoveredChanged)
+			[
+				SNew(SHorizontalBox)
+				+ SHorizontalBox::Slot()
+				.AutoWidth()
+				.Padding(0, 0, 5, 0)
+				[
+					SNew(SBox)
+					.WidthOverride(TextIconSize)
+					.VAlign(VAlign_Center)
+					[
+						SNew(STextBlock)
+						.Font(FAppStyle::Get().GetFontStyle("FontAwesome.10"))
+						.Text(FNiagaraStackEditorWidgetsUtilities::GetIconTextForInputMode(UNiagaraStackFunctionInput::EValueMode::Data))
+						.ToolTipText(FNiagaraStackEditorWidgetsUtilities::GetIconToolTipForInputMode(UNiagaraStackFunctionInput::EValueMode::Data))
+						.ColorAndOpacity(FNiagaraEditorWidgetsStyle::Get().GetColor(FNiagaraStackEditorWidgetsUtilities::GetIconColorNameForInputMode(UNiagaraStackFunctionInput::EValueMode::Data)))
+					]
+				]
+				+ SHorizontalBox::Slot()
+				[
+					SNew(STextBlock)
+					.TextStyle(FNiagaraEditorStyle::Get(), "NiagaraEditor.ParameterText")
+					.Text(FunctionInput->GetInputType().GetClass()->GetDisplayNameText())
+				]
+			]
+		]
+		+ SVerticalBox::Slot()
+		.Padding(0, 5, 0, 0)
+		[
+			SNew(SNiagaraSelectedObjectsDetails, SelectedObject)
+		]
+	];
+}
+
+FGuid SNiagaraStackInlineDynamicInput::GetHoveredInputId() const
+{
+	return HoveredInputId;
+}
+
+void SNiagaraStackInlineDynamicInput::OnFormatTokenOutlineHoveredChanged(FGuid InputId, bool bOutlineIsHovered)
+{
+	if (bOutlineIsHovered)
+	{
+		HoveredInputId = InputId;
+	}
+	else
+	{
+		if (HoveredInputId == InputId)
+		{
+			HoveredInputId = FGuid();
+		}
+	}
+}
+
+EVisibility SNiagaraStackInlineDynamicInput::GetEditConditionCheckBoxVisibility(UNiagaraStackFunctionInput* FunctionInput) const
+{
+	return FunctionInput->GetHasEditCondition() && FunctionInput->GetShowEditConditionInline() ? EVisibility::Visible : EVisibility::Collapsed;
+}
+
+ECheckBoxState SNiagaraStackInlineDynamicInput::GetEditConditionCheckState(UNiagaraStackFunctionInput* FunctionInput) const
+{
+	return FunctionInput->GetHasEditCondition() && FunctionInput->GetEditConditionEnabled() ? ECheckBoxState::Checked : ECheckBoxState::Unchecked;
+}
+
+void SNiagaraStackInlineDynamicInput::OnEditConditionCheckStateChanged(ECheckBoxState InCheckState, UNiagaraStackFunctionInput* FunctionInput)
+{
+	FunctionInput->SetEditConditionEnabled(InCheckState == ECheckBoxState::Checked);
+}
+
+EVisibility SNiagaraStackInlineDynamicInput::GetCompactActionMenuButtonVisibility(FGuid InputId) const
+{
+	return InputId == HoveredInputId ? EVisibility::Visible : EVisibility::Hidden;
+}
+
+#undef LOCTEXT_NAMESPACE
\ No newline at end of file
diff --git a/Engine/Plugins/FX/Niagara/Source/NiagaraEditorWidgets/Private/Stack/SNiagaraStackInlineDynamicInput.h b/Engine/Plugins/FX/Niagara/Source/NiagaraEditorWidgets/Private/Stack/SNiagaraStackInlineDynamicInput.h
new file mode 100644
--- /dev/null
+++ b/Engine/Plugins/FX/Niagara/Source/NiagaraEditorWidgets/Private/Stack/SNiagaraStackInlineDynamicInput.h
@@ -0,0 +1,124 @@
+// Copyright Epic Games, Inc. All Rights Reserved.
+
+#pragma once
+
+#include "ViewModels/Stack/NiagaraStackFunctionInput.h"
+#include "Widgets/DeclarativeSyntaxSupport.h"
+#include "Widgets/SCompoundWidget.h"
+
+class SWrapBox;
+class SGridPanel;
+class UNiagaraNodeFunctionCall;
+
+class SNiagaraStackInlineDynamicInput : public SCompoundWidget
+{
+public:
+	SLATE_BEGIN_ARGS(SNiagaraStackInlineDynamicInput) {}
+	SLATE_END_ARGS();
+
+	void Construct(const FArguments& InArgs, UNiagaraStackFunctionInput* InFunctionInput);
+
+private:
+	struct FInputDisplayEntries
+	{
+		FInputDisplayEntries()
+			: InputEntry(nullptr)
+			, ChildObjectEntry(nullptr)
+		{
+		}
+
+		FORCEINLINE bool operator ==(const FInputDisplayEntries& Other) const
+		{
+			return 
+				InputEntry == Other.InputEntry &&
+				ChildInputDisplayEntries == Other.ChildInputDisplayEntries &&
+				ChildObjectEntry == Other.ChildObjectEntry;	
+		}
+
+		FORCEINLINE bool operator !=(const FInputDisplayEntries& Other) const { return !(*this == Other); }
+
+		UNiagaraStackFunctionInput* InputEntry;
+		TArray<FInputDisplayEntries> ChildInputDisplayEntries;
+		UNiagaraStackObject* ChildObjectEntry;
+	};
+
+	class FWidthSynchronize
+	{
+	public:
+		void AddWidget(TSharedRef<SWidget> InWidget);
+		FOptionalSize GetMaxDesiredWidth() const;
+	private:
+		TArray<TWeakPtr<SWidget>> WeakWidgets;
+	};
+
+private:
+	static FInputDisplayEntries CollectInputDisplayEntriesRecursive(UNiagaraStackFunctionInput* Input);
+
+	void RootFunctionInputStructureChanged(ENiagaraStructureChangedFlags Flags);
+
+	void RootFunctionInputValueChanged();
+
+	static bool TryMatchFormatInputs(
+		const TArray<FInputDisplayEntries>& ChildInputs,
+		const TArray<FNiagaraInlineDynamicInputFormatToken>& InlineFormat,
+		TArray<TPair<FNiagaraInlineDynamicInputFormatToken, FInputDisplayEntries>>& OutMatchedFormatInputs);
+
+	void ConstructChildren();
+
+	void ConstructWrappedExpressionInputWidgets(TSharedRef<SWrapBox> ExpressionWrapBox, TSharedRef<SVerticalBox> DataInterfaceVerticalBox, const FInputDisplayEntries& InputDisplayEntries, FGuid InputId, bool bIsRootInput);
+
+	void ConstructWrappedExpressionDynamicInputWidgets(TSharedRef<SWrapBox> ExpressionWrapBox, TSharedRef<SVerticalBox> DataInterfaceVerticalBox, const FInputDisplayEntries& InputDisplayEntries, FGuid InputId, bool bIsRootInput);
+
+	TSharedRef<SWidget> ConstructGraphInputWidgetForFormatTokenInputPair(
+		const TPair<FNiagaraInlineDynamicInputFormatToken, FInputDisplayEntries>& FormatTokenInputPair,
+		const UNiagaraStackFunctionInput* OwningInput,
+		FGuid InputId);
+
+	void GenerateFormattedGroups(
+		const TArray<FInputDisplayEntries>& InputDisplayEntries,
+		const TArray<FNiagaraInlineDynamicInputFormatToken>& InlineFormat,
+		bool bIgnoreLinebreaks,
+		bool bTreatInputsWithoutChildrenLikeDecorators,
+		bool& bOutFormatMatched,
+		TArray<TArray<TPair<FNiagaraInlineDynamicInputFormatToken, FInputDisplayEntries>>>& OutFormattedGroups);
+
+	TSharedRef<SWidget> ConstructHorizontalGraphInputWidgets(
+		const TArray<FInputDisplayEntries>& InputDisplayEntries,
+		const UNiagaraStackFunctionInput* OwningInput,
+		TSharedRef<SVerticalBox> DataInterfaceVerticalBox,
+		TSharedPtr<FWidthSynchronize> InputWidthSynchronizer);
+
+	TSharedRef<SWidget> ConstructVerticalGraphInputWidgets(
+		const TArray<FInputDisplayEntries>& InputDisplayEntries,
+		const UNiagaraStackFunctionInput* OwningInput,
+		TSharedRef<SVerticalBox> DataInterfaceVerticalBox);
+
+	TSharedRef<SWidget> ConstructHybridGraphInputWidgets(
+		const TArray<FInputDisplayEntries>& InputDisplayEntries,
+		const UNiagaraStackFunctionInput* OwningInput,
+		TSharedRef<SVerticalBox> DataInterfaceVerticalBox);
+
+	void ConstructDataInterfaceWidgets(TSharedRef<SVerticalBox> DataInterfaceVerticalBox, const FInputDisplayEntries& InputDisplayEntries, FGuid InputId);
+
+	FGuid GetHoveredInputId() const;
+
+	void OnFormatTokenOutlineHoveredChanged(FGuid InputId, bool bIsHovered);
+
+	EVisibility GetEditConditionCheckBoxVisibility(UNiagaraStackFunctionInput* FunctionInput) const;
+
+	ECheckBoxState GetEditConditionCheckState(UNiagaraStackFunctionInput* FunctionInput) const;
+
+	void OnEditConditionCheckStateChanged(ECheckBoxState InCheckState, UNiagaraStackFunctionInput* FunctionInput);
+
+	EVisibility GetCompactActionMenuButtonVisibility(FGuid InputId) const;
+
+private:
+	UNiagaraStackFunctionInput* RootFunctionInput;
+	FInputDisplayEntries RootInputDisplayEntries;
+	TArray<TSharedPtr<FWidthSynchronize>> WidthSynchronizers;
+
+	bool bSingleValueMode;
+	FGuid HoveredInputId;
+
+	static const float TextIconSize;
+};
\ No newline at end of file
diff --git a/Engine/Plugins/FX/Niagara/Source/NiagaraEditorWidgets/Private/Stack/SNiagaraStackTableRow.cpp b/Engine/Plugins/FX/Niagara/Source/NiagaraEditorWidgets/Private/Stack/SNiagaraStackTableRow.cpp
--- a/Engine/Plugins/FX/Niagara/Source/NiagaraEditorWidgets/Private/Stack/SNiagaraStackTableRow.cpp
+++ b/Engine/Plugins/FX/Niagara/Source/NiagaraEditorWidgets/Private/Stack/SNiagaraStackTableRow.cpp
@@ -32,7 +32,7 @@
 
 void SNiagaraStackTableRow::Construct(const FArguments& InArgs, UNiagaraStackViewModel* InStackViewModel, UNiagaraStackEntry* InStackEntry, TSharedRef<FNiagaraStackCommandContext> InStackCommandContext, const TSharedRef<STreeView<UNiagaraStackEntry*>>& InOwnerTree)
 {
-	ContentPadding = InArgs._ContentPadding;
+	DefaultContentPadding = InArgs._ContentPadding;
 	bIsCategoryIconHighlighted = InArgs._IsCategoryIconHighlighted;
 	bShowExecutionCategoryIcon = InArgs._ShowExecutionCategoryIcon;
 	NameColumnWidth = InArgs._NameColumnWidth;
@@ -45,8 +45,7 @@ void SNiagaraStackTableRow::Construct(const FArguments& InArgs, UNiagaraStackVie
 	StackCommandContext = InStackCommandContext;
 	OwnerTree = InOwnerTree;
 
-	NameVerticalAlignment = VAlign_Center;
-	NameHorizontalAlignment = HAlign_Fill;
+	Reset();
 
 	ExpandedImage = FCoreStyle::Get().GetBrush("TreeArrow_Expanded");
 	CollapsedImage = FCoreStyle::Get().GetBrush("TreeArrow_Collapsed");
@@ -67,6 +66,21 @@ void SNiagaraStackTableRow::Construct(const FArguments& InArgs, UNiagaraStackVie
 		, OwnerTree.ToSharedRef());
 }
 
+void SNiagaraStackTableRow::Reset()
+{
+	ChildSlot[SNullWidget::NullWidget];
+	ContentPadding = DefaultContentPadding;
+	NameMinWidth.Reset();
+	NameMaxWidth.Reset();
+	NameVerticalAlignment = VAlign_Center;
+	NameHorizontalAlignment = HAlign_Fill;
+	ValueMinWidth.Reset();
+	ValueMaxWidth.Reset();
+	ValueVerticalAlignment = VAlign_Center;
+	ValueHorizontalAlignment = HAlign_Fill;
+	OnFillRowContextMenuHanders.Empty();
+}
+
 void SNiagaraStackTableRow::SetOverrideNameWidth(TOptional<float> InMinWidth, TOptional<float> InMaxWidth)
 {
 	NameMinWidth = InMinWidth;
diff --git a/Engine/Plugins/FX/Niagara/Source/NiagaraEditorWidgets/Private/Stack/SNiagaraStackTableRow.h b/Engine/Plugins/FX/Niagara/Source/NiagaraEditorWidgets/Private/Stack/SNiagaraStackTableRow.h
--- a/Engine/Plugins/FX/Niagara/Source/NiagaraEditorWidgets/Private/Stack/SNiagaraStackTableRow.h
+++ b/Engine/Plugins/FX/Niagara/Source/NiagaraEditorWidgets/Private/Stack/SNiagaraStackTableRow.h
@@ -45,6 +45,8 @@ public:
 
 	void Construct(const FArguments& InArgs, UNiagaraStackViewModel* InStackViewModel, UNiagaraStackEntry* InStackEntry, TSharedRef<FNiagaraStackCommandContext> InStackCommandContext, const TSharedRef<STreeView<UNiagaraStackEntry*>>& InOwnerTree);
 
+	void Reset();
+
 	void SetOverrideNameWidth(TOptional<float> InMinWidth, TOptional<float> InMaxWidth);
 
 	void SetOverrideNameAlignment(EHorizontalAlignment InHAlign, EVerticalAlignment InVAlign);
@@ -116,6 +118,7 @@ private:
 
 	FText ExecutionCategoryToolTipText;
 
+	FMargin DefaultContentPadding;
 	FMargin ContentPadding;
 
 	EHorizontalAlignment NameHorizontalAlignment;
