Trivial fix for linux build

diff --git a/Engine/Plugins/FX/Niagara/Source/NiagaraEditorWidgets/Private/Stack/SNiagaraStackInlineDynamicInput.cpp b/Engine/Plugins/FX/Niagara/Source/NiagaraEditorWidgets/Private/Stack/SNiagaraStackInlineDynamicInput.cpp
--- a/Engine/Plugins/FX/Niagara/Source/NiagaraEditorWidgets/Private/Stack/SNiagaraStackInlineDynamicInput.cpp
+++ b/Engine/Plugins/FX/Niagara/Source/NiagaraEditorWidgets/Private/Stack/SNiagaraStackInlineDynamicInput.cpp
@@ -444,7 +444,7 @@ void SNiagaraStackInlineDynamicInput::ConstructWrappedExpressionDynamicInputWidg
 			];
 		}
 
-		for (const TPair<FNiagaraInlineDynamicInputFormatToken, FInputDisplayEntries> FormatTokenInputPair : OutMatchedFormatInputs)
+		for (const TPair<FNiagaraInlineDynamicInputFormatToken, FInputDisplayEntries>& FormatTokenInputPair : OutMatchedFormatInputs)
 		{
 			if (FormatTokenInputPair.Key.Usage == ENiagaraInlineDynamicInputFormatTokenUsage::Input)
 			{
