- Copy the default DI in CopyParametersFromAsset
- This fixes an issue where SetAsset is called from AssetA which uses a bone list (for example) and we set another asset which shares the same DI name the bone list would be incorrect

diff --git a/Engine/Plugins/FX/Niagara/Source/Niagara/Private/NiagaraComponent.cpp b/Engine/Plugins/FX/Niagara/Source/Niagara/Private/NiagaraComponent.cpp
--- a/Engine/Plugins/FX/Niagara/Source/Niagara/Private/NiagaraComponent.cpp
+++ b/Engine/Plugins/FX/Niagara/Source/Niagara/Private/NiagaraComponent.cpp
@@ -3314,12 +3314,15 @@ void UNiagaraComponent::CopyParametersFromAsset(bool bResetExistingOverrideParam
 		UNiagaraDataInterface* OverrideDataInterface = OverrideParameters.GetDataInterface(i);
 		if (AssetDataInterface != OverrideDataInterface)
 		{
-			continue;
+			// We must copy the data regardless as there is an expectation that data interfaces are always reset to the original state
+			AssetDataInterface->CopyTo(OverrideDataInterface);
+		}
+		else
+		{
+			OverrideDataInterface = NewObject<UNiagaraDataInterface>(this, AssetDataInterface->GetClass(), NAME_None, RF_Transactional | RF_Public);
+			AssetDataInterface->CopyTo(OverrideDataInterface);
+			OverrideParameters.SetDataInterface(OverrideDataInterface, i);
 		}
-
-		OverrideDataInterface = NewObject<UNiagaraDataInterface>(this, AssetDataInterface->GetClass(), NAME_None, RF_Transactional | RF_Public);
-		AssetDataInterface->CopyTo(OverrideDataInterface);
-		OverrideParameters.SetDataInterface(OverrideDataInterface, i);
 	}
 }
 
