- Remove unused / empty function

#rnx

diff --git a/Engine/Plugins/FX/Niagara/Source/Niagara/Private/NiagaraRenderer.cpp b/Engine/Plugins/FX/Niagara/Source/Niagara/Private/NiagaraRenderer.cpp
--- a/Engine/Plugins/FX/Niagara/Source/Niagara/Private/NiagaraRenderer.cpp
+++ b/Engine/Plugins/FX/Niagara/Source/Niagara/Private/NiagaraRenderer.cpp
@@ -351,13 +351,8 @@ FNiagaraDataBuffer* FNiagaraDynamicDataBase::GetParticleDataToRender(bool bIsLow
 	return Ret;
 }
 
-void FNiagaraDynamicDataBase::SetVertexFactoryData(class FNiagaraVertexFactoryBase& VertexFactory)
-{
-}
-
 //////////////////////////////////////////////////////////////////////////
 
-
 FNiagaraRenderer::FNiagaraRenderer(ERHIFeatureLevel::Type InFeatureLevel, const UNiagaraRendererProperties *InProps, const FNiagaraEmitterInstance* Emitter)
 	: DynamicDataRender(nullptr)
 	, bLocalSpace(Emitter->GetCachedEmitterData()->bLocalSpace)
diff --git a/Engine/Plugins/FX/Niagara/Source/Niagara/Private/NiagaraRendererRibbons.cpp b/Engine/Plugins/FX/Niagara/Source/Niagara/Private/NiagaraRendererRibbons.cpp
--- a/Engine/Plugins/FX/Niagara/Source/Niagara/Private/NiagaraRendererRibbons.cpp
+++ b/Engine/Plugins/FX/Niagara/Source/Niagara/Private/NiagaraRendererRibbons.cpp
@@ -2051,8 +2051,6 @@ inline void FNiagaraRendererRibbons::SetupMeshBatchAndCollectorResourceForView(c
 	check(MaterialRenderProxy);
 	
 	// Set common data on vertex factory
-	DynamicDataRibbon->SetVertexFactoryData(RenderingViewResources->VertexFactory);
-
 	FNiagaraRibbonVFLooseParameters VFLooseParams;
 	VFLooseParams.SortedIndices = VertexBuffers.SortedIndicesBuffer.SRV;
 	VFLooseParams.TangentsAndDistances = VertexBuffers.TangentsAndDistancesBuffer.SRV;
diff --git a/Engine/Plugins/FX/Niagara/Source/Niagara/Public/NiagaraRenderer.h b/Engine/Plugins/FX/Niagara/Source/Niagara/Public/NiagaraRenderer.h
--- a/Engine/Plugins/FX/Niagara/Source/Niagara/Public/NiagaraRenderer.h
+++ b/Engine/Plugins/FX/Niagara/Source/Niagara/Public/NiagaraRenderer.h
@@ -43,7 +43,6 @@ struct NIAGARA_API FNiagaraDynamicDataBase
 
 	FORCEINLINE FNiagaraSystemInstanceID GetSystemInstanceID() const { return SystemInstanceID; }
 
-	void SetVertexFactoryData(class FNiagaraVertexFactoryBase& VertexFactory);
 	virtual void ApplyMaterialOverride(int32 MaterialIndex, UMaterialInterface* MaterialOverride) {};
 
 protected:
