Fix GPU crash on IOS 13 as well as some PowerVR android devices caused by [REDACTED_PERSON] shader.

diff --git a/Engine/Plugins/FX/Niagara/Shaders/Private/NiagaraMeshParticleUtils.ush b/Engine/Plugins/FX/Niagara/Shaders/Private/NiagaraMeshParticleUtils.ush
--- a/Engine/Plugins/FX/Niagara/Shaders/Private/NiagaraMeshParticleUtils.ush
+++ b/Engine/Plugins/FX/Niagara/Shaders/Private/NiagaraMeshParticleUtils.ush
@@ -257,8 +257,7 @@ float3x3 NiagaraGetMeshFacingMatrix(FLWCVector3 ParticleSimPosition, NiagaraMesh
 			break;
 		}
 
-		default:
-		case MESH_FACING_CAMERA_PLANE:
+		default: // case MESH_FACING_DEFAULT and MESH_FACING_CAMERA_PLANE
 		{
 			FacingDir = NiagaraWorldToSimVec(-Params.CameraForwardDir, Params);
 			break;
