Fix for loc errors with NiagaraGraphDigest
#rnx

diff --git a/Engine/Plugins/FX/Niagara/Source/NiagaraEditor/Private/NiagaraGraphDigest.cpp b/Engine/Plugins/FX/Niagara/Source/NiagaraEditor/Private/NiagaraGraphDigest.cpp
--- a/Engine/Plugins/FX/Niagara/Source/NiagaraEditor/Private/NiagaraGraphDigest.cpp
+++ b/Engine/Plugins/FX/Niagara/Source/NiagaraEditor/Private/NiagaraGraphDigest.cpp
@@ -34,7 +34,7 @@
 
 #include "Algo/RemoveIf.h"
 
-#define LOCTEXT_NAMESPACE "NiagaraDigestCompiler"
+#define LOCTEXT_NAMESPACE "NiagaraCompiler"
 
 #define NIAGARA_GRAPH_DIGEST_NODE_TYPE(name) NIAGARA_GRAPH_DIGEST_NODE_IMPLEMENT_BODY(name)
 
@@ -3496,7 +3496,7 @@ void FNiagaraCompilationNodeConvert::Compile(FTranslator* Translator, TArray<int
 			int32 CompiledInput = Translator->CompileInputPin(&InputPin);
 			if (CompiledInput == INDEX_NONE)
 			{
-				Translator->Error(LOCTEXT("InputError", "Error compiling input for convert node."), this, &InputPin);
+				Translator->Error(LOCTEXT("ConvertInputError", "Error compiling input for convert node."), this, &InputPin);
 			}
 			CompileInputs.Add(CompiledInput);
 		}
@@ -3600,7 +3600,7 @@ void FNiagaraCompilationNodeParameterMapGet::Compile(FTranslator* Translator, TA
 			int32 CompiledInput = Translator->CompileInputPin(&ParamMapInputPin);
 			if (CompiledInput == INDEX_NONE)
 			{
-				Translator->Error(LOCTEXT("InputError", "Error compiling input for param map get node."), this, &ParamMapInputPin);
+				Translator->Error(LOCTEXT("MapGetInputError", "Error compiling input for param map get node."), this, &ParamMapInputPin);
 			}
 
 			CompileInputs[0] = CompiledInput;
@@ -3784,7 +3784,7 @@ void FNiagaraCompilationNodeParameterMapSet::Compile(FTranslator* Translator, TA
 			int32 CompiledInput = Translator->CompileInputPin(&InputPin);
 			if (CompiledInput == INDEX_NONE)
 			{
-				Translator->Error(LOCTEXT("InputError", "Error compiling input for set node."), this, &InputPin);
+				Translator->Error(LOCTEXT("MapSetInputError", "Error compiling input for set node."), this, &InputPin);
 			}
 			CompileInputs.Emplace(CompiledInput, &InputPin);
 		}
