- Speculative fix for rare emitter instance crash
- Store a separate task ref for the instance batch tick and wait on them individually

#rnx

diff --git a/Engine/Plugins/FX/Niagara/Source/Niagara/Private/NiagaraSystemInstance.cpp b/Engine/Plugins/FX/Niagara/Source/Niagara/Private/NiagaraSystemInstance.cpp
--- a/Engine/Plugins/FX/Niagara/Source/Niagara/Private/NiagaraSystemInstance.cpp
+++ b/Engine/Plugins/FX/Niagara/Source/Niagara/Private/NiagaraSystemInstance.cpp
@@ -2288,6 +2288,7 @@ void FNiagaraSystemInstance::DumpStalledInfo()
 	TStringBuilder<128> Builder;
 	Builder.Appendf(TEXT("System (%s)\n"), *GetNameSafe(GetSystem()));
 	Builder.Appendf(TEXT("ConcurrentTickGraphEvent Complete (%d)\n"), ConcurrentTickGraphEvent ? ConcurrentTickGraphEvent->IsComplete() : true);
+	Builder.Appendf(TEXT("ConcurrentTickBatchGraphEvent Complete (%d)\n"), ConcurrentTickBatchGraphEvent ? ConcurrentTickBatchGraphEvent->IsComplete() : true);
 	Builder.Appendf(TEXT("FinalizePending (%d)\n"), FinalizeRef.IsPending());
 	Builder.Appendf(TEXT("SystemInstanceIndex (%d)\n"), SystemInstanceIndex);
 	Builder.Appendf(TEXT("SystemInstanceState (%d)\n"), SystemInstanceState);
@@ -2300,7 +2301,9 @@ void FNiagaraSystemInstance::WaitForConcurrentTickDoNotFinalize(bool bEnsureComp
 	check(IsInGameThread());
 
 	// Wait for any concurrent ticking for our task
-	if (ConcurrentTickGraphEvent && !ConcurrentTickGraphEvent->IsComplete())
+	const bool bConcurrentTickPending = ConcurrentTickGraphEvent && !ConcurrentTickGraphEvent->IsComplete();
+	const bool bConcurrentBatchTickPending = ConcurrentTickBatchGraphEvent && !ConcurrentTickBatchGraphEvent->IsComplete();
+	if (bConcurrentTickPending || bConcurrentBatchTickPending)
 	{
 		ensureAlwaysMsgf(!bEnsureComplete, TEXT("FNiagaraSystemInstance::WaitForConcurrentTickDoNotFinalize - Async Work not complete and is expected to be. %s"), *GetSystem()->GetPathName());
 		SCOPE_CYCLE_COUNTER(STAT_NiagaraSystemWaitForAsyncTick);
@@ -2310,10 +2313,11 @@ void FNiagaraSystemInstance::WaitForConcurrentTickDoNotFinalize(bool bEnsureComp
 		const double WarnSeconds = 5.0;
 		const uint64 WarnCycles = StartCycles + uint64(WarnSeconds / FPlatformTime::GetSecondsPerCycle64());
 
-		extern int32 GNiagaraSystemSimulationTaskStallTimeout;
-		if (GNiagaraSystemSimulationTaskStallTimeout > 0)
+		// Note must be done in this order as ConcurrentTickBatchGraphEvent is assigned in the async task
+		if (bConcurrentTickPending)
 		{
-			do
+			extern int32 GNiagaraSystemSimulationTaskStallTimeout;
+			if (GNiagaraSystemSimulationTaskStallTimeout > 0)
 			{
 				const double EndTimeoutSeconds = FPlatformTime::Seconds() + (double(GNiagaraSystemSimulationTaskStallTimeout) / 1000.0);
 				LowLevelTasks::BusyWaitUntil(
@@ -2327,14 +2331,16 @@ void FNiagaraSystemInstance::WaitForConcurrentTickDoNotFinalize(bool bEnsureComp
 						return ConcurrentTickGraphEvent->IsComplete();
 					}
 				);
-			} while (ConcurrentTickGraphEvent && !ConcurrentTickGraphEvent->IsComplete());
-		}
-		else
-		{
-			do
+			}
+			else
 			{
 				FTaskGraphInterface::Get().WaitUntilTaskCompletes(ConcurrentTickGraphEvent, ENamedThreads::GameThread_Local);
-			} while (ConcurrentTickGraphEvent && !ConcurrentTickGraphEvent->IsComplete());
+			}
+		}
+		FPlatformMisc::MemoryBarrier();
+		if (ConcurrentTickBatchGraphEvent && !ConcurrentTickBatchGraphEvent->IsComplete())
+		{
+			FTaskGraphInterface::Get().WaitUntilTaskCompletes(ConcurrentTickBatchGraphEvent, ENamedThreads::GameThread_Local);
 		}
 
 		const double StallTimeMS = FPlatformTime::ToMilliseconds64(FPlatformTime::Cycles64() - StartCycles);
@@ -2345,6 +2351,7 @@ void FNiagaraSystemInstance::WaitForConcurrentTickDoNotFinalize(bool bEnsureComp
 		}
 	}
 	ConcurrentTickGraphEvent = nullptr;
+	ConcurrentTickBatchGraphEvent = nullptr;
 }
 
 void FNiagaraSystemInstance::WaitForConcurrentTickAndFinalize(bool bEnsureComplete)
@@ -2620,6 +2627,8 @@ void FNiagaraSystemInstance::FinalizeTick_GameThread(bool bEnqueueGPUTickIfNeede
 	// Ensure concurrent work is complete and clear the finalize ref
 	check(ConcurrentTickGraphEvent == nullptr || ConcurrentTickGraphEvent->IsComplete());
 	ConcurrentTickGraphEvent = nullptr;
+	check(ConcurrentTickBatchGraphEvent == nullptr || ConcurrentTickBatchGraphEvent->IsComplete());
+	ConcurrentTickBatchGraphEvent = nullptr;
 	FinalizeRef.ConditionalClear();
 
 	//Temporarily force FX to update their own LODDistance on frames where it is not provided by the scalability manager.
diff --git a/Engine/Plugins/FX/Niagara/Source/Niagara/Private/NiagaraSystemSimulation.cpp b/Engine/Plugins/FX/Niagara/Source/Niagara/Private/NiagaraSystemSimulation.cpp
--- a/Engine/Plugins/FX/Niagara/Source/Niagara/Private/NiagaraSystemSimulation.cpp
+++ b/Engine/Plugins/FX/Niagara/Source/Niagara/Private/NiagaraSystemSimulation.cpp
@@ -1072,7 +1072,8 @@ void FNiagaraSystemSimulation::FlushTickBatch(FNiagaraSystemSimulationTickContex
 
 			for (FNiagaraSystemInstance* Inst : TickBatch)
 			{
-				Inst->ConcurrentTickGraphEvent = InstanceAsyncGraphEvent;
+				Inst->ConcurrentTickGraphEvent = nullptr;
+				Inst->ConcurrentTickBatchGraphEvent = InstanceAsyncGraphEvent;
 			}
 
 			// Queue finalize task which will run after the instances are complete, track with our all completion event
@@ -2155,6 +2156,7 @@ void FNiagaraSystemSimulation::RemoveInstance(FNiagaraSystemInstance* Instance)
 
 	// We did not finalize and will not so clear the reference
 	Instance->ConcurrentTickGraphEvent = nullptr;
+	Instance->ConcurrentTickBatchGraphEvent = nullptr;
 	Instance->FinalizeRef.ConditionalClear();
 
 	// Remove the instance if it is still valid
diff --git a/Engine/Plugins/FX/Niagara/Source/Niagara/Public/NiagaraSystemInstance.h b/Engine/Plugins/FX/Niagara/Source/Niagara/Public/NiagaraSystemInstance.h
--- a/Engine/Plugins/FX/Niagara/Source/Niagara/Public/NiagaraSystemInstance.h
+++ b/Engine/Plugins/FX/Niagara/Source/Niagara/Public/NiagaraSystemInstance.h
@@ -590,6 +590,7 @@ private:
 
 	/** Graph event to track pending concurrent work. */
 	FGraphEventRef ConcurrentTickGraphEvent;
+	FGraphEventRef ConcurrentTickBatchGraphEvent;
 
 	/** When using concurrent ticking this will be valid until it's complete. */
 	FNiagaraSystemInstanceFinalizeRef FinalizeRef;
