- Fix for ribbon normals tangent x/y being flipped

diff --git a/Engine/Plugins/FX/Niagara/Shaders/Private/NiagaraRibbonVertexFactory.ush b/Engine/Plugins/FX/Niagara/Shaders/Private/NiagaraRibbonVertexFactory.ush
--- a/Engine/Plugins/FX/Niagara/Shaders/Private/NiagaraRibbonVertexFactory.ush
+++ b/Engine/Plugins/FX/Niagara/Shaders/Private/NiagaraRibbonVertexFactory.ush
@@ -583,8 +583,6 @@ FRibbonPositionIntermediates GetRibbonPositionIntermediates(
 		Facing = -SinTwist * OriginalRight + CosTwist * Facing;
 	}
 
-	float3 NormalizedFacing = cross(Intermediates.Direction, Right);
-	
 	Intermediates.TangentY = Intermediates.Direction;
 	Intermediates.TangentX = Right;
 	
@@ -602,6 +600,11 @@ FRibbonPositionIntermediates GetRibbonPositionIntermediates(
 	
 	Intermediates.TangentX = cross(Intermediates.TangentZ, Intermediates.TangentY);
 
+	// Tangent X/Y are incorrectly swapped above detangling that results in 2 paths so less messy to swap here
+	float3 t = Intermediates.TangentX;
+	Intermediates.TangentX = Intermediates.TangentY;
+	Intermediates.TangentY = t;
+
 	// Overrides (FVertexFactoryIntermediates)0; 
 	Intermediates.RibbonWidth = 1.0f;
 
