- Fix for nullptr crash with SVT caching

diff --git a/Engine/Plugins/FX/Niagara/Source/Niagara/Private/NiagaraDataInterfaceRenderTargetVolume.cpp b/Engine/Plugins/FX/Niagara/Source/Niagara/Private/NiagaraDataInterfaceRenderTargetVolume.cpp
--- a/Engine/Plugins/FX/Niagara/Source/Niagara/Private/NiagaraDataInterfaceRenderTargetVolume.cpp
+++ b/Engine/Plugins/FX/Niagara/Source/Niagara/Private/NiagaraDataInterfaceRenderTargetVolume.cpp
@@ -852,7 +852,7 @@ bool UNiagaraDataInterfaceRenderTargetVolume::SimCacheReadFrame(UObject* Storage
 		ENQUEUE_RENDER_COMMAND(NDIRenderTargetVolumeUpdate)
 			(
 				[RT_Proxy, RT_InstanceID = SystemInstance->GetId(),
-				RT_VolumeResolution = VolumeResolution, RT_SVTRenderResources = SVTFrame->GetTextureRenderResources(),
+				RT_VolumeResolution = VolumeResolution, RT_SVTRenderResources = SVTFrame ? SVTFrame->GetTextureRenderResources() : nullptr,
 				FeatureLevel = SystemInstance->GetFeatureLevel()](FRHICommandListImmediate& RHICmdList)
 				{
 					if (RT_SVTRenderResources == nullptr)
