- Ensure we don't cache data interfaces more than once per frame
- The same named DI can live in multiple scripts and we only need to visit it once

#rnx

diff --git a/Engine/Plugins/FX/Niagara/Source/Niagara/Private/NiagaraSimCache.cpp b/Engine/Plugins/FX/Niagara/Source/Niagara/Private/NiagaraSimCache.cpp
--- a/Engine/Plugins/FX/Niagara/Source/Niagara/Private/NiagaraSimCache.cpp
+++ b/Engine/Plugins/FX/Niagara/Source/Niagara/Private/NiagaraSimCache.cpp
@@ -194,6 +194,7 @@ bool UNiagaraSimCache::BeginWrite(FNiagaraSimCacheCreateParameters InCreateParam
 	// Find data interfaces we may want to cache
 	if ( CreateParameters.bAllowDataInterfaceCaching )
 	{
+		TSet<FNiagaraVariableBase> VisitedDataInterfaces;
 		FNiagaraDataInterfaceUtilities::ForEachDataInterface(
 			Helper.SystemInstance,
 			[&](const FNiagaraVariableBase& Variable, UNiagaraDataInterface* DataInterface)
@@ -206,6 +207,12 @@ bool UNiagaraSimCache::BeginWrite(FNiagaraSimCacheCreateParameters InCreateParam
 
 				if (INiagaraSimCacheCustomStorageInterface* SimCacheCustomStorageInterface = Cast<INiagaraSimCacheCustomStorageInterface>(DataInterface))
 				{
+					if (VisitedDataInterfaces.Contains(Variable))
+					{
+						return true;
+					}
+					VisitedDataInterfaces.Add(Variable);
+
 					const void* PerInstanceData = Helper.SystemInstance->FindDataInterfaceInstanceData(DataInterface);
 					if (UObject* DICacheStorage = SimCacheCustomStorageInterface->SimCacheBeginWrite(this, Helper.SystemInstance, PerInstanceData, FeedbackContext))
 					{
@@ -375,6 +382,7 @@ bool UNiagaraSimCache::WriteFrame(UNiagaraComponent* NiagaraComponent, FNiagaraS
 		const int FrameIndex = CacheFrames.Num() - 1;
 		bool bDataInterfacesSucess = true;
 		FString DataInterfaceName;
+		TSet<FNiagaraVariableBase> VisitedDataInterfaces;
 
 		FNiagaraDataInterfaceUtilities::ForEachDataInterface(
 			Helper.SystemInstance,
@@ -382,6 +390,12 @@ bool UNiagaraSimCache::WriteFrame(UNiagaraComponent* NiagaraComponent, FNiagaraS
 			{
 				if ( UObject* StorageObject = DataInterfaceStorage.FindRef(Variable) )
 				{
+					if (VisitedDataInterfaces.Contains(Variable))
+					{
+						return true;
+					}
+					VisitedDataInterfaces.Add(Variable);
+
 					INiagaraSimCacheCustomStorageInterface* SimCacheCustomStorageInterface = CastChecked<INiagaraSimCacheCustomStorageInterface>(DataInterface);
 					const void* PerInstanceData = Helper.SystemInstance->FindDataInterfaceInstanceData(DataInterface);
 					if (!SimCacheCustomStorageInterface->SimCacheWriteFrame(StorageObject, FrameIndex, Helper.SystemInstance, PerInstanceData, FeedbackContext))
@@ -688,6 +702,7 @@ bool UNiagaraSimCache::ReadFrame(int32 FrameIndex, float FrameFraction, FNiagara
 	//-OPT: We shouldn't need to search all the time here
 	if (DataInterfaceStorage.IsEmpty() == false)
 	{
+		TSet<FNiagaraVariableBase> VisitedDataInterfaces;
 		bool bDataInterfacesSucess = true;
 
 		FNiagaraDataInterfaceUtilities::ForEachDataInterface(
@@ -698,6 +713,12 @@ bool UNiagaraSimCache::ReadFrame(int32 FrameIndex, float FrameFraction, FNiagara
 				{
 					if (INiagaraSimCacheCustomStorageInterface* SimCacheCustomStorageInterface = Cast<INiagaraSimCacheCustomStorageInterface>(DataInterface))
 					{
+						if (VisitedDataInterfaces.Contains(Variable))
+						{
+							return true;
+						}
+						VisitedDataInterfaces.Add(Variable);
+
 						void* PerInstanceData = Helper.SystemInstance->FindDataInterfaceInstanceData(DataInterface);
 						bDataInterfacesSucess &= SimCacheCustomStorageInterface->SimCacheReadFrame(StorageObject, FrameIndex, NextFrameIndex, FrameFraction, Helper.SystemInstance, PerInstanceData);
 					}
