Fixed a crash that could happen when undoing an emitter add transaction

diff --git a/Engine/Plugins/FX/Niagara/Source/NiagaraEditor/Private/ViewModels/NiagaraSystemViewModel.cpp b/Engine/Plugins/FX/Niagara/Source/NiagaraEditor/Private/ViewModels/NiagaraSystemViewModel.cpp
--- a/Engine/Plugins/FX/Niagara/Source/NiagaraEditor/Private/ViewModels/NiagaraSystemViewModel.cpp
+++ b/Engine/Plugins/FX/Niagara/Source/NiagaraEditor/Private/ViewModels/NiagaraSystemViewModel.cpp
@@ -2779,6 +2779,15 @@ void FNiagaraSystemViewModel::UpdateEmitterHandleSelectionFromSequencer()
 		return;
 	}
 	
+	for (TSharedRef<FNiagaraEmitterHandleViewModel> EmitterHandleViewModel : EmitterHandleViewModels)
+	{
+		if (!EmitterHandleViewModel->GetEmitterHandle() || ::IsValid(EmitterHandleViewModel->GetEmitterHandle()->GetInstance().Emitter) == false)
+		{
+			// we can run into pending kill emitters here when undoing an emitter add transaction. This method gets called as part of sequencer's PostUndo() refresh broadcast, which happens BEFORE this view model's PostUndo() could clean everything up.
+			return;
+		}
+	}
+	
 	TArray<FGuid> NewSelectedEmitterHandleIds;
 
 	TArray<UMovieSceneTrack*> SelectedTracks;
