- Add scope for NPC ticking

#rnx

diff --git a/Engine/Plugins/FX/Niagara/Source/Niagara/Private/NiagaraWorldManager.cpp b/Engine/Plugins/FX/Niagara/Source/Niagara/Private/NiagaraWorldManager.cpp
--- a/Engine/Plugins/FX/Niagara/Source/Niagara/Private/NiagaraWorldManager.cpp
+++ b/Engine/Plugins/FX/Niagara/Source/Niagara/Private/NiagaraWorldManager.cpp
@@ -37,6 +37,7 @@ DECLARE_CYCLE_STAT(TEXT("Niagara Manager Tick [GT]"), STAT_NiagaraWorldManTick,
 DECLARE_CYCLE_STAT(TEXT("Niagara Manager Wait On Render [GT]"), STAT_NiagaraWorldManWaitOnRender, STATGROUP_Niagara);
 DECLARE_CYCLE_STAT(TEXT("Niagara Manager Wait Pre Garbage Collect [GT]"), STAT_NiagaraWorldManWaitPreGC, STATGROUP_Niagara); 
 DECLARE_CYCLE_STAT(TEXT("Niagara Manager Refresh Owner Allows Scalability"), STAT_NiagaraWorldManRefreshOwnerAllowsScalability, STATGROUP_Niagara);
+DECLARE_CYCLE_STAT(TEXT("Niagara Manager Tick Parameter Collections [GT]"), STAT_NiagaraWorldManTickParamCollections, STATGROUP_Niagara);
 
 static int GNiagaraAllowAsyncWorkToEndOfFrame = 1;
 static FAutoConsoleVariableRef CVarNiagaraAllowAsyncWorkToEndOfFrame(
@@ -498,6 +499,18 @@ UNiagaraCullProxyComponent* FNiagaraWorldManager::GetCullProxy(UNiagaraComponent
 	return nullptr;
 }
 
+void FNiagaraWorldManager::TickParameterCollections()
+{
+	SCOPE_CYCLE_COUNTER(STAT_NiagaraWorldManTickParamCollections);
+
+	//-TODO: Do we need to do this per tick group?
+	for (TPair<UNiagaraParameterCollection*, UNiagaraParameterCollectionInstance*> CollectionInstPair : ParameterCollections)
+	{
+		check(CollectionInstPair.Value);
+		CollectionInstPair.Value->Tick(World);
+	}
+}
+
 UNiagaraParameterCollectionInstance* FNiagaraWorldManager::GetParameterCollection(UNiagaraParameterCollection* Collection)
 {
 	if (!Collection || bIsTearingDown)
@@ -1237,13 +1250,7 @@ void FNiagaraWorldManager::Tick(ETickingGroup TickGroup, float DeltaSeconds, ELe
 
 		UpdateScalabilityManagers(DeltaSeconds, false);
 
-		//Tick our collections to push any changes to bound stores.
-		//-TODO: Do we need to do this per tick group?
-		for (TPair<UNiagaraParameterCollection*, UNiagaraParameterCollectionInstance*> CollectionInstPair : ParameterCollections)
-		{
-			check(CollectionInstPair.Value);
-			CollectionInstPair.Value->Tick(World);
-		}
+		TickParameterCollections();
 	}
 
 	// If we are in paused don't do anything
diff --git a/Engine/Plugins/FX/Niagara/Source/Niagara/Public/NiagaraWorldManager.h b/Engine/Plugins/FX/Niagara/Source/Niagara/Public/NiagaraWorldManager.h
--- a/Engine/Plugins/FX/Niagara/Source/Niagara/Public/NiagaraWorldManager.h
+++ b/Engine/Plugins/FX/Niagara/Source/Niagara/Public/NiagaraWorldManager.h
@@ -103,9 +103,13 @@ public:
 	virtual void AddReferencedObjects(FReferenceCollector& Collector) override;
 	virtual FString GetReferencerName() const override;
 	//~ GCObject Interface
-	
+
+private:
+	void TickParameterCollections();
+public:
 	UNiagaraParameterCollectionInstance* GetParameterCollection(UNiagaraParameterCollection* Collection);
 	void CleanupParameterCollections();
+
 	FNiagaraSystemSimulationRef GetSystemSimulation(ETickingGroup TickGroup, UNiagaraSystem* System);
 	void DestroySystemSimulation(UNiagaraSystem* System);
 	void DestroySystemInstance(FNiagaraSystemInstancePtr& InPtr);	
