Fix incorrect transitions for the InstanceGPUTransformsBuffer buffer.

We can let RDG handle the state of this buffer initially, and extract it after the NiagaraGPURayTracingTransforms pass runs. The issue could be reproduced by running the BoolArray Niagara engine test with -rhivalidation.

#rnx

diff --git a/Engine/Plugins/FX/Niagara/Source/Niagara/Private/NiagaraRendererMeshes.cpp b/Engine/Plugins/FX/Niagara/Source/Niagara/Private/NiagaraRendererMeshes.cpp
--- a/Engine/Plugins/FX/Niagara/Source/Niagara/Private/NiagaraRendererMeshes.cpp
+++ b/Engine/Plugins/FX/Niagara/Source/Niagara/Private/NiagaraRendererMeshes.cpp
@@ -1691,9 +1691,6 @@ void FNiagaraRendererMeshes::GetDynamicRayTracingInstances(FRayTracingMaterialGa
 					FRDGBufferDesc::CreateStructuredDesc(4 * sizeof(float), 3 * NumInstances),
 					TEXT("InstanceGPUTransformsBuffer"));
 
-				const TRefCountPtr<FRDGPooledBuffer>& ExternalBuffer = Context.GraphBuilder.ConvertToExternalBuffer(InstanceGPUTransformsBufferRef);
-				RayTracingInstance.InstanceGPUTransformsSRV = ExternalBuffer->GetOrCreateSRV(FRHIBufferSRVCreateInfo());
-
 				const FLargeWorldRenderPosition AbsoluteViewOrigin(View->ViewMatrices.GetViewOrigin());
 				const FVector ViewTileOffset = AbsoluteViewOrigin.GetTileOffset();
 				const FVector RelativePreViewTranslation = View->ViewMatrices.GetPreViewTranslation() + ViewTileOffset;
@@ -1728,7 +1725,7 @@ void FNiagaraRendererMeshes::GetDynamicRayTracingInstances(FRayTracingMaterialGa
 					//PassParameters.ParticleDataHalfBuffer		= ParticleMeshRenderData.ParticleHalfSRV;
 					PassParameters->ParticleDataIntBuffer		= ParticleMeshRenderData.ParticleIntSRV;
 					PassParameters->GPUInstanceCountBuffer		= ComputeDispatchInterface->GetGPUInstanceCounterManager().GetInstanceCountBuffer().SRV;
-					PassParameters->TLASTransforms				= ExternalBuffer->GetOrCreateUAV(FRHIBufferUAVCreateInfo());
+					PassParameters->TLASTransforms				= Context.GraphBuilder.CreateUAV(FRDGBufferUAVDesc(InstanceGPUTransformsBufferRef));
 
 					PassParameters->ViewTilePosition			= ViewTilePosition;
 					PassParameters->RelativePreViewTranslation	= FVector3f(RelativePreViewTranslation);
@@ -1740,14 +1737,15 @@ void FNiagaraRendererMeshes::GetDynamicRayTracingInstances(FRayTracingMaterialGa
 					ERDGPassFlags::Compute,
 					[PassParameters, Pass_FeatureLevel = FeatureLevel, NumInstances](FRHICommandList& RHICmdList)
 				{
-					RHICmdList.Transition(FRHITransitionInfo(PassParameters->TLASTransforms, ERHIAccess::Unknown, ERHIAccess::UAVCompute));
-
 					FNiagaraGPURayTracingTransformsCS::FPermutationDomain PermutationVector;
 					TShaderMapRef<FNiagaraGPURayTracingTransformsCS> ComputeShader(GetGlobalShaderMap(Pass_FeatureLevel), PermutationVector);
 					FComputeShaderUtils::Dispatch(RHICmdList, ComputeShader, *PassParameters, FComputeShaderUtils::GetGroupCount(int32(NumInstances), int32(FNiagaraGPURayTracingTransformsCS::ThreadGroupSize)));
-
-					RHICmdList.Transition(FRHITransitionInfo(PassParameters->TLASTransforms, ERHIAccess::Unknown, ERHIAccess::SRVCompute));
 				});
+
+				Context.GraphBuilder.UseExternalAccessMode(InstanceGPUTransformsBufferRef, ERHIAccess::SRVMask);
+
+				const TRefCountPtr<FRDGPooledBuffer>& ExternalBuffer = Context.GraphBuilder.ConvertToExternalBuffer(InstanceGPUTransformsBufferRef);
+				RayTracingInstance.InstanceGPUTransformsSRV = ExternalBuffer->GetOrCreateSRV(FRHIBufferSRVCreateInfo());
 			}
 		}
 
diff --git a/Engine/Plugins/FX/Niagara/Source/NiagaraVertexFactories/Public/NiagaraGPURayTracingTransformsShader.h b/Engine/Plugins/FX/Niagara/Source/NiagaraVertexFactories/Public/NiagaraGPURayTracingTransformsShader.h
--- a/Engine/Plugins/FX/Niagara/Source/NiagaraVertexFactories/Public/NiagaraGPURayTracingTransformsShader.h
+++ b/Engine/Plugins/FX/Niagara/Source/NiagaraVertexFactories/Public/NiagaraGPURayTracingTransformsShader.h
@@ -62,7 +62,7 @@ class FNiagaraGPURayTracingTransformsCS : public FGlobalShader
 		SHADER_PARAMETER(FVector3f,	ViewTilePosition)
 		SHADER_PARAMETER(FVector3f, RelativePreViewTranslation)
 
-		SHADER_PARAMETER_UAV(RWStructuredBuffer<float3x4>, TLASTransforms)
+		SHADER_PARAMETER_RDG_BUFFER_UAV(RWStructuredBuffer<float3x4>, TLASTransforms)
 	END_SHADER_PARAMETER_STRUCT()
 
 };
