- Fix for custom depth collision
- We must derive the normal from custom depth rather than use the gbuffer normal

#rnx

diff --git a/Engine/Plugins/FX/Niagara/Shaders/Private/NiagaraDataInterfaceCollisionQuery.ush b/Engine/Plugins/FX/Niagara/Shaders/Private/NiagaraDataInterfaceCollisionQuery.ush
--- a/Engine/Plugins/FX/Niagara/Shaders/Private/NiagaraDataInterfaceCollisionQuery.ush
+++ b/Engine/Plugins/FX/Niagara/Shaders/Private/NiagaraDataInterfaceCollisionQuery.ush
@@ -168,15 +168,6 @@ void NDICollisionQuery_QueryCustomDepthGPU(in float3 In_SamplePos, in float3 In_
 		float SceneDepth = NDICollisionQuery_GetCustomSceneDepth(ScreenUV);
 		FLWCVector3 WorldPosition = WorldPositionFromSceneDepth(ScreenPosition.xy, SceneDepth);
 
-#if NDICOLLISIONQUERY_USE_GBUFFER_NORMAL
-#if STRATA_ENABLED
-		uint2 ScreenPos = uint2(ScreenUV * View.BufferSizeAndInvSize.xy);
-		float3 WorldNormal = StrataPublic_GetWorldNormal(ScreenPos);
-#else
-		float3 WorldNormal = Texture2DSampleLevel(SceneTexturesStruct.GBufferATexture, SceneTexturesStruct_GBufferATextureSampler, ScreenUV, 0).xyz * 2.0 - 1.0;
-#endif
-
-#else
 		float CollisionDepthBounds = 500.0f;
 		float SceneDepth0 = NDICollisionQuery_GetCustomSceneDepth(ScreenUV + float2(View.BufferSizeAndInvSize.z, 0.0));
 		float SceneDepth1 = NDICollisionQuery_GetCustomSceneDepth(ScreenUV + float2(0.0, View.BufferSizeAndInvSize.w));
@@ -192,7 +183,6 @@ void NDICollisionQuery_QueryCustomDepthGPU(in float3 In_SamplePos, in float3 In_
 		FLWCVector3 WorldPosition0 = WorldPositionFromSceneDepth(ScreenPosition.xy + float2(2 * View.ViewSizeAndInvSize.z, 0.0), SceneDepth0);
 		FLWCVector3 WorldPosition1 = WorldPositionFromSceneDepth(ScreenPosition.xy - float2(0.0, 2 * View.ViewSizeAndInvSize.w), SceneDepth1);
 		float3 WorldNormal = normalize(cross(LWCSubtract(WorldPosition0, WorldPosition).Offset, LWCSubtract(WorldPosition1, WorldPosition).Offset));
-#endif
 
 		// convert LWC position back to simulation space
 		LWCSetTile(WorldPosition, LWCGetTile(WorldPosition) - In_LWCTile);
