Various small fixes with staticswitch nodes
- Added reallocate pins when a staticswitch node references a different, preexisting script variable via renaming. This fixes a crash where the visuals expect the pins to be created under the new static switch variable
- Added 'true' argument to allow the deletion of a static switch variable when no static switch is referencing it anymore
- Removed 'GetNumberOfCases' virtual function in favor of 'GetOptionValues'. They both had overlapping purpose but sim target selectors didn't override GeNumberOfCases, which led to mismatch in data.
- Remove & Add Integer case buttons will now only add/remove enum style entries if the script variable is set to enum style. This makes it so that if you convert back to default view, and remove a few integer cases, you aren't also removing enum entries that other static switches might be using.
- 'remove' integer button for staticswitches now differentiates between enum style mode & default mode as these numbers are not necessarily synced

diff --git a/Engine/Plugins/FX/Niagara/Source/NiagaraEditor/Private/NiagaraGraph.cpp b/Engine/Plugins/FX/Niagara/Source/NiagaraEditor/Private/NiagaraGraph.cpp
--- a/Engine/Plugins/FX/Niagara/Source/NiagaraEditor/Private/NiagaraGraph.cpp
+++ b/Engine/Plugins/FX/Niagara/Source/NiagaraEditor/Private/NiagaraGraph.cpp
@@ -2851,7 +2851,7 @@ bool UNiagaraGraph::RenameStaticSwitch(UNiagaraNodeStaticSwitch* SwitchNode, FNa
 
 		if (!StaticSwitchInputs.Contains(OrigVariable))
 		{
-			RemoveParameter(OrigVariable);
+			RemoveParameter(OrigVariable, true);
 		}
 	}
 
diff --git a/Engine/Plugins/FX/Niagara/Source/NiagaraEditor/Private/NiagaraGraphDigest.cpp b/Engine/Plugins/FX/Niagara/Source/NiagaraEditor/Private/NiagaraGraphDigest.cpp
--- a/Engine/Plugins/FX/Niagara/Source/NiagaraEditor/Private/NiagaraGraphDigest.cpp
+++ b/Engine/Plugins/FX/Niagara/Source/NiagaraEditor/Private/NiagaraGraphDigest.cpp
@@ -3849,7 +3849,7 @@ FNiagaraCompilationNodeStaticSwitch::FNiagaraCompilationNodeStaticSwitch(const U
 	bSetByCompiler = InNode->IsSetByCompiler();
 	bSetByPin = InNode->IsSetByPin();
 	SwitchType = InNode->SwitchTypeData.SwitchType;
-	SwitchBranchCount = InNode->GetNumberOfCases();
+	SwitchBranchCount = InNode->GetOptionValues().Num();
 	InputParameterName = InNode->InputParameterName;
 	SwitchConstant = InNode->SwitchTypeData.SwitchConstant;
 	InputType = InNode->GetInputType();
@@ -4369,7 +4369,7 @@ FNiagaraCompilationNodeSelect::FNiagaraCompilationNodeSelect(const UNiagaraNodeS
 	, SelectorPinType(InNode->SelectorPinType)
 {
 	SelectorPinIndex = GetInputPinIndexByPersistentId(InNode->SelectorPinGuid);
-	NumOptionsPerVariable = InNode->GetNumberOfCases();
+	NumOptionsPerVariable = InNode->GetOptionValues().Num();
 	SelectorValues = InNode->GetOptionValues();
 }
 
diff --git a/Engine/Plugins/FX/Niagara/Source/NiagaraEditor/Private/NiagaraNodeSimTargetSelector.cpp b/Engine/Plugins/FX/Niagara/Source/NiagaraEditor/Private/NiagaraNodeSimTargetSelector.cpp
--- a/Engine/Plugins/FX/Niagara/Source/NiagaraEditor/Private/NiagaraNodeSimTargetSelector.cpp
+++ b/Engine/Plugins/FX/Niagara/Source/NiagaraEditor/Private/NiagaraNodeSimTargetSelector.cpp
@@ -47,6 +47,17 @@ FString UNiagaraNodeSimTargetSelector::GetInputCaseName(int32 Case) const
 	return Case == 0 ? TEXT("CPU VM") : TEXT("GPU Shader");
 }
 
+TArray<int32> UNiagaraNodeSimTargetSelector::GetOptionValues() const
+{
+	TArray<int32> OptionValues;
+	for(int32 Value = 0; Value < StaticEnum<ENiagaraSimTarget>()->NumEnums(); Value++)
+	{
+		OptionValues.Add(Value);
+	}
+	
+	return OptionValues; 
+}
+
 void UNiagaraNodeSimTargetSelector::Compile(FTranslator* Translator, TArray<int32>& Outputs) const
 {
 	FPinCollectorArray InputPins;
diff --git a/Engine/Plugins/FX/Niagara/Source/NiagaraEditor/Private/NiagaraNodeSimTargetSelector.h b/Engine/Plugins/FX/Niagara/Source/NiagaraEditor/Private/NiagaraNodeSimTargetSelector.h
--- a/Engine/Plugins/FX/Niagara/Source/NiagaraEditor/Private/NiagaraNodeSimTargetSelector.h
+++ b/Engine/Plugins/FX/Niagara/Source/NiagaraEditor/Private/NiagaraNodeSimTargetSelector.h
@@ -26,5 +26,6 @@ public:
 protected:
 	//~ Begin UNiagaraNodeUsageSelector Interface
 	virtual FString GetInputCaseName(int32 Case) const override;
+	virtual TArray<int32> GetOptionValues() const override;
 	//~ End UNiagaraNodeUsageSelector Interface
 };
diff --git a/Engine/Plugins/FX/Niagara/Source/NiagaraEditor/Private/NiagaraNodeStaticSwitch.cpp b/Engine/Plugins/FX/Niagara/Source/NiagaraEditor/Private/NiagaraNodeStaticSwitch.cpp
--- a/Engine/Plugins/FX/Niagara/Source/NiagaraEditor/Private/NiagaraNodeStaticSwitch.cpp
+++ b/Engine/Plugins/FX/Niagara/Source/NiagaraEditor/Private/NiagaraNodeStaticSwitch.cpp
@@ -214,21 +214,6 @@ void UNiagaraNodeStaticSwitch::PostEditChangeProperty(struct FPropertyChangedEve
 	}
 }
 
-int32 UNiagaraNodeStaticSwitch::GetNumberOfCases() const
-{
-	if(SwitchTypeData.SwitchType == ENiagaraStaticSwitchType::Integer)
-	{
-		if(UNiagaraScriptVariable* ScriptVariable = GetStaticSwitchScriptVariable())
-		{
-			if(ScriptVariable->Metadata.WidgetCustomization.WidgetType == ENiagaraInputWidgetType::EnumStyle)
-			{
-				return ScriptVariable->Metadata.WidgetCustomization.EnumStyleDropdownValues.Num();
-			}
-		}
-	}
-
-	return Super::GetNumberOfCases();
-}
 #endif
 
 TArray<int32> UNiagaraNodeStaticSwitch::GetOptionValues() const
@@ -240,17 +225,18 @@ TArray<int32> UNiagaraNodeStaticSwitch::GetOptionValues() const
 	}
 	else if (SwitchTypeData.SwitchType == ENiagaraStaticSwitchType::Integer)
 	{
-		// we want complete sync of the enum entries, if selected
-		bool bEnsureMinimumCases = true;
+		int32 EnumEntries = INDEX_NONE;
+		bool bUseEnumDisplay = false;
 		if(UNiagaraScriptVariable* ScriptVariable = GetStaticSwitchScriptVariable())
 		{
 			if(ScriptVariable->Metadata.WidgetCustomization.WidgetType == ENiagaraInputWidgetType::EnumStyle)
 			{
-				bEnsureMinimumCases = false;
+				EnumEntries = ScriptVariable->Metadata.WidgetCustomization.EnumStyleDropdownValues.Num();
+				bUseEnumDisplay = true;
 			}
 		}
 		
-		int32 NewOptionsCount = bEnsureMinimumCases ? FMath::Max(2, GetNumberOfCases()) : GetNumberOfCases();
+		int32 NewOptionsCount = bUseEnumDisplay ? EnumEntries : NumOptionsPerVariable;
 		for(int32 Counter = 0; Counter < NewOptionsCount; Counter++)
 		{
 			OptionValues.Add(Counter);
@@ -364,6 +350,8 @@ FName UNiagaraNodeStaticSwitch::GetOptionPinName(const FNiagaraVariable& Variabl
 void UNiagaraNodeStaticSwitch::ChangeSwitchParameterName(const FName& NewName)
 {
 	GetNiagaraGraph()->RenameStaticSwitch(this, NewName);
+	// due to renaming of the static switch, it is possible we are now pointing to another pre-existing static switch variable, so we need to reallocate pins.
+	ReallocatePins();
 	VisualsChangedDelegate.Broadcast(this);
 }
 
@@ -487,9 +475,7 @@ void UNiagaraNodeStaticSwitch::AllocateDefaultPins()
 	}
 	
 	TArray<int32> OptionValues = GetOptionValues();
-	NumOptionsPerVariable = OptionValues.Num();
-	
-	for (int32 i = 0; i < NumOptionsPerVariable; i++)
+	for (int32 i = 0; i < OptionValues.Num(); i++)
 	{
 		for (const FNiagaraVariable& Variable : OutputVars)
 		{
@@ -916,11 +902,15 @@ void UNiagaraNodeStaticSwitch::AddIntegerInputPin()
 	
 	if(UNiagaraScriptVariable* ScriptVariable = GetStaticSwitchScriptVariable())
 	{
-		ScriptVariable->PreEditChange(nullptr);
-		FNiagaraWidgetNamedIntegerInputValue& NewEntry = ScriptVariable->Metadata.WidgetCustomization.EnumStyleDropdownValues.AddDefaulted_GetRef();
-		NewEntry.DisplayName = FText::FromString(FString::FromInt(NumOptionsPerVariable - 1));
-		FPropertyChangedEvent PropertyChangedEvent(nullptr);
-		ScriptVariable->PostEditChangeProperty(PropertyChangedEvent);
+		// we only want to add a new inline enum entry if we are in enum style mode
+		if(ScriptVariable->Metadata.WidgetCustomization.WidgetType == ENiagaraInputWidgetType::EnumStyle)
+		{
+			ScriptVariable->PreEditChange(nullptr);
+			FNiagaraWidgetNamedIntegerInputValue& NewEntry = ScriptVariable->Metadata.WidgetCustomization.EnumStyleDropdownValues.AddDefaulted_GetRef();
+			NewEntry.DisplayName = FText::FromString(FString::FromInt(ScriptVariable->Metadata.WidgetCustomization.EnumStyleDropdownValues.Num() - 1));
+			FPropertyChangedEvent PropertyChangedEvent(nullptr);
+			ScriptVariable->PostEditChangeProperty(PropertyChangedEvent);
+		}
 	}
 	
 	ReallocatePins();
@@ -936,10 +926,14 @@ void UNiagaraNodeStaticSwitch::RemoveIntegerInputPin()
 	
 	if(UNiagaraScriptVariable* ScriptVariable = GetStaticSwitchScriptVariable())
 	{
-		ScriptVariable->PreEditChange(nullptr);
-		ScriptVariable->Metadata.WidgetCustomization.EnumStyleDropdownValues.RemoveAt(ScriptVariable->Metadata.WidgetCustomization.EnumStyleDropdownValues.Num() - 1);
-		FPropertyChangedEvent PropertyChangedEvent(nullptr);
-		ScriptVariable->PostEditChangeProperty(PropertyChangedEvent);
+		// we only want to remove an inline enum entry if we are in enum style mode
+		if(ScriptVariable->Metadata.WidgetCustomization.WidgetType == ENiagaraInputWidgetType::EnumStyle)
+		{
+			ScriptVariable->PreEditChange(nullptr);
+			ScriptVariable->Metadata.WidgetCustomization.EnumStyleDropdownValues.RemoveAt(ScriptVariable->Metadata.WidgetCustomization.EnumStyleDropdownValues.Num() - 1);
+			FPropertyChangedEvent PropertyChangedEvent(nullptr);
+			ScriptVariable->PostEditChangeProperty(PropertyChangedEvent);
+		}
 	}
 	
 	ReallocatePins();
@@ -962,7 +956,23 @@ EVisibility UNiagaraNodeStaticSwitch::ShowAddIntegerButton() const
 
 EVisibility UNiagaraNodeStaticSwitch::ShowRemoveIntegerButton() const
 {
-	return SwitchTypeData.SwitchType == ENiagaraStaticSwitchType::Integer && NumOptionsPerVariable > 2 ? EVisibility::Visible : EVisibility::Collapsed;
+	if(SwitchTypeData.SwitchType != ENiagaraStaticSwitchType::Integer)
+	{
+		return EVisibility::Collapsed;
+	}
+	
+	if(UNiagaraScriptVariable* ScriptVariable = GetStaticSwitchScriptVariable())
+	{
+		if(ScriptVariable->Metadata.WidgetCustomization.WidgetType == ENiagaraInputWidgetType::EnumStyle)
+		{
+			if(ScriptVariable->Metadata.WidgetCustomization.EnumStyleDropdownValues.Num() > 2)
+			{
+				return EVisibility::Visible;
+			}
+		}
+	}
+	
+	return NumOptionsPerVariable > 2 ? EVisibility::Visible : EVisibility::Collapsed;
 }
 
 void UNiagaraNodeStaticSwitch::PostLoad()
diff --git a/Engine/Plugins/FX/Niagara/Source/NiagaraEditor/Private/NiagaraNodeStaticSwitch.h b/Engine/Plugins/FX/Niagara/Source/NiagaraEditor/Private/NiagaraNodeStaticSwitch.h
--- a/Engine/Plugins/FX/Niagara/Source/NiagaraEditor/Private/NiagaraNodeStaticSwitch.h
+++ b/Engine/Plugins/FX/Niagara/Source/NiagaraEditor/Private/NiagaraNodeStaticSwitch.h
@@ -101,9 +101,7 @@ public:
 	virtual ENiagaraNumericOutputTypeSelectionMode GetNumericOutputTypeSelectionMode() const override;
 	//~ End UNiagaraNode Interface
 
-	//~ Begin UNiagaraNodeUsageSelector Interface
-	virtual int32 GetNumberOfCases() const override;
-	// ~End UNiagaraNodeUsageSelector Interface
+	virtual TArray<int32> GetOptionValues() const override;
 
 #if WITH_EDITOR
 	virtual void PostEditChangeProperty(struct FPropertyChangedEvent& PropertyChangedEvent) override;
@@ -123,7 +121,6 @@ protected:
 	
 	virtual FString GetInputCaseName(int32 Case) const override;
 	virtual FName GetOptionPinName(const FNiagaraVariable& Variable, int32 Value) const override;
-	virtual TArray<int32> GetOptionValues() const override;
 	//~ End UNiagaraNodeUsageSelector Interface
 
 	virtual bool CanModifyPin(const UEdGraphPin* Pin) const override;
diff --git a/Engine/Plugins/FX/Niagara/Source/NiagaraEditor/Private/NiagaraNodeUsageSelector.cpp b/Engine/Plugins/FX/Niagara/Source/NiagaraEditor/Private/NiagaraNodeUsageSelector.cpp
--- a/Engine/Plugins/FX/Niagara/Source/NiagaraEditor/Private/NiagaraNodeUsageSelector.cpp
+++ b/Engine/Plugins/FX/Niagara/Source/NiagaraEditor/Private/NiagaraNodeUsageSelector.cpp
@@ -260,8 +260,6 @@ void UNiagaraNodeUsageSelector::AddWidgetsToInputBox(TSharedPtr<SVerticalBox> In
 		return;
 	}
 
-	TArray<FString> Cases;
-
 	TArray<UEdGraphPin*> InputPins;
 	GetInputPins(InputPins);
 
@@ -273,8 +271,7 @@ void UNiagaraNodeUsageSelector::AddWidgetsToInputBox(TSharedPtr<SVerticalBox> In
 	int32 WidgetIndexOffset = 0;
 	// insert separators to make the grouping apparent.
 	TArray<int32> OptionValues = GetOptionValues();
-	ensure(OptionValues.Num() == GetNumberOfCases());
-	for (int32 Idx = 0; Idx < OutputVars.Num() * GetNumberOfCases(); Idx += OutputVars.Num())
+	for (int32 Idx = 0; Idx < OutputVars.Num() * OptionValues.Num(); Idx += OutputVars.Num())
 	{
 		const int32 OptionValueIndex = Idx / OutputVars.Num();
 		const int32 OptionValue = OptionValues[OptionValueIndex];
diff --git a/Engine/Plugins/FX/Niagara/Source/NiagaraEditor/Private/NiagaraNodeUsageSelector.h b/Engine/Plugins/FX/Niagara/Source/NiagaraEditor/Private/NiagaraNodeUsageSelector.h
--- a/Engine/Plugins/FX/Niagara/Source/NiagaraEditor/Private/NiagaraNodeUsageSelector.h
+++ b/Engine/Plugins/FX/Niagara/Source/NiagaraEditor/Private/NiagaraNodeUsageSelector.h
@@ -46,7 +46,6 @@ public:
 
 	/** Helper function to create a variable to add to the OutputVars and FGuid to add to OutputVarGuids. */
 	FGuid AddOutput(FNiagaraTypeDefinition Type, const FName& Name);
-	virtual int32 GetNumberOfCases() const { return NumOptionsPerVariable; }
 
 	void AttemptUpdatePins();
 	bool AreInputPinsOutdated() const;
@@ -84,6 +83,7 @@ protected:
 	//~ End UNiagaraNodeWithDynamicPins Interface
 
 protected:
+	/** Primarily used for integer options where the number of entries is not derivable from something else (i.e. enum or bool which clearly define entries) */
 	UPROPERTY()
 	int32 NumOptionsPerVariable = 0;
 private:
