Fixes for how source graphs are referenced and digested graphs are created for the AsyncTask compilation path for NiagaraSystems
-changed the digest database storage of graphs to be an LruCache
-updated the graph hash to not just be the graph's changeid but also to incorporate the changeid of referenced graphs (recursively)
-made the compilation task hold onto strong referencs of the shared data rather than indirect through the digest database
-cleaned up the handle interface to just a container for graph keys
#rnx

diff --git a/Engine/Plugins/FX/Niagara/Source/NiagaraEditor/Private/NiagaraCompilationCopy.cpp b/Engine/Plugins/FX/Niagara/Source/NiagaraEditor/Private/NiagaraCompilationCopy.cpp
--- a/Engine/Plugins/FX/Niagara/Source/NiagaraEditor/Private/NiagaraCompilationCopy.cpp
+++ b/Engine/Plugins/FX/Niagara/Source/NiagaraEditor/Private/NiagaraCompilationCopy.cpp
@@ -171,12 +171,11 @@ void FNiagaraPrecompileData::FinishPrecompile(
 	ENiagaraScriptCompileStatusEnum = StaticEnum<ENiagaraScriptCompileStatus>();
 	ENiagaraScriptUsageEnum = StaticEnum<ENiagaraScriptUsage>();
 
-	FNiagaraCompilationGraphHandle::FGraphPtr SourceGraph = SourceGraphHandle.Resolve();
 	TArray<const FNiagaraCompilationNodeOutput*> OutputNodes;
 
-	if (SourceGraph)
+	if (DigestedSourceGraph)
 	{
-		SourceGraph->FindOutputNodes(OutputNodes);
+		DigestedSourceGraph->FindOutputNodes(OutputNodes);
 	}
 
 	SortOutputNodesByDependencies(OutputNodes, SimStages);
@@ -246,7 +245,7 @@ void FNiagaraPrecompileData::FinishPrecompile(
 			// Map all for this output node
 			FParameterMapHistoryWithMetaDataBuilder Builder;
 			*Builder.ConstantResolver = ConstantResolver;
-			Builder.AddGraphToCallingGraphContextStack(SourceGraph.Get());
+			Builder.AddGraphToCallingGraphContextStack(DigestedSourceGraph.Get());
 			Builder.RegisterEncounterableVariables(EncounterableVariables);
 			Builder.RegisterExternalStaticVariables(StaticVariables);
 			CompilationTask.GetAvailableCollections(Builder.AvailableCollections->EditCollections());
@@ -499,9 +498,9 @@ UNiagaraDataInterface* FNiagaraCompilationCopyData::GetDuplicatedDataInterfaceCD
 	return nullptr;
 }
 
-void FNiagaraCompilationCopyData::InstantiateCompilationCopy(FNiagaraCompilationGraphHandle GraphHandle, const FNiagaraPrecompileData* PrecompileData, ENiagaraScriptUsage InUsage, FNiagaraFixedConstantResolver ConstantResolver)
+void FNiagaraCompilationCopyData::InstantiateCompilationCopy(const FNiagaraCompilationGraph& SourceGraph, const FNiagaraPrecompileData* PrecompileData, ENiagaraScriptUsage InUsage, FNiagaraFixedConstantResolver ConstantResolver)
 {
-	InstantiatedGraph = GraphHandle.Resolve()->Instantiate(PrecompileData, this, ValidUsages, ConstantResolver);
+	InstantiatedGraph = SourceGraph.Instantiate(PrecompileData, this, ValidUsages, ConstantResolver);
 
 	if (InstantiatedGraph)
 	{
diff --git a/Engine/Plugins/FX/Niagara/Source/NiagaraEditor/Private/NiagaraCompilationPrivate.h b/Engine/Plugins/FX/Niagara/Source/NiagaraEditor/Private/NiagaraCompilationPrivate.h
--- a/Engine/Plugins/FX/Niagara/Source/NiagaraEditor/Private/NiagaraCompilationPrivate.h
+++ b/Engine/Plugins/FX/Niagara/Source/NiagaraEditor/Private/NiagaraCompilationPrivate.h
@@ -139,7 +139,7 @@ public:
 	TArray<FNiagaraVariable> EncounteredVariables;
 	FString EmitterUniqueName;
 	TArray<FSharedPrecompileData> EmitterData;
-	FNiagaraCompilationGraphHandle SourceGraphHandle;
+	FNiagaraDigestedGraphPtr DigestedSourceGraph;
 	FString SourceName;
 	bool bUseRapidIterationParams = true;
 	bool bDisableDebugSwitches = false;
@@ -176,7 +176,7 @@ public:
 	TArray<FParameterMapHistory>& GetPrecomputedHistories() { return PrecompiledHistories; }
 	const TArray<FParameterMapHistory>& GetPrecomputedHistories() const { return PrecompiledHistories; }
 
-	void InstantiateCompilationCopy(FNiagaraCompilationGraphHandle GraphHandle, const FNiagaraPrecompileData* PrecompileData, ENiagaraScriptUsage InUsage, FNiagaraFixedConstantResolver ConstantResolver);
+	void InstantiateCompilationCopy(const FNiagaraCompilationGraph& SourceGraph, const FNiagaraPrecompileData* PrecompileData, ENiagaraScriptUsage InUsage, FNiagaraFixedConstantResolver ConstantResolver);
 	void CreateParameterMapHistory(const FNiagaraSystemCompilationTask& CompilationTask, const TArray<FNiagaraVariable>& EncounterableVariables, const TArray<FNiagaraVariable>& InStaticVariables, FNiagaraFixedConstantResolver ConstantResolver, TConstArrayView<FNiagaraSimulationStageInfo> SimStages);
 
 	int32 GetDependentRequestCount() const { return EmitterData.Num(); }
diff --git a/Engine/Plugins/FX/Niagara/Source/NiagaraEditor/Private/NiagaraCompilationTasks.cpp b/Engine/Plugins/FX/Niagara/Source/NiagaraEditor/Private/NiagaraCompilationTasks.cpp
--- a/Engine/Plugins/FX/Niagara/Source/NiagaraEditor/Private/NiagaraCompilationTasks.cpp
+++ b/Engine/Plugins/FX/Niagara/Source/NiagaraEditor/Private/NiagaraCompilationTasks.cpp
@@ -311,6 +311,24 @@ namespace NiagaraCompilationCopyImpl
 			Output.AddUnique(InputValue);
 		}
 	}
+
+	const UNiagaraGraph* GetGraphFromScriptSource(const UNiagaraScriptSourceBase* ScriptSourceBase)
+	{
+		if (const UNiagaraScriptSource* ScriptSource = CastChecked<const UNiagaraScriptSource>(ScriptSourceBase))
+		{
+			return ScriptSource->NodeGraph;
+		}
+		return nullptr;
+	}
+
+	const UNiagaraGraph* GetGraphFromScript(const UNiagaraScript* Script)
+	{
+		if (Script)
+		{
+			return GetGraphFromScriptSource(Script->GetLatestSource());
+		}
+		return nullptr;
+	}
 } // NiagaraCompilationCopyImpl
 
 FNiagaraSystemCompilationTask::FNiagaraSystemCompilationTask(FNiagaraCompilationTaskHandle InTaskHandle, UNiagaraSystem* InSystem)
@@ -383,7 +401,10 @@ void FNiagaraSystemCompilationTask::FCompileGroupInfo::InstantiateCompileGraph(c
 		{
 			const FNiagaraPrecompileData* EmitterRequestData = static_cast<const FNiagaraPrecompileData*>(ParentTask.SystemPrecompileData->GetDependentRequest(EmitterIt).Get());
 
-			EmitterPtr->InstantiateCompilationCopy(EmitterInfo.SourceGraphHandle, EmitterRequestData, ENiagaraScriptUsage::EmitterSpawnScript, EmitterInfo.ConstantResolver);
+			if (const FNiagaraCompilationGraph* SourceGraph = EmitterInfo.SourceGraph.Get())
+			{
+				EmitterPtr->InstantiateCompilationCopy(*SourceGraph, EmitterRequestData, ENiagaraScriptUsage::EmitterSpawnScript, EmitterInfo.ConstantResolver);
+			}
 		}
 		EmitterPtr->ValidUsages = BasePtr.ValidUsages;
 		BasePtr.EmitterData.Add(EmitterPtr);
@@ -397,7 +418,11 @@ void FNiagaraSystemCompilationTask::FCompileGroupInfo::InstantiateCompileGraph(c
 		// skip the deep copy if we're not compiling the system scripts
 		if (BasePtr.ValidUsages.Contains(ENiagaraScriptUsage::SystemSpawnScript))
 		{
-			BasePtr.InstantiateCompilationCopy(ParentTask.SystemInfo.SystemSourceGraphHandle, ParentTask.SystemPrecompileData.Get(), ENiagaraScriptUsage::SystemSpawnScript, ParentTask.SystemInfo.ConstantResolver);
+			if (const FNiagaraCompilationGraph* SourceGraph = ParentTask.SystemInfo.SystemSourceGraph.Get())
+			{
+				BasePtr.InstantiateCompilationCopy(*SourceGraph, ParentTask.SystemPrecompileData.Get(), ENiagaraScriptUsage::SystemSpawnScript, ParentTask.SystemInfo.ConstantResolver);
+			}
+
 			BasePtr.CreateParameterMapHistory(ParentTask, EncounterableSystemVariables, ParentTask.SystemInfo.StaticVariableResults, ParentTask.SystemInfo.ConstantResolver, {});
 
 			// bubble up the referenced DI classes to the system
@@ -806,14 +831,26 @@ void FNiagaraSystemCompilationTask::FDispatchDataCachePutRequests::Launch(FNiaga
 
 void FNiagaraSystemCompilationTask::DigestSystemInfo()
 {
+	using namespace NiagaraCompilationCopyImpl;
+
+	FNiagaraDigestDatabase& DigestDatabase = FNiagaraDigestDatabase::Get();
+	const UNiagaraScript* SystemSpawnScript = System_GT->GetSystemSpawnScript();
+	const UNiagaraScript* SystemUpdateScript = System_GT->GetSystemUpdateScript();
+
+	const UNiagaraGraph* SystemGraph = GetGraphFromScript(SystemSpawnScript);
+	check(SystemGraph == GetGraphFromScript(SystemUpdateScript));
+
+	FNiagaraGraphChangeIdBuilder ChangeIdBuilder;
+	ChangeIdBuilder.ParseReferencedGraphs(SystemGraph);
+
 	SystemInfo.SystemName = System_GT->GetName();
 	SystemInfo.SystemPackageName = System_GT->GetOutermost()->GetFName();
 	SystemInfo.bUseRapidIterationParams = System_GT->ShouldUseRapidIterationParameters();
 	SystemInfo.bDisableDebugSwitches = System_GT->ShouldDisableDebugSwitches();
-	SystemInfo.SystemSourceGraphHandle = FNiagaraCompilationGraphHandle(System_GT->GetSystemSpawnScript()->GetLatestSource());
 	SystemInfo.ConstantResolver = FNiagaraFixedConstantResolver(FCompileConstantResolver(System_GT.Get(), ENiagaraScriptUsage::SystemSpawnScript));
-	SystemInfo.OwnedScriptKeys = { System_GT->GetSystemSpawnScript(), System_GT->GetSystemUpdateScript() };
-
+	SystemInfo.OwnedScriptKeys = { SystemSpawnScript, SystemUpdateScript };
+	SystemInfo.SystemSourceGraph = DigestDatabase.CreateGraphDigest(SystemGraph, ChangeIdBuilder);
+	
 	{
 		TArray<FNiagaraVariable> StaticVariablesFromEmitters;
 
@@ -835,12 +872,15 @@ void FNiagaraSystemCompilationTask::DigestSystemInfo()
 		FEmitterInfo& EmitterInfo = SystemInfo.EmitterInfo.AddDefaulted_GetRef();
 		const FVersionedNiagaraEmitter& HandleInstance = Handle.GetInstance();
 
+		const UNiagaraGraph* EmitterGraph = GetGraphFromScriptSource(Handle.GetEmitterData()->GraphSource);
+		ChangeIdBuilder.ParseReferencedGraphs(EmitterGraph);
+
 		EmitterInfo.UniqueEmitterName = HandleInstance.Emitter->GetUniqueEmitterName();
-		EmitterInfo.SourceGraphHandle = FNiagaraCompilationGraphHandle(Handle.GetEmitterData()->GraphSource);
 		EmitterInfo.UniqueInstanceName = Handle.GetUniqueInstanceName();
 		EmitterInfo.Enabled = Handle.GetIsEnabled();
 		EmitterInfo.ConstantResolver = FNiagaraFixedConstantResolver(FCompileConstantResolver(HandleInstance, ENiagaraScriptUsage::EmitterSpawnScript));
 		EmitterInfo.EmitterIndex = EmitterIndex;
+		EmitterInfo.SourceGraph = DigestDatabase.CreateGraphDigest(EmitterGraph, ChangeIdBuilder);
 
 		{
 			constexpr bool bCompilableOnly = false;
@@ -895,25 +935,6 @@ void FNiagaraSystemCompilationTask::DigestSystemInfo()
 	System_GT->GetExposedParameters().GetParameters(SystemInfo.OriginalExposedParams);
 }
 
-void FNiagaraSystemCompilationTask::DigestGraphs(TConstArrayView<UNiagaraScript*> Scripts)
-{
-	const FGuid ScriptVersion = FGuid();
-	DigestedScriptGraphs.Reserve(Scripts.Num());
-	for (UNiagaraScript* Script : Scripts)
-	{
-		if (Script)
-		{
-			if (const UNiagaraScriptSource* ScriptSource = Cast<const UNiagaraScriptSource>(Script->GetSource(ScriptVersion)))
-			{
-				if (const UNiagaraGraph* Graph = ScriptSource->NodeGraph)
-				{
-					DigestedScriptGraphs.Add(Script, FNiagaraDigestDatabase::Get().CreateCompilationCopy(Graph));
-				}
-			}
-		}
-	}
-}
-
 void FNiagaraSystemCompilationTask::DigestParameterCollections(TConstArrayView<TWeakObjectPtr<UNiagaraParameterCollection>> Collections)
 {
 	DigestedParameterCollections.Reserve(Collections.Num());
@@ -1154,7 +1175,7 @@ void FNiagaraSystemCompilationTask::Precompile()
 
 	SystemPrecompileData->SourceName = SystemInfo.SystemName;
 
-	SystemPrecompileData->SourceGraphHandle = SystemInfo.SystemSourceGraphHandle;
+	SystemPrecompileData->DigestedSourceGraph = SystemInfo.SystemSourceGraph;
 	SystemPrecompileData->bUseRapidIterationParams = SystemInfo.bUseRapidIterationParams;
 	SystemPrecompileData->bDisableDebugSwitches = SystemInfo.bDisableDebugSwitches;
 
@@ -1171,7 +1192,7 @@ void FNiagaraSystemCompilationTask::Precompile()
 		TSharedPtr<FNiagaraPrecompileData, ESPMode::ThreadSafe> EmitterPtr = MakeShared<FNiagaraPrecompileData, ESPMode::ThreadSafe>();
 		EmitterPtr->EmitterUniqueName = EmitterInfo.UniqueEmitterName;
 		EmitterPtr->SourceName = SystemPrecompileData->SourceName;
-		EmitterPtr->SourceGraphHandle = EmitterInfo.SourceGraphHandle;
+		EmitterPtr->DigestedSourceGraph = EmitterInfo.SourceGraph;
 		EmitterPtr->bUseRapidIterationParams = SystemPrecompileData->bUseRapidIterationParams;
 		EmitterPtr->bDisableDebugSwitches = SystemPrecompileData->bDisableDebugSwitches;
 		EmitterPtr->SharedCompileDataInterfaceData = SystemPrecompileData->SharedCompileDataInterfaceData;
@@ -1289,7 +1310,7 @@ struct FNiagaraSystemCompilationTask::FCollectStaticVariablesTaskBuilder
 	FCollectStaticVariablesTaskBuilder(FNiagaraSystemCompilationTask& CompilationTask)
 	: FoundStaticVariables(CompilationTask.SystemInfo.StaticVariableResults)
 	{
-		if (FNiagaraCompilationGraph* SystemGraph = CompilationTask.SystemInfo.SystemSourceGraphHandle.Resolve().Get())
+		if (FNiagaraCompilationGraph* SystemGraph = CompilationTask.SystemInfo.SystemSourceGraph.Get())
 		{
 			FStaticVariableBuilderTaskHandle BuilderTask = MakeShared<FStaticVariableBuilderTask, ESPMode::ThreadSafe>();
 			BuilderTask->CompilationTask = &CompilationTask;
@@ -1306,7 +1327,7 @@ struct FNiagaraSystemCompilationTask::FCollectStaticVariablesTaskBuilder
 	FCollectStaticVariablesTaskBuilder(const FNiagaraSystemCompilationTask& CompilationTask, FEmitterInfo& EmitterInfo)
 	: FoundStaticVariables(EmitterInfo.StaticVariableResults)
 	{
-		if (FNiagaraCompilationGraph* EmitterGraph = EmitterInfo.SourceGraphHandle.Resolve().Get())
+		if (FNiagaraCompilationGraph* EmitterGraph = EmitterInfo.SourceGraph.Get())
 		{
 			TArray<const FNiagaraCompilationNodeOutput*> OutputNodes;
 			EmitterGraph->FindOutputNodes(OutputNodes);
@@ -1430,6 +1451,7 @@ struct FNiagaraSystemCompilationTask::FCollectStaticVariablesTaskBuilder
 						BuilderTask->CompilationTask->GetAvailableCollections(Builder.AvailableCollections->EditCollections());
 						Builder.BeginTranslation(FNiagaraConstants::EmitterNamespaceString);
 						Builder.EnableScriptAllowList(true, OutputNode->Usage);
+						Builder.IncludeStaticVariablesOnly();
 						Builder.BeginUsage(OutputNode->Usage, BuilderTask->StageNames[OutputNodeIt]);
 						Builder.BuildParameterMaps(OutputNode, true);
 						Builder.EndUsage();
@@ -1483,7 +1505,7 @@ struct FNiagaraSystemCompilationTask::FBuildRapidIterationTaskBuilder
 	
 		if (ScriptInfo.EmitterIndex == INDEX_NONE)
 		{
-			Graph = CompilationTask.SystemInfo.SystemSourceGraphHandle.Resolve().Get();
+			Graph = CompilationTask.SystemInfo.SystemSourceGraph.Get();
 			FoundStaticVariables = &CompilationTask.SystemInfo.StaticVariableResults;
 			ConstantResolver = CompilationTask.SystemInfo.ConstantResolver;
 		}
@@ -1491,7 +1513,7 @@ struct FNiagaraSystemCompilationTask::FBuildRapidIterationTaskBuilder
 		{
 			FEmitterInfo& EmitterInfo = CompilationTask.SystemInfo.EmitterInfo[ScriptInfo.EmitterIndex];
 
-			Graph = EmitterInfo.SourceGraphHandle.Resolve().Get();
+			Graph = EmitterInfo.SourceGraph.Get();
 			FoundStaticVariables = &EmitterInfo.StaticVariableResults;
 			UniqueEmitterName = EmitterInfo.UniqueEmitterName;
 			ConstantResolver = EmitterInfo.ConstantResolver;
diff --git a/Engine/Plugins/FX/Niagara/Source/NiagaraEditor/Private/NiagaraCompilationTasks.h b/Engine/Plugins/FX/Niagara/Source/NiagaraEditor/Private/NiagaraCompilationTasks.h
--- a/Engine/Plugins/FX/Niagara/Source/NiagaraEditor/Private/NiagaraCompilationTasks.h
+++ b/Engine/Plugins/FX/Niagara/Source/NiagaraEditor/Private/NiagaraCompilationTasks.h
@@ -31,7 +31,6 @@ struct FNiagaraSystemCompilationTask
 	void Abort();
 
 	void DigestSystemInfo();
-	void DigestGraphs(TConstArrayView<UNiagaraScript*> Scripts);
 	void DigestParameterCollections(TConstArrayView<TWeakObjectPtr<UNiagaraParameterCollection>> Collections);
 	void AddScript(int32 EmitterIndex, UNiagaraScript* ScriptToCompile, bool bRequiresCompilation);
 	UE::Tasks::FTask BeginTasks();
@@ -62,7 +61,7 @@ struct FNiagaraSystemCompilationTask
 	{
 		FString UniqueEmitterName;
 		FString UniqueInstanceName;
-		FNiagaraCompilationGraphHandle SourceGraphHandle;
+		FNiagaraDigestedGraphPtr SourceGraph;
 		TArray<FNiagaraVariable> InitialStaticVariables;
 		TArray<FNiagaraVariable> StaticVariableResults;
 		TArray<FNiagaraSimulationStageInfo> SimStages;
@@ -82,7 +81,7 @@ struct FNiagaraSystemCompilationTask
 		TArray<FNiagaraVariable> InitialStaticVariables;
 		TArray<FNiagaraVariable> StaticVariableResults;
 		TArray<FNiagaraVariable> OriginalExposedParams;
-		FNiagaraCompilationGraphHandle SystemSourceGraphHandle;
+		FNiagaraDigestedGraphPtr SystemSourceGraph;
 		FNiagaraFixedConstantResolver ConstantResolver;
 		TArray<TObjectKey<UNiagaraScript>> OwnedScriptKeys;
 
@@ -191,8 +190,8 @@ private:
 	TUniquePtr<FDispatchDataCachePutRequests> PutRequestHelper;
 
 	TMap<TObjectKey<UNiagaraScript>, FScriptInfo> DigestedScriptInfo;
+	TSet<FNiagaraDigestedGraphPtr> DigestedGraphs;
 
-	TMap<TObjectKey<UNiagaraScript>, FNiagaraCompilationGraphHandle> DigestedScriptGraphs;
 	TMap<TObjectKey<UNiagaraParameterCollection>, FNiagaraCompilationNPCHandle> DigestedParameterCollections;
 
 	UE::Tasks::FTaskEvent CompileCompletionEvent = UE::Tasks::FTaskEvent(UE_SOURCE_LOCATION);
diff --git a/Engine/Plugins/FX/Niagara/Source/NiagaraEditor/Private/NiagaraDigestDatabase.cpp b/Engine/Plugins/FX/Niagara/Source/NiagaraEditor/Private/NiagaraDigestDatabase.cpp
--- a/Engine/Plugins/FX/Niagara/Source/NiagaraEditor/Private/NiagaraDigestDatabase.cpp
+++ b/Engine/Plugins/FX/Niagara/Source/NiagaraEditor/Private/NiagaraDigestDatabase.cpp
@@ -14,22 +14,13 @@
 namespace NiagaraGraphDigestDatabaseImpl
 {
 
-// todo - two tiered cache for digested objects (see NiagaraDigestDatabase.h)
-//static int GDigestGraphCacheSize = 256;
-//static FAutoConsoleVariableRef CVarDigestGraphCacheSize(
-//	TEXT("fx.Niagara.DigestGraphCacheSize"),
-//	GDigestGraphCacheSize,
-//	TEXT("Defines the size of the cache for digested Niagara graphs."),
-//	ECVF_ReadOnly
-//);
-//
-//static int GDigestCollectionCacheSize = 32;
-//static FAutoConsoleVariableRef CVarDigestCollectionCacheSize(
-//	TEXT("fx.Niagara.DigestCollectionCacheSize"),
-//	GDigestCollectionCacheSize,
-//	TEXT("Defines the size of the cache for digested Niagara parameter collections."),
-//	ECVF_ReadOnly
-//);
+static int32 GDigestGraphCacheSize = 512;
+static FAutoConsoleVariableRef CVarDigestGraphCacheSize(
+	TEXT("fx.Niagara.DigestGraphCacheSize"),
+	GDigestGraphCacheSize,
+	TEXT("Defines the size of the cache for digested Niagara graphs."),
+	ECVF_ReadOnly
+);
 
 FNiagaraDigestDatabase* GDigestDatabase = nullptr;
 
@@ -51,9 +42,7 @@ void FNiagaraDigestDatabase::Shutdown()
 }
 
 FNiagaraDigestDatabase::FNiagaraDigestDatabase()
-	// todo - two tiered cache for digested objects (see NiagaraDigestDatabase.h)
-	//: CompilationGraphCache(NiagaraGraphDigestDatabaseImpl::GDigestGraphCacheSize)
-	//, CompilationNPCCache(NiagaraGraphDigestDatabaseImpl::GDigestCollectionCacheSize)
+	: CompilationGraphCache(NiagaraGraphDigestDatabaseImpl::GDigestGraphCacheSize)
 {
 
 }
@@ -95,71 +84,42 @@ FString FNiagaraDigestDatabase::GetReferencerName() const
 //////////////////////////////////////////////////////////////////////////
 /// Digested graphs
 
-FNiagaraCompilationGraphHandle FNiagaraDigestDatabase::CreateCompilationCopy(const UNiagaraGraph* Graph)
+FNiagaraDigestedGraphPtr FNiagaraDigestDatabase::CreateGraphDigest(const UNiagaraGraph* Graph, const FNiagaraGraphChangeIdBuilder& Digester)
 {
 	check(IsInGameThread());
-	FNiagaraCompilationGraphHandle GraphHash(Graph);
-	FNiagaraCompilationGraph* PendingGraph = nullptr;
+	FNiagaraCompilationGraphHandle GraphHash(Graph, Digester);
+	FNiagaraDigestedGraphPtr PendingGraph;
 
 	{
 		FWriteScopeLock WriteScope(DigestCacheLock);
 
-		if (const FNiagaraCompilationGraphHandle::FGraphPtr* CompilationGraph = CompilationGraphCache.Find(GraphHash))
+		if (const FNiagaraDigestedGraphPtr* CompilationGraph = CompilationGraphCache.FindAndTouch(GraphHash))
 		{
 			++GraphCacheHits;
-			return GraphHash;
+			return *CompilationGraph;
 		}
 
 		++GraphCacheMisses;
-		FNiagaraCompilationGraphHandle::FGraphPtr CompilationGraph = MakeShared<FNiagaraCompilationGraphHandle::FGraphPtr::ElementType, FNiagaraCompilationGraphHandle::FGraphPtr::Mode>();
-		PendingGraph = CompilationGraph.Get();
+		PendingGraph = MakeShared<FNiagaraDigestedGraphPtr::ElementType, FNiagaraDigestedGraphPtr::Mode>();
 
-		CompilationGraphCache.Add(GraphHash, CompilationGraph);
+		CompilationGraphCache.Add(GraphHash, PendingGraph);
 	}
 
 	if (PendingGraph)
 	{
-		PendingGraph->Create(Graph);
+		PendingGraph->Create(Graph, Digester);
 	}
 
-	return GraphHash;
+	return PendingGraph;
 }
 
-FNiagaraCompilationGraphHandle::FGraphPtr FNiagaraDigestDatabase::Resolve(const FNiagaraCompilationGraphHandle& Handle) const
-{
-	FNiagaraCompilationGraphHandle::FGraphPtr GraphRef;
-
-	{
-		FReadScopeLock ReadScope(DigestCacheLock);
-		GraphRef = CompilationGraphCache.FindRef(Handle);
-	}
-
-	check(GraphRef.IsValid());
 
-	return GraphRef;
-}
-
-FNiagaraCompilationGraphHandle::FNiagaraCompilationGraphHandle(const UNiagaraScriptSourceBase* ScriptSourceBase)
-{
-	if (const UNiagaraScriptSource* ScriptSource = CastChecked<const UNiagaraScriptSource>(ScriptSourceBase))
-	{
-		if (const UNiagaraGraph* Graph = ScriptSource->NodeGraph)
-		{
-			AssetKey = Graph;
-			Hash = Graph->GetChangeID();
-		}
-	}
-}
-
-FNiagaraCompilationGraphHandle::FNiagaraCompilationGraphHandle(const UNiagaraGraph* Graph)
+FNiagaraCompilationGraphHandle::FNiagaraCompilationGraphHandle(const UNiagaraGraph* Graph, const FNiagaraGraphChangeIdBuilder& Builder)
 {
 	AssetKey = Graph;
-	Hash = Graph ? Graph->GetChangeID() : FGuid();
-}
 
-FNiagaraCompilationGraphHandle::FGraphPtr FNiagaraCompilationGraphHandle::Resolve() const
-{
-	return FNiagaraDigestDatabase::Get().Resolve(*this);
+	// generate the hash key based on the provided ChangeIdBuilder
+	Hash = Builder.FindChangeId(Graph);
 }
 
 //////////////////////////////////////////////////////////////////////////
diff --git a/Engine/Plugins/FX/Niagara/Source/NiagaraEditor/Private/NiagaraDigestDatabase.h b/Engine/Plugins/FX/Niagara/Source/NiagaraEditor/Private/NiagaraDigestDatabase.h
--- a/Engine/Plugins/FX/Niagara/Source/NiagaraEditor/Private/NiagaraDigestDatabase.h
+++ b/Engine/Plugins/FX/Niagara/Source/NiagaraEditor/Private/NiagaraDigestDatabase.h
@@ -3,6 +3,7 @@
 #pragma once
 
 #include "CoreMinimal.h"
+#include "Containers/LruCache.h"
 #include "NiagaraCompileHash.h"
 #include "NiagaraGraphDigest.h"
 #include "UObject/GCObject.h"
@@ -12,50 +13,29 @@ class UNiagaraGraph;
 class UNiagaraParameterCollection;
 class UNiagaraScriptSourceBase;
 
-// base class used to reference a digested version of a UObject.  
-template<typename AssetType, typename DigestType, typename HashType>
-class FNiagaraDigestedAssetHandle
+using FNiagaraDigestedGraphPtr = TSharedPtr<FNiagaraCompilationGraph, ESPMode::ThreadSafe>;
+
+class FNiagaraCompilationGraphHandle
 {
 public:
-	using FDigestPtr = TSharedPtr<DigestType, ESPMode::ThreadSafe>;
-	FNiagaraDigestedAssetHandle() = default;
-	virtual ~FNiagaraDigestedAssetHandle() = default;
-	FNiagaraDigestedAssetHandle(const FNiagaraDigestedAssetHandle& Handle) = default;
-
-	bool IsValid() const
-	{
-		return Resolve().IsValid();
-	}
-	virtual FDigestPtr Resolve() const = 0;
+	FNiagaraCompilationGraphHandle() = default;
+	FNiagaraCompilationGraphHandle(const UNiagaraGraph* Graph, const FNiagaraGraphChangeIdBuilder& ChangeIdBuilder);
 
 protected:
-	TObjectKey<AssetType> AssetKey;
-	HashType Hash;
+	TObjectKey<UNiagaraGraph> AssetKey;
+	FGuid Hash;
 
-	friend FORCEINLINE uint32 GetTypeHash(const FNiagaraDigestedAssetHandle& CompilationCopy)
+	friend FORCEINLINE uint32 GetTypeHash(const FNiagaraCompilationGraphHandle& CompilationCopy)
 	{
 		return HashCombine(GetTypeHash(CompilationCopy.AssetKey), GetTypeHash(CompilationCopy.Hash));
 	}
 
-	friend FORCEINLINE bool operator==(const FNiagaraDigestedAssetHandle& Lhs, const FNiagaraDigestedAssetHandle& Rhs)
+	friend FORCEINLINE bool operator==(const FNiagaraCompilationGraphHandle& Lhs, const FNiagaraCompilationGraphHandle& Rhs)
 	{
 		return Lhs.AssetKey == Rhs.AssetKey && Lhs.Hash == Rhs.Hash;
 	}
 };
 
-class FNiagaraCompilationGraphHandle : public FNiagaraDigestedAssetHandle<UNiagaraGraph, FNiagaraCompilationGraph, FGuid>
-{
-public:
-	using Super = FNiagaraDigestedAssetHandle<UNiagaraGraph, FNiagaraCompilationGraph, FGuid>;
-	using FGraphPtr = Super::FDigestPtr;
-
-	FNiagaraCompilationGraphHandle() = default;
-	FNiagaraCompilationGraphHandle(const UNiagaraScriptSourceBase* ScriptSource);
-	FNiagaraCompilationGraphHandle(const UNiagaraGraph* Graph);
-
-	virtual FGraphPtr Resolve() const override;
-};
-
 class FNiagaraCompilationNPC
 {
 public:
@@ -72,18 +52,37 @@ public:
 	TMap<FNiagaraVariableBase, UNiagaraDataInterface*> DefaultDataInterfaces;
 };
 
-class FNiagaraCompilationNPCHandle : public FNiagaraDigestedAssetHandle<UNiagaraParameterCollection, FNiagaraCompilationNPC, FNiagaraCompileHash>
+class FNiagaraCompilationNPCHandle
 {
 public:
-	using Super = FNiagaraDigestedAssetHandle<UNiagaraParameterCollection, FNiagaraCompilationNPC, FNiagaraCompileHash>;
+	using FDigestPtr = TSharedPtr<FNiagaraCompilationNPC, ESPMode::ThreadSafe>;
 
 	FNiagaraCompilationNPCHandle() = default;
 	FNiagaraCompilationNPCHandle(const FNiagaraCompilationNPCHandle& Handle) = default;
 	FNiagaraCompilationNPCHandle(const UNiagaraParameterCollection* Collection);
 
-	virtual FDigestPtr Resolve() const override;
+	bool IsValid() const
+	{
+		return Resolve().IsValid();
+	}
+
+	FDigestPtr Resolve() const;
 
 	FName Namespace = NAME_None;
+
+protected:
+	TObjectKey<UNiagaraParameterCollection> AssetKey;
+	FNiagaraCompileHash Hash;
+
+	friend FORCEINLINE uint32 GetTypeHash(const FNiagaraCompilationNPCHandle& CompilationCopy)
+	{
+		return HashCombine(GetTypeHash(CompilationCopy.AssetKey), GetTypeHash(CompilationCopy.Hash));
+	}
+
+	friend FORCEINLINE bool operator==(const FNiagaraCompilationNPCHandle& Lhs, const FNiagaraCompilationNPCHandle& Rhs)
+	{
+		return Lhs.AssetKey == Rhs.AssetKey && Lhs.Hash == Rhs.Hash;
+	}
 };
 
 class FNiagaraDigestedParameterCollections
@@ -109,8 +108,7 @@ public:
 	static FNiagaraDigestDatabase& Get();
 	static void Shutdown();
 
-	FNiagaraCompilationGraphHandle CreateCompilationCopy(const UNiagaraGraph* Graph);
-	FNiagaraCompilationGraphHandle::FGraphPtr Resolve(const FNiagaraCompilationGraphHandle& Handle) const;
+	FNiagaraDigestedGraphPtr CreateGraphDigest(const UNiagaraGraph* Graph, const FNiagaraGraphChangeIdBuilder& Digester);
 
 	FNiagaraCompilationNPCHandle CreateCompilationCopy(const UNiagaraParameterCollection* Collection);
 	FNiagaraCompilationNPCHandle::FDigestPtr Resolve(const FNiagaraCompilationNPCHandle& Handle) const;
@@ -127,11 +125,7 @@ protected:
 	uint32 CollectionCacheHits = 0;
 	uint32 CollectionCacheMisses = 0;
 
-	// todo - replace these maps with some kind of two tiered system where we've got an TLruCache backing
-	// those objects we've digested recently and then another layer of digested objects that we've got pinned
-	// for the lifetime of a compilation task.  For now we'll just have an ever growing map
-
-	using FCompilationGraphCache = TMap<FNiagaraCompilationGraphHandle, FNiagaraCompilationGraphHandle::FGraphPtr>;
+	using FCompilationGraphCache = TLruCache<FNiagaraCompilationGraphHandle, FNiagaraDigestedGraphPtr>;
 	FCompilationGraphCache CompilationGraphCache;
 
 	using FCompilationNPCCache = TMap<FNiagaraCompilationNPCHandle, FNiagaraCompilationNPCHandle::FDigestPtr>;
diff --git a/Engine/Plugins/FX/Niagara/Source/NiagaraEditor/Private/NiagaraGraphDigest.cpp b/Engine/Plugins/FX/Niagara/Source/NiagaraEditor/Private/NiagaraGraphDigest.cpp
--- a/Engine/Plugins/FX/Niagara/Source/NiagaraEditor/Private/NiagaraGraphDigest.cpp
+++ b/Engine/Plugins/FX/Niagara/Source/NiagaraEditor/Private/NiagaraGraphDigest.cpp
@@ -44,12 +44,14 @@ NIAGARA_GRAPH_DIGEST_NODE_TYPE_LIST;
 
 struct FNiagaraCompilationGraphCreateContext
 {
-	FNiagaraCompilationGraphCreateContext(FNiagaraCompilationGraph& InParentGraph)
+	FNiagaraCompilationGraphCreateContext(FNiagaraCompilationGraph& InParentGraph, const FNiagaraGraphChangeIdBuilder& InChangeIdBuilder)
 		: ParentGraph(InParentGraph)
+		, ChangeIdBuilder(InChangeIdBuilder)
 	{
 	}
 
 	FNiagaraCompilationGraph& ParentGraph;
+	const FNiagaraGraphChangeIdBuilder& ChangeIdBuilder;
 };
 
 struct FNiagaraCompilationGraphDuplicateContext
@@ -462,7 +464,74 @@ bool GetStaticSwitchValueFromPin(FNiagaraCompilationGraphInstanceContext& Contex
 
 }; // NiagaraCompilationImpl
 
-void FNiagaraCompilationGraph::Create(const UNiagaraGraph* InGraph)
+void FNiagaraGraphChangeIdBuilder::ParseReferencedGraphs(const UNiagaraGraph* Graph)
+{
+	TSet<const UNiagaraGraph*> CurrentGraphChain;
+	RecursiveBuildGraphChangeId(Graph, CurrentGraphChain);
+}
+
+FGuid FNiagaraGraphChangeIdBuilder::FindChangeId(const UNiagaraGraph* Graph) const
+{
+	const FGuid* ChangeId = ChangeIdMap.Find(Graph);
+	if (ensure(ChangeId))
+	{
+		return *ChangeId;
+	}
+
+	return FGuid();
+}
+
+FGuid FNiagaraGraphChangeIdBuilder::RecursiveBuildGraphChangeId(const UNiagaraGraph* Graph, TSet<const UNiagaraGraph*>& CurrentGraphChain)
+{
+	if (!Graph)
+	{
+		return FGuid();
+	}
+
+	// see if the graph has already been processed
+	if (const FGuid* ExistingChangeId = ChangeIdMap.Find(Graph))
+	{
+		return *ExistingChangeId;
+	}
+
+	FGuid ChangeId = Graph->GetChangeID();
+
+	// check for a cycle in the graph chain we're currently processing
+	const bool bCycleInGraphChain = CurrentGraphChain.Contains(Graph);
+	if (bCycleInGraphChain)
+	{
+		ensure(false);
+		return ChangeId;
+	}
+
+	CurrentGraphChain.Add(Graph);
+
+	for (const UEdGraphNode* Node : Graph->Nodes)
+	{
+		const UNiagaraGraph* SubGraph = nullptr;
+
+		if (const UNiagaraNodeFunctionCall* FunctionCallNode = Cast<const UNiagaraNodeFunctionCall>(Node))
+		{
+			SubGraph = FunctionCallNode->GetCalledGraph();
+		}
+		else if (const UNiagaraNodeEmitter* EmitterNode = Cast<const UNiagaraNodeEmitter>(Node))
+		{
+			SubGraph = EmitterNode->GetCalledGraph();
+		}
+
+		if (SubGraph)
+		{
+			ChangeId = FGuid::Combine(ChangeId, SubGraph->GetChangeID());
+			ChangeId = FGuid::Combine(ChangeId, RecursiveBuildGraphChangeId(SubGraph, CurrentGraphChain));
+		}
+	}
+
+	ChangeIdMap.Add(Graph, ChangeId);
+
+	return ChangeId;
+}
+
+void FNiagaraCompilationGraph::Create(const UNiagaraGraph* InGraph, const FNiagaraGraphChangeIdBuilder& ChangeIdBuilder)
 {
 	using namespace NiagaraCompilationImpl;
 
@@ -477,7 +546,7 @@ void FNiagaraCompilationGraph::Create(const UNiagaraGraph* InGraph)
 
 	TMap<const UEdGraphNode*, int32> NodeIndexMap;
 
-	FNiagaraCompilationGraphCreateContext NodeContext(*this);
+	FNiagaraCompilationGraphCreateContext NodeContext(*this, ChangeIdBuilder);
 
 	for (const UEdGraphNode* SourceNode : InGraph->Nodes)
 	{
@@ -1932,7 +2001,7 @@ FNiagaraCompilationNodeEmitter::FNiagaraCompilationNodeEmitter(const UNiagaraNod
 
 	if (const UNiagaraGraph* DependentGraph = InNode->GetCalledGraph())
 	{
-		CalledGraph = FNiagaraDigestDatabase::Get().CreateCompilationCopy(DependentGraph).Resolve();
+		CalledGraph = FNiagaraDigestDatabase::Get().CreateGraphDigest(DependentGraph, Context.ChangeIdBuilder);
 	}
 
 	Usage = InNode->GetUsage();
@@ -2168,7 +2237,7 @@ FNiagaraCompilationNodeFunctionCall::FNiagaraCompilationNodeFunctionCall(const U
 {
 	if (const UNiagaraGraph* DependentGraph = InNode->GetCalledGraph())
 	{
-		CalledGraph = FNiagaraDigestDatabase::Get().CreateCompilationCopy(DependentGraph).Resolve();
+		CalledGraph = FNiagaraDigestDatabase::Get().CreateGraphDigest(DependentGraph, Context.ChangeIdBuilder);
 	}
 
 	// Top level functions (functions invoked in the root graph) will use the serialized DebugState value, all others will
diff --git a/Engine/Plugins/FX/Niagara/Source/NiagaraEditor/Private/NiagaraGraphDigest.h b/Engine/Plugins/FX/Niagara/Source/NiagaraEditor/Private/NiagaraGraphDigest.h
--- a/Engine/Plugins/FX/Niagara/Source/NiagaraEditor/Private/NiagaraGraphDigest.h
+++ b/Engine/Plugins/FX/Niagara/Source/NiagaraEditor/Private/NiagaraGraphDigest.h
@@ -234,12 +234,28 @@ protected:
 	bool CompileInputPins(FTranslator* Translator, TArray<int32>& OutInputResults) const;
 };
 
+// Helper class which can recursively generate the effective ChangeIds for a graph and all graphs that it
+// references.  With our current setup for digested graph, this new ChangeId reflects the fact that if a
+// graph references a changed graph (through function call or emitter node) then we'll need to re-digest
+// both caller and callee
+class FNiagaraGraphChangeIdBuilder
+{
+public:
+	void ParseReferencedGraphs(const UNiagaraGraph* Graph);
+	FGuid FindChangeId(const UNiagaraGraph* Graph) const;
+
+protected:
+	FGuid RecursiveBuildGraphChangeId(const UNiagaraGraph* Graph, TSet<const UNiagaraGraph*>& CurrentGraphChain);
+
+	TMap<const UNiagaraGraph*, FGuid> ChangeIdMap;
+};
+
 class FNiagaraCompilationGraph
 {
 public:
 	using FSharedPtr = TSharedPtr<FNiagaraCompilationGraph, ESPMode::ThreadSafe>;
 
-	void Create(const UNiagaraGraph* InGraph);
+	void Create(const UNiagaraGraph* InGraph, const FNiagaraGraphChangeIdBuilder& Digester);
 	bool IsValid() const { return true; } // todo - when creation is offloaded to a task we'll need this to synchronize completion
 
 	TArray<TUniquePtr<FNiagaraCompilationNode>> Nodes;
diff --git a/Engine/Plugins/FX/Niagara/Source/NiagaraEditor/Private/NiagaraSystemCompilingManager.cpp b/Engine/Plugins/FX/Niagara/Source/NiagaraEditor/Private/NiagaraSystemCompilingManager.cpp
--- a/Engine/Plugins/FX/Niagara/Source/NiagaraEditor/Private/NiagaraSystemCompilingManager.cpp
+++ b/Engine/Plugins/FX/Niagara/Source/NiagaraEditor/Private/NiagaraSystemCompilingManager.cpp
@@ -195,8 +195,7 @@ FNiagaraCompilationTaskHandle FNiagaraSystemCompilingManager::AddSystem(UNiagara
 
 	CompilationTask->PrepareStartTime = FPlatformTime::Seconds();
 
-	// we're going to have to compile something so let's digest all the graphs and build our compilation task
-	CompilationTask->DigestGraphs(AllScripts);
+	// we're going to have to compile something so let's digest all the collections and build our compilation task
 	CompilationTask->DigestParameterCollections(CompileOptions.ParameterCollections);
 	CompilationTask->DigestSystemInfo();
 
