Fixed sibling drag & drop not excluding the dragged item itself in case the name is identical. This stopped us from reordering categories within the same hierarchy

diff --git a/Engine/Plugins/FX/Niagara/Source/NiagaraEditor/Private/ViewModels/HierarchyEditor/NiagaraHierarchyViewModelBase.cpp b/Engine/Plugins/FX/Niagara/Source/NiagaraEditor/Private/ViewModels/HierarchyEditor/NiagaraHierarchyViewModelBase.cpp
--- a/Engine/Plugins/FX/Niagara/Source/NiagaraEditor/Private/ViewModels/HierarchyEditor/NiagaraHierarchyViewModelBase.cpp
+++ b/Engine/Plugins/FX/Niagara/Source/NiagaraEditor/Private/ViewModels/HierarchyEditor/NiagaraHierarchyViewModelBase.cpp
@@ -291,7 +291,7 @@ FNiagaraHierarchyItemViewModelBase::FCanPerformActionResults FNiagaraHierarchyCa
 			{
 				bool bContainsSiblingWithSameName = SiblingCategories.ContainsByPredicate([DraggedItem](TSharedPtr<FNiagaraHierarchyCategoryViewModel> HierarchyCategoryViewModel)
 					{
-						return DraggedItem->ToString() == HierarchyCategoryViewModel->ToString();
+						return DraggedItem->ToString() == HierarchyCategoryViewModel->ToString() && DraggedItem != HierarchyCategoryViewModel;
 					});
 
 				if(bContainsSiblingWithSameName)
