Niagara - Change a stack entry reference to an owning emitter handle view model from a shared ptr to a weak ptr to prevent stack entry uobjects from keeping emitter view models alive.

diff --git a/Engine/Plugins/FX/Niagara/Source/NiagaraEditor/Private/ViewModels/NiagaraScriptGraphViewModel.cpp b/Engine/Plugins/FX/Niagara/Source/NiagaraEditor/Private/ViewModels/NiagaraScriptGraphViewModel.cpp
--- a/Engine/Plugins/FX/Niagara/Source/NiagaraEditor/Private/ViewModels/NiagaraScriptGraphViewModel.cpp
+++ b/Engine/Plugins/FX/Niagara/Source/NiagaraEditor/Private/ViewModels/NiagaraScriptGraphViewModel.cpp
@@ -44,7 +44,7 @@ FNiagaraScriptGraphViewModel::FNiagaraScriptGraphViewModel(TAttribute<FText> InD
 
 FNiagaraScriptGraphViewModel::~FNiagaraScriptGraphViewModel()
 {
-	if (bIsForDataProcessingOnly == false)
+	if (bIsForDataProcessingOnly == false && GEditor != nullptr)
 	{
 		GEditor->UnregisterForUndo(this);
 	}
diff --git a/Engine/Plugins/FX/Niagara/Source/NiagaraEditor/Private/ViewModels/Stack/NiagaraStackItemGroup.cpp b/Engine/Plugins/FX/Niagara/Source/NiagaraEditor/Private/ViewModels/Stack/NiagaraStackItemGroup.cpp
--- a/Engine/Plugins/FX/Niagara/Source/NiagaraEditor/Private/ViewModels/Stack/NiagaraStackItemGroup.cpp
+++ b/Engine/Plugins/FX/Niagara/Source/NiagaraEditor/Private/ViewModels/Stack/NiagaraStackItemGroup.cpp
@@ -41,7 +41,7 @@ FText UNiagaraStackItemGroup::GetTooltipText() const
 
 bool UNiagaraStackItemGroup::GetIsEnabled() const
 {
-	return OwningEmitterHandleViewModel.IsValid() == false || OwningEmitterHandleViewModel->GetIsEnabled();
+	return OwningEmitterHandleViewModelWeak.IsValid() == false || OwningEmitterHandleViewModelWeak.Pin()->GetIsEnabled();
 }
 
 INiagaraStackItemGroupAddUtilities* UNiagaraStackItemGroup::GetAddUtilities() const
@@ -99,14 +99,16 @@ void UNiagaraStackItemGroup::RefreshChildrenInternal(const TArray<UNiagaraStackE
 	TSharedPtr<FNiagaraEmitterViewModel> OwningEmitterViewModel = GetEmitterViewModel();
 	if (OwningEmitterViewModel.IsValid())
 	{
+		TSharedPtr<FNiagaraEmitterHandleViewModel> OwningEmitterHandleViewModel = OwningEmitterHandleViewModelWeak.Pin();
 		if (OwningEmitterHandleViewModel.IsValid() == false || OwningEmitterHandleViewModel->GetEmitterViewModel() != OwningEmitterViewModel)
 		{
-			OwningEmitterHandleViewModel = GetSystemViewModel()->GetEmitterHandleViewModelForEmitter(OwningEmitterViewModel->GetEmitter());
+			OwningEmitterHandleViewModelWeak = TWeakPtr<FNiagaraEmitterHandleViewModel>(
+				GetSystemViewModel()->GetEmitterHandleViewModelForEmitter(OwningEmitterViewModel->GetEmitter()));
 		}
 	}
 	else
 	{
-		OwningEmitterHandleViewModel.Reset();
+		OwningEmitterHandleViewModelWeak.Reset();
 	}
 }
 
diff --git a/Engine/Plugins/FX/Niagara/Source/NiagaraEditor/Public/ViewModels/Stack/NiagaraStackItemGroup.h b/Engine/Plugins/FX/Niagara/Source/NiagaraEditor/Public/ViewModels/Stack/NiagaraStackItemGroup.h
--- a/Engine/Plugins/FX/Niagara/Source/NiagaraEditor/Public/ViewModels/Stack/NiagaraStackItemGroup.h
+++ b/Engine/Plugins/FX/Niagara/Source/NiagaraEditor/Public/ViewModels/Stack/NiagaraStackItemGroup.h
@@ -63,7 +63,7 @@ private:
 	/** The highest severity of issues along this entry's tree. */
 	mutable TOptional<EStackIssueSeverity> HighestIssueSeverity;
 
-	TSharedPtr<FNiagaraEmitterHandleViewModel> OwningEmitterHandleViewModel;
+	TWeakPtr<FNiagaraEmitterHandleViewModel> OwningEmitterHandleViewModelWeak;
 };
 
 UCLASS()
