Fix for default compilation mode not releasing the hold on the compilation copies of the niagara graph source.
#rnx

diff --git a/Engine/Plugins/FX/Niagara/Source/Niagara/Private/NiagaraAsyncCompile.cpp b/Engine/Plugins/FX/Niagara/Source/Niagara/Private/NiagaraAsyncCompile.cpp
--- a/Engine/Plugins/FX/Niagara/Source/Niagara/Private/NiagaraAsyncCompile.cpp
+++ b/Engine/Plugins/FX/Niagara/Source/Niagara/Private/NiagaraAsyncCompile.cpp
@@ -488,6 +488,12 @@ bool FNiagaraAsyncCompileTask::CompilationIdMatchesRequest() const
 	return IdMatches;
 }
 
+FNiagaraActiveCompilationDefault::~FNiagaraActiveCompilationDefault()
+{
+	// make sure that we release all of the compilation copies that have been created for the tasks
+	Reset();
+}
+
 void FNiagaraActiveCompilationDefault::AddReferencedObjects(FReferenceCollector& Collector)
 {
 	Collector.AddStableReferenceSet(&RootObjects);
diff --git a/Engine/Plugins/FX/Niagara/Source/Niagara/Private/NiagaraAsyncCompile.h b/Engine/Plugins/FX/Niagara/Source/Niagara/Private/NiagaraAsyncCompile.h
--- a/Engine/Plugins/FX/Niagara/Source/Niagara/Private/NiagaraAsyncCompile.h
+++ b/Engine/Plugins/FX/Niagara/Source/Niagara/Private/NiagaraAsyncCompile.h
@@ -117,6 +117,7 @@ public:
 class FNiagaraActiveCompilationDefault : public FNiagaraActiveCompilation
 {
 public:
+	virtual ~FNiagaraActiveCompilationDefault();
 	virtual void AddReferencedObjects(FReferenceCollector& Collector) override;
 	virtual FString GetReferencerName() const override
 	{
