- Refactored user parameter customization
- Split classes into different files
- Now supports type utilities as value widgets. This means we can have custom widgets & min max support etc.
- Added 'reset to default' buttons back to the User Parameters tab
- Added callback to rebuild UI when a script variable property updates so it updates live when widget customization is changed

diff --git a/Engine/Plugins/FX/Niagara/Source/NiagaraEditor/Private/Customizations/NiagaraComponentDetails.cpp b/Engine/Plugins/FX/Niagara/Source/NiagaraEditor/Private/Customizations/NiagaraComponentDetails.cpp
--- a/Engine/Plugins/FX/Niagara/Source/NiagaraEditor/Private/Customizations/NiagaraComponentDetails.cpp
+++ b/Engine/Plugins/FX/Niagara/Source/NiagaraEditor/Private/Customizations/NiagaraComponentDetails.cpp
@@ -5,21 +5,14 @@
 #include "DetailLayoutBuilder.h"
 #include "DetailCategoryBuilder.h"
 #include "NiagaraComponent.h"
-#include "NiagaraDataInterface.h"
-#include "IDetailPropertyRow.h"
 #include "DetailWidgetRow.h"
 #include "NiagaraSimCacheCapture.h"
-#include "Widgets/Text/SInlineEditableTextBlock.h"
 #include "PropertyCustomizationHelpers.h"
 #include "Widgets/Text/STextBlock.h"
-#include "Widgets/Images/SImage.h"
 #include "Widgets/Layout/SBox.h"
 #include "Styling/AppStyle.h"
 #include "ViewModels/NiagaraParameterViewModel.h"
-#include "NiagaraEditorUtilities.h"
 #include "INiagaraEditorTypeUtilities.h"
-#include "NiagaraTypes.h"
-#include "ScopedTransaction.h"
 #include "NiagaraEditorModule.h"
 #include "PropertyEditorModule.h"
 #include "Modules/ModuleManager.h"
@@ -28,584 +21,16 @@
 #include "NiagaraNodeInput.h"
 #include "GameDelegates.h"
 #include "IAssetTools.h"
-#include "NiagaraEditorStyle.h"
 #include "IDetailChildrenBuilder.h"
-#include "NiagaraClipboard.h"
-#include "NiagaraConstants.h"
-#include "NiagaraScriptVariable.h"
 #include "NiagaraSettings.h"
-#include "NiagaraSystemEditorData.h"
-#include "NiagaraUserRedirectionParameterStore.h"
 #include "Widgets/Layout/SUniformGridPanel.h"
 #include "Widgets/SNiagaraDebugger.h"
-#include "Widgets/SNiagaraParameterName.h"
-#include "ViewModels/HierarchyEditor/NiagaraUserParametersHierarchyViewModel.h"
-#include "Widgets/Input/SCheckBox.h"
-#include "Widgets/Input/SComboButton.h"
 #include "Widgets/SNiagaraSystemUserParameters.h"
-#include "Toolkits/SystemToolkitModes/NiagaraSystemToolkitModeBase.h"
-#include "Widgets/Input/SEditableTextBox.h"
+#include "UserParameters/NiagaraComponentUserParametersBuilder.h"
+#include "UserParameters/NiagaraSystemUserParametersBuilder.h"
 
 #define LOCTEXT_NAMESPACE "NiagaraComponentDetails"
 
-static FNiagaraVariant GetParameterValueFromSystem(const FNiagaraVariableBase& Parameter, const UNiagaraSystem* System)
-{
-	const FNiagaraUserRedirectionParameterStore& UserParameterStore = System->GetExposedParameters();
-	if (Parameter.IsDataInterface())
-	{
-		int32 Index = UserParameterStore.IndexOf(Parameter);
-		if (Index != INDEX_NONE)
-		{
-			return FNiagaraVariant(UserParameterStore.GetDataInterfaces()[Index]);
-		}
-	}
-	
-	if (Parameter.IsUObject())
-	{
-		int32 Index = UserParameterStore.IndexOf(Parameter);
-		if (Index != INDEX_NONE)
-		{
-			return FNiagaraVariant(UserParameterStore.GetUObjects()[Index]);
-		}
-	}
-
-	if (Parameter.GetType() == FNiagaraTypeDefinition::GetPositionDef())
-	{
-		const FVector* Value = UserParameterStore.GetPositionParameterValue(Parameter.GetName());
-		return Value == nullptr ? FNiagaraVariant() : FNiagaraVariant(Value, sizeof(FVector));
-	}
-
-	TArray<uint8> DataValue;
-	DataValue.AddUninitialized(Parameter.GetSizeInBytes());
-	if (UserParameterStore.CopyParameterData(Parameter, DataValue.GetData()))
-	{
-		return FNiagaraVariant(DataValue);
-	}
-	return FNiagaraVariant();
-}
-
-static FNiagaraVariant GetCurrentParameterValue(const FNiagaraVariableBase& Parameter, const UObject* Owner)
-{
-	if(const UNiagaraComponent* Component = Cast<UNiagaraComponent>(Owner))
-	{
-		FNiagaraVariant CurrentValue = Component->GetCurrentParameterValue(Parameter);
-		if (CurrentValue.IsValid())
-		{
-			return CurrentValue;
-		}
-		
-		return GetParameterValueFromSystem(Parameter, Component->GetAsset());
-	}
-	else if(const UNiagaraSystem* System = Cast<UNiagaraSystem>(Owner))
-	{
-		return GetParameterValueFromSystem(Parameter, System);
-	}
-
-	ensure(false);
-	return FNiagaraVariant();
-}
-
-// Proxy class to allow us to override values on the component that are not yet overridden.
-class FNiagaraParameterProxy
-{
-public:
-	FNiagaraParameterProxy(TWeakObjectPtr<UObject> InOwner, const FNiagaraVariableBase& InKey, const FNiagaraVariant& InValue, const FSimpleDelegate& InOnRebuild, TArray<TSharedPtr<IPropertyHandle>> InPropertyHandles)
-	{
-		bResettingToDefault = false;
-		Owner = InOwner;
-		ParameterKey = InKey;
-		ParameterValue = InValue;
-		OnRebuild = InOnRebuild;
-		PropertyHandles = InPropertyHandles;
-	}
-
-	FReply OnResetToDefaultClicked()
-	{
-		OnResetToDefault();
-		return FReply::Handled();
-	}
-
-	void OnResetToDefault()
-	{
-		if(Owner->IsA<UNiagaraComponent>())
-		{
-			TWeakObjectPtr<UNiagaraComponent> RawComponent = GetComponent();
-			if (RawComponent != nullptr)
-			{
-				FScopedTransaction ScopedTransaction(NSLOCTEXT("UnrealEd", "PropertyWindowResetToDefault", "Reset to Default"));
-				RawComponent->Modify();
-
-				bResettingToDefault = true;
-
-				for (TSharedPtr<IPropertyHandle> PropertyHandle : PropertyHandles)
-				{
-					PropertyHandle->NotifyPreChange();
-				}
-
-				RawComponent->RemoveParameterOverride(ParameterKey);
-
-				for (TSharedPtr<IPropertyHandle> PropertyHandle : PropertyHandles)
-				{
-					PropertyHandle->NotifyPostChange(EPropertyChangeType::ValueSet);
-				}
-
-				OnRebuild.ExecuteIfBound();
-				bResettingToDefault = false;
-			}
-		}
-	}
-
-	EVisibility GetResetToDefaultVisibility() const
-	{
-		TWeakObjectPtr<UNiagaraComponent> RawComponent = GetComponent();
-		if (RawComponent.IsValid())
-		{
-			return RawComponent->HasParameterOverride(ParameterKey) ? EVisibility::Visible : EVisibility::Hidden;
-		}
-		return EVisibility::Hidden;
-	}
-
-	FNiagaraVariant FindExistingOverride() const
-	{
-		TWeakObjectPtr<UNiagaraComponent> RawComponent = GetComponent();
-		if (RawComponent.IsValid())
-		{
-			return RawComponent->FindParameterOverride(ParameterKey);
-		}
-		return FNiagaraVariant();
-	}
-
-	void OnParameterPreChange()
-	{
-		if (Owner->IsA<UNiagaraComponent>())
-		{
-			TWeakObjectPtr<UNiagaraComponent> RawComponent = GetComponent();
-			if(RawComponent.IsValid())
-			{
-				RawComponent->Modify();
-
-				for (TSharedPtr<IPropertyHandle> PropertyHandle : PropertyHandles)
-				{
-					PropertyHandle->NotifyPreChange();
-				}
-			}
-		}
-		else if(Owner->IsA<UNiagaraSystem>())
-		{
-			TWeakObjectPtr<UNiagaraSystem> System = GetSystem();
-			if(System.IsValid())
-			{
-				System->Modify();
-			}
-		}
-	}
-
-	void OnParameterChanged(const FPropertyChangedEvent& PropertyChangedEvent)
-	{
-		if(Owner->IsA<UNiagaraComponent>())
-		{
-			if (bResettingToDefault)
-			{
-				return;
-			}
-
-			TWeakObjectPtr<UNiagaraComponent> RawComponent = GetComponent();
-			if(RawComponent.IsValid())
-			{
-				RawComponent->SetParameterOverride(ParameterKey, ParameterValue);
-
-				for (TSharedPtr<IPropertyHandle> PropertyHandle : PropertyHandles)
-				{
-					PropertyHandle->NotifyPostChange(PropertyChangedEvent.ChangeType);
-				}
-			}
-		}
-		else if(Owner->IsA<UNiagaraSystem>())
-		{
-			TWeakObjectPtr<UNiagaraSystem> System = GetSystem();
-			if(System.IsValid())
-			{
-				if(ParameterValue.GetMode() == ENiagaraVariantMode::Bytes)
-				{
-					// we have to special case the position definition because, while SetParameterData does handle Position variables, it expects 12 Bytes of data
-					if(ParameterKey.GetType() == FNiagaraTypeDefinition::GetPositionDef())
-					{
-						FVector Value = *reinterpret_cast<const FVector*>(ParameterValue.GetBytes());
-						System->GetExposedParameters().SetPositionParameterValue(Value, ParameterKey.GetName());
-					}
-					else
-					{
-						System->GetExposedParameters().SetParameterData(ParameterValue.GetBytes(), ParameterKey);
-					}
-				}
-				else if(ParameterValue.GetMode() == ENiagaraVariantMode::DataInterface)
-				{
-					// we are operating on the DI objects in the parameter store directly, so no need to call SetDataInterface
-					System->GetExposedParameters().OnInterfaceChange();
-				}
-				else if(ParameterValue.GetMode() == ENiagaraVariantMode::Object)
-				{
-					System->GetExposedParameters().SetUObject(ParameterValue.GetUObject(), ParameterKey);
-				}
-
-				for (TSharedPtr<IPropertyHandle> PropertyHandle : PropertyHandles)
-				{
-					PropertyHandle->NotifyPostChange(PropertyChangedEvent.ChangeType);
-				}			
-			}
-		}
-	}
-
-	void OnAssetSelectedFromPicker(const FAssetData& InAssetData)
-	{
-		if(Owner->IsA<UNiagaraComponent>())
-		{
-			TWeakObjectPtr<UNiagaraComponent> RawComponent = GetComponent();
-
-			if(RawComponent.IsValid())
-			{
-				UObject* Asset = InAssetData.GetAsset();
-				if (Asset == nullptr || Asset->GetClass()->IsChildOf(ParameterKey.GetType().GetClass()))
-				{
-					FScopedTransaction ScopedTransaction(LOCTEXT("ChangeAsset", "Change asset"));
-					RawComponent->Modify();
-					RawComponent->SetParameterOverride(ParameterKey, FNiagaraVariant(Asset));
-
-					for (TSharedPtr<IPropertyHandle> PropertyHandle : PropertyHandles)
-					{
-						PropertyHandle->NotifyPostChange(EPropertyChangeType::ValueSet);
-					}
-				}
-			}
-		}
-		else if(Owner->IsA<UNiagaraSystem>())
-		{
-			TWeakObjectPtr<UNiagaraSystem> System = GetSystem();
-
-			if(System.IsValid())
-			{
-				UObject* Asset = InAssetData.GetAsset();
-				if (Asset == nullptr || Asset->GetClass()->IsChildOf(ParameterKey.GetType().GetClass()))
-				{
-					FScopedTransaction ScopedTransaction(LOCTEXT("ChangeAsset", "Change asset"));
-					System->Modify();
-					System->GetExposedParameters().SetUObject(Asset, ParameterKey);
-
-					for (TSharedPtr<IPropertyHandle> PropertyHandle : PropertyHandles)
-					{
-						PropertyHandle->NotifyPostChange(EPropertyChangeType::ValueSet);
-					}
-				}
-			}
-		}
-	}
-
-	FString GetCurrentAssetPath() const
-	{
-		UObject* CurrentObject = nullptr;
-
-		if(Owner->IsA<UNiagaraComponent>())
-		{
-			TWeakObjectPtr<UNiagaraComponent> RawComponent = GetComponent();
-
-			if(RawComponent.IsValid())
-			{
-				FNiagaraVariant CurrentValue = FindExistingOverride();
-				if (CurrentValue.IsValid())
-				{
-					CurrentObject = CurrentValue.GetUObject();
-				}
-				else
-				{
-					// fetch from asset
-					UNiagaraSystem* System = RawComponent->GetAsset();
-					if (System != nullptr)
-					{
-						FNiagaraUserRedirectionParameterStore& AssetParamStore = System->GetExposedParameters();
-						CurrentObject = AssetParamStore.GetUObject(ParameterKey);
-					}
-				}
-			}
-		}
-		else if(Owner->IsA<UNiagaraSystem>())
-		{			
-			TWeakObjectPtr<UNiagaraSystem> System = GetSystem();
-
-			if(System.IsValid())
-			{
-				FNiagaraUserRedirectionParameterStore& AssetParamStore = System->GetExposedParameters();
-				CurrentObject = AssetParamStore.GetUObject(ParameterKey);			
-			}
-		}
-
-		return CurrentObject != nullptr ? CurrentObject->GetPathName() : FString();
-	}
-	
-	const FNiagaraVariableBase& Key() const { return ParameterKey; }
-	FNiagaraVariant& Value() { return ParameterValue; }
-
-	TWeakObjectPtr<UNiagaraComponent> GetComponent() const { return Cast<UNiagaraComponent>(Owner); }
-	TWeakObjectPtr<UNiagaraSystem> GetSystem() const { return Cast<UNiagaraSystem>(Owner); }
-
-private:
-	TWeakObjectPtr<UObject> Owner;
-	TArray<TSharedPtr<IPropertyHandle>> PropertyHandles;
-	FNiagaraVariableBase ParameterKey;
-	FNiagaraVariant ParameterValue;
-	FSimpleDelegate OnRebuild;
-	bool bResettingToDefault;
-};
-#undef LOCTEXT_NAMESPACE
-
-#define LOCTEXT_NAMESPACE "NiagaraUserParameterDetails"
-
-void FNiagaraUserParameterNodeBuilder::GenerateRowForUserParameter(IDetailChildrenBuilder& ChildrenBuilder, FNiagaraVariable UserParameter, UObject* Owner)
-{
-	FNiagaraVariable ChoppedUserParameter(UserParameter);
-
-	if(UserParameter.IsInNameSpace(FNiagaraConstants::UserNamespaceString))
-	{
-		ChoppedUserParameter.SetName(FName(UserParameter.GetName().ToString().RightChop(5)));
-	}
-	
-	FNiagaraVariant ParameterValue = GetCurrentParameterValue(ChoppedUserParameter, Owner);
-	if (!ParameterValue.IsValid())
-	{
-		return;
-	}
-		
-	if (UserParameter.IsDataInterface())
-	{
-		if(Owner->IsA<UNiagaraComponent>())
-		{
-			ParameterValue = FNiagaraVariant(DuplicateObject(ParameterValue.GetDataInterface(), Owner));
-		}
-		else if(Owner->IsA<UNiagaraSystem>())
-		{
-			ParameterValue = FNiagaraVariant(ParameterValue.GetDataInterface());
-		}
-	}
-
-	TSharedPtr<FNiagaraParameterProxy> ParameterProxy = ParameterProxies.Add_GetRef(MakeShared<FNiagaraParameterProxy>(Owner, UserParameter, ParameterValue, OnRebuildChildren, OverridePropertyHandles));
-
-	IDetailPropertyRow* Row = nullptr;
-
-	TSharedPtr<SWidget> CustomNameWidget = GenerateCustomNameWidget(UserParameter);
- 	TSharedPtr<SWidget> CustomValueWidget;
-	
-	if (UserParameter.IsDataInterface())
-	{
-		// if no changes are made, then it'll just be the same as the asset
-		TArray<UObject*> Objects { ParameterProxy->Value().GetDataInterface() };
-	
-		FAddPropertyParams Params = FAddPropertyParams()
-			.UniqueId(ChoppedUserParameter.GetName())
-			.AllowChildren(true)
-			.CreateCategoryNodes(false);
-	
-		Row = ChildrenBuilder.AddExternalObjectProperty(Objects, NAME_None, Params);
-		
-		CustomValueWidget =
-			SNew(STextBlock)
-			.TextStyle(FNiagaraEditorStyle::Get(), "NiagaraEditor.ParameterText")
-			.Text(FText::FromString(FName::NameToDisplayString(UserParameter.GetType().GetClass()->GetName(), false)));
-	}
-	else if(UserParameter.IsUObject())
-	{
-		TArray<UObject*> Objects { ParameterProxy->Value().GetUObject() };
-
-		FAddPropertyParams Params = FAddPropertyParams()
-			.UniqueId(ChoppedUserParameter.GetName())
-			.AllowChildren(false) // Don't show the material's properties
-			.CreateCategoryNodes(false);
-
-		Row = ChildrenBuilder.AddExternalObjectProperty(Objects, NAME_None, Params);
-		
-		CustomValueWidget = SNew(SObjectPropertyEntryBox)
-			.ObjectPath(ParameterProxy.ToSharedRef(), &FNiagaraParameterProxy::GetCurrentAssetPath)
-			.AllowedClass(UserParameter.GetType().GetClass())
-			.OnObjectChanged(ParameterProxy.ToSharedRef(), &FNiagaraParameterProxy::OnAssetSelectedFromPicker)
-			.DisplayUseSelected(true)
-			.DisplayBrowse(true)
-			.DisplayThumbnail(true)
-			.NewAssetFactories(TArray<UFactory*>());
-	}
-	else
-	{
-		FNiagaraTypeDefinition Type = UserParameter.GetType();
-		UStruct* Struct = Type.GetStruct();
-			
-		if (Type == FNiagaraTypeDefinition::GetPositionDef())
-		{
-			static UPackage* CoreUObjectPkg = FindObjectChecked<UPackage>(nullptr, TEXT("/Script/CoreUObject"));
-			static UScriptStruct* VectorStruct = FindObjectChecked<UScriptStruct>(CoreUObjectPkg, TEXT("Vector"));
-			Struct = VectorStruct;
-		}
-
-		// the details panel uses byte properties with enums assigned to display enums
-		if (Type.IsEnum())
-		{
-			static UPackage* NiagaraEditorPkg = FindObjectChecked<UPackage>(nullptr, TEXT("/Script/NiagaraEditor"));
-			static UScriptStruct* ByteStruct = FindObjectChecked<UScriptStruct>(NiagaraEditorPkg, TEXT("NiagaraEnumToByteHelper"));
-			Struct = ByteStruct;
-		}
-			
-		TSharedPtr<FStructOnScope> StructOnScope = MakeShareable(new FStructOnScope(Struct, ParameterProxy->Value().GetBytes()));
-
-		FAddPropertyParams Params = FAddPropertyParams()
-			.UniqueId(ChoppedUserParameter.GetName());
-
-		Row = ChildrenBuilder.AddExternalStructureProperty(StructOnScope.ToSharedRef(), NAME_None, Params);
-
-		// we set the enum of the contained byte property to the enum carried with the type
-		if(Type.IsEnum())
-		{
-			FProperty* Property = Row->GetPropertyHandle()->GetProperty();
-			TSharedPtr<IPropertyHandle> ValuePropertyHandle = Row->GetPropertyHandle()->GetChildHandle(0);
-			FProperty* ValueProperty = ValuePropertyHandle->GetProperty();
-			if(ValueProperty->IsA<FByteProperty>())
-			{
-				if(FByteProperty* ByteProperty = CastField<FByteProperty>(ValueProperty))
-				{
-					ByteProperty->Enum = Type.GetEnum();
-				}
-			}
-		}
-	
-		ParameterToDisplayStruct.Add(UserParameter, TWeakPtr<FStructOnScope>(StructOnScope));
-	}
-
-	TSharedPtr<SWidget> DefaultNameWidget;
-	TSharedPtr<SWidget> DefaultValueWidget;
-	
-	Row->DisplayName(FText::FromName(ChoppedUserParameter.GetName()));
-	
-	FDetailWidgetRow& CustomWidget = Row->CustomWidget(true);
-
-	if(UserParameter.GetType().IsDataInterface())
-	{
-		FUIAction CopyAction;
-		CopyAction.ExecuteAction = FExecuteAction::CreateLambda([=]()
-		{
-			UNiagaraClipboardContent* ClipboardContent = UNiagaraClipboardContent::Create();
-			const UNiagaraClipboardFunctionInput* Input = UNiagaraClipboardFunctionInput::CreateDataValue(ClipboardContent, UserParameter.GetName(), UserParameter.GetType().GetClass(), TOptional<bool>(), ParameterProxy->Value().GetDataInterface());
-			ClipboardContent->FunctionInputs.Add(Input);
-			
-			FNiagaraEditorModule::Get().GetClipboard().SetClipboardContent(ClipboardContent);
-		});
-		CustomWidget.CopyAction(CopyAction);
-
-		FUIAction PasteAction;
-		PasteAction.ExecuteAction = FExecuteAction::CreateLambda([=]()
-		{
-			const UNiagaraClipboardContent* ClipboardContent = FNiagaraEditorModule::Get().GetClipboard().GetClipboardContent();
-
-			if(ClipboardContent->FunctionInputs.Num() == 1 && ClipboardContent->FunctionInputs[0]->InputType.IsDataInterface())
-			{
-				UNiagaraDataInterface* CopiedDI = ClipboardContent->FunctionInputs[0]->Data;
-				if(CopiedDI->IsA(ParameterProxy->Value().GetDataInterface()->GetClass()))
-				{
-					CopiedDI->CopyTo(ParameterProxy->Value().GetDataInterface());
-				}
-			}			
-		});
-		CustomWidget.PasteAction(PasteAction);
-	}
-	
-	AddCustomMenuActionsForParameter(CustomWidget, UserParameter);
-	
-	UNiagaraSystem* System = nullptr;
-	if(ParameterProxy->GetSystem().IsValid())
-	{
-		System = ParameterProxy->GetSystem().Get();
-	}
-	else if(ParameterProxy->GetComponent().IsValid())
-	{
-		System = ParameterProxy->GetComponent()->GetAsset();
-	}
-
-	if(System)
-	{
-		UNiagaraScriptVariable* ScriptVariable = FNiagaraEditorUtilities::GetScriptVariableForUserParameter(UserParameter, *System);
-			
-		Row->GetDefaultWidgets(DefaultNameWidget, DefaultValueWidget, CustomWidget);
-
-		Row->GetPropertyHandle()->MarkResetToDefaultCustomized(true);
-		
-		Row->GetPropertyHandle()->SetOnPropertyValuePreChange(FSimpleDelegate::CreateSP(ParameterProxy.ToSharedRef(), &FNiagaraParameterProxy::OnParameterPreChange));
-		Row->GetPropertyHandle()->SetOnChildPropertyValuePreChange(FSimpleDelegate::CreateSP(ParameterProxy.ToSharedRef(), &FNiagaraParameterProxy::OnParameterPreChange));
-		Row->GetPropertyHandle()->SetOnPropertyValueChangedWithData(TDelegate<void(const FPropertyChangedEvent&)>::CreateSP(ParameterProxy.ToSharedRef(), &FNiagaraParameterProxy::OnParameterChanged));
-		Row->GetPropertyHandle()->SetOnChildPropertyValueChangedWithData(TDelegate<void(const FPropertyChangedEvent&)>::CreateSP(ParameterProxy.ToSharedRef(), &FNiagaraParameterProxy::OnParameterChanged));
-		Row->GetPropertyHandle()->SetOnPropertyResetToDefault(FSimpleDelegate::CreateSP(ParameterProxy.ToSharedRef(), &FNiagaraParameterProxy::OnResetToDefault));
-
-		TSharedPtr<SWidget> NameWidget = SNew(SBox)
-			.Padding(FMargin(0.0f, 2.0f))
-			[
-				SNew(STextBlock)
-				.TextStyle(FNiagaraEditorStyle::Get(), "NiagaraEditor.ParameterText")
-				.Text(FText::FromName(ChoppedUserParameter.GetName()))
-			];
-		
-		if (CustomNameWidget.IsValid())
-		{
-			NameWidget = CustomNameWidget;
-		}
-
-		// regardless of which widget we use for the name, we want to use the script variable's description as tooltip
-		TAttribute<FText> Tooltip = FText::GetEmpty();
-		if(ScriptVariable)
-		{
-			Tooltip = TAttribute<FText>::CreateLambda([ScriptVariable]
-			{
-				return ScriptVariable->Metadata.Description;
-			});
-		}
-		NameWidget->SetToolTipText(Tooltip);
-		
-		TSharedPtr<SWidget> ValueWidget = DefaultValueWidget;
-		if (CustomValueWidget.IsValid())
-		{
-			ValueWidget = CustomValueWidget;
-		}
-		
-		CustomWidget
-		.FilterString(FText::FromName(ChoppedUserParameter.GetName()))
-		.NameContent()
-		.HAlign(HAlign_Left)
-		[
-			NameWidget.ToSharedRef()
-		]
-		.ValueContent()
-		[
-			SNew(SHorizontalBox)
-			+ SHorizontalBox::Slot()
-			.HAlign(HAlign_Fill)
-			.Padding(4.0f)
-			[
-				// Add in the parameter editor factoried above.
-				ValueWidget.ToSharedRef()
-			]
-			+ SHorizontalBox::Slot()
-			.VAlign(VAlign_Center)
-			.AutoWidth()
-			[
-				// Add in the "reset to default" buttons
-				SNew(SButton)
-				.OnClicked(ParameterProxy.ToSharedRef(), &FNiagaraParameterProxy::OnResetToDefaultClicked)
-				.Visibility(ParameterProxy.ToSharedRef(), &FNiagaraParameterProxy::GetResetToDefaultVisibility)
-				.ContentPadding(FMargin(5.f, 0.f))
-				.ToolTipText(LOCTEXT("ResetToDefaultToolTip", "Reset to Default"))
-				.ButtonStyle(FAppStyle::Get(), "NoBorder")
-				.Content()
-				[
-					SNew(SImage)
-					.Image(FAppStyle::GetBrush("PropertyWindow.DiffersFromDefault"))
-				]
-			]
-		];
-	}
-}
-
 TSharedRef<IDetailCustomization> FNiagaraSystemUserParameterDetails::MakeInstance()
 {
 	return MakeShared<FNiagaraSystemUserParameterDetails>();
@@ -639,669 +64,6 @@ void FNiagaraSystemUserParameterDetails::CustomizeDetails(IDetailLayoutBuilder&
 	}
 }
 
-/** The category builder will display all sub-categories & parameters contained within a given category in the details panel. */
-class FNiagaraUserParameterCategoryBuilder : public IDetailCustomNodeBuilder, public TSharedFromThis<FNiagaraUserParameterCategoryBuilder>
-{
-public:
-	FNiagaraUserParameterCategoryBuilder(UObject* InOwner, const UNiagaraHierarchyCategory* InNiagaraHierarchyCategory, FName InCustomBuilderRowName, FNiagaraUserParameterNodeBuilder& InNodeBuilder) : NodeBuilder(InNodeBuilder) 
-	{
-		HierarchyCategory = InNiagaraHierarchyCategory;
-		CustomBuilderRowName = InCustomBuilderRowName;
-
-		Owner = InOwner;
-	}
-	
-	virtual void GenerateHeaderRowContent(FDetailWidgetRow& NodeRow) override
-	{
-		NodeRow.NameContent()
-		[
-			SNew(STextBlock)
-			.TextStyle(FNiagaraEditorStyle::Get(), "NiagaraEditor.ParameterText")
-			.Text_UObject(HierarchyCategory.Get(), &UNiagaraHierarchyCategory::GetCategoryAsText)
-			.ToolTipText_UObject(HierarchyCategory.Get(), &UNiagaraHierarchyCategory::GetTooltip)
-		];
-
-		NodeRow.FilterString(FText::FromName(CustomBuilderRowName));
-	}
-	
-	virtual void Tick(float DeltaTime) override {}
-	virtual bool RequiresTick() const override { return false; }
-	virtual bool InitiallyCollapsed() const override { return false; }
-
-	virtual FName GetName() const  override
-	{
-		return CustomBuilderRowName;
-	}
-	
-	virtual void GenerateChildContent(IDetailChildrenBuilder& ChildrenBuilder) override
-	{
-		TArray<const UNiagaraHierarchyUserParameter*> ContainedParameters;
-		TArray<const UNiagaraHierarchyCategory*> ContainedCategories;
-
-		for(const UNiagaraHierarchyItemBase* Child : HierarchyCategory->GetChildren())
-		{
-			if(const UNiagaraHierarchyUserParameter* UserParameter = Cast<UNiagaraHierarchyUserParameter>(Child))
-			{
-				ContainedParameters.Add(UserParameter);
-			}
-		}
-
-		for(UNiagaraHierarchyItemBase* Child : HierarchyCategory->GetChildren())
-		{
-			if(UNiagaraHierarchyCategory* Category = Cast<UNiagaraHierarchyCategory>(Child))
-			{
-				ContainedCategories.Add(Category);
-			}
-		}
-
-		for(const UNiagaraHierarchyCategory* Category : ContainedCategories)
-		{
-			if(Category->DoesOneChildExist<UNiagaraHierarchyUserParameter>(true))
-			{
-				ChildrenBuilder.AddCustomBuilder(MakeShared<FNiagaraUserParameterCategoryBuilder>(
-					Owner.Get(),Category, Category->GetCategoryName(), NodeBuilder));
-			}
-		}
-		
-		for(const UNiagaraHierarchyUserParameter* HierarchyUserParameter : ContainedParameters)
-		{
-			NodeBuilder.GenerateRowForUserParameter(ChildrenBuilder, HierarchyUserParameter->GetUserParameter(), Owner.Get());
-		}
-	}
-
-	TWeakObjectPtr<UNiagaraComponent> GetComponent() const { return Cast<UNiagaraComponent>(Owner); }
-	TWeakObjectPtr<UNiagaraSystem> GetSystem() const { return Cast<UNiagaraSystem>(Owner); }
-
-private:
-	TWeakObjectPtr<UObject> Owner;
-	TWeakObjectPtr<const UNiagaraHierarchyCategory> HierarchyCategory;
-	FName CustomBuilderRowName;
-	FNiagaraUserParameterNodeBuilder& NodeBuilder;
-};
-
-void FNiagaraUserParameterNodeBuilder::GenerateUserParameterRows(IDetailChildrenBuilder& ChildrenBuilder, UNiagaraHierarchyRoot& UserParameterHierarchyRoot, UObject* Owner)
-{
-	UNiagaraSystem* SystemAsset = nullptr;
-
-	if(UNiagaraSystem* System = Cast<UNiagaraSystem>(Owner))
-	{
-		SystemAsset = System;
-	}
-	else if(UNiagaraComponent* Component = Cast<UNiagaraComponent>(Owner))
-	{
-		SystemAsset = Component->GetAsset();
-	}
-
-	if(SystemAsset == nullptr)
-	{
-		return;
-	}
-
-	TArray<const UNiagaraHierarchySection*> OrderedSections;
-	TMap<const UNiagaraHierarchySection*, TArray<const UNiagaraHierarchyCategory*>> SectionCategoryMap;
-
-	UNiagaraHierarchyRoot* Root = &UserParameterHierarchyRoot;
-	for(const UNiagaraHierarchyItemBase* Child : Root->GetChildren())
-	{
-		if(const UNiagaraHierarchyCategory* HierarchyCategory = Cast<UNiagaraHierarchyCategory>(Child))
-		{
-			if(HierarchyCategory->GetSection() != nullptr)
-			{
-				SectionCategoryMap.FindOrAdd(HierarchyCategory->GetSection()).Add(HierarchyCategory);	
-			}
-		}
-	}
-
-	// we create a custom row for user param sections here, but only if we have at least one section specified
-	if(SectionCategoryMap.Num() > 0)
-	{
-		TSharedRef<SHorizontalBox> SectionsBox = SNew(SHorizontalBox);
-
-		// maps don't guarantee order, so we iterate over the original section data instead
-		for(auto& Section : Root->GetSectionData())
-		{
-			// if the map doesn't contain at the entry, it means there are 0 categories for that section, so we skip it
-			if(!SectionCategoryMap.Contains(Section))
-			{
-				continue;
-			}
-			
-			TArray<const UNiagaraHierarchyCategory*>& HierarchyCategories = SectionCategoryMap[Section];
-			
-			if(HierarchyCategories.Num() > 0)
-			{
-				// we only want to add a section if any of its categories contain at least one user parameter
-				bool bDoesUserParamExist = false;
-				for(const UNiagaraHierarchyCategory* Category : HierarchyCategories)
-				{
-					bDoesUserParamExist |= Category->DoesOneChildExist<UNiagaraHierarchyUserParameter>(true);
-
-					if(bDoesUserParamExist)
-					{
-						break;
-					}
-				}
-
-				if(bDoesUserParamExist)
-				{
-					SectionsBox->AddSlot()
-					.AutoWidth()
-					.Padding(FMargin(2.f))
-					[
-						SNew(SBox)
-						.Padding(FMargin(0.f, 4.f, 0.f, 0.f))
-						[
-							SNew(SCheckBox)
-							.Style(FAppStyle::Get(), "DetailsView.SectionButton")
-							.OnCheckStateChanged_Lambda([this, Section](ECheckBoxState NewState)
-							{
-								ActiveSection = Section;
-								OnRebuildChildren.ExecuteIfBound();
-							})
-						.IsChecked_Lambda([this, Section]()
-							{
-								return ActiveSection == Section ? ECheckBoxState::Checked : ECheckBoxState::Unchecked;
-							})
-							[
-								SNew(STextBlock)
-								.TextStyle(FAppStyle::Get(), "SmallText")
-								.Text_UObject(Section, &UNiagaraHierarchySection::GetSectionNameAsText)
-								.ToolTipText_UObject(Section, &UNiagaraHierarchySection::GetTooltip)
-							]
-						]
-					];
-				}
-			}
-		}
-
-		// if we have at least one custom section, we add a default "All" section
-		if(SectionsBox->GetChildren()->Num() > 0)
-		{
-			SectionsBox->AddSlot()
-			.AutoWidth()
-			.Padding(FMargin(2.f))
-			[
-				SNew(SBox)
-				.Padding(FMargin(0.f, 4.f, 0.f, 0.f))
-				[
-					SNew(SCheckBox)
-					.Style(FAppStyle::Get(), "DetailsView.SectionButton")
-					.OnCheckStateChanged_Lambda([this](ECheckBoxState NewState)
-					{
-						ActiveSection = nullptr;
-						OnRebuildChildren.ExecuteIfBound();
-					})
-					.IsChecked_Lambda([this]()
-					{
-						return ActiveSection == nullptr ? ECheckBoxState::Checked : ECheckBoxState::Unchecked;
-					})
-					[
-						SNew(STextBlock)
-						.TextStyle(FAppStyle::Get(), "SmallText")
-						.Text(FText::FromString("All"))
-					]
-				]
-			];
-			
-			FDetailWidgetRow& Row = ChildrenBuilder.AddCustomRow(FText::FromString("Parameter Sections"));
-			Row.WholeRowContent()
-			[
-				SectionsBox
-			];
-		}
-	}
-	
-	// we add the categories now if the active section is set to all (nullptr) or if the section it belongs to is active
-	for(const UNiagaraHierarchyItemBase* Child : Root->GetChildren())
-	{		
-		if(const UNiagaraHierarchyCategory* HierarchyCategory = Cast<UNiagaraHierarchyCategory>(Child))
-		{
-			if((ActiveSection == nullptr || HierarchyCategory->GetSection() == ActiveSection) && HierarchyCategory->DoesOneChildExist<UNiagaraHierarchyUserParameter>(true))
-			{
-				ChildrenBuilder.AddCustomBuilder(MakeShared<FNiagaraUserParameterCategoryBuilder>(
-					Owner, HierarchyCategory, HierarchyCategory->GetCategoryName(), *this));
-			
-			}
-		}
-		else if(const UNiagaraHierarchyUserParameter* UserParameter = Cast<UNiagaraHierarchyUserParameter>(Child))
-		{
-			// we only want to display parameters within the root in the all section
-			if(ActiveSection == nullptr)
-			{
-				GenerateRowForUserParameter(ChildrenBuilder, UserParameter->GetUserParameter(), Owner);
-			}
-		}
-	}
-
-	// at last we add all parameters that haven't been setup at all, if the active section is set to "All"
-	if(ActiveSection == nullptr)
-	{
-		TArray<FNiagaraVariable> LeftoverUserParameters;
-		SystemAsset->GetExposedParameters().GetUserParameters(LeftoverUserParameters);
-
-		TArray<UNiagaraHierarchyUserParameter*> HierarchyUserParameters;
-		Root->GetChildrenOfType<>(HierarchyUserParameters, true);
-
-		TArray<FNiagaraVariable> RedirectedUserParameters;
-		for(FNiagaraVariable& Parameter : LeftoverUserParameters)
-		{
-			SystemAsset->GetExposedParameters().RedirectUserVariable(Parameter);
-		}
-		
-		for(UNiagaraHierarchyUserParameter* HierarchyUserParameter : HierarchyUserParameters)
-		{
-			if(LeftoverUserParameters.Contains(HierarchyUserParameter->GetUserParameter()))
-			{
-				LeftoverUserParameters.Remove(HierarchyUserParameter->GetUserParameter());
-			}
-		}
-
-		for(FNiagaraVariable& LeftoverUserParameter : LeftoverUserParameters)
-		{			
-			GenerateRowForUserParameter(ChildrenBuilder, LeftoverUserParameter, Owner);
-		}
-	}
-}
-
-void FNiagaraComponentUserParametersNodeBuilder::GenerateChildContent(IDetailChildrenBuilder& ChildrenBuilder)
-{
-	if(Component.IsValid())
-	{
-		if (bDelegatesInitialized == false)
-		{
-			Component->OnSynchronizedWithAssetParameters().AddSP(this, &FNiagaraComponentUserParametersNodeBuilder::Rebuild);
-			Component->GetOverrideParameters().AddOnChangedHandler(
-				FNiagaraParameterStore::FOnChanged::FDelegate::CreateSP(this, &FNiagaraComponentUserParametersNodeBuilder::ParameterValueChanged));
-
-			RegisterRebuildOnHierarchyChanged();
-			bDelegatesInitialized = true;
-		}
-
-		ParameterProxies.Reset();
-
-		UNiagaraSystem* SystemAsset = Component->GetAsset();
-		if (SystemAsset == nullptr)
-		{
-			return;
-		}
-			
-		TArray<FNiagaraVariable> UserParameters;
-		SystemAsset->GetExposedParameters().GetUserParameters(UserParameters);
-			
-		ParameterProxies.Reserve(UserParameters.Num());
-
-		ParameterToDisplayStruct.Empty();
-
-		GenerateUserParameterRows(ChildrenBuilder, *Cast<UNiagaraSystemEditorData>(SystemAsset->GetEditorData())->UserParameterHierarchy.Get(),Component.Get());
-	}
-}
-
-void FNiagaraComponentUserParametersNodeBuilder::RegisterRebuildOnHierarchyChanged()
-{
-	if(Component.IsValid())
-	{
-		if(UNiagaraSystem* Asset = Component->GetAsset())
-		{
-			TSharedPtr<FNiagaraSystemViewModel> SystemViewModel = TNiagaraViewModelManager<UNiagaraSystem, FNiagaraSystemViewModel>::GetExistingViewModelForObject(Asset);
-			if(SystemViewModel.IsValid())
-			{
-				SystemViewModel->GetUserParametersHierarchyViewModel()->OnHierarchyChanged().RemoveAll(this);
-				SystemViewModel->GetUserParametersHierarchyViewModel()->OnHierarchyChanged().Add(UNiagaraHierarchyViewModelBase::FOnHierarchyChanged::FDelegate::CreateSP(this, &FNiagaraComponentUserParametersNodeBuilder::Rebuild));
-			}
-		}
-	}
-}
-
-void FNiagaraComponentUserParametersNodeBuilder::ParameterValueChanged()
-{
-	if (Component.IsValid())
-	{
-		const FNiagaraParameterStore& OverrideParameters = Component->GetOverrideParameters();
-		TArray<FNiagaraVariable> UserParameters;
-		OverrideParameters.GetParameters(UserParameters);
-		for (const FNiagaraVariable& UserParameter : UserParameters)
-		{
-			if (UserParameter.IsUObject() == false)
-			{
-				TWeakPtr<FStructOnScope>* DisplayStructPtr = ParameterToDisplayStruct.Find(UserParameter);
-				if (DisplayStructPtr != nullptr && DisplayStructPtr->IsValid())
-				{
-					TSharedPtr<FStructOnScope> DisplayStruct = DisplayStructPtr->Pin();
-					if (UserParameter.GetType() == FNiagaraTypeDefinition::GetPositionDef())
-					{
-						const FVector* PositionValue = OverrideParameters.GetPositionParameterValue(UserParameter.GetName());
-						if (ensureMsgf(PositionValue != nullptr, TEXT("Position user parameter %s was missing it's position data.  Path: %s"),
-							*UserParameter.GetName().ToString(), *Component->GetPathName()))
-						{
-							FMemory::Memcpy(DisplayStruct->GetStructMemory(), PositionValue, DisplayStruct->GetStruct()->GetStructureSize());
-						}
-					}
-					else
-					{
-						OverrideParameters.CopyParameterData(UserParameter, DisplayStruct->GetStructMemory());
-					}
-				}
-			}
-		}
-	}
-}
-
-FNiagaraSystemUserParameterBuilder::FNiagaraSystemUserParameterBuilder(TSharedPtr<FNiagaraSystemViewModel> InSystemViewModel, FName InCustomBuilderRowName)
-{
-	CustomBuilderRowName = InCustomBuilderRowName;
-	System = &InSystemViewModel->GetSystem();
-	SystemViewModel = InSystemViewModel;
-	bDelegatesInitialized = false;
-}
-
-TSharedPtr<SWidget> FNiagaraSystemUserParameterBuilder::GenerateCustomNameWidget(FNiagaraVariable UserParameter)
-{	
-	TSharedRef<SNiagaraParameterNameTextBlock> ParameterName = SNew(SNiagaraParameterNameTextBlock)
-		// if this is specified, we avoid the default behavior, which is "when clicked when the widget had keyboard focus already, enter editing mode"
-		// we don't want to be able to rename by clicking, only by context menu
-		.IsSelected_Lambda([]()
-		{
-			return false;
-		})
-		.ParameterText(FText::FromName(UserParameter.GetName()))
-		.OnDragDetected(this, &FNiagaraSystemUserParameterBuilder::GenerateParameterDragDropOp, UserParameter)
-		.OnVerifyTextChanged_Lambda([this, UserParameter](const FText& InNewName, FText& OutErrorMessage)
-		{
-			if(InNewName.IsEmptyOrWhitespace())
-			{
-				OutErrorMessage = LOCTEXT("UserParameterNameCantBeEmpty", "Can not accept empty name");
-				return false;
-			}
-			
-			FNiagaraVariable NewUserParameter = UserParameter;
-			NewUserParameter.SetName(FName(InNewName.ToString()));
-			if(UserParameter == NewUserParameter)
-			{
-				// we return true, but renaming won't do anything if the name is the same as before
-				return true;
-			}
-			else
-			{
-				TArray<FNiagaraVariable> ExistingParameters;
-				UserParamToWidgetMap.GenerateKeyArray(ExistingParameters);
-
-				bool bNameExistsAlready = false;
-				for(FNiagaraVariable& ExistingUserParameter : ExistingParameters)
-				{
-					if(ExistingUserParameter.GetName().IsEqual(FName(InNewName.ToString())))
-					{
-						bNameExistsAlready = true;
-						break;
-					}
-				}
-
-				if(bNameExistsAlready)
-				{
-					OutErrorMessage = LOCTEXT("UserParameterAlreadyExists", "A User Parameter with the same name already exists");
-					return false;					
-				}
-
-				return true;
-			}
-		})
-		.OnTextCommitted_Lambda([this, UserParameter](const FText& InText, ETextCommit::Type Type)
-		{
-			RenameParameter(UserParameter, FName(InText.ToString()));
-		});
-	
-	UserParamToWidgetMap.Add(UserParameter, ParameterName);
-	
-	return ParameterName;
-}
-
-void FNiagaraSystemUserParameterBuilder::GenerateChildContent(IDetailChildrenBuilder& ChildrenBuilder)
-{
-	check(System.IsValid());
-
-	if (System.IsValid())
-	{
-		if(bDelegatesInitialized == false)
-		{
-			System->GetExposedParameters().OnStructureChanged().Add(FNiagaraParameterStore::FOnStructureChanged::FDelegate::CreateSP(this, &FNiagaraSystemUserParameterBuilder::Rebuild));
-
-			SystemViewModel.Pin()->GetUserParametersHierarchyViewModel()->OnHierarchyChanged().RemoveAll(this);
-			SystemViewModel.Pin()->GetUserParametersHierarchyViewModel()->OnHierarchyChanged().Add(UNiagaraHierarchyViewModelBase::FOnHierarchyChanged::FDelegate::CreateSP(this, &FNiagaraSystemUserParameterBuilder::Rebuild));
-
-			TSharedPtr<FNiagaraUserParameterPanelViewModel> UserParameterPanelViewModel = SystemViewModel.Pin()->GetUserParameterPanelViewModel();
-			if(ensure(UserParameterPanelViewModel.IsValid()))
-			{
-				UserParameterPanelViewModel->OnRefreshRequested().BindSP(this, &FNiagaraSystemUserParameterBuilder::Rebuild);
-				UserParameterPanelViewModel->OnParameterAdded().BindSP(this, &FNiagaraSystemUserParameterBuilder::OnParameterAdded);
-			}
-			
-			bDelegatesInitialized = true;
-		}
-			
-		ParameterProxies.Reset();
-			
-		TArray<FNiagaraVariable> UserParameters;
-		System->GetExposedParameters().GetUserParameters(UserParameters);
-		ParameterProxies.Reserve(UserParameters.Num());
-
-		ParameterToDisplayStruct.Empty();
-
-		UNiagaraHierarchyRoot* Root = Cast<UNiagaraSystemEditorData>(System->GetEditorData())->UserParameterHierarchy;
-		GenerateUserParameterRows(ChildrenBuilder, *Root, System.Get());
-	}
-}
-
-void FNiagaraSystemUserParameterBuilder::AddCustomMenuActionsForParameter(FDetailWidgetRow& WidgetRow, FNiagaraVariable UserParameter)
-{
-	FNiagaraUserParameterNodeBuilder::AddCustomMenuActionsForParameter(WidgetRow, UserParameter);
-	
-	FUIAction RenameAction(FExecuteAction::CreateSP(this, &FNiagaraSystemUserParameterBuilder::RequestRename, UserParameter));
-	WidgetRow.AddCustomContextMenuAction(RenameAction, LOCTEXT("RenameParameterAction", "Rename"), LOCTEXT("RenameParameterActionTooltip", "Rename this user parameter"));
-
-	FUIAction DeleteAction(FExecuteAction::CreateSP(this, &FNiagaraSystemUserParameterBuilder::DeleteParameter, UserParameter));
-	WidgetRow.AddCustomContextMenuAction(DeleteAction, LOCTEXT("DeleteParameterAction", "Delete"), LOCTEXT("DeleteParameterActionTooltip", "Delete this user parameter"));
-}
-
-TSharedRef<SWidget> FNiagaraSystemUserParameterBuilder::GetAdditionalHeaderWidgets()
-{
-	if(!AdditionalHeaderWidgetsContainer.IsValid())
-	{
-		AdditionalHeaderWidgetsContainer = SNew(SBox)
-		.HAlign(HAlign_Right)
-		.VAlign(VAlign_Center)
-		.Padding(2.f)
-		[
-			SNew(SHorizontalBox)
-			+ SHorizontalBox::Slot()
-			.HAlign(HAlign_Center)
-			.VAlign(VAlign_Center)
-			.AutoWidth()
-			.Padding(2.f, 1.f)
-			[
-				SNew(SButton)
-				.OnClicked(this, &FNiagaraSystemUserParameterBuilder::SummonHierarchyEditor)
-				.Text(LOCTEXT("SummonUserParametersHierarchyButtonLabel", "Edit Hierarchy"))
-				.ButtonStyle(FAppStyle::Get(), "RoundButton")
-				.ToolTipText(LOCTEXT("SummonUserParametersHierarchyButtonTooltip", "Summon the Hierarchy Editor to add categories, sections and tooltips to User Parameters."))
-			]
-			+ SHorizontalBox::Slot()
-			.HAlign(HAlign_Center)
-			.VAlign(VAlign_Center)
-			.AutoWidth()
-			.Padding(2.f)
-			[
-				SAssignNew(AddParameterButton, SComboButton)
-				.HAlign(HAlign_Center)
-				.VAlign(VAlign_Center)
-				.ButtonStyle(FAppStyle::Get(), "RoundButton")
-				.ForegroundColor(FAppStyle::GetSlateColor("DefaultForeground"))
-				.ContentPadding(FMargin(2, 3.5))
-				.ToolTipText(LOCTEXT("AddUserParameterButtonTooltip", "Add a new User Parameter to this system.\nUser Parameters exist once for the entire system and their values can be changed at runtime from Blueprints, C++ or the sequencer, among other systems."))
-				.OnGetMenuContent(this, &FNiagaraSystemUserParameterBuilder::GetAddParameterMenu)
-				.OnComboBoxOpened_Lambda([this]
-				{
-					AddParameterButton->SetMenuContentWidgetToFocus(AddParameterMenu->GetSearchBox());
-				})
-				.HasDownArrow(false)
-				.ButtonContent()
-				[
-					SNew(SImage)
-					.Image(FAppStyle::GetBrush("Plus"))
-				]
-			]
-		];
-	}
-	
-	return AdditionalHeaderWidgetsContainer.ToSharedRef(); 
-}
-
-TSharedRef<SWidget> FNiagaraSystemUserParameterBuilder::GetAddParameterMenu()
-{
-	FNiagaraParameterPanelCategory UserCategory = GetDefault<UNiagaraEditorSettings>()->GetMetaDataForNamespaces({FNiagaraConstants::UserNamespace});
-	AddParameterMenu = SNew(SNiagaraAddParameterFromPanelMenu)
-		.Graphs(SystemViewModel.Pin()->GetParameterPanelViewModel()->GetEditableGraphsConst())
-		.OnNewParameterRequested(this, &FNiagaraSystemUserParameterBuilder::AddParameter)
-		.OnAllowMakeType(this, &FNiagaraSystemUserParameterBuilder::OnAllowMakeType)
-		.NamespaceId(UserCategory.NamespaceMetaData.GetGuid())
-		.ShowNamespaceCategory(false)
-		.ShowGraphParameters(false)
-		.AutoExpandMenu(false);
-
-	return AddParameterMenu.ToSharedRef();
-}
-
-void FNiagaraSystemUserParameterBuilder::AddParameter(FNiagaraVariable NewParameter) const
-{
-	FNiagaraParameterPanelCategory UserCategory = GetDefault<UNiagaraEditorSettings>()->GetMetaDataForNamespaces({FNiagaraConstants::UserNamespace});
-	FNiagaraEditorUtilities::AddParameter(NewParameter, SystemViewModel.Pin()->GetSystem().GetExposedParameters(), SystemViewModel.Pin()->GetSystem(), nullptr);
-	SystemViewModel.Pin()->GetUserParameterPanelViewModel()->OnParameterAdded().Execute(NewParameter);
-}
-
-bool FNiagaraSystemUserParameterBuilder::OnAllowMakeType(const FNiagaraTypeDefinition& InType) const
-{
-	FNiagaraNamespaceMetadata UserNameSpaceMetaData = GetDefault<UNiagaraEditorSettings>()->GetMetaDataForNamespaces({FNiagaraConstants::UserNamespace});
-	FNiagaraParameterPanelCategory Category(UserNameSpaceMetaData);
-	return SystemViewModel.Pin()->GetParameterPanelViewModel()->CanAddType(InType, Category);
-}
-
-void FNiagaraSystemUserParameterBuilder::ParameterValueChanged()
-{
-	if (System.IsValid())
-	{
-		const FNiagaraParameterStore& ExposedParameters = System->GetExposedParameters();
-		TArray<FNiagaraVariable> UserParameters;
-		ExposedParameters.GetParameters(UserParameters);
-		for (const FNiagaraVariable& UserParameter : UserParameters)
-		{
-			if (UserParameter.IsUObject() == false)
-			{
-				TWeakPtr<FStructOnScope>* DisplayStructPtr = ParameterToDisplayStruct.Find(UserParameter);
-				if (DisplayStructPtr != nullptr && DisplayStructPtr->IsValid())
-				{
-					TSharedPtr<FStructOnScope> DisplayStruct = DisplayStructPtr->Pin();
-					if (UserParameter.GetType() == FNiagaraTypeDefinition::GetPositionDef())
-					{
-						FMemory::Memcpy(DisplayStruct->GetStructMemory(), ExposedParameters.GetPositionParameterValue(UserParameter.GetName()), UserParameter.GetSizeInBytes());
-					}
-					else
-					{
-						ExposedParameters.CopyParameterData(UserParameter, DisplayStruct->GetStructMemory());
-					}
-				}
-			}
-		}
-	}
-}
-
-FReply FNiagaraSystemUserParameterBuilder::SummonHierarchyEditor()
-{
-	if(TSharedPtr<FNiagaraSystemViewModel> SystemViewModelPinned = SystemViewModel.Pin())
-	{
-		SystemViewModelPinned->FocusTab(FNiagaraSystemToolkitModeBase::UserParametersHierarchyTabID, true);
-		return FReply::Handled();
-	}
-
-	return FReply::Unhandled();
-}
-
-void FNiagaraSystemUserParameterBuilder::DeleteParameter(FNiagaraVariable UserParameter) const
-{
-	FScopedTransaction Transaction(LOCTEXT("DeleteParameterTransaction", "Deleted user parameter"));
-	System->Modify();
-	
-	if(System->GetExposedParameters().IndexOf(UserParameter) != INDEX_NONE)
-	{
-		System->GetExposedParameters().RemoveParameter(UserParameter);		
-	}
-
-	Cast<UNiagaraSystemEditorData>(System->GetEditorData())->UserParameterHierarchy->Modify();
-	System->HandleVariableRemoved(UserParameter, true);
-}
-
-void FNiagaraSystemUserParameterBuilder::SelectAllSection()
-{
-	if(ActiveSection != nullptr)
-	{
-		ActiveSection = nullptr;
-		OnRebuildChildren.Execute();
-	}
-}
-
-void FNiagaraSystemUserParameterBuilder::OnParameterAdded(FNiagaraVariable UserParameter)
-{
-	// when adding a new parameter it defaults to the 'All' section, so we select it in order to visualize it before renaming
-	SelectAllSection();
-	RequestRename(UserParameter);
-}
-
-void FNiagaraSystemUserParameterBuilder::RequestRename(FNiagaraVariable UserParameter)
-{
-	SelectedParameter = UserParameter;
-	
-	// we add a timer for next frame, as requesting a rename on a parameter that has just been created will not preselect the entire text. Waiting a frame fixes that.
-	UserParamToWidgetMap[UserParameter]->RegisterActiveTimer(0.f, FWidgetActiveTimerDelegate::CreateLambda([this, UserParameter](double CurrentTime, float DeltaTime)
-	{
-		UserParamToWidgetMap[UserParameter]->EnterEditingMode();
-		return EActiveTimerReturnType::Stop;
-	}));
-}
-
-void FNiagaraSystemUserParameterBuilder::RenameParameter(FNiagaraVariable UserParameter, FName NewName)
-{
-	if(UserParameter.GetName() == NewName)
-	{
-		return;
-	}
-	
-	if(System->GetExposedParameters().IndexOf(UserParameter) != INDEX_NONE)
-	{		
-		TArray<FNiagaraVariable> Parameters;
-		System->GetExposedParameters().GetParameters(Parameters);
-
-		TSet<FName> Names;
-		for(FNiagaraVariable& Parameter : Parameters)
-		{
-			Names.Add(Parameter.GetName());
-		}
-		
-		FName UniqueName = FNiagaraUtilities::GetUniqueName(NewName, Names);
-
-		FNiagaraVariable VariableWithNewName(UserParameter);
-		VariableWithNewName.SetName(UniqueName);
-		FNiagaraUserRedirectionParameterStore::MakeUserVariable(VariableWithNewName);
-		
-		SystemViewModel.Pin()->RenameParameter(UserParameter, NewName, ENiagaraGetGraphParameterReferencesMode::AllGraphs);
-	}
-}
-
-FReply FNiagaraSystemUserParameterBuilder::GenerateParameterDragDropOp(const FGeometry& Geometry, const FPointerEvent& MouseEvent, FNiagaraVariable UserParameter) const
-{
-	TSharedPtr<FEdGraphSchemaAction> ParameterDragAction = MakeShared<FNiagaraParameterAction>(UserParameter, FText::GetEmpty(), FText::GetEmpty(), FText::GetEmpty(), 0, FText(), nullptr, 0/*SectionID*/);
-	TSharedRef<FNiagaraParameterDragOperation> ParameterDragOperation = MakeShared<FNiagaraParameterDragOperation>(ParameterDragAction);
-	ParameterDragOperation->SetupDefaults();
-	ParameterDragOperation->Construct();
-	return FReply::Handled().BeginDragDrop(ParameterDragOperation);
-}
-
 TSharedRef<IDetailCustomization> FNiagaraComponentDetails::MakeInstance()
 {
 	return MakeShareable(new FNiagaraComponentDetails);
diff --git a/Engine/Plugins/FX/Niagara/Source/NiagaraEditor/Private/Customizations/NiagaraComponentDetails.h b/Engine/Plugins/FX/Niagara/Source/NiagaraEditor/Private/Customizations/NiagaraComponentDetails.h
--- a/Engine/Plugins/FX/Niagara/Source/NiagaraEditor/Private/Customizations/NiagaraComponentDetails.h
+++ b/Engine/Plugins/FX/Niagara/Source/NiagaraEditor/Private/Customizations/NiagaraComponentDetails.h
@@ -4,22 +4,13 @@
 
 #include "CoreMinimal.h"
 #include "IDetailCustomization.h"
-#include "IDetailCustomNodeBuilder.h"
 #include "Input/Reply.h"
 #include "NiagaraComponent.h"
 #include "NiagaraSimCacheCapture.h"
 #include "ViewModels/TNiagaraViewModelManager.h"
 #include "ViewModels/HierarchyEditor/NiagaraHierarchyViewModelBase.h"
-#include "ViewModels/HierarchyEditor/NiagaraUserParametersHierarchyViewModel.h"
 #include "NiagaraComponentDetails.generated.h"
 
-class IDetailLayoutBuilder;
-class IPropertyHandle;
-class FNiagaraParameterViewModelCustomDetails;
-class INiagaraParameterViewModel;
-class SNiagaraParameterEditor;
-class FNiagaraParameterViewModelCustomDetails;
-class FNiagaraSimCacheCapture;
 class UNiagaraSystem;
 
 USTRUCT()
@@ -31,6 +22,16 @@ struct FNiagaraEnumToByteHelper
 	uint8 Value = 0;
 };
 
+UCLASS()
+class UNiagaraObjectAssetHelper : public UObject
+{
+	GENERATED_BODY()
+
+	/** We are customizing the instance property metadata for this to restrict the allowed classes in the UI */
+	UPROPERTY(EditAnywhere, Category = Asset)
+	FSoftObjectPath Path;
+};
+
 class FNiagaraComponentDetails : public IDetailCustomization
 {
 public:
@@ -67,167 +68,3 @@ public:
 protected:
 	TWeakObjectPtr<UNiagaraSystem> System;
 };
-
-/** A base class for generating rows based on a given hierarchy root that contains sections, categories and user parameters. */
-class FNiagaraUserParameterNodeBuilder : public IDetailCustomNodeBuilder, public TSharedFromThis<FNiagaraUserParameterNodeBuilder>, FSelfRegisteringEditorUndoClient
-{
-public:
-	virtual void GenerateHeaderRowContent(FDetailWidgetRow& NodeRow) override {}
-	virtual void Tick(float DeltaTime) override {}
-	virtual bool RequiresTick() const override { return false; }
-	virtual bool InitiallyCollapsed() const override { return false; }
-	virtual FName GetName() const  override
-	{
-		return CustomBuilderRowName;
-	}
-	
-	void Rebuild() const
-	{
-		OnRebuildChildren.ExecuteIfBound();
-	}
-	
-	virtual void SetOnRebuildChildren(FSimpleDelegate InOnRegenerateChildren) override
-	{
-		OnRebuildChildren = InOnRegenerateChildren;
-	}
-
-	virtual void PostUndo(bool bSuccess) override { Rebuild(); }
-	virtual void PostRedo(bool bSuccess) override { Rebuild(); }
-
-	/** Allows the addition of right-click context actions for a specific user parameter */
-	virtual void AddCustomMenuActionsForParameter(FDetailWidgetRow& WidgetRow, FNiagaraVariable UserParameter) {}
-	/** Allows the creation of a custom name widget for a user parameter. */
-	virtual TSharedPtr<SWidget> GenerateCustomNameWidget(FNiagaraVariable UserParameter) { return nullptr; }
-
-	/** Generates all rows for the node builder. Includes setting up sections, categories and user parameters */
-	void GenerateUserParameterRows(IDetailChildrenBuilder& ChildrenBuilder, UNiagaraHierarchyRoot& UserParameterHierarchyRoot, UObject* Owner);
-	/** Generates a single row for a specific user parameter. */
-	void GenerateRowForUserParameter(IDetailChildrenBuilder& ChildrenBuilder, FNiagaraVariable UserParameter, UObject* Owner);
-	
-	void OnActiveSectionChanged(ECheckBoxState State, const UNiagaraHierarchySection* NewSelection)
-	{
-		ActiveSection = NewSelection;
-		OnRebuildChildren.Execute();
-	}
-
-	ECheckBoxState IsSectionActive(const UNiagaraHierarchySection* Section) const
-	{
-		return ActiveSection == Section ? ECheckBoxState::Checked : ECheckBoxState::Unchecked;
-	}
-	
-protected:
-	bool bDelegatesInitialized = false;
-	TArray<TSharedPtr<IPropertyHandle>> OverridePropertyHandles;
-	FSimpleDelegate OnRebuildChildren;
-	TArray<TSharedPtr<class FNiagaraParameterProxy>> ParameterProxies;
-	TMap<FNiagaraVariableBase, TWeakPtr<FStructOnScope>> ParameterToDisplayStruct;
-	FName CustomBuilderRowName;
-	const UNiagaraHierarchySection* ActiveSection = nullptr;
-	TArray<UNiagaraHierarchySection*> AvailableSections;
-};
-
-/** A specialized class to generate user parameters for a component. */
-class FNiagaraComponentUserParametersNodeBuilder : public FNiagaraUserParameterNodeBuilder
-{
-public:
-	FNiagaraComponentUserParametersNodeBuilder(UNiagaraComponent* InComponent, TArray<TSharedPtr<IPropertyHandle>> InOverridePropertyHandles, FName InCustomBuilderRowName) 
-	{
-		OverridePropertyHandles = InOverridePropertyHandles;
-		CustomBuilderRowName = InCustomBuilderRowName;
-
-		Component = InComponent;
-		bDelegatesInitialized = false;
-	}
-
-	virtual ~FNiagaraComponentUserParametersNodeBuilder() override
-	{
-		if (Component.IsValid() && bDelegatesInitialized)
-		{
-			Component->OnSynchronizedWithAssetParameters().RemoveAll(this);
-			Component->GetOverrideParameters().RemoveAllOnChangedHandlers(this);
-
-			if(UNiagaraSystem* System = Component->GetAsset())
-			{
-				TSharedPtr<FNiagaraSystemViewModel> SystemViewModel = TNiagaraViewModelManager<UNiagaraSystem, FNiagaraSystemViewModel>::GetExistingViewModelForObject(System);
-				if(SystemViewModel.IsValid() && SystemViewModel->GetUserParametersHierarchyViewModel())
-				{
-					SystemViewModel->GetUserParametersHierarchyViewModel()->OnHierarchyChanged().RemoveAll(this);
-				}
-			}
-		}
-	}
-	
-	virtual void GenerateChildContent(IDetailChildrenBuilder& ChildrenBuilder) override;
-
-private:
-	void RegisterRebuildOnHierarchyChanged();
-
-	void ParameterValueChanged();
-
-private:
-	TWeakObjectPtr<UNiagaraComponent> Component;
-};
-
-/** A specialized class to generate user parameters for a system asset's editor. Due to reliance on the system view model this should only be used within the asset editor. */
-class FNiagaraSystemUserParameterBuilder : public FNiagaraUserParameterNodeBuilder
-{
-public:
-	FNiagaraSystemUserParameterBuilder(TSharedPtr<FNiagaraSystemViewModel> InSystemViewModel, FName InCustomBuilderRowName);
-
-	virtual ~FNiagaraSystemUserParameterBuilder() override
-	{
-		if (System.IsValid() && bDelegatesInitialized)
-		{
-			System->GetExposedParameters().OnStructureChanged().RemoveAll(this);
-
-			if(SystemViewModel.IsValid() && SystemViewModel.Pin()->GetUserParameterPanelViewModel().IsValid())
-			{
-				if (TSharedPtr<FNiagaraUserParameterPanelViewModel> PanelViewModel = SystemViewModel.Pin()->GetUserParameterPanelViewModel())
-				{
-					PanelViewModel->OnRefreshRequested().Unbind();
-				}
-			}
-			
-			if(SystemViewModel.IsValid() && SystemViewModel.Pin()->GetUserParametersHierarchyViewModel())
-			{
-				if (UNiagaraUserParametersHierarchyViewModel* HierarchyViewModel = SystemViewModel.Pin()->GetUserParametersHierarchyViewModel())
-				{
-					HierarchyViewModel->OnHierarchyChanged().RemoveAll(this);
-				}
-			}
-		}
-	}
-
-	virtual void GenerateChildContent(IDetailChildrenBuilder& ChildrenBuilder) override;
-
-	/** Instead of pure text, we use the parameter widgets for display in a system asset. */
-	virtual TSharedPtr<SWidget> GenerateCustomNameWidget(FNiagaraVariable UserParameter) override;
-	
-	/** We want to add rename & delete actions within a system asset */
-	virtual void AddCustomMenuActionsForParameter(FDetailWidgetRow& WidgetRow, FNiagaraVariable UserParameter) override;
-
-	TSharedRef<SWidget> GetAdditionalHeaderWidgets();
-private:
-	TSharedRef<SWidget> GetAddParameterMenu();
-	void AddParameter(FNiagaraVariable NewParameter) const;
-	bool OnAllowMakeType(const FNiagaraTypeDefinition& InType) const;
-	
-	void ParameterValueChanged();
-
-	FReply SummonHierarchyEditor();
-	void DeleteParameter(FNiagaraVariable UserParameter) const;
-	void SelectAllSection();
-	void OnParameterAdded(FNiagaraVariable UserParameter);
-	void RequestRename(FNiagaraVariable UserParameter);
-	void RenameParameter(FNiagaraVariable UserParameter, FName NewName);
-
-	FReply GenerateParameterDragDropOp(const FGeometry& Geometry, const FPointerEvent& MouseEvent, FNiagaraVariable UserParameter) const;
-private:
-	TWeakObjectPtr<UNiagaraSystem> System;
-	TWeakPtr<FNiagaraSystemViewModel> SystemViewModel;
-	TOptional<FNiagaraVariable> SelectedParameter;
-	TMap<FNiagaraVariable, TSharedPtr<class SNiagaraParameterNameTextBlock>> UserParamToWidgetMap;
-	TSharedPtr<SWidget> AdditionalHeaderWidgetsContainer;
-	TSharedPtr<class SComboButton> AddParameterButton;
-	TSharedPtr<class SNiagaraAddParameterFromPanelMenu> AddParameterMenu;
-};
diff --git a/Engine/Plugins/FX/Niagara/Source/NiagaraEditor/Private/Customizations/UserParameters/NiagaraComponentUserParametersBuilder.cpp b/Engine/Plugins/FX/Niagara/Source/NiagaraEditor/Private/Customizations/UserParameters/NiagaraComponentUserParametersBuilder.cpp
new file mode 100644
--- /dev/null
+++ b/Engine/Plugins/FX/Niagara/Source/NiagaraEditor/Private/Customizations/UserParameters/NiagaraComponentUserParametersBuilder.cpp
@@ -0,0 +1,506 @@
+// Copyright Epic Games, Inc. All Rights Reserved.
+
+#include "NiagaraComponentUserParametersBuilder.h"
+#include "ViewModels/HierarchyEditor/NiagaraUserParametersHierarchyViewModel.h"
+#include "PropertyHandle.h"
+#include "NiagaraComponent.h"
+#include "DetailWidgetRow.h"
+#include "PropertyCustomizationHelpers.h"
+#include "NiagaraTypes.h"
+#include "ScopedTransaction.h"
+#include "NiagaraEditorStyle.h"
+#include "IDetailChildrenBuilder.h"
+#include "NiagaraConstants.h"
+#include "NiagaraScriptVariable.h"
+#include "NiagaraSettings.h"
+#include "NiagaraSystemEditorData.h"
+#include "NiagaraUserRedirectionParameterStore.h"
+#include "SNiagaraParameterEditor.h"
+#include "Widgets/SNiagaraDebugger.h"
+#include "Widgets/SNiagaraSystemUserParameters.h"
+#include "Widgets/Text/STextBlock.h"
+#include "Widgets/Input/SButton.h"
+
+#define LOCTEXT_NAMESPACE "NiagaraComponentUserParametersBuilder"
+
+// Proxy class to allow us to override values on the component that are not yet overridden.
+class FNiagaraParameterProxy : public TSharedFromThis<FNiagaraParameterProxy>
+{
+public:
+	FNiagaraParameterProxy(TWeakObjectPtr<UNiagaraComponent> InComponent, const FNiagaraVariableBase& InKey, const FNiagaraVariant& InValue, const FSimpleDelegate& InOnRebuild, TArray<TSharedPtr<IPropertyHandle>> InPropertyHandles)
+	{
+		bResettingToDefault = false;
+		Component = InComponent;
+		ParameterKey = InKey;
+		ParameterValue = InValue;
+		OnRebuild = InOnRebuild;
+		OverridePropertyHandles = InPropertyHandles;
+	}
+
+	void OnPropertyResetToDefault(TSharedPtr<IPropertyHandle> ResetPropertyHandle)
+	{
+		TWeakObjectPtr<UNiagaraComponent> RawComponent = GetComponent();
+		if (RawComponent != nullptr)
+		{
+			FScopedTransaction ScopedTransaction(FNiagaraUserParameterNodeBuilder::ResetUserParameterTransactionText);
+			RawComponent->Modify();
+
+			bResettingToDefault = true;
+
+			for (TSharedPtr<IPropertyHandle> PropertyHandle : OverridePropertyHandles)
+			{
+				PropertyHandle->NotifyPreChange();
+			}
+
+			RawComponent->RemoveParameterOverride(ParameterKey);
+
+			for (TSharedPtr<IPropertyHandle> PropertyHandle : OverridePropertyHandles)
+			{
+				PropertyHandle->NotifyPostChange(EPropertyChangeType::ValueSet);
+			}
+
+			OnRebuild.ExecuteIfBound();
+			bResettingToDefault = false;
+		}
+	}
+
+	void OnObjectAssetResetToDefault()
+	{
+		OnPropertyResetToDefault(nullptr);
+	}
+
+	bool IsPropertyResetToDefaultVisible(TSharedPtr<IPropertyHandle> ResetPropertyHandle) const
+	{
+		TWeakObjectPtr<UNiagaraComponent> RawComponent = GetComponent();
+		if (RawComponent.IsValid())
+		{
+			bool bHasParameterOverride = RawComponent->HasParameterOverride(ParameterKey);
+			return bHasParameterOverride;
+		}
+		return false;
+	}
+
+	bool IsObjectAssetResetToDefaultVisible() const
+	{
+		return IsPropertyResetToDefaultVisible(nullptr);		
+	}
+
+	FNiagaraVariant FindExistingOverride() const
+	{
+		TWeakObjectPtr<UNiagaraComponent> RawComponent = GetComponent();
+		if (RawComponent.IsValid())
+		{
+			return RawComponent->FindParameterOverride(ParameterKey);
+		}
+		return FNiagaraVariant();
+	}
+
+	void OnParameterPreChange()
+	{
+		TWeakObjectPtr<UNiagaraComponent> RawComponent = GetComponent();
+		if(RawComponent.IsValid())
+		{
+			RawComponent->Modify();
+
+			for (TSharedPtr<IPropertyHandle> PropertyHandle : OverridePropertyHandles)
+			{
+				PropertyHandle->NotifyPreChange();
+			}
+		}
+	}
+
+	/** This function handles value changes coming from the widget. It syncs local parameter editor values -> parameter value. */
+	void OnParameterEditorValueChanged()
+	{
+		if(!ensure(DisplayData.IsValid() && DisplayDataPropertyHandle.IsValid()))
+		{
+			return;
+		}
+
+		TSharedRef<FStructOnScope> ParameterEditorData = MakeShared<FStructOnScope>(ParameterKey.GetType().GetStruct());
+		ParameterEditor->UpdateStructFromInternalValue(ParameterEditorData);
+		
+		FScopedTransaction ScopedTransaction(FNiagaraUserParameterNodeBuilder::ChangedUserParameterTransactionText);
+		
+		bool bCancelTransaction = true;
+		// with 'BroadcastValueChangesPerKey' enabled, this function can get called even when just clicking in or out of the value without modifying it. We don't want that to create a transaction.
+		if(FMemory::Memcmp(DisplayData->GetStructMemory(), ParameterEditorData->GetStructMemory(), ParameterKey.GetType().GetSize()) != 0)
+		{
+			bCancelTransaction = false;
+		    Component->Modify();
+		}
+		
+		ParameterValue.SetBytes(ParameterEditorData->GetStructMemory(), ParameterKey.GetType().GetSize());
+
+		DisplayDataPropertyHandle->NotifyPreChange();
+		ParameterEditor->UpdateStructFromInternalValue(DisplayData.ToSharedRef());
+		// NotifyPostChange will trigger 'OnParameterPostChange', actually setting up the override on the component & calling NotifyPostChange on the override handles, triggering PostEditChangeProperty on the component
+		DisplayDataPropertyHandle->NotifyPostChange(EPropertyChangeType::ValueSet);
+		DisplayDataPropertyHandle->NotifyFinishedChangingProperties();
+		
+		if(bCancelTransaction == true)
+		{
+			ScopedTransaction.Cancel();
+		}
+	}
+	
+	void OnParameterPostChange(const FPropertyChangedEvent& PropertyChangedEvent)
+	{
+		if (bResettingToDefault)
+		{
+			return;
+		}
+	
+		TWeakObjectPtr<UNiagaraComponent> RawComponent = GetComponent();
+		if(RawComponent.IsValid())
+		{
+			// if this post change is called on a 'local user parameter' (not a DI or object), we have to sync up the display data -> parameter editor
+			if(DisplayData.IsValid())
+			{
+				if(ParameterEditor.IsValid())
+				{
+					// we have to make sure the internal parameter editor data is in sync with the display data (copy paste, reset etc. write to display data)
+					ParameterEditor->UpdateInternalValueFromStruct(DisplayData.ToSharedRef());
+				}
+				
+				ParameterValue.SetBytes(DisplayData->GetStructMemory(), ParameterValue.GetNumBytes());
+			}
+			
+			RawComponent->SetParameterOverride(ParameterKey, ParameterValue);
+	
+			for (TSharedPtr<IPropertyHandle> PropertyHandle : OverridePropertyHandles)
+			{
+				PropertyHandle->NotifyPostChange(PropertyChangedEvent.ChangeType);
+			}
+		}
+	}
+
+	void OnAssetSelectedFromPicker(const FAssetData& InAssetData)
+	{
+		TWeakObjectPtr<UNiagaraComponent> RawComponent = GetComponent();
+
+		if(RawComponent.IsValid())
+		{
+			UObject* Asset = InAssetData.GetAsset();
+			if (Asset == nullptr || Asset->GetClass()->IsChildOf(ParameterKey.GetType().GetClass()))
+			{
+				FScopedTransaction ScopedTransaction(LOCTEXT("ChangeAsset", "Change asset"));
+				RawComponent->Modify();
+				RawComponent->SetParameterOverride(ParameterKey, FNiagaraVariant(Asset));
+
+				for (TSharedPtr<IPropertyHandle> PropertyHandle : OverridePropertyHandles)
+				{
+					PropertyHandle->NotifyPostChange(EPropertyChangeType::ValueSet);
+				}
+			}
+		}
+	}
+
+	FString GetCurrentAssetPath() const
+	{
+		UObject* CurrentObject = nullptr;
+		
+		TWeakObjectPtr<UNiagaraComponent> RawComponent = GetComponent();
+
+		if(RawComponent.IsValid())
+		{
+			FNiagaraVariant CurrentValue = FindExistingOverride();
+			if (CurrentValue.IsValid())
+			{
+				CurrentObject = CurrentValue.GetUObject();
+			}
+			else
+			{
+				// fetch from asset
+				UNiagaraSystem* System = RawComponent->GetAsset();
+				if (System != nullptr)
+				{
+					FNiagaraUserRedirectionParameterStore& AssetParamStore = System->GetExposedParameters();
+					CurrentObject = AssetParamStore.GetUObject(ParameterKey);
+				}
+			}
+		}
+
+		return CurrentObject != nullptr ? CurrentObject->GetPathName() : FString();
+	}
+
+	void SetParameterEditor(TSharedRef<SNiagaraParameterEditor> InParameterEditor)
+	{
+		ParameterEditor = InParameterEditor;
+	}
+
+	void SetDisplayData(TSharedRef<FStructOnScope> InDisplayData)
+	{
+		DisplayData = InDisplayData;
+	}
+
+	void SetDisplayDataPropertyHandle(TSharedPtr<IPropertyHandle> InDisplayDataPropertyHandle)
+	{
+		DisplayDataPropertyHandle = InDisplayDataPropertyHandle;
+	}
+	
+	const FNiagaraVariableBase& Key() const { return ParameterKey; }
+	FNiagaraVariant& Value() { return ParameterValue; }
+
+	TWeakObjectPtr<UNiagaraComponent> GetComponent() const { return Component; }
+
+private:
+	TWeakObjectPtr<UNiagaraComponent> Component;
+	TArray<TSharedPtr<IPropertyHandle>> OverridePropertyHandles;
+	FNiagaraVariableBase ParameterKey;
+	FNiagaraVariant ParameterValue;
+	FSimpleDelegate OnRebuild;
+	bool bResettingToDefault;
+	TSharedPtr<SNiagaraParameterEditor> ParameterEditor;
+	TSharedPtr<FStructOnScope> DisplayData;
+	TSharedPtr<IPropertyHandle> DisplayDataPropertyHandle;
+};
+
+FNiagaraComponentUserParametersNodeBuilder::FNiagaraComponentUserParametersNodeBuilder(UNiagaraComponent* InComponent, TArray<TSharedPtr<IPropertyHandle>> InOverridePropertyHandles, FName InCustomBuilderRowName)
+{
+	OverridePropertyHandles = InOverridePropertyHandles;
+	CustomBuilderRowName = InCustomBuilderRowName;
+
+	Component = InComponent;
+	bDelegatesInitialized = false;
+}
+
+FNiagaraComponentUserParametersNodeBuilder::~FNiagaraComponentUserParametersNodeBuilder()
+{
+	if (Component.IsValid() && bDelegatesInitialized)
+	{
+		Component->OnSynchronizedWithAssetParameters().RemoveAll(this);
+		Component->GetOverrideParameters().RemoveAllOnChangedHandlers(this);
+
+		if(UNiagaraSystem* System = Component->GetAsset())
+		{
+			TSharedPtr<FNiagaraSystemViewModel> SystemViewModel = TNiagaraViewModelManager<UNiagaraSystem, FNiagaraSystemViewModel>::GetExistingViewModelForObject(System);
+			if(SystemViewModel.IsValid() && SystemViewModel->GetUserParametersHierarchyViewModel())
+			{
+				SystemViewModel->GetUserParametersHierarchyViewModel()->OnHierarchyChanged().RemoveAll(this);
+				SystemViewModel->GetUserParametersHierarchyViewModel()->OnHierarchyPropertiesChanged().RemoveAll(this);
+			}
+		}
+	}
+}
+
+void FNiagaraComponentUserParametersNodeBuilder::GenerateChildContent(IDetailChildrenBuilder& ChildrenBuilder)
+{
+	if(Component.IsValid())
+	{
+		if (bDelegatesInitialized == false)
+		{
+			Component->OnSynchronizedWithAssetParameters().AddSP(this, &FNiagaraComponentUserParametersNodeBuilder::Rebuild);
+			
+			RegisterRebuildOnHierarchyChanged();
+			bDelegatesInitialized = true;
+		}
+
+		ParameterProxies.Reset();
+
+		UNiagaraSystem* SystemAsset = Component->GetAsset();
+		if (SystemAsset == nullptr)
+		{
+			return;
+		}
+			
+		TArray<FNiagaraVariable> UserParameters;
+		SystemAsset->GetExposedParameters().GetUserParameters(UserParameters);
+			
+		ParameterProxies.Reserve(UserParameters.Num());
+
+		GenerateUserParameterRows(ChildrenBuilder, *Cast<UNiagaraSystemEditorData>(SystemAsset->GetEditorData())->UserParameterHierarchy.Get());
+	}
+}
+
+FNiagaraVariant FNiagaraComponentUserParametersNodeBuilder::GetCurrentParameterValue(const FNiagaraVariableBase& UserParameter) const
+{
+	FNiagaraVariant CurrentValue = Component->GetCurrentParameterValue(UserParameter);
+	if (CurrentValue.IsValid())
+	{
+		return CurrentValue;
+	}
+		
+	return GetParameterValueFromSystem(UserParameter, *Component->GetAsset());
+}
+
+UNiagaraSystem* FNiagaraComponentUserParametersNodeBuilder::GetSystem() const
+{
+	return Component->GetAsset();
+}
+
+void FNiagaraComponentUserParametersNodeBuilder::GenerateRowForUserParameterInternal(IDetailChildrenBuilder& ChildrenBuilder, FNiagaraVariable UserParameter)
+{
+	UNiagaraSystem* SystemAsset = GetSystem();
+
+	if(SystemAsset == nullptr)
+	{
+		return;
+	}
+	
+	FNiagaraVariable ChoppedUserParameter(UserParameter);
+	if(UserParameter.IsInNameSpace(FNiagaraConstants::UserNamespaceString))
+	{
+		ChoppedUserParameter.SetName(FName(UserParameter.GetName().ToString().RightChop(5)));
+	}
+	
+	FNiagaraVariant CurrentParameterValue = GetCurrentParameterValue(ChoppedUserParameter);
+	if (!CurrentParameterValue.IsValid())
+	{
+		return;
+	}
+
+	// if the DI is not overridden yet, we don't want to change the original system's DI so we duplicate here just in case
+	if (UserParameter.IsDataInterface())
+	{
+		CurrentParameterValue = FNiagaraVariant(DuplicateObject(CurrentParameterValue.GetDataInterface(), Component.Get()));
+	}
+
+	TSharedPtr<FNiagaraParameterProxy> ParameterProxy = ParameterProxies.Add_GetRef(MakeShared<FNiagaraParameterProxy>(Component.Get(), UserParameter, CurrentParameterValue, OnRebuildChildren, OverridePropertyHandles));
+
+	if (UserParameter.IsDataInterface())
+	{
+		// if no changes are made, then it'll just be the same as the asset
+		TArray<UObject*> Objects { CurrentParameterValue.GetDataInterface() };
+	
+		FAddPropertyParams Params = FAddPropertyParams()
+			.UniqueId(ChoppedUserParameter.GetName())
+			.AllowChildren(true)
+			.CreateCategoryNodes(false);
+	
+		IDetailPropertyRow* DataInterfaceRow  = ChildrenBuilder.AddExternalObjectProperty(Objects, NAME_None, Params);			
+		DataInterfaceRow->DisplayName(FText::FromName(ChoppedUserParameter.GetName()));
+		
+		FDetailWidgetRow& CustomWidget = DataInterfaceRow->CustomWidget(true);
+
+		DataInterfaceRow->GetPropertyHandle()->SetPropertyDisplayName(FText::FromName(ChoppedUserParameter.GetName()));
+		DataInterfaceRow->GetPropertyHandle()->SetOnPropertyValuePreChange(FSimpleDelegate::CreateSP(ParameterProxy.ToSharedRef(), &FNiagaraParameterProxy::OnParameterPreChange));
+		DataInterfaceRow->GetPropertyHandle()->SetOnChildPropertyValuePreChange(FSimpleDelegate::CreateSP(ParameterProxy.ToSharedRef(), &FNiagaraParameterProxy::OnParameterPreChange));
+		DataInterfaceRow->GetPropertyHandle()->SetOnPropertyValueChangedWithData(TDelegate<void(const FPropertyChangedEvent&)>::CreateSP(ParameterProxy.ToSharedRef(), &FNiagaraParameterProxy::OnParameterPostChange));
+		DataInterfaceRow->GetPropertyHandle()->SetOnChildPropertyValueChangedWithData(TDelegate<void(const FPropertyChangedEvent&)>::CreateSP(ParameterProxy.ToSharedRef(), &FNiagaraParameterProxy::OnParameterPostChange));
+		DataInterfaceRow->OverrideResetToDefault(
+			FResetToDefaultOverride::Create(
+				FIsResetToDefaultVisible::CreateSP(ParameterProxy.ToSharedRef(), &FNiagaraParameterProxy::IsPropertyResetToDefaultVisible),
+				FResetToDefaultHandler::CreateSP(ParameterProxy.ToSharedRef(), &FNiagaraParameterProxy::OnPropertyResetToDefault)));
+
+		CustomWidget
+		.FilterString(FText::FromName(ChoppedUserParameter.GetName()))
+		.NameContent()
+		.HAlign(HAlign_Left)
+		[
+			CreateUserParameterNameWidget(ChoppedUserParameter)
+		]
+		.ValueContent()
+		[
+			SNew(SHorizontalBox)
+			+ SHorizontalBox::Slot()
+			.HAlign(HAlign_Fill)
+			.Padding(4.0f)
+			[
+				SNew(STextBlock)
+				.TextStyle(FNiagaraEditorStyle::Get(), "NiagaraEditor.ParameterText")
+				.Text(FText::FromString(FName::NameToDisplayString(UserParameter.GetType().GetClass()->GetName(), false)))
+			]
+		];
+
+		AddCustomMenuActionsForParameter(CustomWidget, UserParameter);
+	}
+	else if(UserParameter.IsUObject())
+	{		
+		TAttribute<bool> IsObjectAssetResetVisible;
+		IsObjectAssetResetVisible.Bind(ParameterProxy.ToSharedRef(), &FNiagaraParameterProxy::IsObjectAssetResetToDefaultVisible);
+		
+		FDetailWidgetRow& DetailWidgetRow = ChildrenBuilder.AddCustomRow(FText::FromName(ChoppedUserParameter.GetName()));
+		DetailWidgetRow.OverrideResetToDefault(
+			FResetToDefaultOverride::Create(
+				IsObjectAssetResetVisible,
+				FSimpleDelegate::CreateSP(ParameterProxy.ToSharedRef(), &FNiagaraParameterProxy::OnObjectAssetResetToDefault)));
+		
+		DetailWidgetRow
+		.NameContent()
+		[
+			CreateUserParameterNameWidget(ChoppedUserParameter)
+		]
+		.ValueContent()
+		[
+			SNew(SObjectPropertyEntryBox)
+			.ObjectPath(ParameterProxy.ToSharedRef(), &FNiagaraParameterProxy::GetCurrentAssetPath)
+			.AllowedClass(UserParameter.GetType().GetClass())
+			.OnObjectChanged(ParameterProxy.ToSharedRef(), &FNiagaraParameterProxy::OnAssetSelectedFromPicker)
+			.DisplayUseSelected(true)
+			.DisplayBrowse(true)
+			.DisplayThumbnail(true)
+			.NewAssetFactories(TArray<UFactory*>())
+		];
+		
+		AddCustomMenuActionsForParameter(DetailWidgetRow, UserParameter);		
+	}
+	else
+	{
+		IDetailPropertyRow* DetailPropertyRow = AddValueParameterAsRow(ChildrenBuilder, ChoppedUserParameter);
+
+		// DisplayData should always exist after the call above
+		if(ensure(DisplayData.Contains(ChoppedUserParameter)))
+		{
+			// A niagara parameter editor exists only for certain types. If it exists, setup the value change callback.
+			if(NiagaraParameterEditors.Contains(ChoppedUserParameter))
+			{
+				NiagaraParameterEditors[ChoppedUserParameter]->SetOnBeginValueChange(SNiagaraParameterEditor::FOnValueChange::CreateSP(ParameterProxy.ToSharedRef(), &FNiagaraParameterProxy::OnParameterPreChange));
+				NiagaraParameterEditors[ChoppedUserParameter]->SetOnEndValueChange(SNiagaraParameterEditor::FOnValueChange::CreateSP(ParameterProxy.ToSharedRef(), &FNiagaraParameterProxy::OnParameterEditorValueChanged));
+				ParameterProxy->SetParameterEditor(NiagaraParameterEditors[ChoppedUserParameter]);
+			}
+			
+			ParameterProxy->SetDisplayData(DisplayData[ChoppedUserParameter]);
+			ParameterProxy->SetDisplayDataPropertyHandle(DetailPropertyRow->GetPropertyHandle());
+		}
+
+		DetailPropertyRow->GetPropertyHandle()->SetOnChildPropertyValuePreChange(FSimpleDelegate::CreateSP(ParameterProxy.ToSharedRef(), &FNiagaraParameterProxy::OnParameterPreChange));
+		DetailPropertyRow->GetPropertyHandle()->SetOnPropertyValuePreChange(FSimpleDelegate::CreateSP(ParameterProxy.ToSharedRef(), &FNiagaraParameterProxy::OnParameterPreChange));
+		/** We make sure to react properly whenever a value changes. This is also executed after pasting or resetting to default. */
+		DetailPropertyRow->GetPropertyHandle()->SetOnPropertyValueChangedWithData(TDelegate<void(const FPropertyChangedEvent&)>::CreateSP(ParameterProxy.ToSharedRef(), &FNiagaraParameterProxy::OnParameterPostChange));
+		DetailPropertyRow->GetPropertyHandle()->SetOnChildPropertyValueChangedWithData(TDelegate<void(const FPropertyChangedEvent&)>::CreateSP(ParameterProxy.ToSharedRef(), &FNiagaraParameterProxy::OnParameterPostChange));
+		DetailPropertyRow->GetPropertyHandle()->SetPropertyDisplayName(FText::FromName(ChoppedUserParameter.GetName()));
+		DetailPropertyRow->OverrideResetToDefault(
+			FResetToDefaultOverride::Create(
+				FIsResetToDefaultVisible::CreateSP(ParameterProxy.ToSharedRef(), &FNiagaraParameterProxy::IsPropertyResetToDefaultVisible),
+				FResetToDefaultHandler::CreateSP(ParameterProxy.ToSharedRef(), &FNiagaraParameterProxy::OnPropertyResetToDefault)));
+	}
+}
+
+void FNiagaraComponentUserParametersNodeBuilder::OnParameterEditorValueChanged(FNiagaraVariable UserParameter)
+{
+	// this should not be called. FNiagaraParameterProxy::OnParameterEditorValueChanged should be called instead.
+	ensure(false);
+}
+
+void FNiagaraComponentUserParametersNodeBuilder::OnObjectAssetChanged(const FAssetData& NewObject, FNiagaraVariable UserParameter)
+{
+	for(TSharedPtr<FNiagaraParameterProxy> Proxy : ParameterProxies)
+	{
+		if(Proxy->Key() == UserParameter)
+		{
+			Proxy->OnAssetSelectedFromPicker(NewObject);
+			return;
+		}
+	}
+}
+
+void FNiagaraComponentUserParametersNodeBuilder::RegisterRebuildOnHierarchyChanged()
+{
+	if(Component.IsValid())
+	{
+		if(UNiagaraSystem* Asset = Component->GetAsset())
+		{
+			TSharedPtr<FNiagaraSystemViewModel> SystemViewModel = TNiagaraViewModelManager<UNiagaraSystem, FNiagaraSystemViewModel>::GetExistingViewModelForObject(Asset);
+			if(SystemViewModel.IsValid())
+			{
+				SystemViewModel->GetUserParametersHierarchyViewModel()->OnHierarchyChanged().RemoveAll(this);
+				SystemViewModel->GetUserParametersHierarchyViewModel()->OnHierarchyChanged().Add(UNiagaraHierarchyViewModelBase::FOnHierarchyChanged::FDelegate::CreateSP(this, &FNiagaraComponentUserParametersNodeBuilder::Rebuild));
+				SystemViewModel->GetUserParametersHierarchyViewModel()->OnHierarchyPropertiesChanged().RemoveAll(this);
+				SystemViewModel->GetUserParametersHierarchyViewModel()->OnHierarchyPropertiesChanged().Add(UNiagaraHierarchyViewModelBase::FOnHierarchyPropertiesChanged::FDelegate::CreateSP(this, &FNiagaraComponentUserParametersNodeBuilder::Rebuild));
+			}
+		}
+	}
+}
+
+#undef LOCTEXT_NAMESPACE
diff --git a/Engine/Plugins/FX/Niagara/Source/NiagaraEditor/Private/Customizations/UserParameters/NiagaraComponentUserParametersBuilder.h b/Engine/Plugins/FX/Niagara/Source/NiagaraEditor/Private/Customizations/UserParameters/NiagaraComponentUserParametersBuilder.h
new file mode 100644
--- /dev/null
+++ b/Engine/Plugins/FX/Niagara/Source/NiagaraEditor/Private/Customizations/UserParameters/NiagaraComponentUserParametersBuilder.h
@@ -0,0 +1,42 @@
+// Copyright Epic Games, Inc. All Rights Reserved.
+
+#pragma once
+
+#include "NiagaraUserParametersBuilderBase.h"
+#include "CoreMinimal.h"
+#include "NiagaraComponent.h"
+#include "ViewModels/HierarchyEditor/NiagaraUserParametersHierarchyViewModel.h"
+
+class IDetailLayoutBuilder;
+class IPropertyHandle;
+class FNiagaraParameterViewModelCustomDetails;
+class INiagaraParameterViewModel;
+class SNiagaraParameterEditor;
+class FNiagaraParameterViewModelCustomDetails;
+class FNiagaraSimCacheCapture;
+class UNiagaraSystem;
+
+/** A specialized class to generate user parameters for a component. */
+class FNiagaraComponentUserParametersNodeBuilder : public FNiagaraUserParameterNodeBuilder
+{
+public:
+	FNiagaraComponentUserParametersNodeBuilder(UNiagaraComponent* InComponent, TArray<TSharedPtr<IPropertyHandle>> InOverridePropertyHandles, FName InCustomBuilderRowName);
+	virtual ~FNiagaraComponentUserParametersNodeBuilder() override;
+	
+	virtual void GenerateChildContent(IDetailChildrenBuilder& ChildrenBuilder) override;
+
+private:
+	/** NiagaraUserParameterNodeBuilder Interface */
+	virtual UNiagaraSystem* GetSystem() const override;
+	virtual FNiagaraVariant GetCurrentParameterValue(const FNiagaraVariableBase& UserParameter) const override;
+	virtual void GenerateRowForUserParameterInternal(IDetailChildrenBuilder& ChildrenBuilder, FNiagaraVariable UserParameter) override;
+	virtual void OnParameterEditorValueChanged(FNiagaraVariable UserParameter) override;
+	virtual void OnObjectAssetChanged(const FAssetData& NewObject, FNiagaraVariable UserParameter) override;
+	
+	void RegisterRebuildOnHierarchyChanged();
+
+private:
+	TWeakObjectPtr<UNiagaraComponent> Component;
+	TArray<TSharedPtr<class FNiagaraParameterProxy>> ParameterProxies;
+	TArray<TSharedPtr<IPropertyHandle>> OverridePropertyHandles;
+};
diff --git a/Engine/Plugins/FX/Niagara/Source/NiagaraEditor/Private/Customizations/UserParameters/NiagaraSystemUserParametersBuilder.cpp b/Engine/Plugins/FX/Niagara/Source/NiagaraEditor/Private/Customizations/UserParameters/NiagaraSystemUserParametersBuilder.cpp
new file mode 100644
--- /dev/null
+++ b/Engine/Plugins/FX/Niagara/Source/NiagaraEditor/Private/Customizations/UserParameters/NiagaraSystemUserParametersBuilder.cpp
@@ -0,0 +1,667 @@
+// Copyright Epic Games, Inc. All Rights Reserved.
+
+#include "NiagaraSystemUserParametersBuilder.h"
+#include "ViewModels/HierarchyEditor/NiagaraUserParametersHierarchyViewModel.h"
+#include "Toolkits/SystemToolkitModes/NiagaraSystemToolkitModeBase.h"
+#include "DetailWidgetRow.h"
+#include "PropertyCustomizationHelpers.h"
+#include "Styling/AppStyle.h"
+#include "NiagaraEditorUtilities.h"
+#include "NiagaraTypes.h"
+#include "ScopedTransaction.h"
+#include "NiagaraNodeInput.h"
+#include "IDetailChildrenBuilder.h"
+#include "NiagaraConstants.h"
+#include "NiagaraScriptVariable.h"
+#include "NiagaraSettings.h"
+#include "NiagaraSystemEditorData.h"
+#include "NiagaraUserRedirectionParameterStore.h"
+#include "SNiagaraParameterEditor.h"
+#include "Widgets/Input/SButton.h"
+#include "Widgets/SNiagaraDebugger.h"
+#include "Widgets/SNiagaraParameterName.h"
+#include "Widgets/Images/SImage.h"
+#include "Widgets/Layout/SBox.h"
+#include "Widgets/Input/SComboButton.h"
+#include "Widgets/SNiagaraSystemUserParameters.h"
+#include "Widgets/SNiagaraParameterMenu.h"
+#include "Widgets/Input/SEditableTextBox.h"
+
+#define LOCTEXT_NAMESPACE "NiagaraSystemUserParametersBuilder"
+
+FNiagaraSystemUserParameterBuilder::FNiagaraSystemUserParameterBuilder(TSharedPtr<FNiagaraSystemViewModel> InSystemViewModel, FName InCustomBuilderRowName)
+{
+	CustomBuilderRowName = InCustomBuilderRowName;
+	System = &InSystemViewModel->GetSystem();
+	SystemViewModel = InSystemViewModel;
+	bDelegatesInitialized = false;
+}
+
+FNiagaraSystemUserParameterBuilder::~FNiagaraSystemUserParameterBuilder()
+{
+	if (System.IsValid() && bDelegatesInitialized)
+	{
+		System->GetExposedParameters().OnStructureChanged().RemoveAll(this);
+
+		if(SystemViewModel.IsValid() && SystemViewModel.Pin()->GetUserParameterPanelViewModel().IsValid())
+		{
+			if (TSharedPtr<FNiagaraUserParameterPanelViewModel> PanelViewModel = SystemViewModel.Pin()->GetUserParameterPanelViewModel())
+			{
+				PanelViewModel->OnRefreshRequested().Unbind();
+			}
+		}
+			
+		if(SystemViewModel.IsValid() && SystemViewModel.Pin()->GetUserParametersHierarchyViewModel())
+		{
+			if (UNiagaraUserParametersHierarchyViewModel* HierarchyViewModel = SystemViewModel.Pin()->GetUserParametersHierarchyViewModel())
+			{
+				HierarchyViewModel->OnHierarchyChanged().RemoveAll(this);
+				HierarchyViewModel->OnHierarchyPropertiesChanged().RemoveAll(this);
+			}
+		}
+	}
+}
+
+void FNiagaraSystemUserParameterBuilder::GenerateChildContent(IDetailChildrenBuilder& ChildrenBuilder)
+{
+	check(System.IsValid());
+
+	if (System.IsValid())
+	{
+		if(bDelegatesInitialized == false)
+		{
+			System->GetExposedParameters().OnStructureChanged().Add(FNiagaraParameterStore::FOnStructureChanged::FDelegate::CreateSP(this, &FNiagaraSystemUserParameterBuilder::Rebuild));
+
+			SystemViewModel.Pin()->GetUserParametersHierarchyViewModel()->OnHierarchyChanged().RemoveAll(this);
+			SystemViewModel.Pin()->GetUserParametersHierarchyViewModel()->OnHierarchyChanged().Add(UNiagaraHierarchyViewModelBase::FOnHierarchyChanged::FDelegate::CreateSP(this, &FNiagaraSystemUserParameterBuilder::Rebuild));
+			SystemViewModel.Pin()->GetUserParametersHierarchyViewModel()->OnHierarchyPropertiesChanged().RemoveAll(this);
+			SystemViewModel.Pin()->GetUserParametersHierarchyViewModel()->OnHierarchyPropertiesChanged().Add(UNiagaraHierarchyViewModelBase::FOnHierarchyPropertiesChanged::FDelegate::CreateSP(this, &FNiagaraSystemUserParameterBuilder::Rebuild));
+			
+			TSharedPtr<FNiagaraUserParameterPanelViewModel> UserParameterPanelViewModel = SystemViewModel.Pin()->GetUserParameterPanelViewModel();
+			if(ensure(UserParameterPanelViewModel.IsValid()))
+			{
+				UserParameterPanelViewModel->OnRefreshRequested().BindSP(this, &FNiagaraSystemUserParameterBuilder::Rebuild);
+				UserParameterPanelViewModel->OnParameterAdded().BindSP(this, &FNiagaraSystemUserParameterBuilder::OnParameterAdded);
+			}
+			
+			bDelegatesInitialized = true;
+		}
+		
+		UNiagaraHierarchyRoot* Root = Cast<UNiagaraSystemEditorData>(System->GetEditorData())->UserParameterHierarchy;
+		GenerateUserParameterRows(ChildrenBuilder, *Root);
+	}
+}
+
+void FNiagaraSystemUserParameterBuilder::AddCustomMenuActionsForParameter(FDetailWidgetRow& WidgetRow, FNiagaraVariable UserParameter)
+{
+	/** We want to add rename & delete actions within a system asset */
+	
+	FNiagaraUserParameterNodeBuilder::AddCustomMenuActionsForParameter(WidgetRow, UserParameter);
+	
+	FUIAction RenameAction(FExecuteAction::CreateSP(this, &FNiagaraSystemUserParameterBuilder::RequestRename, UserParameter));
+	WidgetRow.AddCustomContextMenuAction(RenameAction, LOCTEXT("RenameParameterAction", "Rename"), LOCTEXT("RenameParameterActionTooltip", "Rename this user parameter"));
+
+	FUIAction DeleteAction(FExecuteAction::CreateSP(this, &FNiagaraSystemUserParameterBuilder::DeleteParameter, UserParameter));
+	WidgetRow.AddCustomContextMenuAction(DeleteAction, LOCTEXT("DeleteParameterAction", "Delete"), LOCTEXT("DeleteParameterActionTooltip", "Delete this user parameter"));
+}
+
+TSharedRef<SWidget> FNiagaraSystemUserParameterBuilder::GetAdditionalHeaderWidgets()
+{
+	if(!AdditionalHeaderWidgetsContainer.IsValid())
+	{
+		AdditionalHeaderWidgetsContainer = SNew(SBox)
+		.HAlign(HAlign_Right)
+		.VAlign(VAlign_Center)
+		.Padding(2.f)
+		[
+			SNew(SHorizontalBox)
+			+ SHorizontalBox::Slot()
+			.HAlign(HAlign_Center)
+			.VAlign(VAlign_Center)
+			.AutoWidth()
+			.Padding(2.f, 1.f)
+			[
+				SNew(SButton)
+				.OnClicked(this, &FNiagaraSystemUserParameterBuilder::SummonHierarchyEditor)
+				.Text(LOCTEXT("SummonUserParametersHierarchyButtonLabel", "Edit Hierarchy"))
+				.ButtonStyle(FAppStyle::Get(), "RoundButton")
+				.ToolTipText(LOCTEXT("SummonUserParametersHierarchyButtonTooltip", "Summon the Hierarchy Editor to add categories, sections and tooltips to User Parameters."))
+			]
+			+ SHorizontalBox::Slot()
+			.HAlign(HAlign_Center)
+			.VAlign(VAlign_Center)
+			.AutoWidth()
+			.Padding(2.f)
+			[
+				SAssignNew(AddParameterButton, SComboButton)
+				.HAlign(HAlign_Center)
+				.VAlign(VAlign_Center)
+				.ButtonStyle(FAppStyle::Get(), "RoundButton")
+				.ForegroundColor(FAppStyle::GetSlateColor("DefaultForeground"))
+				.ContentPadding(FMargin(2, 3.5))
+				.ToolTipText(LOCTEXT("AddUserParameterButtonTooltip", "Add a new User Parameter to this system.\nUser Parameters exist once for the entire system and their values can be changed at runtime from Blueprints, C++ or the sequencer, among other systems."))
+				.OnGetMenuContent(this, &FNiagaraSystemUserParameterBuilder::GetAddParameterMenu)
+				.OnComboBoxOpened_Lambda([this]
+				{
+					AddParameterButton->SetMenuContentWidgetToFocus(AddParameterMenu->GetSearchBox());
+				})
+				.HasDownArrow(false)
+				.ButtonContent()
+				[
+					SNew(SImage)
+					.Image(FAppStyle::GetBrush("Plus"))
+				]
+			]
+		];
+	}
+	
+	return AdditionalHeaderWidgetsContainer.ToSharedRef(); 
+}
+
+FNiagaraVariant FNiagaraSystemUserParameterBuilder::GetCurrentParameterValue(const FNiagaraVariableBase& UserParameter) const
+{
+	return GetParameterValueFromSystem(UserParameter, *System.Get());
+}
+
+UNiagaraSystem* FNiagaraSystemUserParameterBuilder::GetSystem() const
+{
+	return System.Get();
+}
+
+void FNiagaraSystemUserParameterBuilder::GenerateRowForUserParameterInternal(IDetailChildrenBuilder& ChildrenBuilder, FNiagaraVariable UserParameter)
+{
+	UNiagaraSystem* SystemAsset = GetSystem();
+
+	if(SystemAsset == nullptr)
+	{
+		return;
+	}
+	
+	FNiagaraVariable ChoppedUserParameter(UserParameter);
+	if(UserParameter.IsInNameSpace(FNiagaraConstants::UserNamespaceString))
+	{
+		ChoppedUserParameter.SetName(FName(UserParameter.GetName().ToString().RightChop(5)));
+	}
+	
+	FNiagaraVariant CurrentParameterValue = GetCurrentParameterValue(ChoppedUserParameter);
+	if (!CurrentParameterValue.IsValid())
+	{
+		return;
+	}
+	
+	if (UserParameter.IsDataInterface())
+	{
+		// if no changes are made, then it'll just be the same as the asset
+		TArray<UObject*> Objects { CurrentParameterValue.GetDataInterface() };
+	
+		FAddPropertyParams Params = FAddPropertyParams()
+			.UniqueId(ChoppedUserParameter.GetName())
+			.AllowChildren(true)
+			.CreateCategoryNodes(false);
+	
+		IDetailPropertyRow* DataInterfaceRow = ChildrenBuilder.AddExternalObjectProperty(Objects, NAME_None, Params);			
+		DataInterfaceRow->DisplayName(FText::FromName(ChoppedUserParameter.GetName()));
+		
+		FDetailWidgetRow& CustomWidget = DataInterfaceRow->CustomWidget(true);
+
+		DataInterfaceRow->GetPropertyHandle()->SetPropertyDisplayName(FText::FromName(ChoppedUserParameter.GetName()));
+		DataInterfaceRow->GetPropertyHandle()->SetOnPropertyValuePreChange(FSimpleDelegate::CreateSP(this, &FNiagaraSystemUserParameterBuilder::OnDataInterfacePropertyValuePreChange));
+		DataInterfaceRow->GetPropertyHandle()->SetOnChildPropertyValuePreChange(FSimpleDelegate::CreateSP(this, &FNiagaraSystemUserParameterBuilder::OnDataInterfacePropertyValuePreChange));
+		DataInterfaceRow->GetPropertyHandle()->SetOnPropertyValueChangedWithData(TDelegate<void(const FPropertyChangedEvent&)>::CreateSP(this, &FNiagaraSystemUserParameterBuilder::OnDataInterfacePropertyValueChangedWithData));
+		DataInterfaceRow->GetPropertyHandle()->SetOnChildPropertyValueChangedWithData(TDelegate<void(const FPropertyChangedEvent&)>::CreateSP(this, &FNiagaraSystemUserParameterBuilder::OnDataInterfacePropertyValueChangedWithData));
+		DataInterfaceRow->OverrideResetToDefault(
+			FResetToDefaultOverride::Create(
+				FIsResetToDefaultVisible::CreateSP(this, &FNiagaraSystemUserParameterBuilder::IsDataInterfaceResetToDefaultVisible, ChoppedUserParameter),
+				FResetToDefaultHandler::CreateSP(this, &FNiagaraSystemUserParameterBuilder::OnHandleDataInterfaceReset, ChoppedUserParameter)));
+		
+		CustomWidget
+		.FilterString(FText::FromName(ChoppedUserParameter.GetName()))
+		.NameContent()
+		.HAlign(HAlign_Left)
+		[
+			CreateUserParameterNameWidget(UserParameter)
+		]
+		.ValueContent()
+		[
+			SNew(STextBlock)
+			.TextStyle(FNiagaraEditorStyle::Get(), "NiagaraEditor.ParameterText")
+			.Text(FText::FromString(FName::NameToDisplayString(UserParameter.GetType().GetClass()->GetName(), false)))
+		];
+		
+		AddCustomMenuActionsForParameter(CustomWidget, UserParameter);
+	}
+	else if(UserParameter.IsUObject())
+	{
+		FDetailWidgetRow& DetailWidgetRow = ChildrenBuilder.AddCustomRow(FText::FromName(ChoppedUserParameter.GetName()));
+		DetailWidgetRow.OverrideResetToDefault(
+					FResetToDefaultOverride::Create(
+						TAttribute<bool>::CreateSP(this, &FNiagaraSystemUserParameterBuilder::IsObjectAssetCustomRowResetToDefaultVisible, ChoppedUserParameter),
+						FSimpleDelegate::CreateSP(this, &FNiagaraSystemUserParameterBuilder::OnHandleObjectAssetReset, ChoppedUserParameter)));
+		
+		DetailWidgetRow
+		.NameContent()
+		[
+			CreateUserParameterNameWidget(UserParameter)
+		]
+		.ValueContent()
+		[
+			SNew(SObjectPropertyEntryBox)
+			.ObjectPath(TAttribute<FString>::CreateSP(this, &FNiagaraSystemUserParameterBuilder::GetObjectAssetPathForUserParameter, ChoppedUserParameter))
+			.AllowedClass(UserParameter.GetType().GetClass())
+			.OnObjectChanged(this, &FNiagaraSystemUserParameterBuilder::OnObjectAssetChanged, ChoppedUserParameter)
+			.DisplayUseSelected(true)
+			.DisplayBrowse(true)
+			.DisplayThumbnail(true)
+			.NewAssetFactories(TArray<UFactory*>())
+		];
+
+		AddCustomMenuActionsForParameter(DetailWidgetRow, UserParameter);
+	}
+	else
+	{
+		IDetailPropertyRow* DetailPropertyRow = AddValueParameterAsRow(ChildrenBuilder, ChoppedUserParameter);
+		if(!ensure(DisplayData.Contains(ChoppedUserParameter) && DisplayDataPropertyHandles.Contains(ChoppedUserParameter)))
+		{
+			return;
+		}
+		
+		/** We make sure to react properly whenever a value changes. This is also executed after pasting or resetting to default. */
+		DetailPropertyRow->GetPropertyHandle()->SetOnPropertyValueChangedWithData(TDelegate<void(const FPropertyChangedEvent&)>::CreateSP(this, &FNiagaraSystemUserParameterBuilder::OnDisplayDataPostChange, ChoppedUserParameter));
+		DetailPropertyRow->GetPropertyHandle()->SetOnChildPropertyValueChangedWithData(TDelegate<void(const FPropertyChangedEvent&)>::CreateSP(this, &FNiagaraSystemUserParameterBuilder::OnDisplayDataPostChange, ChoppedUserParameter));
+		DetailPropertyRow->GetPropertyHandle()->SetPropertyDisplayName(FText::FromName(ChoppedUserParameter.GetName()));
+		DetailPropertyRow->OverrideResetToDefault(
+			FResetToDefaultOverride::Create(
+				FIsResetToDefaultVisible::CreateSP(this, &FNiagaraSystemUserParameterBuilder::IsLocalPropertyResetToDefaultVisible, ChoppedUserParameter),
+				FResetToDefaultHandler::CreateSP(this, &FNiagaraSystemUserParameterBuilder::OnHandleLocalPropertyReset, ChoppedUserParameter)));
+	}
+}
+
+TSharedRef<SWidget> FNiagaraSystemUserParameterBuilder::GetAddParameterMenu()
+{	
+	FNiagaraParameterPanelCategory UserCategory = GetDefault<UNiagaraEditorSettings>()->GetMetaDataForNamespaces({FNiagaraConstants::UserNamespace});
+	AddParameterMenu = SNew(SNiagaraAddParameterFromPanelMenu)
+		.Graphs(SystemViewModel.Pin()->GetParameterPanelViewModel()->GetEditableGraphsConst())
+		.OnNewParameterRequested(this, &FNiagaraSystemUserParameterBuilder::AddParameter)
+		.OnAllowMakeType(this, &FNiagaraSystemUserParameterBuilder::OnAllowMakeType)
+		.NamespaceId(UserCategory.NamespaceMetaData.GetGuid())
+		.ShowNamespaceCategory(false)
+		.ShowGraphParameters(false)
+		.AutoExpandMenu(false);
+
+	return AddParameterMenu.ToSharedRef();
+}
+
+void FNiagaraSystemUserParameterBuilder::AddParameter(FNiagaraVariable NewParameter) const
+{
+	FNiagaraParameterPanelCategory UserCategory = GetDefault<UNiagaraEditorSettings>()->GetMetaDataForNamespaces({FNiagaraConstants::UserNamespace});
+	FNiagaraEditorUtilities::AddParameter(NewParameter, SystemViewModel.Pin()->GetSystem().GetExposedParameters(), SystemViewModel.Pin()->GetSystem(), nullptr);
+	SystemViewModel.Pin()->GetUserParameterPanelViewModel()->OnParameterAdded().Execute(NewParameter);
+}
+
+bool FNiagaraSystemUserParameterBuilder::OnAllowMakeType(const FNiagaraTypeDefinition& InType) const
+{
+	FNiagaraNamespaceMetadata UserNameSpaceMetaData = GetDefault<UNiagaraEditorSettings>()->GetMetaDataForNamespaces({FNiagaraConstants::UserNamespace});
+	FNiagaraParameterPanelCategory Category(UserNameSpaceMetaData);
+	return SystemViewModel.Pin()->GetParameterPanelViewModel()->CanAddType(InType, Category);
+}
+
+FReply FNiagaraSystemUserParameterBuilder::SummonHierarchyEditor() const
+{
+	if(TSharedPtr<FNiagaraSystemViewModel> SystemViewModelPinned = SystemViewModel.Pin())
+	{
+		SystemViewModelPinned->FocusTab(FNiagaraSystemToolkitModeBase::UserParametersHierarchyTabID, true);
+		return FReply::Handled();
+	}
+
+	return FReply::Unhandled();
+}
+
+void FNiagaraSystemUserParameterBuilder::DeleteParameter(FNiagaraVariable UserParameter) const
+{
+	FScopedTransaction Transaction(LOCTEXT("DeleteParameterTransaction", "Deleted user parameter"));
+	System->Modify();
+	
+	if(System->GetExposedParameters().IndexOf(UserParameter) != INDEX_NONE)
+	{
+		System->GetExposedParameters().RemoveParameter(UserParameter);		
+	}
+
+	Cast<UNiagaraSystemEditorData>(System->GetEditorData())->UserParameterHierarchy->Modify();
+	System->HandleVariableRemoved(UserParameter, true);
+}
+
+TSharedRef<SWidget> FNiagaraSystemUserParameterBuilder::CreateUserParameterNameWidget(FNiagaraVariable UserParameter)
+{
+	// we make sure that within a system context, the name parameter is always prefixed with User.
+	FNiagaraUserRedirectionParameterStore::MakeUserVariable(UserParameter);
+	
+	/** Instead of pure text, we use the parameter widgets for display in a system asset. */
+	TSharedRef<SNiagaraParameterNameTextBlock> ParameterName = SNew(SNiagaraParameterNameTextBlock)
+		// if this is specified, we avoid the default behavior, which is "when clicked when the widget had keyboard focus already, enter editing mode"
+		// we don't want to be able to rename by clicking, only by context menu
+		.IsSelected_Lambda([]()
+		{
+			return false;
+		})
+		.ParameterText(FText::FromName(UserParameter.GetName()))
+		.OnDragDetected(this, &FNiagaraSystemUserParameterBuilder::GenerateParameterDragDropOp, UserParameter)
+		.OnVerifyTextChanged_Lambda([this, UserParameter](const FText& InNewName, FText& OutErrorMessage)
+		{
+			if(InNewName.IsEmptyOrWhitespace())
+			{
+				OutErrorMessage = LOCTEXT("UserParameterNameCantBeEmpty", "Can not accept empty name");
+				return false;
+			}
+			
+			FNiagaraVariable NewUserParameter = UserParameter;
+			NewUserParameter.SetName(FName(InNewName.ToString()));
+			if(UserParameter == NewUserParameter)
+			{
+				// we return true, but renaming won't do anything if the name is the same as before
+				return true;
+			}
+			else
+			{
+				TArray<FNiagaraVariable> ExistingParameters;
+				UserParamToWidgetMap.GenerateKeyArray(ExistingParameters);
+
+				bool bNameExistsAlready = false;
+				for(FNiagaraVariable& ExistingUserParameter : ExistingParameters)
+				{
+					if(ExistingUserParameter.GetName().IsEqual(FName(InNewName.ToString())))
+					{
+						bNameExistsAlready = true;
+						break;
+					}
+				}
+
+				if(bNameExistsAlready)
+				{
+					OutErrorMessage = LOCTEXT("UserParameterAlreadyExists", "A User Parameter with the same name already exists");
+					return false;					
+				}
+
+				return true;
+			}
+		})
+		.OnTextCommitted_Lambda([this, UserParameter](const FText& InText, ETextCommit::Type Type)
+		{
+			RenameParameter(UserParameter, FName(InText.ToString()));
+		});
+	
+	UserParamToWidgetMap.Add(UserParameter, ParameterName);
+
+	TAttribute<FText> Tooltip = FText::GetEmpty();
+	if(UNiagaraScriptVariable* ScriptVariable = FNiagaraEditorUtilities::GetScriptVariableForUserParameter(UserParameter, *GetSystem()))
+	{
+		Tooltip = TAttribute<FText>::CreateLambda([ScriptVariable]
+		{
+			return ScriptVariable->Metadata.Description;
+		});
+	}
+	ParameterName->SetToolTipText(Tooltip);
+	
+	return ParameterName;
+}
+
+void FNiagaraSystemUserParameterBuilder::SelectAllSection()
+{
+	if(ActiveSection != nullptr)
+	{
+		ActiveSection = nullptr;
+		OnRebuildChildren.Execute();
+	}
+}
+
+void FNiagaraSystemUserParameterBuilder::OnParameterAdded(FNiagaraVariable UserParameter)
+{
+	// when adding a new parameter it defaults to the 'All' section, so we select it in order to visualize it before renaming
+	SelectAllSection();
+	RequestRename(UserParameter);
+}
+
+void FNiagaraSystemUserParameterBuilder::RequestRename(FNiagaraVariable UserParameter)
+{
+	FNiagaraUserRedirectionParameterStore::MakeUserVariable(UserParameter);
+	
+	// we add a timer for next frame, as requesting a rename on a parameter that has just been created will not preselect the entire text. Waiting a frame fixes that.
+	UserParamToWidgetMap[UserParameter]->RegisterActiveTimer(0.f, FWidgetActiveTimerDelegate::CreateLambda([this, UserParameter](double CurrentTime, float DeltaTime)
+	{
+		UserParamToWidgetMap[UserParameter]->EnterEditingMode();
+		return EActiveTimerReturnType::Stop;
+	}));
+}
+
+void FNiagaraSystemUserParameterBuilder::RenameParameter(FNiagaraVariable UserParameter, FName NewName) const
+{
+	if(UserParameter.GetName() == NewName)
+	{
+		return;
+	}
+	
+	if(System->GetExposedParameters().IndexOf(UserParameter) != INDEX_NONE)
+	{		
+		SystemViewModel.Pin()->RenameParameter(UserParameter, NewName, ENiagaraGetGraphParameterReferencesMode::AllGraphs);
+	}
+}
+
+FReply FNiagaraSystemUserParameterBuilder::GenerateParameterDragDropOp(const FGeometry& Geometry, const FPointerEvent& MouseEvent, FNiagaraVariable UserParameter) const
+{
+	TSharedPtr<FEdGraphSchemaAction> ParameterDragAction = MakeShared<FNiagaraParameterAction>(UserParameter, FText::GetEmpty(), FText::GetEmpty(), FText::GetEmpty(), 0, FText(), nullptr, 0/*SectionID*/);
+	TSharedRef<FNiagaraParameterDragOperation> ParameterDragOperation = MakeShared<FNiagaraParameterDragOperation>(ParameterDragAction);
+	ParameterDragOperation->SetupDefaults();
+	ParameterDragOperation->Construct();
+	return FReply::Handled().BeginDragDrop(ParameterDragOperation);
+}
+
+void FNiagaraSystemUserParameterBuilder::OnParameterEditorValueChanged(FNiagaraVariable UserParameter)
+{
+	UNiagaraSystem* SystemAsset = GetSystem();
+	
+	if(UserParameter.IsDataInterface() || UserParameter.IsUObject())
+	{
+		ensure(false);
+		return;
+	}
+
+	TSharedRef<FStructOnScope> ParameterEditorData = MakeShared<FStructOnScope>(UserParameter.GetType().GetStruct());
+	NiagaraParameterEditors[UserParameter]->UpdateStructFromInternalValue(ParameterEditorData);
+	if(FMemory::Memcmp(ParameterEditorData->GetStructMemory(), DisplayData[UserParameter]->GetStructMemory(), UserParameter.GetType().GetSize()) == 0)
+	{
+		return;
+	}
+	
+	FScopedTransaction ScopedTransaction(ChangedUserParameterTransactionText);
+	SystemAsset->Modify();
+
+	// we forward the change in the parameter editor into the display data. NotifyPostChange will take care of forwarding the change into the user parameter store.
+	DisplayDataPropertyHandles[UserParameter]->NotifyPreChange();
+	NiagaraParameterEditors[UserParameter]->UpdateStructFromInternalValue(DisplayData[UserParameter]);
+	DisplayDataPropertyHandles[UserParameter]->NotifyPostChange(EPropertyChangeType::ValueSet);
+	DisplayDataPropertyHandles[UserParameter]->NotifyFinishedChangingProperties();
+}
+
+void FNiagaraSystemUserParameterBuilder::OnDisplayDataPostChange(const FPropertyChangedEvent& ChangedEvent, FNiagaraVariable UserParameter)
+{
+	if(ChangedEvent.ChangeType != EPropertyChangeType::Interactive)
+	{
+		UNiagaraSystem* SystemAsset = GetSystem();
+
+		if(NiagaraParameterEditors.Contains(UserParameter))
+		{
+			// we have to make sure the internal parameter editor data is in sync with the display data (copy paste, reset etc. write to display data)
+			NiagaraParameterEditors[UserParameter]->UpdateInternalValueFromStruct(DisplayData[UserParameter]);
+		}
+		
+		TSharedRef<FStructOnScope> UserParameterDisplayData = DisplayData[UserParameter];
+
+		// we have to special case the position definition because, while SetParameterData does handle Position variables, it expects 12 Bytes of data
+		if(UserParameter.GetType() == FNiagaraTypeDefinition::GetPositionDef())
+		{
+			FVector Value = *reinterpret_cast<const FVector*>(UserParameterDisplayData->GetStructMemory());
+			SystemAsset->GetExposedParameters().SetPositionParameterValue(Value, UserParameter.GetName());
+		}
+		else
+		{
+			SystemAsset->GetExposedParameters().SetParameterData(UserParameterDisplayData ->GetStructMemory(), UserParameter);
+		}
+	}
+}
+
+void FNiagaraSystemUserParameterBuilder::OnObjectAssetChanged(const FAssetData& NewObject, FNiagaraVariable UserParameter)
+{
+	UNiagaraSystem* SystemAsset = GetSystem();
+	
+	if(UserParameter.IsUObject() == false || UserParameter.IsDataInterface())
+	{
+		ensure(false);
+		return;
+	}
+
+	FScopedTransaction ScopedTransaction(LOCTEXT("ChangeAsset", "Changed User Parameter Asset"));
+	SystemAsset->Modify();
+	SystemAsset->GetExposedParameters().SetUObject(NewObject.GetAsset(), UserParameter);
+}
+
+FString FNiagaraSystemUserParameterBuilder::GetObjectAssetPathForUserParameter(FNiagaraVariable UserParameter) const
+{
+	ensure(UserParameter.GetType().IsUObject() && UserParameter.GetType().IsDataInterface() == false);
+	if(UObject* Object = GetCurrentParameterValue(UserParameter).GetUObject())
+	{
+		return Object->GetPathName();
+	}
+	
+	return FString();
+}
+
+bool FNiagaraSystemUserParameterBuilder::IsDataInterfaceResetToDefaultVisible(TSharedPtr<IPropertyHandle> PropertyHandle, FNiagaraVariable UserParameter) const
+{
+	UNiagaraDataInterface* DataInterface = GetSystem()->GetExposedParameters().GetDataInterface(UserParameter);
+
+	if(DataInterface == nullptr)
+	{
+		return false;
+	}
+
+	bool bIdentical = true;
+	UObject* DataInterfaceCDO = DataInterface->GetClass()->GetDefaultObject();
+	for(TFieldIterator<FProperty> It(DataInterface->GetClass()); It; ++It)
+	{
+		if(It->HasAnyPropertyFlags(CPF_Edit) == false)
+		{
+			continue;
+		}
+		
+		bIdentical &= It->Identical_InContainer(DataInterface, DataInterfaceCDO, 0, PPF_DeepComparison);
+
+		if(bIdentical == false)
+		{
+			break;
+		}
+	}
+
+	return bIdentical == false;
+}
+
+void FNiagaraSystemUserParameterBuilder::OnHandleDataInterfaceReset(TSharedPtr<IPropertyHandle> PropertyHandle,	FNiagaraVariable UserParameter)
+{
+	UNiagaraDataInterface* DataInterface = GetSystem()->GetExposedParameters().GetDataInterface(UserParameter);
+
+	if(DataInterface == nullptr)
+	{
+		return;
+	}
+	
+	UObject* DataInterfaceCDO = DataInterface->GetClass()->GetDefaultObject();
+	for(TFieldIterator<FProperty> It(DataInterface->GetClass()); It; ++It)
+	{
+		if(It->HasAnyPropertyFlags(CPF_Edit) == false)
+		{
+			continue;
+		}
+		
+		It->CopyCompleteValue_InContainer(DataInterface, DataInterfaceCDO);
+	}
+
+	// to clean up reset widgets for child properties, 
+	Rebuild();
+}
+
+bool FNiagaraSystemUserParameterBuilder::IsLocalPropertyResetToDefaultVisible(TSharedPtr<IPropertyHandle> PropertyHandle, FNiagaraVariable UserParameter) const
+{
+	if(UserParameter.GetType() == FNiagaraTypeDefinition::GetPositionDef())
+	{
+		static UPackage* CoreUObjectPkg = FindObjectChecked<UPackage>(nullptr, TEXT("/Script/CoreUObject"));
+		static UScriptStruct* VectorStruct = FindObjectChecked<UScriptStruct>(CoreUObjectPkg, TEXT("Vector"));
+
+		TSharedRef<FStructOnScope> VectorData = MakeShared<FStructOnScope>(VectorStruct);
+		VectorStruct->InitializeDefaultValue(VectorData->GetStructMemory());
+		return FMemory::Memcmp(DisplayData[UserParameter]->GetStructMemory(), VectorData->GetStructMemory(), VectorStruct->GetStructureSize()) != 0;
+	}
+	
+	FNiagaraEditorUtilities::ResetVariableToDefaultValue(UserParameter);
+	return FMemory::Memcmp(DisplayData[UserParameter]->GetStructMemory(), UserParameter.GetData(), UserParameter.GetType().GetSize()) != 0;
+}
+
+void FNiagaraSystemUserParameterBuilder::OnHandleLocalPropertyReset(TSharedPtr<IPropertyHandle> PropertyHandle, FNiagaraVariable UserParameter)
+{
+	FScopedTransaction Transaction(ResetUserParameterTransactionText);
+	GetSystem()->Modify();
+	
+	FNiagaraTypeDefinition Type = UserParameter.GetType();
+
+	UStruct* Struct = Type.GetStruct();
+	TSharedPtr<FStructOnScope> DefaultData;
+	if(UserParameter.GetType() == FNiagaraTypeDefinition::GetPositionDef())
+	{
+		static UPackage* CoreUObjectPkg = FindObjectChecked<UPackage>(nullptr, TEXT("/Script/CoreUObject"));
+		static UScriptStruct* VectorStruct = FindObjectChecked<UScriptStruct>(CoreUObjectPkg, TEXT("Vector"));
+
+		DefaultData = MakeShared<FStructOnScope>(VectorStruct);
+		VectorStruct->InitializeDefaultValue(DefaultData->GetStructMemory());
+		Struct = VectorStruct;
+
+		FVector* Vector = reinterpret_cast<FVector*>(DefaultData->GetStructMemory());
+		FVector* CurrentVector = reinterpret_cast<FVector*>(DisplayData[UserParameter]->GetStructMemory());
+
+		UE_LOG(LogNiagaraEditor, Warning, TEXT("Current: %s"), *CurrentVector->ToString());
+		UE_LOG(LogNiagaraEditor, Warning, TEXT("Default: %s"), *Vector->ToString());
+
+	}
+	else
+	{
+		FNiagaraEditorUtilities::ResetVariableToDefaultValue(UserParameter);
+		DefaultData = MakeShared<FStructOnScope>(Struct, UserParameter.GetData());
+	}
+	
+	
+	DisplayDataPropertyHandles[UserParameter]->NotifyPreChange();
+	FMemory::Memcpy(DisplayData[UserParameter]->GetStructMemory(), DefaultData->GetStructMemory(), Struct->GetStructureSize());
+	DisplayDataPropertyHandles[UserParameter]->NotifyPostChange(EPropertyChangeType::ValueSet);
+	DisplayDataPropertyHandles[UserParameter]->NotifyFinishedChangingProperties();
+}
+
+bool FNiagaraSystemUserParameterBuilder::IsObjectAssetCustomRowResetToDefaultVisible(FNiagaraVariable UserParameter) const
+{	
+	return GetSystem()->GetExposedParameters().GetUObject(UserParameter) != nullptr;
+}
+
+void FNiagaraSystemUserParameterBuilder::OnHandleObjectAssetReset(FNiagaraVariable UserParameter)
+{
+	FScopedTransaction Transaction(ResetUserParameterTransactionText);
+	GetSystem()->Modify();
+
+	GetSystem()->GetExposedParameters().SetUObject(nullptr, UserParameter);
+}
+
+void FNiagaraSystemUserParameterBuilder::OnDataInterfacePropertyValuePreChange()
+{
+	GetSystem()->Modify();
+}
+
+void FNiagaraSystemUserParameterBuilder::OnDataInterfacePropertyValueChangedWithData(const FPropertyChangedEvent& PropertyChangedEvent)
+{
+	GetSystem()->GetExposedParameters().OnInterfaceChange();
+}
+
+#undef LOCTEXT_NAMESPACE
diff --git a/Engine/Plugins/FX/Niagara/Source/NiagaraEditor/Private/Customizations/UserParameters/NiagaraSystemUserParametersBuilder.h b/Engine/Plugins/FX/Niagara/Source/NiagaraEditor/Private/Customizations/UserParameters/NiagaraSystemUserParametersBuilder.h
new file mode 100644
--- /dev/null
+++ b/Engine/Plugins/FX/Niagara/Source/NiagaraEditor/Private/Customizations/UserParameters/NiagaraSystemUserParametersBuilder.h
@@ -0,0 +1,81 @@
+// Copyright Epic Games, Inc. All Rights Reserved.
+
+#pragma once
+
+#include "NiagaraUserParametersBuilderBase.h"
+#include "CoreMinimal.h"
+#include "Input/Reply.h"
+#include "ViewModels/TNiagaraViewModelManager.h"
+#include "ViewModels/HierarchyEditor/NiagaraHierarchyViewModelBase.h"
+#include "ViewModels/HierarchyEditor/NiagaraUserParametersHierarchyViewModel.h"
+
+class IDetailLayoutBuilder;
+class IPropertyHandle;
+class FNiagaraParameterViewModelCustomDetails;
+class INiagaraParameterViewModel;
+class SNiagaraParameterEditor;
+class FNiagaraParameterViewModelCustomDetails;
+class FNiagaraSimCacheCapture;
+class UNiagaraSystem;
+
+/** A specialized class to generate user parameters for a system asset's editor. Due to reliance on the system view model this should only be used within the asset editor. */
+class FNiagaraSystemUserParameterBuilder : public FNiagaraUserParameterNodeBuilder
+{
+public:
+	FNiagaraSystemUserParameterBuilder(TSharedPtr<FNiagaraSystemViewModel> InSystemViewModel, FName InCustomBuilderRowName);
+	virtual ~FNiagaraSystemUserParameterBuilder() override;
+
+	virtual void GenerateChildContent(IDetailChildrenBuilder& ChildrenBuilder) override;
+	
+	TSharedRef<SWidget> GetAdditionalHeaderWidgets();
+private:
+	/** NiagaraUserParameterNodeBuilder Interface */
+	virtual UNiagaraSystem* GetSystem() const override;
+	virtual FNiagaraVariant GetCurrentParameterValue(const FNiagaraVariableBase& UserParameter) const override;
+	virtual void GenerateRowForUserParameterInternal(IDetailChildrenBuilder& ChildrenBuilder, FNiagaraVariable UserParameter) override;
+	virtual void AddCustomMenuActionsForParameter(FDetailWidgetRow& WidgetRow, FNiagaraVariable UserParameter) override;
+	virtual void OnParameterEditorValueChanged(FNiagaraVariable UserParameter) override;
+	virtual void OnObjectAssetChanged(const FAssetData& NewObject, FNiagaraVariable UserParameter) override;
+
+	void OnDisplayDataPostChange(const FPropertyChangedEvent& ChangedEvent, FNiagaraVariable UserParameter);
+
+	TSharedRef<SWidget> GetAddParameterMenu();
+	void AddParameter(FNiagaraVariable NewParameter) const;
+	void OnParameterAdded(FNiagaraVariable UserParameter);
+	bool OnAllowMakeType(const FNiagaraTypeDefinition& InType) const;
+	void DeleteParameter(FNiagaraVariable UserParameter) const;
+
+	virtual TSharedRef<SWidget> CreateUserParameterNameWidget(FNiagaraVariable UserParameter) override;
+	
+	FReply SummonHierarchyEditor() const;
+	void SelectAllSection();
+	
+	void RequestRename(FNiagaraVariable UserParameter);
+	void RenameParameter(FNiagaraVariable UserParameter, FName NewName) const;
+
+	FReply GenerateParameterDragDropOp(const FGeometry& Geometry, const FPointerEvent& MouseEvent, FNiagaraVariable UserParameter) const;
+private:
+	FString GetObjectAssetPathForUserParameter(FNiagaraVariable UserParameter) const;
+
+	bool IsDataInterfaceResetToDefaultVisible(TSharedPtr<IPropertyHandle> PropertyHandle, FNiagaraVariable UserParameter) const;
+	void OnHandleDataInterfaceReset(TSharedPtr<IPropertyHandle> PropertyHandle, FNiagaraVariable UserParameter);
+	
+	/** These can be used for Reset to Default overrides that are added to rows with valid property nodes (external properties etc.) */
+	bool IsLocalPropertyResetToDefaultVisible(TSharedPtr<IPropertyHandle> PropertyHandle, FNiagaraVariable UserParameter) const;
+	void OnHandleLocalPropertyReset(TSharedPtr<IPropertyHandle> PropertyHandle, FNiagaraVariable UserParameter);
+
+	/** These can be used for Reset to Default overrides that are added to custom rows (no valid property nodes) */
+	bool IsObjectAssetCustomRowResetToDefaultVisible(FNiagaraVariable UserParameter) const;
+	void OnHandleObjectAssetReset(FNiagaraVariable UserParameter);
+	
+	void OnDataInterfacePropertyValuePreChange();
+	void OnDataInterfacePropertyValueChangedWithData(const FPropertyChangedEvent& PropertyChangedEvent);
+	
+private:
+	TWeakObjectPtr<UNiagaraSystem> System;
+	TWeakPtr<FNiagaraSystemViewModel> SystemViewModel;
+	TMap<FNiagaraVariable, TSharedPtr<class SNiagaraParameterNameTextBlock>> UserParamToWidgetMap;
+	TSharedPtr<SWidget> AdditionalHeaderWidgetsContainer;
+	TSharedPtr<class SComboButton> AddParameterButton;
+	TSharedPtr<class SNiagaraAddParameterFromPanelMenu> AddParameterMenu;
+};
diff --git a/Engine/Plugins/FX/Niagara/Source/NiagaraEditor/Private/Customizations/UserParameters/NiagaraUserParametersBuilderBase.cpp b/Engine/Plugins/FX/Niagara/Source/NiagaraEditor/Private/Customizations/UserParameters/NiagaraUserParametersBuilderBase.cpp
new file mode 100644
--- /dev/null
+++ b/Engine/Plugins/FX/Niagara/Source/NiagaraEditor/Private/Customizations/UserParameters/NiagaraUserParametersBuilderBase.cpp
@@ -0,0 +1,506 @@
+// Copyright Epic Games, Inc. All Rights Reserved.
+
+#include "NiagaraUserParametersBuilderBase.h"
+#include "PropertyCustomizationHelpers.h"
+#include "NiagaraTypes.h"
+#include "IDetailChildrenBuilder.h"
+#include "INiagaraEditorTypeUtilities.h"
+#include "NiagaraClipboard.h"
+#include "NiagaraEditorStyle.h"
+#include "NiagaraScriptVariable.h"
+#include "NiagaraSettings.h"
+#include "NiagaraUserRedirectionParameterStore.h"
+#include "ViewModels/HierarchyEditor/NiagaraUserParametersHierarchyViewModel.h"
+#include "Widgets/SNiagaraDebugger.h"
+#include "Widgets/SNiagaraSystemUserParameters.h"
+#include "Widgets/Input/SCheckBox.h"
+#include "NiagaraEditorDataBase.h"
+#include "PropertyEditorClipboard.h"
+#include "Customizations/NiagaraComponentDetails.h"
+
+#define LOCTEXT_NAMESPACE "NiagaraUserParametersBuilderBase"
+
+FText FNiagaraUserParameterNodeBuilder::ResetUserParameterTransactionText = LOCTEXT("ResetUserParameter", "User Parameter reset to default value");
+FText FNiagaraUserParameterNodeBuilder::ChangedUserParameterTransactionText = LOCTEXT("EditUserParameter", "Edit User Parameter");
+
+FNiagaraVariant FNiagaraUserParameterNodeBuilder::GetParameterValueFromSystem(const FNiagaraVariableBase& Parameter, const UNiagaraSystem& System)
+{
+	const FNiagaraUserRedirectionParameterStore& UserParameterStore = System.GetExposedParameters();
+	if (Parameter.IsDataInterface())
+	{
+		int32 Index = UserParameterStore.IndexOf(Parameter);
+		if (Index != INDEX_NONE)
+		{
+			return FNiagaraVariant(UserParameterStore.GetDataInterfaces()[Index]);
+		}
+	}
+	
+	if (Parameter.IsUObject())
+	{
+		int32 Index = UserParameterStore.IndexOf(Parameter);
+		if (Index != INDEX_NONE)
+		{
+			return FNiagaraVariant(UserParameterStore.GetUObjects()[Index]);
+		}
+	}
+
+	if (Parameter.GetType() == FNiagaraTypeDefinition::GetPositionDef())
+	{
+		const FVector* Value = UserParameterStore.GetPositionParameterValue(Parameter.GetName());
+		return Value == nullptr ? FNiagaraVariant() : FNiagaraVariant(Value, sizeof(FVector));
+	}
+
+	TArray<uint8> DataValue;
+	DataValue.AddUninitialized(Parameter.GetSizeInBytes());
+	if (UserParameterStore.CopyParameterData(Parameter, DataValue.GetData()))
+	{
+		return FNiagaraVariant(DataValue);
+	}
+	return FNiagaraVariant();
+}
+
+IDetailPropertyRow* FNiagaraUserParameterNodeBuilder::AddValueParameterAsRow(IDetailChildrenBuilder& ChildrenBuilder, FNiagaraVariable ChoppedUserParameter)
+{
+	FNiagaraVariant CurrentParameterValue = GetCurrentParameterValue(ChoppedUserParameter);
+	if(!ensure(CurrentParameterValue.GetMode() == ENiagaraVariantMode::Bytes))
+	{
+		return nullptr;
+	}
+	
+	FNiagaraTypeDefinition ParameterType = ChoppedUserParameter.GetType();
+	UStruct* Struct = ParameterType.GetStruct();
+
+	// we don't want to create niagara editors for certain types.
+	// - Position types would use the Vector3f widget right now which wouldn't give us double precision
+	bool bShouldCreateNiagaraEditor = ParameterType != FNiagaraTypeDefinition::GetPositionDef();
+
+	// for position types, we want to display an FVector instead of the FVector3f that the position actually is. So we replace the struct here
+    if(ParameterType == FNiagaraTypeDefinition::GetPositionDef())
+    {
+    	static UPackage* CoreUObjectPkg = FindObjectChecked<UPackage>(nullptr, TEXT("/Script/CoreUObject"));
+    	static UScriptStruct* VectorStruct = FindObjectChecked<UScriptStruct>(CoreUObjectPkg, TEXT("Vector"));
+    	Struct = VectorStruct;
+    }
+	
+	TSharedRef<FStructOnScope> CurrentParameterValueData = MakeShared<FStructOnScope>(Struct);
+	FMemory::Memcpy(CurrentParameterValueData->GetStructMemory(), CurrentParameterValue.GetBytes(), Struct->GetStructureSize());
+
+	DisplayData.Add(ChoppedUserParameter, CurrentParameterValueData);
+	
+	if(bShouldCreateNiagaraEditor)
+	{
+		UNiagaraScriptVariable* UserParameterScriptVariable = FNiagaraEditorUtilities::GetScriptVariableForUserParameter(ChoppedUserParameter, *GetSystem());
+		FNiagaraEditorModule& NiagaraEditorModule = FNiagaraEditorModule::Get();
+		if(TSharedPtr<INiagaraEditorTypeUtilities> TypeUtilities = NiagaraEditorModule.GetTypeUtilities(ParameterType))
+		{
+			if(TypeUtilities->CanCreateParameterEditor())
+			{
+				TSharedPtr<SNiagaraParameterEditor> ParameterEditor = TypeUtilities->CreateParameterEditor(ParameterType, UserParameterScriptVariable->Metadata.DisplayUnit, UserParameterScriptVariable->Metadata.WidgetCustomization).ToSharedRef();
+				// since the color editor is in a wrap box using the allotted size, we have to make sure to give it plenty of space so it doesn't start wrapping immediately
+				if(ParameterType == FNiagaraTypeDefinition::GetColorDef())
+				{
+					ParameterEditor->SetMinimumDesiredWidth(350.f);
+				}
+				ParameterEditor->UpdateInternalValueFromStruct(CurrentParameterValueData);
+				ParameterEditor->SetOnEndValueChange(SNiagaraParameterEditor::FOnValueChange::CreateSP(this, &FNiagaraUserParameterNodeBuilder::OnParameterEditorValueChanged, ChoppedUserParameter));			
+				
+				NiagaraParameterEditors.Add(ChoppedUserParameter, ParameterEditor.ToSharedRef());
+			}
+		}
+	}
+	
+	FAddPropertyParams Params;
+	Params.UniqueId(ChoppedUserParameter.GetName());
+
+	/** We add an external property here, pointing at the display data we have created
+	 *  Data can be written in three ways:
+	 *  1) Using the custom parameter editor. This requires syncing from parameter editor -> display data -> user parameter store
+	 *  2) Using copy & paste. This requires syncing from display data -> parameter editor -> user parameter store
+	 *  3) Using reset to default. This requires syncing the default value to display data -> parameter editor -> user parameter store
+	 */
+	IDetailPropertyRow* DetailPropertyRow = ChildrenBuilder.AddExternalStructure(DisplayData[ChoppedUserParameter], ChoppedUserParameter.GetName());
+	DisplayDataPropertyHandles.Add(ChoppedUserParameter, DetailPropertyRow->GetPropertyHandle());
+	
+	TSharedPtr<SWidget> DefaultNameWidget;
+	TSharedPtr<SWidget> DefaultValueWidget;
+	DetailPropertyRow->GetDefaultWidgets(DefaultNameWidget, DefaultValueWidget);
+	
+	FDetailWidgetRow& DetailWidgetRow = DetailPropertyRow->CustomWidget();
+	DetailWidgetRow
+	.FilterString(FText::FromName(ChoppedUserParameter.GetName()))
+	.NameContent()
+	[
+		CreateUserParameterNameWidget(ChoppedUserParameter)
+	]
+	.ValueContent()
+	[
+		NiagaraParameterEditors.Contains(ChoppedUserParameter) ? NiagaraParameterEditors[ChoppedUserParameter] : DefaultValueWidget.ToSharedRef()
+	];
+
+	AddCustomMenuActionsForParameter(DetailWidgetRow, ChoppedUserParameter);
+	
+	return DetailPropertyRow;
+}
+
+IDetailPropertyRow* FNiagaraUserParameterNodeBuilder::AddObjectAssetParameterAsRow(IDetailChildrenBuilder& ChildrenBuilder, FNiagaraVariable ChoppedUserParameter)
+{
+	UObject* AssetHelper = ObjectAssetHelpers.Add_GetRef(NewObject<UNiagaraObjectAssetHelper>());
+
+	FAddPropertyParams AddPropertyParams;
+	AddPropertyParams.UniqueId(ChoppedUserParameter.GetName());
+	return ChildrenBuilder.AddExternalObjectProperty({AssetHelper}, FName("Path"), AddPropertyParams);
+}
+
+void FNiagaraUserParameterNodeBuilder::GenerateUserParameterRows(IDetailChildrenBuilder& ChildrenBuilder, UNiagaraHierarchyRoot& UserParameterHierarchyRoot)
+{
+	UNiagaraSystem* SystemAsset = GetSystem();
+
+	if(SystemAsset == nullptr)
+	{
+		return;
+	}
+
+	TArray<const UNiagaraHierarchySection*> OrderedSections;
+	TMap<const UNiagaraHierarchySection*, TArray<const UNiagaraHierarchyCategory*>> SectionCategoryMap;
+
+	UNiagaraHierarchyRoot* Root = &UserParameterHierarchyRoot;
+	for(const UNiagaraHierarchyItemBase* Child : Root->GetChildren())
+	{
+		if(const UNiagaraHierarchyCategory* HierarchyCategory = Cast<UNiagaraHierarchyCategory>(Child))
+		{
+			if(HierarchyCategory->GetSection() != nullptr)
+			{
+				SectionCategoryMap.FindOrAdd(HierarchyCategory->GetSection()).Add(HierarchyCategory);	
+			}
+		}
+	}
+
+	// we create a custom row for user param sections here, but only if we have at least one section specified
+	if(SectionCategoryMap.Num() > 0)
+	{
+		TSharedRef<SHorizontalBox> SectionsBox = SNew(SHorizontalBox);
+
+		// maps don't guarantee order, so we iterate over the original section data instead
+		for(auto& Section : Root->GetSectionData())
+		{
+			// if the map doesn't contain at the entry, it means there are 0 categories for that section, so we skip it
+			if(!SectionCategoryMap.Contains(Section))
+			{
+				continue;
+			}
+			
+			TArray<const UNiagaraHierarchyCategory*>& HierarchyCategories = SectionCategoryMap[Section];
+			
+			if(HierarchyCategories.Num() > 0)
+			{
+				// we only want to add a section if any of its categories contain at least one user parameter
+				bool bDoesUserParamExist = false;
+				for(const UNiagaraHierarchyCategory* Category : HierarchyCategories)
+				{
+					bDoesUserParamExist |= Category->DoesOneChildExist<UNiagaraHierarchyUserParameter>(true);
+
+					if(bDoesUserParamExist)
+					{
+						break;
+					}
+				}
+
+				if(bDoesUserParamExist)
+				{
+					SectionsBox->AddSlot()
+					.AutoWidth()
+					.Padding(FMargin(2.f))
+					[
+						SNew(SBox)
+						.Padding(FMargin(0.f, 4.f, 0.f, 0.f))
+						[
+							SNew(SCheckBox)
+							.Style(FAppStyle::Get(), "DetailsView.SectionButton")
+							.OnCheckStateChanged_Lambda([this, Section](ECheckBoxState NewState)
+							{
+								ActiveSection = Section;
+								OnRebuildChildren.ExecuteIfBound();
+							})
+							.IsChecked_Lambda([this, Section]()
+							{
+								return ActiveSection == Section ? ECheckBoxState::Checked : ECheckBoxState::Unchecked;
+							})
+							[
+								SNew(STextBlock)
+								.TextStyle(FAppStyle::Get(), "SmallText")
+								.Text_UObject(Section, &UNiagaraHierarchySection::GetSectionNameAsText)
+								.ToolTipText_UObject(Section, &UNiagaraHierarchySection::GetTooltip)
+							]
+						]
+					];
+				}
+			}
+		}
+
+		// if we have at least one custom section, we add a default "All" section
+		if(SectionsBox->GetChildren()->Num() > 0)
+		{
+			SectionsBox->AddSlot()
+			.AutoWidth()
+			.Padding(FMargin(2.f))
+			[
+				SNew(SBox)
+				.Padding(FMargin(0.f, 4.f, 0.f, 0.f))
+				[
+					SNew(SCheckBox)
+					.Style(FAppStyle::Get(), "DetailsView.SectionButton")
+					.OnCheckStateChanged_Lambda([this](ECheckBoxState NewState)
+					{
+						ActiveSection = nullptr;
+						OnRebuildChildren.ExecuteIfBound();
+					})
+					.IsChecked_Lambda([this]()
+					{
+						return ActiveSection == nullptr ? ECheckBoxState::Checked : ECheckBoxState::Unchecked;
+					})
+					[
+						SNew(STextBlock)
+						.TextStyle(FAppStyle::Get(), "SmallText")
+						.Text(FText::FromString("All"))
+					]
+				]
+			];
+			
+			FDetailWidgetRow& Row = ChildrenBuilder.AddCustomRow(FText::FromString("Parameter Sections"));
+			Row.WholeRowContent()
+			[
+				SectionsBox
+			];
+		}
+	}
+	
+	// we add the categories now if the active section is set to all (nullptr) or if the section it belongs to is active
+	for(const UNiagaraHierarchyItemBase* Child : Root->GetChildren())
+	{		
+		if(const UNiagaraHierarchyCategory* HierarchyCategory = Cast<UNiagaraHierarchyCategory>(Child))
+		{
+			if((ActiveSection == nullptr || HierarchyCategory->GetSection() == ActiveSection) && HierarchyCategory->DoesOneChildExist<UNiagaraHierarchyUserParameter>(true))
+			{
+				ChildrenBuilder.AddCustomBuilder(MakeShared<FNiagaraUserParameterCategoryBuilder>(HierarchyCategory, HierarchyCategory->GetCategoryName(), *this));
+			
+			}
+		}
+		else if(const UNiagaraHierarchyUserParameter* UserParameter = Cast<UNiagaraHierarchyUserParameter>(Child))
+		{
+			// we only want to display parameters within the root in the all section
+			if(ActiveSection == nullptr)
+			{
+				GenerateRowForUserParameter(ChildrenBuilder, UserParameter->GetUserParameter());
+			}
+		}
+	}
+
+	// at last we add all parameters that haven't been setup at all, if the active section is set to "All"
+	if(ActiveSection == nullptr)
+	{
+		TArray<FNiagaraVariable> LeftoverUserParameters;
+		SystemAsset->GetExposedParameters().GetUserParameters(LeftoverUserParameters);
+
+		TArray<UNiagaraHierarchyUserParameter*> HierarchyUserParameters;
+		Root->GetChildrenOfType<>(HierarchyUserParameters, true);
+
+		TArray<FNiagaraVariable> RedirectedUserParameters;
+		for(FNiagaraVariable& Parameter : LeftoverUserParameters)
+		{
+			SystemAsset->GetExposedParameters().RedirectUserVariable(Parameter);
+		}
+		
+		for(UNiagaraHierarchyUserParameter* HierarchyUserParameter : HierarchyUserParameters)
+		{
+			if(LeftoverUserParameters.Contains(HierarchyUserParameter->GetUserParameter()))
+			{
+				LeftoverUserParameters.Remove(HierarchyUserParameter->GetUserParameter());
+			}
+		}
+
+		for(FNiagaraVariable& LeftoverUserParameter : LeftoverUserParameters)
+		{			
+			GenerateRowForUserParameter(ChildrenBuilder, LeftoverUserParameter);
+		}
+	}
+}
+
+void FNiagaraUserParameterNodeBuilder::GenerateRowForUserParameter(IDetailChildrenBuilder& ChildrenBuilder, FNiagaraVariable UserParameter)
+{
+	GenerateRowForUserParameterInternal(ChildrenBuilder, UserParameter);
+}
+
+void FNiagaraUserParameterNodeBuilder::AddCustomMenuActionsForParameter(FDetailWidgetRow& WidgetRow, FNiagaraVariable UserParameter)
+{
+	if(UserParameter.IsDataInterface())
+	{
+		FUIAction CopyAction;
+		CopyAction.ExecuteAction = FExecuteAction::CreateLambda([=]()
+		{
+			FNiagaraVariant CurrentParameterValue = GetCurrentParameterValue(UserParameter);
+			if(CurrentParameterValue.IsValid() == false || CurrentParameterValue.GetDataInterface() == nullptr)
+			{
+				return;
+			}
+			
+			UNiagaraClipboardContent* ClipboardContent = UNiagaraClipboardContent::Create();
+			const UNiagaraClipboardFunctionInput* Input = UNiagaraClipboardFunctionInput::CreateDataValue(ClipboardContent, UserParameter.GetName(), UserParameter.GetType().GetClass(), TOptional<bool>(), CurrentParameterValue.GetDataInterface());
+			ClipboardContent->FunctionInputs.Add(Input);
+			
+			FNiagaraEditorModule::Get().GetClipboard().SetClipboardContent(ClipboardContent);
+		});
+		WidgetRow.CopyAction(CopyAction);
+
+		FUIAction PasteAction;
+		PasteAction.ExecuteAction = FExecuteAction::CreateLambda([=]()
+		{
+			FNiagaraVariant CurrentParameterValue = GetCurrentParameterValue(UserParameter);
+			if(CurrentParameterValue.IsValid() == false || CurrentParameterValue.GetDataInterface() == nullptr)
+			{
+				return;
+			}
+			
+			const UNiagaraClipboardContent* ClipboardContent = FNiagaraEditorModule::Get().GetClipboard().GetClipboardContent();
+
+			if(ClipboardContent->FunctionInputs.Num() == 1 && ClipboardContent->FunctionInputs[0]->InputType.IsDataInterface())
+			{
+				UNiagaraDataInterface* CopiedDI = ClipboardContent->FunctionInputs[0]->Data;
+				if(CopiedDI->IsA(CurrentParameterValue.GetDataInterface()->GetClass()))
+				{
+					CopiedDI->CopyTo(CurrentParameterValue.GetDataInterface());
+				}
+			}			
+		});
+		WidgetRow.PasteAction(PasteAction);
+	}
+	else if(UserParameter.IsUObject())
+	{
+		FUIAction CopyAction;
+		CopyAction.ExecuteAction = FExecuteAction::CreateLambda([=]()
+		{
+			FNiagaraVariant CurrentParameterValue = GetCurrentParameterValue(UserParameter);
+			if(CurrentParameterValue.IsValid() == false)
+			{
+				return;
+			}
+
+			FPropertyEditorClipboard::ClipboardCopy(*CurrentParameterValue.GetUObject()->GetPathName());
+		});
+		WidgetRow.CopyAction(CopyAction);
+
+		FUIAction PasteAction;
+		PasteAction.ExecuteAction = FExecuteAction::CreateLambda([=]()
+		{
+			FString PastedString;
+			FPropertyEditorClipboard::ClipboardPaste(PastedString);
+
+			FSoftObjectPath Path(PastedString);
+			if(Path.IsAsset())
+			{
+				UObject* Asset = Path.TryLoad();
+				OnObjectAssetChanged(FAssetData(Asset), UserParameter);
+			}
+			else
+			{
+				OnObjectAssetChanged(FAssetData(nullptr), UserParameter);
+			}
+		});
+		WidgetRow.PasteAction(PasteAction);
+	}
+}
+
+TSharedRef<SWidget> FNiagaraUserParameterNodeBuilder::CreateUserParameterNameWidget(FNiagaraVariable UserParameter)
+{
+	UNiagaraScriptVariable* ScriptVariable = FNiagaraEditorUtilities::GetScriptVariableForUserParameter(UserParameter, *GetSystem());
+	TAttribute<FText> Tooltip = FText::GetEmpty();
+	if(ScriptVariable)
+	{
+		Tooltip = TAttribute<FText>::CreateLambda([ScriptVariable]
+		{
+			return ScriptVariable->Metadata.Description;
+		});
+	}
+	
+	return SNew(SBox)
+	.ToolTipText(Tooltip)
+	.Padding(FMargin(0.0f, 2.0f))
+	[
+		SNew(STextBlock)
+		.TextStyle(FNiagaraEditorStyle::Get(), "NiagaraEditor.ParameterText")
+		.Text(FText::FromName(UserParameter.GetName()))
+	];
+}
+
+void FNiagaraUserParameterNodeBuilder::OnScriptVariablePropertyChanged(UNiagaraScriptVariable* ScriptVariable, const FPropertyChangedEvent& PropertyChangedEvent) const
+{
+	if(ScriptVariable->GetOuter() == GetSystem()->GetEditorData())
+	{
+		Rebuild();
+	}
+}
+
+void FNiagaraUserParameterNodeBuilder::AddReferencedObjects(FReferenceCollector& Collector)
+{
+	Collector.AddReferencedObjects(ObjectAssetHelpers);
+}
+
+FString FNiagaraUserParameterNodeBuilder::GetReferencerName() const
+{
+	return "UserParametersBuilder";
+}
+
+FNiagaraUserParameterCategoryBuilder::FNiagaraUserParameterCategoryBuilder(const UNiagaraHierarchyCategory* InNiagaraHierarchyCategory, FName InCustomBuilderRowName, FNiagaraUserParameterNodeBuilder& InNodeBuilder) : NodeBuilder(InNodeBuilder)
+{
+	HierarchyCategory = InNiagaraHierarchyCategory;
+	CustomBuilderRowName = InCustomBuilderRowName;
+}
+
+void FNiagaraUserParameterCategoryBuilder::GenerateHeaderRowContent(FDetailWidgetRow& NodeRow)
+{
+	NodeRow.NameContent()
+	[
+		SNew(STextBlock)
+		.TextStyle(FNiagaraEditorStyle::Get(), "NiagaraEditor.ParameterText")
+		.Text_UObject(HierarchyCategory.Get(), &UNiagaraHierarchyCategory::GetCategoryAsText)
+		.ToolTipText_UObject(HierarchyCategory.Get(), &UNiagaraHierarchyCategory::GetTooltip)
+	];
+
+	NodeRow.FilterString(FText::FromName(CustomBuilderRowName));
+}
+
+void FNiagaraUserParameterCategoryBuilder::GenerateChildContent(IDetailChildrenBuilder& ChildrenBuilder)
+{
+	TArray<const UNiagaraHierarchyUserParameter*> ContainedParameters;
+	TArray<const UNiagaraHierarchyCategory*> ContainedCategories;
+
+	for(const UNiagaraHierarchyItemBase* Child : HierarchyCategory->GetChildren())
+	{
+		if(const UNiagaraHierarchyUserParameter* UserParameter = Cast<UNiagaraHierarchyUserParameter>(Child))
+		{
+			ContainedParameters.Add(UserParameter);
+		}
+	}
+
+	for(UNiagaraHierarchyItemBase* Child : HierarchyCategory->GetChildren())
+	{
+		if(UNiagaraHierarchyCategory* Category = Cast<UNiagaraHierarchyCategory>(Child))
+		{
+			ContainedCategories.Add(Category);
+		}
+	}
+
+	for(const UNiagaraHierarchyCategory* Category : ContainedCategories)
+	{
+		if(Category->DoesOneChildExist<UNiagaraHierarchyUserParameter>(true))
+		{
+			ChildrenBuilder.AddCustomBuilder(MakeShared<FNiagaraUserParameterCategoryBuilder>(Category, Category->GetCategoryName(), NodeBuilder));
+		}
+	}
+		
+	for(const UNiagaraHierarchyUserParameter* HierarchyUserParameter : ContainedParameters)
+	{
+		NodeBuilder.GenerateRowForUserParameter(ChildrenBuilder, HierarchyUserParameter->GetUserParameter());
+	}
+}
+
+#undef LOCTEXT_NAMESPACE
diff --git a/Engine/Plugins/FX/Niagara/Source/NiagaraEditor/Private/Customizations/UserParameters/NiagaraUserParametersBuilderBase.h b/Engine/Plugins/FX/Niagara/Source/NiagaraEditor/Private/Customizations/UserParameters/NiagaraUserParametersBuilderBase.h
new file mode 100644
--- /dev/null
+++ b/Engine/Plugins/FX/Niagara/Source/NiagaraEditor/Private/Customizations/UserParameters/NiagaraUserParametersBuilderBase.h
@@ -0,0 +1,111 @@
+// Copyright Epic Games, Inc. All Rights Reserved.
+
+#pragma once
+
+#include "CoreMinimal.h"
+#include "IDetailCustomNodeBuilder.h"
+#include "NiagaraComponent.h"
+#include "SNiagaraParameterEditor.h"
+#include "ViewModels/HierarchyEditor/NiagaraHierarchyViewModelBase.h"
+
+/** A base class for generating rows based on a given hierarchy root that contains sections, categories and user parameters. */
+class FNiagaraUserParameterNodeBuilder : public IDetailCustomNodeBuilder, public TSharedFromThis<FNiagaraUserParameterNodeBuilder>, FSelfRegisteringEditorUndoClient, FGCObject
+{
+public:
+	virtual void GenerateHeaderRowContent(FDetailWidgetRow& NodeRow) override {}
+	virtual bool InitiallyCollapsed() const override { return false; }
+	
+	virtual FName GetName() const  override
+	{
+		return CustomBuilderRowName;
+	}
+	
+	void Rebuild() const
+	{
+		OnRebuildChildren.ExecuteIfBound();
+	}
+	
+	virtual void SetOnRebuildChildren(FSimpleDelegate InOnRegenerateChildren) override
+	{
+		OnRebuildChildren = InOnRegenerateChildren;
+	}
+	
+	/** Generates all rows for the node builder. Includes setting up sections, categories and user parameters */
+	void GenerateUserParameterRows(IDetailChildrenBuilder& ChildrenBuilder, UNiagaraHierarchyRoot& UserParameterHierarchyRoot);
+	/** Generates a single row for a specific user parameter. */
+	void GenerateRowForUserParameter(IDetailChildrenBuilder& ChildrenBuilder, FNiagaraVariable UserParameter);
+	
+	void OnActiveSectionChanged(ECheckBoxState State, const UNiagaraHierarchySection* NewSelection)
+	{
+		ActiveSection = NewSelection;
+		OnRebuildChildren.Execute();
+	}
+
+	ECheckBoxState IsSectionActive(const UNiagaraHierarchySection* Section) const
+	{
+		return ActiveSection == Section ? ECheckBoxState::Checked : ECheckBoxState::Unchecked;
+	}
+
+	static FText ResetUserParameterTransactionText;
+	static FText ChangedUserParameterTransactionText;	
+protected:	
+	virtual UNiagaraSystem* GetSystem() const = 0;
+	virtual FNiagaraVariant GetCurrentParameterValue(const FNiagaraVariableBase& UserParameter) const = 0;
+	/** Allows the addition of right-click context actions for a specific user parameter */
+	virtual void AddCustomMenuActionsForParameter(FDetailWidgetRow& WidgetRow, FNiagaraVariable UserParameter);
+	virtual void GenerateRowForUserParameterInternal(IDetailChildrenBuilder& ChildrenBuilder, FNiagaraVariable UserParameter) = 0;
+	virtual void OnParameterEditorValueChanged(FNiagaraVariable UserParameter) = 0;
+	/** Since object asset references are handled individually via custom rows, they need to implement this themselves. This will be called with copy & paste. */
+	virtual void OnObjectAssetChanged(const FAssetData& NewObject, FNiagaraVariable UserParameter) = 0;
+
+	static FNiagaraVariant GetParameterValueFromSystem(const FNiagaraVariableBase& UserParameter, const UNiagaraSystem& System);
+	
+	/** Main function to generate a value parameter row. Sets up the parameter editor & value changed callback, display data & property handle. */
+	class IDetailPropertyRow* AddValueParameterAsRow(IDetailChildrenBuilder& ChildrenBuilder, FNiagaraVariable ChoppedUserParameter);
+	class IDetailPropertyRow* AddObjectAssetParameterAsRow(IDetailChildrenBuilder& ChildrenBuilder, FNiagaraVariable ChoppedUserParameter);
+	
+	virtual TSharedRef<SWidget> CreateUserParameterNameWidget(FNiagaraVariable UserParameter);
+	
+	virtual void PostUndo(bool bSuccess) override { Rebuild(); }
+	virtual void PostRedo(bool bSuccess) override { Rebuild(); }
+
+	void OnScriptVariablePropertyChanged(UNiagaraScriptVariable* ScriptVariable, const FPropertyChangedEvent& PropertyChangedEvent) const;
+
+	virtual void AddReferencedObjects(FReferenceCollector& Collector) override;
+	virtual FString GetReferencerName() const override;
+protected:
+	bool bDelegatesInitialized = false;
+	FSimpleDelegate OnRebuildChildren;
+	FName CustomBuilderRowName;
+	const UNiagaraHierarchySection* ActiveSection = nullptr;
+	TArray<UNiagaraHierarchySection*> AvailableSections;
+	TMap<FNiagaraVariable, TSharedRef<SNiagaraParameterEditor>> NiagaraParameterEditors;
+	/** The struct on scopes contained here are owned by the builder. When parameter editor value changes, they get copied into the display data, and from there into the parameter store.
+	 * This is done because the types the parameter editors operate on are not the same as the struct on scopes we are editing (bool vs. FNiagaraBool etc.) */
+	TMap<FNiagaraVariable, TSharedRef<FStructOnScope>> DisplayData;
+	TMap<FNiagaraVariable, TSharedPtr<IPropertyHandle>> DisplayDataPropertyHandles;
+	TArray<UObject*> ObjectAssetHelpers;
+};
+
+/** The category builder will display all sub-categories & parameters contained within a given category in the details panel. */
+class FNiagaraUserParameterCategoryBuilder : public IDetailCustomNodeBuilder, public TSharedFromThis<FNiagaraUserParameterCategoryBuilder>
+{
+public:
+	FNiagaraUserParameterCategoryBuilder(const UNiagaraHierarchyCategory* InNiagaraHierarchyCategory, FName InCustomBuilderRowName, FNiagaraUserParameterNodeBuilder& InNodeBuilder);
+	
+	virtual void GenerateHeaderRowContent(FDetailWidgetRow& NodeRow) override;
+	
+	virtual bool InitiallyCollapsed() const override { return false; }
+
+	virtual FName GetName() const  override
+	{
+		return CustomBuilderRowName;
+	}
+	
+	virtual void GenerateChildContent(IDetailChildrenBuilder& ChildrenBuilder) override;
+
+private:
+	TWeakObjectPtr<const UNiagaraHierarchyCategory> HierarchyCategory;
+	FName CustomBuilderRowName;
+	FNiagaraUserParameterNodeBuilder& NodeBuilder;
+};
diff --git a/Engine/Plugins/FX/Niagara/Source/NiagaraEditor/Private/ViewModels/HierarchyEditor/NiagaraUserParametersHierarchyViewModel.cpp b/Engine/Plugins/FX/Niagara/Source/NiagaraEditor/Private/ViewModels/HierarchyEditor/NiagaraUserParametersHierarchyViewModel.cpp
--- a/Engine/Plugins/FX/Niagara/Source/NiagaraEditor/Private/ViewModels/HierarchyEditor/NiagaraUserParametersHierarchyViewModel.cpp
+++ b/Engine/Plugins/FX/Niagara/Source/NiagaraEditor/Private/ViewModels/HierarchyEditor/NiagaraUserParametersHierarchyViewModel.cpp
@@ -13,7 +13,7 @@
 
 #define LOCTEXT_NAMESPACE "NiagaraUserParametersHierarchyEditor"
 
-void UNiagaraHierarchyUserParameter::Initialize(UNiagaraScriptVariable& InUserParameterScriptVariable, UNiagaraSystem& InSystem)
+void UNiagaraHierarchyUserParameter::Initialize(UNiagaraScriptVariable& InUserParameterScriptVariable)
 {
 	UserParameterScriptVariable = &InUserParameterScriptVariable;
 	SetIdentity(FNiagaraHierarchyIdentity({InUserParameterScriptVariable.Metadata.GetVariableGuid()}, {}));
@@ -122,7 +122,7 @@ void UNiagaraUserParametersHierarchyViewModel::PrepareSourceItems(UNiagaraHierar
 		
 		// since the source items are transient we need to create them here and keep them around until the end of the tool's lifetime
 		UNiagaraHierarchyUserParameter* UserParameterHierarchyObject = NewObject<UNiagaraHierarchyUserParameter>(SourceRoot);
-		UserParameterHierarchyObject->Initialize(*ScriptVariable.Get(), GetSystemViewModel()->GetSystem());
+		UserParameterHierarchyObject->Initialize(*ScriptVariable.Get());
 		SourceRoot->GetChildrenMutable().Add(UserParameterHierarchyObject);
 	}
 }
diff --git a/Engine/Plugins/FX/Niagara/Source/NiagaraEditor/Public/ViewModels/HierarchyEditor/NiagaraUserParametersHierarchyViewModel.h b/Engine/Plugins/FX/Niagara/Source/NiagaraEditor/Public/ViewModels/HierarchyEditor/NiagaraUserParametersHierarchyViewModel.h
--- a/Engine/Plugins/FX/Niagara/Source/NiagaraEditor/Public/ViewModels/HierarchyEditor/NiagaraUserParametersHierarchyViewModel.h
+++ b/Engine/Plugins/FX/Niagara/Source/NiagaraEditor/Public/ViewModels/HierarchyEditor/NiagaraUserParametersHierarchyViewModel.h
@@ -28,11 +28,13 @@ public:
 	UNiagaraHierarchyUserParameter() {}
 	virtual ~UNiagaraHierarchyUserParameter() override {}
 	
-	void Initialize(UNiagaraScriptVariable& InUserParameterScriptVariable, UNiagaraSystem& System);
+	void Initialize(UNiagaraScriptVariable& InUserParameterScriptVariable);
 	
 	const FNiagaraVariable& GetUserParameter() const { return UserParameterScriptVariable->Variable; }
 
 	virtual FString ToString() const override { return GetUserParameter().GetName().ToString(); }
+
+	const UNiagaraScriptVariable* GetScriptVariable() const { return UserParameterScriptVariable; }
 private:
 	UPROPERTY()
 	TObjectPtr<UNiagaraScriptVariable> UserParameterScriptVariable;
