Add component and life cycle track by default when adding a Niagara actor to sequencer

diff --git a/Engine/Plugins/FX/Niagara/Source/NiagaraEditor/NiagaraEditor.Build.cs b/Engine/Plugins/FX/Niagara/Source/NiagaraEditor/NiagaraEditor.Build.cs
--- a/Engine/Plugins/FX/Niagara/Source/NiagaraEditor/NiagaraEditor.Build.cs
+++ b/Engine/Plugins/FX/Niagara/Source/NiagaraEditor/NiagaraEditor.Build.cs
@@ -47,7 +47,8 @@ public class NiagaraEditor : ModuleRules
 				"AssetDefinition",
 				"ContentBrowser",
 				"ToolWidgets",
-				"AssetTools"
+				"AssetTools",
+				"LevelSequence",
 			}
 		);
 
diff --git a/Engine/Plugins/FX/Niagara/Source/NiagaraEditor/Private/NiagaraEditorModule.cpp b/Engine/Plugins/FX/Niagara/Source/NiagaraEditor/Private/NiagaraEditorModule.cpp
--- a/Engine/Plugins/FX/Niagara/Source/NiagaraEditor/Private/NiagaraEditorModule.cpp
+++ b/Engine/Plugins/FX/Niagara/Source/NiagaraEditor/Private/NiagaraEditorModule.cpp
@@ -4,7 +4,6 @@
 #include "NiagaraModule.h"
 #include "NiagaraEditorTickables.h"
 #include "Modules/ModuleManager.h"
-#include "IAssetTypeActions.h"
 #include "AssetToolsModule.h"
 #include "ISequencerModule.h"
 #include "ISettingsModule.h"
@@ -99,7 +98,6 @@
 
 #include "Customizations/NiagaraComponentDetails.h"
 #include "Customizations/NiagaraFunctionCallNodeDetails.h"
-#include "Customizations/NiagaraEventScriptPropertiesCustomization.h"
 #include "Customizations/NiagaraParameterBindingCustomization.h"
 #include "Customizations/NiagaraPlatformSetCustomization.h"
 #include "Customizations/NiagaraScriptVariableCustomization.h"
@@ -122,7 +120,6 @@
 #include "NiagaraMeshRendererProperties.h"
 
 #include "HAL/IConsoleManager.h"
-#include "NiagaraHlslTranslator.h"
 #include "NiagaraThumbnailRenderer.h"
 #include "Misc/FeedbackContext.h"
 #include "NiagaraNodeFunctionCall.h"
@@ -131,6 +128,7 @@
 #include "INiagaraEditorOnlyDataUtlities.h"
 
 #include "Editor.h"
+#include "ILevelSequenceModule.h"
 #include "HAL/PlatformFileManager.h"
 #include "HAL/FileManager.h"
 #include "ISourceControlOperation.h"
@@ -166,7 +164,6 @@
 #include "NiagaraDataChannel.h"
 
 #include "Engine/AssetManager.h"
-#include "Engine/AssetManagerSettings.h"
 #include "ViewModels/HierarchyEditor/NiagaraHierarchyCommands.h"
 
 #include UE_INLINE_GENERATED_CPP_BY_NAME(NiagaraEditorModule)
@@ -1430,6 +1427,8 @@ void FNiagaraEditorModule::StartupModule()
 	INiagaraDataInterfaceNodeActionProvider::Register<UNiagaraDataInterfaceDataChannelRead, FNiagaraDataInterfaceNodeActionProvider_DataChannelRead>();
 
 	FAssetCompilingManager::Get().RegisterManager(&FNiagaraSystemCompilingManager::Get());
+	ILevelSequenceModule& LevelSequenceModule = FModuleManager::LoadModuleChecked<ILevelSequenceModule>("LevelSequence");
+	DefaultTrackHandle = LevelSequenceModule.OnNewActorTrackAdded().AddStatic(&FNiagaraSystemTrackEditor::AddDefaultSystemTracks);
 }
 
 void FNiagaraEditorModule::ShutdownModule()
@@ -1461,6 +1460,11 @@ void FNiagaraEditorModule::ShutdownModule()
 		FPropertyEditorModule& PropertyModule = FModuleManager::LoadModuleChecked<FPropertyEditorModule>("PropertyEditor");
 		PropertyModule.UnregisterCustomClassLayout("NiagaraComponent");
 	}
+	if (FModuleManager::Get().IsModuleLoaded("LevelSequence"))
+	{
+		ILevelSequenceModule& LevelSequenceModule = FModuleManager::LoadModuleChecked<ILevelSequenceModule>("LevelSequence");
+		LevelSequenceModule.OnNewActorTrackAdded().Remove(DefaultTrackHandle);
+	}
 
 	FNiagaraEditorStyle::Unregister();
 
diff --git a/Engine/Plugins/FX/Niagara/Source/NiagaraEditor/Private/Sequencer/LevelSequence/NiagaraSystemTrackEditor.cpp b/Engine/Plugins/FX/Niagara/Source/NiagaraEditor/Private/Sequencer/LevelSequence/NiagaraSystemTrackEditor.cpp
--- a/Engine/Plugins/FX/Niagara/Source/NiagaraEditor/Private/Sequencer/LevelSequence/NiagaraSystemTrackEditor.cpp
+++ b/Engine/Plugins/FX/Niagara/Source/NiagaraEditor/Private/Sequencer/LevelSequence/NiagaraSystemTrackEditor.cpp
@@ -178,6 +178,71 @@ TSharedPtr<SWidget> FNiagaraSystemTrackEditor::BuildOutlinerEditWidget(const FGu
 	return FMovieSceneTrackEditor::BuildOutlinerEditWidget(ObjectBinding, Track, Params);
 }
 
+void FNiagaraSystemTrackEditor::AddDefaultSystemTracks(const AActor& SourceActor, const FGuid& Binding,	TSharedPtr<ISequencer> Sequencer)
+{
+	ENiagaraAddDefaultsTrackMode TrackMode = GetDefault<UNiagaraEditorSettings>()->DefaultsSequencerSubtracks;
+	if (!SourceActor.IsA<ANiagaraActor>() || TrackMode == ENiagaraAddDefaultsTrackMode::NoSubtracks)
+	{
+		return;
+	}
+
+	const ANiagaraActor* NiagaraActor = Cast<ANiagaraActor>(&SourceActor);
+	UNiagaraComponent* NiagaraComponent = NiagaraActor->GetNiagaraComponent();
+	if (NiagaraComponent && NiagaraComponent->GetAsset())
+	{
+		UMovieSceneSequence* Sequence = Sequencer->GetFocusedMovieSceneSequence();
+		UMovieScene* MovieScene = Sequence->GetMovieScene();
+		FGuid ComponentBinding = Sequencer->GetHandleToObject(NiagaraComponent);
+		
+		if (MovieScene->FindSpawnable(Binding))
+		{
+			// we only want to add tracks for possessables
+			return;
+		}
+
+		if (TrackMode == ENiagaraAddDefaultsTrackMode::ComponentTrackOnly)
+		{
+			FMovieScenePossessable* BoundComponent = MovieScene->FindPossessable(ComponentBinding);
+			if (!BoundComponent)
+			{
+				Sequence->CreatePossessable(NiagaraComponent);
+			}
+		}
+		else if (TrackMode == ENiagaraAddDefaultsTrackMode::LifecycleTrack)
+		{
+			UClass* TrackClass = UMovieSceneNiagaraSystemTrack::StaticClass();
+			UMovieSceneTrack* NewTrack = MovieScene->FindTrack(TrackClass, ComponentBinding);
+			if (!NewTrack)
+			{
+				UNiagaraSystem* System = NiagaraComponent->GetAsset();
+
+				UMovieSceneNiagaraSystemTrack* NiagaraSystemTrack = MovieScene->AddTrack<UMovieSceneNiagaraSystemTrack>(ComponentBinding);
+				NiagaraSystemTrack->SetDisplayName(LOCTEXT("SystemLifeCycleTrackName", "System Life Cycle"));
+				UMovieSceneNiagaraSystemSpawnSection* SpawnSection = NewObject<UMovieSceneNiagaraSystemSpawnSection>(NiagaraSystemTrack, NAME_None, RF_Transactional);
+
+				FFrameRate FrameResolution = MovieScene->GetTickResolution();
+				FFrameTime SpawnSectionStartTime = Sequencer->GetLocalTime().ConvertTo(FrameResolution);
+				FFrameTime SpawnSectionDuration;
+
+				UNiagaraSystemEditorData* SystemEditorData = Cast<UNiagaraSystemEditorData>(System->GetEditorData());
+				if (SystemEditorData != nullptr && SystemEditorData->GetPlaybackRange().HasLowerBound() && SystemEditorData->GetPlaybackRange().HasUpperBound())
+				{
+					SpawnSectionDuration = FrameResolution.AsFrameTime(SystemEditorData->GetPlaybackRange().Size<float>());
+				}
+				else
+				{
+					SpawnSectionDuration = FrameResolution.AsFrameTime(5.0);
+				}
+
+				SpawnSection->SetRange(TRange<FFrameNumber>(
+					SpawnSectionStartTime.RoundToFrame(),
+					(SpawnSectionStartTime + SpawnSectionDuration).RoundToFrame()));
+				NiagaraSystemTrack->AddSection(*SpawnSection);
+			}
+		}
+	}
+}
+
 void FNiagaraSystemTrackEditor::AddNiagaraSystemTrack(TArray<FGuid> ObjectBindings)
 {
 	UMovieScene* MovieScene = GetSequencer()->GetFocusedMovieSceneSequence()->GetMovieScene();
diff --git a/Engine/Plugins/FX/Niagara/Source/NiagaraEditor/Private/Sequencer/LevelSequence/NiagaraSystemTrackEditor.h b/Engine/Plugins/FX/Niagara/Source/NiagaraEditor/Private/Sequencer/LevelSequence/NiagaraSystemTrackEditor.h
--- a/Engine/Plugins/FX/Niagara/Source/NiagaraEditor/Private/Sequencer/LevelSequence/NiagaraSystemTrackEditor.h
+++ b/Engine/Plugins/FX/Niagara/Source/NiagaraEditor/Private/Sequencer/LevelSequence/NiagaraSystemTrackEditor.h
@@ -32,6 +32,7 @@ class FNiagaraSystemTrackEditor : public FMovieSceneTrackEditor
 {
 public:
 	static TSharedRef<ISequencerTrackEditor> CreateTrackEditor(TSharedRef<ISequencer> InSequencer);
+	static void AddDefaultSystemTracks(const AActor& SourceActor, const FGuid& Binding, TSharedPtr<ISequencer> Sequencer);
 
 public:
 	FNiagaraSystemTrackEditor(TSharedRef<ISequencer> InSequencer);
diff --git a/Engine/Plugins/FX/Niagara/Source/NiagaraEditor/Public/NiagaraEditorModule.h b/Engine/Plugins/FX/Niagara/Source/NiagaraEditor/Public/NiagaraEditorModule.h
--- a/Engine/Plugins/FX/Niagara/Source/NiagaraEditor/Public/NiagaraEditorModule.h
+++ b/Engine/Plugins/FX/Niagara/Source/NiagaraEditor/Public/NiagaraEditorModule.h
@@ -8,7 +8,6 @@
 #include "Toolkits/AssetEditorToolkit.h"
 #include "NiagaraTypes.h"
 #include "INiagaraCompiler.h"
-#include "AssetTypeCategories.h"
 #include "NiagaraCompilationTypes.h"
 #include "NiagaraPerfBaseline.h"
 #include "NiagaraDebuggerCommon.h"
@@ -388,6 +387,8 @@ private:
 	FDelegateHandle AssetRegistryOnLoadCompleteHandle;
 	
 	FDelegateHandle OnAssetsPreDeleteHandle;
+
+	FDelegateHandle DefaultTrackHandle;
 	
 	TObjectPtr<USequencerSettings> SequencerSettings;
 
diff --git a/Engine/Plugins/FX/Niagara/Source/NiagaraEditor/Public/NiagaraEditorSettings.h b/Engine/Plugins/FX/Niagara/Source/NiagaraEditor/Public/NiagaraEditorSettings.h
--- a/Engine/Plugins/FX/Niagara/Source/NiagaraEditor/Public/NiagaraEditorSettings.h
+++ b/Engine/Plugins/FX/Niagara/Source/NiagaraEditor/Public/NiagaraEditorSettings.h
@@ -76,6 +76,19 @@ enum class ENiagaraNamespaceMetadataOptions
 	HideInDefinitions,
 };
 
+UENUM()
+enum class ENiagaraAddDefaultsTrackMode : uint8
+{
+	/** Adding a Niagara actor to a sequence will not add any additional subtracks by default */
+	NoSubtracks,
+
+	/** Adding a Niagara actor will add a component subtrack, but nothing else */
+	ComponentTrackOnly,
+	
+	/** Adding a Niagara actor to a sequence will also add a life cycle track to it by default  */
+	LifecycleTrack,
+};
+
 USTRUCT()
 struct NIAGARAEDITOR_API FNiagaraNamespaceMetadata
 {
@@ -300,6 +313,9 @@ public:
 
 	UPROPERTY(config, EditAnywhere, Category = Niagara)
 	ENiagaraCategoryExpandState RendererCategoryExpandState = ENiagaraCategoryExpandState::Default;
+	
+	UPROPERTY(config, EditAnywhere, Category = Niagara)
+	ENiagaraAddDefaultsTrackMode DefaultsSequencerSubtracks = ENiagaraAddDefaultsTrackMode::NoSubtracks;
 
 	TArray<float> GetPlaybackSpeeds() const;
 
