Fix for bool NDC param reads.
Extra early out in SpawnConditional for reduced overhead when spawning nothing.

diff --git a/Engine/Plugins/FX/Niagara/Source/Niagara/Private/DataInterface/NiagaraDataInterfaceDataChannelRead.cpp b/Engine/Plugins/FX/Niagara/Source/Niagara/Private/DataInterface/NiagaraDataInterfaceDataChannelRead.cpp
--- a/Engine/Plugins/FX/Niagara/Source/Niagara/Private/DataInterface/NiagaraDataInterfaceDataChannelRead.cpp
+++ b/Engine/Plugins/FX/Niagara/Source/Niagara/Private/DataInterface/NiagaraDataInterfaceDataChannelRead.cpp
@@ -1077,8 +1077,6 @@ void UNiagaraDataInterfaceDataChannelRead::SpawnConditional(FVectorVMExternalFun
 	//This should only be called from emitter scripts and since it has per instance data then we process them individually.
 	check(Context.GetNumInstances() == 1);
 
-	FNDIRandomHelperFromStream RandHelper(Context);
-
 	VectorVM::FUserPtrHandler<FNDIDataChannelReadInstanceData> InstData(Context);
 
 	//Binding info can be null here as we can be spawning without any conditions, i.e. no variadic parameters to the function.
@@ -1096,9 +1094,11 @@ void UNiagaraDataInterfaceDataChannelRead::SpawnConditional(FVectorVMExternalFun
 	FNiagaraDataChannelData* DataChannelData = InstData->DataChannelData.Get();
 	FNiagaraDataBuffer* Data = DataChannelData ? DataChannelData->GetCPUData(!bReadCurrentFrame) : nullptr;
 
-	bool bSpawn = INiagaraModule::DataChannelsEnabled() && Data && EmittterInst && InEnabled.GetAndAdvance();
+	bool bSpawn = INiagaraModule::DataChannelsEnabled() && Data && Data->GetNumInstances() > 0 && EmittterInst && InEnabled.GetAndAdvance();
 	if(bSpawn)
 	{
+		FNDIRandomHelperFromStream RandHelper(Context);
+
 		TArray<FNiagaraSpawnInfo>& EmitterSpawnInfos = EmittterInst->GetSpawnInfo();
 		ENDIDataChannelSpawnMode Mode = (ENDIDataChannelSpawnMode)InMode.GetAndAdvance();
 		ENiagaraConditionalOperator Op = (ENiagaraConditionalOperator)InOp.GetAndAdvance();
@@ -1346,6 +1346,12 @@ void UNiagaraDataInterfaceDataChannelRead::GetParameterDefinitionHLSL(FNiagaraDa
 					FunctionParameterComponentType = HlslGenContext.GetStructHlslTypeName(FNiagaraTypeDefinition(Struct));
 					OutCode += FString::Format(bRead ? *ReadDataTemplate : *WriteDataTemplate, HlslTemplateArgs);
 				}
+				else if (Struct == FNiagaraTypeDefinition::GetBoolStruct())
+				{
+					FunctionParameterComponentBufferType = TEXT("Int32");
+					FunctionParameterComponentType = HlslGenContext.GetStructHlslTypeName(FNiagaraTypeDefinition::GetIntDef());
+					OutCode += FString::Format(bRead ? *ReadDataTemplate : *WriteDataTemplate, HlslTemplateArgs);
+				}
 				else
 				{
 					FString PropertyBaseName = FunctionParameterComponentName.StringValue;
