Fix lambda parameters

diff --git a/Engine/Plugins/FX/Niagara/Source/NiagaraEditor/Private/ViewModels/NiagaraParameterPanelViewModel.cpp b/Engine/Plugins/FX/Niagara/Source/NiagaraEditor/Private/ViewModels/NiagaraParameterPanelViewModel.cpp
--- a/Engine/Plugins/FX/Niagara/Source/NiagaraEditor/Private/ViewModels/NiagaraParameterPanelViewModel.cpp
+++ b/Engine/Plugins/FX/Niagara/Source/NiagaraEditor/Private/ViewModels/NiagaraParameterPanelViewModel.cpp
@@ -1659,7 +1659,7 @@ bool FNiagaraSystemToolkitParameterPanelViewModel::GetCanSetParameterCustomNames
 
 TSharedRef<SWidget> FNiagaraSystemToolkitParameterPanelViewModel::CreateAddParameterMenuForAssignmentNode(UNiagaraNodeAssignment* AssignmentNode, const TSharedPtr<SComboButton>& AddButton) 
 {
-	auto AddParameterLambda = [this, AssignmentNode](FNiagaraVariable& NewParameter) {
+	auto AddParameterLambda = [this, AssignmentNode](FNiagaraVariable NewParameter) {
 		// If an assignment target is already setting the associated parameter, make a unique name here so that the assignment target is valid.
 		// Because the unique name is made before calling FNiagaraSystemToolkitParameterPanelViewModel::AddParameter, pass bMakeUniqueName = false.
 		if (AssignmentNode->GetAssignmentTargets().Contains(NewParameter))
@@ -1682,7 +1682,7 @@ TSharedRef<SWidget> FNiagaraSystemToolkitParameterPanelViewModel::CreateAddParam
 		AssignmentNode->AddParameter(NewParameter, VarDefaultValue);
 	};
 
-	auto AddExistingParameterLambda = [this, AssignmentNode](FNiagaraVariable& NewParameter) {
+	auto AddExistingParameterLambda = [this, AssignmentNode](const FNiagaraVariable& NewParameter) {
 		// If an assignment target is already setting the associated parameter, we simply don't do anything
 		if (AssignmentNode->GetAssignmentTargets().Contains(NewParameter))
 		{
