Optionally include cascade FX in a limited fashion in the Debug HUD

diff --git a/Engine/Plugins/FX/Niagara/Source/Niagara/Private/NiagaraDebugHud.cpp b/Engine/Plugins/FX/Niagara/Source/Niagara/Private/NiagaraDebugHud.cpp
--- a/Engine/Plugins/FX/Niagara/Source/Niagara/Private/NiagaraDebugHud.cpp
+++ b/Engine/Plugins/FX/Niagara/Source/Niagara/Private/NiagaraDebugHud.cpp
@@ -28,6 +28,8 @@
 #include "SceneView.h"
 #include "UObject/UObjectIterator.h"
 
+#include "ParticleEmitterInstances.h"
+
 #if WITH_NIAGARA_DEBUGGER
 
 namespace NiagaraDebugLocal
@@ -999,37 +1001,39 @@ void FNiagaraDebugHud::GatherSystemInfo()
 	}
 
 	// Iterate all components looking for active ones in the world we are in
-	for (TObjectIterator<UNiagaraComponent> It; It; ++It)
+	for (TObjectIterator<UFXSystemComponent> It; It; ++It)
 	{
-		UNiagaraComponent* NiagaraComponent = *It;
-		if (!IsValidChecked(NiagaraComponent) || NiagaraComponent->IsUnreachable() || NiagaraComponent->HasAnyFlags(EObjectFlags::RF_ClassDefaultObject))
+		UFXSystemComponent* FXComponent = *It;
+		UNiagaraComponent* NiagaraComponent = Cast<UNiagaraComponent>(FXComponent);
+		UParticleSystemComponent* CascadeComponent = Cast<UParticleSystemComponent>(FXComponent);
+
+		FNiagaraSystemInstanceControllerPtr SystemInstanceController = NiagaraComponent ? NiagaraComponent->GetSystemInstanceController() : nullptr;
+		FNiagaraSystemInstance* SystemInstance = SystemInstanceController.IsValid() ? SystemInstanceController->GetSystemInstance_Unsafe() : nullptr;
+		
+
+		if (!IsValidChecked(FXComponent) || (!IsValid(NiagaraComponent) && !IsValid(CascadeComponent)) || FXComponent->IsUnreachable() || FXComponent->HasAnyFlags(EObjectFlags::RF_ClassDefaultObject))
 		{
 			continue;
 		}
-		if (NiagaraComponent->GetWorld() != World)
+		if (FXComponent->GetWorld() != World)
 		{
 			continue;
 		}
-		if (NiagaraComponent->GetAsset() == nullptr)
+		if (FXComponent->GetFXSystemAsset() == nullptr)
 		{
 			continue;
 		}
 
-		FNiagaraSystemInstanceControllerPtr SystemInstanceController = NiagaraComponent->GetSystemInstanceController();
-		FNiagaraSystemInstance* SystemInstance = SystemInstanceController.IsValid() ? SystemInstanceController->GetSystemInstance_Unsafe() : nullptr;
-
-		check(NiagaraComponent->GetAsset() != nullptr);
+		const bool bIsActive = FXComponent->IsActive();
+		const bool bHasScalability = NiagaraComponent ? NiagaraComponent->IsRegisteredWithScalabilityManager() : CascadeComponent->bIsManagingSignificance;
 
-		const bool bIsActive = NiagaraComponent->IsActive();
-		const bool bHasScalability = NiagaraComponent->IsRegisteredWithScalabilityManager();
-
-		FSystemDebugInfo& SystemDebugInfo = PerSystemDebugInfo.FindOrAdd(NiagaraComponent->GetAsset()->GetFName());
+		FSystemDebugInfo& SystemDebugInfo = PerSystemDebugInfo.FindOrAdd(FXComponent->GetFXSystemAsset()->GetFName());
 		if (SystemDebugInfo.SystemName.IsEmpty())
 		{
-			SystemDebugInfo.SystemName = GetNameSafe(NiagaraComponent->GetAsset());
+			SystemDebugInfo.SystemName = GetNameSafe(FXComponent->GetFXSystemAsset());
 		}
 	#if WITH_EDITORONLY_DATA
-		SystemDebugInfo.bCompileForEdit = NiagaraComponent->GetAsset()->bCompileForEdit;
+		SystemDebugInfo.bCompileForEdit = NiagaraComponent ? NiagaraComponent->GetAsset()->bCompileForEdit : false;
 	#endif
 		SystemDebugInfo.bShowInWorld = Settings.bSystemFilterEnabled && SystemDebugInfo.SystemName.MatchesWildcard(Settings.SystemFilter);
 		SystemDebugInfo.bPassesSystemFilter = !Settings.bSystemFilterEnabled || SystemDebugInfo.SystemName.MatchesWildcard(Settings.SystemFilter);
@@ -1041,23 +1045,23 @@ void FNiagaraDebugHud::GatherSystemInfo()
 			// Filter by actor
 			if ( Settings.bActorFilterEnabled )
 			{
-				AActor* Actor = NiagaraComponent->GetOwner();
+				AActor* Actor = FXComponent->GetOwner();
 				bIsMatch &= (Actor != nullptr) && Actor->GetActorNameOrLabel().MatchesWildcard(Settings.ActorFilter);
 			}
 
 			// Filter by component
 			if ( bIsMatch && Settings.bComponentFilterEnabled )
 			{
-				bIsMatch &= NiagaraComponent->GetName().MatchesWildcard(Settings.ComponentFilter);
+				bIsMatch &= FXComponent->GetName().MatchesWildcard(Settings.ComponentFilter);
 			}
 
-			if (bIsMatch)
+			if (bIsMatch && NiagaraComponent)//TODO: Handle Cascade?
 			{
 				InWorldComponents.Add(NiagaraComponent);
 			}
 		}
 
-		if (NiagaraComponent->IsRegistered())
+		if (FXComponent->IsRegistered())
 		{
 			++GlobalTotalRegistered;
 			++SystemDebugInfo.TotalRegistered;
@@ -1069,13 +1073,13 @@ void FNiagaraDebugHud::GatherSystemInfo()
 			++SystemDebugInfo.TotalScalability;
 		}
 
-		if (NiagaraComponent->IsLocalPlayerEffect() && bIsActive)
+		if (NiagaraComponent && NiagaraComponent->IsLocalPlayerEffect() && bIsActive)
 		{
 			++GlobalTotalPlayerSystems;
 			++SystemDebugInfo.TotalPlayerSystems;
 		}
 
-		if (NiagaraComponent->IsRegisteredWithScalabilityManager())
+		if (NiagaraComponent && NiagaraComponent->IsRegisteredWithScalabilityManager())
 		{
 #if WITH_NIAGARA_DEBUGGER
 			if (NiagaraComponent->DebugCachedScalabilityState.bCulled)
@@ -1114,26 +1118,44 @@ void FNiagaraDebugHud::GatherSystemInfo()
 
 		if (bIsActive)
 		{
-			check(SystemInstance);
 			// Accumulate totals
 			int32 ActiveEmitters = 0;
 			int32 TotalEmitters = 0;
 			int32 ActiveParticles = 0;
-
-			for (const TSharedRef<FNiagaraEmitterInstance, ESPMode::ThreadSafe>& EmitterInstance : SystemInstance->GetEmitters())
+			if( NiagaraComponent)
 			{
-				UNiagaraEmitter* NiagaraEmitter = EmitterInstance->GetCachedEmitter().Emitter;
-				if (NiagaraEmitter == nullptr)
+				check(SystemInstance);
+				for (const TSharedRef<FNiagaraEmitterInstance, ESPMode::ThreadSafe>& EmitterInstance : SystemInstance->GetEmitters())
 				{
-					continue;
-				}
+					UNiagaraEmitter* NiagaraEmitter = EmitterInstance->GetCachedEmitter().Emitter;
+					if (NiagaraEmitter == nullptr)
+					{
+						continue;
+					}
 
-				++TotalEmitters;
-				if (EmitterInstance->GetExecutionState() == ENiagaraExecutionState::Active)
-				{
-					++ActiveEmitters;
+					++TotalEmitters;
+					if (EmitterInstance->GetExecutionState() == ENiagaraExecutionState::Active)
+					{
+						++ActiveEmitters;
+					}
+					ActiveParticles += EmitterInstance->GetNumParticles();
+				}
+			}
+			else 
+			{
+				check(CascadeComponent);				
+				for(FParticleEmitterInstance* Emitter : CascadeComponent->EmitterInstances)
+				{					
+					if(Emitter)
+					{
+						++TotalEmitters;
+						ActiveParticles += Emitter->ActiveParticles;
+						if(Emitter->HasCompleted() == false)
+						{
+							++ActiveEmitters;
+						}
+					}
 				}
-				ActiveParticles += EmitterInstance->GetNumParticles();
 			}
 
 			++SystemDebugInfo.TotalActive;
@@ -1147,13 +1169,13 @@ void FNiagaraDebugHud::GatherSystemInfo()
 #if WITH_PER_SYSTEM_PARTICLE_PERF_STATS
 			if(StatsListener)
 			{
-				SystemDebugInfo.PerfStats = StatsListener->GetSystemStats(NiagaraComponent->GetAsset());
+				SystemDebugInfo.PerfStats = StatsListener->GetSystemStats(FXComponent->GetFXSystemAsset());
 			}
 #endif
 		}
 
 		//Generate a unique-ish random color for use in graphs and the world to help visually ID this system.
-		FRandomStream Rands(GetTypeHash(NiagaraComponent->GetAsset()->GetName()) + Settings.SystemColorSeed);
+		FRandomStream Rands(GetTypeHash(FXComponent->GetFXSystemAsset()->GetName()) + Settings.SystemColorSeed);
 		uint8 RandomHue = (uint8)Rands.RandRange(int32(Settings.SystemColorHSVMin.X), int32(Settings.SystemColorHSVMax.X));
 		uint8 RandomSat = (uint8)Rands.RandRange(int32(Settings.SystemColorHSVMin.Y), int32(Settings.SystemColorHSVMax.Y));
 		uint8 RandomValue = (uint8)Rands.RandRange(int32(Settings.SystemColorHSVMin.Z), int32(Settings.SystemColorHSVMax.Z));
@@ -3202,7 +3224,7 @@ void FNiagaraDebugHudStatHistory::GetHistoryFrames_GPU(TArray<double>& OutHistor
 	} while (WriteFrame != CurrFrameGPU);
 }
 
-TSharedPtr<FNiagaraDebugHUDPerfStats> FNiagaraDebugHUDStatsListener::GetSystemStats(UNiagaraSystem* System)
+TSharedPtr<FNiagaraDebugHUDPerfStats> FNiagaraDebugHUDStatsListener::GetSystemStats(UFXSystemAsset* System)
 {
 	FScopeLock Lock(&SystemStatsGuard);//Locking on every get isn't great but it's debug hud code so meh.
 	if (TSharedPtr<FNiagaraDebugHUDPerfStats>* Stats = SystemStats.Find(System))
diff --git a/Engine/Plugins/FX/Niagara/Source/Niagara/Private/NiagaraDebugHud.h b/Engine/Plugins/FX/Niagara/Source/Niagara/Private/NiagaraDebugHud.h
--- a/Engine/Plugins/FX/Niagara/Source/Niagara/Private/NiagaraDebugHud.h
+++ b/Engine/Plugins/FX/Niagara/Source/Niagara/Private/NiagaraDebugHud.h
@@ -77,7 +77,7 @@ public:
 
 	virtual bool Tick()override;
 	virtual void TickRT()override;
-	TSharedPtr<FNiagaraDebugHUDPerfStats> GetSystemStats(UNiagaraSystem* System);
+	TSharedPtr<FNiagaraDebugHUDPerfStats> GetSystemStats(UFXSystemAsset* System);
 	FNiagaraDebugHUDPerfStats& GetGlobalStats();
 
 	virtual void OnAddSystem(const TWeakObjectPtr<const UFXSystemAsset>& NewSystem)override;
diff --git a/Engine/Plugins/FX/Niagara/Source/Niagara/Public/NiagaraDebuggerCommon.h b/Engine/Plugins/FX/Niagara/Source/Niagara/Public/NiagaraDebuggerCommon.h
--- a/Engine/Plugins/FX/Niagara/Source/Niagara/Public/NiagaraDebuggerCommon.h
+++ b/Engine/Plugins/FX/Niagara/Source/Niagara/Public/NiagaraDebuggerCommon.h
@@ -456,6 +456,10 @@ struct NIAGARA_API FNiagaraDebugHUDSettingsData
 	UPROPERTY(EditAnywhere, Category = "Debug Overview", meta = (DisplayName = "Debug Overview Enabled"))
 	bool bOverviewEnabled = false;
 
+	/** When enabled the overview display will include cascade FX. */
+	UPROPERTY(EditAnywhere, Category = "Debug Overview")
+	bool bIncludeCascade = true;
+
 	UPROPERTY(EditAnywhere, Category = "Debug Overview", meta = (DisplayName = "Debug Overview Mode"))
 	ENiagaraDebugHUDOverviewMode OverviewMode = ENiagaraDebugHUDOverviewMode::Overview;
 
