Fixed bug when creating a SparseVolumeTextureViewer, binding a vertex buffer to a byte address buffer on Mac OS uses buffer backed texture overwritting Texture slot 0

diff --git a/Engine/Plugins/FX/Niagara/Source/Niagara/Private/NiagaraDataInterfaceRenderTargetVolume.cpp b/Engine/Plugins/FX/Niagara/Source/Niagara/Private/NiagaraDataInterfaceRenderTargetVolume.cpp
--- a/Engine/Plugins/FX/Niagara/Source/Niagara/Private/NiagaraDataInterfaceRenderTargetVolume.cpp
+++ b/Engine/Plugins/FX/Niagara/Source/Niagara/Private/NiagaraDataInterfaceRenderTargetVolume.cpp
@@ -910,7 +910,7 @@ bool UNiagaraDataInterfaceRenderTargetVolume::SimCacheReadFrame(UObject* Storage
 						PassParameters->TileDataTextureSampler = TStaticSamplerState<SF_Point, AM_Clamp, AM_Clamp, AM_Clamp>::GetRHI();
 						PassParameters->SparseVolumeTexturePageTable = PageTableTexture ? PageTableTexture : GBlackUintVolumeTexture->TextureRHI.GetReference();
 						PassParameters->SparseVolumeTextureA = TextureA ? TextureA : GBlackVolumeTexture->TextureRHI.GetReference();
-						PassParameters->StreamingInfoBuffer = StreamingInfoBufferSRV ? StreamingInfoBufferSRV : GEmptyVertexBufferWithUAV->ShaderResourceViewRHI.GetReference();
+						PassParameters->StreamingInfoBuffer = StreamingInfoBufferSRV ? StreamingInfoBufferSRV : GEmptyStructuredBufferWithUAV->ShaderResourceViewRHI.GetReference();
 							
 						PassParameters->PackedSVTUniforms0 = CurrentPackedUniforms0;
 						PassParameters->PackedSVTUniforms1 = CurrentPackedUniforms1;
diff --git a/Engine/Plugins/FX/Niagara/Source/Niagara/Private/NiagaraDataInterfaceSparseVolumeTexture.cpp b/Engine/Plugins/FX/Niagara/Source/Niagara/Private/NiagaraDataInterfaceSparseVolumeTexture.cpp
--- a/Engine/Plugins/FX/Niagara/Source/Niagara/Private/NiagaraDataInterfaceSparseVolumeTexture.cpp
+++ b/Engine/Plugins/FX/Niagara/Source/Niagara/Private/NiagaraDataInterfaceSparseVolumeTexture.cpp
@@ -285,7 +285,7 @@ void UNiagaraDataInterfaceSparseVolumeTexture::SetShaderParameters(const FNiagar
 	Parameters->PageTableTexture = GBlackUintVolumeTexture->TextureRHI;
 	Parameters->PhysicalTileDataATexture = GBlackVolumeTexture->TextureRHI;
 	Parameters->PhysicalTileDataBTexture = GBlackVolumeTexture->TextureRHI;
-	Parameters->StreamingInfoBuffer = GEmptyVertexBufferWithUAV->ShaderResourceViewRHI;
+	Parameters->StreamingInfoBuffer = GEmptyStructuredBufferWithUAV->ShaderResourceViewRHI;
 	Parameters->PackedUniforms0 = FUintVector4();
 	Parameters->PackedUniforms1 = FUintVector4();
 	Parameters->TextureSize = FIntVector3::ZeroValue;
diff --git a/Engine/Source/Runtime/Engine/Private/Materials/MaterialUniformExpressions.cpp b/Engine/Source/Runtime/Engine/Private/Materials/MaterialUniformExpressions.cpp
--- a/Engine/Source/Runtime/Engine/Private/Materials/MaterialUniformExpressions.cpp
+++ b/Engine/Source/Runtime/Engine/Private/Materials/MaterialUniformExpressions.cpp
@@ -1280,7 +1280,7 @@ void FUniformExpressionSet::FillUniformBuffer(const FMaterialRenderContext& Mate
 			*ResourceTableTexturePageTablePtr = GBlackUintVolumeTexture->TextureRHI;
 			*ResourceTableTexturePhysicalAPtr = GBlackVolumeTexture->TextureRHI;
 			*ResourceTableTexturePhysicalBPtr = GBlackVolumeTexture->TextureRHI;
-			*ResourceTableStreamingInfoBufferPtr = GEmptyVertexBufferWithUAV->ShaderResourceViewRHI;
+			*ResourceTableStreamingInfoBufferPtr = GEmptyStructuredBufferWithUAV->ShaderResourceViewRHI;
 			check(GBlackVolumeTexture->SamplerStateRHI);
 			*ResourceTablePhysicalSamplerPtr = GBlackVolumeTexture->SamplerStateRHI;
 			
diff --git a/Engine/Source/Runtime/Renderer/Private/SparseVolumeTexture/SparseVolumeTextureViewerRendering.cpp b/Engine/Source/Runtime/Renderer/Private/SparseVolumeTexture/SparseVolumeTextureViewerRendering.cpp
--- a/Engine/Source/Runtime/Renderer/Private/SparseVolumeTexture/SparseVolumeTextureViewerRendering.cpp
+++ b/Engine/Source/Runtime/Renderer/Private/SparseVolumeTexture/SparseVolumeTextureViewerRendering.cpp
@@ -135,7 +135,7 @@ void AddSparseVolumeTextureViewerRenderPass(FRDGBuilder& GraphBuilder, FSceneRen
 			PsPassParameters->SparseVolumeTexturePageTable = GBlackUintVolumeTexture->TextureRHI;
 			PsPassParameters->SparseVolumeTextureA = GBlackVolumeTexture->TextureRHI;
 			PsPassParameters->SparseVolumeTextureB = GBlackVolumeTexture->TextureRHI;
-			PsPassParameters->StreamingInfoBuffer = GEmptyVertexBufferWithUAV->ShaderResourceViewRHI;
+			PsPassParameters->StreamingInfoBuffer = GEmptyStructuredBufferWithUAV->ShaderResourceViewRHI;
 			PsPassParameters->PackedSVTUniforms0 = FUintVector4(0);
 			PsPassParameters->PackedSVTUniforms1 = FUintVector4(0);
 			PsPassParameters->SparseVolumeTextureResolution = SVTSceneProxy->VolumeResolution;
