- Add an option to remove system simulations when they become unique (i.e. referenced only by the WM)
- This can reduce the load on the World Manager at the cost of potentially having to spin up more simulations

diff --git a/Engine/Plugins/FX/Niagara/Source/Niagara/Private/NiagaraWorldManager.cpp b/Engine/Plugins/FX/Niagara/Source/Niagara/Private/NiagaraWorldManager.cpp
--- a/Engine/Plugins/FX/Niagara/Source/Niagara/Private/NiagaraWorldManager.cpp
+++ b/Engine/Plugins/FX/Niagara/Source/Niagara/Private/NiagaraWorldManager.cpp
@@ -39,6 +39,14 @@ DECLARE_CYCLE_STAT(TEXT("Niagara Manager Wait Pre Garbage Collect [GT]"), STAT_N
 DECLARE_CYCLE_STAT(TEXT("Niagara Manager Refresh Owner Allows Scalability"), STAT_NiagaraWorldManRefreshOwnerAllowsScalability, STATGROUP_Niagara);
 DECLARE_CYCLE_STAT(TEXT("Niagara Manager Tick Parameter Collections [GT]"), STAT_NiagaraWorldManTickParamCollections, STATGROUP_Niagara);
 
+bool GNiagaraWorldManagerKillUniqueSims = 1;
+static FAutoConsoleVariableRef CVarNiagaraWorldManagerKillUniqueSims(
+	TEXT("fx.Niagara.WorldManager.KillUniqueSims"),
+	GNiagaraWorldManagerKillUniqueSims,
+	TEXT("System simulations will be removed when unique rather that waiting for GC."),
+	ECVF_Default
+);
+
 static int GNiagaraAllowAsyncWorkToEndOfFrame = 1;
 static FAutoConsoleVariableRef CVarNiagaraAllowAsyncWorkToEndOfFrame(
 	TEXT("fx.Niagara.AllowAsyncWorkToEndOfFrame"),
@@ -1270,12 +1278,14 @@ void FNiagaraWorldManager::Tick(ETickingGroup TickGroup, float DeltaSeconds, ELe
 
 	ActiveNiagaraTickGroup = ActualTickGroup;
 
+	const bool bKillUniqueSims = GNiagaraWorldManagerKillUniqueSims;
 	bool bTickFunctionCanOverlapTickGroups = GNiagaraAllowAsyncWorkToEndOfFrame != 0;
 	for (auto SimIt=SystemSimulations[ActualTickGroup].CreateIterator(); SimIt; ++SimIt)
 	{
 		FNiagaraSystemSimulation* Sim = &SimIt.Value().Get();
 
-		if (Sim->IsValid())
+		const bool bUniqueKill = bKillUniqueSims && SimIt.Value().IsUnique();
+		if (!bUniqueKill && Sim->IsValid())
 		{
 			Sim->Tick_GameThread(DeltaSeconds, MyCompletionGraphEvent);
 
