- Save ~20% of NPC binding parameter store pushing in test case

#rnx

diff --git a/Engine/Plugins/FX/Niagara/Source/Niagara/Private/NiagaraParameterStore.cpp b/Engine/Plugins/FX/Niagara/Source/Niagara/Private/NiagaraParameterStore.cpp
--- a/Engine/Plugins/FX/Niagara/Source/Niagara/Private/NiagaraParameterStore.cpp
+++ b/Engine/Plugins/FX/Niagara/Source/Niagara/Private/NiagaraParameterStore.cpp
@@ -318,30 +318,21 @@ void FNiagaraParameterStoreBinding::MatchParameters(FNiagaraParameterStore* Dest
 	}
 }
 
-void FNiagaraParameterStoreBinding::Tick(FNiagaraParameterStore* DestStore, FNiagaraParameterStore* SrcStore, bool bForce)
+void FNiagaraParameterStoreBinding::CopyParameters(FNiagaraParameterStore* DestStore, const FNiagaraParameterStore* SrcStore) const
 {
-	if (SrcStore->GetParametersDirty() || bForce)
+	for (const FParameterBinding& Binding : ParameterBindings)
 	{
-		for (FParameterBinding& Binding : ParameterBindings)
-		{
-			DestStore->SetParameterData(SrcStore->GetParameterData(Binding.SrcOffset), Binding.DestOffset, Binding.Size);
-		}
+		DestStore->SetParameterData(SrcStore->GetParameterData(Binding.SrcOffset), Binding.DestOffset, Binding.Size);
 	}
 
-	if (SrcStore->GetInterfacesDirty() || bForce)
+	for (const FInterfaceBinding& Binding : InterfaceBindings)
 	{
-		for (FInterfaceBinding& Binding : InterfaceBindings)
-		{
-			DestStore->SetDataInterface(SrcStore->GetDataInterface(Binding.SrcOffset), Binding.DestOffset);
-		}
+		DestStore->SetDataInterface(SrcStore->GetDataInterface(Binding.SrcOffset), Binding.DestOffset);
 	}
 
-	if (SrcStore->GetUObjectsDirty() || bForce)
+	for (const FUObjectBinding& Binding : UObjectBindings)
 	{
-		for (FUObjectBinding& Binding : UObjectBindings)
-		{
-			DestStore->SetUObject(SrcStore->GetUObject(Binding.SrcOffset), Binding.DestOffset);
-		}
+		DestStore->SetUObject(SrcStore->GetUObject(Binding.SrcOffset), Binding.DestOffset);
 	}
 
 #if NIAGARA_NAN_CHECKING
@@ -415,7 +406,7 @@ bool FNiagaraParameterStoreBinding::BindParameters(FNiagaraParameterStore* DestS
 	if (bAnyBinding)
 	{
 		//Force an initial tick to prime our values in the destination store.
-		Tick(DestStore, SrcStore, true);
+		CopyParameters(DestStore, SrcStore);
 	}
 	return bAnyBinding;
 }
@@ -553,13 +544,40 @@ void FNiagaraParameterStore::ConvertParameterType(const FNiagaraVariable& Existi
 }
 #endif
 
-void FNiagaraParameterStore::TickBindings()
+void FNiagaraParameterStore::TickBindings() const
 {
 	SCOPE_CYCLE_COUNTER(STAT_NiagaraParameterStoreTick);
-	for (TPair<FNiagaraParameterStore*, FNiagaraParameterStoreBinding>& Binding : Bindings)
+
+	for (const TPair<FNiagaraParameterStore*, FNiagaraParameterStoreBinding>& Binding : Bindings)
 	{
-		Binding.Value.Tick(Binding.Key, this);
+		FNiagaraParameterStore* DestStore = Binding.Key;
+		if ( bParametersDirty )
+		{
+			for (const FNiagaraParameterStoreBinding::FParameterBinding& ParamBinding : Binding.Value.ParameterBindings)
+			{
+				DestStore->SetParameterData(GetParameterData(ParamBinding.SrcOffset), ParamBinding.DestOffset, ParamBinding.Size);
+			}
+		}
+		if ( bInterfacesDirty )
+		{
+			for (const FNiagaraParameterStoreBinding::FInterfaceBinding& InterfaceBinding : Binding.Value.InterfaceBindings)
+			{
+				DestStore->SetDataInterface(GetDataInterface(InterfaceBinding.SrcOffset), InterfaceBinding.DestOffset);
+			}
+		}
+		if ( bUObjectsDirty )
+		{
+			for (const FNiagaraParameterStoreBinding::FUObjectBinding& UObjectBinding : Binding.Value.UObjectBindings)
+			{
+				DestStore->SetUObject(GetUObject(UObjectBinding.SrcOffset), UObjectBinding.DestOffset);
+			}
+		}
+
+#if NIAGARA_NAN_CHECKING
+		DestStore->CheckForNaNs();
+#endif
 	}
+
 	Dump();
 }
 
@@ -592,7 +610,7 @@ void FNiagaraParameterStore::UnbindFromSourceStores()
 	SourceStores.Empty();
 }
 
-void FNiagaraParameterStore::DumpParameters(bool bDumpBindings)const
+void FNiagaraParameterStore::DumpParameters(bool bDumpBindings) const
 {
 	for (const FNiagaraVariableWithOffset& VariableBase : ReadParameterVariables())
 	{
@@ -624,7 +642,7 @@ FString FNiagaraParameterStore::ToString() const
 	return Value;
 }
 
-void FNiagaraParameterStore::Dump()
+void FNiagaraParameterStore::Dump() const
 {
 #if WITH_EDITORONLY_DATA
 	if (GbDumpParticleParameterStores && GetParametersDirty())
diff --git a/Engine/Plugins/FX/Niagara/Source/Niagara/Public/NiagaraParameterStore.h b/Engine/Plugins/FX/Niagara/Source/Niagara/Public/NiagaraParameterStore.h
--- a/Engine/Plugins/FX/Niagara/Source/Niagara/Public/NiagaraParameterStore.h
+++ b/Engine/Plugins/FX/Niagara/Source/Niagara/Public/NiagaraParameterStore.h
@@ -111,17 +111,11 @@ struct FNiagaraParameterStoreBinding
 	/** Bindings of UObject params. Src and Dest offsets.*/
 	TArray<FUObjectBinding> UObjectBindings;
 
-	FNiagaraParameterStoreBinding()
-	{
-	}
-	
 	FORCEINLINE_DEBUGGABLE void Empty(FNiagaraParameterStore* DestStore, FNiagaraParameterStore* SrcStore);
 	FORCEINLINE_DEBUGGABLE bool Initialize(FNiagaraParameterStore* DestStore, FNiagaraParameterStore* SrcStore, const FNiagaraBoundParameterArray* BoundParameters = nullptr);
 	FORCEINLINE_DEBUGGABLE bool VerifyBinding(const FNiagaraParameterStore* DestStore, const FNiagaraParameterStore* SrcStore)const;
-	void Tick(FNiagaraParameterStore* DestStore, FNiagaraParameterStore* SrcStore, bool bForce = false);
+	void CopyParameters(FNiagaraParameterStore* DestStore, const FNiagaraParameterStore* SrcStore) const;
 	FORCEINLINE_DEBUGGABLE void Dump(const FNiagaraParameterStore* DestStore, const FNiagaraParameterStore* SrcStore)const;
-	/** TODO: Merge contiguous ranges into a single binding? */
-	//FORCEINLINE_DEBUGGABLE void Optimize();
 
 	static void GetBindingData(FNiagaraParameterStore* DestStore, FNiagaraParameterStore* SrcStore, FNiagaraBoundParameterArray& OutBoundParameters);
 
@@ -273,7 +267,7 @@ public:
 	void SetOwner(UObject* InOwner);
 	UObject* GetOwner() const { return Owner.Get(); }
 
-	void Dump();
+	void Dump() const;
 	void DumpParameters(bool bDumpBindings = false)const;
 
 	SIZE_T GetResourceSize() const
@@ -671,7 +665,7 @@ public:
 	void TriggerOnLayoutChanged() { OnLayoutChange(); }
 
 protected:
-	void TickBindings();
+	void TickBindings() const;
 	void OnLayoutChange();
 	void CopySortedParameterOffsets(TArrayView<const FNiagaraVariableWithOffset> Src);
 	void AssignParameterData(TConstArrayView<uint8> SourceParameterData);
@@ -801,12 +795,6 @@ FORCEINLINE_DEBUGGABLE bool FNiagaraParameterStoreBinding::VerifyBinding(const F
 	return bBindingValid;
 }
 
-/** Merge contiguous ranges into a single binding. */
-// FORCEINLINE_DEBUGGABLE void FNiagaraParameterStoreBinding::Optimize()
-// {
-// 	//TODO
-// }
-
 FORCEINLINE_DEBUGGABLE void FNiagaraParameterStoreBinding::Dump(const FNiagaraParameterStore* DestStore, const FNiagaraParameterStore* SrcStore)const
 {
 #if WITH_EDITORONLY_DATA
