Niagara - Fix niagara systems being marked dirty on load.
+ When creating editor only sequencer data make them transient so that debugging systems being marked dirty on load is easier to debug.
+ Fix an issue in the curve overview view models which was causing modify to be called on curves in systems which was caused by a refactor to the curve editor model classes.

diff --git a/Engine/Plugins/FX/Niagara/Source/NiagaraEditor/Private/Sequencer/NiagaraSequence/MovieSceneNiagaraEmitterTrack.cpp b/Engine/Plugins/FX/Niagara/Source/NiagaraEditor/Private/Sequencer/NiagaraSequence/MovieSceneNiagaraEmitterTrack.cpp
--- a/Engine/Plugins/FX/Niagara/Source/NiagaraEditor/Private/Sequencer/NiagaraSequence/MovieSceneNiagaraEmitterTrack.cpp
+++ b/Engine/Plugins/FX/Niagara/Source/NiagaraEditor/Private/Sequencer/NiagaraSequence/MovieSceneNiagaraEmitterTrack.cpp
@@ -297,7 +297,7 @@ void UMovieSceneNiagaraEmitterTrack::CreateSections(const FFrameRate& InFrameRes
 
 					if(EmitterSection == nullptr)
 					{
-						EmitterSection = NewObject<UMovieSceneNiagaraEmitterSectionBase>(this, SectionClass);
+						EmitterSection = NewObject<UMovieSceneNiagaraEmitterSectionBase>(this, SectionClass, NAME_None, RF_Transient);
 						EmitterSection->Initialize(*SystemViewModel, EmitterHandleViewModel.Pin().ToSharedRef());
 						bSectionCreated = true;
 					}
diff --git a/Engine/Plugins/FX/Niagara/Source/NiagaraEditor/Private/ViewModels/NiagaraSystemViewModel.cpp b/Engine/Plugins/FX/Niagara/Source/NiagaraEditor/Private/ViewModels/NiagaraSystemViewModel.cpp
--- a/Engine/Plugins/FX/Niagara/Source/NiagaraEditor/Private/ViewModels/NiagaraSystemViewModel.cpp
+++ b/Engine/Plugins/FX/Niagara/Source/NiagaraEditor/Private/ViewModels/NiagaraSystemViewModel.cpp
@@ -2137,6 +2137,7 @@ void FNiagaraSystemViewModel::RefreshSequencerTracks()
 		for (TSharedRef<FNiagaraEmitterHandleViewModel> EmitterHandleViewModel : EmitterHandleViewModels)
 		{
 			UMovieSceneNiagaraEmitterTrack* EmitterTrack = Cast<UMovieSceneNiagaraEmitterTrack>(NiagaraSequence->GetMovieScene()->AddTrack(UMovieSceneNiagaraEmitterTrack::StaticClass()));
+			EmitterTrack->SetFlags(RF_Transient);
 			EmitterTrack->Initialize(*this, EmitterHandleViewModel, NiagaraSequence->GetMovieScene()->GetTickResolution());
 			EmitterHandleIdToTrackMap.Add(EmitterHandleViewModel->GetId(), EmitterTrack);
 		}
@@ -2148,7 +2149,7 @@ void FNiagaraSystemViewModel::RefreshSequencerTracks()
 		UNiagaraSystemEditorFolder& RootFolder = SystemEditorData.GetRootFolder();
 		for (const UNiagaraSystemEditorFolder* RootChildFolder : RootFolder.GetChildFolders())
 		{
-			UMovieSceneFolder* MovieSceneRootFolder = NewObject<UMovieSceneFolder>(MovieScene, RootChildFolder->GetFolderName(), RF_Transactional);
+			UMovieSceneFolder* MovieSceneRootFolder = NewObject<UMovieSceneFolder>(MovieScene, RootChildFolder->GetFolderName(), RF_Transactional | RF_Transient);
 			MovieSceneRootFolder->SetFolderName(RootChildFolder->GetFolderName());
 			MovieScene->AddRootFolder(MovieSceneRootFolder);
 			PopulateChildMovieSceneFoldersFromNiagaraFolders(RootChildFolder, MovieSceneRootFolder, EmitterHandleIdToTrackMap);
@@ -2201,7 +2202,7 @@ void FNiagaraSystemViewModel::SetupSequencer()
 		return;
 	}
 	NiagaraSequence = NewObject<UNiagaraSequence>(GetTransientPackage());
-	UMovieScene* MovieScene = NewObject<UMovieScene>(NiagaraSequence, FName("Niagara System MovieScene"), RF_Transactional);
+	UMovieScene* MovieScene = NewObject<UMovieScene>(NiagaraSequence, FName("Niagara System MovieScene"), RF_Transactional | RF_Transient);
 	MovieScene->SetDisplayRate(FFrameRate(240, 1));
 
 	NiagaraSequence->Initialize(this, MovieScene);
diff --git a/Engine/Plugins/FX/Niagara/Source/NiagaraEditorWidgets/Private/DetailCustomizations/NiagaraDataInterfaceCurveDetails.cpp b/Engine/Plugins/FX/Niagara/Source/NiagaraEditorWidgets/Private/DetailCustomizations/NiagaraDataInterfaceCurveDetails.cpp
--- a/Engine/Plugins/FX/Niagara/Source/NiagaraEditorWidgets/Private/DetailCustomizations/NiagaraDataInterfaceCurveDetails.cpp
+++ b/Engine/Plugins/FX/Niagara/Source/NiagaraEditorWidgets/Private/DetailCustomizations/NiagaraDataInterfaceCurveDetails.cpp
@@ -121,7 +121,8 @@ struct FNiagaraCurveDetailsTreeItem : public ICurveEditorTreeItem, TSharedFromTh
 		{
 			NewCurve->SetShortDisplayName(FText::FromName(CurveData.Name));
 		}
-		NewCurve->SetColor(CurveData.Color);
+		bool bModify = false; // Don't want to modify (and mark dirty) the underlying data when creating a view model for it.
+		NewCurve->SetColor(CurveData.Color, bModify);
 		NewCurve->OnCurveModified().AddSP(this, &FNiagaraCurveDetailsTreeItem::CurveChanged);
 		OutCurveModels.Add(MoveTemp(NewCurve));
 	}
