Removed the 'RemovedInBase' entries during summary view merge process. The name was misleading: these items were not removed in base (parent), but they were originally added in base, then removed in child.
Since we don't support user deletions of items in the child, we shouldn't maintain the deletions in the child process as it is possible the child got into a bad state.
Items might still be deleted (i.e. deleting a dynamic input that was added in the parent),  but these will be cleaned up during sync and it's preferrable if the child assumes all the parent entries are supposed to be kept at first, then cleaned after.

diff --git a/Engine/Plugins/FX/Niagara/Source/NiagaraEditor/Private/NiagaraScriptMergeManager.cpp b/Engine/Plugins/FX/Niagara/Source/NiagaraEditor/Private/NiagaraScriptMergeManager.cpp
--- a/Engine/Plugins/FX/Niagara/Source/NiagaraEditor/Private/NiagaraScriptMergeManager.cpp
+++ b/Engine/Plugins/FX/Niagara/Source/NiagaraEditor/Private/NiagaraScriptMergeManager.cpp
@@ -1203,7 +1203,6 @@ bool FNiagaraEmitterDiffResults::IsEmpty() const
 		ModifiedOtherRenderers.Num() == 0 &&
 		AddedSummaryEntriesInOther.Num() == 0 &&
 		AddedSummarySectionsInOther.Num() == 0 &&
-		RemovedSummaryEntriesInBase.Num() == 0 && 
 		ModifiedStackEntryDisplayNames.Num() == 0 &&
 		bScratchPadModified == false &&
 		NewShouldShowSummaryViewValue.IsSet() == false;
@@ -2479,8 +2478,7 @@ void FNiagaraScriptMergeManager::DiffEmitterSummary(const UNiagaraEmitterEditorD
 		BaseItems,
 		OtherItems,
 		[](UNiagaraHierarchyItemBase* HierarchyItem) { return HierarchyItem->GetPersistentIdentity(); });
-
-	DiffResults.RemovedSummaryEntriesInBase.Append(SummaryItemDiff.RemovedBaseValues);
+	
 	DiffResults.AddedSummaryEntriesInOther.Append(SummaryItemDiff.AddedOtherValues);
 	
 	TArray<UNiagaraHierarchySection*> BaseSections = BaseRoot->GetSectionData();
@@ -3778,11 +3776,6 @@ FNiagaraScriptMergeManager::FApplyDiffResults FNiagaraScriptMergeManager::ApplyE
 		}
 	}
 
-	for (UNiagaraHierarchyItemBase* RemovedInputInBase : DiffResults.RemovedSummaryEntriesInBase)
-	{
-		bool bChildRemoved = BaseRoot->RemoveChildWithIdentity(RemovedInputInBase->GetPersistentIdentity(), true);
-	}
-
 	if (DiffResults.NewShouldShowSummaryViewValue.IsSet())
 	{
 		EditorData->SetShowSummaryView(DiffResults.NewShouldShowSummaryViewValue.GetValue());
diff --git a/Engine/Plugins/FX/Niagara/Source/NiagaraEditor/Private/NiagaraScriptMergeManager.h b/Engine/Plugins/FX/Niagara/Source/NiagaraEditor/Private/NiagaraScriptMergeManager.h
--- a/Engine/Plugins/FX/Niagara/Source/NiagaraEditor/Private/NiagaraScriptMergeManager.h
+++ b/Engine/Plugins/FX/Niagara/Source/NiagaraEditor/Private/NiagaraScriptMergeManager.h
@@ -409,7 +409,6 @@ struct FNiagaraEmitterDiffResults
 	TArray<TSharedRef<FNiagaraRendererMergeAdapter>> ModifiedBaseRenderers;
 	TArray<TSharedRef<FNiagaraRendererMergeAdapter>> ModifiedOtherRenderers;
 
-	TArray<UNiagaraHierarchyItemBase*> RemovedSummaryEntriesInBase;
 	TArray<UNiagaraHierarchyItemBase*> AddedSummaryEntriesInOther;	
 	TArray<UNiagaraHierarchySection*> AddedSummarySectionsInOther;
 
