- Add null check for data interface in FindDataInterfaceDependencies
- DIs can be in an invalid state from emitter merging, this will clean up post compile but we need to handle a null user DI here until they are resolved

#rnx

diff --git a/Engine/Plugins/FX/Niagara/Source/Niagara/Private/NiagaraSystem.cpp b/Engine/Plugins/FX/Niagara/Source/Niagara/Private/NiagaraSystem.cpp
--- a/Engine/Plugins/FX/Niagara/Source/Niagara/Private/NiagaraSystem.cpp
+++ b/Engine/Plugins/FX/Niagara/Source/Niagara/Private/NiagaraSystem.cpp
@@ -1603,59 +1603,66 @@ void UNiagaraSystem::FindEventDependencies(FVersionedNiagaraEmitterData* Emitter
 
 void UNiagaraSystem::FindDataInterfaceDependencies(FVersionedNiagaraEmitterData* EmitterData, UNiagaraScript* Script, TArray<FVersionedNiagaraEmitter>& Dependencies)
 {
-	if (const FNiagaraScriptExecutionParameterStore* ParameterStore = Script->GetExecutionReadyParameterStore(EmitterData->SimTarget))
+	const FNiagaraScriptExecutionParameterStore* ParameterStore = Script->GetExecutionReadyParameterStore(EmitterData->SimTarget);
+	if (ParameterStore == nullptr)
 	{
-		if (EmitterData->SimTarget == ENiagaraSimTarget::CPUSim)
+		return;
+	}
+
+	// Note we need to null check here as merging can leave parameter stores with nullptrs to user data interfaces, once the compilation is complete they will be valid.
+	if (EmitterData->SimTarget == ENiagaraSimTarget::CPUSim)
+	{
+		for (UNiagaraDataInterface* DataInterface : ParameterStore->GetDataInterfaces())
 		{
-			for (UNiagaraDataInterface* DataInterface : ParameterStore->GetDataInterfaces())
+			if (DataInterface != nullptr)
 			{
 				DataInterface->GetEmitterDependencies(this, Dependencies);
 			}
 		}
-		else
+	}
+	else
+	{
+		const TArray<UNiagaraDataInterface*>& StoreDataInterfaces = ParameterStore->GetDataInterfaces();
+		if (StoreDataInterfaces.Num() > 0)
 		{
-			const TArray<UNiagaraDataInterface*>& StoreDataInterfaces = ParameterStore->GetDataInterfaces();
-			if (StoreDataInterfaces.Num() > 0)
+			auto FindCachedDefaultDI =
+				[](UNiagaraScript* Script, const FNiagaraVariableBase& Variable) -> UNiagaraDataInterface*
 			{
-				auto FindCachedDefaultDI =
-					[](UNiagaraScript* Script, const FNiagaraVariableBase& Variable) -> UNiagaraDataInterface*
+				if (Script)
 				{
-					if (Script)
+					for (const FNiagaraScriptResolvedDataInterfaceInfo& DataInterfaceInfo : Script->GetResolvedDataInterfaces())
 					{
-						for (const FNiagaraScriptResolvedDataInterfaceInfo& DataInterfaceInfo : Script->GetResolvedDataInterfaces())
+						if (DataInterfaceInfo.ParameterStoreVariable == Variable)
 						{
-							if (DataInterfaceInfo.ParameterStoreVariable == Variable)
-							{
-								return DataInterfaceInfo.ResolvedDataInterface;
-							}
+							return DataInterfaceInfo.ResolvedDataInterface;
 						}
 					}
-					return nullptr;
-				};
+				}
+				return nullptr;
+			};
 
-				for (const FNiagaraVariableWithOffset& Variable : ParameterStore->ReadParameterVariables())
+			for (const FNiagaraVariableWithOffset& Variable : ParameterStore->ReadParameterVariables())
+			{
+				if (!Variable.IsDataInterface())
 				{
-					if (!Variable.IsDataInterface())
-					{
-						continue;
-					}
+					continue;
+				}
 
-					if (UNiagaraDataInterface* DefaultDI = FindCachedDefaultDI(SystemSpawnScript, Variable))
-					{
-						DefaultDI->GetEmitterDependencies(this, Dependencies);
-						continue;
-					}
+				if (UNiagaraDataInterface* DefaultDI = FindCachedDefaultDI(SystemSpawnScript, Variable))
+				{
+					DefaultDI->GetEmitterDependencies(this, Dependencies);
+					continue;
+				}
 
-					if (UNiagaraDataInterface* DefaultDI = FindCachedDefaultDI(SystemUpdateScript, Variable))
-					{
-						DefaultDI->GetEmitterDependencies(this, Dependencies);
-						continue;
-					}
+				if (UNiagaraDataInterface* DefaultDI = FindCachedDefaultDI(SystemUpdateScript, Variable))
+				{
+					DefaultDI->GetEmitterDependencies(this, Dependencies);
+					continue;
+				}
 
-					if (StoreDataInterfaces[Variable.Offset] != nullptr)
-					{
-						StoreDataInterfaces[Variable.Offset]->GetEmitterDependencies(this, Dependencies);
-					}
+				if (StoreDataInterfaces[Variable.Offset] != nullptr)
+				{
+					StoreDataInterfaces[Variable.Offset]->GetEmitterDependencies(this, Dependencies);
 				}
 			}
 		}
