- Support new RHI readback with RenderTargetVolume DI
- Remove vdb caching for RenderTargetVolume DI
- Remove vdb build support from Niagara

diff --git a/Engine/Plugins/FX/Niagara/Source/Niagara/Niagara.Build.cs b/Engine/Plugins/FX/Niagara/Source/Niagara/Niagara.Build.cs
--- a/Engine/Plugins/FX/Niagara/Source/Niagara/Niagara.Build.cs
+++ b/Engine/Plugins/FX/Niagara/Source/Niagara/Niagara.Build.cs
@@ -10,18 +10,6 @@ public class Niagara : ModuleRules
 		NumIncludedBytesPerUnityCPPOverride = 491520; // best unity size found from using UBT ProfileUnitySizes mode
 
 		PrivateIncludePaths.Add("../../../../Shaders/Shared");
-
-		// Enable OpenVDB under Windows
-		// @todo: we need builds for OpenVDB for platforms other than windows
-		if (Target.IsInPlatformGroup(UnrealPlatformGroup.Windows) && Target.WindowsPlatform.Architecture != UnrealArch.Arm64)
-		{
-			// Specific to OpenVDB support
-			bUseRTTI = true;
-			bEnableExceptions = true;
-
-			AddEngineThirdPartyPrivateStaticDependencies(Target, "OpenVDB");
-			PublicDefinitions.Add("UE_USE_OPENVDB=1");
-		}
 		
 		PrivateDependencyModuleNames.AddRange(
             new string[] {
@@ -78,16 +66,6 @@ public class Niagara : ModuleRules
             });
         }
 
-		if (Target.IsInPlatformGroup(UnrealPlatformGroup.Windows) && Target.WindowsPlatform.Architecture != UnrealArch.Arm64)
-		{ 
-			AddEngineThirdPartyPrivateStaticDependencies(Target,
-				"IntelTBB",
-				"Blosc",
-				"zlib",
-				"Boost"
-			);
-		}
-
 		PublicDefinitions.AddRange(
             new string[]
             {
diff --git a/Engine/Plugins/FX/Niagara/Source/Niagara/Private/NiagaraDataInterfaceRenderTargetVolume.cpp b/Engine/Plugins/FX/Niagara/Source/Niagara/Private/NiagaraDataInterfaceRenderTargetVolume.cpp
--- a/Engine/Plugins/FX/Niagara/Source/Niagara/Private/NiagaraDataInterfaceRenderTargetVolume.cpp
+++ b/Engine/Plugins/FX/Niagara/Source/Niagara/Private/NiagaraDataInterfaceRenderTargetVolume.cpp
@@ -94,8 +94,7 @@ namespace NDIRenderTargetVolumeLocal
 	enum SimCacheStorageMode
 	{
 		Dense = 0,
-		OpenVDB = 1,
-		SVT = 2
+		SVT = 1
 	};
 
 	int32 GSimCacheDataStorageMode = SimCacheStorageMode::SVT;
@@ -588,53 +587,7 @@ bool UNiagaraDataInterfaceRenderTargetVolume::GetExposedVariableValue(const FNia
 
 UObject* UNiagaraDataInterfaceRenderTargetVolume::SimCacheBeginWrite(UObject* InSimCache, FNiagaraSystemInstance* NiagaraSystemInstance, const void* OptionalPerInstanceData, FNiagaraSimCacheFeedbackContext& FeedbackContext) const
 {
-	if (NDIRenderTargetVolumeLocal::GSimCacheDataStorageMode == NDIRenderTargetVolumeLocal::SimCacheStorageMode::OpenVDB)
-	{		
-		UVolumeCache* OpenVDBSimCacheData = nullptr;
-
-		UNiagaraSimCache* SimCache = CastChecked<UNiagaraSimCache>(InSimCache);
-		OpenVDBSimCacheData = NewObject<UVolumeCache>(SimCache);
-
-		FString SystemInstanceName = NiagaraSystemInstance->GetSystem()->GetName();
-
-		if (GetDefault<UNiagaraSettings>()->SimCacheAuxiliaryFileBasePath == "")
-		{
-			FeedbackContext.Errors.Emplace(TEXT("UNiagaraDataInterfaceRenderTargetVolume - You must set SimCacheAuxiliaryFileBasePath in project settings"));
-			return nullptr;
-		}
-
-		const FGuid& CacheGuid = SimCache->GetCacheGuid();
-		const FString DIName = Proxy->SourceDIName.ToString();
-		FString FullFilePathSpec = GetDefault<UNiagaraSettings>()->SimCacheAuxiliaryFileBasePath + "/" + CacheGuid.ToString() + "/" + DIName + "_SimCache.{FrameIndex}.vdb";
-		FullFilePathSpec.ReplaceInline(TEXT("//"), TEXT("/"));		
-		
-		FullFilePathSpec.ReplaceInline(TEXT("{project_dir}"), *FPaths::ProjectDir());
-
-		FullFilePathSpec = FPaths::ConvertRelativePathToFull(FullFilePathSpec);
-
-		// Create output directory		
-		const FString FileDirectory = FString(FPathViews::GetPath(FullFilePathSpec));
-
-		IPlatformFile& PlatformFile = FPlatformFileManager::Get().GetPlatformFile();
-		if (!PlatformFile.DirectoryExists(*FileDirectory))
-		{
-			if (!PlatformFile.CreateDirectoryTree(*FileDirectory))
-			{
-				FeedbackContext.Errors.Emplace(FString::Printf(TEXT("Cannot Create Directory : %s"), *FileDirectory));
-				return nullptr;
-			}
-		}
-
-		OpenVDBSimCacheData->FilePath = FullFilePathSpec;
-		OpenVDBSimCacheData->CacheType = EVolumeCacheType::OpenVDB;
-		OpenVDBSimCacheData->Resolution = FIntVector3(1,1,1);
-		OpenVDBSimCacheData->FrameRangeStart = TNumericLimits<int32>::Lowest();
-		OpenVDBSimCacheData->FrameRangeEnd = TNumericLimits<int32>::Lowest();
-		OpenVDBSimCacheData->InitData();
-
-		return OpenVDBSimCacheData;
-	}
-	else if (NDIRenderTargetVolumeLocal::GSimCacheDataStorageMode == NDIRenderTargetVolumeLocal::SimCacheStorageMode::SVT)
+	if (NDIRenderTargetVolumeLocal::GSimCacheDataStorageMode == NDIRenderTargetVolumeLocal::SimCacheStorageMode::SVT)
 	{		
 		UNiagaraSimCache* SimCache = CastChecked<UNiagaraSimCache>(InSimCache);
 		UAnimatedSparseVolumeTexture* CurrCache = NewObject<UAnimatedSparseVolumeTexture>(SimCache);			
@@ -670,20 +623,48 @@ bool UNiagaraDataInterfaceRenderTargetVolume::SimCacheWriteFrame(UObject* Storag
 		const FTextureRenderTargetResource* RT_TargetTexture = InstanceData_GT->TargetTexture->GameThread_GetRenderTargetResource();
 
 		//-OPT: Currently we are flushing rendering commands.  Do not remove this until making access to the frame data safe across threads.
-		TArray<FFloat16Color> TextureData;
+		TArray<uint8> TextureData;
+
 		ENQUEUE_RENDER_COMMAND(NDIRenderTargetVolume_CacheFrame)
 		(
-			[RT_Proxy, RT_InstanceID=SystemInstance->GetId(), RT_TargetTexture, RT_TextureData=&TextureData](FRHICommandListImmediate& RHICmdList)
+			[RT_Proxy, RT_InstanceID=SystemInstance->GetId(), RT_TargetTexture, RT_TextureData = &TextureData, RT_VolumeResolution = InstanceData_GT->Size](FRHICommandListImmediate& RHICmdList)
 			{
 				if (const FRenderTargetVolumeRWInstanceData_RenderThread* InstanceData_RT = RT_Proxy->SystemInstancesToProxyData_RT.Find(RT_InstanceID))
-				{
-					// Readback TextureData
-					RHICmdList.Read3DSurfaceFloatData(
-						InstanceData_RT->RenderTarget->GetRHI(),
-						FIntRect(0, 0, InstanceData_RT->Size.X, InstanceData_RT->Size.Y),
-						FIntPoint(0, InstanceData_RT->Size.Z),
-						*RT_TextureData
-					);
+				{					
+					FRHIGPUTextureReadback RenderTargetReadback("ReadVolumeTexture");
+
+					RenderTargetReadback.EnqueueCopy(RHICmdList, InstanceData_RT->RenderTarget->GetRHI(), FIntVector(0,0,0), 0, RT_VolumeResolution);
+					
+					// Sync the GPU. Unfortunately we can't use the fences because not all RHIs implement them yet.
+					RHICmdList.BlockUntilGPUIdle();
+					RHICmdList.FlushResources();
+
+					//Lock the readback staging texture
+					int32 RowPitchInPixels;
+					int32 BufferHeight;
+					const uint8* LockedData = (const uint8*) RenderTargetReadback.Lock(RowPitchInPixels, &BufferHeight);
+					
+					uint32 BlockBytes = GPixelFormats[InstanceData_RT->RenderTarget->GetRHI()->GetFormat()].BlockBytes;
+					int32 Count = InstanceData_RT->Size.X * InstanceData_RT->Size.Y * InstanceData_RT->Size.Z * BlockBytes;
+					RT_TextureData->AddUninitialized(Count);
+
+					const uint8* SliceStart = LockedData;
+					for (int32 Z = 0; Z < InstanceData_RT->Size.Z; ++Z)
+					{
+						const uint8* RowStart = SliceStart;
+						for (int32 Y = 0; Y < InstanceData_RT->Size.Y; ++Y)
+						{									
+							int32 Offset = 0 + Y * InstanceData_RT->Size.X + Z * InstanceData_RT->Size.X * InstanceData_RT->Size.Y;
+							FMemory::Memcpy(RT_TextureData->GetData() + Offset * BlockBytes, RowStart, BlockBytes * InstanceData_RT->Size.X);
+
+							RowStart += RowPitchInPixels * BlockBytes;
+						}
+
+						SliceStart += BufferHeight * RowPitchInPixels * BlockBytes;
+					}
+
+					//Unlock the staging texture
+					RenderTargetReadback.Unlock();
 				}
 			}
 		);
@@ -691,21 +672,7 @@ bool UNiagaraDataInterfaceRenderTargetVolume::SimCacheWriteFrame(UObject* Storag
 
 		if (TextureData.Num() > 0)
 		{
-			if (NDIRenderTargetVolumeLocal::GSimCacheDataStorageMode == NDIRenderTargetVolumeLocal::SimCacheStorageMode::OpenVDB)
-			{
-#if UE_USE_OPENVDB
-				UVolumeCache* OpenVDBSimCacheData = CastChecked<UVolumeCache>(StorageObject);
-				
-				FString FullPath = OpenVDBSimCacheData->GetAssetPath(FrameIndex);
-
-				OpenVDBTools::WriteImageDataToOpenVDBFile(FullPath, InstanceData_GT->Size, TextureData, (bool) NDIRenderTargetVolumeLocal::GSimCacheUseOpenVDBFloatGrids);
-
-				OpenVDBSimCacheData->FrameRangeStart = FMath::Min(FrameIndex, OpenVDBSimCacheData->FrameRangeStart);
-				OpenVDBSimCacheData->FrameRangeEnd = FMath::Max(FrameIndex, OpenVDBSimCacheData->FrameRangeEnd);
-				OpenVDBSimCacheData->Resolution = InstanceData_GT->Size;
-#endif
-			}
-			else if (NDIRenderTargetVolumeLocal::GSimCacheDataStorageMode == NDIRenderTargetVolumeLocal::SimCacheStorageMode::SVT)
+			if (NDIRenderTargetVolumeLocal::GSimCacheDataStorageMode == NDIRenderTargetVolumeLocal::SimCacheStorageMode::SVT)
 			{
 #if WITH_EDITOR
 				UAnimatedSparseVolumeTexture* CurrCache = CastChecked<UAnimatedSparseVolumeTexture>(StorageObject);
@@ -715,7 +682,7 @@ bool UNiagaraDataInterfaceRenderTargetVolume::SimCacheWriteFrame(UObject* Storag
 				SVTCreateInfo.VirtualVolumeAABBMax = InstanceData_GT->Size;
 				SVTCreateInfo.FallbackValues[0] = FVector4f(0, 0, 0, 0);
 				SVTCreateInfo.FallbackValues[1] = FVector4f(0, 0, 0, 0);
-				SVTCreateInfo.AttributesFormats[0] = PF_FloatRGBA;
+				SVTCreateInfo.AttributesFormats[0] = InstanceData_GT->Format;
 				SVTCreateInfo.AttributesFormats[1] = PF_Unknown;
 
 				UE::SVT::FTextureData SparseTextureData{};
@@ -732,7 +699,7 @@ bool UNiagaraDataInterfaceRenderTargetVolume::SimCacheWriteFrame(UObject* Storag
 				FNDIRenderTargetVolumeSimCacheFrame* CacheFrame = &SimCacheData->Frames[FrameIndex];
 				
 				CacheFrame->Size = InstanceData_GT->Size;
-				CacheFrame->Format = EPixelFormat::PF_FloatRGBA;
+				CacheFrame->Format = InstanceData_GT->Format;
 
 				const FName CompressionType = SimCacheData->CompressionType;
 				const bool bUseCompression = CompressionType.IsNone() == false;
@@ -790,38 +757,7 @@ bool UNiagaraDataInterfaceRenderTargetVolume::SimCacheReadFrame(UObject* Storage
 
 	if (Cast<UVolumeCache>(StorageObject))
 	{
-#if PLATFORM_WINDOWS
-		UVolumeCache *OpenVDBSimCacheData = CastChecked<UVolumeCache>(StorageObject);
-
-		TSharedPtr<FVolumeCacheData> VolumeCacheData = OpenVDBSimCacheData->GetData();
-
-		InstanceData_GT->Size = OpenVDBSimCacheData->Resolution;
-		InstanceData_GT->Format = PF_FloatRGBA;
-		InstanceData_GT->Filter = TextureFilter::TF_Default;
-
-		PerInstanceTick(InstanceData_GT, SystemInstance, 0.0f);
-		PerInstanceTickPostSimulate(InstanceData_GT, SystemInstance, 0.0f);
-
-		const int FrameIndex = Interp >= 0.5f ? FrameB : FrameA;
-		OpenVDBSimCacheData->LoadFile(FrameIndex);
-
-		//  write to the volume texture
-		if (InstanceData_GT->TargetTexture)
-		{
-			FNiagaraDataInterfaceProxyRenderTargetVolumeProxy* RT_Proxy = GetProxyAs<FNiagaraDataInterfaceProxyRenderTargetVolumeProxy>();
-
-			FTextureRenderTargetResource* RT_TargetTexture = InstanceData_GT->TargetTexture->GameThread_GetRenderTargetResource();
-			ENQUEUE_RENDER_COMMAND(NDIRenderTargetVolumeUpdate)
-				(
-					[RT_Proxy, RT_VolumeCacheData = OpenVDBSimCacheData->GetData(), RT_InstanceID = SystemInstance->GetId(), RT_TargetTexture, RT_FrameIndex = FrameIndex](FRHICommandListImmediate& RHICmdList)
-					{
-						if (FRenderTargetVolumeRWInstanceData_RenderThread* InstanceData_RT = RT_Proxy->SystemInstancesToProxyData_RT.Find(RT_InstanceID))
-						{
-							RT_VolumeCacheData->Fill3DTexture_RenderThread(RT_FrameIndex, InstanceData_RT->RenderTarget->GetRHI(), RHICmdList);
-						}
-					});
-		}
-#endif
+		UE_LOG(LogNiagara, Error, TEXT("vdb caches are no longer supported.  Regenerate your Sim Cache, which will use SVT to store volume data."));
 	}
 	else if (Cast<UAnimatedSparseVolumeTexture>(StorageObject))
 	{			
@@ -839,7 +775,7 @@ bool UNiagaraDataInterfaceRenderTargetVolume::SimCacheReadFrame(UObject* Storage
 		FIntVector VolumeResolution = SVT->GetVolumeResolution();						
 
 		InstanceData_GT->Size = VolumeResolution;
-		InstanceData_GT->Format = PF_FloatRGBA;
+		InstanceData_GT->Format = SVT->GetFormat(0);
 		InstanceData_GT->Filter = TextureFilter::TF_Default;
 
 		PerInstanceTick(InstanceData_GT, SystemInstance, 0.0f);
@@ -851,7 +787,7 @@ bool UNiagaraDataInterfaceRenderTargetVolume::SimCacheReadFrame(UObject* Storage
 		// Execute compute shader
 		ENQUEUE_RENDER_COMMAND(NDIRenderTargetVolumeUpdate)
 			(
-				[RT_Proxy, RT_InstanceID = SystemInstance->GetId(),
+				[RT_Proxy, RT_InstanceID = SystemInstance->GetId(), RT_Format = InstanceData_GT->Format,
 				RT_VolumeResolution = VolumeResolution, RT_SVTRenderResources = SVTFrame ? SVTFrame->GetTextureRenderResources() : nullptr,
 				FeatureLevel = SystemInstance->GetFeatureLevel()](FRHICommandListImmediate& RHICmdList)
 				{
@@ -896,7 +832,7 @@ bool UNiagaraDataInterfaceRenderTargetVolume::SimCacheReadFrame(UObject* Storage
 						else
 						{
 							PassParameters->DestinationBuffer = GraphBuilder.CreateUAV(GraphBuilder.CreateTexture(
-									FRDGTextureDesc::Create3D(FIntVector(1, 1, 1), PF_A16B16G16R16, FClearValueBinding::Black, ETextureCreateFlags::ShaderResource | ETextureCreateFlags::UAV),
+									FRDGTextureDesc::Create3D(FIntVector(1, 1, 1), RT_Format, FClearValueBinding::Black, ETextureCreateFlags::ShaderResource | ETextureCreateFlags::UAV),
 									TEXT("NiagaraEmptyTextureUAV::Texture3D")
 								),
 								ERDGUnorderedAccessViewFlags::SkipBarrier
@@ -965,7 +901,7 @@ bool UNiagaraDataInterfaceRenderTargetVolume::SimCacheReadFrame(UObject* Storage
 			FTextureRenderTargetResource* RT_TargetTexture = InstanceData_GT->TargetTexture->GameThread_GetRenderTargetResource();
 			ENQUEUE_RENDER_COMMAND(NDIRenderTargetVolumeUpdate)
 			(
-				[RT_Proxy, RT_InstanceID=SystemInstance->GetId(), RT_TargetTexture, RT_CacheFrame=CacheFrame, RT_CompressionType=SimCacheData->CompressionType](FRHICommandListImmediate& RHICmdList)
+				[RT_Proxy, RT_InstanceID=SystemInstance->GetId(), RT_TargetTexture, RT_CacheFrame=CacheFrame, RT_CompressionType=SimCacheData->CompressionType, RT_Format = InstanceData_GT->Format](FRHICommandListImmediate& RHICmdList)
 				{
 					if (FRenderTargetVolumeRWInstanceData_RenderThread* InstanceData_RT = RT_Proxy->SystemInstancesToProxyData_RT.Find(RT_InstanceID))
 					{
@@ -973,7 +909,9 @@ bool UNiagaraDataInterfaceRenderTargetVolume::SimCacheReadFrame(UObject* Storage
 
 						FUpdateTexture3DData UpdateTexture = RHICmdList.BeginUpdateTexture3D(InstanceData_RT->RenderTarget->GetRHI(), 0, UpdateRegion);
 
-						const int32 SrcRowPitch = InstanceData_RT->Size.X * sizeof(FFloat16Color);
+						uint32 BlockBytes = GPixelFormats[RT_Format].BlockBytes;
+
+						const int32 SrcRowPitch = InstanceData_RT->Size.X * BlockBytes;
 						const int32 SrcDepthPitch = InstanceData_RT->Size.Y * SrcRowPitch;
 						if (UpdateTexture.RowPitch == SrcRowPitch && UpdateTexture.DepthPitch == SrcDepthPitch)
 						{
@@ -991,7 +929,7 @@ bool UNiagaraDataInterfaceRenderTargetVolume::SimCacheReadFrame(UObject* Storage
 							TArray<uint8> Decompressed;
 							if (RT_CacheFrame->CompressedSize > 0)
 							{
-								Decompressed.AddUninitialized(sizeof(FFloat16Color) * InstanceData_RT->Size.X * InstanceData_RT->Size.Y * InstanceData_RT->Size.Z);
+								Decompressed.AddUninitialized(sizeof(BlockBytes) * InstanceData_RT->Size.X * InstanceData_RT->Size.Y * InstanceData_RT->Size.Z);
 								FCompression::UncompressMemory(RT_CompressionType, Decompressed.GetData(), Decompressed.Num(), RT_CacheFrame->GetPixelData(), RT_CacheFrame->CompressedSize);
 							}
 
diff --git a/Engine/Plugins/FX/Niagara/Source/NiagaraEditor/Private/NiagaraBakerRenderer.cpp b/Engine/Plugins/FX/Niagara/Source/NiagaraEditor/Private/NiagaraBakerRenderer.cpp
--- a/Engine/Plugins/FX/Niagara/Source/NiagaraEditor/Private/NiagaraBakerRenderer.cpp
+++ b/Engine/Plugins/FX/Niagara/Source/NiagaraEditor/Private/NiagaraBakerRenderer.cpp
@@ -743,11 +743,9 @@ bool FNiagaraBakerRenderer::ExportVolume(FStringView FilePath, FIntVector ImageS
 	const FString FileExtension = FPaths::GetExtension(FilePath.GetData(), true);
 	if (FileExtension == TEXT(".vdb"))
 	{
-#if UE_USE_OPENVDB
-		return OpenVDBTools::WriteImageDataToOpenVDBFile(FilePath, ImageSize, ImageData, false);
-#else
+		UE_LOG(LogNiagara, Warning, TEXT("Exporting vdb grids from the Niagara Baker is no longer supported."));
+
 		return false;
-#endif
 	}
 	else
 	{
