cis fix

diff --git a/Engine/Plugins/FX/Niagara/Source/NiagaraEditor/Private/NiagaraBakerRenderer.cpp b/Engine/Plugins/FX/Niagara/Source/NiagaraEditor/Private/NiagaraBakerRenderer.cpp
--- a/Engine/Plugins/FX/Niagara/Source/NiagaraEditor/Private/NiagaraBakerRenderer.cpp
+++ b/Engine/Plugins/FX/Niagara/Source/NiagaraEditor/Private/NiagaraBakerRenderer.cpp
@@ -9,7 +9,7 @@
 #include "NiagaraComputeExecutionContext.h"
 #include "NiagaraGpuComputeDispatchInterface.h"
 #include "NiagaraBatchedElements.h"
-#include "NiagaraUseOpenVDB.h"
+#include "NiagaraTypes.h"
 
 #include "NiagaraDataInterfaceGrid2DCollection.h"
 #include "NiagaraDataInterfaceRenderTarget2D.h"
