cis fix

diff --git a/Engine/Plugins/FX/Niagara/Source/NiagaraEditor/Private/NiagaraBakerRenderer.cpp b/Engine/Plugins/FX/Niagara/Source/NiagaraEditor/Private/NiagaraBakerRenderer.cpp
--- a/Engine/Plugins/FX/Niagara/Source/NiagaraEditor/Private/NiagaraBakerRenderer.cpp
+++ b/Engine/Plugins/FX/Niagara/Source/NiagaraEditor/Private/NiagaraBakerRenderer.cpp
@@ -34,6 +34,8 @@
 #include "TextureResource.h"
 #include "SceneInterface.h"
 
+
+DEFINE_LOG_CATEGORY_STATIC(LogNiagaraBaker, Log, All);
 //////////////////////////////////////////////////////////////////////////
 
 const FString FNiagaraBakerOutputBindingHelper::STRING_SceneCaptureSource("SceneCaptureSource");
@@ -743,7 +745,7 @@ bool FNiagaraBakerRenderer::ExportVolume(FStringView FilePath, FIntVector ImageS
 	const FString FileExtension = FPaths::GetExtension(FilePath.GetData(), true);
 	if (FileExtension == TEXT(".vdb"))
 	{
-		UE_LOG(LogNiagara, Warning, TEXT("Exporting vdb grids from the Niagara Baker is no longer supported."));
+		UE_LOG(LogNiagaraBaker, Warning, TEXT("Exporting vdb grids from the Niagara Baker is no longer supported."));
 
 		return false;
 	}
