reintroduce logging

diff --git a/Engine/Plugins/FX/Niagara/Source/NiagaraEditor/Private/NiagaraBakerRenderer.cpp b/Engine/Plugins/FX/Niagara/Source/NiagaraEditor/Private/NiagaraBakerRenderer.cpp
--- a/Engine/Plugins/FX/Niagara/Source/NiagaraEditor/Private/NiagaraBakerRenderer.cpp
+++ b/Engine/Plugins/FX/Niagara/Source/NiagaraEditor/Private/NiagaraBakerRenderer.cpp
@@ -9,7 +9,7 @@
 #include "NiagaraComputeExecutionContext.h"
 #include "NiagaraGpuComputeDispatchInterface.h"
 #include "NiagaraBatchedElements.h"
-#include "NiagaraTypes.h"
+#include "NiagaraEditorCommon.h"
 
 #include "NiagaraDataInterfaceGrid2DCollection.h"
 #include "NiagaraDataInterfaceRenderTarget2D.h"
@@ -34,8 +34,6 @@
 #include "TextureResource.h"
 #include "SceneInterface.h"
 
-
-// DEFINE_LOG_CATEGORY_STATIC(LogNiagaraBaker, Log, All);
 //////////////////////////////////////////////////////////////////////////
 
 const FString FNiagaraBakerOutputBindingHelper::STRING_SceneCaptureSource("SceneCaptureSource");
@@ -745,7 +743,7 @@ bool FNiagaraBakerRenderer::ExportVolume(FStringView FilePath, FIntVector ImageS
 	const FString FileExtension = FPaths::GetExtension(FilePath.GetData(), true);
 	if (FileExtension == TEXT(".vdb"))
 	{
-		// UE_LOG(LogNiagaraBaker, Warning, TEXT("Exporting vdb grids from the Niagara Baker is no longer supported."));
+		UE_LOG(LogNiagaraEditor, Warning, TEXT("Exporting vdb grids from the Niagara Baker is no longer supported."));
 
 		return false;
 	}
