Some utility classes for writing groups of persistent data to Niagara Data Channels.

diff --git a/Engine/Plugins/FX/Niagara/Source/Niagara/Internal/NiagaraDataChannelCommon.h b/Engine/Plugins/FX/Niagara/Source/Niagara/Internal/NiagaraDataChannelCommon.h
--- a/Engine/Plugins/FX/Niagara/Source/Niagara/Internal/NiagaraDataChannelCommon.h
+++ b/Engine/Plugins/FX/Niagara/Source/Niagara/Internal/NiagaraDataChannelCommon.h
@@ -18,119 +18,6 @@ class UNiagaraDataInterfaceDataChannelWrite;
 class UNiagaraDataInterfaceDataChannelRead;
 struct FNiagaraDataChannelDataProxy;
 
-
-/** Buffer containing a single FNiagaraVariable's data at the game level. AoS layout. LWC Types. */
-struct FNiagaraDataChannelVariableBuffer
-{
-	TArray<uint8> Data;
-	TArray<uint8> PrevData;
-	int32 Size = 0;
-
-	void Init(const FNiagaraVariableBase& Var)
-	{
-		//Position types are a special case where we have to store an LWC Vector in game level data and convert to a simulation friendly FVector3f as it enters Niagara.
-		if (Var.GetType() == FNiagaraTypeDefinition::GetPositionDef())
-		{
-			Size = FNiagaraTypeHelper::GetVectorDef().GetSize();
-		}
-		else
-		{
-			Size = Var.GetSizeInBytes();
-		}
-	}
-
-	void Empty() { Data.Empty(); }
-
-	void BeginFrame(bool bKeepPrevious)
-	{
-		if (bKeepPrevious)
-		{
-			Swap(Data, PrevData);
-		}
-		Data.Reset();
-	}
-
-	template<typename T>
-	bool Write(int32 Index, const T& InData)
-	{
-		check(sizeof(T) == Size);
-		if (Index >= 0 && Index < Num())
-		{
-			T* Dest = ((T*)Data.GetData()) + Index;
-			*Dest = InData;
-
-			return true;
-		}
-		return false;
-	}
-
-	template<typename T>
-	bool Read(int32 Index, T& OutData, bool bPreviousFrameData)const
-	{
-		check(sizeof(T) == Size);
-
-		//Should we fallback to previous if we don't have current?
-		//bPreviousFrameData |= Data.Num() == 0;
-
-		int32 NumElems = bPreviousFrameData ? PrevNum() : Num();
-		if (Index >= 0 && Index < NumElems)
-		{
-			const uint8* DataPtr = bPreviousFrameData ? PrevData.GetData() : Data.GetData();
-			T* Src = ((T*)DataPtr) + Index;
-			OutData = *Src;
-
-			return true;
-		}
-		OutData = T();
-		return false;
-	}
-
-	void SetNum(int32 Num)
-	{
-		Data.SetNumZeroed(Size * Num);
-	}
-
-	int32 Num()const { return Data.Num() / Size; }
-	int32 PrevNum()const { return PrevData.Num() / Size; }
-
-	int32 GetElementSize()const { return Size; }
-};
-
-/** Storage for game level DataChannels generated by BP / C++ */
-struct FNiagaraDataChannelGameData : public TSharedFromThis<FNiagaraDataChannelGameData>
-{
-private:
-
-	/** Per variable storage buffers. */
-	TArray<FNiagaraDataChannelVariableBuffer> VariableData;
-
-	int32 NumElements = 0;
-	int32 PrevNumElements = 0;
-
-	TWeakObjectPtr<const UNiagaraDataChannel> DataChannel;
-
-public:
-
-	void Init(const UNiagaraDataChannel* InDataChannel);
-	void Empty();
-	void BeginFrame();
-	void SetNum(int32 NewNum);
-	int32 Num() const { return NumElements; }
-	int32 PrevNum()const { return PrevNumElements; }
-
-	FNiagaraDataChannelVariableBuffer* FindVariableBuffer(const FNiagaraVariableBase& Var);
-
-	void WriteToDataSet(FNiagaraDataBuffer* DestBuffer, int32 DestStartIdx, FVector3f SimulationLwcTile);
-
-	void AppendFromGameData(const FNiagaraDataChannelGameData& GameData);
-
-	void AppendFromDataSet(const FNiagaraDataBuffer* SrcBuffer, FVector3f SimulationLwcTile);
-
-	const UNiagaraDataChannel* GetDataChannel()const { return DataChannel.Get(); }
-
-	const TConstArrayView<FNiagaraDataChannelVariableBuffer> GetVariableBuffers()const { return VariableData; }
-};
-
 /** A request to publish data into a Niagara Data Channel.  */
 struct FNiagaraDataChannelPublishRequest
 {
diff --git a/Engine/Plugins/FX/Niagara/Source/Niagara/Private/NiagaraDataChannel.cpp b/Engine/Plugins/FX/Niagara/Source/Niagara/Private/NiagaraDataChannel.cpp
--- a/Engine/Plugins/FX/Niagara/Source/Niagara/Private/NiagaraDataChannel.cpp
+++ b/Engine/Plugins/FX/Niagara/Source/Niagara/Private/NiagaraDataChannel.cpp
@@ -84,6 +84,15 @@ void FNiagaraDataChannelGameData::SetNum(int32 NewNum)
 	}
 }
 
+void FNiagaraDataChannelGameData::Reserve(int32 NewNum)
+{
+	//NumElements = NewNum;
+	for (auto& Buffer : VariableData)
+	{
+		Buffer.Reserve(NewNum);
+	}
+}
+
 FNiagaraDataChannelVariableBuffer* FNiagaraDataChannelGameData::FindVariableBuffer(const FNiagaraVariableBase& Var)
 {
 	const FNiagaraDataChannelGameDataLayout& Layout = DataChannel->GetGameDataLayout();
diff --git a/Engine/Plugins/FX/Niagara/Source/Niagara/Private/NiagaraDataChannelAccessor.cpp b/Engine/Plugins/FX/Niagara/Source/Niagara/Private/NiagaraDataChannelAccessor.cpp
--- a/Engine/Plugins/FX/Niagara/Source/Niagara/Private/NiagaraDataChannelAccessor.cpp
+++ b/Engine/Plugins/FX/Niagara/Source/Niagara/Private/NiagaraDataChannelAccessor.cpp
@@ -5,10 +5,13 @@
 #include "UObject/UObjectGlobals.h"
 #include "UObject/Package.h"
 
+#include "NiagaraDataChannelPublic.h"
 #include "NiagaraDataChannelCommon.h"
 #include "NiagaraDataChannel.h"
 #include "NiagaraDataChannelHandler.h"
 
+#include "NiagaraWorldManager.h"
+
 //////////////////////////////////////////////////////////////////////////
 
 namespace NiagaraDataChannelAccesorLocal
@@ -264,3 +267,72 @@ void UNiagaraDataChannelWriter::WritePosition(FName VarName, int32 Index, FVecto
 {
 	WriteData(FNiagaraVariableBase(FNiagaraTypeDefinition::GetPositionDef(), VarName), Index, InData);
 }
+
+
+//////////////////////////////////////////////////////////////////////////
+
+void FNiagaraDataChannelGameDataWriterBase::BeginWrite()
+{
+	if(Data)
+	{
+		Data->BeginFrame();
+	}
+}
+
+void FNiagaraDataChannelGameDataWriterBase::Publish(FNiagaraDataChannelDataPtr& Destination)
+{
+	if (Data->Num() > 0)
+	{
+		FNiagaraDataChannelPublishRequest PublishRequest;
+		PublishRequest.bVisibleToCPUSims = true;
+		PublishRequest.bVisibleToGPUSims = true;
+		PublishRequest.GameData = Data;
+		Destination->Publish(PublishRequest);
+	}
+}
+
+void FNiagaraDataChannelGameDataWriterBase::SetNum(int32 Num)
+{
+	if (Data)
+	{
+		Data->SetNum(Num);
+	}
+}
+
+void FNiagaraDataChannelGameDataWriterBase::Reserve(int32 Num)
+{
+	if (Data)
+	{
+		Data->Reserve(Num);
+	}
+}
+
+int32 FNiagaraDataChannelGameDataWriterBase::Add(int32 Count)
+{
+	if (Data)
+	{
+		return Data->Add(Count);
+	}
+	return INDEX_NONE;
+}
+
+//////////////////////////////////////////////////////////////////////////
+
+FNiagaraDataChannelDataPtr FNiagaraDataChannelGameDataGroupedWriterBase::FindDataChannelData(UWorld* World, const UNiagaraDataChannel* NDC, const FNiagaraDataChannelSearchParameters& ItemSearchParams)
+{
+	if(World == nullptr || NDC == nullptr)
+	{
+		return nullptr;
+	}
+
+	if (FNiagaraWorldManager* WorldMan = FNiagaraWorldManager::Get(World))
+	{
+		if (UNiagaraDataChannelHandler* NDCHandler = WorldMan->FindDataChannelHandler(NDC))
+		{
+			return NDCHandler->FindData(ItemSearchParams, ENiagaraResourceAccess::WriteOnly);
+		}
+	}
+	return nullptr;
+}
+
+//////////////////////////////////////////////////////////////////////////
diff --git a/Engine/Plugins/FX/Niagara/Source/Niagara/Public/NiagaraDataChannelAccessor.h b/Engine/Plugins/FX/Niagara/Source/Niagara/Public/NiagaraDataChannelAccessor.h
--- a/Engine/Plugins/FX/Niagara/Source/Niagara/Public/NiagaraDataChannelAccessor.h
+++ b/Engine/Plugins/FX/Niagara/Source/Niagara/Public/NiagaraDataChannelAccessor.h
@@ -121,4 +121,269 @@ public:
 
 	UFUNCTION(BlueprintCallable, Category = NiagaraDataChannel, meta = (Keywords = "niagara DataChannel"))
 	void WritePosition(FName VarName, int32 Index, FVector InData);
-};
\ No newline at end of file
+};
+
+
+
+/** Base clas for template data channel writer objects. Provides common, non-template functionality. */
+struct NIAGARA_API FNiagaraDataChannelGameDataWriterBase
+{
+	FNiagaraDataChannelGameDataWriterBase(FNiagaraDataChannelGameDataPtr InData)
+	: Data(InData)
+	{
+	}
+
+	/** Call before we begin writing to this data. */
+	void BeginWrite();
+
+	/** Publish the contents of this writer's Data buffer to the given data channel data. */
+	void Publish(FNiagaraDataChannelDataPtr& Destination);
+
+	//Binds a specific variable to an index.
+	bool BindVariable(int32 VariableIndex, const FNiagaraVariableBase& Var);
+
+	// Set the size of buffers for all variables.
+	void SetNum(int32 Num);
+
+	// Reserve space in buffers for all variables.
+	void Reserve(int32 Num);
+
+	/** Adds count elements to all variable buffers. */
+	int32 Add(int32 Count = 1);
+
+	FNiagaraDataChannelGameDataPtr& GetData(){ return Data; }
+
+protected:
+	FNiagaraDataChannelGameDataPtr Data = nullptr;
+};
+
+//Templated utility writer class.
+//Provides an interface for pushing specific types easily into Niagara Data Channels.
+template<typename T>
+struct FNiagaraDataChannelGameDataWriter : public FNiagaraDataChannelGameDataWriterBase
+{
+	FNiagaraDataChannelGameDataWriter(FNiagaraDataChannelGameDataPtr InData)
+		: FNiagaraDataChannelGameDataWriterBase(InData)
+	{
+	}
+
+	/** Write data to a specific variable index. */
+	bool Write(int32 Index, const T& InData)
+	{
+		//static_assert(false, "All implementations must provide a write function");		
+	}
+};
+
+/**
+* Below is an example implementation of FNiagaraDataChannelGameDataWriter.
+* It allows pushing of a simple set of data into a Niagara Data Channel.
+struct FExampleNDCData
+{
+	FVector Position;
+	FVector PrevPosition;
+	FVector Velocity;
+	FNiagaraID ID;
+};
+
+template<typename T>
+struct NIAGARA_API FNiagaraDataChannelGameDataWriter<FExampleNDCData> : public FNiagaraDataChannelGameDataWriterBase
+{
+	FNiagaraDataChannelGameDataWriter(FNiagaraDataChannelGameDataPtr InData)
+		: FNiagaraDataChannelGameDataWriterBase(InData)
+	{
+		if (FNiagaraDataChannelGameData* DataPtr = Data.Get())
+		{
+			PositionBuffer = DataPtr->FindVariableBuffer(FNiagaraVariable(FNiagaraTypeDefinition::GetPositionDef(), TEXT("Position")));
+			PrevPositionBuffer = = DataPtr->FindVariableBuffer(FNiagaraVariable(FNiagaraTypeDefinition::GetPositionDef(), TEXT("PrevPosition")));
+			VelocityBuffer = = DataPtr->FindVariableBuffer(FNiagaraVariable(FNiagaraTypeHelper::GetVectorDef(), TEXT("Velocity")));
+			IDBuffer = = DataPtr->FindVariableBuffer(FNiagaraVariable(FNiagaraTypeDefinition::GetIDDef(), TEXT("ID")));
+
+			if(PositionBuffer == nullptr || PreviousPositionBuffer == nullptr || VelocityBuffer == nullptr || IDBuffer == nullptr)
+			{
+				Data.Reset();
+			}
+		}
+	}
+
+	bool Write(int32 Index, const FLWProjectileNDCData& InData);
+	{
+		if(Data)
+		{
+			check(PositionBuffer && PreviousPositionBuffer && VelocityBuffer && IDBuffer);
+
+			PositionBuffer->Write(InData.Position);
+			PrevPositionBuffer->Write(InData.PrevPosition);
+			VelocityBuffer->Write(InData.Velocity);
+			IDBuffer->Write(InData.ID);
+		}
+	}
+
+	const FNiagaraDataChannelVariableBuffer* PositionBuffer = nullptr;
+	const FNiagaraDataChannelVariableBuffer* PrevPositionBuffer = nullptr;
+	const FNiagaraDataChannelVariableBuffer* VelocityBuffer = nullptr;
+	const FNiagaraDataChannelVariableBuffer* IDBuffer = nullptr;
+};
+
+*/
+
+//////////////////////////////////////////////////////////////////////////
+// FNiagaraDataChannelGameDataGroupedWriter 
+// A utility class allowing for many persistent items to be written each frame to one or more Data Channels in efficient batches.
+// Items are grouped together for allocation and publishing to the NDC by the destination NDC data they'll write to based on the FNiagaraDataChannelSearchParameters.
+// Must use a type the has an implementation of FNiagaraDataChannelGameDataWriter.
+// Usage:
+// 
+// Construct:
+// Create the FNiagaraDataChannelGameDataGroupedWriter with desired type.
+// 
+// Add/Remove
+// When new Items you want to track are created call AddItem(). This will return a GroupIndex that you should store with each Item.
+// Items can be added to be written into different Niagara Data Channels if needed. Items are grouped by the specific NDC data they write to.
+// Call RemoveItem when the items you're tracking are destroyed.
+// 
+// Writing:
+// Ideally write once per frame.
+// To write, construct an FScopedWriter from your FNiagaraDataChannelGameDataGroupedWriter.
+// This new scoped writer will allocate space to write as many items as have been added in each group via AddItem.
+// Writing more than this will not cause any major harm but will cause additional allocations.
+// Results will be published to the NDC when the scoped writer is destroyed.
+// Call Append(), passing each item's GroupIndex and the data to write.
+
+struct NIAGARA_API FNiagaraDataChannelGameDataGroupedWriterBase
+{
+	FNiagaraDataChannelDataPtr FindDataChannelData(UWorld* World, const UNiagaraDataChannel* NDC, const FNiagaraDataChannelSearchParameters& ItemSearchParams);
+};
+
+template<typename T>
+struct FNiagaraDataChannelGameDataGroupedWriter : public FNiagaraDataChannelGameDataGroupedWriterBase
+{
+private:
+	struct FGroupInfo
+	{
+		//A writer helper for each group;
+		FNiagaraDataChannelGameDataWriter<T> Writer;
+				
+		int32 NumItems = 0;
+
+		FGroupInfo(const UNiagaraDataChannel* DataChannel, FNiagaraDataChannelGameDataGroupedWriter& Owner)
+		: Writer(DataChannel ? DataChannel->CreateGameData() : nullptr)
+		{
+		}
+	};
+
+	TMap<FNiagaraDataChannelDataPtr, int32> GroupMap;
+	TArray<FGroupInfo> Groups;
+
+public:
+
+	/**
+	 Scoped utility writer class for appending items to a FNiagaraDataChannelGameDataGroupedWriter.
+	*/
+	struct FScopedWriter
+	{
+	public:
+		FScopedWriter(FNiagaraDataChannelGameDataGroupedWriter<T>& InOwner)	: Owner(InOwner)
+		{
+			Owner.BeginWrite();
+		}
+
+		~FScopedWriter()
+		{
+			Owner.EndWrite();
+		}
+
+		void Append(int32 GroupIdx, const T& InData)
+		{
+			FNiagaraDataChannelGameDataWriter<T>& Writer = Owner.GetWriter(GroupIdx);
+			int32 DataIndex = Writer.Add();
+			Writer.Write(DataIndex, InData);
+		}
+
+	private:
+		FNiagaraDataChannelGameDataGroupedWriter<T>& Owner;
+	};
+
+	FScopedWriter CreateScopedWriter(){ return FScopedWriter(*this); }
+
+	void BeginWrite()
+	{
+		for (auto It = GroupMap.CreateIterator(); It; ++It)
+		{
+			FNiagaraDataChannelDataPtr& Data = It.Key();
+			int32 Index = It.Value();
+
+			if (Groups.IsValidIndex(Index) && Groups[Index].NumItems > 0)
+			{
+				FGroupInfo& Group = Groups[Index];
+				Group.Writer.BeginWrite();
+				Group.Writer.Reserve(Group.NumItems);
+			}
+			else
+			{
+				It.RemoveCurrent();
+			}
+		}
+	}
+
+	void EndWrite()
+	{
+		for (auto& GroupPair : GroupMap)
+		{
+			FNiagaraDataChannelDataPtr& Data = GroupPair.Key;
+			int32 Index = GroupPair.Value;
+			if (Groups.IsValidIndex(Index))
+			{
+				Groups[Index].Writer.Publish(Data);
+			}
+		}
+	}
+
+	/** 
+	Adds a new item to be written to the correct data channel group based on the search params.
+	Returns the group index this item now belongs to.
+	*/
+	int32 AddItem(UWorld* World, const UNiagaraDataChannel* NDC, const FNiagaraDataChannelSearchParameters& ItemSearchParams)
+	{
+		if(FNiagaraDataChannelDataPtr DataChannelData = FindDataChannelData(World, NDC, ItemSearchParams))
+		{
+			int32 LocalGroupIndex = INDEX_NONE;
+			if (int32* FoundGroup = GroupMap.Find(DataChannelData))
+			{
+				LocalGroupIndex = *FoundGroup;
+			}
+			else
+			{
+				//We don't have an group for this data. Create a new group.
+				//Try to reuse an old group first.
+				LocalGroupIndex = Groups.IndexOfByPredicate([](const auto& Info) { return Info.NumItems == 0; });
+				if (LocalGroupIndex == INDEX_NONE)
+				{
+					LocalGroupIndex = Groups.Num();
+					Groups.Emplace(NDC, *this);
+				}
+				GroupMap.Add(DataChannelData) = LocalGroupIndex;
+			}
+
+			++Groups[LocalGroupIndex].NumItems;
+			return LocalGroupIndex;
+		}
+		return INDEX_NONE;
+	}
+
+	void RemoveItem(int32 GroupIndex)
+	{
+		--Groups[GroupIndex].NumItems;
+	}
+
+	/** Returns the writer object for the given group index. */
+	FNiagaraDataChannelGameDataWriter<T>& GetWriter(int32 GroupIndex)
+	{
+		return Groups[GroupIndex].Writer;
+	}
+
+	/** Returns the number of items in a given group. */
+	int32 GetNumItems(int32 GroupIndex)const
+	{
+		return Groups[GroupIndex].NumItems;
+	}
+};
diff --git a/Engine/Plugins/FX/Niagara/Source/Niagara/Public/NiagaraDataChannelPublic.h b/Engine/Plugins/FX/Niagara/Source/Niagara/Public/NiagaraDataChannelPublic.h
--- a/Engine/Plugins/FX/Niagara/Source/Niagara/Public/NiagaraDataChannelPublic.h
+++ b/Engine/Plugins/FX/Niagara/Source/Niagara/Public/NiagaraDataChannelPublic.h
@@ -97,4 +97,131 @@ public:
 };
 
 
-#endif
\ No newline at end of file
+#endif//UE_BUILD_SHIPPING
+
+
+/** Buffer containing a single FNiagaraVariable's data at the game level. AoS layout. LWC Types. */
+struct FNiagaraDataChannelVariableBuffer
+{
+	TArray<uint8> Data;
+	TArray<uint8> PrevData;
+	int32 Size = 0;
+
+	void Init(const FNiagaraVariableBase& Var)
+	{
+		//Position types are a special case where we have to store an LWC Vector in game level data and convert to a simulation friendly FVector3f as it enters Niagara.
+		if (Var.GetType() == FNiagaraTypeDefinition::GetPositionDef())
+		{
+			Size = FNiagaraTypeHelper::GetVectorDef().GetSize();
+		}
+		else
+		{
+			Size = Var.GetSizeInBytes();
+		}
+	}
+
+	void Empty() { Data.Empty(); }
+
+	void BeginFrame(bool bKeepPrevious)
+	{
+		if (bKeepPrevious)
+		{
+			Swap(Data, PrevData);
+		}
+		Data.Reset();
+	}
+
+	template<typename T>
+	bool Write(int32 Index, const T& InData)
+	{
+		check(sizeof(T) == Size);
+		if (Index >= 0 && Index < Num())
+		{
+			T* Dest = ((T*)Data.GetData()) + Index;
+			*Dest = InData;
+
+			return true;
+		}
+		return false;
+	}
+
+	template<typename T>
+	bool Read(int32 Index, T& OutData, bool bPreviousFrameData)const
+	{
+		check(sizeof(T) == Size);
+
+		//Should we fallback to previous if we don't have current?
+		//bPreviousFrameData |= Data.Num() == 0;
+
+		int32 NumElems = bPreviousFrameData ? PrevNum() : Num();
+		if (Index >= 0 && Index < NumElems)
+		{
+			const uint8* DataPtr = bPreviousFrameData ? PrevData.GetData() : Data.GetData();
+			T* Src = ((T*)DataPtr) + Index;
+			OutData = *Src;
+
+			return true;
+		}
+		OutData = T();
+		return false;
+	}
+
+	void SetNum(int32 Num)
+	{
+		Data.SetNumZeroed(Size * Num);
+	}
+
+	void Reserve(int32 Num)
+	{
+		Data.Reserve(Size * Num);
+	}
+
+	int32 Num()const { return Data.Num() / Size; }
+	int32 PrevNum()const { return PrevData.Num() / Size; }
+
+	int32 GetElementSize()const { return Size; }
+};
+
+
+/** Storage for game level DataChannels generated by BP / C++ */
+struct NIAGARA_API FNiagaraDataChannelGameData : public TSharedFromThis<FNiagaraDataChannelGameData>
+{
+private:
+
+	/** Per variable storage buffers. */
+	TArray<FNiagaraDataChannelVariableBuffer> VariableData;
+
+	int32 NumElements = 0;
+	int32 PrevNumElements = 0;
+
+	TWeakObjectPtr<const UNiagaraDataChannel> DataChannel;
+
+public:
+
+	void Init(const UNiagaraDataChannel* InDataChannel);
+	void Empty();
+	void BeginFrame();
+	void SetNum(int32 NewNum);
+	void Reserve(int32 NewNum);
+	int32 Num() const { return NumElements; }
+	int32 PrevNum()const { return PrevNumElements; }
+
+	int32 Add(int32 Count=1)
+	{
+		int32 Ret = NumElements;
+		SetNum(Ret + Count);
+		return Ret;
+	}
+
+	FNiagaraDataChannelVariableBuffer* FindVariableBuffer(const FNiagaraVariableBase& Var);
+
+	void WriteToDataSet(FNiagaraDataBuffer* DestBuffer, int32 DestStartIdx, FVector3f SimulationLwcTile);
+
+	void AppendFromGameData(const FNiagaraDataChannelGameData& GameData);
+
+	void AppendFromDataSet(const FNiagaraDataBuffer* SrcBuffer, FVector3f SimulationLwcTile);
+
+	const UNiagaraDataChannel* GetDataChannel()const { return DataChannel.Get(); }
+
+	const TConstArrayView<FNiagaraDataChannelVariableBuffer> GetVariableBuffers()const { return VariableData; }
+};
