- Change Triangle -> TriangleIndex to fix platform compatability

#rnx

diff --git a/Engine/Plugins/FX/Niagara/Shaders/Private/NiagaraDataInterfaceDynamicMeshTemplate.ush b/Engine/Plugins/FX/Niagara/Shaders/Private/NiagaraDataInterfaceDynamicMeshTemplate.ush
--- a/Engine/Plugins/FX/Niagara/Shaders/Private/NiagaraDataInterfaceDynamicMeshTemplate.ush
+++ b/Engine/Plugins/FX/Niagara/Shaders/Private/NiagaraDataInterfaceDynamicMeshTemplate.ush
@@ -140,14 +140,14 @@ void AllocateSectionTriangles_{ParameterName}_UEImpureCall(bool bExecute, int Se
 	}
 }
 
-void SetTriangleVertices_{ParameterName}_UEImpureCall(in bool bExecute, in int Triangle, in int Index0, in int Index1, in int Index2)
+void SetTriangleVertices_{ParameterName}_UEImpureCall(in bool bExecute, in int TriangleIndex, in int Index0, in int Index1, in int Index2)
 {
-	if ( bExecute && Triangle >= 0 && Triangle < {ParameterName}_NumTriangles )
+	if ( bExecute && TriangleIndex >= 0 && TriangleIndex < {ParameterName}_NumTriangles )
 	{
 		//bSuccess = true;
-		{ParameterName}_IndexBuffer[Triangle * 3 + 0] = (uint)Index0;
-		{ParameterName}_IndexBuffer[Triangle * 3 + 1] = (uint)Index1;
-		{ParameterName}_IndexBuffer[Triangle * 3 + 2] = (uint)Index2;
+		{ParameterName}_IndexBuffer[TriangleIndex * 3 + 0] = (uint)Index0;
+		{ParameterName}_IndexBuffer[TriangleIndex * 3 + 1] = (uint)Index1;
+		{ParameterName}_IndexBuffer[TriangleIndex * 3 + 2] = (uint)Index2;
 	}
 	else
 	{
@@ -288,14 +288,14 @@ void AppendTriangle_{ParameterName}_UEImpureCall(
 ////////////////////////////////////////////////////////
 // Imutable Functions
 
-void GetTriangleVertices_{ParameterName}(in int Triangle, out bool bValid, out int Index0, out int Index1, out int Index2)
+void GetTriangleVertices_{ParameterName}(in int TriangleIndex, out bool bValid, out int Index0, out int Index1, out int Index2)
 {
-	if ( Triangle >= 0 && Triangle < {ParameterName}_NumTriangles )
+	if ( TriangleIndex >= 0 && TriangleIndex < {ParameterName}_NumTriangles )
 	{
 		bValid = true;
-		Index0 = (int){ParameterName}_IndexBuffer[Triangle * 3 + 0];
-		Index1 = (int){ParameterName}_IndexBuffer[Triangle * 3 + 1];
-		Index2 = (int){ParameterName}_IndexBuffer[Triangle * 3 + 2];
+		Index0 = (int){ParameterName}_IndexBuffer[TriangleIndex * 3 + 0];
+		Index1 = (int){ParameterName}_IndexBuffer[TriangleIndex * 3 + 1];
+		Index2 = (int){ParameterName}_IndexBuffer[TriangleIndex * 3 + 2];
 	}
 	else
 	{
