- Fix for Niagara GPU stat memory accounting
- If the user never called release we would 'leak' in the stats, this was especially obvious with Grids

#rnx

diff --git a/Engine/Plugins/FX/Niagara/Source/Niagara/Public/NiagaraRenderGraphUtils.h b/Engine/Plugins/FX/Niagara/Source/Niagara/Public/NiagaraRenderGraphUtils.h
--- a/Engine/Plugins/FX/Niagara/Source/Niagara/Public/NiagaraRenderGraphUtils.h
+++ b/Engine/Plugins/FX/Niagara/Source/Niagara/Public/NiagaraRenderGraphUtils.h
@@ -7,6 +7,8 @@
 // Helper for wrapping a persistent RW buffer
 struct NIAGARA_API FNiagaraPooledRWBuffer
 {
+	~FNiagaraPooledRWBuffer() { Release(); }
+
 	void Initialize(FRDGBuilder& GraphBuilder, const TCHAR* ResourceName, EPixelFormat PixelFormat, const FRDGBufferDesc& BufferDesc);
 	void Initialize(FRDGBuilder& GraphBuilder, const TCHAR* ResourceName, EPixelFormat PixelFormat, const uint32 BytesPerElemenet, const uint32 NumElements, EBufferUsageFlags UsageFlags = EBufferUsageFlags::None);
 	void Release();
@@ -38,6 +40,8 @@ private:
 // Helper for wrapping a persistent RW texture
 struct NIAGARA_API FNiagaraPooledRWTexture
 {
+	~FNiagaraPooledRWTexture() { Release(); }
+
 	void Initialize(FRDGBuilder& GraphBuilder, const TCHAR* ResourceName, const FRDGTextureDesc& TextureDesc);
 	void Release();
 
