- Fix RT enabled crash when using meshes that have no RT geomety

#rnx

diff --git a/Engine/Plugins/FX/Niagara/Source/Niagara/Private/NiagaraRendererMeshes.cpp b/Engine/Plugins/FX/Niagara/Source/Niagara/Private/NiagaraRendererMeshes.cpp
--- a/Engine/Plugins/FX/Niagara/Source/Niagara/Private/NiagaraRendererMeshes.cpp
+++ b/Engine/Plugins/FX/Niagara/Source/Niagara/Private/NiagaraRendererMeshes.cpp
@@ -1408,7 +1408,7 @@ void FNiagaraRendererMeshes::GetDynamicRayTracingInstances(FRayTracingMaterialGa
 
 		INiagaraRenderableMesh::FLODModelData LODModel;
 		MeshData.RenderableMesh->GetLODModelData(LODModel);
-		if (LODModel.LODIndex == INDEX_NONE)
+		if (LODModel.LODIndex == INDEX_NONE || LODModel.RayTracingGeometry == nullptr)
 		{
 			continue;
 		}
