[Niagara]
* Moved dllexport from type to method/staticvar in order to improve compiletime, bin/pdb size and performance

diff --git a/Engine/Plugins/FX/Niagara/Source/Niagara/Classes/NiagaraBakerOutput.h b/Engine/Plugins/FX/Niagara/Source/Niagara/Classes/NiagaraBakerOutput.h
--- a/Engine/Plugins/FX/Niagara/Source/Niagara/Classes/NiagaraBakerOutput.h
+++ b/Engine/Plugins/FX/Niagara/Source/Niagara/Classes/NiagaraBakerOutput.h
@@ -22,8 +22,8 @@ struct FNiagaraBakerTextureSource
 	FName SourceName;
 };
 
-UCLASS(Abstract)
-class NIAGARA_API UNiagaraBakerOutput : public UObject
+UCLASS(Abstract, MinimalAPI)
+class UNiagaraBakerOutput : public UObject
 {
 	GENERATED_BODY()
 public:
@@ -35,18 +35,18 @@ public:
 	virtual bool Equals(const UNiagaraBakerOutput& Other) const { return OutputName == Other.OutputName; }
 
 	/** Make an OutputName used when creating a new output */
-	virtual FString MakeOutputName() const;
+	NIAGARA_API virtual FString MakeOutputName() const;
 
-	static FString SanitizeOutputName(FString Name);
+	static NIAGARA_API FString SanitizeOutputName(FString Name);
 
 #if WITH_EDITOR
 	/** Find any warnings about the output. */
-	virtual void FindWarnings(TArray<FText>& OutWarnings) const;
+	NIAGARA_API virtual void FindWarnings(TArray<FText>& OutWarnings) const;
 
-	FString GetAssetPath(FString PathFormat, int32 FrameIndex) const;
-	FString GetAssetFolder(FString PathFormat, int32 FrameIndex) const;
-	FString GetExportPath(FString PathFormat, int32 FrameIndex) const;
-	FString GetExportFolder(FString PathFormat, int32 FrameIndex) const;
+	NIAGARA_API FString GetAssetPath(FString PathFormat, int32 FrameIndex) const;
+	NIAGARA_API FString GetAssetFolder(FString PathFormat, int32 FrameIndex) const;
+	NIAGARA_API FString GetExportPath(FString PathFormat, int32 FrameIndex) const;
+	NIAGARA_API FString GetExportFolder(FString PathFormat, int32 FrameIndex) const;
 
 	template<class TObjectClass>
 	TObjectClass* GetAsset(FString PathFormat, int32 FrameIndex)
@@ -57,7 +57,7 @@ public:
 		return TSoftObjectPtr<TObjectClass>(AssetPath).LoadSynchronous();
 	}
 
-	static UObject* GetOrCreateAsset(const FString& PackagePath, UClass* ObjectClass, UClass* FactoryClass);
+	static NIAGARA_API UObject* GetOrCreateAsset(const FString& PackagePath, UClass* ObjectClass, UClass* FactoryClass);
 
 	template<typename TObjectClass, typename TFactoryClass>
 	static TObjectClass* GetOrCreateAsset(const FString& PackagePath)
@@ -67,5 +67,5 @@ public:
 	}
 #endif
 
-	virtual void PostInitProperties() override;
+	NIAGARA_API virtual void PostInitProperties() override;
 };
diff --git a/Engine/Plugins/FX/Niagara/Source/Niagara/Classes/NiagaraBakerOutputSimCache.h b/Engine/Plugins/FX/Niagara/Source/Niagara/Classes/NiagaraBakerOutputSimCache.h
--- a/Engine/Plugins/FX/Niagara/Source/Niagara/Classes/NiagaraBakerOutputSimCache.h
+++ b/Engine/Plugins/FX/Niagara/Source/Niagara/Classes/NiagaraBakerOutputSimCache.h
@@ -6,8 +6,8 @@
 #include "NiagaraSimCache.h"
 #include "NiagaraBakerOutputSimCache.generated.h"
 
-UCLASS(meta = (DisplayName = "Simulation Cache Output (Experimental)"))
-class NIAGARA_API UNiagaraBakerOutputSimCache : public UNiagaraBakerOutput
+UCLASS(meta = (DisplayName = "Simulation Cache Output (Experimental)"), MinimalAPI)
+class UNiagaraBakerOutputSimCache : public UNiagaraBakerOutput
 {
 	GENERATED_BODY()
 
@@ -22,13 +22,13 @@ public:
 	UPROPERTY(EditAnywhere, Category = "Settings")
 	FNiagaraSimCacheCreateParameters CreateParameters;
 
-	virtual bool Equals(const UNiagaraBakerOutput& Other) const override;
+	NIAGARA_API virtual bool Equals(const UNiagaraBakerOutput& Other) const override;
 
 #if WITH_EDITOR
-	FString MakeOutputName() const override;
+	NIAGARA_API FString MakeOutputName() const override;
 #endif
 
 #if WITH_EDITORONLY_DATA
-	virtual void PostEditChangeProperty(struct FPropertyChangedEvent& PropertyChangedEvent) override;
+	NIAGARA_API virtual void PostEditChangeProperty(struct FPropertyChangedEvent& PropertyChangedEvent) override;
 #endif
 };
diff --git a/Engine/Plugins/FX/Niagara/Source/Niagara/Classes/NiagaraBakerOutputTexture2D.h b/Engine/Plugins/FX/Niagara/Source/Niagara/Classes/NiagaraBakerOutputTexture2D.h
--- a/Engine/Plugins/FX/Niagara/Source/Niagara/Classes/NiagaraBakerOutputTexture2D.h
+++ b/Engine/Plugins/FX/Niagara/Source/Niagara/Classes/NiagaraBakerOutputTexture2D.h
@@ -6,8 +6,8 @@
 #include "NiagaraBakerOutput.h"
 #include "NiagaraBakerOutputTexture2D.generated.h"
 
-UCLASS(meta = (DisplayName = "Texture 2D Output"))
-class NIAGARA_API UNiagaraBakerOutputTexture2D : public UNiagaraBakerOutput
+UCLASS(meta = (DisplayName = "Texture 2D Output"), MinimalAPI)
+class UNiagaraBakerOutputTexture2D : public UNiagaraBakerOutput
 {
 	GENERATED_BODY()
 
@@ -75,16 +75,16 @@ public:
 	UPROPERTY(EditAnywhere, Category="Settings", meta=(EditCondition="bExportFrames"))
 	FString FramesExportPathFormat = TEXT("{SavedDir}/NiagaraBaker/{AssetName}_{OutputName}/Frame_{FrameIndex}.png");
 
-	virtual bool Equals(const UNiagaraBakerOutput& Other) const override;
+	NIAGARA_API virtual bool Equals(const UNiagaraBakerOutput& Other) const override;
 
-	virtual void PostInitProperties() override;
+	NIAGARA_API virtual void PostInitProperties() override;
 
 #if WITH_EDITOR
-	virtual FString MakeOutputName() const override;
-	virtual void FindWarnings(TArray<FText>& OutWarnings) const override;
+	NIAGARA_API virtual FString MakeOutputName() const override;
+	NIAGARA_API virtual void FindWarnings(TArray<FText>& OutWarnings) const override;
 #endif
 
 #if WITH_EDITORONLY_DATA
-	virtual void PostEditChangeProperty(struct FPropertyChangedEvent& PropertyChangedEvent) override;
+	NIAGARA_API virtual void PostEditChangeProperty(struct FPropertyChangedEvent& PropertyChangedEvent) override;
 #endif
 };
diff --git a/Engine/Plugins/FX/Niagara/Source/Niagara/Classes/NiagaraBakerOutputVolumeTexture.h b/Engine/Plugins/FX/Niagara/Source/Niagara/Classes/NiagaraBakerOutputVolumeTexture.h
--- a/Engine/Plugins/FX/Niagara/Source/Niagara/Classes/NiagaraBakerOutputVolumeTexture.h
+++ b/Engine/Plugins/FX/Niagara/Source/Niagara/Classes/NiagaraBakerOutputVolumeTexture.h
@@ -5,8 +5,8 @@
 #include "NiagaraBakerOutput.h"
 #include "NiagaraBakerOutputVolumeTexture.generated.h"
 
-UCLASS(meta = (DisplayName = "Volume Texture Output"))
-class NIAGARA_API UNiagaraBakerOutputVolumeTexture : public UNiagaraBakerOutput
+UCLASS(meta = (DisplayName = "Volume Texture Output"), MinimalAPI)
+class UNiagaraBakerOutputVolumeTexture : public UNiagaraBakerOutput
 {
 	GENERATED_BODY()
 
@@ -46,13 +46,13 @@ public:
 	UPROPERTY(EditAnywhere, Category="Settings", meta=(EditCondition="bExportFrames"))
 	FString FramesExportPathFormat = TEXT("{SavedDir}/NiagaraBaker/{AssetName}_{OutputName}/Frame_{FrameIndex}.png");
 
-	virtual bool Equals(const UNiagaraBakerOutput& Other) const override;
+	NIAGARA_API virtual bool Equals(const UNiagaraBakerOutput& Other) const override;
 
 #if WITH_EDITOR
-	FString MakeOutputName() const override;
+	NIAGARA_API FString MakeOutputName() const override;
 #endif
 
 #if WITH_EDITORONLY_DATA
-	virtual void PostEditChangeProperty(struct FPropertyChangedEvent& PropertyChangedEvent) override;
+	NIAGARA_API virtual void PostEditChangeProperty(struct FPropertyChangedEvent& PropertyChangedEvent) override;
 #endif
 };
diff --git a/Engine/Plugins/FX/Niagara/Source/Niagara/Classes/NiagaraBakerSettings.h b/Engine/Plugins/FX/Niagara/Source/Niagara/Classes/NiagaraBakerSettings.h
--- a/Engine/Plugins/FX/Niagara/Source/Niagara/Classes/NiagaraBakerSettings.h
+++ b/Engine/Plugins/FX/Niagara/Source/Niagara/Classes/NiagaraBakerSettings.h
@@ -140,13 +140,13 @@ struct FNiagaraBakerOutputFrameIndices
 	float	Interp = 0.0f;
 };
 
-UCLASS(config=Niagara, defaultconfig)
-class NIAGARA_API UNiagaraBakerSettings : public UObject
+UCLASS(config=Niagara, defaultconfig, MinimalAPI)
+class UNiagaraBakerSettings : public UObject
 {
 	GENERATED_BODY()
 
 public:
-	UNiagaraBakerSettings(const FObjectInitializer& Init);
+	NIAGARA_API UNiagaraBakerSettings(const FObjectInitializer& Init);
 
 	/**
 	This is the start time of the simulation where we begin the capture.
@@ -195,32 +195,32 @@ public:
 	UPROPERTY(EditAnywhere, Category = "Settings")
 	uint8 bRenderComponentOnly : 1;
 
-	bool Equals(const UNiagaraBakerSettings& Other) const;
+	NIAGARA_API bool Equals(const UNiagaraBakerSettings& Other) const;
 
 	float GetSeekDelta() const { return 1.0f / float(FramesPerSecond); }
 
-	FVector GetCameraLocation() const;
-	FRotator GetCameraRotation() const;
-	FMatrix GetViewportMatrix() const;
-	FMatrix GetViewMatrix() const;
-	FMatrix GetProjectionMatrix() const;
+	NIAGARA_API FVector GetCameraLocation() const;
+	NIAGARA_API FRotator GetCameraRotation() const;
+	NIAGARA_API FMatrix GetViewportMatrix() const;
+	NIAGARA_API FMatrix GetViewMatrix() const;
+	NIAGARA_API FMatrix GetProjectionMatrix() const;
 
 	FNiagaraBakerCameraSettings& GetCurrentCamera() { return CameraSettings[CurrentCameraIndex]; }
 	const FNiagaraBakerCameraSettings& GetCurrentCamera() const { return CameraSettings[CurrentCameraIndex]; }
 
-	int GetOutputNumFrames(UNiagaraBakerOutput* BakerOutput) const;
-	FNiagaraBakerOutputFrameIndices GetOutputFrameIndices(UNiagaraBakerOutput* BakerOutput, float RelativeTime) const;
+	NIAGARA_API int GetOutputNumFrames(UNiagaraBakerOutput* BakerOutput) const;
+	NIAGARA_API FNiagaraBakerOutputFrameIndices GetOutputFrameIndices(UNiagaraBakerOutput* BakerOutput, float RelativeTime) const;
 
-	int GetOutputNumFrames(int OutputIndex) const;
-	FNiagaraBakerOutputFrameIndices GetOutputFrameIndices(int OutputIndex, float RelativeTime) const;
+	NIAGARA_API int GetOutputNumFrames(int OutputIndex) const;
+	NIAGARA_API FNiagaraBakerOutputFrameIndices GetOutputFrameIndices(int OutputIndex, float RelativeTime) const;
 
-	virtual void PostLoad() override;
+	NIAGARA_API virtual void PostLoad() override;
 #if WITH_EDITORONLY_DATA
-	static void DeclareConstructClasses(TArray<FTopLevelAssetPath>& OutConstructClasses, const UClass* SpecificSubclass);
+	static NIAGARA_API void DeclareConstructClasses(TArray<FTopLevelAssetPath>& OutConstructClasses, const UClass* SpecificSubclass);
 #endif
 
 #if WITH_EDITORONLY_DATA
-	virtual void PostEditChangeProperty(struct FPropertyChangedEvent& PropertyChangedEvent) override;
+	NIAGARA_API virtual void PostEditChangeProperty(struct FPropertyChangedEvent& PropertyChangedEvent) override;
 #endif
 
 	// Deprecated properties
diff --git a/Engine/Plugins/FX/Niagara/Source/Niagara/Classes/NiagaraConstants.h b/Engine/Plugins/FX/Niagara/Source/Niagara/Classes/NiagaraConstants.h
--- a/Engine/Plugins/FX/Niagara/Source/Niagara/Classes/NiagaraConstants.h
+++ b/Engine/Plugins/FX/Niagara/Source/Niagara/Classes/NiagaraConstants.h
@@ -169,87 +169,87 @@ struct FNiagaraKnownConstantInfo
 	const ENiagaraKnownConstantType ConstantType = ENiagaraKnownConstantType::Other;
 };
 
-struct NIAGARA_API FNiagaraConstants
+struct FNiagaraConstants
 {
-	static void Init();
-	static const TArray<FNiagaraVariable>& GetEngineConstants();
-	static const TArray<FNiagaraVariable>& GetTranslatorConstants();
-	static const TArray<FNiagaraVariable>& GetStaticSwitchConstants();
-	static FNiagaraVariable UpdateEngineConstant(const FNiagaraVariable& InVar);
-	static const FNiagaraVariable *FindEngineConstant(const FNiagaraVariable& InVar);
-	static FText GetEngineConstantDescription(const FNiagaraVariable& InVar);
-	static const TArray<FNiagaraVariable>& GetOldPositionTypeVariables();
+	static NIAGARA_API void Init();
+	static NIAGARA_API const TArray<FNiagaraVariable>& GetEngineConstants();
+	static NIAGARA_API const TArray<FNiagaraVariable>& GetTranslatorConstants();
+	static NIAGARA_API const TArray<FNiagaraVariable>& GetStaticSwitchConstants();
+	static NIAGARA_API FNiagaraVariable UpdateEngineConstant(const FNiagaraVariable& InVar);
+	static NIAGARA_API const FNiagaraVariable *FindEngineConstant(const FNiagaraVariable& InVar);
+	static NIAGARA_API FText GetEngineConstantDescription(const FNiagaraVariable& InVar);
+	static NIAGARA_API const TArray<FNiagaraVariable>& GetOldPositionTypeVariables();
 
-	static const TArray<FNiagaraVariable>& GetCommonParticleAttributes();
-	static FText GetAttributeDescription(const FNiagaraVariable& InVar);
-	static FString GetAttributeDefaultValue(const FNiagaraVariable& InVar);
-	static FNiagaraVariable GetAttributeWithDefaultValue(const FNiagaraVariable& InAttribute);
-	static FNiagaraVariable GetAttributeAsParticleDataSetKey(const FNiagaraVariable& InAttribute);
-	static FNiagaraVariable GetAttributeAsEmitterDataSetKey(const FNiagaraVariable& InAttribute);
-	static FNiagaraVariableAttributeBinding GetAttributeDefaultBinding(const FNiagaraVariable& InAttribute);
+	static NIAGARA_API const TArray<FNiagaraVariable>& GetCommonParticleAttributes();
+	static NIAGARA_API FText GetAttributeDescription(const FNiagaraVariable& InVar);
+	static NIAGARA_API FString GetAttributeDefaultValue(const FNiagaraVariable& InVar);
+	static NIAGARA_API FNiagaraVariable GetAttributeWithDefaultValue(const FNiagaraVariable& InAttribute);
+	static NIAGARA_API FNiagaraVariable GetAttributeAsParticleDataSetKey(const FNiagaraVariable& InAttribute);
+	static NIAGARA_API FNiagaraVariable GetAttributeAsEmitterDataSetKey(const FNiagaraVariable& InAttribute);
+	static NIAGARA_API FNiagaraVariableAttributeBinding GetAttributeDefaultBinding(const FNiagaraVariable& InAttribute);
 
-	static bool IsNiagaraConstant(const FNiagaraVariable& InVar);
-	static const FNiagaraVariableMetaData* GetConstantMetaData(const FNiagaraVariable& InVar);
+	static NIAGARA_API bool IsNiagaraConstant(const FNiagaraVariable& InVar);
+	static NIAGARA_API const FNiagaraVariableMetaData* GetConstantMetaData(const FNiagaraVariable& InVar);
 
-	static const FNiagaraVariable* GetKnownConstant(const FName& InName, bool bAllowPartialNameMatch);
-	static const FNiagaraKnownConstantInfo GetKnownConstantInfo(const FName& InName, bool bAllowPartialNameMatch);
-	static const FNiagaraVariable *FindStaticSwitchConstant(const FName& InName);
+	static NIAGARA_API const FNiagaraVariable* GetKnownConstant(const FName& InName, bool bAllowPartialNameMatch);
+	static NIAGARA_API const FNiagaraKnownConstantInfo GetKnownConstantInfo(const FName& InName, bool bAllowPartialNameMatch);
+	static NIAGARA_API const FNiagaraVariable *FindStaticSwitchConstant(const FName& InName);
 
-	static bool IsEngineManagedAttribute(const FNiagaraVariable& Var);
+	static NIAGARA_API bool IsEngineManagedAttribute(const FNiagaraVariable& Var);
 
-	static const FName InputPinName;
-	static const FName OutputPinName;
+	static NIAGARA_API const FName InputPinName;
+	static NIAGARA_API const FName OutputPinName;
 
 
 	/** Reserved Namespace Names and Scope Names */
-	static const FName UserNamespace;
-	static const FName EngineNamespace;
-	static const FName SystemNamespace;
-	static const FName EmitterNamespace;
-	static const FName ParticleAttributeNamespace;
-	static const FName ModuleNamespace;
-	static const FName OutputNamespace;
-	static const FName TransientNamespace;
-	static const FName StackContextNamespace;
-	static const FName DataInstanceNamespace;
-	static const FName StaticSwitchNamespace;
-	static const FName ArrayNamespace;
-	static const FName ParameterCollectionNamespace;
-	static const FString InitialPrefix;
-	static const FName LocalNamespace;
-	static const FName InitialNamespace;
-	static const FName PreviousNamespace;
-	static const FName OwnerNamespace;
-
-	static const FName EngineOwnerScopeName;
-	static const FName EngineSystemScopeName;
-	static const FName EngineEmitterScopeName;
-
-	static const FName ScriptTransientScopeName;
-	static const FName ScriptPersistentScopeName;
-	static const FName InputScopeName;
-	static const FName OutputScopeName;
-	static const FName UniqueOutputScopeName;
-	static const FName CustomScopeName;
+	static NIAGARA_API const FName UserNamespace;
+	static NIAGARA_API const FName EngineNamespace;
+	static NIAGARA_API const FName SystemNamespace;
+	static NIAGARA_API const FName EmitterNamespace;
+	static NIAGARA_API const FName ParticleAttributeNamespace;
+	static NIAGARA_API const FName ModuleNamespace;
+	static NIAGARA_API const FName OutputNamespace;
+	static NIAGARA_API const FName TransientNamespace;
+	static NIAGARA_API const FName StackContextNamespace;
+	static NIAGARA_API const FName DataInstanceNamespace;
+	static NIAGARA_API const FName StaticSwitchNamespace;
+	static NIAGARA_API const FName ArrayNamespace;
+	static NIAGARA_API const FName ParameterCollectionNamespace;
+	static NIAGARA_API const FString InitialPrefix;
+	static NIAGARA_API const FName LocalNamespace;
+	static NIAGARA_API const FName InitialNamespace;
+	static NIAGARA_API const FName PreviousNamespace;
+	static NIAGARA_API const FName OwnerNamespace;
+
+	static NIAGARA_API const FName EngineOwnerScopeName;
+	static NIAGARA_API const FName EngineSystemScopeName;
+	static NIAGARA_API const FName EngineEmitterScopeName;
+
+	static NIAGARA_API const FName ScriptTransientScopeName;
+	static NIAGARA_API const FName ScriptPersistentScopeName;
+	static NIAGARA_API const FName InputScopeName;
+	static NIAGARA_API const FName OutputScopeName;
+	static NIAGARA_API const FName UniqueOutputScopeName;
+	static NIAGARA_API const FName CustomScopeName;
 
 	/** String version of namespaces */
-	static const FString AssignmentNodePrefixString;
-	static const FString EngineNamespaceString;
-	static const FString EmitterNamespaceString;
-	static const FString ModuleNamespaceString;
-	static const FString OutputNamespaceString;
-	static const FString ParameterCollectionNamespaceString;
-	static const FString ParticleAttributeNamespaceString;
-	static const FString RapidIterationParametersNamespaceString;
-	static const FString StackContextNamespaceString;
-	static const FString SystemNamespaceString;
-	static const FString UserNamespaceString;
-	static const FString InternalNamespaceString;
-	static const FString PreviousNamespaceString;
-
-	static const int32 MaxCategoryNameLength;
-	static const int32 MaxParameterLength;
-	static const int32 MaxScriptNameLength;
+	static NIAGARA_API const FString AssignmentNodePrefixString;
+	static NIAGARA_API const FString EngineNamespaceString;
+	static NIAGARA_API const FString EmitterNamespaceString;
+	static NIAGARA_API const FString ModuleNamespaceString;
+	static NIAGARA_API const FString OutputNamespaceString;
+	static NIAGARA_API const FString ParameterCollectionNamespaceString;
+	static NIAGARA_API const FString ParticleAttributeNamespaceString;
+	static NIAGARA_API const FString RapidIterationParametersNamespaceString;
+	static NIAGARA_API const FString StackContextNamespaceString;
+	static NIAGARA_API const FString SystemNamespaceString;
+	static NIAGARA_API const FString UserNamespaceString;
+	static NIAGARA_API const FString InternalNamespaceString;
+	static NIAGARA_API const FString PreviousNamespaceString;
+
+	static NIAGARA_API const int32 MaxCategoryNameLength;
+	static NIAGARA_API const int32 MaxParameterLength;
+	static NIAGARA_API const int32 MaxScriptNameLength;
 private:
 	static TArray<FNiagaraVariable> SystemParameters;
 	static TArray<FNiagaraVariable> TranslatorParameters;
@@ -266,4 +266,4 @@ private:
 	
 	static TArray<FNiagaraVariable> EngineManagedAttributes;
 
-};
\ No newline at end of file
+};
diff --git a/Engine/Plugins/FX/Niagara/Source/Niagara/Classes/NiagaraDataInterface.h b/Engine/Plugins/FX/Niagara/Source/Niagara/Classes/NiagaraDataInterface.h
--- a/Engine/Plugins/FX/Niagara/Source/Niagara/Classes/NiagaraDataInterface.h
+++ b/Engine/Plugins/FX/Niagara/Source/Niagara/Classes/NiagaraDataInterface.h
@@ -250,7 +250,7 @@ private:
 //////////////////////////////////////////////////////////////////////////
 
 #if WITH_NIAGARA_DEBUGGER
-struct NIAGARA_API FNDIDrawDebugHudContext
+struct FNDIDrawDebugHudContext
 {
 	FNDIDrawDebugHudContext(bool bInVerbose, UWorld* InWorld, UCanvas* InCanvas, FNiagaraSystemInstance* InSystemInstance)
 		: bVerbose(bInVerbose)
@@ -277,7 +277,7 @@ protected:
 
 //////////////////////////////////////////////////////////////////////////
 
-struct NIAGARA_API FNDIGpuComputeContext
+struct FNDIGpuComputeContext
 {
 	FNDIGpuComputeContext(FRDGBuilder& InGraphBuilder, const FNiagaraGpuComputeDispatchInterface& InComputeDispatchInterface)
 		: GraphBuilder(InGraphBuilder)
@@ -286,7 +286,7 @@ struct NIAGARA_API FNDIGpuComputeContext
 	}
 
 	FRDGBuilder& GetGraphBuilder() const { return GraphBuilder; }
-	FRDGExternalAccessQueue& GetRDGExternalAccessQueue() const;
+	NIAGARA_API FRDGExternalAccessQueue& GetRDGExternalAccessQueue() const;
 
 	const FNiagaraGpuComputeDispatchInterface& GetComputeDispatchInterface() const { return ComputeDispatchInterface; }
 
@@ -295,7 +295,7 @@ protected:
 	const FNiagaraGpuComputeDispatchInterface& ComputeDispatchInterface;
 };
 
-struct NIAGARA_API FNDIGpuComputeResetContext : public FNDIGpuComputeContext
+struct FNDIGpuComputeResetContext : public FNDIGpuComputeContext
 {
 	FNDIGpuComputeResetContext(FRDGBuilder& InGraphBuilder, const FNiagaraGpuComputeDispatchInterface& InComputeDispatchInterface, FNiagaraSystemInstanceID InSystemInstanceID)
 		: FNDIGpuComputeContext(InGraphBuilder, InComputeDispatchInterface)
@@ -309,7 +309,7 @@ protected:
 	FNiagaraSystemInstanceID SystemInstanceID = FNiagaraSystemInstanceID();
 };
 
-struct NIAGARA_API FNDIGpuComputePrePostStageContext : public FNDIGpuComputeContext
+struct FNDIGpuComputePrePostStageContext : public FNDIGpuComputeContext
 {
 	FNDIGpuComputePrePostStageContext(FRDGBuilder& InGraphBuilder, const FNiagaraGpuComputeDispatchInterface& InComputeDispatchInterface, const FNiagaraGPUSystemTick& InSystemTick, const FNiagaraComputeInstanceData& InComputeInstanceData, const FNiagaraSimStageData& InSimStageData)
 		: FNDIGpuComputeContext(InGraphBuilder, InComputeDispatchInterface)
@@ -322,11 +322,11 @@ struct NIAGARA_API FNDIGpuComputePrePostStageContext : public FNDIGpuComputeCont
 	const FNiagaraComputeInstanceData& GetComputeInstanceData() const { return ComputeInstanceData; }
 	const FNiagaraSimStageData& GetSimStageData() const { return SimStageData; }
 
-	FNiagaraSystemInstanceID GetSystemInstanceID() const;
-	FVector3f GetSystemLWCTile() const;
-	bool IsOutputStage() const;
-	bool IsInputStage() const;
-	bool IsIterationStage() const;
+	NIAGARA_API FNiagaraSystemInstanceID GetSystemInstanceID() const;
+	NIAGARA_API FVector3f GetSystemLWCTile() const;
+	NIAGARA_API bool IsOutputStage() const;
+	NIAGARA_API bool IsInputStage() const;
+	NIAGARA_API bool IsIterationStage() const;
 
 	void SetDataInterfaceProxy(FNiagaraDataInterfaceProxy* InDataInterfaceProxy) { DataInterfaceProxy = InDataInterfaceProxy; }
 
@@ -340,7 +340,7 @@ protected:
 using FNDIGpuComputePreStageContext = FNDIGpuComputePrePostStageContext;
 using FNDIGpuComputePostStageContext = FNDIGpuComputePrePostStageContext;
 
-struct NIAGARA_API FNDIGpuComputePostSimulateContext : public FNDIGpuComputeContext
+struct FNDIGpuComputePostSimulateContext : public FNDIGpuComputeContext
 {
 	FNDIGpuComputePostSimulateContext(FRDGBuilder& InGraphBuilder, const FNiagaraGpuComputeDispatchInterface& InComputeDispatchInterface, FNiagaraSystemInstanceID InSystemInstanceID, bool InFinalPostSimulate)
 		: FNDIGpuComputeContext(InGraphBuilder, InComputeDispatchInterface)
@@ -396,7 +396,7 @@ struct FNiagaraDataInterfaceProxy
 
 //////////////////////////////////////////////////////////////////////////
 
-struct NIAGARA_API FNiagaraDataInterfaceSetShaderParametersContext
+struct FNiagaraDataInterfaceSetShaderParametersContext
 {
 	FNiagaraDataInterfaceSetShaderParametersContext(FRDGBuilder& InGraphBuilder, const FNiagaraGpuComputeDispatchInterface& InComputeDispatchInterface, const FNiagaraGPUSystemTick& InSystemTick, const FNiagaraComputeInstanceData& InComputeInstanceData, const FNiagaraSimStageData& InSimStageData, const FNiagaraShaderRef& InShaderRef, const FNiagaraShaderScriptParametersMetadata& InShaderParametersMetadata, uint8* InBaseParameters)
 		: GraphBuilder(InGraphBuilder)
@@ -411,7 +411,7 @@ struct NIAGARA_API FNiagaraDataInterfaceSetShaderParametersContext
 	}
 
 	FRDGBuilder& GetGraphBuilder() const { return GraphBuilder; }
-	FRDGExternalAccessQueue& GetRDGExternalAccessQueue() const;
+	NIAGARA_API FRDGExternalAccessQueue& GetRDGExternalAccessQueue() const;
 
 	template<typename T> T& GetProxy() const { check(DataInterfaceProxy); return static_cast<T&>(*DataInterfaceProxy); }
 	const FNiagaraGpuComputeDispatchInterface& GetComputeDispatchInterface() const { return ComputeDispatchInterface; }
@@ -419,12 +419,12 @@ struct NIAGARA_API FNiagaraDataInterfaceSetShaderParametersContext
 	const FNiagaraComputeInstanceData& GetComputeInstanceData() const { return ComputeInstanceData; }
 	const FNiagaraSimStageData& GetSimStageData() const { return SimStageData; }
 
-	FNiagaraSystemInstanceID GetSystemInstanceID() const;
-	FVector3f GetSystemLWCTile() const;
-	bool IsResourceBound(const void* ResourceAddress) const;
-	bool IsParameterBound(const void* ParameterAddress) const;
-	bool IsOutputStage() const;
-	bool IsIterationStage() const;
+	NIAGARA_API FNiagaraSystemInstanceID GetSystemInstanceID() const;
+	NIAGARA_API FVector3f GetSystemLWCTile() const;
+	NIAGARA_API bool IsResourceBound(const void* ResourceAddress) const;
+	NIAGARA_API bool IsParameterBound(const void* ParameterAddress) const;
+	NIAGARA_API bool IsOutputStage() const;
+	NIAGARA_API bool IsIterationStage() const;
 
 	template<typename T> bool IsStructBound(const T* StructAddress) const
 	{
@@ -473,8 +473,8 @@ struct NIAGARA_API FNiagaraDataInterfaceSetShaderParametersContext
 	}
 
 private:
-	bool IsStructBoundInternal(const void* StructAddress, uint32 StructSize) const;
-	uint16 GetParameterIncludedStructInternal(const FShaderParametersMetadata* StructMetadata) const;
+	NIAGARA_API bool IsStructBoundInternal(const void* StructAddress, uint32 StructSize) const;
+	NIAGARA_API uint16 GetParameterIncludedStructInternal(const FShaderParametersMetadata* StructMetadata) const;
 
 private:
 	FRDGBuilder& GraphBuilder;
@@ -529,18 +529,18 @@ struct FNDIStageTickHandler
 //////////////////////////////////////////////////////////////////////////
 
 /** Base class for all Niagara data interfaces. */
-UCLASS(abstract, EditInlineNew)
-class NIAGARA_API UNiagaraDataInterface : public UNiagaraDataInterfaceBase
+UCLASS(abstract, EditInlineNew, MinimalAPI)
+class UNiagaraDataInterface : public UNiagaraDataInterfaceBase
 {
 	GENERATED_UCLASS_BODY()
 
 public:
-	virtual ~UNiagaraDataInterface() override;
+	NIAGARA_API virtual ~UNiagaraDataInterface() override;
 
 	// UObject Interface
-	virtual void PostLoad() override;
+	NIAGARA_API virtual void PostLoad() override;
 #if WITH_EDITOR
-	virtual void PostEditChangeProperty(struct FPropertyChangedEvent& PropertyChangedEvent) override;
+	NIAGARA_API virtual void PostEditChangeProperty(struct FPropertyChangedEvent& PropertyChangedEvent) override;
 	// UObject Interface END
 
 	/** Does this data interface need setup and teardown for each stage when working a sim stage sim source? */
@@ -587,12 +587,12 @@ public:
 	
 #if WITH_EDITORONLY_DATA
 	/** Allows the generic class defaults version of this class to specify any dependencies/version/etc that might invalidate the compile. It should never depend on the value of specific properties.*/
-	virtual bool AppendCompileHash(FNiagaraCompileHashVisitor* InVisitor) const;
+	NIAGARA_API virtual bool AppendCompileHash(FNiagaraCompileHashVisitor* InVisitor) const;
 #endif
 
 #if WITH_EDITOR
 	/** Allows data interfaces to influence the compilation of GPU shaders and is only called on the CDO object not the instance. */
-	virtual void ModifyCompilationEnvironment(EShaderPlatform ShaderPlatform, struct FShaderCompilerEnvironment& OutEnvironment) const;
+	NIAGARA_API virtual void ModifyCompilationEnvironment(EShaderPlatform ShaderPlatform, struct FShaderCompilerEnvironment& OutEnvironment) const;
 
 	/** Allows data interfaces to prevent compilation of GPU shaders and is only called on the CDO object not the instance. */
 	virtual bool ShouldCompile(EShaderPlatform ShaderPlatform) const { return true; };
@@ -641,10 +641,10 @@ public:
 	virtual void GetVMExternalFunction(const FVMExternalFunctionBindingInfo& BindingInfo, void* InstanceData, FVMExternalFunction &OutFunc) { };
 	
 	/** Copies the contents of this DataInterface to another.*/
-	bool CopyTo(UNiagaraDataInterface* Destination) const;
+	NIAGARA_API bool CopyTo(UNiagaraDataInterface* Destination) const;
 
 	/** Determines if this DataInterface is the same as another.*/
-	virtual bool Equals(const UNiagaraDataInterface* Other) const;
+	NIAGARA_API virtual bool Equals(const UNiagaraDataInterface* Other) const;
 
 	virtual bool CanExecuteOnTarget(ENiagaraSimTarget Target)const { return false; }
 
@@ -677,13 +677,13 @@ public:
 	virtual ETickingGroup CalculateTickGroup(const void* PerInstanceData) const { return NiagaraFirstTickGroup; }
 
 	/** Used to determine if we need to create CPU resources for the emitter. */
-	bool IsUsedWithCPUEmitter() const;
+	NIAGARA_API bool IsUsedWithCPUEmitter() const;
 
 	/** Used to determine if we need to create GPU resources for the emitter. */
-	bool IsUsedWithGPUEmitter() const;
+	NIAGARA_API bool IsUsedWithGPUEmitter() const;
 
 	/** Determines if this type definition matches to a known data interface type.*/
-	static bool IsDataInterfaceType(const FNiagaraTypeDefinition& TypeDef);
+	static NIAGARA_API bool IsDataInterfaceType(const FNiagaraTypeDefinition& TypeDef);
 
 #if WITH_EDITORONLY_DATA
 	/** Allows data interfaces to provide common functionality that will be shared across interfaces on that type. */
@@ -691,9 +691,9 @@ public:
 	{
 	}
 	
-	virtual void GetParameterDefinitionHLSL(FNiagaraDataInterfaceHlslGenerationContext& HlslGenContext, FString& OutHLSL);
+	NIAGARA_API virtual void GetParameterDefinitionHLSL(FNiagaraDataInterfaceHlslGenerationContext& HlslGenContext, FString& OutHLSL);
 
-	virtual bool GetFunctionHLSL(FNiagaraDataInterfaceHlslGenerationContext& HlslGenContext, FString& OutHLSL);
+	NIAGARA_API virtual bool GetFunctionHLSL(FNiagaraDataInterfaceHlslGenerationContext& HlslGenContext, FString& OutHLSL);
 
 	virtual void GetParameterDefinitionHLSL(const FNiagaraDataInterfaceGPUParamInfo& ParamInfo, FString& OutHLSL)
 	{
@@ -714,7 +714,7 @@ public:
 	}
 
 	/** Formats and appends a template file onto the output HLSL */
-	void AppendTemplateHLSL(FString& OutHLSL, const TCHAR* TemplateShaderFile, const TMap<FString, FStringFormatArg>& TemplateArgs) const;
+	NIAGARA_API void AppendTemplateHLSL(FString& OutHLSL, const TCHAR* TemplateShaderFile, const TMap<FString, FStringFormatArg>& TemplateArgs) const;
 #endif
 
 #if WITH_NIAGARA_DEBUGGER
@@ -746,22 +746,22 @@ public:
 		Also, InAsset or InComponent may be null values, as the UI for DataInterfaces is displayed in a variety of locations. 
 		In these cases, only provide information that is relevant to that context.
 	*/
-	virtual void GetFeedback(UNiagaraSystem* InAsset, UNiagaraComponent* InComponent, TArray<FNiagaraDataInterfaceError>& OutErrors, TArray<FNiagaraDataInterfaceFeedback>& OutWarnings, TArray<FNiagaraDataInterfaceFeedback>& OutInfo);
+	NIAGARA_API virtual void GetFeedback(UNiagaraSystem* InAsset, UNiagaraComponent* InComponent, TArray<FNiagaraDataInterfaceError>& OutErrors, TArray<FNiagaraDataInterfaceFeedback>& OutWarnings, TArray<FNiagaraDataInterfaceFeedback>& OutInfo);
 
-	static void GetFeedback(UNiagaraDataInterface* DataInterface, TArray<FNiagaraDataInterfaceError>& Errors, TArray<FNiagaraDataInterfaceFeedback>& Warnings,
+	static NIAGARA_API void GetFeedback(UNiagaraDataInterface* DataInterface, TArray<FNiagaraDataInterfaceError>& Errors, TArray<FNiagaraDataInterfaceFeedback>& Warnings,
 		TArray<FNiagaraDataInterfaceFeedback>& Info);
 
 	/** Validates a function being compiled and allows interface classes to post custom compile errors when their API changes. */
-	virtual void ValidateFunction(const FNiagaraFunctionSignature& Function, TArray<FText>& OutValidationErrors);
+	NIAGARA_API virtual void ValidateFunction(const FNiagaraFunctionSignature& Function, TArray<FText>& OutValidationErrors);
 
-	void RefreshErrors();
+	NIAGARA_API void RefreshErrors();
 
-	FSimpleMulticastDelegate& OnErrorsRefreshed();
+	NIAGARA_API FSimpleMulticastDelegate& OnErrorsRefreshed();
 
 #endif
 
     /** Method to add asset tags that are specific to this data interface. By default we add in how many instances of this class exist in the list.*/
-	virtual void GetAssetTagsForContext(const UObject* InAsset, FGuid AssetVersion, const TArray<const UNiagaraDataInterface*>& InProperties, TMap<FName, uint32>& NumericKeys, TMap<FName, FString>& StringKeys) const;
+	NIAGARA_API virtual void GetAssetTagsForContext(const UObject* InAsset, FGuid AssetVersion, const TArray<const UNiagaraDataInterface*>& InProperties, TMap<FName, uint32>& NumericKeys, TMap<FName, FString>& StringKeys) const;
 	virtual bool CanExposeVariables() const { return false; }
 	virtual void GetExposedVariables(TArray<FNiagaraVariableBase>& OutVariables) const {}
 	virtual bool GetExposedVariableValue(const FNiagaraVariableBase& InVariable, void* InPerInstanceData, FNiagaraSystemInstance* InSystemInstance, void* OutData) const { return false; }
@@ -850,7 +850,7 @@ protected:
 		return TypedProxy;
 	}
 
-	virtual bool CopyToInternal(UNiagaraDataInterface* Destination) const;
+	NIAGARA_API virtual bool CopyToInternal(UNiagaraDataInterface* Destination) const;
 
 	TUniquePtr<FNiagaraDataInterfaceProxy> Proxy;
 
diff --git a/Engine/Plugins/FX/Niagara/Source/Niagara/Classes/NiagaraDataInterface2DArrayTexture.h b/Engine/Plugins/FX/Niagara/Source/Niagara/Classes/NiagaraDataInterface2DArrayTexture.h
--- a/Engine/Plugins/FX/Niagara/Source/Niagara/Classes/NiagaraDataInterface2DArrayTexture.h
+++ b/Engine/Plugins/FX/Niagara/Source/Niagara/Classes/NiagaraDataInterface2DArrayTexture.h
@@ -10,8 +10,8 @@
 class UTexture2DArray;
 
 /** Data Interface allowing sampling of a texture */
-UCLASS(EditInlineNew, Category = "Texture", CollapseCategories, meta = (DisplayName = "2D Array Texture Sample"))
-class NIAGARA_API UNiagaraDataInterface2DArrayTexture : public UNiagaraDataInterface
+UCLASS(EditInlineNew, Category = "Texture", CollapseCategories, meta = (DisplayName = "2D Array Texture Sample"), MinimalAPI)
+class UNiagaraDataInterface2DArrayTexture : public UNiagaraDataInterface
 {
 	GENERATED_UCLASS_BODY()
 
@@ -29,44 +29,44 @@ public:
 	FNiagaraUserParameterBinding TextureUserParameter;
 
 	//UObject Interface
-	virtual void PostInitProperties()override;
+	NIAGARA_API virtual void PostInitProperties()override;
 	//UObject Interface End
 
 	//UNiagaraDataInterface Interface
-	virtual void GetFunctions(TArray<FNiagaraFunctionSignature>& OutFunctions)override;
-	virtual void GetVMExternalFunction(const FVMExternalFunctionBindingInfo& BindingInfo, void* InstanceData, FVMExternalFunction &OutFunc) override;
+	NIAGARA_API virtual void GetFunctions(TArray<FNiagaraFunctionSignature>& OutFunctions)override;
+	NIAGARA_API virtual void GetVMExternalFunction(const FVMExternalFunctionBindingInfo& BindingInfo, void* InstanceData, FVMExternalFunction &OutFunc) override;
 	virtual bool CanExecuteOnTarget(ENiagaraSimTarget Target)const override { return Target == ENiagaraSimTarget::GPUComputeSim; }
 
-	virtual int32 PerInstanceDataSize() const override;
-	virtual bool InitPerInstanceData(void* PerInstanceData, FNiagaraSystemInstance* SystemInstance) override;
-	virtual void DestroyPerInstanceData(void* PerInstanceData, FNiagaraSystemInstance* SystemInstance) override;
+	NIAGARA_API virtual int32 PerInstanceDataSize() const override;
+	NIAGARA_API virtual bool InitPerInstanceData(void* PerInstanceData, FNiagaraSystemInstance* SystemInstance) override;
+	NIAGARA_API virtual void DestroyPerInstanceData(void* PerInstanceData, FNiagaraSystemInstance* SystemInstance) override;
 
 	virtual bool HasPreSimulateTick() const override { return true; }
-	virtual bool PerInstanceTick(void* PerInstanceData, FNiagaraSystemInstance* SystemInstance, float DeltaSeconds) override;
+	NIAGARA_API virtual bool PerInstanceTick(void* PerInstanceData, FNiagaraSystemInstance* SystemInstance, float DeltaSeconds) override;
 	//UNiagaraDataInterface Interface End
 
-	void GetTextureDimensions(FVectorVMExternalFunctionContext& Context);
+	NIAGARA_API void GetTextureDimensions(FVectorVMExternalFunctionContext& Context);
 
-	virtual bool Equals(const UNiagaraDataInterface* Other) const override;
+	NIAGARA_API virtual bool Equals(const UNiagaraDataInterface* Other) const override;
 
 	// GPU sim functionality
 #if WITH_EDITORONLY_DATA
-	virtual bool AppendCompileHash(FNiagaraCompileHashVisitor* InVisitor) const override;
-	virtual void GetParameterDefinitionHLSL(const FNiagaraDataInterfaceGPUParamInfo& ParamInfo, FString& OutHLSL) override;
-	virtual bool GetFunctionHLSL(const FNiagaraDataInterfaceGPUParamInfo& ParamInfo, const FNiagaraDataInterfaceGeneratedFunction& FunctionInfo, int FunctionInstanceIndex, FString& OutHLSL) override;
+	NIAGARA_API virtual bool AppendCompileHash(FNiagaraCompileHashVisitor* InVisitor) const override;
+	NIAGARA_API virtual void GetParameterDefinitionHLSL(const FNiagaraDataInterfaceGPUParamInfo& ParamInfo, FString& OutHLSL) override;
+	NIAGARA_API virtual bool GetFunctionHLSL(const FNiagaraDataInterfaceGPUParamInfo& ParamInfo, const FNiagaraDataInterfaceGeneratedFunction& FunctionInfo, int FunctionInstanceIndex, FString& OutHLSL) override;
 #endif
-	virtual void BuildShaderParameters(FNiagaraShaderParametersBuilder& ShaderParametersBuilder) const override;
-	virtual void SetShaderParameters(const FNiagaraDataInterfaceSetShaderParametersContext& Context) const override;
+	NIAGARA_API virtual void BuildShaderParameters(FNiagaraShaderParametersBuilder& ShaderParametersBuilder) const override;
+	NIAGARA_API virtual void SetShaderParameters(const FNiagaraDataInterfaceSetShaderParametersContext& Context) const override;
 
-	void SetTexture(UTexture* InTexture);
+	NIAGARA_API void SetTexture(UTexture* InTexture);
 
 protected:
-	virtual bool CopyToInternal(UNiagaraDataInterface* Destination) const override;
+	NIAGARA_API virtual bool CopyToInternal(UNiagaraDataInterface* Destination) const override;
 
 protected:
-	static const TCHAR* TemplateShaderFilePath;
-	static const FName LoadTextureName;
-	static const FName GatherRedTextureName;
-	static const FName SampleTextureName;
-	static const FName TextureDimsName;
+	static NIAGARA_API const TCHAR* TemplateShaderFilePath;
+	static NIAGARA_API const FName LoadTextureName;
+	static NIAGARA_API const FName GatherRedTextureName;
+	static NIAGARA_API const FName SampleTextureName;
+	static NIAGARA_API const FName TextureDimsName;
 };
diff --git a/Engine/Plugins/FX/Niagara/Source/Niagara/Classes/NiagaraDataInterfaceArray.h b/Engine/Plugins/FX/Niagara/Source/Niagara/Classes/NiagaraDataInterfaceArray.h
--- a/Engine/Plugins/FX/Niagara/Source/Niagara/Classes/NiagaraDataInterfaceArray.h
+++ b/Engine/Plugins/FX/Niagara/Source/Niagara/Classes/NiagaraDataInterfaceArray.h
@@ -69,17 +69,17 @@ struct INDIArrayProxyBase : public FNiagaraDataInterfaceProxyRW
 	virtual void SetShaderParameters(FShaderParameters* ShaderParameters, FNiagaraSystemInstanceID SystemInstanceID) const = 0;
 };
 
-UCLASS(abstract, EditInlineNew)
-class NIAGARA_API UNiagaraDataInterfaceArray : public UNiagaraDataInterfaceRWBase
+UCLASS(abstract, EditInlineNew, MinimalAPI)
+class UNiagaraDataInterfaceArray : public UNiagaraDataInterfaceRWBase
 {
 	GENERATED_UCLASS_BODY()
 
 public:
 	//UObject Interface
-	virtual void PostInitProperties() override;
-	virtual void PostLoad() override;
+	NIAGARA_API virtual void PostInitProperties() override;
+	NIAGARA_API virtual void PostLoad() override;
 #if WITH_EDITOR
-	virtual void PostEditChangeProperty(struct FPropertyChangedEvent& PropertyChangedEvent) override;
+	NIAGARA_API virtual void PostEditChangeProperty(struct FPropertyChangedEvent& PropertyChangedEvent) override;
 #endif
 	//UObject Interface End
 
@@ -89,7 +89,7 @@ public:
 #if WITH_EDITORONLY_DATA
 	virtual void GetParameterDefinitionHLSL(const FNiagaraDataInterfaceGPUParamInfo& ParamInfo, FString& OutHLSL) override { GetProxyAs<INDIArrayProxyBase>()->GetParameterDefinitionHLSL(ParamInfo, OutHLSL); }
 	virtual bool GetFunctionHLSL(const FNiagaraDataInterfaceGPUParamInfo& ParamInfo, const FNiagaraDataInterfaceGeneratedFunction& FunctionInfo, int FunctionInstanceIndex, FString& OutHLSL) override { return GetProxyAs<INDIArrayProxyBase>()->GetFunctionHLSL(ParamInfo, FunctionInfo, FunctionInstanceIndex, OutHLSL); }
-	virtual bool AppendCompileHash(FNiagaraCompileHashVisitor* InVisitor) const override;
+	NIAGARA_API virtual bool AppendCompileHash(FNiagaraCompileHashVisitor* InVisitor) const override;
 	virtual bool UpgradeFunctionCall(FNiagaraFunctionSignature& FunctionSignature) override { return GetProxyAs<INDIArrayProxyBase>()->UpgradeFunctionCall(FunctionSignature); }
 #endif
 	virtual bool CanExecuteOnTarget(ENiagaraSimTarget Target) const override { return true; }
@@ -98,8 +98,8 @@ public:
 	virtual void DrawDebugHud(FNDIDrawDebugHudContext& DebugHudContext) const { GetProxyAs<INDIArrayProxyBase>()->DrawDebugHud(DebugHudContext); }
 #endif
 
-	virtual bool CopyToInternal(UNiagaraDataInterface* Destination) const;
-	virtual bool Equals(const UNiagaraDataInterface* Other) const;
+	NIAGARA_API virtual bool CopyToInternal(UNiagaraDataInterface* Destination) const;
+	NIAGARA_API virtual bool Equals(const UNiagaraDataInterface* Other) const;
 
 	virtual int32 PerInstanceDataSize() const override { return GetProxyAs<INDIArrayProxyBase>()->PerInstanceDataSize(); }
 	virtual bool InitPerInstanceData(void* InPerInstanceData, FNiagaraSystemInstance* SystemInstance) override { return GetProxyAs<INDIArrayProxyBase>()->InitPerInstanceData(this, InPerInstanceData, SystemInstance); }
@@ -107,8 +107,8 @@ public:
 
 	virtual void ProvidePerInstanceDataForRenderThread(void* DataForRenderThread, void* PerInstanceData, const FNiagaraSystemInstanceID& SystemInstance) override { GetProxyAs<INDIArrayProxyBase>()->ProvidePerInstanceDataForRenderThread(DataForRenderThread, PerInstanceData, SystemInstance); }
 
-	virtual void BuildShaderParameters(FNiagaraShaderParametersBuilder& ShaderParametersBuilder) const override;
-	virtual void SetShaderParameters(const FNiagaraDataInterfaceSetShaderParametersContext& Context) const override;
+	NIAGARA_API virtual void BuildShaderParameters(FNiagaraShaderParametersBuilder& ShaderParametersBuilder) const override;
+	NIAGARA_API virtual void SetShaderParameters(const FNiagaraDataInterfaceSetShaderParametersContext& Context) const override;
 	//UNiagaraDataInterface Interface
 
 	/** ReadWrite lock to ensure safe access to the underlying array. */
diff --git a/Engine/Plugins/FX/Niagara/Source/Niagara/Classes/NiagaraDataInterfaceArrayFloat.h b/Engine/Plugins/FX/Niagara/Source/Niagara/Classes/NiagaraDataInterfaceArrayFloat.h
--- a/Engine/Plugins/FX/Niagara/Source/Niagara/Classes/NiagaraDataInterfaceArrayFloat.h
+++ b/Engine/Plugins/FX/Niagara/Source/Niagara/Classes/NiagaraDataInterfaceArrayFloat.h
@@ -5,8 +5,8 @@
 #include "NiagaraDataInterfaceArray.h"
 #include "NiagaraDataInterfaceArrayFloat.generated.h"
 
-UCLASS(EditInlineNew, Category = "Array", CollapseCategories, meta = (DisplayName = "Float Array"), Blueprintable, BlueprintType)
-class NIAGARA_API UNiagaraDataInterfaceArrayFloat : public UNiagaraDataInterfaceArray
+UCLASS(EditInlineNew, Category = "Array", CollapseCategories, meta = (DisplayName = "Float Array"), Blueprintable, BlueprintType, MinimalAPI)
+class UNiagaraDataInterfaceArrayFloat : public UNiagaraDataInterfaceArray
 {
 public:
 	GENERATED_BODY()
@@ -17,8 +17,8 @@ public:
 	NDIARRAY_GENERATE_BODY(UNiagaraDataInterfaceArrayFloat, float, FloatData)
 };
 
-UCLASS(EditInlineNew, Category = "Array", CollapseCategories, meta = (DisplayName = "Vector 2D Array"), Blueprintable, BlueprintType)
-class NIAGARA_API UNiagaraDataInterfaceArrayFloat2 : public UNiagaraDataInterfaceArray
+UCLASS(EditInlineNew, Category = "Array", CollapseCategories, meta = (DisplayName = "Vector 2D Array"), Blueprintable, BlueprintType, MinimalAPI)
+class UNiagaraDataInterfaceArrayFloat2 : public UNiagaraDataInterfaceArray
 {
 public:
 	GENERATED_BODY()
@@ -34,8 +34,8 @@ public:
 	NDIARRAY_GENERATE_BODY_LWC(UNiagaraDataInterfaceArrayFloat2, FVector2f, FloatData)
 };
 
-UCLASS(EditInlineNew, Category = "Array", CollapseCategories, meta = (DisplayName = "Vector Array"), Blueprintable, BlueprintType)
-class NIAGARA_API UNiagaraDataInterfaceArrayFloat3 : public UNiagaraDataInterfaceArray
+UCLASS(EditInlineNew, Category = "Array", CollapseCategories, meta = (DisplayName = "Vector Array"), Blueprintable, BlueprintType, MinimalAPI)
+class UNiagaraDataInterfaceArrayFloat3 : public UNiagaraDataInterfaceArray
 {
 public:
 	GENERATED_BODY()
@@ -51,8 +51,8 @@ public:
 	NDIARRAY_GENERATE_BODY_LWC(UNiagaraDataInterfaceArrayFloat3, FVector3f, FloatData)
 };
 
-UCLASS(EditInlineNew, Category = "Array", CollapseCategories, meta = (DisplayName = "Position Array"), Blueprintable, BlueprintType)
-class NIAGARA_API UNiagaraDataInterfaceArrayPosition : public UNiagaraDataInterfaceArray
+UCLASS(EditInlineNew, Category = "Array", CollapseCategories, meta = (DisplayName = "Position Array"), Blueprintable, BlueprintType, MinimalAPI)
+class UNiagaraDataInterfaceArrayPosition : public UNiagaraDataInterfaceArray
 {
 public:
 	GENERATED_BODY()
@@ -63,8 +63,8 @@ public:
 	NDIARRAY_GENERATE_BODY(UNiagaraDataInterfaceArrayPosition, FNiagaraPosition, PositionData)
 };
 
-UCLASS(EditInlineNew, Category = "Array", CollapseCategories, meta = (DisplayName = "Vector 4 Array"), Blueprintable, BlueprintType)
-class NIAGARA_API UNiagaraDataInterfaceArrayFloat4 : public UNiagaraDataInterfaceArray
+UCLASS(EditInlineNew, Category = "Array", CollapseCategories, meta = (DisplayName = "Vector 4 Array"), Blueprintable, BlueprintType, MinimalAPI)
+class UNiagaraDataInterfaceArrayFloat4 : public UNiagaraDataInterfaceArray
 {
 public:
 	GENERATED_BODY()
@@ -80,8 +80,8 @@ public:
 	NDIARRAY_GENERATE_BODY_LWC(UNiagaraDataInterfaceArrayFloat4, FVector4f, FloatData)
 };
 
-UCLASS(EditInlineNew, Category = "Array", CollapseCategories, meta = (DisplayName = "Color Array"), Blueprintable, BlueprintType)
-class NIAGARA_API UNiagaraDataInterfaceArrayColor : public UNiagaraDataInterfaceArray
+UCLASS(EditInlineNew, Category = "Array", CollapseCategories, meta = (DisplayName = "Color Array"), Blueprintable, BlueprintType, MinimalAPI)
+class UNiagaraDataInterfaceArrayColor : public UNiagaraDataInterfaceArray
 {
 public:
 	GENERATED_BODY()
@@ -92,8 +92,8 @@ public:
 	NDIARRAY_GENERATE_BODY(UNiagaraDataInterfaceArrayColor, FLinearColor, ColorData)
 };
 
-UCLASS(EditInlineNew, Category = "Array", CollapseCategories, meta = (DisplayName = "Quaternion Array"), Blueprintable, BlueprintType)
-class NIAGARA_API UNiagaraDataInterfaceArrayQuat : public UNiagaraDataInterfaceArray
+UCLASS(EditInlineNew, Category = "Array", CollapseCategories, meta = (DisplayName = "Quaternion Array"), Blueprintable, BlueprintType, MinimalAPI)
+class UNiagaraDataInterfaceArrayQuat : public UNiagaraDataInterfaceArray
 {
 public:
 	GENERATED_BODY()
@@ -109,8 +109,8 @@ public:
 	NDIARRAY_GENERATE_BODY_LWC(UNiagaraDataInterfaceArrayQuat, FQuat4f, QuatData)
 };
 
-UCLASS(EditInlineNew, Category = "Array", CollapseCategories, meta = (DisplayName = "Matrix Array"), Blueprintable, BlueprintType)
-class NIAGARA_API UNiagaraDataInterfaceArrayMatrix : public UNiagaraDataInterfaceArray
+UCLASS(EditInlineNew, Category = "Array", CollapseCategories, meta = (DisplayName = "Matrix Array"), Blueprintable, BlueprintType, MinimalAPI)
+class UNiagaraDataInterfaceArrayMatrix : public UNiagaraDataInterfaceArray
 {
 public:
 	GENERATED_BODY()
diff --git a/Engine/Plugins/FX/Niagara/Source/Niagara/Classes/NiagaraDataInterfaceArrayFunctionLibrary.h b/Engine/Plugins/FX/Niagara/Source/Niagara/Classes/NiagaraDataInterfaceArrayFunctionLibrary.h
--- a/Engine/Plugins/FX/Niagara/Source/Niagara/Classes/NiagaraDataInterfaceArrayFunctionLibrary.h
+++ b/Engine/Plugins/FX/Niagara/Source/Niagara/Classes/NiagaraDataInterfaceArrayFunctionLibrary.h
@@ -10,165 +10,165 @@ class UNiagaraComponent;
 /**
 * C++ and Blueprint library for accessing array types
 */
-UCLASS()
-class NIAGARA_API UNiagaraDataInterfaceArrayFunctionLibrary : public UBlueprintFunctionLibrary
+UCLASS(MinimalAPI)
+class UNiagaraDataInterfaceArrayFunctionLibrary : public UBlueprintFunctionLibrary
 {
 	GENERATED_BODY()
 
 public:
 	/** Sets Niagara Array Float Data. */
 	UFUNCTION(BlueprintCallable, Category = Niagara, meta=(DisplayName = "Niagara Set Float Array"))
-	static void SetNiagaraArrayFloat(UNiagaraComponent* NiagaraSystem, FName OverrideName, const TArray<float>& ArrayData);
+	static NIAGARA_API void SetNiagaraArrayFloat(UNiagaraComponent* NiagaraSystem, FName OverrideName, const TArray<float>& ArrayData);
 	/** Sets Niagara Array FVector2D Data. */
 	UFUNCTION(BlueprintCallable, Category = Niagara, meta = (DisplayName = "Niagara Set Vector 2D Array"))
-	static void SetNiagaraArrayVector2D(UNiagaraComponent* NiagaraSystem, FName OverrideName, const TArray<FVector2D>& ArrayData);
+	static NIAGARA_API void SetNiagaraArrayVector2D(UNiagaraComponent* NiagaraSystem, FName OverrideName, const TArray<FVector2D>& ArrayData);
 	/** Sets Niagara Array FVector Data. */
 	UFUNCTION(BlueprintCallable, Category = Niagara, meta = (DisplayName = "Niagara Set Vector Array"))
-	static void SetNiagaraArrayVector(UNiagaraComponent* NiagaraSystem, FName OverrideName, const TArray<FVector>& ArrayData);
+	static NIAGARA_API void SetNiagaraArrayVector(UNiagaraComponent* NiagaraSystem, FName OverrideName, const TArray<FVector>& ArrayData);
 	/** Sets Niagara Array FVector Data. */
 	UFUNCTION(BlueprintCallable, Category = Niagara, meta = (DisplayName = "Niagara Set Position Array"))
-	static void SetNiagaraArrayPosition(UNiagaraComponent* NiagaraSystem, FName OverrideName, const TArray<FVector>& ArrayData);
+	static NIAGARA_API void SetNiagaraArrayPosition(UNiagaraComponent* NiagaraSystem, FName OverrideName, const TArray<FVector>& ArrayData);
 	/** Sets Niagara Array FVector4 Data. */
 	UFUNCTION(BlueprintCallable, Category = Niagara, meta = (DisplayName = "Niagara Set Vector 4 Array"))
-	static void SetNiagaraArrayVector4(UNiagaraComponent* NiagaraSystem, FName OverrideName, const TArray<FVector4>& ArrayData);
+	static NIAGARA_API void SetNiagaraArrayVector4(UNiagaraComponent* NiagaraSystem, FName OverrideName, const TArray<FVector4>& ArrayData);
 	/** Sets Niagara Array FLinearColor Data. */
 	UFUNCTION(BlueprintCallable, Category = Niagara, meta = (DisplayName = "Niagara Set Color Array"))
-	static void SetNiagaraArrayColor(UNiagaraComponent* NiagaraSystem, FName OverrideName, const TArray<FLinearColor>& ArrayData);
+	static NIAGARA_API void SetNiagaraArrayColor(UNiagaraComponent* NiagaraSystem, FName OverrideName, const TArray<FLinearColor>& ArrayData);
 	/** Sets Niagara Array FQuat Data. */
 	UFUNCTION(BlueprintCallable, Category = Niagara, meta = (DisplayName = "Niagara Set Quaternion Array"))
-	static void SetNiagaraArrayQuat(UNiagaraComponent* NiagaraSystem, FName OverrideName, const TArray<FQuat>& ArrayData);
+	static NIAGARA_API void SetNiagaraArrayQuat(UNiagaraComponent* NiagaraSystem, FName OverrideName, const TArray<FQuat>& ArrayData);
 	/**
 	 * Sets Niagara Array FMatrix Data.
 	 * @param bApplyLWCRebase When enabled the matrix translation will have the simulation tile offset subtracted from it
 	 */
 	UFUNCTION(BlueprintCallable, Category = Niagara, meta = (DisplayName = "Niagara Set Matrix Array"))
-	static void SetNiagaraArrayMatrix(UNiagaraComponent* NiagaraSystem, FName OverrideName, const TArray<FMatrix>& ArrayData, bool bApplyLWCRebase = true);
+	static NIAGARA_API void SetNiagaraArrayMatrix(UNiagaraComponent* NiagaraSystem, FName OverrideName, const TArray<FMatrix>& ArrayData, bool bApplyLWCRebase = true);
 	/** Sets Niagara Array Int32 Data. */
 	UFUNCTION(BlueprintCallable, Category = Niagara, meta = (DisplayName = "Niagara Set Int32 Array"))
-	static void SetNiagaraArrayInt32(UNiagaraComponent* NiagaraSystem, FName OverrideName, const TArray<int32>& ArrayData);
+	static NIAGARA_API void SetNiagaraArrayInt32(UNiagaraComponent* NiagaraSystem, FName OverrideName, const TArray<int32>& ArrayData);
 	/** Sets Niagara Array UInt8 Data. */
 	UFUNCTION(BlueprintCallable, Category = Niagara, meta = (DisplayName = "Niagara Set UInt8 Array"))
-	static void SetNiagaraArrayUInt8(UNiagaraComponent* NiagaraSystem, FName OverrideName, const TArray<int32>& ArrayData);
+	static NIAGARA_API void SetNiagaraArrayUInt8(UNiagaraComponent* NiagaraSystem, FName OverrideName, const TArray<int32>& ArrayData);
 	/** Sets Niagara Array Bool Data. */
 	UFUNCTION(BlueprintCallable, Category = Niagara, meta = (DisplayName = "Niagara Set Bool Array"))
-	static void SetNiagaraArrayBool(UNiagaraComponent* NiagaraSystem, FName OverrideName, const TArray<bool>& ArrayData);
+	static NIAGARA_API void SetNiagaraArrayBool(UNiagaraComponent* NiagaraSystem, FName OverrideName, const TArray<bool>& ArrayData);
 
 	/** Gets a copy of Niagara Float Data. */
 	UFUNCTION(BlueprintCallable, Category = Niagara, meta = (DisplayName = "Niagara Get Float Array"))
-	static TArray<float> GetNiagaraArrayFloat(UNiagaraComponent* NiagaraSystem, FName OverrideName);
+	static NIAGARA_API TArray<float> GetNiagaraArrayFloat(UNiagaraComponent* NiagaraSystem, FName OverrideName);
 	/** Gets a copy of Niagara FVector2D Data. */
 	UFUNCTION(BlueprintCallable, Category = Niagara, meta = (DisplayName = "Niagara Get Vector 2D Array"))
-	static TArray<FVector2D> GetNiagaraArrayVector2D(UNiagaraComponent* NiagaraSystem, FName OverrideName);
+	static NIAGARA_API TArray<FVector2D> GetNiagaraArrayVector2D(UNiagaraComponent* NiagaraSystem, FName OverrideName);
 	/** Gets a copy of Niagara FVector Data. */
 	UFUNCTION(BlueprintCallable, Category = Niagara, meta = (DisplayName = "Niagara Get Vector Array"))
-	static TArray<FVector> GetNiagaraArrayVector(UNiagaraComponent* NiagaraSystem, FName OverrideName);
+	static NIAGARA_API TArray<FVector> GetNiagaraArrayVector(UNiagaraComponent* NiagaraSystem, FName OverrideName);
 	/** Gets a copy of Niagara Position Data. */
 	UFUNCTION(BlueprintCallable, Category = Niagara, meta = (DisplayName = "Niagara Get Position Array"))
-	static TArray<FVector> GetNiagaraArrayPosition(UNiagaraComponent* NiagaraSystem, FName OverrideName);
+	static NIAGARA_API TArray<FVector> GetNiagaraArrayPosition(UNiagaraComponent* NiagaraSystem, FName OverrideName);
 	/** Gets a copy of Niagara FVector4 Data. */
 	UFUNCTION(BlueprintCallable, Category = Niagara, meta = (DisplayName = "Niagara Get Vector 4 Array"))
-	static TArray<FVector4> GetNiagaraArrayVector4(UNiagaraComponent* NiagaraSystem, FName OverrideName);
+	static NIAGARA_API TArray<FVector4> GetNiagaraArrayVector4(UNiagaraComponent* NiagaraSystem, FName OverrideName);
 	/** Gets a copy of Niagara FLinearColor Data. */
 	UFUNCTION(BlueprintCallable, Category = Niagara, meta = (DisplayName = "Niagara Get Color Array"))
-	static TArray<FLinearColor> GetNiagaraArrayColor(UNiagaraComponent* NiagaraSystem, FName OverrideName);
+	static NIAGARA_API TArray<FLinearColor> GetNiagaraArrayColor(UNiagaraComponent* NiagaraSystem, FName OverrideName);
 	/** Gets a copy of Niagara FQuat Data. */
 	UFUNCTION(BlueprintCallable, Category = Niagara, meta = (DisplayName = "Niagara Get Quaternion Array"))
-	static TArray<FQuat> GetNiagaraArrayQuat(UNiagaraComponent* NiagaraSystem, FName OverrideName);
+	static NIAGARA_API TArray<FQuat> GetNiagaraArrayQuat(UNiagaraComponent* NiagaraSystem, FName OverrideName);
 	/**
 	 * Gets a copy of Niagara FMatrix Data.
 	 * @param bApplyLWCRebase When enabled the matrix translation will have the simulation tile offset added to it
 	 */
 	UFUNCTION(BlueprintCallable, Category = Niagara, meta = (DisplayName = "Niagara Get Matrix Array"))
-	static TArray<FMatrix> GetNiagaraArrayMatrix(UNiagaraComponent* NiagaraSystem, FName OverrideName, bool bApplyLWCRebase = true);
+	static NIAGARA_API TArray<FMatrix> GetNiagaraArrayMatrix(UNiagaraComponent* NiagaraSystem, FName OverrideName, bool bApplyLWCRebase = true);
 	/** Gets a copy of Niagara Int32 Data. */
 	UFUNCTION(BlueprintCallable, Category = Niagara, meta = (DisplayName = "Niagara Get Int32 Array"))
-	static TArray<int32> GetNiagaraArrayInt32(UNiagaraComponent* NiagaraSystem, FName OverrideName);
+	static NIAGARA_API TArray<int32> GetNiagaraArrayInt32(UNiagaraComponent* NiagaraSystem, FName OverrideName);
 	/** Gets a copy of Niagara UInt8 Data. */
 	UFUNCTION(BlueprintCallable, Category = Niagara, meta = (DisplayName = "Niagara Get UInt8 Array"))
-	static TArray<int32> GetNiagaraArrayUInt8(UNiagaraComponent* NiagaraSystem, FName OverrideName);
+	static NIAGARA_API TArray<int32> GetNiagaraArrayUInt8(UNiagaraComponent* NiagaraSystem, FName OverrideName);
 	/** Gets a copy of Niagara Bool Data. */
 	UFUNCTION(BlueprintCallable, Category = Niagara, meta = (DisplayName = "Niagara Get Bool Array"))
-	static TArray<bool> GetNiagaraArrayBool(UNiagaraComponent* NiagaraSystem, FName OverrideName);
+	static NIAGARA_API TArray<bool> GetNiagaraArrayBool(UNiagaraComponent* NiagaraSystem, FName OverrideName);
 
 	/** Sets a single value within a Niagara Array Float. */
 	UFUNCTION(BlueprintCallable, Category = Niagara, meta = (DisplayName = "Niagara Set Float Array Value"))
-	static void SetNiagaraArrayFloatValue(UNiagaraComponent* NiagaraSystem, FName OverrideName, int Index, float Value, bool bSizeToFit);
+	static NIAGARA_API void SetNiagaraArrayFloatValue(UNiagaraComponent* NiagaraSystem, FName OverrideName, int Index, float Value, bool bSizeToFit);
 	/** Sets a single value within a Niagara Array FVector2D. */
 	UFUNCTION(BlueprintCallable, Category = Niagara, meta = (DisplayName = "Niagara Set Vector 2D Array Value"))
-	static void SetNiagaraArrayVector2DValue(UNiagaraComponent* NiagaraSystem, FName OverrideName, int Index, const FVector2D& Value, bool bSizeToFit);
+	static NIAGARA_API void SetNiagaraArrayVector2DValue(UNiagaraComponent* NiagaraSystem, FName OverrideName, int Index, const FVector2D& Value, bool bSizeToFit);
 	/** Sets a single value within a Niagara Array FVector. */
 	UFUNCTION(BlueprintCallable, Category = Niagara, meta = (DisplayName = "Niagara Set Vector Array Value"))
-	static void SetNiagaraArrayVectorValue(UNiagaraComponent* NiagaraSystem, FName OverrideName, int Index, const FVector& Value, bool bSizeToFit);
+	static NIAGARA_API void SetNiagaraArrayVectorValue(UNiagaraComponent* NiagaraSystem, FName OverrideName, int Index, const FVector& Value, bool bSizeToFit);
 	UFUNCTION(BlueprintCallable, Category = Niagara, meta = (DisplayName = "Niagara Set Position Array Value"))
-	static void SetNiagaraArrayPositionValue(UNiagaraComponent* NiagaraSystem, FName OverrideName, int Index, const FVector& Value, bool bSizeToFit);
+	static NIAGARA_API void SetNiagaraArrayPositionValue(UNiagaraComponent* NiagaraSystem, FName OverrideName, int Index, const FVector& Value, bool bSizeToFit);
 	/** Sets a single value within a Niagara Array FVector4. */
 	UFUNCTION(BlueprintCallable, Category = Niagara, meta = (DisplayName = "Niagara Set Vector 4 Array Value"))
-	static void SetNiagaraArrayVector4Value(UNiagaraComponent* NiagaraSystem, FName OverrideName, int Index, const FVector4& Value, bool bSizeToFit);
+	static NIAGARA_API void SetNiagaraArrayVector4Value(UNiagaraComponent* NiagaraSystem, FName OverrideName, int Index, const FVector4& Value, bool bSizeToFit);
 	/** Sets a single value within a Niagara Array FLinearColor. */
 	UFUNCTION(BlueprintCallable, Category = Niagara, meta = (DisplayName = "Niagara Set Color Array Value"))
-	static void SetNiagaraArrayColorValue(UNiagaraComponent* NiagaraSystem, FName OverrideName, int Index, const FLinearColor& Value, bool bSizeToFit);
+	static NIAGARA_API void SetNiagaraArrayColorValue(UNiagaraComponent* NiagaraSystem, FName OverrideName, int Index, const FLinearColor& Value, bool bSizeToFit);
 	/** Sets a single value within a Niagara Array FQuat. */
 	UFUNCTION(BlueprintCallable, Category = Niagara, meta = (DisplayName = "Niagara Set Quaternion Array Value"))
-	static void SetNiagaraArrayQuatValue(UNiagaraComponent* NiagaraSystem, FName OverrideName, int Index, const FQuat& Value, bool bSizeToFit);
+	static NIAGARA_API void SetNiagaraArrayQuatValue(UNiagaraComponent* NiagaraSystem, FName OverrideName, int Index, const FQuat& Value, bool bSizeToFit);
 	/**
 	 * Sets a single value within a Niagara Array FMatrix.
 	 * @param bApplyLWCRebase When enabled the matrix translation will have the simulation tile offset subtracted from it
 	 */
 	UFUNCTION(BlueprintCallable, Category = Niagara, meta = (DisplayName = "Niagara Set Matrix Array Value"))
-	static void SetNiagaraArrayMatrixValue(UNiagaraComponent* NiagaraSystem, FName OverrideName, int Index, const FMatrix& Value, bool bSizeToFit, bool bApplyLWCRebase = true);
+	static NIAGARA_API void SetNiagaraArrayMatrixValue(UNiagaraComponent* NiagaraSystem, FName OverrideName, int Index, const FMatrix& Value, bool bSizeToFit, bool bApplyLWCRebase = true);
 	/** Sets a single value within a Niagara Array Int32. */
 	UFUNCTION(BlueprintCallable, Category = Niagara, meta = (DisplayName = "Niagara Set Int32 Array Value"))
-	static void SetNiagaraArrayInt32Value(UNiagaraComponent* NiagaraSystem, FName OverrideName, int Index, int32 Value, bool bSizeToFit);
+	static NIAGARA_API void SetNiagaraArrayInt32Value(UNiagaraComponent* NiagaraSystem, FName OverrideName, int Index, int32 Value, bool bSizeToFit);
 	/** Sets a single value within a Niagara Array UInt8. */
 	UFUNCTION(BlueprintCallable, Category = Niagara, meta = (DisplayName = "Niagara Set UInt8 Array Value"))
-	static void SetNiagaraArrayUInt8Value(UNiagaraComponent* NiagaraSystem, FName OverrideName, int Index, int32 Value, bool bSizeToFit);
+	static NIAGARA_API void SetNiagaraArrayUInt8Value(UNiagaraComponent* NiagaraSystem, FName OverrideName, int Index, int32 Value, bool bSizeToFit);
 	/** Sets a single value within a Niagara Array Bool. */
 	UFUNCTION(BlueprintCallable, Category = Niagara, meta = (DisplayName = "Niagara Set Bool Array Value"))
-	static void SetNiagaraArrayBoolValue(UNiagaraComponent* NiagaraSystem, FName OverrideName, int Index, const bool& Value, bool bSizeToFit);
+	static NIAGARA_API void SetNiagaraArrayBoolValue(UNiagaraComponent* NiagaraSystem, FName OverrideName, int Index, const bool& Value, bool bSizeToFit);
 
 	/** Gets a single value within a Niagara Array Float. */
 	UFUNCTION(BlueprintCallable, Category = Niagara, meta = (DisplayName = "Niagara Get Float Array Value"))
-	static float GetNiagaraArrayFloatValue(UNiagaraComponent* NiagaraSystem, FName OverrideName, int Index);
+	static NIAGARA_API float GetNiagaraArrayFloatValue(UNiagaraComponent* NiagaraSystem, FName OverrideName, int Index);
 	/** Gets a single value within a Niagara Array FVector2D. */
 	UFUNCTION(BlueprintCallable, Category = Niagara, meta = (DisplayName = "Niagara Get Vector 2D Array Value"))
-	static FVector2D GetNiagaraArrayVector2DValue(UNiagaraComponent* NiagaraSystem, FName OverrideName, int Index);
+	static NIAGARA_API FVector2D GetNiagaraArrayVector2DValue(UNiagaraComponent* NiagaraSystem, FName OverrideName, int Index);
 	/** Gets a single value within a Niagara Array FVector. */
 	UFUNCTION(BlueprintCallable, Category = Niagara, meta = (DisplayName = "Niagara Get Vector Array Value"))
-	static FVector GetNiagaraArrayVectorValue(UNiagaraComponent* NiagaraSystem, FName OverrideName, int Index);
+	static NIAGARA_API FVector GetNiagaraArrayVectorValue(UNiagaraComponent* NiagaraSystem, FName OverrideName, int Index);
 	/** Gets a single value within a Niagara Array Position. */
 	UFUNCTION(BlueprintCallable, Category = Niagara, meta = (DisplayName = "Niagara Get Position Array Value"))
-	static FVector GetNiagaraArrayPositionValue(UNiagaraComponent* NiagaraSystem, FName OverrideName, int Index);
+	static NIAGARA_API FVector GetNiagaraArrayPositionValue(UNiagaraComponent* NiagaraSystem, FName OverrideName, int Index);
 	/** Gets a single value within a Niagara Array FVector4. */
 	UFUNCTION(BlueprintCallable, Category = Niagara, meta = (DisplayName = "Niagara Get Vector 4 Array Value"))
-	static FVector4 GetNiagaraArrayVector4Value(UNiagaraComponent* NiagaraSystem, FName OverrideName, int Index);
+	static NIAGARA_API FVector4 GetNiagaraArrayVector4Value(UNiagaraComponent* NiagaraSystem, FName OverrideName, int Index);
 	/** Gets a single value within a Niagara Array FLinearColor. */
 	UFUNCTION(BlueprintCallable, Category = Niagara, meta = (DisplayName = "Niagara Get Color Array Value"))
-	static FLinearColor GetNiagaraArrayColorValue(UNiagaraComponent* NiagaraSystem, FName OverrideName, int Index);
+	static NIAGARA_API FLinearColor GetNiagaraArrayColorValue(UNiagaraComponent* NiagaraSystem, FName OverrideName, int Index);
 	/** Gets a single value within a Niagara Array FQuat. */
 	UFUNCTION(BlueprintCallable, Category = Niagara, meta = (DisplayName = "Niagara Get Quaternion Array Value"))
-	static FQuat GetNiagaraArrayQuatValue(UNiagaraComponent* NiagaraSystem, FName OverrideName, int Index);
+	static NIAGARA_API FQuat GetNiagaraArrayQuatValue(UNiagaraComponent* NiagaraSystem, FName OverrideName, int Index);
 	/**
 	 * Gets a single value within a Niagara Array FMatrix.
 	 * @param bApplyLWCRebase When enabled the matrix translation will have the simulation tile offset added to it
 	 */
 	UFUNCTION(BlueprintCallable, Category = Niagara, meta = (DisplayName = "Niagara Get Matrix Array Value"))
-	static FMatrix GetNiagaraArrayMatrixValue(UNiagaraComponent* NiagaraSystem, FName OverrideName, int Index, bool bApplyLWCRebase = true);
+	static NIAGARA_API FMatrix GetNiagaraArrayMatrixValue(UNiagaraComponent* NiagaraSystem, FName OverrideName, int Index, bool bApplyLWCRebase = true);
 	/** Gets a single value within a Niagara Array Int32. */
 	UFUNCTION(BlueprintCallable, Category = Niagara, meta = (DisplayName = "Niagara Get Int32 Array Value"))
-	static int32 GetNiagaraArrayInt32Value(UNiagaraComponent* NiagaraSystem, FName OverrideName, int Index);
+	static NIAGARA_API int32 GetNiagaraArrayInt32Value(UNiagaraComponent* NiagaraSystem, FName OverrideName, int Index);
 	/** Gets a single value within a Niagara Array UInt8. */
 	UFUNCTION(BlueprintCallable, Category = Niagara, meta = (DisplayName = "Niagara Get UInt8 Array Value"))
-	static int32 GetNiagaraArrayUInt8Value(UNiagaraComponent* NiagaraSystem, FName OverrideName, int Index);
+	static NIAGARA_API int32 GetNiagaraArrayUInt8Value(UNiagaraComponent* NiagaraSystem, FName OverrideName, int Index);
 	/** Gets a single value within a Niagara Array Bool. */
 	UFUNCTION(BlueprintCallable, Category = Niagara, meta = (DisplayName = "Niagara Get Bool Array Value"))
-	static bool GetNiagaraArrayBoolValue(UNiagaraComponent* NiagaraSystem, FName OverrideName, int Index);
+	static NIAGARA_API bool GetNiagaraArrayBoolValue(UNiagaraComponent* NiagaraSystem, FName OverrideName, int Index);
 
 	//////////////////////////////////////////////////////////////////////////////////////////////////////////////////
 	// None BP compatable set functions
-	static void SetNiagaraArrayVector2D(UNiagaraComponent* NiagaraSystem, FName OverrideName, TConstArrayView<FVector2f> ArrayData);
-	static void SetNiagaraArrayVector(UNiagaraComponent* NiagaraSystem, FName OverrideName, TConstArrayView<FVector3f> ArrayData);
-	static void SetNiagaraArrayVector4(UNiagaraComponent* NiagaraSystem, FName OverrideName, TConstArrayView<FVector4f> ArrayData);
-	static void SetNiagaraArrayQuat(UNiagaraComponent* NiagaraSystem, FName OverrideName, TConstArrayView<FQuat4f> ArrayData);
-	static void SetNiagaraArrayMatrix(UNiagaraComponent* NiagaraSystem, FName OverrideName, TConstArrayView<FMatrix44f> ArrayData);
-	static void SetNiagaraArrayUInt8(UNiagaraComponent* NiagaraSystem, FName OverrideName, TConstArrayView<uint8> ArrayData);
+	static NIAGARA_API void SetNiagaraArrayVector2D(UNiagaraComponent* NiagaraSystem, FName OverrideName, TConstArrayView<FVector2f> ArrayData);
+	static NIAGARA_API void SetNiagaraArrayVector(UNiagaraComponent* NiagaraSystem, FName OverrideName, TConstArrayView<FVector3f> ArrayData);
+	static NIAGARA_API void SetNiagaraArrayVector4(UNiagaraComponent* NiagaraSystem, FName OverrideName, TConstArrayView<FVector4f> ArrayData);
+	static NIAGARA_API void SetNiagaraArrayQuat(UNiagaraComponent* NiagaraSystem, FName OverrideName, TConstArrayView<FQuat4f> ArrayData);
+	static NIAGARA_API void SetNiagaraArrayMatrix(UNiagaraComponent* NiagaraSystem, FName OverrideName, TConstArrayView<FMatrix44f> ArrayData);
+	static NIAGARA_API void SetNiagaraArrayUInt8(UNiagaraComponent* NiagaraSystem, FName OverrideName, TConstArrayView<uint8> ArrayData);
 };
diff --git a/Engine/Plugins/FX/Niagara/Source/Niagara/Classes/NiagaraDataInterfaceArrayInt.h b/Engine/Plugins/FX/Niagara/Source/Niagara/Classes/NiagaraDataInterfaceArrayInt.h
--- a/Engine/Plugins/FX/Niagara/Source/Niagara/Classes/NiagaraDataInterfaceArrayInt.h
+++ b/Engine/Plugins/FX/Niagara/Source/Niagara/Classes/NiagaraDataInterfaceArrayInt.h
@@ -5,8 +5,8 @@
 #include "NiagaraDataInterfaceArray.h"
 #include "NiagaraDataInterfaceArrayInt.generated.h"
 
-UCLASS(EditInlineNew, Category = "Array", CollapseCategories, meta = (DisplayName = "Int32 Array"), Blueprintable, BlueprintType)
-class NIAGARA_API UNiagaraDataInterfaceArrayInt32 : public UNiagaraDataInterfaceArray
+UCLASS(EditInlineNew, Category = "Array", CollapseCategories, meta = (DisplayName = "Int32 Array"), Blueprintable, BlueprintType, MinimalAPI)
+class UNiagaraDataInterfaceArrayInt32 : public UNiagaraDataInterfaceArray
 {
 public:
 	GENERATED_BODY()
@@ -17,8 +17,8 @@ public:
 	NDIARRAY_GENERATE_BODY(UNiagaraDataInterfaceArrayInt32, int32, IntData)
 };
 
-UCLASS(EditInlineNew, Category = "Array", CollapseCategories, meta = (DisplayName = "UInt8 Array"), Blueprintable, BlueprintType)
-class NIAGARA_API UNiagaraDataInterfaceArrayUInt8 : public UNiagaraDataInterfaceArray
+UCLASS(EditInlineNew, Category = "Array", CollapseCategories, meta = (DisplayName = "UInt8 Array"), Blueprintable, BlueprintType, MinimalAPI)
+class UNiagaraDataInterfaceArrayUInt8 : public UNiagaraDataInterfaceArray
 {
 public:
 	GENERATED_BODY()
@@ -34,8 +34,8 @@ public:
 	NDIARRAY_GENERATE_BODY_LWC(UNiagaraDataInterfaceArrayUInt8, uint8, IntData)
 };
 
-UCLASS(EditInlineNew, Category = "Array", CollapseCategories, meta = (DisplayName = "Bool Array"), Blueprintable, BlueprintType)
-class NIAGARA_API UNiagaraDataInterfaceArrayBool : public UNiagaraDataInterfaceArray
+UCLASS(EditInlineNew, Category = "Array", CollapseCategories, meta = (DisplayName = "Bool Array"), Blueprintable, BlueprintType, MinimalAPI)
+class UNiagaraDataInterfaceArrayBool : public UNiagaraDataInterfaceArray
 {
 public:
 	GENERATED_BODY()
diff --git a/Engine/Plugins/FX/Niagara/Source/Niagara/Classes/NiagaraDataInterfaceArrayNiagaraID.h b/Engine/Plugins/FX/Niagara/Source/Niagara/Classes/NiagaraDataInterfaceArrayNiagaraID.h
--- a/Engine/Plugins/FX/Niagara/Source/Niagara/Classes/NiagaraDataInterfaceArrayNiagaraID.h
+++ b/Engine/Plugins/FX/Niagara/Source/Niagara/Classes/NiagaraDataInterfaceArrayNiagaraID.h
@@ -5,8 +5,8 @@
 #include "NiagaraDataInterfaceArray.h"
 #include "NiagaraDataInterfaceArrayNiagaraID.generated.h"
 
-UCLASS(EditInlineNew, Category = "Array", meta = (DisplayName = "NiagaraID Array", Experimental), Blueprintable, BlueprintType)
-class NIAGARA_API UNiagaraDataInterfaceArrayNiagaraID : public UNiagaraDataInterfaceArray
+UCLASS(EditInlineNew, Category = "Array", meta = (DisplayName = "NiagaraID Array", Experimental), Blueprintable, BlueprintType, MinimalAPI)
+class UNiagaraDataInterfaceArrayNiagaraID : public UNiagaraDataInterfaceArray
 {
 public:
 	GENERATED_BODY()
diff --git a/Engine/Plugins/FX/Niagara/Source/Niagara/Classes/NiagaraDataInterfaceAsyncGpuTrace.h b/Engine/Plugins/FX/Niagara/Source/Niagara/Classes/NiagaraDataInterfaceAsyncGpuTrace.h
--- a/Engine/Plugins/FX/Niagara/Source/Niagara/Classes/NiagaraDataInterfaceAsyncGpuTrace.h
+++ b/Engine/Plugins/FX/Niagara/Source/Niagara/Classes/NiagaraDataInterfaceAsyncGpuTrace.h
@@ -11,8 +11,8 @@
 class FNiagaraSystemInstance;
 
 /** Data interface for handling latent (delayed 1 frame) traces against the scene for GPU simulations. */
-UCLASS(EditInlineNew, Category = "Collision", CollapseCategories, meta = (DisplayName = "Async Gpu Trace"))
-class NIAGARA_API UNiagaraDataInterfaceAsyncGpuTrace : public UNiagaraDataInterface
+UCLASS(EditInlineNew, Category = "Collision", CollapseCategories, meta = (DisplayName = "Async Gpu Trace"), MinimalAPI)
+class UNiagaraDataInterfaceAsyncGpuTrace : public UNiagaraDataInterface
 {
 	GENERATED_UCLASS_BODY()
 
@@ -31,39 +31,39 @@ public:
 	TEnumAsByte<ENDICollisionQuery_AsyncGpuTraceProvider::Type> TraceProvider = ENDICollisionQuery_AsyncGpuTraceProvider::Default;
 
 	//UObject Interface
-	virtual void PostInitProperties() override;
-	virtual void PostLoad() override;
+	NIAGARA_API virtual void PostInitProperties() override;
+	NIAGARA_API virtual void PostLoad() override;
 
 #if WITH_EDITOR
-	virtual void PostEditChangeProperty(struct FPropertyChangedEvent& PropertyChangedEvent) override;
-	virtual void GetFeedback(UNiagaraSystem* InAsset, UNiagaraComponent* InComponent, TArray<FNiagaraDataInterfaceError>& OutErrors, TArray<FNiagaraDataInterfaceFeedback>& OutWarnings, TArray<FNiagaraDataInterfaceFeedback>& OutInfo) override;
+	NIAGARA_API virtual void PostEditChangeProperty(struct FPropertyChangedEvent& PropertyChangedEvent) override;
+	NIAGARA_API virtual void GetFeedback(UNiagaraSystem* InAsset, UNiagaraComponent* InComponent, TArray<FNiagaraDataInterfaceError>& OutErrors, TArray<FNiagaraDataInterfaceFeedback>& OutWarnings, TArray<FNiagaraDataInterfaceFeedback>& OutInfo) override;
 #endif
 	//UObject Interface End
 
-	virtual int32 PerInstanceDataSize() const override;
-	virtual bool InitPerInstanceData(void* PerInstanceData, FNiagaraSystemInstance* InSystemInstance) override;
-	virtual void DestroyPerInstanceData(void* PerInstanceData, FNiagaraSystemInstance* InSystemInstance) override;
+	NIAGARA_API virtual int32 PerInstanceDataSize() const override;
+	NIAGARA_API virtual bool InitPerInstanceData(void* PerInstanceData, FNiagaraSystemInstance* InSystemInstance) override;
+	NIAGARA_API virtual void DestroyPerInstanceData(void* PerInstanceData, FNiagaraSystemInstance* InSystemInstance) override;
 
-	virtual void GetFunctions(TArray<FNiagaraFunctionSignature>& OutFunctions) override;
+	NIAGARA_API virtual void GetFunctions(TArray<FNiagaraFunctionSignature>& OutFunctions) override;
 
 	virtual bool CanExecuteOnTarget(ENiagaraSimTarget Target) const override { return Target == ENiagaraSimTarget::GPUComputeSim; }
-	virtual bool RequiresDistanceFieldData() const override;
-	virtual bool RequiresRayTracingScene() const override;
+	NIAGARA_API virtual bool RequiresDistanceFieldData() const override;
+	NIAGARA_API virtual bool RequiresRayTracingScene() const override;
 
 #if WITH_EDITORONLY_DATA
-	virtual bool AppendCompileHash(FNiagaraCompileHashVisitor* InVisitor) const override;
-	virtual void GetCommonHLSL(FString& OutHLSL) override;
-	virtual void GetParameterDefinitionHLSL(const FNiagaraDataInterfaceGPUParamInfo& ParamInfo, FString& OutHLSL) override;
-	virtual bool GetFunctionHLSL(const FNiagaraDataInterfaceGPUParamInfo& ParamInfo, const FNiagaraDataInterfaceGeneratedFunction& FunctionInfo, int FunctionInstanceIndex, FString& OutHLSL) override;
-	virtual bool UpgradeFunctionCall(FNiagaraFunctionSignature& FunctionSignature) override;
+	NIAGARA_API virtual bool AppendCompileHash(FNiagaraCompileHashVisitor* InVisitor) const override;
+	NIAGARA_API virtual void GetCommonHLSL(FString& OutHLSL) override;
+	NIAGARA_API virtual void GetParameterDefinitionHLSL(const FNiagaraDataInterfaceGPUParamInfo& ParamInfo, FString& OutHLSL) override;
+	NIAGARA_API virtual bool GetFunctionHLSL(const FNiagaraDataInterfaceGPUParamInfo& ParamInfo, const FNiagaraDataInterfaceGeneratedFunction& FunctionInfo, int FunctionInstanceIndex, FString& OutHLSL) override;
+	NIAGARA_API virtual bool UpgradeFunctionCall(FNiagaraFunctionSignature& FunctionSignature) override;
 #endif
-	virtual void BuildShaderParameters(FNiagaraShaderParametersBuilder& ShaderParametersBuilder) const override;
-	virtual void SetShaderParameters(const FNiagaraDataInterfaceSetShaderParametersContext& Context) const override;
+	NIAGARA_API virtual void BuildShaderParameters(FNiagaraShaderParametersBuilder& ShaderParametersBuilder) const override;
+	NIAGARA_API virtual void SetShaderParameters(const FNiagaraDataInterfaceSetShaderParametersContext& Context) const override;
 
 private:
-	virtual bool Equals(const UNiagaraDataInterface* Other) const override;
+	NIAGARA_API virtual bool Equals(const UNiagaraDataInterface* Other) const override;
 
 protected:
-	virtual bool CopyToInternal(UNiagaraDataInterface* Destination) const override;
-	virtual void PushToRenderThreadImpl() override;
+	NIAGARA_API virtual bool CopyToInternal(UNiagaraDataInterface* Destination) const override;
+	NIAGARA_API virtual void PushToRenderThreadImpl() override;
 };
diff --git a/Engine/Plugins/FX/Niagara/Source/Niagara/Classes/NiagaraDataInterfaceAudio.h b/Engine/Plugins/FX/Niagara/Source/Niagara/Classes/NiagaraDataInterfaceAudio.h
--- a/Engine/Plugins/FX/Niagara/Source/Niagara/Classes/NiagaraDataInterfaceAudio.h
+++ b/Engine/Plugins/FX/Niagara/Source/Niagara/Classes/NiagaraDataInterfaceAudio.h
@@ -14,7 +14,7 @@
 #endif
 
 /** Class used to to capture the audio stream of an arbitrary submix. */
-class NIAGARA_API FNiagaraSubmixListener : public ISubmixBufferListener
+class FNiagaraSubmixListener : public ISubmixBufferListener
 {
 public:
 
@@ -25,7 +25,7 @@ public:
 	 *                               data to be overwritten by threads that produce audio. High 
 	 *                               values require more memory.
 	 */
-	FNiagaraSubmixListener(Audio::FPatchMixer& InMixer, int32 InNumSamplesToBuffer, Audio::FDeviceId InDeviceId, USoundSubmix* InSoundSubmix);
+	NIAGARA_API FNiagaraSubmixListener(Audio::FPatchMixer& InMixer, int32 InNumSamplesToBuffer, Audio::FDeviceId InDeviceId, USoundSubmix* InSoundSubmix);
 
 	FNiagaraSubmixListener(const FNiagaraSubmixListener& Other)
 	{
@@ -35,27 +35,27 @@ public:
 	}
 
 	/** Move submix listener. */
-	FNiagaraSubmixListener(FNiagaraSubmixListener&& Other);
+	NIAGARA_API FNiagaraSubmixListener(FNiagaraSubmixListener&& Other);
 
-	void RegisterToSubmix();
+	NIAGARA_API void RegisterToSubmix();
 
-	virtual ~FNiagaraSubmixListener();
+	NIAGARA_API virtual ~FNiagaraSubmixListener();
 
 	/** Returns the current sample rate of the current submix. */
-	float GetSampleRate() const;
+	NIAGARA_API float GetSampleRate() const;
 
 	/** Returns the number of channels of the current submix. */
-	int32 GetNumChannels() const;
+	NIAGARA_API int32 GetNumChannels() const;
 
 	// Begin ISubmixBufferListener overrides
-	virtual void OnNewSubmixBuffer(const USoundSubmix* OwningSubmix, float* AudioData, int32 NumSamples, int32 NumChannels, const int32 SampleRate, double AudioClock) override;
+	NIAGARA_API virtual void OnNewSubmixBuffer(const USoundSubmix* OwningSubmix, float* AudioData, int32 NumSamples, int32 NumChannels, const int32 SampleRate, double AudioClock) override;
 	// End ISubmixBufferListener overrides
 
 private:
-	FNiagaraSubmixListener();
+	NIAGARA_API FNiagaraSubmixListener();
 
 
-	void UnregisterFromSubmix();
+	NIAGARA_API void UnregisterFromSubmix();
 
 	TAtomic<int32> NumChannelsInSubmix;
 	TAtomic<int32> SubmixSampleRate;
@@ -132,8 +132,8 @@ private:
 };
 
 /** Data Interface allowing sampling of recent audio data. */
-UCLASS(EditInlineNew, Category = "Audio", CollapseCategories, meta = (DisplayName = "Audio Submix"))
-class NIAGARA_API UNiagaraDataInterfaceAudioSubmix : public UNiagaraDataInterface
+UCLASS(EditInlineNew, Category = "Audio", CollapseCategories, meta = (DisplayName = "Audio Submix"), MinimalAPI)
+class UNiagaraDataInterfaceAudioSubmix : public UNiagaraDataInterface
 {
 	GENERATED_UCLASS_BODY()
 public:
@@ -143,18 +143,18 @@ public:
 	TObjectPtr<USoundSubmix> Submix;
 
 
-	virtual bool Equals(const UNiagaraDataInterface* Other) const override;
+	NIAGARA_API virtual bool Equals(const UNiagaraDataInterface* Other) const override;
 
 #if WITH_EDITOR
-	virtual void PostEditChangeProperty(struct FPropertyChangedEvent& PropertyChangedEvent) override;
+	NIAGARA_API virtual void PostEditChangeProperty(struct FPropertyChangedEvent& PropertyChangedEvent) override;
 #endif // WITH_EDITOR
 
-	virtual void PostInitProperties() override;
-	virtual void BeginDestroy() override;
-	virtual void PostLoad() override;
+	NIAGARA_API virtual void PostInitProperties() override;
+	NIAGARA_API virtual void BeginDestroy() override;
+	NIAGARA_API virtual void PostLoad() override;
 
 protected:
-	virtual bool CopyToInternal(UNiagaraDataInterface* Destination) const override;
+	NIAGARA_API virtual bool CopyToInternal(UNiagaraDataInterface* Destination) const override;
 private:
 
 };
diff --git a/Engine/Plugins/FX/Niagara/Source/Niagara/Classes/NiagaraDataInterfaceAudioOscilloscope.h b/Engine/Plugins/FX/Niagara/Source/Niagara/Classes/NiagaraDataInterfaceAudioOscilloscope.h
--- a/Engine/Plugins/FX/Niagara/Source/Niagara/Classes/NiagaraDataInterfaceAudioOscilloscope.h
+++ b/Engine/Plugins/FX/Niagara/Source/Niagara/Classes/NiagaraDataInterfaceAudioOscilloscope.h
@@ -87,8 +87,8 @@ private:
 };
 
 /** Data Interface allowing sampling of recent audio data. */
-UCLASS(EditInlineNew, Category = "Audio", CollapseCategories, meta = (DisplayName = "Audio Oscilloscope"))
-class NIAGARA_API UNiagaraDataInterfaceAudioOscilloscope : public UNiagaraDataInterface
+UCLASS(EditInlineNew, Category = "Audio", CollapseCategories, meta = (DisplayName = "Audio Oscilloscope"), MinimalAPI)
+class UNiagaraDataInterfaceAudioOscilloscope : public UNiagaraDataInterface
 {
 	GENERATED_UCLASS_BODY()
 
@@ -114,11 +114,11 @@ public:
 	float ScopeInMilliseconds;
 
 	//VM function overrides:
-	void SampleAudio(FVectorVMExternalFunctionContext& Context);
-	void GetNumChannels(FVectorVMExternalFunctionContext& Context);
+	NIAGARA_API void SampleAudio(FVectorVMExternalFunctionContext& Context);
+	NIAGARA_API void GetNumChannels(FVectorVMExternalFunctionContext& Context);
 
-	virtual void GetFunctions(TArray<FNiagaraFunctionSignature>& OutFunctions)override;
-	virtual void GetVMExternalFunction(const FVMExternalFunctionBindingInfo& BindingInfo, void* InstanceData, FVMExternalFunction &OutFunc) override;
+	NIAGARA_API virtual void GetFunctions(TArray<FNiagaraFunctionSignature>& OutFunctions)override;
+	NIAGARA_API virtual void GetVMExternalFunction(const FVMExternalFunctionBindingInfo& BindingInfo, void* InstanceData, FVMExternalFunction &OutFunc) override;
 
 	virtual bool CanExecuteOnTarget(ENiagaraSimTarget Target) const override
 	{
@@ -128,25 +128,25 @@ public:
 	virtual bool RequiresDistanceFieldData() const override { return false; }
 
 #if WITH_EDITORONLY_DATA
-	virtual bool AppendCompileHash(FNiagaraCompileHashVisitor* InVisitor) const override;
-	virtual bool GetFunctionHLSL(const FNiagaraDataInterfaceGPUParamInfo& ParamInfo, const FNiagaraDataInterfaceGeneratedFunction& FunctionInfo, int FunctionInstanceIndex, FString& OutHLSL) override;
-	virtual void GetParameterDefinitionHLSL(const FNiagaraDataInterfaceGPUParamInfo& ParamInfo, FString& OutHLSL) override;
+	NIAGARA_API virtual bool AppendCompileHash(FNiagaraCompileHashVisitor* InVisitor) const override;
+	NIAGARA_API virtual bool GetFunctionHLSL(const FNiagaraDataInterfaceGPUParamInfo& ParamInfo, const FNiagaraDataInterfaceGeneratedFunction& FunctionInfo, int FunctionInstanceIndex, FString& OutHLSL) override;
+	NIAGARA_API virtual void GetParameterDefinitionHLSL(const FNiagaraDataInterfaceGPUParamInfo& ParamInfo, FString& OutHLSL) override;
 #endif
-	virtual void BuildShaderParameters(FNiagaraShaderParametersBuilder& ShaderParametersBuilder) const override;
-	virtual void SetShaderParameters(const FNiagaraDataInterfaceSetShaderParametersContext& Context) const override;
+	NIAGARA_API virtual void BuildShaderParameters(FNiagaraShaderParametersBuilder& ShaderParametersBuilder) const override;
+	NIAGARA_API virtual void SetShaderParameters(const FNiagaraDataInterfaceSetShaderParametersContext& Context) const override;
 
-	virtual bool Equals(const UNiagaraDataInterface* Other) const override;
+	NIAGARA_API virtual bool Equals(const UNiagaraDataInterface* Other) const override;
 
 #if WITH_EDITOR
-	virtual void PostEditChangeProperty(struct FPropertyChangedEvent& PropertyChangedEvent) override;
+	NIAGARA_API virtual void PostEditChangeProperty(struct FPropertyChangedEvent& PropertyChangedEvent) override;
 #endif // WITH_EDITOR
 
-	virtual void PostInitProperties() override;
-	virtual void BeginDestroy() override;
-	virtual void PostLoad() override;
+	NIAGARA_API virtual void PostInitProperties() override;
+	NIAGARA_API virtual void BeginDestroy() override;
+	NIAGARA_API virtual void PostLoad() override;
 	
 
 protected:
-	virtual bool CopyToInternal(UNiagaraDataInterface* Destination) const override;
+	NIAGARA_API virtual bool CopyToInternal(UNiagaraDataInterface* Destination) const override;
 };
 
diff --git a/Engine/Plugins/FX/Niagara/Source/Niagara/Classes/NiagaraDataInterfaceAudioPlayer.h b/Engine/Plugins/FX/Niagara/Source/Niagara/Classes/NiagaraDataInterfaceAudioPlayer.h
--- a/Engine/Plugins/FX/Niagara/Source/Niagara/Classes/NiagaraDataInterfaceAudioPlayer.h
+++ b/Engine/Plugins/FX/Niagara/Source/Niagara/Classes/NiagaraDataInterfaceAudioPlayer.h
@@ -31,13 +31,13 @@ struct FPersistentAudioParticleData
 	TFunction<void(struct FAudioPlayerInterface_InstanceData*,UAudioComponent*,FNiagaraSystemInstance*)> UpdateCallback;
 };
 
-UCLASS(EditInlineNew, Category = "Audio", meta = (DisplayName = "Niagara Audio Player Settings"), Blueprintable, BlueprintType)
-class NIAGARA_API UNiagaraDataInterfaceAudioPlayerSettings : public UObject
+UCLASS(EditInlineNew, Category = "Audio", meta = (DisplayName = "Niagara Audio Player Settings"), Blueprintable, BlueprintType, MinimalAPI)
+class UNiagaraDataInterfaceAudioPlayerSettings : public UObject
 {
 public:
 	GENERATED_BODY()
 
-	UNiagaraDataInterfaceAudioPlayerSettings();
+	NIAGARA_API UNiagaraDataInterfaceAudioPlayerSettings();
 
 	UPROPERTY(EditAnywhere, BlueprintReadWrite, Category = "Audio", meta=(InlineEditConditionToggle))
 	bool bOverrideConcurrency = true;
@@ -80,8 +80,8 @@ struct FAudioPlayerInterface_InstanceData
 };
 
 /** This Data Interface can be used to play one-shot audio effects driven by particle data. */
-UCLASS(EditInlineNew, Category = "Audio", meta = (DisplayName = "Audio Player"))
-class NIAGARA_API UNiagaraDataInterfaceAudioPlayer : public UNiagaraDataInterface
+UCLASS(EditInlineNew, Category = "Audio", meta = (DisplayName = "Audio Player"), MinimalAPI)
+class UNiagaraDataInterfaceAudioPlayer : public UNiagaraDataInterface
 {
 	GENERATED_UCLASS_BODY()
 
@@ -129,22 +129,22 @@ public:
 	UPROPERTY(EditAnywhere, AdvancedDisplay, Category = "Audio")
 	bool bOnlyActiveDuringGameplay = false;
 
-	virtual bool UpgradeFunctionCall(FNiagaraFunctionSignature& FunctionSignature) override;
+	NIAGARA_API virtual bool UpgradeFunctionCall(FNiagaraFunctionSignature& FunctionSignature) override;
 #endif
 	
 	//UObject Interface
-	virtual void PostInitProperties() override;
+	NIAGARA_API virtual void PostInitProperties() override;
 	//UObject Interface End
 
 	//UNiagaraDataInterface Interface
-	virtual void GetFunctions(TArray<FNiagaraFunctionSignature>& OutFunctions) override;
-	virtual void GetVMExternalFunction(const FVMExternalFunctionBindingInfo& BindingInfo, void* InstanceData, FVMExternalFunction &OutFunc) override;
-	virtual bool InitPerInstanceData(void* PerInstanceData, FNiagaraSystemInstance* SystemInstance) override;
-	virtual void DestroyPerInstanceData(void* PerInstanceData, FNiagaraSystemInstance* SystemInstance) override;
+	NIAGARA_API virtual void GetFunctions(TArray<FNiagaraFunctionSignature>& OutFunctions) override;
+	NIAGARA_API virtual void GetVMExternalFunction(const FVMExternalFunctionBindingInfo& BindingInfo, void* InstanceData, FVMExternalFunction &OutFunc) override;
+	NIAGARA_API virtual bool InitPerInstanceData(void* PerInstanceData, FNiagaraSystemInstance* SystemInstance) override;
+	NIAGARA_API virtual void DestroyPerInstanceData(void* PerInstanceData, FNiagaraSystemInstance* SystemInstance) override;
 	virtual int32 PerInstanceDataSize() const override { return sizeof(FAudioPlayerInterface_InstanceData); }
-	virtual bool PerInstanceTick(void* PerInstanceData, FNiagaraSystemInstance* SystemInstance, float DeltaSeconds) override;
-	virtual bool PerInstanceTickPostSimulate(void* PerInstanceData, FNiagaraSystemInstance* SystemInstance, float DeltaSeconds) override;
-	virtual bool Equals(const UNiagaraDataInterface* Other) const override;
+	NIAGARA_API virtual bool PerInstanceTick(void* PerInstanceData, FNiagaraSystemInstance* SystemInstance, float DeltaSeconds) override;
+	NIAGARA_API virtual bool PerInstanceTickPostSimulate(void* PerInstanceData, FNiagaraSystemInstance* SystemInstance, float DeltaSeconds) override;
+	NIAGARA_API virtual bool Equals(const UNiagaraDataInterface* Other) const override;
 	virtual bool CanExecuteOnTarget(ENiagaraSimTarget Target) const override { return Target == ENiagaraSimTarget::CPUSim; }
 
 	virtual bool HasPreSimulateTick() const override { return true; }
@@ -152,29 +152,29 @@ public:
 	virtual bool PostSimulateCanOverlapFrames() const { return false; }
 	//UNiagaraDataInterface Interface
 
-	virtual void PlayOneShotAudio(FVectorVMExternalFunctionContext& Context);
-	virtual void PlayPersistentAudio(FVectorVMExternalFunctionContext& Context);
-	virtual void SetParameterBool(FVectorVMExternalFunctionContext& Context);
-	virtual void SetParameterInteger(FVectorVMExternalFunctionContext& Context);
-	virtual void SetParameterFloat(FVectorVMExternalFunctionContext& Context);
-	virtual void UpdateVolume(FVectorVMExternalFunctionContext& Context);
-	virtual void UpdatePitch(FVectorVMExternalFunctionContext& Context);
-	virtual void UpdateLocation(FVectorVMExternalFunctionContext& Context);
-	virtual void UpdateRotation(FVectorVMExternalFunctionContext& Context);
-	virtual void SetPausedState(FVectorVMExternalFunctionContext& Context);
+	NIAGARA_API virtual void PlayOneShotAudio(FVectorVMExternalFunctionContext& Context);
+	NIAGARA_API virtual void PlayPersistentAudio(FVectorVMExternalFunctionContext& Context);
+	NIAGARA_API virtual void SetParameterBool(FVectorVMExternalFunctionContext& Context);
+	NIAGARA_API virtual void SetParameterInteger(FVectorVMExternalFunctionContext& Context);
+	NIAGARA_API virtual void SetParameterFloat(FVectorVMExternalFunctionContext& Context);
+	NIAGARA_API virtual void UpdateVolume(FVectorVMExternalFunctionContext& Context);
+	NIAGARA_API virtual void UpdatePitch(FVectorVMExternalFunctionContext& Context);
+	NIAGARA_API virtual void UpdateLocation(FVectorVMExternalFunctionContext& Context);
+	NIAGARA_API virtual void UpdateRotation(FVectorVMExternalFunctionContext& Context);
+	NIAGARA_API virtual void SetPausedState(FVectorVMExternalFunctionContext& Context);
 
 protected:
-	virtual bool CopyToInternal(UNiagaraDataInterface* Destination) const override;
+	NIAGARA_API virtual bool CopyToInternal(UNiagaraDataInterface* Destination) const override;
 	
 private:
-	static const FName PlayAudioName;
-	static const FName PlayPersistentAudioName;
-	static const FName SetPersistentAudioVolumeName;
-	static const FName SetPersistentAudioPitchName;
-	static const FName SetPersistentAudioLocationName;
-	static const FName SetPersistentAudioRotationName;
-	static const FName SetPersistentAudioBoolParamName;
-	static const FName SetPersistentAudioIntegerParamName;
-	static const FName SetPersistentAudioFloatParamName;
-	static const FName PausePersistentAudioName;
+	static NIAGARA_API const FName PlayAudioName;
+	static NIAGARA_API const FName PlayPersistentAudioName;
+	static NIAGARA_API const FName SetPersistentAudioVolumeName;
+	static NIAGARA_API const FName SetPersistentAudioPitchName;
+	static NIAGARA_API const FName SetPersistentAudioLocationName;
+	static NIAGARA_API const FName SetPersistentAudioRotationName;
+	static NIAGARA_API const FName SetPersistentAudioBoolParamName;
+	static NIAGARA_API const FName SetPersistentAudioIntegerParamName;
+	static NIAGARA_API const FName SetPersistentAudioFloatParamName;
+	static NIAGARA_API const FName PausePersistentAudioName;
 };
diff --git a/Engine/Plugins/FX/Niagara/Source/Niagara/Classes/NiagaraDataInterfaceAudioSpectrum.h b/Engine/Plugins/FX/Niagara/Source/Niagara/Classes/NiagaraDataInterfaceAudioSpectrum.h
--- a/Engine/Plugins/FX/Niagara/Source/Niagara/Classes/NiagaraDataInterfaceAudioSpectrum.h
+++ b/Engine/Plugins/FX/Niagara/Source/Niagara/Classes/NiagaraDataInterfaceAudioSpectrum.h
@@ -127,8 +127,8 @@ private:
 };
 
 /** Data Interface allowing sampling of recent audio spectrum. */
-UCLASS(EditInlineNew, Category = "Audio", meta = (DisplayName = "Audio Spectrum"))
-class NIAGARA_API UNiagaraDataInterfaceAudioSpectrum : public UNiagaraDataInterfaceAudioSubmix
+UCLASS(EditInlineNew, Category = "Audio", meta = (DisplayName = "Audio Spectrum"), MinimalAPI)
+class UNiagaraDataInterfaceAudioSpectrum : public UNiagaraDataInterfaceAudioSubmix
 {
 	GENERATED_UCLASS_BODY()
 
@@ -140,13 +140,13 @@ class NIAGARA_API UNiagaraDataInterfaceAudioSpectrum : public UNiagaraDataInterf
 
 public:
 	// VM function names
-	static const FName GetSpectrumFunctionName;
-	static const FName GetNumChannelsFunctionName;
+	static NIAGARA_API const FName GetSpectrumFunctionName;
+	static NIAGARA_API const FName GetNumChannelsFunctionName;
 
 	// Global variable prefixes
-	static const FString GetSpectrumName;
-	static const FString NumChannelsName;
-	static const FString ResolutionName;
+	static NIAGARA_API const FString GetSpectrumName;
+	static NIAGARA_API const FString NumChannelsName;
+	static NIAGARA_API const FString ResolutionName;
 	
 	/** The number of spectrum samples to pass to the GPU */
 	UPROPERTY(EditAnywhere, Category = "Spectrum", meta = (ClampMin = "16", ClampMax = "1024") )
@@ -165,11 +165,11 @@ public:
 	float NoiseFloorDb;
 
 	//VM function overrides:
-	void GetSpectrumValue(FVectorVMExternalFunctionContext& Context);
-	void GetNumChannels(FVectorVMExternalFunctionContext& Context);
+	NIAGARA_API void GetSpectrumValue(FVectorVMExternalFunctionContext& Context);
+	NIAGARA_API void GetNumChannels(FVectorVMExternalFunctionContext& Context);
 
-	virtual void GetFunctions(TArray<FNiagaraFunctionSignature>& OutFunctions) override;
-	virtual void GetVMExternalFunction(const FVMExternalFunctionBindingInfo& BindingInfo, void* InstanceData, FVMExternalFunction &OutFunc) override;
+	NIAGARA_API virtual void GetFunctions(TArray<FNiagaraFunctionSignature>& OutFunctions) override;
+	NIAGARA_API virtual void GetVMExternalFunction(const FVMExternalFunctionBindingInfo& BindingInfo, void* InstanceData, FVMExternalFunction &OutFunc) override;
 
 	virtual bool CanExecuteOnTarget(ENiagaraSimTarget Target) const override 
 	{
@@ -177,25 +177,25 @@ public:
 	}
 
 #if WITH_EDITORONLY_DATA
-	virtual bool AppendCompileHash(FNiagaraCompileHashVisitor* InVisitor) const override;
-	virtual bool GetFunctionHLSL(const FNiagaraDataInterfaceGPUParamInfo& ParamInfo, const FNiagaraDataInterfaceGeneratedFunction& FunctionInfo, int FunctionInstanceIndex, FString& OutHLSL) override;
-	virtual void GetParameterDefinitionHLSL(const FNiagaraDataInterfaceGPUParamInfo& ParamInfo, FString& OutHLSL) override;
+	NIAGARA_API virtual bool AppendCompileHash(FNiagaraCompileHashVisitor* InVisitor) const override;
+	NIAGARA_API virtual bool GetFunctionHLSL(const FNiagaraDataInterfaceGPUParamInfo& ParamInfo, const FNiagaraDataInterfaceGeneratedFunction& FunctionInfo, int FunctionInstanceIndex, FString& OutHLSL) override;
+	NIAGARA_API virtual void GetParameterDefinitionHLSL(const FNiagaraDataInterfaceGPUParamInfo& ParamInfo, FString& OutHLSL) override;
 #endif
-	virtual void BuildShaderParameters(FNiagaraShaderParametersBuilder& ShaderParametersBuilder) const override;
-	virtual void SetShaderParameters(const FNiagaraDataInterfaceSetShaderParametersContext& Context) const override;
+	NIAGARA_API virtual void BuildShaderParameters(FNiagaraShaderParametersBuilder& ShaderParametersBuilder) const override;
+	NIAGARA_API virtual void SetShaderParameters(const FNiagaraDataInterfaceSetShaderParametersContext& Context) const override;
 
-	virtual bool Equals(const UNiagaraDataInterface* Other) const override;
+	NIAGARA_API virtual bool Equals(const UNiagaraDataInterface* Other) const override;
 
 #if WITH_EDITOR
-	virtual void PostEditChangeProperty(struct FPropertyChangedEvent& PropertyChangedEvent) override;
+	NIAGARA_API virtual void PostEditChangeProperty(struct FPropertyChangedEvent& PropertyChangedEvent) override;
 #endif // WITH_EDITOR
 
 
-	virtual void PostInitProperties() override;
-	virtual void PostLoad() override;
+	NIAGARA_API virtual void PostInitProperties() override;
+	NIAGARA_API virtual void PostLoad() override;
 
 protected:
-	virtual bool CopyToInternal(UNiagaraDataInterface* Destination) const override;
+	NIAGARA_API virtual bool CopyToInternal(UNiagaraDataInterface* Destination) const override;
 private:
 
 };
