[Niagara]

diff --git a/Engine/Plugins/FX/Niagara/Source/Niagara/Classes/NiagaraDataInterfaceCamera.h b/Engine/Plugins/FX/Niagara/Source/Niagara/Classes/NiagaraDataInterfaceCamera.h
--- a/Engine/Plugins/FX/Niagara/Source/Niagara/Classes/NiagaraDataInterfaceCamera.h
+++ b/Engine/Plugins/FX/Niagara/Source/Niagara/Classes/NiagaraDataInterfaceCamera.h
@@ -23,8 +23,8 @@ struct FCameraDataInterface_InstanceData
 	TArray<FDistanceData> ParticlesSortedByDistance;	
 };
 
-UCLASS(EditInlineNew, Category = "Camera", CollapseCategories, meta = (DisplayName = "Camera Query"))
-class NIAGARA_API UNiagaraDataInterfaceCamera : public UNiagaraDataInterface
+UCLASS(EditInlineNew, Category = "Camera", CollapseCategories, meta = (DisplayName = "Camera Query"), MinimalAPI)
+class UNiagaraDataInterfaceCamera : public UNiagaraDataInterface
 {
 	GENERATED_UCLASS_BODY()
 
@@ -44,54 +44,54 @@ public:
 	bool bRequireCurrentFrameData = true;
 
 	//UObject Interface
-	virtual void PostInitProperties() override;
+	NIAGARA_API virtual void PostInitProperties() override;
 	//UObject Interface End
 
 	//UNiagaraDataInterface Interface
-	virtual void GetFunctions(TArray<FNiagaraFunctionSignature>& OutFunctions)override;
-	virtual void GetVMExternalFunction(const FVMExternalFunctionBindingInfo& BindingInfo, void* InstanceData, FVMExternalFunction &OutFunc) override;
-	virtual bool InitPerInstanceData(void* PerInstanceData, FNiagaraSystemInstance* SystemInstance) override;
+	NIAGARA_API virtual void GetFunctions(TArray<FNiagaraFunctionSignature>& OutFunctions)override;
+	NIAGARA_API virtual void GetVMExternalFunction(const FVMExternalFunctionBindingInfo& BindingInfo, void* InstanceData, FVMExternalFunction &OutFunc) override;
+	NIAGARA_API virtual bool InitPerInstanceData(void* PerInstanceData, FNiagaraSystemInstance* SystemInstance) override;
 	virtual int32 PerInstanceDataSize() const override { return sizeof(FCameraDataInterface_InstanceData); }
-	virtual bool PerInstanceTick(void* PerInstanceData, FNiagaraSystemInstance* SystemInstance, float DeltaSeconds) override;
+	NIAGARA_API virtual bool PerInstanceTick(void* PerInstanceData, FNiagaraSystemInstance* SystemInstance, float DeltaSeconds) override;
 #if WITH_EDITORONLY_DATA
-	virtual bool AppendCompileHash(FNiagaraCompileHashVisitor* InVisitor) const override;
-	virtual bool GetFunctionHLSL(const FNiagaraDataInterfaceGPUParamInfo& ParamInfo, const FNiagaraDataInterfaceGeneratedFunction& FunctionInfo, int FunctionInstanceIndex, FString& OutHLSL) override;
-	virtual bool UpgradeFunctionCall(FNiagaraFunctionSignature& FunctionSignature) override;
-	virtual void GetParameterDefinitionHLSL(const FNiagaraDataInterfaceGPUParamInfo& ParamInfo, FString& OutHLSL) override;
+	NIAGARA_API virtual bool AppendCompileHash(FNiagaraCompileHashVisitor* InVisitor) const override;
+	NIAGARA_API virtual bool GetFunctionHLSL(const FNiagaraDataInterfaceGPUParamInfo& ParamInfo, const FNiagaraDataInterfaceGeneratedFunction& FunctionInfo, int FunctionInstanceIndex, FString& OutHLSL) override;
+	NIAGARA_API virtual bool UpgradeFunctionCall(FNiagaraFunctionSignature& FunctionSignature) override;
+	NIAGARA_API virtual void GetParameterDefinitionHLSL(const FNiagaraDataInterfaceGPUParamInfo& ParamInfo, FString& OutHLSL) override;
 #endif
-	virtual void BuildShaderParameters(FNiagaraShaderParametersBuilder& ShaderParametersBuilder) const override;
-	virtual void SetShaderParameters(const FNiagaraDataInterfaceSetShaderParametersContext& Context) const override;
+	NIAGARA_API virtual void BuildShaderParameters(FNiagaraShaderParametersBuilder& ShaderParametersBuilder) const override;
+	NIAGARA_API virtual void SetShaderParameters(const FNiagaraDataInterfaceSetShaderParametersContext& Context) const override;
 	virtual bool CanExecuteOnTarget(ENiagaraSimTarget Target) const override { return true; }
 	virtual bool HasTickGroupPrereqs() const override { return true; }
-	virtual ETickingGroup CalculateTickGroup(const void* PerInstanceData) const override;
+	NIAGARA_API virtual ETickingGroup CalculateTickGroup(const void* PerInstanceData) const override;
 	virtual bool RequiresEarlyViewData() const override { return true; }
-	virtual bool Equals(const UNiagaraDataInterface* Other) const override;
+	NIAGARA_API virtual bool Equals(const UNiagaraDataInterface* Other) const override;
 	virtual bool HasPreSimulateTick() const override { return true; }
 #if WITH_EDITOR	
-	virtual void GetFeedback(UNiagaraSystem* Asset, UNiagaraComponent* Component, TArray<FNiagaraDataInterfaceError>& OutErrors,
+	NIAGARA_API virtual void GetFeedback(UNiagaraSystem* Asset, UNiagaraComponent* Component, TArray<FNiagaraDataInterfaceError>& OutErrors,
         TArray<FNiagaraDataInterfaceFeedback>& Warnings, TArray<FNiagaraDataInterfaceFeedback>& Info) override;
 #endif
 	//UNiagaraDataInterface Interface
 
-	void CalculateParticleDistances(FVectorVMExternalFunctionContext& Context);
-	void GetClosestParticles(FVectorVMExternalFunctionContext& Context);
-	void GetCameraFOV(FVectorVMExternalFunctionContext& Context);
-	void GetCameraProperties(FVectorVMExternalFunctionContext& Context);
+	NIAGARA_API void CalculateParticleDistances(FVectorVMExternalFunctionContext& Context);
+	NIAGARA_API void GetClosestParticles(FVectorVMExternalFunctionContext& Context);
+	NIAGARA_API void GetCameraFOV(FVectorVMExternalFunctionContext& Context);
+	NIAGARA_API void GetCameraProperties(FVectorVMExternalFunctionContext& Context);
 
 protected:
-	virtual bool CopyToInternal(UNiagaraDataInterface* Destination) const override;
+	NIAGARA_API virtual bool CopyToInternal(UNiagaraDataInterface* Destination) const override;
 private:
-	static const TCHAR* TemplateShaderFilePath;
-	static const FName CalculateDistancesName;
-	static const FName QueryClosestName;
-	static const FName GetViewPropertiesName;
-	static const FName GetClipSpaceTransformsName;
-	static const FName GetViewSpaceTransformsName;
-	static const FName GetCameraPropertiesName;
-	static const FName GetFieldOfViewName;
-	static const FName GetTAAJitterName;
-	static const FName ApplyPreViewTranslationToPosition;
-	static const FName RemovePreViewTranslationFromPosition;
+	static NIAGARA_API const TCHAR* TemplateShaderFilePath;
+	static NIAGARA_API const FName CalculateDistancesName;
+	static NIAGARA_API const FName QueryClosestName;
+	static NIAGARA_API const FName GetViewPropertiesName;
+	static NIAGARA_API const FName GetClipSpaceTransformsName;
+	static NIAGARA_API const FName GetViewSpaceTransformsName;
+	static NIAGARA_API const FName GetCameraPropertiesName;
+	static NIAGARA_API const FName GetFieldOfViewName;
+	static NIAGARA_API const FName GetTAAJitterName;
+	static NIAGARA_API const FName ApplyPreViewTranslationToPosition;
+	static NIAGARA_API const FName RemovePreViewTranslationFromPosition;
 };
 
 struct FNiagaraDataIntefaceProxyCameraQuery : public FNiagaraDataInterfaceProxy
diff --git a/Engine/Plugins/FX/Niagara/Source/Niagara/Classes/NiagaraDataInterfaceCollisionQuery.h b/Engine/Plugins/FX/Niagara/Source/Niagara/Classes/NiagaraDataInterfaceCollisionQuery.h
--- a/Engine/Plugins/FX/Niagara/Source/Niagara/Classes/NiagaraDataInterfaceCollisionQuery.h
+++ b/Engine/Plugins/FX/Niagara/Source/Niagara/Classes/NiagaraDataInterfaceCollisionQuery.h
@@ -22,8 +22,8 @@ struct CQDIPerInstanceData
 
 
 /** Data Interface that can be used to query collision related data, like geometry traces or sampling the depth buffer. */
-UCLASS(EditInlineNew, Category = "Collision", meta = (DisplayName = "Collision Query"))
-class NIAGARA_API UNiagaraDataInterfaceCollisionQuery : public UNiagaraDataInterface
+UCLASS(EditInlineNew, Category = "Collision", meta = (DisplayName = "Collision Query"), MinimalAPI)
+class UNiagaraDataInterfaceCollisionQuery : public UNiagaraDataInterface
 {
 	GENERATED_UCLASS_BODY()
 
@@ -33,43 +33,43 @@ class NIAGARA_API UNiagaraDataInterfaceCollisionQuery : public UNiagaraDataInter
 
 public:
 	//UObject Interface
-	virtual void PostInitProperties() override;
+	NIAGARA_API virtual void PostInitProperties() override;
 	//UObject Interface End
 
 	/** Initializes the per instance data for this interface. Returns false if there was some error and the simulation should be disabled. */
-	virtual bool InitPerInstanceData(void* PerInstanceData, FNiagaraSystemInstance* InSystemInstance) override;
+	NIAGARA_API virtual bool InitPerInstanceData(void* PerInstanceData, FNiagaraSystemInstance* InSystemInstance) override;
 	/** Destroys the per instence data for this interface. */
-	virtual void DestroyPerInstanceData(void* PerInstanceData, FNiagaraSystemInstance* InSystemInstance) override;
+	NIAGARA_API virtual void DestroyPerInstanceData(void* PerInstanceData, FNiagaraSystemInstance* InSystemInstance) override;
 
 	/** Ticks the per instance data for this interface, if it has any. */
-	virtual bool PerInstanceTick(void* PerInstanceData, FNiagaraSystemInstance* SystemInstance, float DeltaSeconds) override;
-	virtual bool PerInstanceTickPostSimulate(void* PerInstanceData, FNiagaraSystemInstance* SystemInstance, float DeltaSeconds) override;
+	NIAGARA_API virtual bool PerInstanceTick(void* PerInstanceData, FNiagaraSystemInstance* SystemInstance, float DeltaSeconds) override;
+	NIAGARA_API virtual bool PerInstanceTickPostSimulate(void* PerInstanceData, FNiagaraSystemInstance* SystemInstance, float DeltaSeconds) override;
 	virtual int32 PerInstanceDataSize() const override { return sizeof(CQDIPerInstanceData); }
 	
-	virtual void GetFunctions(TArray<FNiagaraFunctionSignature>& OutFunctions) override;
-	virtual void GetVMExternalFunction(const FVMExternalFunctionBindingInfo& BindingInfo, void* InstanceData, FVMExternalFunction &OutFunc) override;
-	virtual void GetAssetTagsForContext(const UObject* InAsset, FGuid AssetVersion, const TArray<const UNiagaraDataInterface*>& InProperties, TMap<FName, uint32>& NumericKeys, TMap<FName, FString>& StringKeys) const override;
+	NIAGARA_API virtual void GetFunctions(TArray<FNiagaraFunctionSignature>& OutFunctions) override;
+	NIAGARA_API virtual void GetVMExternalFunction(const FVMExternalFunctionBindingInfo& BindingInfo, void* InstanceData, FVMExternalFunction &OutFunc) override;
+	NIAGARA_API virtual void GetAssetTagsForContext(const UObject* InAsset, FGuid AssetVersion, const TArray<const UNiagaraDataInterface*>& InProperties, TMap<FName, uint32>& NumericKeys, TMap<FName, FString>& StringKeys) const override;
 
 	// VM functions
-	void PerformQuerySyncCPU(FVectorVMExternalFunctionContext& Context);
-	void PerformQueryAsyncCPU(FVectorVMExternalFunctionContext& Context);
+	NIAGARA_API void PerformQuerySyncCPU(FVectorVMExternalFunctionContext& Context);
+	NIAGARA_API void PerformQueryAsyncCPU(FVectorVMExternalFunctionContext& Context);
 
 	virtual bool CanExecuteOnTarget(ENiagaraSimTarget Target) const override { return true; }
 	virtual bool RequiresDistanceFieldData() const override { return true; }
 	virtual bool RequiresDepthBuffer() const override { return true; }
 
 #if WITH_EDITORONLY_DATA
-	virtual bool AppendCompileHash(FNiagaraCompileHashVisitor* InVisitor) const override;
-	virtual void GetCommonHLSL(FString& OutHLSL) override;
-	virtual void GetParameterDefinitionHLSL(const FNiagaraDataInterfaceGPUParamInfo& ParamInfo, FString& OutHLSL) override;
-	virtual bool GetFunctionHLSL(const FNiagaraDataInterfaceGPUParamInfo& ParamInfo, const FNiagaraDataInterfaceGeneratedFunction& FunctionInfo, int FunctionInstanceIndex, FString& OutHLSL) override;
-	virtual bool UpgradeFunctionCall(FNiagaraFunctionSignature& FunctionSignature) override;
+	NIAGARA_API virtual bool AppendCompileHash(FNiagaraCompileHashVisitor* InVisitor) const override;
+	NIAGARA_API virtual void GetCommonHLSL(FString& OutHLSL) override;
+	NIAGARA_API virtual void GetParameterDefinitionHLSL(const FNiagaraDataInterfaceGPUParamInfo& ParamInfo, FString& OutHLSL) override;
+	NIAGARA_API virtual bool GetFunctionHLSL(const FNiagaraDataInterfaceGPUParamInfo& ParamInfo, const FNiagaraDataInterfaceGeneratedFunction& FunctionInfo, int FunctionInstanceIndex, FString& OutHLSL) override;
+	NIAGARA_API virtual bool UpgradeFunctionCall(FNiagaraFunctionSignature& FunctionSignature) override;
 #endif
-	virtual void BuildShaderParameters(FNiagaraShaderParametersBuilder& ShaderParametersBuilder) const override;
-	virtual void SetShaderParameters(const FNiagaraDataInterfaceSetShaderParametersContext& Context) const override;
+	NIAGARA_API virtual void BuildShaderParameters(FNiagaraShaderParametersBuilder& ShaderParametersBuilder) const override;
+	NIAGARA_API virtual void SetShaderParameters(const FNiagaraDataInterfaceSetShaderParametersContext& Context) const override;
 
 #if WITH_EDITOR	
-	virtual void ValidateFunction(const FNiagaraFunctionSignature& Function, TArray<FText>& OutValidationErrors) override;
+	NIAGARA_API virtual void ValidateFunction(const FNiagaraFunctionSignature& Function, TArray<FText>& OutValidationErrors) override;
 #endif
 	
 	virtual bool HasPreSimulateTick() const override{ return true; }
@@ -77,6 +77,6 @@ public:
 	virtual bool PostSimulateCanOverlapFrames() const { return false; }
 
 private:
-	static FCriticalSection CriticalSection;
+	static NIAGARA_API FCriticalSection CriticalSection;
 	UEnum* TraceChannelEnum;
 };
diff --git a/Engine/Plugins/FX/Niagara/Source/Niagara/Classes/NiagaraDataInterfaceColorCurve.h b/Engine/Plugins/FX/Niagara/Source/Niagara/Classes/NiagaraDataInterfaceColorCurve.h
--- a/Engine/Plugins/FX/Niagara/Source/Niagara/Classes/NiagaraDataInterfaceColorCurve.h
+++ b/Engine/Plugins/FX/Niagara/Source/Niagara/Classes/NiagaraDataInterfaceColorCurve.h
@@ -17,8 +17,8 @@ class FNiagaraSystemInstance;
 
 
 /** Data Interface allowing sampling of color curves. */
-UCLASS(EditInlineNew, Category="Curves", CollapseCategories, meta = (DisplayName = "Curve for Colors"))
-class NIAGARA_API UNiagaraDataInterfaceColorCurve : public UNiagaraDataInterfaceCurveBase
+UCLASS(EditInlineNew, Category="Curves", CollapseCategories, meta = (DisplayName = "Curve for Colors"), MinimalAPI)
+class UNiagaraDataInterfaceColorCurve : public UNiagaraDataInterfaceCurveBase
 {
 	GENERATED_UCLASS_BODY()
 public:
@@ -49,8 +49,8 @@ public:
 #endif
 
 	//UObject Interface
-	virtual void PostInitProperties() override;
-	virtual void Serialize(FArchive& Ar) override;
+	NIAGARA_API virtual void PostInitProperties() override;
+	NIAGARA_API virtual void Serialize(FArchive& Ar) override;
 	//UObject Interface End
 
 	enum
@@ -58,23 +58,23 @@ public:
 		CurveLUTNumElems = 4,
 	};
 
-	virtual void UpdateTimeRanges() override;
-	virtual TArray<float> BuildLUT(int32 NumEntries) const override;
+	NIAGARA_API virtual void UpdateTimeRanges() override;
+	NIAGARA_API virtual TArray<float> BuildLUT(int32 NumEntries) const override;
 
-	virtual void GetFunctions(TArray<FNiagaraFunctionSignature>& OutFunctions)override;
-	virtual void GetVMExternalFunction(const FVMExternalFunctionBindingInfo& BindingInfo, void* InstanceData, FVMExternalFunction &OutFunc) override;
+	NIAGARA_API virtual void GetFunctions(TArray<FNiagaraFunctionSignature>& OutFunctions)override;
+	NIAGARA_API virtual void GetVMExternalFunction(const FVMExternalFunctionBindingInfo& BindingInfo, void* InstanceData, FVMExternalFunction &OutFunc) override;
 
 	template<typename UseLUT>
 	void SampleCurve(FVectorVMExternalFunctionContext& Context);
 
-	virtual bool Equals(const UNiagaraDataInterface* Other) const override;
+	NIAGARA_API virtual bool Equals(const UNiagaraDataInterface* Other) const override;
 
 	//~ UNiagaraDataInterfaceCurveBase interface
-	virtual void GetCurveData(TArray<FCurveData>& OutCurveData) override;
+	NIAGARA_API virtual void GetCurveData(TArray<FCurveData>& OutCurveData) override;
 
 #if WITH_EDITORONLY_DATA
-	virtual void SyncCurvesToAsset() override;
-	virtual bool GetFunctionHLSL(const FNiagaraDataInterfaceGPUParamInfo& ParamInfo, const FNiagaraDataInterfaceGeneratedFunction& FunctionInfo, int FunctionInstanceIndex, FString& OutHLSL) override;
+	NIAGARA_API virtual void SyncCurvesToAsset() override;
+	NIAGARA_API virtual bool GetFunctionHLSL(const FNiagaraDataInterfaceGPUParamInfo& ParamInfo, const FNiagaraDataInterfaceGeneratedFunction& FunctionInfo, int FunctionInstanceIndex, FString& OutHLSL) override;
 #endif
 
 	virtual int32 GetCurveNumElems() const override { return CurveLUTNumElems; }
@@ -83,11 +83,11 @@ public:
 #endif
 
 protected:
-	virtual bool CopyToInternal(UNiagaraDataInterface* Destination) const override;
+	NIAGARA_API virtual bool CopyToInternal(UNiagaraDataInterface* Destination) const override;
 
 private:
 	template<typename UseLUT>
 	FORCEINLINE_DEBUGGABLE FLinearColor SampleCurveInternal(float X);
 
-	static const FName SampleCurveName;
+	static NIAGARA_API const FName SampleCurveName;
 };
diff --git a/Engine/Plugins/FX/Niagara/Source/Niagara/Classes/NiagaraDataInterfaceCubeTexture.h b/Engine/Plugins/FX/Niagara/Source/Niagara/Classes/NiagaraDataInterfaceCubeTexture.h
--- a/Engine/Plugins/FX/Niagara/Source/Niagara/Classes/NiagaraDataInterfaceCubeTexture.h
+++ b/Engine/Plugins/FX/Niagara/Source/Niagara/Classes/NiagaraDataInterfaceCubeTexture.h
@@ -10,8 +10,8 @@
 class UTextureCube;
 
 /** Data Interface allowing sampling of a texture */
-UCLASS(EditInlineNew, Category = "Texture", CollapseCategories, meta = (DisplayName = "Cube Texture Sample"))
-class NIAGARA_API UNiagaraDataInterfaceCubeTexture : public UNiagaraDataInterface
+UCLASS(EditInlineNew, Category = "Texture", CollapseCategories, meta = (DisplayName = "Cube Texture Sample"), MinimalAPI)
+class UNiagaraDataInterfaceCubeTexture : public UNiagaraDataInterface
 {
 	GENERATED_UCLASS_BODY()
 
@@ -29,43 +29,43 @@ public:
 	FNiagaraUserParameterBinding TextureUserParameter;
 
 	//UObject Interface
-	virtual void PostInitProperties()override;
+	NIAGARA_API virtual void PostInitProperties()override;
 	//UObject Interface End
 
 	//UNiagaraDataInterface Interface
-	virtual void GetFunctions(TArray<FNiagaraFunctionSignature>& OutFunctions)override;
-	virtual void GetVMExternalFunction(const FVMExternalFunctionBindingInfo& BindingInfo, void* InstanceData, FVMExternalFunction &OutFunc) override;
+	NIAGARA_API virtual void GetFunctions(TArray<FNiagaraFunctionSignature>& OutFunctions)override;
+	NIAGARA_API virtual void GetVMExternalFunction(const FVMExternalFunctionBindingInfo& BindingInfo, void* InstanceData, FVMExternalFunction &OutFunc) override;
 	virtual bool CanExecuteOnTarget(ENiagaraSimTarget Target)const override { return Target == ENiagaraSimTarget::GPUComputeSim; }
 
-	virtual int32 PerInstanceDataSize() const override;
-	virtual bool InitPerInstanceData(void* PerInstanceData, FNiagaraSystemInstance* SystemInstance) override;
-	virtual void DestroyPerInstanceData(void* PerInstanceData, FNiagaraSystemInstance* SystemInstance) override;
+	NIAGARA_API virtual int32 PerInstanceDataSize() const override;
+	NIAGARA_API virtual bool InitPerInstanceData(void* PerInstanceData, FNiagaraSystemInstance* SystemInstance) override;
+	NIAGARA_API virtual void DestroyPerInstanceData(void* PerInstanceData, FNiagaraSystemInstance* SystemInstance) override;
 
 	virtual bool HasPreSimulateTick() const override { return true; }
-	virtual bool PerInstanceTick(void* PerInstanceData, FNiagaraSystemInstance* SystemInstance, float DeltaSeconds) override;
+	NIAGARA_API virtual bool PerInstanceTick(void* PerInstanceData, FNiagaraSystemInstance* SystemInstance, float DeltaSeconds) override;
 	//UNiagaraDataInterface Interface End
 
-	void SampleCubeTexture(FVectorVMExternalFunctionContext& Context);
-	void GetTextureDimensions(FVectorVMExternalFunctionContext& Context);
+	NIAGARA_API void SampleCubeTexture(FVectorVMExternalFunctionContext& Context);
+	NIAGARA_API void GetTextureDimensions(FVectorVMExternalFunctionContext& Context);
 
-	virtual bool Equals(const UNiagaraDataInterface* Other) const override;
+	NIAGARA_API virtual bool Equals(const UNiagaraDataInterface* Other) const override;
 
 	// GPU sim functionality
 #if WITH_EDITORONLY_DATA
-	virtual bool AppendCompileHash(FNiagaraCompileHashVisitor* InVisitor) const override;
-	virtual void GetParameterDefinitionHLSL(const FNiagaraDataInterfaceGPUParamInfo& ParamInfo, FString& OutHLSL) override;
-	virtual bool GetFunctionHLSL(const FNiagaraDataInterfaceGPUParamInfo& ParamInfo, const FNiagaraDataInterfaceGeneratedFunction& FunctionInfo, int FunctionInstanceIndex, FString& OutHLSL) override;
+	NIAGARA_API virtual bool AppendCompileHash(FNiagaraCompileHashVisitor* InVisitor) const override;
+	NIAGARA_API virtual void GetParameterDefinitionHLSL(const FNiagaraDataInterfaceGPUParamInfo& ParamInfo, FString& OutHLSL) override;
+	NIAGARA_API virtual bool GetFunctionHLSL(const FNiagaraDataInterfaceGPUParamInfo& ParamInfo, const FNiagaraDataInterfaceGeneratedFunction& FunctionInfo, int FunctionInstanceIndex, FString& OutHLSL) override;
 #endif
-	virtual void BuildShaderParameters(FNiagaraShaderParametersBuilder& ShaderParametersBuilder) const override;
-	virtual void SetShaderParameters(const FNiagaraDataInterfaceSetShaderParametersContext& Context) const override;
+	NIAGARA_API virtual void BuildShaderParameters(FNiagaraShaderParametersBuilder& ShaderParametersBuilder) const override;
+	NIAGARA_API virtual void SetShaderParameters(const FNiagaraDataInterfaceSetShaderParametersContext& Context) const override;
 
-	void SetTexture(UTexture* InTexture);
+	NIAGARA_API void SetTexture(UTexture* InTexture);
 
 protected:
-	virtual bool CopyToInternal(UNiagaraDataInterface* Destination) const override;
+	NIAGARA_API virtual bool CopyToInternal(UNiagaraDataInterface* Destination) const override;
 
 protected:
-	static const TCHAR* TemplateShaderFilePath;
-	static const FName SampleCubeTextureName;
-	static const FName TextureDimsName;
+	static NIAGARA_API const TCHAR* TemplateShaderFilePath;
+	static NIAGARA_API const FName SampleCubeTextureName;
+	static NIAGARA_API const FName TextureDimsName;
 };
diff --git a/Engine/Plugins/FX/Niagara/Source/Niagara/Classes/NiagaraDataInterfaceCurlNoise.h b/Engine/Plugins/FX/Niagara/Source/Niagara/Classes/NiagaraDataInterfaceCurlNoise.h
--- a/Engine/Plugins/FX/Niagara/Source/Niagara/Classes/NiagaraDataInterfaceCurlNoise.h
+++ b/Engine/Plugins/FX/Niagara/Source/Niagara/Classes/NiagaraDataInterfaceCurlNoise.h
@@ -8,8 +8,8 @@
 #include "NiagaraDataInterfaceCurlNoise.generated.h"
 
 /** Data Interface allowing sampling of curl noise LUT. */
-UCLASS(EditInlineNew, Category = "Curl Noise LUT", CollapseCategories, meta = (DisplayName = "Curl Noise"))
-class NIAGARA_API UNiagaraDataInterfaceCurlNoise : public UNiagaraDataInterface
+UCLASS(EditInlineNew, Category = "Curl Noise LUT", CollapseCategories, meta = (DisplayName = "Curl Noise"), MinimalAPI)
+class UNiagaraDataInterfaceCurlNoise : public UNiagaraDataInterface
 {
 	GENERATED_UCLASS_BODY()
 
@@ -25,36 +25,36 @@ public:
 	FVector3f OffsetFromSeed;
 
 	//UObject Interface
-	virtual void PostInitProperties()override;
-	virtual void PostLoad() override;
+	NIAGARA_API virtual void PostInitProperties()override;
+	NIAGARA_API virtual void PostLoad() override;
 #if WITH_EDITOR
-	virtual void PreEditChange(FProperty* PropertyAboutToChange) override;
-	virtual void PostEditChangeProperty(struct FPropertyChangedEvent& PropertyChangedEvent) override;
+	NIAGARA_API virtual void PreEditChange(FProperty* PropertyAboutToChange) override;
+	NIAGARA_API virtual void PostEditChangeProperty(struct FPropertyChangedEvent& PropertyChangedEvent) override;
 #endif
 	//UObject Interface End
 
 	//UNiagaraDataInterface Interface
-	virtual void GetFunctions(TArray<FNiagaraFunctionSignature>& OutFunctions) override;
-	virtual void GetVMExternalFunction(const FVMExternalFunctionBindingInfo& BindingInfo, void* InstanceData, FVMExternalFunction &OutFunc) override;
+	NIAGARA_API virtual void GetFunctions(TArray<FNiagaraFunctionSignature>& OutFunctions) override;
+	NIAGARA_API virtual void GetVMExternalFunction(const FVMExternalFunctionBindingInfo& BindingInfo, void* InstanceData, FVMExternalFunction &OutFunc) override;
 	virtual bool CanExecuteOnTarget(ENiagaraSimTarget Target)const override { return true; }
-	virtual bool Equals(const UNiagaraDataInterface* Other) const override;
+	NIAGARA_API virtual bool Equals(const UNiagaraDataInterface* Other) const override;
 	//UNiagaraDataInterface Interface End
 
-	void SampleNoiseField(FVectorVMExternalFunctionContext& Context);
+	NIAGARA_API void SampleNoiseField(FVectorVMExternalFunctionContext& Context);
 
 	// GPU sim functionality
 #if WITH_EDITORONLY_DATA
-	virtual bool AppendCompileHash(FNiagaraCompileHashVisitor* InVisitor) const override;
-	virtual void GetParameterDefinitionHLSL(const FNiagaraDataInterfaceGPUParamInfo& ParamInfo, FString& OutHLSL) override;
-	virtual bool GetFunctionHLSL(const FNiagaraDataInterfaceGPUParamInfo& ParamInfo, const FNiagaraDataInterfaceGeneratedFunction& FunctionInfo, int FunctionInstanceIndex, FString& OutHLSL) override;
+	NIAGARA_API virtual bool AppendCompileHash(FNiagaraCompileHashVisitor* InVisitor) const override;
+	NIAGARA_API virtual void GetParameterDefinitionHLSL(const FNiagaraDataInterfaceGPUParamInfo& ParamInfo, FString& OutHLSL) override;
+	NIAGARA_API virtual bool GetFunctionHLSL(const FNiagaraDataInterfaceGPUParamInfo& ParamInfo, const FNiagaraDataInterfaceGeneratedFunction& FunctionInfo, int FunctionInstanceIndex, FString& OutHLSL) override;
 #endif
-	virtual void BuildShaderParameters(FNiagaraShaderParametersBuilder& ShaderParametersBuilder) const override;
-	virtual void SetShaderParameters(const FNiagaraDataInterfaceSetShaderParametersContext& Context) const override;
+	NIAGARA_API virtual void BuildShaderParameters(FNiagaraShaderParametersBuilder& ShaderParametersBuilder) const override;
+	NIAGARA_API virtual void SetShaderParameters(const FNiagaraDataInterfaceSetShaderParametersContext& Context) const override;
 
 protected:
-	virtual bool CopyToInternal(UNiagaraDataInterface* Destination) const override;
+	NIAGARA_API virtual bool CopyToInternal(UNiagaraDataInterface* Destination) const override;
 
-	virtual void PushToRenderThreadImpl() override;
+	NIAGARA_API virtual void PushToRenderThreadImpl() override;
 };
 
 struct FNiagaraDataInterfaceProxyCurlNoise : public FNiagaraDataInterfaceProxy
diff --git a/Engine/Plugins/FX/Niagara/Source/Niagara/Classes/NiagaraDataInterfaceCurve.h b/Engine/Plugins/FX/Niagara/Source/Niagara/Classes/NiagaraDataInterfaceCurve.h
--- a/Engine/Plugins/FX/Niagara/Source/Niagara/Classes/NiagaraDataInterfaceCurve.h
+++ b/Engine/Plugins/FX/Niagara/Source/Niagara/Classes/NiagaraDataInterfaceCurve.h
@@ -11,8 +11,8 @@
 
 
 /** Data Interface allowing sampling of float curves. */
-UCLASS(EditInlineNew, Category = "Curves", CollapseCategories, meta = (DisplayName = "Curve for Floats"))
-class NIAGARA_API UNiagaraDataInterfaceCurve : public UNiagaraDataInterfaceCurveBase
+UCLASS(EditInlineNew, Category = "Curves", CollapseCategories, meta = (DisplayName = "Curve for Floats"), MinimalAPI)
+class UNiagaraDataInterfaceCurve : public UNiagaraDataInterfaceCurveBase
 {
 	GENERATED_UCLASS_BODY()
 
@@ -33,27 +33,27 @@ public:
 	};
 
 	//UObject Interface
-	virtual void PostInitProperties() override;
-	virtual void Serialize(FArchive& Ar) override;
+	NIAGARA_API virtual void PostInitProperties() override;
+	NIAGARA_API virtual void Serialize(FArchive& Ar) override;
 	//UObject Interface End
 
-	virtual void UpdateTimeRanges() override;
-	virtual TArray<float> BuildLUT(int32 NumEntries) const override;
+	NIAGARA_API virtual void UpdateTimeRanges() override;
+	NIAGARA_API virtual TArray<float> BuildLUT(int32 NumEntries) const override;
 
-	virtual void GetFunctions(TArray<FNiagaraFunctionSignature>& OutFunctions)override;
-	virtual void GetVMExternalFunction(const FVMExternalFunctionBindingInfo& BindingInfo, void* InstanceData, FVMExternalFunction &OutFunc) override;
+	NIAGARA_API virtual void GetFunctions(TArray<FNiagaraFunctionSignature>& OutFunctions)override;
+	NIAGARA_API virtual void GetVMExternalFunction(const FVMExternalFunctionBindingInfo& BindingInfo, void* InstanceData, FVMExternalFunction &OutFunc) override;
 
 	template<typename UseLUT>
 	void SampleCurve(FVectorVMExternalFunctionContext& Context);
 	
-	virtual bool Equals(const UNiagaraDataInterface* Other) const override;
+	NIAGARA_API virtual bool Equals(const UNiagaraDataInterface* Other) const override;
 
 	//~ UNiagaraDataInterfaceCurveBase interface
-	virtual void GetCurveData(TArray<FCurveData>& OutCurveData) override;
+	NIAGARA_API virtual void GetCurveData(TArray<FCurveData>& OutCurveData) override;
 
 #if WITH_EDITORONLY_DATA
-	virtual void SyncCurvesToAsset() override;
-	virtual bool GetFunctionHLSL(const FNiagaraDataInterfaceGPUParamInfo& ParamInfo, const FNiagaraDataInterfaceGeneratedFunction& FunctionInfo, int FunctionInstanceIndex, FString& OutHLSL) override;
+	NIAGARA_API virtual void SyncCurvesToAsset() override;
+	NIAGARA_API virtual bool GetFunctionHLSL(const FNiagaraDataInterfaceGPUParamInfo& ParamInfo, const FNiagaraDataInterfaceGeneratedFunction& FunctionInfo, int FunctionInstanceIndex, FString& OutHLSL) override;
 #endif
 
 	virtual int32 GetCurveNumElems() const override { return CurveLUTNumElems; }
@@ -62,11 +62,11 @@ public:
 #endif
 
 protected:
-	virtual bool CopyToInternal(UNiagaraDataInterface* Destination) const override;
+	NIAGARA_API virtual bool CopyToInternal(UNiagaraDataInterface* Destination) const override;
 
 private:
 	template<typename UseLUT>
 	FORCEINLINE_DEBUGGABLE float SampleCurveInternal(float X);
 
-	static const FName SampleCurveName;
+	static NIAGARA_API const FName SampleCurveName;
 };
diff --git a/Engine/Plugins/FX/Niagara/Source/Niagara/Classes/NiagaraDataInterfaceCurveBase.h b/Engine/Plugins/FX/Niagara/Source/Niagara/Classes/NiagaraDataInterfaceCurveBase.h
--- a/Engine/Plugins/FX/Niagara/Source/Niagara/Classes/NiagaraDataInterfaceCurveBase.h
+++ b/Engine/Plugins/FX/Niagara/Source/Niagara/Classes/NiagaraDataInterfaceCurveBase.h
@@ -10,8 +10,8 @@
 #include "NiagaraDataInterfaceCurveBase.generated.h"
 
 /** Base class for curve data interfaces which facilitates handling the curve data in a standardized way. */
-UCLASS(EditInlineNew, abstract)
-class NIAGARA_API UNiagaraDataInterfaceCurveBase : public UNiagaraDataInterface
+UCLASS(EditInlineNew, abstract, MinimalAPI)
+class UNiagaraDataInterfaceCurveBase : public UNiagaraDataInterface
 {
 public:
 	static constexpr float DefaultOptimizeThreshold = 0.01f;
@@ -87,12 +87,12 @@ public:
 
 #if WITH_EDITOR	
 	/** Refreshes and returns the errors detected with the corresponding data, if any.*/
-	virtual TArray<FNiagaraDataInterfaceError> GetErrors() override;
+	NIAGARA_API virtual TArray<FNiagaraDataInterfaceError> GetErrors() override;
 #endif
 
 public:
-	UNiagaraDataInterfaceCurveBase();
-	UNiagaraDataInterfaceCurveBase(FObjectInitializer const& ObjectInitializer);
+	NIAGARA_API UNiagaraDataInterfaceCurveBase();
+	NIAGARA_API UNiagaraDataInterfaceCurveBase(FObjectInitializer const& ObjectInitializer);
 
 	enum
 	{
@@ -117,47 +117,47 @@ public:
 	};
 
 	//UObject Interface
-	virtual void PostLoad() override;
-	virtual void Serialize(FArchive& Ar) override;
+	NIAGARA_API virtual void PostLoad() override;
+	NIAGARA_API virtual void Serialize(FArchive& Ar) override;
 #if WITH_EDITOR
-	virtual void PostEditChangeProperty(struct FPropertyChangedEvent& PropertyChangedEvent) override;
+	NIAGARA_API virtual void PostEditChangeProperty(struct FPropertyChangedEvent& PropertyChangedEvent) override;
 #endif
 	//UObject Interface End
 
 	/** Gets information for all of the curves owned by this curve data interface. */
 	virtual void GetCurveData(TArray<FCurveData>& OutCurveData) { }
 
-	virtual void CacheStaticBuffers(struct FNiagaraSystemStaticBuffers& StaticBuffers, const FNiagaraVariable& ResolvedVariable, bool bUsedByCPU, bool bUsedByGPU) override;
+	NIAGARA_API virtual void CacheStaticBuffers(struct FNiagaraSystemStaticBuffers& StaticBuffers, const FNiagaraVariable& ResolvedVariable, bool bUsedByCPU, bool bUsedByGPU) override;
 #if WITH_EDITORONLY_DATA
-	virtual bool AppendCompileHash(FNiagaraCompileHashVisitor* InVisitor) const override;
-	virtual void GetParameterDefinitionHLSL(const FNiagaraDataInterfaceGPUParamInfo& ParamInfo, FString& OutHLSL) override;
+	NIAGARA_API virtual bool AppendCompileHash(FNiagaraCompileHashVisitor* InVisitor) const override;
+	NIAGARA_API virtual void GetParameterDefinitionHLSL(const FNiagaraDataInterfaceGPUParamInfo& ParamInfo, FString& OutHLSL) override;
 #endif
-	virtual void BuildShaderParameters(FNiagaraShaderParametersBuilder& ShaderParametersBuilder) const override;
-	virtual void SetShaderParameters(const FNiagaraDataInterfaceSetShaderParametersContext& Context) const override;
-	virtual void PostCompile() override;
+	NIAGARA_API virtual void BuildShaderParameters(FNiagaraShaderParametersBuilder& ShaderParametersBuilder) const override;
+	NIAGARA_API virtual void SetShaderParameters(const FNiagaraDataInterfaceSetShaderParametersContext& Context) const override;
+	NIAGARA_API virtual void PostCompile() override;
 
-	void SetDefaultLUT();
+	NIAGARA_API void SetDefaultLUT();
 #if WITH_EDITORONLY_DATA
-	void UpdateLUT(bool bFromSerialize = false);
-	void OptimizeLUT();
-	void UpdateExposedTexture();
+	NIAGARA_API void UpdateLUT(bool bFromSerialize = false);
+	NIAGARA_API void OptimizeLUT();
+	NIAGARA_API void UpdateExposedTexture();
 	virtual void SyncCurvesToAsset() {};
 #endif
 
 	//UNiagaraDataInterface interface
-	virtual bool Equals(const UNiagaraDataInterface* Other) const override;
+	NIAGARA_API virtual bool Equals(const UNiagaraDataInterface* Other) const override;
 	virtual bool CanExecuteOnTarget(ENiagaraSimTarget Target) const override { return true; }
 
 	virtual bool CanExposeVariables() const override { return true; }
-	virtual void GetExposedVariables(TArray<FNiagaraVariableBase>& OutVariables) const override;
-	virtual bool GetExposedVariableValue(const FNiagaraVariableBase& InVariable, void* InPerInstanceData, FNiagaraSystemInstance* InSystemInstance, void* OutData) const override;
+	NIAGARA_API virtual void GetExposedVariables(TArray<FNiagaraVariableBase>& OutVariables) const override;
+	NIAGARA_API virtual bool GetExposedVariableValue(const FNiagaraVariableBase& InVariable, void* InPerInstanceData, FNiagaraSystemInstance* InSystemInstance, void* OutData) const override;
 
-	virtual int32 GetCurveNumElems() const PURE_VIRTUAL(UNiagaraDataInterfaceCurveBase::GetCurveNumElems(), return 0;)
+	NIAGARA_API virtual int32 GetCurveNumElems() const PURE_VIRTUAL(UNiagaraDataInterfaceCurveBase::GetCurveNumElems(), return 0;)
 #if WITH_EDITORONLY_DATA
-	virtual FName GetCurveSampleFunctionName() const PURE_VIRTUAL(UNiagaraDataInterfaceCurveBase::GetCurveSampleFunctionName(), return NAME_None;)
+	NIAGARA_API virtual FName GetCurveSampleFunctionName() const PURE_VIRTUAL(UNiagaraDataInterfaceCurveBase::GetCurveSampleFunctionName(), return NAME_None;)
 #endif
 
-	virtual void UpdateTimeRanges() PURE_VIRTUAL(UNiagaraDataInterfaceCurveBase::UpdateTimeRanges(), )
+	NIAGARA_API virtual void UpdateTimeRanges() PURE_VIRTUAL(UNiagaraDataInterfaceCurveBase::UpdateTimeRanges(), )
 	virtual TArray<float> BuildLUT(int32 NumEntries) const PURE_VIRTUAL(UNiagaraDataInterfaceCurveBase::BuildLUT(), return TArray<float>(); )
 
 	FORCEINLINE float GetMinTime()const { return LUTMinTime; }
@@ -165,9 +165,9 @@ public:
 	FORCEINLINE float GetInvTimeRange()const { return LUTInvTimeRange; }
 
 protected:
-	virtual void PushToRenderThreadImpl() override;
-	virtual bool CopyToInternal(UNiagaraDataInterface* Destination) const override;
-	virtual bool CompareLUTS(const TArray<float>& OtherLUT) const;
+	NIAGARA_API virtual void PushToRenderThreadImpl() override;
+	NIAGARA_API virtual bool CopyToInternal(UNiagaraDataInterface* Destination) const override;
+	NIAGARA_API virtual bool CompareLUTS(const TArray<float>& OtherLUT) const;
 	//UNiagaraDataInterface interface END
 };
 
diff --git a/Engine/Plugins/FX/Niagara/Source/Niagara/Classes/NiagaraDataInterfaceExport.h b/Engine/Plugins/FX/Niagara/Source/Niagara/Classes/NiagaraDataInterfaceExport.h
--- a/Engine/Plugins/FX/Niagara/Source/Niagara/Classes/NiagaraDataInterfaceExport.h
+++ b/Engine/Plugins/FX/Niagara/Source/Niagara/Classes/NiagaraDataInterfaceExport.h
@@ -49,8 +49,8 @@ enum class ENDIExport_GPUAllocationMode : uint8
 
 /** This Data Interface can be used to gather particles at execution time and call either a 
 C++ or blueprint object with the gathered particle data each tick. */
-UCLASS(EditInlineNew, Category = "Export", meta = (DisplayName = "Export particle data"))
-class NIAGARA_API UNiagaraDataInterfaceExport : public UNiagaraDataInterface
+UCLASS(EditInlineNew, Category = "Export", meta = (DisplayName = "Export particle data"), MinimalAPI)
+class UNiagaraDataInterfaceExport : public UNiagaraDataInterface
 {
 	GENERATED_UCLASS_BODY()
 
@@ -71,25 +71,25 @@ public:
 	float GPUAllocationPerParticleSize = 1.0f;
 
 	//UObject Interface
-	virtual void PostInitProperties() override;
+	NIAGARA_API virtual void PostInitProperties() override;
 	//UObject Interface End
 
 	//UNiagaraDataInterface Interface
-	virtual void GetFunctions(TArray<FNiagaraFunctionSignature>& OutFunctions) override;
-	virtual void GetVMExternalFunction(const FVMExternalFunctionBindingInfo& BindingInfo, void* InstanceData, FVMExternalFunction &OutFunc) override;
-	virtual bool InitPerInstanceData(void* PerInstanceData, FNiagaraSystemInstance* SystemInstance) override;
-	virtual void DestroyPerInstanceData(void* PerInstanceData, FNiagaraSystemInstance* SystemInstance) override;
-	virtual int32 PerInstanceDataSize() const override;
-	virtual bool PerInstanceTick(void* PerInstanceData, FNiagaraSystemInstance* SystemInstance, float DeltaSeconds) override;
-	virtual bool PerInstanceTickPostSimulate(void* PerInstanceData, FNiagaraSystemInstance* SystemInstance, float DeltaSeconds) override;
-	virtual bool Equals(const UNiagaraDataInterface* Other) const override;
+	NIAGARA_API virtual void GetFunctions(TArray<FNiagaraFunctionSignature>& OutFunctions) override;
+	NIAGARA_API virtual void GetVMExternalFunction(const FVMExternalFunctionBindingInfo& BindingInfo, void* InstanceData, FVMExternalFunction &OutFunc) override;
+	NIAGARA_API virtual bool InitPerInstanceData(void* PerInstanceData, FNiagaraSystemInstance* SystemInstance) override;
+	NIAGARA_API virtual void DestroyPerInstanceData(void* PerInstanceData, FNiagaraSystemInstance* SystemInstance) override;
+	NIAGARA_API virtual int32 PerInstanceDataSize() const override;
+	NIAGARA_API virtual bool PerInstanceTick(void* PerInstanceData, FNiagaraSystemInstance* SystemInstance, float DeltaSeconds) override;
+	NIAGARA_API virtual bool PerInstanceTickPostSimulate(void* PerInstanceData, FNiagaraSystemInstance* SystemInstance, float DeltaSeconds) override;
+	NIAGARA_API virtual bool Equals(const UNiagaraDataInterface* Other) const override;
 #if WITH_EDITORONLY_DATA
-	virtual void GetParameterDefinitionHLSL(const FNiagaraDataInterfaceGPUParamInfo& ParamInfo, FString& OutHLSL) override;
-	virtual bool GetFunctionHLSL(const FNiagaraDataInterfaceGPUParamInfo& ParamInfo, const FNiagaraDataInterfaceGeneratedFunction& FunctionInfo, int FunctionInstanceIndex, FString& OutHLSL) override;
-	virtual bool AppendCompileHash(FNiagaraCompileHashVisitor* InVisitor) const override;
+	NIAGARA_API virtual void GetParameterDefinitionHLSL(const FNiagaraDataInterfaceGPUParamInfo& ParamInfo, FString& OutHLSL) override;
+	NIAGARA_API virtual bool GetFunctionHLSL(const FNiagaraDataInterfaceGPUParamInfo& ParamInfo, const FNiagaraDataInterfaceGeneratedFunction& FunctionInfo, int FunctionInstanceIndex, FString& OutHLSL) override;
+	NIAGARA_API virtual bool AppendCompileHash(FNiagaraCompileHashVisitor* InVisitor) const override;
 #endif
-	virtual void BuildShaderParameters(FNiagaraShaderParametersBuilder& ShaderParametersBuilder) const override;
-	virtual void SetShaderParameters(const FNiagaraDataInterfaceSetShaderParametersContext& Context) const override;
+	NIAGARA_API virtual void BuildShaderParameters(FNiagaraShaderParametersBuilder& ShaderParametersBuilder) const override;
+	NIAGARA_API virtual void SetShaderParameters(const FNiagaraDataInterfaceSetShaderParametersContext& Context) const override;
 
 	virtual bool CanExecuteOnTarget(ENiagaraSimTarget Target) const override { return true; }
 
@@ -100,9 +100,9 @@ public:
 
 	virtual bool HasInternalAttributeReads(const UNiagaraEmitter* OwnerEmitter, const UNiagaraEmitter* Provider) const override { return OwnerEmitter == Provider; };
 
-	virtual void StoreData(FVectorVMExternalFunctionContext& Context);
-	virtual void ExportData(FVectorVMExternalFunctionContext& Context);
+	NIAGARA_API virtual void StoreData(FVectorVMExternalFunctionContext& Context);
+	NIAGARA_API virtual void ExportData(FVectorVMExternalFunctionContext& Context);
 
 protected:
-	virtual bool CopyToInternal(UNiagaraDataInterface* Destination) const override;
+	NIAGARA_API virtual bool CopyToInternal(UNiagaraDataInterface* Destination) const override;
 };
diff --git a/Engine/Plugins/FX/Niagara/Source/Niagara/Classes/NiagaraDataInterfaceGrid2DCollection.h b/Engine/Plugins/FX/Niagara/Source/Niagara/Classes/NiagaraDataInterfaceGrid2DCollection.h
--- a/Engine/Plugins/FX/Niagara/Source/Niagara/Classes/NiagaraDataInterfaceGrid2DCollection.h
+++ b/Engine/Plugins/FX/Niagara/Source/Niagara/Classes/NiagaraDataInterfaceGrid2DCollection.h
@@ -110,8 +110,8 @@ struct FNiagaraDataInterfaceParametersCS_Grid2DCollection : public FNiagaraDataI
 	LAYOUT_FIELD(TMemoryImageArray<uint32>, AttributeChannelCount);
 };
 
-UCLASS(EditInlineNew, Category = "Grid", CollapseCategories, meta = (DisplayName = "Grid2D Collection"), Blueprintable, BlueprintType)
-class NIAGARA_API UNiagaraDataInterfaceGrid2DCollection : public UNiagaraDataInterfaceGrid2D
+UCLASS(EditInlineNew, Category = "Grid", CollapseCategories, meta = (DisplayName = "Grid2D Collection"), Blueprintable, BlueprintType, MinimalAPI)
+class UNiagaraDataInterfaceGrid2DCollection : public UNiagaraDataInterfaceGrid2D
 {
 	GENERATED_UCLASS_BODY()
 
@@ -135,34 +135,34 @@ public:
 	FName PreviewAttribute = NAME_None;
 #endif
 
-	virtual void PostInitProperties() override;
+	NIAGARA_API virtual void PostInitProperties() override;
 	
 	//~ UNiagaraDataInterface interface
 	// VM functionality
-	virtual void GetFunctions(TArray<FNiagaraFunctionSignature>& OutFunctions) override;
-	virtual void GetVMExternalFunction(const FVMExternalFunctionBindingInfo& BindingInfo, void* InstanceData, FVMExternalFunction &OutFunc) override;
+	NIAGARA_API virtual void GetFunctions(TArray<FNiagaraFunctionSignature>& OutFunctions) override;
+	NIAGARA_API virtual void GetVMExternalFunction(const FVMExternalFunctionBindingInfo& BindingInfo, void* InstanceData, FVMExternalFunction &OutFunc) override;
 
 #if WITH_EDITOR	
-	virtual void GetFeedback(UNiagaraSystem* Asset, UNiagaraComponent* Component, TArray<FNiagaraDataInterfaceError>& OutErrors,
+	NIAGARA_API virtual void GetFeedback(UNiagaraSystem* Asset, UNiagaraComponent* Component, TArray<FNiagaraDataInterfaceError>& OutErrors,
 		TArray<FNiagaraDataInterfaceFeedback>& Warnings, TArray<FNiagaraDataInterfaceFeedback>& Info) override;
 #endif
 
 #if WITH_EDITORONLY_DATA
-	virtual bool UpgradeFunctionCall(FNiagaraFunctionSignature& FunctionSignature) override;
+	NIAGARA_API virtual bool UpgradeFunctionCall(FNiagaraFunctionSignature& FunctionSignature) override;
 #endif
 
-	virtual bool Equals(const UNiagaraDataInterface* Other) const override;
+	NIAGARA_API virtual bool Equals(const UNiagaraDataInterface* Other) const override;
 
 	// GPU sim functionality
 #if WITH_EDITORONLY_DATA
-	virtual void GetParameterDefinitionHLSL(const FNiagaraDataInterfaceGPUParamInfo& ParamInfo, FString& OutHLSL) override;
-	virtual bool GetFunctionHLSL(const FNiagaraDataInterfaceGPUParamInfo& ParamInfo, const FNiagaraDataInterfaceGeneratedFunction& FunctionInfo, int FunctionInstanceIndex, FString& OutHLSL) override;
-	virtual bool AppendCompileHash(FNiagaraCompileHashVisitor* InVisitor) const override;
+	NIAGARA_API virtual void GetParameterDefinitionHLSL(const FNiagaraDataInterfaceGPUParamInfo& ParamInfo, FString& OutHLSL) override;
+	NIAGARA_API virtual bool GetFunctionHLSL(const FNiagaraDataInterfaceGPUParamInfo& ParamInfo, const FNiagaraDataInterfaceGeneratedFunction& FunctionInfo, int FunctionInstanceIndex, FString& OutHLSL) override;
+	NIAGARA_API virtual bool AppendCompileHash(FNiagaraCompileHashVisitor* InVisitor) const override;
 #endif
-	virtual void BuildShaderParameters(FNiagaraShaderParametersBuilder& ShaderParametersBuilder) const override;
-	virtual void SetShaderParameters(const FNiagaraDataInterfaceSetShaderParametersContext& Context) const override;
-	virtual FNiagaraDataInterfaceParametersCS* CreateShaderStorage(const FNiagaraDataInterfaceGPUParamInfo& ParameterInfo, const FShaderParameterMap& ParameterMap) const override;
-	virtual const FTypeLayoutDesc* GetShaderStorageType() const override;
+	NIAGARA_API virtual void BuildShaderParameters(FNiagaraShaderParametersBuilder& ShaderParametersBuilder) const override;
+	NIAGARA_API virtual void SetShaderParameters(const FNiagaraDataInterfaceSetShaderParametersContext& Context) const override;
+	NIAGARA_API virtual FNiagaraDataInterfaceParametersCS* CreateShaderStorage(const FNiagaraDataInterfaceGPUParamInfo& ParameterInfo, const FShaderParameterMap& ParameterMap) const override;
+	NIAGARA_API virtual const FTypeLayoutDesc* GetShaderStorageType() const override;
 
 #if WITH_EDITOR
 	virtual ENiagaraGpuDispatchType GetGpuDispatchType() const override { return ENiagaraGpuDispatchType::TwoD; }
@@ -170,140 +170,140 @@ public:
 #endif
 
 	virtual void ProvidePerInstanceDataForRenderThread(void* DataForRenderThread, void* PerInstanceData, const FNiagaraSystemInstanceID& SystemInstance) override {}
-	virtual bool InitPerInstanceData(void* PerInstanceData, FNiagaraSystemInstance* SystemInstance) override;
-	virtual void DestroyPerInstanceData(void* PerInstanceData, FNiagaraSystemInstance* SystemInstance) override;
-	virtual bool PerInstanceTick(void* PerInstanceData, FNiagaraSystemInstance* SystemInstance, float DeltaSeconds) override;
+	NIAGARA_API virtual bool InitPerInstanceData(void* PerInstanceData, FNiagaraSystemInstance* SystemInstance) override;
+	NIAGARA_API virtual void DestroyPerInstanceData(void* PerInstanceData, FNiagaraSystemInstance* SystemInstance) override;
+	NIAGARA_API virtual bool PerInstanceTick(void* PerInstanceData, FNiagaraSystemInstance* SystemInstance, float DeltaSeconds) override;
 	virtual int32 PerInstanceDataSize()const override { return sizeof(FGrid2DCollectionRWInstanceData_GameThread); }
-	virtual bool PerInstanceTickPostSimulate(void* PerInstanceData, FNiagaraSystemInstance* SystemInstance, float DeltaSeconds) override;
+	NIAGARA_API virtual bool PerInstanceTickPostSimulate(void* PerInstanceData, FNiagaraSystemInstance* SystemInstance, float DeltaSeconds) override;
 	virtual bool HasPreSimulateTick() const override { return true; }
 	virtual bool HasPostSimulateTick() const override { return true; }
 
 	virtual bool CanExposeVariables() const override { return true;}
-	virtual void GetExposedVariables(TArray<FNiagaraVariableBase>& OutVariables) const override;
-	virtual bool GetExposedVariableValue(const FNiagaraVariableBase& InVariable, void* InPerInstanceData, FNiagaraSystemInstance* InSystemInstance, void* OutData) const override;
+	NIAGARA_API virtual void GetExposedVariables(TArray<FNiagaraVariableBase>& OutVariables) const override;
+	NIAGARA_API virtual bool GetExposedVariableValue(const FNiagaraVariableBase& InVariable, void* InPerInstanceData, FNiagaraSystemInstance* InSystemInstance, void* OutData) const override;
 
 	virtual bool CanRenderVariablesToCanvas() const { return true; }
-	virtual void GetCanvasVariables(TArray<FNiagaraVariableBase>& OutVariables) const override;
-	virtual bool RenderVariableToCanvas(FNiagaraSystemInstanceID SystemInstanceID, FName VariableName, class FCanvas* Canvas, const FIntRect& DrawRect) const override;
+	NIAGARA_API virtual void GetCanvasVariables(TArray<FNiagaraVariableBase>& OutVariables) const override;
+	NIAGARA_API virtual bool RenderVariableToCanvas(FNiagaraSystemInstanceID SystemInstanceID, FName VariableName, class FCanvas* Canvas, const FIntRect& DrawRect) const override;
 	//~ UNiagaraDataInterface interface END
 
 private:
-	static void CollectAttributesForScript(UNiagaraScript* Script, const UNiagaraDataInterface* DataInterFace, TArray<FNiagaraVariableBase>& OutVariables, TArray<uint32>& OutVariableOffsets, int32& TotalAttributes, TArray<FText>* OutWarnings = nullptr);
+	static NIAGARA_API void CollectAttributesForScript(UNiagaraScript* Script, const UNiagaraDataInterface* DataInterFace, TArray<FNiagaraVariableBase>& OutVariables, TArray<uint32>& OutVariableOffsets, int32& TotalAttributes, TArray<FText>* OutWarnings = nullptr);
 public:
 	/** Finds all attributes by locating the data interface amongst the parameter stores. */
-	void FindAttributes(TArray<FNiagaraVariableBase>& OutVariables, TArray<uint32>& OutVariableOffsets, int32& OutNumAttribChannelsFound, TArray<FText>* OutWarnings = nullptr) const;
+	NIAGARA_API void FindAttributes(TArray<FNiagaraVariableBase>& OutVariables, TArray<uint32>& OutVariableOffsets, int32& OutNumAttribChannelsFound, TArray<FText>* OutWarnings = nullptr) const;
 
 	// Fills a texture render target 2d with the current data from the simulation
 	// #todo(dmp): this will eventually go away when we formalize how data makes it out of Niagara
 	UFUNCTION(BlueprintCallable, Category = Niagara, meta=(DeprecatedFunction, DeprecationMessage = "This function has been replaced by object user variables on the emitter to specify render targets to fill with data."))
-	virtual bool FillTexture2D(const UNiagaraComponent *Component, UTextureRenderTarget2D *dest, int AttributeIndex);
+	NIAGARA_API virtual bool FillTexture2D(const UNiagaraComponent *Component, UTextureRenderTarget2D *dest, int AttributeIndex);
 	
 	UFUNCTION(BlueprintCallable, Category = Niagara, meta=(DeprecatedFunction, DeprecationMessage = "This function has been replaced by object user variables on the emitter to specify render targets to fill with data."))
-	virtual bool FillRawTexture2D(const UNiagaraComponent *Component, UTextureRenderTarget2D *Dest, int &TilesX, int &TilesY);
+	NIAGARA_API virtual bool FillRawTexture2D(const UNiagaraComponent *Component, UTextureRenderTarget2D *Dest, int &TilesX, int &TilesY);
 	
 	UFUNCTION(BlueprintCallable, Category = Niagara, meta = (DeprecatedFunction, DeprecationMessage = "This function has been replaced by object user variables on the emitter to specify render targets to fill with data."))
-	virtual void GetRawTextureSize(const UNiagaraComponent *Component, int &SizeX, int &SizeY);
+	NIAGARA_API virtual void GetRawTextureSize(const UNiagaraComponent *Component, int &SizeX, int &SizeY);
 
 	UFUNCTION(BlueprintCallable, Category = Niagara)
-	virtual void GetTextureSize(const UNiagaraComponent *Component, int &SizeX, int &SizeY);
-
-	void VMGetWorldBBoxSize(FVectorVMExternalFunctionContext& Context);
-	void VMGetCellSize(FVectorVMExternalFunctionContext& Context);
-	void VMGetNumCells(FVectorVMExternalFunctionContext& Context);
-	void VMSetNumCells(FVectorVMExternalFunctionContext& Context);
-	void VMGetAttributeIndex(FVectorVMExternalFunctionContext& Context, const FName& InName, int32 NumChannels);
-
-	static const FString GridName;
-	static const FString OutputGridName;
-	static const FString SamplerName;
-
-	static const FName ClearCellFunctionName;
-	static const FName CopyPreviousToCurrentForCellFunctionName;
-
-	static const FName SetValueFunctionName;
-	static const FName GetValueFunctionName;
-	static const FName SampleGridFunctionName;
-
-	static const FName SetValueAtIndexFunctionName;
-	static const FName GetPreviousValueAtIndexFunctionName;
-	static const FName SamplePreviousGridAtIndexFunctionName;
-
-	static const FName SetVector4ValueFunctionName;
-	static const FName GetVector4ValueFunctionName;
-	static const FName SampleGridVector4FunctionName;
-
-	static const FName SetVector3ValueFunctionName;
-	static const FName GetVector3ValueFunctionName;
-	static const FName SampleGridVector3FunctionName;
-
-	static const FName SetVector2ValueFunctionName;
-	static const FName GetVector2ValueFunctionName;
-	static const FName SampleGridVector2FunctionName;
-
-	static const FName SetFloatValueFunctionName;
-	static const FName GetFloatValueFunctionName;
-	static const FName SampleGridFloatFunctionName;
-
-	static const FName GetPreviousVector4ValueFunctionName;
-	static const FName SamplePreviousGridVector4FunctionName;
-	static const FName SetVectorValueFunctionName;
-	static const FName GetPreviousVectorValueFunctionName;
-	static const FName SamplePreviousGridVectorFunctionName;
-	static const FName SetVector2DValueFunctionName;
-	static const FName GetPreviousVector2DValueFunctionName;
-	static const FName SamplePreviousGridVector2DFunctionName;
-	static const FName GetPreviousFloatValueFunctionName;
-	static const FName SamplePreviousGridFloatFunctionName;
+	NIAGARA_API virtual void GetTextureSize(const UNiagaraComponent *Component, int &SizeX, int &SizeY);
+
+	NIAGARA_API void VMGetWorldBBoxSize(FVectorVMExternalFunctionContext& Context);
+	NIAGARA_API void VMGetCellSize(FVectorVMExternalFunctionContext& Context);
+	NIAGARA_API void VMGetNumCells(FVectorVMExternalFunctionContext& Context);
+	NIAGARA_API void VMSetNumCells(FVectorVMExternalFunctionContext& Context);
+	NIAGARA_API void VMGetAttributeIndex(FVectorVMExternalFunctionContext& Context, const FName& InName, int32 NumChannels);
+
+	static NIAGARA_API const FString GridName;
+	static NIAGARA_API const FString OutputGridName;
+	static NIAGARA_API const FString SamplerName;
+
+	static NIAGARA_API const FName ClearCellFunctionName;
+	static NIAGARA_API const FName CopyPreviousToCurrentForCellFunctionName;
+
+	static NIAGARA_API const FName SetValueFunctionName;
+	static NIAGARA_API const FName GetValueFunctionName;
+	static NIAGARA_API const FName SampleGridFunctionName;
+
+	static NIAGARA_API const FName SetValueAtIndexFunctionName;
+	static NIAGARA_API const FName GetPreviousValueAtIndexFunctionName;
+	static NIAGARA_API const FName SamplePreviousGridAtIndexFunctionName;
+
+	static NIAGARA_API const FName SetVector4ValueFunctionName;
+	static NIAGARA_API const FName GetVector4ValueFunctionName;
+	static NIAGARA_API const FName SampleGridVector4FunctionName;
+
+	static NIAGARA_API const FName SetVector3ValueFunctionName;
+	static NIAGARA_API const FName GetVector3ValueFunctionName;
+	static NIAGARA_API const FName SampleGridVector3FunctionName;
+
+	static NIAGARA_API const FName SetVector2ValueFunctionName;
+	static NIAGARA_API const FName GetVector2ValueFunctionName;
+	static NIAGARA_API const FName SampleGridVector2FunctionName;
+
+	static NIAGARA_API const FName SetFloatValueFunctionName;
+	static NIAGARA_API const FName GetFloatValueFunctionName;
+	static NIAGARA_API const FName SampleGridFloatFunctionName;
+
+	static NIAGARA_API const FName GetPreviousVector4ValueFunctionName;
+	static NIAGARA_API const FName SamplePreviousGridVector4FunctionName;
+	static NIAGARA_API const FName SetVectorValueFunctionName;
+	static NIAGARA_API const FName GetPreviousVectorValueFunctionName;
+	static NIAGARA_API const FName SamplePreviousGridVectorFunctionName;
+	static NIAGARA_API const FName SetVector2DValueFunctionName;
+	static NIAGARA_API const FName GetPreviousVector2DValueFunctionName;
+	static NIAGARA_API const FName SamplePreviousGridVector2DFunctionName;
+	static NIAGARA_API const FName GetPreviousFloatValueFunctionName;
+	static NIAGARA_API const FName SamplePreviousGridFloatFunctionName;
 	
-	static const FString AttributeIndicesBaseName;
-	static const TCHAR* VectorComponentNames[];
+	static NIAGARA_API const FString AttributeIndicesBaseName;
+	static NIAGARA_API const TCHAR* VectorComponentNames[];
 
-	static const FName SetNumCellsFunctionName;
+	static NIAGARA_API const FName SetNumCellsFunctionName;
 
-	static const FName GetVector4AttributeIndexFunctionName;
-	static const FName GetVectorAttributeIndexFunctionName;
-	static const FName GetVector2DAttributeIndexFunctionName;
-	static const FName GetFloatAttributeIndexFunctionName;
+	static NIAGARA_API const FName GetVector4AttributeIndexFunctionName;
+	static NIAGARA_API const FName GetVectorAttributeIndexFunctionName;
+	static NIAGARA_API const FName GetVector2DAttributeIndexFunctionName;
+	static NIAGARA_API const FName GetFloatAttributeIndexFunctionName;
 
-	static const FString AnonymousAttributeString;
+	static NIAGARA_API const FString AnonymousAttributeString;
 
 #if WITH_EDITOR
 	virtual bool SupportsSetupAndTeardownHLSL() const { return true; }
-	virtual bool GenerateSetupHLSL(FNiagaraDataInterfaceGPUParamInfo& DIInstanceInfo, TConstArrayView<FNiagaraVariable> InArguments, bool bSpawnOnly, bool bPartialWrites, TArray<FText>& OutErrors, FString& OutHLSL) const;
-	virtual bool GenerateTeardownHLSL(FNiagaraDataInterfaceGPUParamInfo& DIInstanceInfo, TConstArrayView<FNiagaraVariable> InArguments, bool bSpawnOnly, bool bPartialWrites, TArray<FText>& OutErrors, FString& OutHLSL) const;
+	NIAGARA_API virtual bool GenerateSetupHLSL(FNiagaraDataInterfaceGPUParamInfo& DIInstanceInfo, TConstArrayView<FNiagaraVariable> InArguments, bool bSpawnOnly, bool bPartialWrites, TArray<FText>& OutErrors, FString& OutHLSL) const;
+	NIAGARA_API virtual bool GenerateTeardownHLSL(FNiagaraDataInterfaceGPUParamInfo& DIInstanceInfo, TConstArrayView<FNiagaraVariable> InArguments, bool bSpawnOnly, bool bPartialWrites, TArray<FText>& OutErrors, FString& OutHLSL) const;
 	virtual bool SupportsIterationSourceNamespaceAttributesHLSL() const override { return true; }
-	virtual bool GenerateIterationSourceNamespaceReadAttributesHLSL(FNiagaraDataInterfaceGPUParamInfo& DIInstanceInfo, const FNiagaraVariable& IterationSourceVar, TConstArrayView<FNiagaraVariable> InArguments, TConstArrayView<FNiagaraVariable> InAttributes, TConstArrayView<FString> InAttributeHLSLNames, bool bInSetToDefaults, bool bSpawnOnly, bool bPartialWrites, TArray<FText>& OutErrors, FString& OutHLSL) const override;
-	virtual bool GenerateIterationSourceNamespaceWriteAttributesHLSL(FNiagaraDataInterfaceGPUParamInfo& DIInstanceInfo, const FNiagaraVariable& IterationSourceVar, TConstArrayView<FNiagaraVariable> InArguments, TConstArrayView<FNiagaraVariable> InAttributes, TConstArrayView<FString> InAttributeHLSLNames, TConstArrayView<FNiagaraVariable> InAllAttributes, bool bSpawnOnly, bool bPartialWrites, TArray<FText>& OutErrors, FString& OutHLSL) const override;
+	NIAGARA_API virtual bool GenerateIterationSourceNamespaceReadAttributesHLSL(FNiagaraDataInterfaceGPUParamInfo& DIInstanceInfo, const FNiagaraVariable& IterationSourceVar, TConstArrayView<FNiagaraVariable> InArguments, TConstArrayView<FNiagaraVariable> InAttributes, TConstArrayView<FString> InAttributeHLSLNames, bool bInSetToDefaults, bool bSpawnOnly, bool bPartialWrites, TArray<FText>& OutErrors, FString& OutHLSL) const override;
+	NIAGARA_API virtual bool GenerateIterationSourceNamespaceWriteAttributesHLSL(FNiagaraDataInterfaceGPUParamInfo& DIInstanceInfo, const FNiagaraVariable& IterationSourceVar, TConstArrayView<FNiagaraVariable> InArguments, TConstArrayView<FNiagaraVariable> InAttributes, TConstArrayView<FString> InAttributeHLSLNames, TConstArrayView<FNiagaraVariable> InAllAttributes, bool bSpawnOnly, bool bPartialWrites, TArray<FText>& OutErrors, FString& OutHLSL) const override;
 #endif
 
-	static int32 GetComponentCountFromFuncName(const FName& FuncName);
-	static FNiagaraTypeDefinition GetValueTypeFromFuncName(const FName& FuncName);
-	static bool CanCreateVarFromFuncName(const FName& FuncName);
+	static NIAGARA_API int32 GetComponentCountFromFuncName(const FName& FuncName);
+	static NIAGARA_API FNiagaraTypeDefinition GetValueTypeFromFuncName(const FName& FuncName);
+	static NIAGARA_API bool CanCreateVarFromFuncName(const FName& FuncName);
 
 	TMap<FNiagaraSystemInstanceID, FGrid2DCollectionRWInstanceData_GameThread*>& GetSystemInstancesToProxyData_GT() { return SystemInstancesToProxyData_GT; }
 protected:
 #if WITH_EDITORONLY_DATA
-	void WriteSetHLSL(const FNiagaraDataInterfaceRWAttributeHelper &AttributeHelper, const FNiagaraDataInterfaceGPUParamInfo& ParamInfo, const FNiagaraDataInterfaceGeneratedFunction& FunctionInfo, int FunctionInstanceIndex, int32 InNumChannels, FString& OutHLSL);
-	void WriteGetHLSL(const FNiagaraDataInterfaceRWAttributeHelper& AttributeHelper, const FNiagaraDataInterfaceGPUParamInfo& ParamInfo, const FNiagaraDataInterfaceGeneratedFunction& FunctionInfo, int FunctionInstanceIndex, int32 InNumChannels, FString& OutHLSL);
-	void WriteSampleHLSL(const FNiagaraDataInterfaceRWAttributeHelper& AttributeHelper, const FNiagaraDataInterfaceGPUParamInfo& ParamInfo, const FNiagaraDataInterfaceGeneratedFunction& FunctionInfo, int FunctionInstanceIndex, int32 InNumChannels, FString& OutHLSL);
-	void WriteAttributeGetIndexHLSL(const FNiagaraDataInterfaceRWAttributeHelper& AttributeHelper, const FNiagaraDataInterfaceGPUParamInfo& ParamInfo, const FNiagaraDataInterfaceGeneratedFunction& FunctionInfo, int FunctionInstanceIndex, int32 InNumChannels, FString& OutHLSL);
+	NIAGARA_API void WriteSetHLSL(const FNiagaraDataInterfaceRWAttributeHelper &AttributeHelper, const FNiagaraDataInterfaceGPUParamInfo& ParamInfo, const FNiagaraDataInterfaceGeneratedFunction& FunctionInfo, int FunctionInstanceIndex, int32 InNumChannels, FString& OutHLSL);
+	NIAGARA_API void WriteGetHLSL(const FNiagaraDataInterfaceRWAttributeHelper& AttributeHelper, const FNiagaraDataInterfaceGPUParamInfo& ParamInfo, const FNiagaraDataInterfaceGeneratedFunction& FunctionInfo, int FunctionInstanceIndex, int32 InNumChannels, FString& OutHLSL);
+	NIAGARA_API void WriteSampleHLSL(const FNiagaraDataInterfaceRWAttributeHelper& AttributeHelper, const FNiagaraDataInterfaceGPUParamInfo& ParamInfo, const FNiagaraDataInterfaceGeneratedFunction& FunctionInfo, int FunctionInstanceIndex, int32 InNumChannels, FString& OutHLSL);
+	NIAGARA_API void WriteAttributeGetIndexHLSL(const FNiagaraDataInterfaceRWAttributeHelper& AttributeHelper, const FNiagaraDataInterfaceGPUParamInfo& ParamInfo, const FNiagaraDataInterfaceGeneratedFunction& FunctionInfo, int FunctionInstanceIndex, int32 InNumChannels, FString& OutHLSL);
 
-	FString GenerateAttributeIndexHLSL(const FNiagaraDataInterfaceRWAttributeHelper& AttributeHelper, const FNiagaraDataInterfaceGeneratedFunction& FunctionInfo);
+	NIAGARA_API FString GenerateAttributeIndexHLSL(const FNiagaraDataInterfaceRWAttributeHelper& AttributeHelper, const FNiagaraDataInterfaceGeneratedFunction& FunctionInfo);
 
-	const TCHAR* TypeDefinitionToHLSLTypeString(const FNiagaraTypeDefinition& InDef) const;
-	FName TypeDefinitionToGetFunctionName(const FNiagaraTypeDefinition& InDef) const;
-	FName TypeDefinitionToSetFunctionName(const FNiagaraTypeDefinition& InDef) const;
+	NIAGARA_API const TCHAR* TypeDefinitionToHLSLTypeString(const FNiagaraTypeDefinition& InDef) const;
+	NIAGARA_API FName TypeDefinitionToGetFunctionName(const FNiagaraTypeDefinition& InDef) const;
+	NIAGARA_API FName TypeDefinitionToSetFunctionName(const FNiagaraTypeDefinition& InDef) const;
 #endif
 
 	// Only used with CDO for hlsl generation
 	bool bUseIndirection;
 
 	//~ UNiagaraDataInterface interface
-	virtual bool CopyToInternal(UNiagaraDataInterface* Destination) const override;
+	NIAGARA_API virtual bool CopyToInternal(UNiagaraDataInterface* Destination) const override;
 
 	//~ UNiagaraDataInterface interface END
 
-	static FNiagaraVariableBase ExposedRTVar;
+	static NIAGARA_API FNiagaraVariableBase ExposedRTVar;
 
 	TMap<FNiagaraSystemInstanceID, FGrid2DCollectionRWInstanceData_GameThread*> SystemInstancesToProxyData_GT;
 
diff --git a/Engine/Plugins/FX/Niagara/Source/Niagara/Classes/NiagaraDataInterfaceGrid2DCollectionReader.h b/Engine/Plugins/FX/Niagara/Source/Niagara/Classes/NiagaraDataInterfaceGrid2DCollectionReader.h
--- a/Engine/Plugins/FX/Niagara/Source/Niagara/Classes/NiagaraDataInterfaceGrid2DCollectionReader.h
+++ b/Engine/Plugins/FX/Niagara/Source/Niagara/Classes/NiagaraDataInterfaceGrid2DCollectionReader.h
@@ -13,8 +13,8 @@ class FNiagaraSystemInstance;
 class UTextureRenderTarget;
 class UTextureRenderTargetVolume;
 
-UCLASS(EditInlineNew, Category = "Grid", CollapseCategories, hidecategories = (Grid2DCollection, Grid, Deprecated), meta = (DisplayName = "Grid2D Collection Reader", Experimental), Blueprintable, BlueprintType)
-class NIAGARA_API UNiagaraDataInterfaceGrid2DCollectionReader : public UNiagaraDataInterfaceGrid2DCollection
+UCLASS(EditInlineNew, Category = "Grid", CollapseCategories, hidecategories = (Grid2DCollection, Grid, Deprecated), meta = (DisplayName = "Grid2D Collection Reader", Experimental), Blueprintable, BlueprintType, MinimalAPI)
+class UNiagaraDataInterfaceGrid2DCollectionReader : public UNiagaraDataInterfaceGrid2DCollection
 {
 	GENERATED_UCLASS_BODY()
 
@@ -30,14 +30,14 @@ public:
 
 	virtual void Serialize(FArchive& Ar) override { UNiagaraDataInterfaceRWBase::Serialize(Ar); }
 
-	virtual bool Equals(const UNiagaraDataInterface* Other) const override;
-	virtual bool InitPerInstanceData(void* PerInstanceData, FNiagaraSystemInstance* SystemInstance) override;
-	virtual void GetEmitterDependencies(UNiagaraSystem* Asset, TArray<FVersionedNiagaraEmitter>& Dependencies) const override;
-	virtual void GetFunctions(TArray<FNiagaraFunctionSignature>& OutFunctions) override;
-	virtual bool PerInstanceTickPostSimulate(void* PerInstanceData, FNiagaraSystemInstance* SystemInstance, float DeltaSeconds) override;
+	NIAGARA_API virtual bool Equals(const UNiagaraDataInterface* Other) const override;
+	NIAGARA_API virtual bool InitPerInstanceData(void* PerInstanceData, FNiagaraSystemInstance* SystemInstance) override;
+	NIAGARA_API virtual void GetEmitterDependencies(UNiagaraSystem* Asset, TArray<FVersionedNiagaraEmitter>& Dependencies) const override;
+	NIAGARA_API virtual void GetFunctions(TArray<FNiagaraFunctionSignature>& OutFunctions) override;
+	NIAGARA_API virtual bool PerInstanceTickPostSimulate(void* PerInstanceData, FNiagaraSystemInstance* SystemInstance, float DeltaSeconds) override;
 
 protected:
 	//~ UNiagaraDataInterface interface
-	virtual bool CopyToInternal(UNiagaraDataInterface* Destination) const override;
+	NIAGARA_API virtual bool CopyToInternal(UNiagaraDataInterface* Destination) const override;
 	//~ UNiagaraDataInterface interface END
-};
\ No newline at end of file
+};
diff --git a/Engine/Plugins/FX/Niagara/Source/Niagara/Classes/NiagaraDataInterfaceGrid3DCollection.h b/Engine/Plugins/FX/Niagara/Source/Niagara/Classes/NiagaraDataInterfaceGrid3DCollection.h
--- a/Engine/Plugins/FX/Niagara/Source/Niagara/Classes/NiagaraDataInterfaceGrid3DCollection.h
+++ b/Engine/Plugins/FX/Niagara/Source/Niagara/Classes/NiagaraDataInterfaceGrid3DCollection.h
@@ -120,8 +120,8 @@ struct FNiagaraDataInterfaceParametersCS_Grid3DCollection : public FNiagaraDataI
 	LAYOUT_FIELD(TMemoryImageArray<uint32>, AttributeChannelCount);
 };
 
-UCLASS(EditInlineNew, Category = "Grid", CollapseCategories, meta = (DisplayName = "Grid3D Collection"), Blueprintable, BlueprintType)
-class NIAGARA_API UNiagaraDataInterfaceGrid3DCollection : public UNiagaraDataInterfaceGrid3D
+UCLASS(EditInlineNew, Category = "Grid", CollapseCategories, meta = (DisplayName = "Grid3D Collection"), Blueprintable, BlueprintType, MinimalAPI)
+class UNiagaraDataInterfaceGrid3DCollection : public UNiagaraDataInterfaceGrid3D
 {
 	GENERATED_UCLASS_BODY()
 
@@ -149,54 +149,54 @@ public:
 	FName PreviewAttribute = NAME_None;
 #endif
 
-	virtual void PostInitProperties() override;
+	NIAGARA_API virtual void PostInitProperties() override;
 
 	//~ UNiagaraDataInterface interface
 	// VM functionality
-	virtual void GetFunctions(TArray<FNiagaraFunctionSignature>& OutFunctions) override;
-	virtual void GetVMExternalFunction(const FVMExternalFunctionBindingInfo& BindingInfo, void* InstanceData, FVMExternalFunction& OutFunc) override;
+	NIAGARA_API virtual void GetFunctions(TArray<FNiagaraFunctionSignature>& OutFunctions) override;
+	NIAGARA_API virtual void GetVMExternalFunction(const FVMExternalFunctionBindingInfo& BindingInfo, void* InstanceData, FVMExternalFunction& OutFunc) override;
 
 #if WITH_EDITOR	
-	virtual void GetFeedback(UNiagaraSystem* Asset, UNiagaraComponent* Component, TArray<FNiagaraDataInterfaceError>& OutErrors,
+	NIAGARA_API virtual void GetFeedback(UNiagaraSystem* Asset, UNiagaraComponent* Component, TArray<FNiagaraDataInterfaceError>& OutErrors,
 		TArray<FNiagaraDataInterfaceFeedback>& Warnings, TArray<FNiagaraDataInterfaceFeedback>& Info) override;
 #endif
 
 #if WITH_EDITORONLY_DATA
-	virtual bool UpgradeFunctionCall(FNiagaraFunctionSignature& FunctionSignature) override;
+	NIAGARA_API virtual bool UpgradeFunctionCall(FNiagaraFunctionSignature& FunctionSignature) override;
 #endif
 
-	virtual bool Equals(const UNiagaraDataInterface* Other) const override;
+	NIAGARA_API virtual bool Equals(const UNiagaraDataInterface* Other) const override;
 
 	// GPU sim functionality
 #if WITH_EDITOR
-	virtual bool ShouldCompile(EShaderPlatform ShaderPlatform) const override;
+	NIAGARA_API virtual bool ShouldCompile(EShaderPlatform ShaderPlatform) const override;
 #endif
 #if WITH_EDITORONLY_DATA
-	virtual bool AppendCompileHash(FNiagaraCompileHashVisitor* InVisitor) const override;
-	virtual void GetParameterDefinitionHLSL(const FNiagaraDataInterfaceGPUParamInfo& ParamInfo, FString& OutHLSL) override;
-	virtual bool GetFunctionHLSL(const FNiagaraDataInterfaceGPUParamInfo& ParamInfo, const FNiagaraDataInterfaceGeneratedFunction& FunctionInfo, int FunctionInstanceIndex, FString& OutHLSL) override;
+	NIAGARA_API virtual bool AppendCompileHash(FNiagaraCompileHashVisitor* InVisitor) const override;
+	NIAGARA_API virtual void GetParameterDefinitionHLSL(const FNiagaraDataInterfaceGPUParamInfo& ParamInfo, FString& OutHLSL) override;
+	NIAGARA_API virtual bool GetFunctionHLSL(const FNiagaraDataInterfaceGPUParamInfo& ParamInfo, const FNiagaraDataInterfaceGeneratedFunction& FunctionInfo, int FunctionInstanceIndex, FString& OutHLSL) override;
 #endif
-	virtual void BuildShaderParameters(FNiagaraShaderParametersBuilder& ShaderParametersBuilder) const override;
-	virtual void SetShaderParameters(const FNiagaraDataInterfaceSetShaderParametersContext& Context) const override;
-	virtual FNiagaraDataInterfaceParametersCS* CreateShaderStorage(const FNiagaraDataInterfaceGPUParamInfo& ParameterInfo, const FShaderParameterMap& ParameterMap) const override;
-	virtual const FTypeLayoutDesc* GetShaderStorageType() const override;
+	NIAGARA_API virtual void BuildShaderParameters(FNiagaraShaderParametersBuilder& ShaderParametersBuilder) const override;
+	NIAGARA_API virtual void SetShaderParameters(const FNiagaraDataInterfaceSetShaderParametersContext& Context) const override;
+	NIAGARA_API virtual FNiagaraDataInterfaceParametersCS* CreateShaderStorage(const FNiagaraDataInterfaceGPUParamInfo& ParameterInfo, const FShaderParameterMap& ParameterMap) const override;
+	NIAGARA_API virtual const FTypeLayoutDesc* GetShaderStorageType() const override;
 
 #if WITH_EDITOR
 	virtual ENiagaraGpuDispatchType GetGpuDispatchType() const override { return ENiagaraGpuDispatchType::ThreeD; }
 #endif
 
 	virtual void ProvidePerInstanceDataForRenderThread(void* DataForRenderThread, void* PerInstanceData, const FNiagaraSystemInstanceID& SystemInstance) override {}
-	virtual bool InitPerInstanceData(void* PerInstanceData, FNiagaraSystemInstance* SystemInstance) override;
-	virtual void DestroyPerInstanceData(void* PerInstanceData, FNiagaraSystemInstance* SystemInstance) override;
-	virtual bool PerInstanceTick(void* PerInstanceData, FNiagaraSystemInstance* SystemInstance, float DeltaSeconds) override;
+	NIAGARA_API virtual bool InitPerInstanceData(void* PerInstanceData, FNiagaraSystemInstance* SystemInstance) override;
+	NIAGARA_API virtual void DestroyPerInstanceData(void* PerInstanceData, FNiagaraSystemInstance* SystemInstance) override;
+	NIAGARA_API virtual bool PerInstanceTick(void* PerInstanceData, FNiagaraSystemInstance* SystemInstance, float DeltaSeconds) override;
 	virtual int32 PerInstanceDataSize()const override { return sizeof(FGrid3DCollectionRWInstanceData_GameThread); }
-	virtual bool PerInstanceTickPostSimulate(void* PerInstanceData, FNiagaraSystemInstance* SystemInstance, float DeltaSeconds) override;
+	NIAGARA_API virtual bool PerInstanceTickPostSimulate(void* PerInstanceData, FNiagaraSystemInstance* SystemInstance, float DeltaSeconds) override;
 	virtual bool HasPreSimulateTick() const override { return true; }
 	virtual bool HasPostSimulateTick() const override { return true; }
 
 	virtual bool CanExposeVariables() const override { return true; }
-	virtual void GetExposedVariables(TArray<FNiagaraVariableBase>& OutVariables) const override;
-	virtual bool GetExposedVariableValue(const FNiagaraVariableBase& InVariable, void* InPerInstanceData, FNiagaraSystemInstance* InSystemInstance, void* OutData) const override;
+	NIAGARA_API virtual void GetExposedVariables(TArray<FNiagaraVariableBase>& OutVariables) const override;
+	NIAGARA_API virtual bool GetExposedVariableValue(const FNiagaraVariableBase& InVariable, void* InPerInstanceData, FNiagaraSystemInstance* InSystemInstance, void* OutData) const override;
 	//~ UNiagaraDataInterface interface END
 
 public:
@@ -205,148 +205,148 @@ public:
 	// #todo(dmp): reimplement for 3d
 
 	UFUNCTION(BlueprintCallable, Category = Niagara, meta = (DeprecatedFunction, DeprecationMessage = "This function has been replaced by object user variables on the emitter to specify render targets to fill with data."))
-	virtual bool FillVolumeTexture(const UNiagaraComponent* Component, UVolumeTexture* dest, int AttributeIndex);
+	NIAGARA_API virtual bool FillVolumeTexture(const UNiagaraComponent* Component, UVolumeTexture* dest, int AttributeIndex);
 
 	UFUNCTION(BlueprintCallable, Category = Niagara, meta = (DeprecatedFunction, DeprecationMessage = "This function has been replaced by object user variables on the emitter to specify render targets to fill with data."))
-	virtual bool FillRawVolumeTexture(const UNiagaraComponent* Component, UVolumeTexture* Dest, int& TilesX, int& TilesY, int& TileZ);
+	NIAGARA_API virtual bool FillRawVolumeTexture(const UNiagaraComponent* Component, UVolumeTexture* Dest, int& TilesX, int& TilesY, int& TileZ);
 
 	UFUNCTION(BlueprintCallable, Category = Niagara)
-	virtual void GetRawTextureSize(const UNiagaraComponent* Component, int& SizeX, int& SizeY, int& SizeZ);
+	NIAGARA_API virtual void GetRawTextureSize(const UNiagaraComponent* Component, int& SizeX, int& SizeY, int& SizeZ);
 
 	UFUNCTION(BlueprintCallable, Category = Niagara)
-	virtual void GetTextureSize(const UNiagaraComponent* Component, int& SizeX, int& SizeY, int& SizeZ);
-
-	void VMGetWorldBBoxSize(FVectorVMExternalFunctionContext& Context);
-	void VMGetCellSize(FVectorVMExternalFunctionContext& Context);
-
-	void VMGetNumCells(FVectorVMExternalFunctionContext& Context);
-	void VMSetNumCells(FVectorVMExternalFunctionContext& Context);
-	void VMUnitToFloatIndex(FVectorVMExternalFunctionContext& Context);
-
-	void VMGetAttributeIndex(FVectorVMExternalFunctionContext& Context, const FName& InName, int32 NumChannels);
-
-	static const FString NumTilesName;
-	static const FString OneOverNumTilesName;
-	static const FString UnitClampMinName;
-	static const FString UnitClampMaxName;
-
-	static const FString GridName;
-	static const FString OutputGridName;
-	static const FString SamplerName;
-
-	static const FName ClearCellFunctionName;
-	static const FName CopyPreviousToCurrentForCellFunctionName;
-	static const FName CopyMaskedPreviousToCurrentForCellFunctionName;
-
-	static const FName SetValueFunctionName;
-	static const FName GetValueFunctionName;
-	static const FName SampleGridFunctionName;
-	static const FName CubicSampleGridFunctionName;
-
-	static const FName SetFullGridValueFunctionName;
-	static const FName GetFullGridPreviousValueFunctionName;
-	static const FName SamplePreviousFullGridFunctionName;
-	static const FName CubicSamplePreviousFullGridFunctionName;
-
-	static const FName SetVector4ValueFunctionName;
-	static const FName SetVector3ValueFunctionName;
-	static const FName SetVector2ValueFunctionName;
-	static const FName GetVector2ValueFunctionName;
-	static const FName SetFloatValueFunctionName;
-	static const FName GetPreviousValueAtIndexFunctionName;
-	static const FName SamplePreviousGridAtIndexFunctionName;
-	static const FName CubicSamplePreviousGridAtIndexFunctionName;
-
-	static const FName SetVector4ValueAtIndexFunctionName;
-	static const FName SetVector3ValueAtIndexFunctionName;
-	static const FName SetVector2ValueAtIndexFunctionName;
-	static const FName SetFloatValueAtIndexFunctionName;
-
-	static const FName GetPreviousVector4ValueAtIndexFunctionName;
-	static const FName GetPreviousVector3ValueAtIndexFunctionName;
-	static const FName GetPreviousVector2ValueAtIndexFunctionName;
-	static const FName GetPreviousFloatValueAtIndexFunctionName;
-
-	static const FName SamplePreviousGridVector4AtIndexFunctionName;
-	static const FName SamplePreviousGridVector3AtIndexFunctionName;
-	static const FName SamplePreviousGridVector2AtIndexFunctionName;
-	static const FName SamplePreviousGridFloatAtIndexFunctionName;
-
-	static const FName CubicSamplePreviousGridVector4AtIndexFunctionName;
-	static const FName CubicSamplePreviousGridVector3AtIndexFunctionName;
-	static const FName CubicSamplePreviousGridVector2AtIndexFunctionName;
-	static const FName CubicSamplePreviousGridFloatAtIndexFunctionName;
-
-	static const FName GetPreviousVector4ValueFunctionName;
-	static const FName SamplePreviousGridVector4FunctionName;
-	static const FName CubicSamplePreviousGridVector4FunctionName;
-	static const FName SetVectorValueFunctionName;
-	static const FName GetPreviousVectorValueFunctionName;
-	static const FName SamplePreviousGridVectorFunctionName;
-	static const FName CubicSamplePreviousGridVectorFunctionName;
-	static const FName SetVector2DValueFunctionName;
-	static const FName GetPreviousVector2DValueFunctionName;
-	static const FName SamplePreviousGridVector2DFunctionName;
-	static const FName CubicSamplePreviousGridVector2DFunctionName;
-	static const FName GetPreviousFloatValueFunctionName;
-	static const FName SamplePreviousGridFloatFunctionName;
-	static const FName CubicSamplePreviousGridFloatFunctionName;
-	static const FString AttributeIndicesBaseName;
-	static const FString PerAttributeDataName;
-	static const TCHAR* VectorComponentNames[];
-
-	static const FName SetNumCellsFunctionName;
-
-	static const FName GetVector4AttributeIndexFunctionName;
-	static const FName GetVectorAttributeIndexFunctionName;
-	static const FName GetVector2DAttributeIndexFunctionName;
-	static const FName GetFloatAttributeIndexFunctionName;
-
-	static const FString AnonymousAttributeString;
-
-	static const TCHAR* TemplateShaderFilePath;
+	NIAGARA_API virtual void GetTextureSize(const UNiagaraComponent* Component, int& SizeX, int& SizeY, int& SizeZ);
+
+	NIAGARA_API void VMGetWorldBBoxSize(FVectorVMExternalFunctionContext& Context);
+	NIAGARA_API void VMGetCellSize(FVectorVMExternalFunctionContext& Context);
+
+	NIAGARA_API void VMGetNumCells(FVectorVMExternalFunctionContext& Context);
+	NIAGARA_API void VMSetNumCells(FVectorVMExternalFunctionContext& Context);
+	NIAGARA_API void VMUnitToFloatIndex(FVectorVMExternalFunctionContext& Context);
+
+	NIAGARA_API void VMGetAttributeIndex(FVectorVMExternalFunctionContext& Context, const FName& InName, int32 NumChannels);
+
+	static NIAGARA_API const FString NumTilesName;
+	static NIAGARA_API const FString OneOverNumTilesName;
+	static NIAGARA_API const FString UnitClampMinName;
+	static NIAGARA_API const FString UnitClampMaxName;
+
+	static NIAGARA_API const FString GridName;
+	static NIAGARA_API const FString OutputGridName;
+	static NIAGARA_API const FString SamplerName;
+
+	static NIAGARA_API const FName ClearCellFunctionName;
+	static NIAGARA_API const FName CopyPreviousToCurrentForCellFunctionName;
+	static NIAGARA_API const FName CopyMaskedPreviousToCurrentForCellFunctionName;
+
+	static NIAGARA_API const FName SetValueFunctionName;
+	static NIAGARA_API const FName GetValueFunctionName;
+	static NIAGARA_API const FName SampleGridFunctionName;
+	static NIAGARA_API const FName CubicSampleGridFunctionName;
+
+	static NIAGARA_API const FName SetFullGridValueFunctionName;
+	static NIAGARA_API const FName GetFullGridPreviousValueFunctionName;
+	static NIAGARA_API const FName SamplePreviousFullGridFunctionName;
+	static NIAGARA_API const FName CubicSamplePreviousFullGridFunctionName;
+
+	static NIAGARA_API const FName SetVector4ValueFunctionName;
+	static NIAGARA_API const FName SetVector3ValueFunctionName;
+	static NIAGARA_API const FName SetVector2ValueFunctionName;
+	static NIAGARA_API const FName GetVector2ValueFunctionName;
+	static NIAGARA_API const FName SetFloatValueFunctionName;
+	static NIAGARA_API const FName GetPreviousValueAtIndexFunctionName;
+	static NIAGARA_API const FName SamplePreviousGridAtIndexFunctionName;
+	static NIAGARA_API const FName CubicSamplePreviousGridAtIndexFunctionName;
+
+	static NIAGARA_API const FName SetVector4ValueAtIndexFunctionName;
+	static NIAGARA_API const FName SetVector3ValueAtIndexFunctionName;
+	static NIAGARA_API const FName SetVector2ValueAtIndexFunctionName;
+	static NIAGARA_API const FName SetFloatValueAtIndexFunctionName;
+
+	static NIAGARA_API const FName GetPreviousVector4ValueAtIndexFunctionName;
+	static NIAGARA_API const FName GetPreviousVector3ValueAtIndexFunctionName;
+	static NIAGARA_API const FName GetPreviousVector2ValueAtIndexFunctionName;
+	static NIAGARA_API const FName GetPreviousFloatValueAtIndexFunctionName;
+
+	static NIAGARA_API const FName SamplePreviousGridVector4AtIndexFunctionName;
+	static NIAGARA_API const FName SamplePreviousGridVector3AtIndexFunctionName;
+	static NIAGARA_API const FName SamplePreviousGridVector2AtIndexFunctionName;
+	static NIAGARA_API const FName SamplePreviousGridFloatAtIndexFunctionName;
+
+	static NIAGARA_API const FName CubicSamplePreviousGridVector4AtIndexFunctionName;
+	static NIAGARA_API const FName CubicSamplePreviousGridVector3AtIndexFunctionName;
+	static NIAGARA_API const FName CubicSamplePreviousGridVector2AtIndexFunctionName;
+	static NIAGARA_API const FName CubicSamplePreviousGridFloatAtIndexFunctionName;
+
+	static NIAGARA_API const FName GetPreviousVector4ValueFunctionName;
+	static NIAGARA_API const FName SamplePreviousGridVector4FunctionName;
+	static NIAGARA_API const FName CubicSamplePreviousGridVector4FunctionName;
+	static NIAGARA_API const FName SetVectorValueFunctionName;
+	static NIAGARA_API const FName GetPreviousVectorValueFunctionName;
+	static NIAGARA_API const FName SamplePreviousGridVectorFunctionName;
+	static NIAGARA_API const FName CubicSamplePreviousGridVectorFunctionName;
+	static NIAGARA_API const FName SetVector2DValueFunctionName;
+	static NIAGARA_API const FName GetPreviousVector2DValueFunctionName;
+	static NIAGARA_API const FName SamplePreviousGridVector2DFunctionName;
+	static NIAGARA_API const FName CubicSamplePreviousGridVector2DFunctionName;
+	static NIAGARA_API const FName GetPreviousFloatValueFunctionName;
+	static NIAGARA_API const FName SamplePreviousGridFloatFunctionName;
+	static NIAGARA_API const FName CubicSamplePreviousGridFloatFunctionName;
+	static NIAGARA_API const FString AttributeIndicesBaseName;
+	static NIAGARA_API const FString PerAttributeDataName;
+	static NIAGARA_API const TCHAR* VectorComponentNames[];
+
+	static NIAGARA_API const FName SetNumCellsFunctionName;
+
+	static NIAGARA_API const FName GetVector4AttributeIndexFunctionName;
+	static NIAGARA_API const FName GetVectorAttributeIndexFunctionName;
+	static NIAGARA_API const FName GetVector2DAttributeIndexFunctionName;
+	static NIAGARA_API const FName GetFloatAttributeIndexFunctionName;
+
+	static NIAGARA_API const FString AnonymousAttributeString;
+
+	static NIAGARA_API const TCHAR* TemplateShaderFilePath;
 
 #if WITH_EDITOR
 	virtual bool SupportsSetupAndTeardownHLSL() const { return true; }
-	virtual bool GenerateSetupHLSL(FNiagaraDataInterfaceGPUParamInfo& DIInstanceInfo, TConstArrayView<FNiagaraVariable> InArguments, bool bSpawnOnly, bool bPartialWrites, TArray<FText>& OutErrors, FString& OutHLSL) const;
-	virtual bool GenerateTeardownHLSL(FNiagaraDataInterfaceGPUParamInfo& DIInstanceInfo, TConstArrayView<FNiagaraVariable> InArguments, bool bSpawnOnly, bool bPartialWrites, TArray<FText>& OutErrors, FString& OutHLSL) const;
+	NIAGARA_API virtual bool GenerateSetupHLSL(FNiagaraDataInterfaceGPUParamInfo& DIInstanceInfo, TConstArrayView<FNiagaraVariable> InArguments, bool bSpawnOnly, bool bPartialWrites, TArray<FText>& OutErrors, FString& OutHLSL) const;
+	NIAGARA_API virtual bool GenerateTeardownHLSL(FNiagaraDataInterfaceGPUParamInfo& DIInstanceInfo, TConstArrayView<FNiagaraVariable> InArguments, bool bSpawnOnly, bool bPartialWrites, TArray<FText>& OutErrors, FString& OutHLSL) const;
 	virtual bool SupportsIterationSourceNamespaceAttributesHLSL() const override { return true; }
-	virtual bool GenerateIterationSourceNamespaceReadAttributesHLSL(FNiagaraDataInterfaceGPUParamInfo& DIInstanceInfo, const FNiagaraVariable& IterationSourceVar, TConstArrayView<FNiagaraVariable> InArguments, TConstArrayView<FNiagaraVariable> InAttributes, TConstArrayView<FString> InAttributeHLSLNames, bool bInSetToDefaults, bool bSpawnOnly, bool bPartialWrites, TArray<FText>& OutErrors, FString& OutHLSL) const override;
-	virtual bool GenerateIterationSourceNamespaceWriteAttributesHLSL(FNiagaraDataInterfaceGPUParamInfo& DIInstanceInfo, const FNiagaraVariable& IterationSourceVar, TConstArrayView<FNiagaraVariable> InArguments, TConstArrayView<FNiagaraVariable> InAttributes, TConstArrayView<FString> InAttributeHLSLNames, TConstArrayView<FNiagaraVariable> InAllAttributes, bool bSpawnOnly, bool bPartialWrites, TArray<FText>& OutErrors, FString& OutHLSL) const override;
+	NIAGARA_API virtual bool GenerateIterationSourceNamespaceReadAttributesHLSL(FNiagaraDataInterfaceGPUParamInfo& DIInstanceInfo, const FNiagaraVariable& IterationSourceVar, TConstArrayView<FNiagaraVariable> InArguments, TConstArrayView<FNiagaraVariable> InAttributes, TConstArrayView<FString> InAttributeHLSLNames, bool bInSetToDefaults, bool bSpawnOnly, bool bPartialWrites, TArray<FText>& OutErrors, FString& OutHLSL) const override;
+	NIAGARA_API virtual bool GenerateIterationSourceNamespaceWriteAttributesHLSL(FNiagaraDataInterfaceGPUParamInfo& DIInstanceInfo, const FNiagaraVariable& IterationSourceVar, TConstArrayView<FNiagaraVariable> InArguments, TConstArrayView<FNiagaraVariable> InAttributes, TConstArrayView<FString> InAttributeHLSLNames, TConstArrayView<FNiagaraVariable> InAllAttributes, bool bSpawnOnly, bool bPartialWrites, TArray<FText>& OutErrors, FString& OutHLSL) const override;
 #endif
 
-	static int32 GetComponentCountFromFuncName(const FName& FuncName);
-	static FNiagaraTypeDefinition GetValueTypeFromFuncName(const FName& FuncName);
-	static bool CanCreateVarFromFuncName(const FName& FuncName);
+	static NIAGARA_API int32 GetComponentCountFromFuncName(const FName& FuncName);
+	static NIAGARA_API FNiagaraTypeDefinition GetValueTypeFromFuncName(const FName& FuncName);
+	static NIAGARA_API bool CanCreateVarFromFuncName(const FName& FuncName);
 
 	TMap<FNiagaraSystemInstanceID, FGrid3DCollectionRWInstanceData_GameThread*>& GetSystemInstancesToProxyData_GT() { return SystemInstancesToProxyData_GT; }
 protected:
 #if WITH_EDITORONLY_DATA
-	void WriteSetHLSL(const FNiagaraDataInterfaceGPUParamInfo& ParamInfo, const FNiagaraDataInterfaceGeneratedFunction& FunctionInfo, int FunctionInstanceIndex, int32 InNumChannels, FString& OutHLSL);
-	void WriteSetAtIndexHLSL(const FNiagaraDataInterfaceGPUParamInfo& ParamInfo, const FNiagaraDataInterfaceGeneratedFunction& FunctionInfo, int32 InNumChannels, FString& OutHLSL);
-	void WriteGetHLSL(const FNiagaraDataInterfaceGPUParamInfo& ParamInfo, const FNiagaraDataInterfaceGeneratedFunction& FunctionInfo, int FunctionInstanceIndex, int32 InNumChannels, FString& OutHLSL);
-	void WriteGetAtIndexHLSL(const FNiagaraDataInterfaceGPUParamInfo& ParamInfo, const FNiagaraDataInterfaceGeneratedFunction& FunctionInfo, int32 InNumChannels, FString& OutHLSL);
-	void WriteSampleHLSL(const FNiagaraDataInterfaceGPUParamInfo& ParamInfo, const FNiagaraDataInterfaceGeneratedFunction& FunctionInfo, int FunctionInstanceIndex, int32 InNumChannels, FString SampleFunction, FString& OutHLSL);
-	void WriteSampleAtIndexHLSL(const FNiagaraDataInterfaceGPUParamInfo& ParamInfo, const FNiagaraDataInterfaceGeneratedFunction& FunctionInfo, int32 InNumChannels, FString SampleFunction, FString& OutHLSL);
-	void WriteAttributeGetIndexHLSL(const FNiagaraDataInterfaceGPUParamInfo& ParamInfo, const FNiagaraDataInterfaceGeneratedFunction& FunctionInfo, int FunctionInstanceIndex, int32 InNumChannels, FString& OutHLSL);
-
-	const TCHAR* TypeDefinitionToHLSLTypeString(const FNiagaraTypeDefinition& InDef) const;
-	FName TypeDefinitionToGetFunctionName(const FNiagaraTypeDefinition& InDef) const;
-	FName TypeDefinitionToSetFunctionName(const FNiagaraTypeDefinition& InDef) const;
-	FName TypeDefinitionToAttributeIndexFunctionName(const FNiagaraTypeDefinition& InDef) const;
+	NIAGARA_API void WriteSetHLSL(const FNiagaraDataInterfaceGPUParamInfo& ParamInfo, const FNiagaraDataInterfaceGeneratedFunction& FunctionInfo, int FunctionInstanceIndex, int32 InNumChannels, FString& OutHLSL);
+	NIAGARA_API void WriteSetAtIndexHLSL(const FNiagaraDataInterfaceGPUParamInfo& ParamInfo, const FNiagaraDataInterfaceGeneratedFunction& FunctionInfo, int32 InNumChannels, FString& OutHLSL);
+	NIAGARA_API void WriteGetHLSL(const FNiagaraDataInterfaceGPUParamInfo& ParamInfo, const FNiagaraDataInterfaceGeneratedFunction& FunctionInfo, int FunctionInstanceIndex, int32 InNumChannels, FString& OutHLSL);
+	NIAGARA_API void WriteGetAtIndexHLSL(const FNiagaraDataInterfaceGPUParamInfo& ParamInfo, const FNiagaraDataInterfaceGeneratedFunction& FunctionInfo, int32 InNumChannels, FString& OutHLSL);
+	NIAGARA_API void WriteSampleHLSL(const FNiagaraDataInterfaceGPUParamInfo& ParamInfo, const FNiagaraDataInterfaceGeneratedFunction& FunctionInfo, int FunctionInstanceIndex, int32 InNumChannels, FString SampleFunction, FString& OutHLSL);
+	NIAGARA_API void WriteSampleAtIndexHLSL(const FNiagaraDataInterfaceGPUParamInfo& ParamInfo, const FNiagaraDataInterfaceGeneratedFunction& FunctionInfo, int32 InNumChannels, FString SampleFunction, FString& OutHLSL);
+	NIAGARA_API void WriteAttributeGetIndexHLSL(const FNiagaraDataInterfaceGPUParamInfo& ParamInfo, const FNiagaraDataInterfaceGeneratedFunction& FunctionInfo, int FunctionInstanceIndex, int32 InNumChannels, FString& OutHLSL);
+
+	NIAGARA_API const TCHAR* TypeDefinitionToHLSLTypeString(const FNiagaraTypeDefinition& InDef) const;
+	NIAGARA_API FName TypeDefinitionToGetFunctionName(const FNiagaraTypeDefinition& InDef) const;
+	NIAGARA_API FName TypeDefinitionToSetFunctionName(const FNiagaraTypeDefinition& InDef) const;
+	NIAGARA_API FName TypeDefinitionToAttributeIndexFunctionName(const FNiagaraTypeDefinition& InDef) const;
 #endif
 
 	//~ UNiagaraDataInterface interface
-	virtual bool CopyToInternal(UNiagaraDataInterface* Destination) const override;
+	NIAGARA_API virtual bool CopyToInternal(UNiagaraDataInterface* Destination) const override;
 	//~ UNiagaraDataInterface interface END
 
 	TMap<FNiagaraSystemInstanceID, FGrid3DCollectionRWInstanceData_GameThread*> SystemInstancesToProxyData_GT;
 
-	static FNiagaraVariableBase ExposedRTVar;
+	static NIAGARA_API FNiagaraVariableBase ExposedRTVar;
 
 public:
 	/** Finds all attributes by locating the data interface amongst the parameter stores. */
-	void FindAttributes(TArray<FNiagaraVariableBase>& OutVariables, TArray<uint32>& OutVariableOffsets, int32& OutNumAttribChannelsFound, TArray<FText>* OutWarnings = nullptr) const;
+	NIAGARA_API void FindAttributes(TArray<FNiagaraVariableBase>& OutVariables, TArray<uint32>& OutVariableOffsets, int32& OutNumAttribChannelsFound, TArray<FText>* OutWarnings = nullptr) const;
 protected:
-	static void CollectAttributesForScript(UNiagaraScript* Script, const UNiagaraDataInterface* DataInterface, TArray<FNiagaraVariableBase>& OutVariables, TArray<uint32>& OutVariableOffsets, int32& TotalAttributes, TArray<FText>* OutWarnings = nullptr);
-};
\ No newline at end of file
+	static NIAGARA_API void CollectAttributesForScript(UNiagaraScript* Script, const UNiagaraDataInterface* DataInterface, TArray<FNiagaraVariableBase>& OutVariables, TArray<uint32>& OutVariableOffsets, int32& TotalAttributes, TArray<FText>* OutWarnings = nullptr);
+};
diff --git a/Engine/Plugins/FX/Niagara/Source/Niagara/Classes/NiagaraDataInterfaceGrid3DCollectionReader.h b/Engine/Plugins/FX/Niagara/Source/Niagara/Classes/NiagaraDataInterfaceGrid3DCollectionReader.h
--- a/Engine/Plugins/FX/Niagara/Source/Niagara/Classes/NiagaraDataInterfaceGrid3DCollectionReader.h
+++ b/Engine/Plugins/FX/Niagara/Source/Niagara/Classes/NiagaraDataInterfaceGrid3DCollectionReader.h
@@ -13,8 +13,8 @@ class FNiagaraSystemInstance;
 class UTextureRenderTarget;
 class UTextureRenderTargetVolume;
 
-UCLASS(EditInlineNew, Category = "Grid", CollapseCategories, hidecategories=(Grid3DCollection,Grid, Deprecated), meta = (DisplayName = "Grid3D Collection Reader"), Blueprintable, BlueprintType)
-class NIAGARA_API UNiagaraDataInterfaceGrid3DCollectionReader : public UNiagaraDataInterfaceGrid3DCollection
+UCLASS(EditInlineNew, Category = "Grid", CollapseCategories, hidecategories=(Grid3DCollection,Grid, Deprecated), meta = (DisplayName = "Grid3D Collection Reader"), Blueprintable, BlueprintType, MinimalAPI)
+class UNiagaraDataInterfaceGrid3DCollectionReader : public UNiagaraDataInterfaceGrid3DCollection
 {
 	GENERATED_UCLASS_BODY()
 
@@ -27,14 +27,14 @@ public:
 	UPROPERTY(EditAnywhere, Category = "Reader")
 	FString DIName;
 
-	virtual bool Equals(const UNiagaraDataInterface* Other) const override;
-	virtual bool InitPerInstanceData(void* PerInstanceData, FNiagaraSystemInstance* SystemInstance) override;
-	virtual void GetEmitterDependencies(UNiagaraSystem* Asset, TArray<FVersionedNiagaraEmitter>& Dependencies) const override;
-	virtual void GetFunctions(TArray<FNiagaraFunctionSignature>& OutFunctions) override;
-	virtual bool PerInstanceTickPostSimulate(void* PerInstanceData, FNiagaraSystemInstance* SystemInstance, float DeltaSeconds) override;
+	NIAGARA_API virtual bool Equals(const UNiagaraDataInterface* Other) const override;
+	NIAGARA_API virtual bool InitPerInstanceData(void* PerInstanceData, FNiagaraSystemInstance* SystemInstance) override;
+	NIAGARA_API virtual void GetEmitterDependencies(UNiagaraSystem* Asset, TArray<FVersionedNiagaraEmitter>& Dependencies) const override;
+	NIAGARA_API virtual void GetFunctions(TArray<FNiagaraFunctionSignature>& OutFunctions) override;
+	NIAGARA_API virtual bool PerInstanceTickPostSimulate(void* PerInstanceData, FNiagaraSystemInstance* SystemInstance, float DeltaSeconds) override;
 
 protected:
 	//~ UNiagaraDataInterface interface
-	virtual bool CopyToInternal(UNiagaraDataInterface* Destination) const override;
+	NIAGARA_API virtual bool CopyToInternal(UNiagaraDataInterface* Destination) const override;
 	//~ UNiagaraDataInterface interface END
-};
\ No newline at end of file
+};
