[Niagara]

diff --git a/Engine/Plugins/FX/Niagara/Source/Niagara/Classes/NiagaraDataInterfaceIntRenderTarget2D.h b/Engine/Plugins/FX/Niagara/Source/Niagara/Classes/NiagaraDataInterfaceIntRenderTarget2D.h
--- a/Engine/Plugins/FX/Niagara/Source/Niagara/Classes/NiagaraDataInterfaceIntRenderTarget2D.h
+++ b/Engine/Plugins/FX/Niagara/Source/Niagara/Classes/NiagaraDataInterfaceIntRenderTarget2D.h
@@ -9,8 +9,8 @@
 
 class UTextureRenderTarget2D;
 
-UCLASS(EditInlineNew, Category = "Render Target", CollapseCategories, meta = (DisplayName = "Integer Render Target 2D", Experimental), Blueprintable, BlueprintType)
-class NIAGARA_API UNiagaraDataInterfaceIntRenderTarget2D : public UNiagaraDataInterfaceRWBase
+UCLASS(EditInlineNew, Category = "Render Target", CollapseCategories, meta = (DisplayName = "Integer Render Target 2D", Experimental), Blueprintable, BlueprintType, MinimalAPI)
+class UNiagaraDataInterfaceIntRenderTarget2D : public UNiagaraDataInterfaceRWBase
 {
 	GENERATED_UCLASS_BODY()
 
@@ -20,43 +20,43 @@ class NIAGARA_API UNiagaraDataInterfaceIntRenderTarget2D : public UNiagaraDataIn
 	END_SHADER_PARAMETER_STRUCT()
 
 public:
-	virtual void PostInitProperties() override;
+	NIAGARA_API virtual void PostInitProperties() override;
 	
 	//~ UNiagaraDataInterface interface
 	virtual bool CanExecuteOnTarget(ENiagaraSimTarget Target)const override { return true; }
-	virtual void GetFunctions(TArray<FNiagaraFunctionSignature>& OutFunctions) override;
+	NIAGARA_API virtual void GetFunctions(TArray<FNiagaraFunctionSignature>& OutFunctions) override;
 #if WITH_EDITORONLY_DATA
-	virtual bool UpgradeFunctionCall(FNiagaraFunctionSignature& FunctionSignature) override;
+	NIAGARA_API virtual bool UpgradeFunctionCall(FNiagaraFunctionSignature& FunctionSignature) override;
 #endif
-	virtual void GetVMExternalFunction(const FVMExternalFunctionBindingInfo& BindingInfo, void* InstanceData, FVMExternalFunction &OutFunc) override;
+	NIAGARA_API virtual void GetVMExternalFunction(const FVMExternalFunctionBindingInfo& BindingInfo, void* InstanceData, FVMExternalFunction &OutFunc) override;
 
-	virtual bool Equals(const UNiagaraDataInterface* Other) const override;
-	virtual bool CopyToInternal(UNiagaraDataInterface* Destination) const override;
+	NIAGARA_API virtual bool Equals(const UNiagaraDataInterface* Other) const override;
+	NIAGARA_API virtual bool CopyToInternal(UNiagaraDataInterface* Destination) const override;
 
 #if WITH_EDITORONLY_DATA
-	virtual bool AppendCompileHash(FNiagaraCompileHashVisitor* InVisitor) const override;
-	virtual void GetParameterDefinitionHLSL(const FNiagaraDataInterfaceGPUParamInfo& ParamInfo, FString& OutHLSL) override;
-	virtual bool GetFunctionHLSL(const FNiagaraDataInterfaceGPUParamInfo& ParamInfo, const FNiagaraDataInterfaceGeneratedFunction& FunctionInfo, int FunctionInstanceIndex, FString& OutHLSL) override;
+	NIAGARA_API virtual bool AppendCompileHash(FNiagaraCompileHashVisitor* InVisitor) const override;
+	NIAGARA_API virtual void GetParameterDefinitionHLSL(const FNiagaraDataInterfaceGPUParamInfo& ParamInfo, FString& OutHLSL) override;
+	NIAGARA_API virtual bool GetFunctionHLSL(const FNiagaraDataInterfaceGPUParamInfo& ParamInfo, const FNiagaraDataInterfaceGeneratedFunction& FunctionInfo, int FunctionInstanceIndex, FString& OutHLSL) override;
 #endif
-	virtual void BuildShaderParameters(FNiagaraShaderParametersBuilder& ShaderParametersBuilder) const override;
-	virtual void SetShaderParameters(const FNiagaraDataInterfaceSetShaderParametersContext& Context) const override;
+	NIAGARA_API virtual void BuildShaderParameters(FNiagaraShaderParametersBuilder& ShaderParametersBuilder) const override;
+	NIAGARA_API virtual void SetShaderParameters(const FNiagaraDataInterfaceSetShaderParametersContext& Context) const override;
 
 	virtual void ProvidePerInstanceDataForRenderThread(void* DataForRenderThread, void* PerInstanceData, const FNiagaraSystemInstanceID& SystemInstance) override {}
-	virtual bool InitPerInstanceData(void* PerInstanceData, FNiagaraSystemInstance* SystemInstance) override;
-	virtual void DestroyPerInstanceData(void* PerInstanceData, FNiagaraSystemInstance* SystemInstance) override;
-	virtual int32 PerInstanceDataSize() const override;
-	virtual bool PerInstanceTickPostSimulate(void* PerInstanceData, FNiagaraSystemInstance* InSystemInstance, float DeltaSeconds) override;
+	NIAGARA_API virtual bool InitPerInstanceData(void* PerInstanceData, FNiagaraSystemInstance* SystemInstance) override;
+	NIAGARA_API virtual void DestroyPerInstanceData(void* PerInstanceData, FNiagaraSystemInstance* SystemInstance) override;
+	NIAGARA_API virtual int32 PerInstanceDataSize() const override;
+	NIAGARA_API virtual bool PerInstanceTickPostSimulate(void* PerInstanceData, FNiagaraSystemInstance* InSystemInstance, float DeltaSeconds) override;
 	virtual bool HasPostSimulateTick() const override { return true; }
 
 	virtual bool CanExposeVariables() const override { return true;}
-	virtual void GetExposedVariables(TArray<FNiagaraVariableBase>& OutVariables) const override;
-	virtual bool GetExposedVariableValue(const FNiagaraVariableBase& InVariable, void* InPerInstanceData, FNiagaraSystemInstance* InSystemInstance, void* OutData) const override;
+	NIAGARA_API virtual void GetExposedVariables(TArray<FNiagaraVariableBase>& OutVariables) const override;
+	NIAGARA_API virtual bool GetExposedVariableValue(const FNiagaraVariableBase& InVariable, void* InPerInstanceData, FNiagaraSystemInstance* InSystemInstance, void* OutData) const override;
 	//~ UNiagaraDataInterface interface END
 
-	bool UpdateInstanceTexture(FNiagaraSystemInstance* SystemInstance, struct FNDIIntRenderTarget2DInstanceData_GameThread* InstanceData);
+	NIAGARA_API bool UpdateInstanceTexture(FNiagaraSystemInstance* SystemInstance, struct FNDIIntRenderTarget2DInstanceData_GameThread* InstanceData);
 
-	void VMGetSize(FVectorVMExternalFunctionContext& Context);
-	void VMSetSize(FVectorVMExternalFunctionContext& Context);
+	NIAGARA_API void VMGetSize(FVectorVMExternalFunctionContext& Context);
+	NIAGARA_API void VMSetSize(FVectorVMExternalFunctionContext& Context);
 
 	UPROPERTY(EditAnywhere, Category = "Render Target", meta = (DisplayPriority = 1))
 	FIntPoint Size = FIntPoint::ZeroValue;
@@ -74,7 +74,7 @@ public:
 	FNiagaraUserParameterBinding RenderTargetUserParameter;
 
 protected:
-	static FNiagaraVariableBase ExposedRTVar;
+	static NIAGARA_API FNiagaraVariableBase ExposedRTVar;
 
 	UPROPERTY(Transient, DuplicateTransient)
 	TMap<uint64, TObjectPtr<UTextureRenderTarget2D>> ManagedRenderTargets;
diff --git a/Engine/Plugins/FX/Niagara/Source/Niagara/Classes/NiagaraDataInterfaceLandscape.h b/Engine/Plugins/FX/Niagara/Source/Niagara/Classes/NiagaraDataInterfaceLandscape.h
--- a/Engine/Plugins/FX/Niagara/Source/Niagara/Classes/NiagaraDataInterfaceLandscape.h
+++ b/Engine/Plugins/FX/Niagara/Source/Niagara/Classes/NiagaraDataInterfaceLandscape.h
@@ -34,8 +34,8 @@ enum class ENDILandscape_SourceMode : uint8
 };
 
 /** Data Interface allowing sampling of a Landscape */
-UCLASS(EditInlineNew, Category = "Landscape", CollapseCategories, meta = (DisplayName = "Landscape Sample"))
-class NIAGARA_API UNiagaraDataInterfaceLandscape : public UNiagaraDataInterface
+UCLASS(EditInlineNew, Category = "Landscape", CollapseCategories, meta = (DisplayName = "Landscape Sample"), MinimalAPI)
+class UNiagaraDataInterfaceLandscape : public UNiagaraDataInterface
 {
 	GENERATED_UCLASS_BODY()
 
@@ -50,42 +50,42 @@ public:
 	TArray<TObjectPtr<UPhysicalMaterial>> PhysicalMaterials;
 
 	//UObject Interface
-	virtual void PostInitProperties() override;	
+	NIAGARA_API virtual void PostInitProperties() override;	
 	//UObject Interface End
 
 	//UNiagaraDataInterface Interface
-	virtual void GetFunctions(TArray<FNiagaraFunctionSignature>& OutFunctions) override;
+	NIAGARA_API virtual void GetFunctions(TArray<FNiagaraFunctionSignature>& OutFunctions) override;
 	virtual bool CanExecuteOnTarget(ENiagaraSimTarget Target) const override { return Target == ENiagaraSimTarget::GPUComputeSim; }
 #if WITH_EDITORONLY_DATA
-	virtual bool UpgradeFunctionCall(FNiagaraFunctionSignature& FunctionSignature) override;
+	NIAGARA_API virtual bool UpgradeFunctionCall(FNiagaraFunctionSignature& FunctionSignature) override;
 #endif
 	//UNiagaraDataInterface Interface End
 
-	virtual bool Equals(const UNiagaraDataInterface* Other) const override;
+	NIAGARA_API virtual bool Equals(const UNiagaraDataInterface* Other) const override;
 
 	// GPU sim functionality
 #if WITH_EDITORONLY_DATA
-	virtual void GetParameterDefinitionHLSL(const FNiagaraDataInterfaceGPUParamInfo& ParamInfo, FString& OutHLSL) override;
-	virtual bool GetFunctionHLSL(const FNiagaraDataInterfaceGPUParamInfo& ParamInfo, const FNiagaraDataInterfaceGeneratedFunction& FunctionInfo, int FunctionInstanceIndex, FString& OutHLSL) override;
-	virtual bool AppendCompileHash(FNiagaraCompileHashVisitor* InVisitor) const override;
+	NIAGARA_API virtual void GetParameterDefinitionHLSL(const FNiagaraDataInterfaceGPUParamInfo& ParamInfo, FString& OutHLSL) override;
+	NIAGARA_API virtual bool GetFunctionHLSL(const FNiagaraDataInterfaceGPUParamInfo& ParamInfo, const FNiagaraDataInterfaceGeneratedFunction& FunctionInfo, int FunctionInstanceIndex, FString& OutHLSL) override;
+	NIAGARA_API virtual bool AppendCompileHash(FNiagaraCompileHashVisitor* InVisitor) const override;
 #endif
-	virtual void BuildShaderParameters(FNiagaraShaderParametersBuilder& ShaderParametersBuilder) const override;
-	virtual void SetShaderParameters(const FNiagaraDataInterfaceSetShaderParametersContext& Context) const override;
-
-	virtual void ProvidePerInstanceDataForRenderThread(void* DataForRenderThread, void* PerInstanceData, const FNiagaraSystemInstanceID& SystemInstance) override;
-	virtual bool InitPerInstanceData(void* PerInstanceData, FNiagaraSystemInstance* SystemInstance) override;
-	virtual void DestroyPerInstanceData(void* PerInstanceData, FNiagaraSystemInstance* SystemInstance) override;
-	virtual int32 PerInstanceDataSize() const override;
-	virtual bool PerInstanceTick(void* PerInstanceData, FNiagaraSystemInstance* SystemInstance, float DeltaSeconds) override;
+	NIAGARA_API virtual void BuildShaderParameters(FNiagaraShaderParametersBuilder& ShaderParametersBuilder) const override;
+	NIAGARA_API virtual void SetShaderParameters(const FNiagaraDataInterfaceSetShaderParametersContext& Context) const override;
+
+	NIAGARA_API virtual void ProvidePerInstanceDataForRenderThread(void* DataForRenderThread, void* PerInstanceData, const FNiagaraSystemInstanceID& SystemInstance) override;
+	NIAGARA_API virtual bool InitPerInstanceData(void* PerInstanceData, FNiagaraSystemInstance* SystemInstance) override;
+	NIAGARA_API virtual void DestroyPerInstanceData(void* PerInstanceData, FNiagaraSystemInstance* SystemInstance) override;
+	NIAGARA_API virtual int32 PerInstanceDataSize() const override;
+	NIAGARA_API virtual bool PerInstanceTick(void* PerInstanceData, FNiagaraSystemInstance* SystemInstance, float DeltaSeconds) override;
 	virtual bool HasPreSimulateTick() const override { return true; }
 
 protected:
-	virtual bool CopyToInternal(UNiagaraDataInterface* Destination) const override;
-	void ApplyLandscape(const FNiagaraSystemInstance& SystemInstance, FNDILandscapeData_GameThread& InstanceData) const;
-	ALandscape* GetLandscape(const FNiagaraSystemInstance& SystemInstance, FNDILandscapeData_GameThread& Hint) const;
-
-	static const FName GetBaseColorName;
-	static const FName GetHeightName;
-	static const FName GetWorldNormalName;
-	static const FName GetPhysicalMaterialIndexName;
+	NIAGARA_API virtual bool CopyToInternal(UNiagaraDataInterface* Destination) const override;
+	NIAGARA_API void ApplyLandscape(const FNiagaraSystemInstance& SystemInstance, FNDILandscapeData_GameThread& InstanceData) const;
+	NIAGARA_API ALandscape* GetLandscape(const FNiagaraSystemInstance& SystemInstance, FNDILandscapeData_GameThread& Hint) const;
+
+	static NIAGARA_API const FName GetBaseColorName;
+	static NIAGARA_API const FName GetHeightName;
+	static NIAGARA_API const FName GetWorldNormalName;
+	static NIAGARA_API const FName GetPhysicalMaterialIndexName;
 };
diff --git a/Engine/Plugins/FX/Niagara/Source/Niagara/Classes/NiagaraDataInterfaceMeshRendererInfo.h b/Engine/Plugins/FX/Niagara/Source/Niagara/Classes/NiagaraDataInterfaceMeshRendererInfo.h
--- a/Engine/Plugins/FX/Niagara/Source/Niagara/Classes/NiagaraDataInterfaceMeshRendererInfo.h
+++ b/Engine/Plugins/FX/Niagara/Source/Niagara/Classes/NiagaraDataInterfaceMeshRendererInfo.h
@@ -11,8 +11,8 @@ class UNiagaraMeshRendererProperties;
 class FNDIMeshRendererInfo;
 
 /** This Data Interface can be used to query information about the mesh renderers of an emitter */
-UCLASS(EditInlineNew, Category = "Mesh Particles", CollapseCategories, meta = (DisplayName = "Mesh Renderer Info"))
-class NIAGARA_API UNiagaraDataInterfaceMeshRendererInfo : public UNiagaraDataInterface
+UCLASS(EditInlineNew, Category = "Mesh Particles", CollapseCategories, meta = (DisplayName = "Mesh Renderer Info"), MinimalAPI)
+class UNiagaraDataInterfaceMeshRendererInfo : public UNiagaraDataInterface
 {
 	GENERATED_UCLASS_BODY()
 
@@ -26,45 +26,45 @@ public:
 	UNiagaraMeshRendererProperties* GetMeshRenderer() const { return MeshRenderer; }
 
 	//UObject Interface
-	virtual void PostInitProperties() override;
-	virtual void PostLoad() override;
-	virtual void BeginDestroy() override;
+	NIAGARA_API virtual void PostInitProperties() override;
+	NIAGARA_API virtual void PostLoad() override;
+	NIAGARA_API virtual void BeginDestroy() override;
 #if WITH_EDITOR	
-	virtual void PreEditChange(FProperty* PropertyAboutToChange) override;
-	virtual void PostEditChangeProperty(struct FPropertyChangedEvent& PropertyChangedEvent) override;
+	NIAGARA_API virtual void PreEditChange(FProperty* PropertyAboutToChange) override;
+	NIAGARA_API virtual void PostEditChangeProperty(struct FPropertyChangedEvent& PropertyChangedEvent) override;
 #endif 
 	//UObject Interface End
 
-	void OnMeshRendererChanged(UNiagaraMeshRendererProperties* NewMeshRenderer);
+	NIAGARA_API void OnMeshRendererChanged(UNiagaraMeshRendererProperties* NewMeshRenderer);
 
 	//UNiagaraDataInterface Interface
-	virtual void GetFunctions(TArray<FNiagaraFunctionSignature>& OutFunctions) override;
-	virtual void GetVMExternalFunction(const FVMExternalFunctionBindingInfo& BindingInfo, void* InstanceData, FVMExternalFunction& OutFunc) override;
-	virtual bool Equals(const UNiagaraDataInterface* Other) const override;
+	NIAGARA_API virtual void GetFunctions(TArray<FNiagaraFunctionSignature>& OutFunctions) override;
+	NIAGARA_API virtual void GetVMExternalFunction(const FVMExternalFunctionBindingInfo& BindingInfo, void* InstanceData, FVMExternalFunction& OutFunc) override;
+	NIAGARA_API virtual bool Equals(const UNiagaraDataInterface* Other) const override;
 #if WITH_EDITORONLY_DATA
-	virtual bool AppendCompileHash(FNiagaraCompileHashVisitor* InVisitor) const override;
-	virtual void GetParameterDefinitionHLSL(const FNiagaraDataInterfaceGPUParamInfo& ParamInfo, FString& OutHLSL) override;
-	virtual bool GetFunctionHLSL(const FNiagaraDataInterfaceGPUParamInfo& ParamInfo, const FNiagaraDataInterfaceGeneratedFunction& FunctionInfo, int FunctionInstanceIndex, FString& OutHLSL) override;
+	NIAGARA_API virtual bool AppendCompileHash(FNiagaraCompileHashVisitor* InVisitor) const override;
+	NIAGARA_API virtual void GetParameterDefinitionHLSL(const FNiagaraDataInterfaceGPUParamInfo& ParamInfo, FString& OutHLSL) override;
+	NIAGARA_API virtual bool GetFunctionHLSL(const FNiagaraDataInterfaceGPUParamInfo& ParamInfo, const FNiagaraDataInterfaceGeneratedFunction& FunctionInfo, int FunctionInstanceIndex, FString& OutHLSL) override;
 #endif
-	virtual void BuildShaderParameters(FNiagaraShaderParametersBuilder& ShaderParametersBuilder) const override;
-	virtual void SetShaderParameters(const FNiagaraDataInterfaceSetShaderParametersContext& Context) const override;
+	NIAGARA_API virtual void BuildShaderParameters(FNiagaraShaderParametersBuilder& ShaderParametersBuilder) const override;
+	NIAGARA_API virtual void SetShaderParameters(const FNiagaraDataInterfaceSetShaderParametersContext& Context) const override;
 
 	virtual bool CanExecuteOnTarget(ENiagaraSimTarget Target) const override { return true; }
 #if WITH_EDITOR
-	virtual bool UpgradeFunctionCall(FNiagaraFunctionSignature& FunctionSignature) override;
-	virtual void GetFeedback(UNiagaraSystem* InAsset, UNiagaraComponent* InComponent, TArray<FNiagaraDataInterfaceError>& OutErrors, TArray<FNiagaraDataInterfaceFeedback>& OutWarnings,
+	NIAGARA_API virtual bool UpgradeFunctionCall(FNiagaraFunctionSignature& FunctionSignature) override;
+	NIAGARA_API virtual void GetFeedback(UNiagaraSystem* InAsset, UNiagaraComponent* InComponent, TArray<FNiagaraDataInterfaceError>& OutErrors, TArray<FNiagaraDataInterfaceFeedback>& OutWarnings,
 		TArray<FNiagaraDataInterfaceFeedback>& OutInfo) override;
 #endif
 	//UNiagaraDataInterface Interface	
 
 protected:
-	virtual bool CopyToInternal(UNiagaraDataInterface* Destination) const override;
-	virtual void PushToRenderThreadImpl() override;
-	void UpdateCachedData();
+	NIAGARA_API virtual bool CopyToInternal(UNiagaraDataInterface* Destination) const override;
+	NIAGARA_API virtual void PushToRenderThreadImpl() override;
+	NIAGARA_API void UpdateCachedData();
 
-	void VMGetNumMeshes(FVectorVMExternalFunctionContext& Context);
-	void VMGetMeshLocalBounds(FVectorVMExternalFunctionContext& Context);
-	void VMGetSubUVDetails(FVectorVMExternalFunctionContext& Context);
+	NIAGARA_API void VMGetNumMeshes(FVectorVMExternalFunctionContext& Context);
+	NIAGARA_API void VMGetMeshLocalBounds(FVectorVMExternalFunctionContext& Context);
+	NIAGARA_API void VMGetSubUVDetails(FVectorVMExternalFunctionContext& Context);
 
 	/** The name of the mesh renderer */
 	UPROPERTY(EditAnywhere, Category = "Source")
diff --git a/Engine/Plugins/FX/Niagara/Source/Niagara/Classes/NiagaraDataInterfaceNeighborGrid3D.h b/Engine/Plugins/FX/Niagara/Source/Niagara/Classes/NiagaraDataInterfaceNeighborGrid3D.h
--- a/Engine/Plugins/FX/Niagara/Source/Niagara/Classes/NiagaraDataInterfaceNeighborGrid3D.h
+++ b/Engine/Plugins/FX/Niagara/Source/Niagara/Classes/NiagaraDataInterfaceNeighborGrid3D.h
@@ -56,8 +56,8 @@ struct FNiagaraDataInterfaceProxyNeighborGrid3D : public FNiagaraDataInterfacePr
 	TMap<FNiagaraSystemInstanceID, FNDINeighborGrid3DInstanceData_RT> SystemInstancesToProxyData_RT;
 };
 
-UCLASS(EditInlineNew, Category = "Grid", CollapseCategories, meta = (DisplayName = "Neighbor Grid3D"))
-class NIAGARA_API UNiagaraDataInterfaceNeighborGrid3D : public UNiagaraDataInterfaceGrid3D
+UCLASS(EditInlineNew, Category = "Grid", CollapseCategories, meta = (DisplayName = "Neighbor Grid3D"), MinimalAPI)
+class UNiagaraDataInterfaceNeighborGrid3D : public UNiagaraDataInterfaceGrid3D
 {
 	GENERATED_UCLASS_BODY()
 
@@ -92,44 +92,44 @@ public:
 
 	//~ UNiagaraDataInterface interface
 	// VM functionality
-	virtual void GetFunctions(TArray<FNiagaraFunctionSignature>& OutFunctions) override;
-	virtual void GetVMExternalFunction(const FVMExternalFunctionBindingInfo& BindingInfo, void* InstanceData, FVMExternalFunction &OutFunc) override;
+	NIAGARA_API virtual void GetFunctions(TArray<FNiagaraFunctionSignature>& OutFunctions) override;
+	NIAGARA_API virtual void GetVMExternalFunction(const FVMExternalFunctionBindingInfo& BindingInfo, void* InstanceData, FVMExternalFunction &OutFunc) override;
 
-	virtual bool Equals(const UNiagaraDataInterface* Other) const override;
+	NIAGARA_API virtual bool Equals(const UNiagaraDataInterface* Other) const override;
 
 	// GPU sim functionality
 #if WITH_EDITORONLY_DATA
-	virtual bool AppendCompileHash(FNiagaraCompileHashVisitor* InVisitor) const override;
-	virtual void GetParameterDefinitionHLSL(const FNiagaraDataInterfaceGPUParamInfo& ParamInfo, FString& OutHLSL) override;
-	virtual bool GetFunctionHLSL(const FNiagaraDataInterfaceGPUParamInfo& ParamInfo, const FNiagaraDataInterfaceGeneratedFunction& FunctionInfo, int FunctionInstanceIndex, FString& OutHLSL) override;
+	NIAGARA_API virtual bool AppendCompileHash(FNiagaraCompileHashVisitor* InVisitor) const override;
+	NIAGARA_API virtual void GetParameterDefinitionHLSL(const FNiagaraDataInterfaceGPUParamInfo& ParamInfo, FString& OutHLSL) override;
+	NIAGARA_API virtual bool GetFunctionHLSL(const FNiagaraDataInterfaceGPUParamInfo& ParamInfo, const FNiagaraDataInterfaceGeneratedFunction& FunctionInfo, int FunctionInstanceIndex, FString& OutHLSL) override;
 #endif
-	virtual void BuildShaderParameters(FNiagaraShaderParametersBuilder& ShaderParametersBuilder) const override;
-	virtual void SetShaderParameters(const FNiagaraDataInterfaceSetShaderParametersContext& Context) const override;
+	NIAGARA_API virtual void BuildShaderParameters(FNiagaraShaderParametersBuilder& ShaderParametersBuilder) const override;
+	NIAGARA_API virtual void SetShaderParameters(const FNiagaraDataInterfaceSetShaderParametersContext& Context) const override;
 #if WITH_EDITOR
-	virtual bool ShouldCompile(EShaderPlatform ShaderPlatform) const override;
+	NIAGARA_API virtual bool ShouldCompile(EShaderPlatform ShaderPlatform) const override;
 #endif
 
 	virtual void ProvidePerInstanceDataForRenderThread(void* DataForRenderThread, void* PerInstanceData, const FNiagaraSystemInstanceID& SystemInstance) override {}
-	virtual bool InitPerInstanceData(void* PerInstanceData, FNiagaraSystemInstance* SystemInstance) override;
-	virtual void DestroyPerInstanceData(void* PerInstanceData, FNiagaraSystemInstance* SystemInstance) override;
+	NIAGARA_API virtual bool InitPerInstanceData(void* PerInstanceData, FNiagaraSystemInstance* SystemInstance) override;
+	NIAGARA_API virtual void DestroyPerInstanceData(void* PerInstanceData, FNiagaraSystemInstance* SystemInstance) override;
 	virtual bool PerInstanceTick(void* PerInstanceData, FNiagaraSystemInstance* SystemInstance, float DeltaSeconds) override { return false;  }
 	virtual int32 PerInstanceDataSize()const override { return sizeof(FNDINeighborGrid3DInstanceData_GT); }
-	virtual bool PerInstanceTickPostSimulate(void* PerInstanceData, FNiagaraSystemInstance* SystemInstance, float DeltaSeconds) override;
+	NIAGARA_API virtual bool PerInstanceTickPostSimulate(void* PerInstanceData, FNiagaraSystemInstance* SystemInstance, float DeltaSeconds) override;
 	virtual bool HasPostSimulateTick() const override { return true; }
 	virtual bool HasPreSimulateTick() const override { return true; }
 	//~ UNiagaraDataInterface interface END
 
-	void GetWorldBBoxSize(FVectorVMExternalFunctionContext& Context);
-	void GetNumCells(FVectorVMExternalFunctionContext& Context);
-	void GetMaxNeighborsPerCell(FVectorVMExternalFunctionContext& Context);
-	void SetNumCells(FVectorVMExternalFunctionContext& Context);
+	NIAGARA_API void GetWorldBBoxSize(FVectorVMExternalFunctionContext& Context);
+	NIAGARA_API void GetNumCells(FVectorVMExternalFunctionContext& Context);
+	NIAGARA_API void GetMaxNeighborsPerCell(FVectorVMExternalFunctionContext& Context);
+	NIAGARA_API void SetNumCells(FVectorVMExternalFunctionContext& Context);
 
-	static const FName SetNumCellsFunctionName;
+	static NIAGARA_API const FName SetNumCellsFunctionName;
 
 
 protected:
 	//~ UNiagaraDataInterface interface
-	virtual bool CopyToInternal(UNiagaraDataInterface* Destination) const override;
+	NIAGARA_API virtual bool CopyToInternal(UNiagaraDataInterface* Destination) const override;
 	//~ UNiagaraDataInterface interface END	
 };
 
diff --git a/Engine/Plugins/FX/Niagara/Source/Niagara/Classes/NiagaraDataInterfaceOcclusion.h b/Engine/Plugins/FX/Niagara/Source/Niagara/Classes/NiagaraDataInterfaceOcclusion.h
--- a/Engine/Plugins/FX/Niagara/Source/Niagara/Classes/NiagaraDataInterfaceOcclusion.h
+++ b/Engine/Plugins/FX/Niagara/Source/Niagara/Classes/NiagaraDataInterfaceOcclusion.h
@@ -6,8 +6,8 @@
 #include "NiagaraDataInterface.h"
 #include "NiagaraDataInterfaceOcclusion.generated.h"
 
-UCLASS(EditInlineNew, Category = "Camera", meta = (DisplayName = "Occlusion Query"))
-class NIAGARA_API UNiagaraDataInterfaceOcclusion : public UNiagaraDataInterface
+UCLASS(EditInlineNew, Category = "Camera", meta = (DisplayName = "Occlusion Query"), MinimalAPI)
+class UNiagaraDataInterfaceOcclusion : public UNiagaraDataInterface
 {
 	GENERATED_UCLASS_BODY()
 
@@ -19,20 +19,20 @@ class NIAGARA_API UNiagaraDataInterfaceOcclusion : public UNiagaraDataInterface
 
 public:
 	//UObject Interface
-	virtual void PostInitProperties() override;
+	NIAGARA_API virtual void PostInitProperties() override;
 	//UObject Interface End
 
 	//UNiagaraDataInterface Interface
-	virtual void GetFunctions(TArray<FNiagaraFunctionSignature>& OutFunctions)override;
+	NIAGARA_API virtual void GetFunctions(TArray<FNiagaraFunctionSignature>& OutFunctions)override;
 #if WITH_EDITORONLY_DATA
-	virtual bool AppendCompileHash(FNiagaraCompileHashVisitor* InVisitor) const override;
-	virtual void GetCommonHLSL(FString& OutHLSL) override;
-	virtual bool GetFunctionHLSL(const FNiagaraDataInterfaceGPUParamInfo& ParamInfo, const FNiagaraDataInterfaceGeneratedFunction& FunctionInfo, int FunctionInstanceIndex, FString& OutHLSL) override;
-	virtual bool UpgradeFunctionCall(FNiagaraFunctionSignature& FunctionSignature) override;
-	virtual void GetParameterDefinitionHLSL(const FNiagaraDataInterfaceGPUParamInfo& ParamInfo, FString& OutHLSL) override;
+	NIAGARA_API virtual bool AppendCompileHash(FNiagaraCompileHashVisitor* InVisitor) const override;
+	NIAGARA_API virtual void GetCommonHLSL(FString& OutHLSL) override;
+	NIAGARA_API virtual bool GetFunctionHLSL(const FNiagaraDataInterfaceGPUParamInfo& ParamInfo, const FNiagaraDataInterfaceGeneratedFunction& FunctionInfo, int FunctionInstanceIndex, FString& OutHLSL) override;
+	NIAGARA_API virtual bool UpgradeFunctionCall(FNiagaraFunctionSignature& FunctionSignature) override;
+	NIAGARA_API virtual void GetParameterDefinitionHLSL(const FNiagaraDataInterfaceGPUParamInfo& ParamInfo, FString& OutHLSL) override;
 #endif
-	virtual void BuildShaderParameters(FNiagaraShaderParametersBuilder& ShaderParametersBuilder) const override;
-	virtual void SetShaderParameters(const FNiagaraDataInterfaceSetShaderParametersContext& Context) const override;
+	NIAGARA_API virtual void BuildShaderParameters(FNiagaraShaderParametersBuilder& ShaderParametersBuilder) const override;
+	NIAGARA_API virtual void SetShaderParameters(const FNiagaraDataInterfaceSetShaderParametersContext& Context) const override;
 
 	virtual bool CanExecuteOnTarget(ENiagaraSimTarget Target) const override { return Target == ENiagaraSimTarget::GPUComputeSim; }
 	virtual bool RequiresDepthBuffer() const override { return true; }
diff --git a/Engine/Plugins/FX/Niagara/Source/Niagara/Classes/NiagaraDataInterfaceParticleRead.h b/Engine/Plugins/FX/Niagara/Source/Niagara/Classes/NiagaraDataInterfaceParticleRead.h
--- a/Engine/Plugins/FX/Niagara/Source/Niagara/Classes/NiagaraDataInterfaceParticleRead.h
+++ b/Engine/Plugins/FX/Niagara/Source/Niagara/Classes/NiagaraDataInterfaceParticleRead.h
@@ -7,8 +7,8 @@
 #include "NiagaraEmitterInstance.h"
 #include "NiagaraDataInterfaceParticleRead.generated.h"
 
-UCLASS(EditInlineNew, Category = "ParticleRead", CollapseCategories, meta = (DisplayName = "Particle Attribute Reader"))
-class NIAGARA_API UNiagaraDataInterfaceParticleRead : public UNiagaraDataInterfaceRWBase
+UCLASS(EditInlineNew, Category = "ParticleRead", CollapseCategories, meta = (DisplayName = "Particle Attribute Reader"), MinimalAPI)
+class UNiagaraDataInterfaceParticleRead : public UNiagaraDataInterfaceRWBase
 {
 	GENERATED_UCLASS_BODY()
 
@@ -32,72 +32,72 @@ public:
 	FString EmitterName;
 
 	//UObject Interface
-	virtual void PostInitProperties()override;
+	NIAGARA_API virtual void PostInitProperties()override;
 #if WITH_EDITOR
-	virtual void PostEditChangeProperty(struct FPropertyChangedEvent& PropertyChangedEvent);
+	NIAGARA_API virtual void PostEditChangeProperty(struct FPropertyChangedEvent& PropertyChangedEvent);
 #endif
 	//UObject Interface End
 
 	//UNiagaraDataInterface Interface
-	virtual bool InitPerInstanceData(void* PerInstanceData, FNiagaraSystemInstance* SystemInstance) override;
-	virtual void DestroyPerInstanceData(void* PerInstanceData, FNiagaraSystemInstance* SystemInstance) override;
-	virtual int32 PerInstanceDataSize() const override;
-	virtual void GetFunctions(TArray<FNiagaraFunctionSignature>& OutFunctions) override;
-	virtual void GetVMExternalFunction(const FVMExternalFunctionBindingInfo& BindingInfo, void* InstanceData, FVMExternalFunction& OutFunc) override;
+	NIAGARA_API virtual bool InitPerInstanceData(void* PerInstanceData, FNiagaraSystemInstance* SystemInstance) override;
+	NIAGARA_API virtual void DestroyPerInstanceData(void* PerInstanceData, FNiagaraSystemInstance* SystemInstance) override;
+	NIAGARA_API virtual int32 PerInstanceDataSize() const override;
+	NIAGARA_API virtual void GetFunctions(TArray<FNiagaraFunctionSignature>& OutFunctions) override;
+	NIAGARA_API virtual void GetVMExternalFunction(const FVMExternalFunctionBindingInfo& BindingInfo, void* InstanceData, FVMExternalFunction& OutFunc) override;
 	virtual bool CanExecuteOnTarget(ENiagaraSimTarget Target) const override { return true; }
-	virtual bool Equals(const UNiagaraDataInterface* Other) const override;
+	NIAGARA_API virtual bool Equals(const UNiagaraDataInterface* Other) const override;
 #if WITH_EDITORONLY_DATA
-	virtual bool AppendCompileHash(FNiagaraCompileHashVisitor* InVisitor) const override;
-	virtual void GetParameterDefinitionHLSL(const FNiagaraDataInterfaceGPUParamInfo& ParamInfo, FString& OutHLSL) override;
-	virtual bool GetFunctionHLSL(const FNiagaraDataInterfaceGPUParamInfo& ParamInfo, const FNiagaraDataInterfaceGeneratedFunction& FunctionInfo, int FunctionInstanceIndex, FString& OutHLSL) override;
+	NIAGARA_API virtual bool AppendCompileHash(FNiagaraCompileHashVisitor* InVisitor) const override;
+	NIAGARA_API virtual void GetParameterDefinitionHLSL(const FNiagaraDataInterfaceGPUParamInfo& ParamInfo, FString& OutHLSL) override;
+	NIAGARA_API virtual bool GetFunctionHLSL(const FNiagaraDataInterfaceGPUParamInfo& ParamInfo, const FNiagaraDataInterfaceGeneratedFunction& FunctionInfo, int FunctionInstanceIndex, FString& OutHLSL) override;
 #endif
-	virtual void BuildShaderParameters(FNiagaraShaderParametersBuilder& ShaderParametersBuilder) const override;
-	virtual void SetShaderParameters(const FNiagaraDataInterfaceSetShaderParametersContext& Context) const override;
-	virtual FNiagaraDataInterfaceParametersCS* CreateShaderStorage(const FNiagaraDataInterfaceGPUParamInfo& ParameterInfo, const FShaderParameterMap& ParameterMap) const override;
-	virtual const FTypeLayoutDesc* GetShaderStorageType() const override;
+	NIAGARA_API virtual void BuildShaderParameters(FNiagaraShaderParametersBuilder& ShaderParametersBuilder) const override;
+	NIAGARA_API virtual void SetShaderParameters(const FNiagaraDataInterfaceSetShaderParametersContext& Context) const override;
+	NIAGARA_API virtual FNiagaraDataInterfaceParametersCS* CreateShaderStorage(const FNiagaraDataInterfaceGPUParamInfo& ParameterInfo, const FShaderParameterMap& ParameterMap) const override;
+	NIAGARA_API virtual const FTypeLayoutDesc* GetShaderStorageType() const override;
 
-	virtual void ProvidePerInstanceDataForRenderThread(void* DataForRenderThread, void* PerInstanceData, const FNiagaraSystemInstanceID& SystemInstance) override;
+	NIAGARA_API virtual void ProvidePerInstanceDataForRenderThread(void* DataForRenderThread, void* PerInstanceData, const FNiagaraSystemInstanceID& SystemInstance) override;
 #if WITH_EDITOR	
-	virtual void GetFeedback(UNiagaraSystem* Asset, UNiagaraComponent* Component, TArray<FNiagaraDataInterfaceError>& OutErrors, TArray<FNiagaraDataInterfaceFeedback>& Warnings, TArray<FNiagaraDataInterfaceFeedback>& Info) override;
+	NIAGARA_API virtual void GetFeedback(UNiagaraSystem* Asset, UNiagaraComponent* Component, TArray<FNiagaraDataInterfaceError>& OutErrors, TArray<FNiagaraDataInterfaceFeedback>& Warnings, TArray<FNiagaraDataInterfaceFeedback>& Info) override;
 #endif
 #if WITH_EDITORONLY_DATA
-	virtual bool UpgradeFunctionCall(FNiagaraFunctionSignature& FunctionSignature) override;
+	NIAGARA_API virtual bool UpgradeFunctionCall(FNiagaraFunctionSignature& FunctionSignature) override;
 #endif
-	virtual void GetEmitterDependencies(UNiagaraSystem* Asset, TArray<FVersionedNiagaraEmitter>& Dependencies) const override;
-	virtual void GetEmitterReferencesByName(TArray<FString>& EmitterReferences) const override;
+	NIAGARA_API virtual void GetEmitterDependencies(UNiagaraSystem* Asset, TArray<FVersionedNiagaraEmitter>& Dependencies) const override;
+	NIAGARA_API virtual void GetEmitterReferencesByName(TArray<FString>& EmitterReferences) const override;
 
-	virtual bool HasInternalAttributeReads(const UNiagaraEmitter* OwnerEmitter, const UNiagaraEmitter* Provider) const override;
+	NIAGARA_API virtual bool HasInternalAttributeReads(const UNiagaraEmitter* OwnerEmitter, const UNiagaraEmitter* Provider) const override;
 	//UNiagaraDataInterface Interface End
 
-	void VMGetLocalSpace(FVectorVMExternalFunctionContext& Context);
-	void GetNumSpawnedParticles(FVectorVMExternalFunctionContext& Context);
-	void GetSpawnedIDAtIndex(FVectorVMExternalFunctionContext& Context);
-	void GetNumParticles(FVectorVMExternalFunctionContext& Context);
-	void GetParticleIndex(FVectorVMExternalFunctionContext& Context);
-	void GetParticleIndexFromIDTable(FVectorVMExternalFunctionContext& Context);
-	void ReadInt(FVectorVMExternalFunctionContext& Context, FName AttributeToRead);
-	void ReadBool(FVectorVMExternalFunctionContext& Context, FName AttributeToRead);
-	void ReadFloat(FVectorVMExternalFunctionContext& Context, FName AttributeToRead);
-	void ReadVector2(FVectorVMExternalFunctionContext& Context, FName AttributeToRead);
-	void ReadVector3(FVectorVMExternalFunctionContext& Context, FName AttributeToRead);
-	void ReadVector4(FVectorVMExternalFunctionContext& Context, FName AttributeToRead);
-	void ReadColor(FVectorVMExternalFunctionContext& Context, FName AttributeToRead);
-	void ReadPosition(FVectorVMExternalFunctionContext& Context, FName AttributeToRead);
-	void ReadQuat(FVectorVMExternalFunctionContext& Context, FName AttributeToRead);
-	void ReadID(FVectorVMExternalFunctionContext& Context, FName AttributeToRead);
-	void ReadIntByIndex(FVectorVMExternalFunctionContext& Context, FName AttributeToRead);
-	void ReadBoolByIndex(FVectorVMExternalFunctionContext& Context, FName AttributeToRead);
-	void ReadFloatByIndex(FVectorVMExternalFunctionContext& Context, FName AttributeToRead);
-	void ReadVector2ByIndex(FVectorVMExternalFunctionContext& Context, FName AttributeToRead);
-	void ReadVector3ByIndex(FVectorVMExternalFunctionContext& Context, FName AttributeToRead);
-	void ReadVector4ByIndex(FVectorVMExternalFunctionContext& Context, FName AttributeToRead);
-	void ReadPositionByIndex(FVectorVMExternalFunctionContext& Context, FName AttributeToRead);
-	void ReadColorByIndex(FVectorVMExternalFunctionContext& Context, FName AttributeToRead);
-	void ReadQuatByIndex(FVectorVMExternalFunctionContext& Context, FName AttributeToRead);
-	void ReadIDByIndex(FVectorVMExternalFunctionContext& Context, FName AttributeToRead);
+	NIAGARA_API void VMGetLocalSpace(FVectorVMExternalFunctionContext& Context);
+	NIAGARA_API void GetNumSpawnedParticles(FVectorVMExternalFunctionContext& Context);
+	NIAGARA_API void GetSpawnedIDAtIndex(FVectorVMExternalFunctionContext& Context);
+	NIAGARA_API void GetNumParticles(FVectorVMExternalFunctionContext& Context);
+	NIAGARA_API void GetParticleIndex(FVectorVMExternalFunctionContext& Context);
+	NIAGARA_API void GetParticleIndexFromIDTable(FVectorVMExternalFunctionContext& Context);
+	NIAGARA_API void ReadInt(FVectorVMExternalFunctionContext& Context, FName AttributeToRead);
+	NIAGARA_API void ReadBool(FVectorVMExternalFunctionContext& Context, FName AttributeToRead);
+	NIAGARA_API void ReadFloat(FVectorVMExternalFunctionContext& Context, FName AttributeToRead);
+	NIAGARA_API void ReadVector2(FVectorVMExternalFunctionContext& Context, FName AttributeToRead);
+	NIAGARA_API void ReadVector3(FVectorVMExternalFunctionContext& Context, FName AttributeToRead);
+	NIAGARA_API void ReadVector4(FVectorVMExternalFunctionContext& Context, FName AttributeToRead);
+	NIAGARA_API void ReadColor(FVectorVMExternalFunctionContext& Context, FName AttributeToRead);
+	NIAGARA_API void ReadPosition(FVectorVMExternalFunctionContext& Context, FName AttributeToRead);
+	NIAGARA_API void ReadQuat(FVectorVMExternalFunctionContext& Context, FName AttributeToRead);
+	NIAGARA_API void ReadID(FVectorVMExternalFunctionContext& Context, FName AttributeToRead);
+	NIAGARA_API void ReadIntByIndex(FVectorVMExternalFunctionContext& Context, FName AttributeToRead);
+	NIAGARA_API void ReadBoolByIndex(FVectorVMExternalFunctionContext& Context, FName AttributeToRead);
+	NIAGARA_API void ReadFloatByIndex(FVectorVMExternalFunctionContext& Context, FName AttributeToRead);
+	NIAGARA_API void ReadVector2ByIndex(FVectorVMExternalFunctionContext& Context, FName AttributeToRead);
+	NIAGARA_API void ReadVector3ByIndex(FVectorVMExternalFunctionContext& Context, FName AttributeToRead);
+	NIAGARA_API void ReadVector4ByIndex(FVectorVMExternalFunctionContext& Context, FName AttributeToRead);
+	NIAGARA_API void ReadPositionByIndex(FVectorVMExternalFunctionContext& Context, FName AttributeToRead);
+	NIAGARA_API void ReadColorByIndex(FVectorVMExternalFunctionContext& Context, FName AttributeToRead);
+	NIAGARA_API void ReadQuatByIndex(FVectorVMExternalFunctionContext& Context, FName AttributeToRead);
+	NIAGARA_API void ReadIDByIndex(FVectorVMExternalFunctionContext& Context, FName AttributeToRead);
 
 protected:
-	void GetPersistentIDFunctions(TArray<FNiagaraFunctionSignature>& OutFunctions);
-	void GetIndexFunctions(TArray<FNiagaraFunctionSignature>& OutFunctions);
-	virtual bool CopyToInternal(UNiagaraDataInterface* Destination) const override;
-};
\ No newline at end of file
+	NIAGARA_API void GetPersistentIDFunctions(TArray<FNiagaraFunctionSignature>& OutFunctions);
+	NIAGARA_API void GetIndexFunctions(TArray<FNiagaraFunctionSignature>& OutFunctions);
+	NIAGARA_API virtual bool CopyToInternal(UNiagaraDataInterface* Destination) const override;
+};
diff --git a/Engine/Plugins/FX/Niagara/Source/Niagara/Classes/NiagaraDataInterfacePlatformSet.h b/Engine/Plugins/FX/Niagara/Source/Niagara/Classes/NiagaraDataInterfacePlatformSet.h
--- a/Engine/Plugins/FX/Niagara/Source/Niagara/Classes/NiagaraDataInterfacePlatformSet.h
+++ b/Engine/Plugins/FX/Niagara/Source/Niagara/Classes/NiagaraDataInterfacePlatformSet.h
@@ -9,8 +9,8 @@
 
 
 /** Data Interface allowing querying of the current platform set. */
-UCLASS(EditInlineNew, Category = "Scalability", CollapseCategories, meta = (DisplayName = "Platform Set"))
-class NIAGARA_API UNiagaraDataInterfacePlatformSet : public UNiagaraDataInterface
+UCLASS(EditInlineNew, Category = "Scalability", CollapseCategories, meta = (DisplayName = "Platform Set"), MinimalAPI)
+class UNiagaraDataInterfacePlatformSet : public UNiagaraDataInterface
 {
 	GENERATED_UCLASS_BODY()
 
@@ -19,22 +19,22 @@ public:
 	FNiagaraPlatformSet Platforms;
 
 	//UObject Interface
-	virtual void PostInitProperties() override;
+	NIAGARA_API virtual void PostInitProperties() override;
 	//UObject Interface End
 
 	// UNiagaraDataInterface interface
-	virtual void GetFunctions(TArray<FNiagaraFunctionSignature>& OutFunctions)override;
-	virtual void GetVMExternalFunction(const FVMExternalFunctionBindingInfo& BindingInfo, void* InstanceData, FVMExternalFunction& OutFunc) override;
+	NIAGARA_API virtual void GetFunctions(TArray<FNiagaraFunctionSignature>& OutFunctions)override;
+	NIAGARA_API virtual void GetVMExternalFunction(const FVMExternalFunctionBindingInfo& BindingInfo, void* InstanceData, FVMExternalFunction& OutFunc) override;
 	virtual bool CanExecuteOnTarget(ENiagaraSimTarget Target)const override { return Target == ENiagaraSimTarget::CPUSim; }
 
-	virtual bool Equals(const UNiagaraDataInterface* Other) const override;
+	NIAGARA_API virtual bool Equals(const UNiagaraDataInterface* Other) const override;
 
 protected:
-	virtual bool CopyToInternal(UNiagaraDataInterface* Destination) const override;
+	NIAGARA_API virtual bool CopyToInternal(UNiagaraDataInterface* Destination) const override;
 
 	//~ UNiagaraDataInterface interface
 
 public:
-	void IsActive(FVectorVMExternalFunctionContext& Context);
+	NIAGARA_API void IsActive(FVectorVMExternalFunctionContext& Context);
 };
 
diff --git a/Engine/Plugins/FX/Niagara/Source/Niagara/Classes/NiagaraDataInterfaceRW.h b/Engine/Plugins/FX/Niagara/Source/Niagara/Classes/NiagaraDataInterfaceRW.h
--- a/Engine/Plugins/FX/Niagara/Source/Niagara/Classes/NiagaraDataInterfaceRW.h
+++ b/Engine/Plugins/FX/Niagara/Source/Niagara/Classes/NiagaraDataInterfaceRW.h
@@ -17,7 +17,7 @@ BEGIN_SHADER_PARAMETER_STRUCT(FNDIGpuComputeDispatchArgsGenParameters, )
 	RDG_TEXTURE_ACCESS_ARRAY(TextureAccessArray)
 END_SHADER_PARAMETER_STRUCT()
 
-struct NIAGARA_API FNDIGpuComputeDispatchArgsGenContext : public FNDIGpuComputeContext
+struct FNDIGpuComputeDispatchArgsGenContext : public FNDIGpuComputeContext
 {
 	using FIndirectArgs = TPair<FRDGBuffer*, uint32>;
 	using FCreateIndirectCallback = TFunction<void(FRHICommandList&, FRDGBuffer*, uint32)>;
@@ -42,25 +42,25 @@ struct NIAGARA_API FNDIGpuComputeDispatchArgsGenContext : public FNDIGpuComputeC
 	void SetDirect(uint32 InElementCount, uint32 GpuCountOffset = INDEX_NONE) const { SetDirect(FIntVector3(InElementCount, 1, 1), GpuCountOffset); }
 	void SetDirect(const FIntPoint& InElementCount, uint32 GpuCountOffset = INDEX_NONE) const { SetDirect(FIntVector3(InElementCount.X, InElementCount.Y, 1), GpuCountOffset); }
 	void SetDirect(const FIntVector2& InElementCountXY, uint32 GpuCountOffset = INDEX_NONE) const { SetDirect(FIntVector3(InElementCountXY.X, InElementCountXY.Y, 1), GpuCountOffset); }
-	void SetDirect(const FIntVector3& InElementCountXYZ, uint32 GpuCountOffset = INDEX_NONE) const;
+	NIAGARA_API void SetDirect(const FIntVector3& InElementCountXYZ, uint32 GpuCountOffset = INDEX_NONE) const;
 
 	// Set Indirect buffer to use
-	void SetIndirect(FRDGBuffer* InBuffer, uint32 BufferByteOffset) const;
+	NIAGARA_API void SetIndirect(FRDGBuffer* InBuffer, uint32 BufferByteOffset) const;
 	// Set indirect buffer to use, the callback will be inside the RDG pass
-	void SetIndirect(FRDGBuffer* InBuffer, uint32 BufferByteOffset, FCreateIndirectCallback&& Callback) const;
+	NIAGARA_API void SetIndirect(FRDGBuffer* InBuffer, uint32 BufferByteOffset, FCreateIndirectCallback&& Callback) const;
 	// Create indirect buffer from count values that are held in the GPU count buffer
 	// The returned indirect argument information can be kept around for the current frame
 	FIndirectArgs CreateIndirect(uint32 InCounterOffset) const { return CreateIndirect(FUintVector3(InCounterOffset, INDEX_NONE, INDEX_NONE)); }
 	FIndirectArgs CreateIndirect(const FUintVector2& InCounterOffsets) const { return CreateIndirect(FUintVector3(InCounterOffsets.X, InCounterOffsets.Y, INDEX_NONE)); }
-	FIndirectArgs CreateIndirect(const FUintVector3& InCounterOffsets) const;
+	NIAGARA_API FIndirectArgs CreateIndirect(const FUintVector3& InCounterOffsets) const;
 
 	// Add buffer / texture access used for batching transitions together
-	void AddBufferAccess(FRDGBuffer* InBuffer, ERHIAccess InAccess) const;
-	void AddTextureAccess(FRDGTexture* InTexture, ERHIAccess InAccess) const;
+	NIAGARA_API void AddBufferAccess(FRDGBuffer* InBuffer, ERHIAccess InAccess) const;
+	NIAGARA_API void AddTextureAccess(FRDGTexture* InTexture, ERHIAccess InAccess) const;
 
 protected:
-	FNDIGpuComputeDispatchArgsGenParameters* GetPassParameters() const;
-	void FlushPass() const;
+	NIAGARA_API FNDIGpuComputeDispatchArgsGenParameters* GetPassParameters() const;
+	NIAGARA_API void FlushPass() const;
 
 protected:
 	FNiagaraSystemInstanceID SystemInstanceID;
@@ -105,42 +105,42 @@ public:
 	virtual FNiagaraDataInterfaceProxyRW* AsIterationProxy() override { return this; }
 };
 
-UCLASS(abstract, EditInlineNew)
-class NIAGARA_API UNiagaraDataInterfaceRWBase : public UNiagaraDataInterface
+UCLASS(abstract, EditInlineNew, MinimalAPI)
+class UNiagaraDataInterfaceRWBase : public UNiagaraDataInterface
 {
 	GENERATED_UCLASS_BODY()
 public:
 	// Global HLSL variable base names, used by HLSL.
-	static const FString NumAttributesName;
-	static const FString NumNamedAttributesName;
-	static const FString NumCellsName;
-	static const FString UnitToUVName;
-	static const FString CellSizeName;
-	static const FString WorldBBoxSizeName;
+	static NIAGARA_API const FString NumAttributesName;
+	static NIAGARA_API const FString NumNamedAttributesName;
+	static NIAGARA_API const FString NumCellsName;
+	static NIAGARA_API const FString UnitToUVName;
+	static NIAGARA_API const FString CellSizeName;
+	static NIAGARA_API const FString WorldBBoxSizeName;
 
 	// Attribute names
-	static const FName NAME_Attribute;
+	static NIAGARA_API const FName NAME_Attribute;
 
 	// Global VM function names, also used by the shaders code generation methods.
-	static const FName NumCellsFunctionName;
-	static const FName CellSizeFunctionName;
+	static NIAGARA_API const FName NumCellsFunctionName;
+	static NIAGARA_API const FName CellSizeFunctionName;
 
-	static const FName WorldBBoxSizeFunctionName;
+	static NIAGARA_API const FName WorldBBoxSizeFunctionName;
 
-	static const FName SimulationToUnitFunctionName;
-	static const FName UnitToSimulationFunctionName;
-	static const FName UnitToIndexFunctionName;
-	static const FName UnitToFloatIndexFunctionName;
-	static const FName IndexToUnitFunctionName;
+	static NIAGARA_API const FName SimulationToUnitFunctionName;
+	static NIAGARA_API const FName UnitToSimulationFunctionName;
+	static NIAGARA_API const FName UnitToIndexFunctionName;
+	static NIAGARA_API const FName UnitToFloatIndexFunctionName;
+	static NIAGARA_API const FName IndexToUnitFunctionName;
 
-	static const FName IndexToUnitStaggeredXFunctionName;
-	static const FName IndexToUnitStaggeredYFunctionName;
+	static NIAGARA_API const FName IndexToUnitStaggeredXFunctionName;
+	static NIAGARA_API const FName IndexToUnitStaggeredYFunctionName;
 
-	static const FName IndexToLinearFunctionName;
-	static const FName LinearToIndexFunctionName;
+	static NIAGARA_API const FName IndexToLinearFunctionName;
+	static NIAGARA_API const FName LinearToIndexFunctionName;
 
-	static const FName ExecutionIndexToGridIndexFunctionName;
-	static const FName ExecutionIndexToUnitFunctionName;
+	static NIAGARA_API const FName ExecutionIndexToGridIndexFunctionName;
+	static NIAGARA_API const FName ExecutionIndexToUnitFunctionName;
 	//~ UObject interface
 
 	virtual void PostLoad() override
@@ -185,8 +185,8 @@ public:
 };
 
 
-UCLASS(abstract, EditInlineNew)
-class NIAGARA_API UNiagaraDataInterfaceGrid3D : public UNiagaraDataInterfaceRWBase
+UCLASS(abstract, EditInlineNew, MinimalAPI)
+class UNiagaraDataInterfaceGrid3D : public UNiagaraDataInterfaceRWBase
 {
 	GENERATED_UCLASS_BODY()
 
@@ -221,28 +221,28 @@ public:
 
 	//~ UNiagaraDataInterface interface
 	// VM functionality
-	virtual void GetFunctions(TArray<FNiagaraFunctionSignature>& OutFunctions) override;
-	virtual void GetVMExternalFunction(const FVMExternalFunctionBindingInfo& BindingInfo, void* InstanceData, FVMExternalFunction &OutFunc) override;	
+	NIAGARA_API virtual void GetFunctions(TArray<FNiagaraFunctionSignature>& OutFunctions) override;
+	NIAGARA_API virtual void GetVMExternalFunction(const FVMExternalFunctionBindingInfo& BindingInfo, void* InstanceData, FVMExternalFunction &OutFunc) override;	
 
 
-	virtual bool Equals(const UNiagaraDataInterface* Other) const override;
+	NIAGARA_API virtual bool Equals(const UNiagaraDataInterface* Other) const override;
 
 	// GPU sim functionality
 #if WITH_EDITORONLY_DATA
-	virtual void GetParameterDefinitionHLSL(const FNiagaraDataInterfaceGPUParamInfo& ParamInfo, FString& OutHLSL) override;
-	virtual bool GetFunctionHLSL(const FNiagaraDataInterfaceGPUParamInfo& ParamInfo, const FNiagaraDataInterfaceGeneratedFunction& FunctionInfo, int FunctionInstanceIndex, FString& OutHLSL) override;
+	NIAGARA_API virtual void GetParameterDefinitionHLSL(const FNiagaraDataInterfaceGPUParamInfo& ParamInfo, FString& OutHLSL) override;
+	NIAGARA_API virtual bool GetFunctionHLSL(const FNiagaraDataInterfaceGPUParamInfo& ParamInfo, const FNiagaraDataInterfaceGeneratedFunction& FunctionInfo, int FunctionInstanceIndex, FString& OutHLSL) override;
 #endif
 	//~ UNiagaraDataInterface interface END
 
 protected:
 	//~ UNiagaraDataInterface interface
-	virtual bool CopyToInternal(UNiagaraDataInterface* Destination) const override;
+	NIAGARA_API virtual bool CopyToInternal(UNiagaraDataInterface* Destination) const override;
 	//~ UNiagaraDataInterface interface END	
 };
 
 
-UCLASS(abstract, EditInlineNew)
-class NIAGARA_API UNiagaraDataInterfaceGrid2D : public UNiagaraDataInterfaceRWBase
+UCLASS(abstract, EditInlineNew, MinimalAPI)
+class UNiagaraDataInterfaceGrid2D : public UNiagaraDataInterfaceRWBase
 {
 	GENERATED_UCLASS_BODY()
 
@@ -280,29 +280,29 @@ public:
 
 public:
 
-	virtual void Serialize(FArchive& Ar) override;
+	NIAGARA_API virtual void Serialize(FArchive& Ar) override;
 
 	//~ UNiagaraDataInterface interface
 	// VM functionality
-	virtual void GetFunctions(TArray<FNiagaraFunctionSignature>& OutFunctions) override;
-	virtual void GetVMExternalFunction(const FVMExternalFunctionBindingInfo& BindingInfo, void* InstanceData, FVMExternalFunction &OutFunc) override;
+	NIAGARA_API virtual void GetFunctions(TArray<FNiagaraFunctionSignature>& OutFunctions) override;
+	NIAGARA_API virtual void GetVMExternalFunction(const FVMExternalFunctionBindingInfo& BindingInfo, void* InstanceData, FVMExternalFunction &OutFunc) override;
 
-	virtual bool Equals(const UNiagaraDataInterface* Other) const override;
+	NIAGARA_API virtual bool Equals(const UNiagaraDataInterface* Other) const override;
 
 	// GPU sim functionality
 #if WITH_EDITORONLY_DATA
-	virtual void GetParameterDefinitionHLSL(const FNiagaraDataInterfaceGPUParamInfo& ParamInfo, FString& OutHLSL) override;
-	virtual bool GetFunctionHLSL(const FNiagaraDataInterfaceGPUParamInfo& ParamInfo, const FNiagaraDataInterfaceGeneratedFunction& FunctionInfo, int FunctionInstanceIndex, FString& OutHLSL) override;
+	NIAGARA_API virtual void GetParameterDefinitionHLSL(const FNiagaraDataInterfaceGPUParamInfo& ParamInfo, FString& OutHLSL) override;
+	NIAGARA_API virtual bool GetFunctionHLSL(const FNiagaraDataInterfaceGPUParamInfo& ParamInfo, const FNiagaraDataInterfaceGeneratedFunction& FunctionInfo, int FunctionInstanceIndex, FString& OutHLSL) override;
 #endif
 
 #if WITH_EDITOR		
-	virtual void ValidateFunction(const FNiagaraFunctionSignature& Function, TArray<FText>& OutValidationErrors) override;
+	NIAGARA_API virtual void ValidateFunction(const FNiagaraFunctionSignature& Function, TArray<FText>& OutValidationErrors) override;
 #endif
 	//~ UNiagaraDataInterface interface END
 
 
 protected:
 	//~ UNiagaraDataInterface interface
-	virtual bool CopyToInternal(UNiagaraDataInterface* Destination) const override;
+	NIAGARA_API virtual bool CopyToInternal(UNiagaraDataInterface* Destination) const override;
 	//~ UNiagaraDataInterface interface END	
-};
\ No newline at end of file
+};
diff --git a/Engine/Plugins/FX/Niagara/Source/Niagara/Classes/NiagaraDataInterfaceRasterizationGrid3D.h b/Engine/Plugins/FX/Niagara/Source/Niagara/Classes/NiagaraDataInterfaceRasterizationGrid3D.h
--- a/Engine/Plugins/FX/Niagara/Source/Niagara/Classes/NiagaraDataInterfaceRasterizationGrid3D.h
+++ b/Engine/Plugins/FX/Niagara/Source/Niagara/Classes/NiagaraDataInterfaceRasterizationGrid3D.h
@@ -52,8 +52,8 @@ struct FNiagaraDataInterfaceProxyRasterizationGrid3D : public FNiagaraDataInterf
 	TMap<FNiagaraSystemInstanceID, RasterizationGrid3DRWInstanceData> SystemInstancesToProxyData;
 };
 
-UCLASS(EditInlineNew, Category = "Grid", CollapseCategories, meta = (DisplayName = "Rasterization Grid3D"))
-class NIAGARA_API UNiagaraDataInterfaceRasterizationGrid3D : public UNiagaraDataInterfaceGrid3D
+UCLASS(EditInlineNew, Category = "Grid", CollapseCategories, meta = (DisplayName = "Rasterization Grid3D"), MinimalAPI)
+class UNiagaraDataInterfaceRasterizationGrid3D : public UNiagaraDataInterfaceGrid3D
 {
 	GENERATED_UCLASS_BODY()
 
@@ -84,44 +84,44 @@ public:
 
 	//~ UNiagaraDataInterface interface
 	// VM functionality
-	virtual void GetFunctions(TArray<FNiagaraFunctionSignature>& OutFunctions) override;
-	virtual void GetVMExternalFunction(const FVMExternalFunctionBindingInfo& BindingInfo, void* InstanceData, FVMExternalFunction &OutFunc) override;
+	NIAGARA_API virtual void GetFunctions(TArray<FNiagaraFunctionSignature>& OutFunctions) override;
+	NIAGARA_API virtual void GetVMExternalFunction(const FVMExternalFunctionBindingInfo& BindingInfo, void* InstanceData, FVMExternalFunction &OutFunc) override;
 
-	virtual bool Equals(const UNiagaraDataInterface* Other) const override;
+	NIAGARA_API virtual bool Equals(const UNiagaraDataInterface* Other) const override;
 #if WITH_EDITOR
-	virtual bool ShouldCompile(EShaderPlatform ShaderPlatform) const override;
+	NIAGARA_API virtual bool ShouldCompile(EShaderPlatform ShaderPlatform) const override;
 #endif
 
 	// GPU sim functionality
 #if WITH_EDITORONLY_DATA
-	virtual bool AppendCompileHash(FNiagaraCompileHashVisitor* InVisitor) const override;
-	virtual void GetParameterDefinitionHLSL(const FNiagaraDataInterfaceGPUParamInfo& ParamInfo, FString& OutHLSL) override;
-	virtual bool GetFunctionHLSL(const FNiagaraDataInterfaceGPUParamInfo& ParamInfo, const FNiagaraDataInterfaceGeneratedFunction& FunctionInfo, int FunctionInstanceIndex, FString& OutHLSL) override;
+	NIAGARA_API virtual bool AppendCompileHash(FNiagaraCompileHashVisitor* InVisitor) const override;
+	NIAGARA_API virtual void GetParameterDefinitionHLSL(const FNiagaraDataInterfaceGPUParamInfo& ParamInfo, FString& OutHLSL) override;
+	NIAGARA_API virtual bool GetFunctionHLSL(const FNiagaraDataInterfaceGPUParamInfo& ParamInfo, const FNiagaraDataInterfaceGeneratedFunction& FunctionInfo, int FunctionInstanceIndex, FString& OutHLSL) override;
 #endif
-	virtual void BuildShaderParameters(FNiagaraShaderParametersBuilder& ShaderParametersBuilder) const override;
-	virtual void SetShaderParameters(const FNiagaraDataInterfaceSetShaderParametersContext& Context) const override;
+	NIAGARA_API virtual void BuildShaderParameters(FNiagaraShaderParametersBuilder& ShaderParametersBuilder) const override;
+	NIAGARA_API virtual void SetShaderParameters(const FNiagaraDataInterfaceSetShaderParametersContext& Context) const override;
 #if WITH_EDITOR
 	virtual ENiagaraGpuDispatchType GetGpuDispatchType() const override { return ENiagaraGpuDispatchType::ThreeD; }
 #endif
 
 	virtual void ProvidePerInstanceDataForRenderThread(void* DataForRenderThread, void* PerInstanceData, const FNiagaraSystemInstanceID& SystemInstance) override {}
-	virtual bool InitPerInstanceData(void* PerInstanceData, FNiagaraSystemInstance* SystemInstance) override;
-	virtual void DestroyPerInstanceData(void* PerInstanceData, FNiagaraSystemInstance* SystemInstance) override;
+	NIAGARA_API virtual bool InitPerInstanceData(void* PerInstanceData, FNiagaraSystemInstance* SystemInstance) override;
+	NIAGARA_API virtual void DestroyPerInstanceData(void* PerInstanceData, FNiagaraSystemInstance* SystemInstance) override;
 	virtual bool PerInstanceTick(void* PerInstanceData, FNiagaraSystemInstance* SystemInstance, float DeltaSeconds) override { return false;  }
 	virtual int32 PerInstanceDataSize()const override { return sizeof(RasterizationGrid3DRWInstanceData); }
-	virtual bool PerInstanceTickPostSimulate(void* PerInstanceData, FNiagaraSystemInstance* SystemInstance, float DeltaSeconds) override;
+	NIAGARA_API virtual bool PerInstanceTickPostSimulate(void* PerInstanceData, FNiagaraSystemInstance* SystemInstance, float DeltaSeconds) override;
 	virtual bool HasPostSimulateTick() const override { return true; }
 	virtual bool HasPreSimulateTick() const override { return true; }
 	//~ UNiagaraDataInterface interface END
 
-	void VMGetNumCells(FVectorVMExternalFunctionContext& Context);
-	void VMSetNumCells(FVectorVMExternalFunctionContext& Context);
-	void VMSetNumAttributes(FVectorVMExternalFunctionContext& Context);
-	void VMSetFloatResetValue(FVectorVMExternalFunctionContext& Context);
+	NIAGARA_API void VMGetNumCells(FVectorVMExternalFunctionContext& Context);
+	NIAGARA_API void VMSetNumCells(FVectorVMExternalFunctionContext& Context);
+	NIAGARA_API void VMSetNumAttributes(FVectorVMExternalFunctionContext& Context);
+	NIAGARA_API void VMSetFloatResetValue(FVectorVMExternalFunctionContext& Context);
 
 protected:
 	//~ UNiagaraDataInterface interface
-	virtual bool CopyToInternal(UNiagaraDataInterface* Destination) const override;
+	NIAGARA_API virtual bool CopyToInternal(UNiagaraDataInterface* Destination) const override;
 	//~ UNiagaraDataInterface interface END	
 };
 
diff --git a/Engine/Plugins/FX/Niagara/Source/Niagara/Classes/NiagaraDataInterfaceRenderTarget2D.h b/Engine/Plugins/FX/Niagara/Source/Niagara/Classes/NiagaraDataInterfaceRenderTarget2D.h
--- a/Engine/Plugins/FX/Niagara/Source/Niagara/Classes/NiagaraDataInterfaceRenderTarget2D.h
+++ b/Engine/Plugins/FX/Niagara/Source/Niagara/Classes/NiagaraDataInterfaceRenderTarget2D.h
@@ -85,57 +85,57 @@ struct FNiagaraDataInterfaceProxyRenderTarget2DProxy : public FNiagaraDataInterf
 	TMap<FNiagaraSystemInstanceID, FRenderTarget2DRWInstanceData_RenderThread> SystemInstancesToProxyData_RT;
 };
 
-UCLASS(EditInlineNew, Category = "Rendering", CollapseCategories, meta = (DisplayName = "Render Target 2D"), Blueprintable, BlueprintType)
-class NIAGARA_API UNiagaraDataInterfaceRenderTarget2D : public UNiagaraDataInterfaceRWBase
+UCLASS(EditInlineNew, Category = "Rendering", CollapseCategories, meta = (DisplayName = "Render Target 2D"), Blueprintable, BlueprintType, MinimalAPI)
+class UNiagaraDataInterfaceRenderTarget2D : public UNiagaraDataInterfaceRWBase
 {
 	GENERATED_UCLASS_BODY()
 
 public:
-	virtual void PostInitProperties() override;
+	NIAGARA_API virtual void PostInitProperties() override;
 	
 	//~ UNiagaraDataInterface interface
 	// VM functionality
 	virtual bool CanExecuteOnTarget(ENiagaraSimTarget Target)const override { return true; }
-	virtual void GetFunctions(TArray<FNiagaraFunctionSignature>& OutFunctions) override;
+	NIAGARA_API virtual void GetFunctions(TArray<FNiagaraFunctionSignature>& OutFunctions) override;
 #if WITH_EDITORONLY_DATA
-	virtual bool UpgradeFunctionCall(FNiagaraFunctionSignature& FunctionSignature) override;
+	NIAGARA_API virtual bool UpgradeFunctionCall(FNiagaraFunctionSignature& FunctionSignature) override;
 #endif
-	virtual void GetVMExternalFunction(const FVMExternalFunctionBindingInfo& BindingInfo, void* InstanceData, FVMExternalFunction &OutFunc) override;
+	NIAGARA_API virtual void GetVMExternalFunction(const FVMExternalFunctionBindingInfo& BindingInfo, void* InstanceData, FVMExternalFunction &OutFunc) override;
 
-	virtual bool Equals(const UNiagaraDataInterface* Other) const override;
-	virtual bool CopyToInternal(UNiagaraDataInterface* Destination) const override;
+	NIAGARA_API virtual bool Equals(const UNiagaraDataInterface* Other) const override;
+	NIAGARA_API virtual bool CopyToInternal(UNiagaraDataInterface* Destination) const override;
 
 	// GPU sim functionality
 #if WITH_EDITORONLY_DATA
-	virtual bool AppendCompileHash(FNiagaraCompileHashVisitor* InVisitor) const override;
-	virtual void GetParameterDefinitionHLSL(const FNiagaraDataInterfaceGPUParamInfo& ParamInfo, FString& OutHLSL) override;
-	virtual bool GetFunctionHLSL(const FNiagaraDataInterfaceGPUParamInfo& ParamInfo, const FNiagaraDataInterfaceGeneratedFunction& FunctionInfo, int FunctionInstanceIndex, FString& OutHLSL) override;
+	NIAGARA_API virtual bool AppendCompileHash(FNiagaraCompileHashVisitor* InVisitor) const override;
+	NIAGARA_API virtual void GetParameterDefinitionHLSL(const FNiagaraDataInterfaceGPUParamInfo& ParamInfo, FString& OutHLSL) override;
+	NIAGARA_API virtual bool GetFunctionHLSL(const FNiagaraDataInterfaceGPUParamInfo& ParamInfo, const FNiagaraDataInterfaceGeneratedFunction& FunctionInfo, int FunctionInstanceIndex, FString& OutHLSL) override;
 #endif
-	virtual void BuildShaderParameters(FNiagaraShaderParametersBuilder& ShaderParametersBuilder) const override;
-	virtual void SetShaderParameters(const FNiagaraDataInterfaceSetShaderParametersContext& Context) const override;
+	NIAGARA_API virtual void BuildShaderParameters(FNiagaraShaderParametersBuilder& ShaderParametersBuilder) const override;
+	NIAGARA_API virtual void SetShaderParameters(const FNiagaraDataInterfaceSetShaderParametersContext& Context) const override;
 
 	virtual void ProvidePerInstanceDataForRenderThread(void* DataForRenderThread, void* PerInstanceData, const FNiagaraSystemInstanceID& SystemInstance) override {}
-	virtual bool InitPerInstanceData(void* PerInstanceData, FNiagaraSystemInstance* SystemInstance) override;
-	virtual void DestroyPerInstanceData(void* PerInstanceData, FNiagaraSystemInstance* SystemInstance) override;
-	virtual bool PerInstanceTick(void* PerInstanceData, FNiagaraSystemInstance* SystemInstance, float DeltaSeconds) override;
+	NIAGARA_API virtual bool InitPerInstanceData(void* PerInstanceData, FNiagaraSystemInstance* SystemInstance) override;
+	NIAGARA_API virtual void DestroyPerInstanceData(void* PerInstanceData, FNiagaraSystemInstance* SystemInstance) override;
+	NIAGARA_API virtual bool PerInstanceTick(void* PerInstanceData, FNiagaraSystemInstance* SystemInstance, float DeltaSeconds) override;
 	virtual int32 PerInstanceDataSize()const override { return sizeof(FRenderTarget2DRWInstanceData_GameThread); }
-	virtual bool PerInstanceTickPostSimulate(void* PerInstanceData, FNiagaraSystemInstance* InSystemInstance, float DeltaSeconds) override;
+	NIAGARA_API virtual bool PerInstanceTickPostSimulate(void* PerInstanceData, FNiagaraSystemInstance* InSystemInstance, float DeltaSeconds) override;
 	virtual bool HasPreSimulateTick() const override { return true; }
 	virtual bool HasPostSimulateTick() const override { return true; }
 
 	virtual bool CanExposeVariables() const override { return true;}
-	virtual void GetExposedVariables(TArray<FNiagaraVariableBase>& OutVariables) const override;
-	virtual bool GetExposedVariableValue(const FNiagaraVariableBase& InVariable, void* InPerInstanceData, FNiagaraSystemInstance* InSystemInstance, void* OutData) const override;
+	NIAGARA_API virtual void GetExposedVariables(TArray<FNiagaraVariableBase>& OutVariables) const override;
+	NIAGARA_API virtual bool GetExposedVariableValue(const FNiagaraVariableBase& InVariable, void* InPerInstanceData, FNiagaraSystemInstance* InSystemInstance, void* OutData) const override;
 
 	virtual bool CanRenderVariablesToCanvas() const { return true; }
-	virtual void GetCanvasVariables(TArray<FNiagaraVariableBase>& OutVariables) const override;
-	virtual bool RenderVariableToCanvas(FNiagaraSystemInstanceID SystemInstanceID, FName VariableName, class FCanvas* Canvas, const FIntRect& DrawRect) const override;
+	NIAGARA_API virtual void GetCanvasVariables(TArray<FNiagaraVariableBase>& OutVariables) const override;
+	NIAGARA_API virtual bool RenderVariableToCanvas(FNiagaraSystemInstanceID SystemInstanceID, FName VariableName, class FCanvas* Canvas, const FIntRect& DrawRect) const override;
 	//~ UNiagaraDataInterface interface END
 
-	void VMGetSize(FVectorVMExternalFunctionContext& Context);
-	void VMSetSize(FVectorVMExternalFunctionContext& Context);
-	void VMGetNumMipLevels(FVectorVMExternalFunctionContext& Context);
-	void VMSetFormat(FVectorVMExternalFunctionContext& Context);
+	NIAGARA_API void VMGetSize(FVectorVMExternalFunctionContext& Context);
+	NIAGARA_API void VMSetSize(FVectorVMExternalFunctionContext& Context);
+	NIAGARA_API void VMGetNumMipLevels(FVectorVMExternalFunctionContext& Context);
+	NIAGARA_API void VMSetFormat(FVectorVMExternalFunctionContext& Context);
 
 	UPROPERTY(EditAnywhere, Category = "Render Target", meta = (DisplayPriority = 2, EditCondition = "!bInheritUserParameterSettings", EditConditionHides))
 	FIntPoint Size;
@@ -175,7 +175,7 @@ public:
 	FNiagaraUserParameterBinding RenderTargetUserParameter;
 
 protected:
-	static FNiagaraVariableBase ExposedRTVar;
+	static NIAGARA_API FNiagaraVariableBase ExposedRTVar;
 
 	TMap<FNiagaraSystemInstanceID, FRenderTarget2DRWInstanceData_GameThread*> SystemInstancesToProxyData_GT;
 	
diff --git a/Engine/Plugins/FX/Niagara/Source/Niagara/Classes/NiagaraDataInterfaceRenderTarget2DArray.h b/Engine/Plugins/FX/Niagara/Source/Niagara/Classes/NiagaraDataInterfaceRenderTarget2DArray.h
--- a/Engine/Plugins/FX/Niagara/Source/Niagara/Classes/NiagaraDataInterfaceRenderTarget2DArray.h
+++ b/Engine/Plugins/FX/Niagara/Source/Niagara/Classes/NiagaraDataInterfaceRenderTarget2DArray.h
@@ -81,53 +81,53 @@ struct FNiagaraDataInterfaceProxyRenderTarget2DArrayProxy : public FNiagaraDataI
 	TMap<FNiagaraSystemInstanceID, FRenderTarget2DArrayRWInstanceData_RenderThread> SystemInstancesToProxyData_RT;
 };
 
-UCLASS(EditInlineNew, Category = "Rendering", CollapseCategories, meta = (DisplayName = "Render Target 2D Array"), Blueprintable, BlueprintType)
-class NIAGARA_API UNiagaraDataInterfaceRenderTarget2DArray : public UNiagaraDataInterfaceRWBase
+UCLASS(EditInlineNew, Category = "Rendering", CollapseCategories, meta = (DisplayName = "Render Target 2D Array"), Blueprintable, BlueprintType, MinimalAPI)
+class UNiagaraDataInterfaceRenderTarget2DArray : public UNiagaraDataInterfaceRWBase
 {
 	GENERATED_UCLASS_BODY()
 
 public:
-	virtual void PostInitProperties() override;
+	NIAGARA_API virtual void PostInitProperties() override;
 	
 	//~ UNiagaraDataInterface interface
 	// VM functionality
 	virtual bool CanExecuteOnTarget(ENiagaraSimTarget Target)const override { return true; }
-	virtual void GetFunctions(TArray<FNiagaraFunctionSignature>& OutFunctions) override;
+	NIAGARA_API virtual void GetFunctions(TArray<FNiagaraFunctionSignature>& OutFunctions) override;
 #if WITH_EDITORONLY_DATA
-	virtual bool UpgradeFunctionCall(FNiagaraFunctionSignature& FunctionSignature) override;
+	NIAGARA_API virtual bool UpgradeFunctionCall(FNiagaraFunctionSignature& FunctionSignature) override;
 #endif
-	virtual void GetVMExternalFunction(const FVMExternalFunctionBindingInfo& BindingInfo, void* InstanceData, FVMExternalFunction &OutFunc) override;
+	NIAGARA_API virtual void GetVMExternalFunction(const FVMExternalFunctionBindingInfo& BindingInfo, void* InstanceData, FVMExternalFunction &OutFunc) override;
 
-	virtual bool Equals(const UNiagaraDataInterface* Other) const override;
-	virtual bool CopyToInternal(UNiagaraDataInterface* Destination) const override;
+	NIAGARA_API virtual bool Equals(const UNiagaraDataInterface* Other) const override;
+	NIAGARA_API virtual bool CopyToInternal(UNiagaraDataInterface* Destination) const override;
 
 	// GPU sim functionality
 #if WITH_EDITORONLY_DATA
-	virtual bool AppendCompileHash(FNiagaraCompileHashVisitor* InVisitor) const override;
-	virtual void GetParameterDefinitionHLSL(const FNiagaraDataInterfaceGPUParamInfo& ParamInfo, FString& OutHLSL) override;
-	virtual bool GetFunctionHLSL(const FNiagaraDataInterfaceGPUParamInfo& ParamInfo, const FNiagaraDataInterfaceGeneratedFunction& FunctionInfo, int FunctionInstanceIndex, FString& OutHLSL) override;
+	NIAGARA_API virtual bool AppendCompileHash(FNiagaraCompileHashVisitor* InVisitor) const override;
+	NIAGARA_API virtual void GetParameterDefinitionHLSL(const FNiagaraDataInterfaceGPUParamInfo& ParamInfo, FString& OutHLSL) override;
+	NIAGARA_API virtual bool GetFunctionHLSL(const FNiagaraDataInterfaceGPUParamInfo& ParamInfo, const FNiagaraDataInterfaceGeneratedFunction& FunctionInfo, int FunctionInstanceIndex, FString& OutHLSL) override;
 #endif
-	virtual void BuildShaderParameters(FNiagaraShaderParametersBuilder& ShaderParametersBuilder) const override;
-	virtual void SetShaderParameters(const FNiagaraDataInterfaceSetShaderParametersContext& Context) const override;
+	NIAGARA_API virtual void BuildShaderParameters(FNiagaraShaderParametersBuilder& ShaderParametersBuilder) const override;
+	NIAGARA_API virtual void SetShaderParameters(const FNiagaraDataInterfaceSetShaderParametersContext& Context) const override;
 
 	virtual void ProvidePerInstanceDataForRenderThread(void* DataForRenderThread, void* PerInstanceData, const FNiagaraSystemInstanceID& SystemInstance) override {}
-	virtual bool InitPerInstanceData(void* PerInstanceData, FNiagaraSystemInstance* SystemInstance) override;
-	virtual void DestroyPerInstanceData(void* PerInstanceData, FNiagaraSystemInstance* SystemInstance) override;
-	virtual bool PerInstanceTick(void* PerInstanceData, FNiagaraSystemInstance* SystemInstance, float DeltaSeconds) override;
+	NIAGARA_API virtual bool InitPerInstanceData(void* PerInstanceData, FNiagaraSystemInstance* SystemInstance) override;
+	NIAGARA_API virtual void DestroyPerInstanceData(void* PerInstanceData, FNiagaraSystemInstance* SystemInstance) override;
+	NIAGARA_API virtual bool PerInstanceTick(void* PerInstanceData, FNiagaraSystemInstance* SystemInstance, float DeltaSeconds) override;
 	virtual int32 PerInstanceDataSize()const override { return sizeof(FRenderTarget2DArrayRWInstanceData_GameThread); }
-	virtual bool PerInstanceTickPostSimulate(void* PerInstanceData, FNiagaraSystemInstance* InSystemInstance, float DeltaSeconds) override;
+	NIAGARA_API virtual bool PerInstanceTickPostSimulate(void* PerInstanceData, FNiagaraSystemInstance* InSystemInstance, float DeltaSeconds) override;
 	virtual bool HasPreSimulateTick() const override { return true; }
 	virtual bool HasPostSimulateTick() const override { return true; }
 
 	virtual bool CanExposeVariables() const override { return true;}
-	virtual void GetExposedVariables(TArray<FNiagaraVariableBase>& OutVariables) const override;
-	virtual bool GetExposedVariableValue(const FNiagaraVariableBase& InVariable, void* InPerInstanceData, FNiagaraSystemInstance* InSystemInstance, void* OutData) const override;
+	NIAGARA_API virtual void GetExposedVariables(TArray<FNiagaraVariableBase>& OutVariables) const override;
+	NIAGARA_API virtual bool GetExposedVariableValue(const FNiagaraVariableBase& InVariable, void* InPerInstanceData, FNiagaraSystemInstance* InSystemInstance, void* OutData) const override;
 
 	//~ UNiagaraDataInterface interface END
-	void VMGetSize(FVectorVMExternalFunctionContext& Context);
-	void VMSetSize(FVectorVMExternalFunctionContext& Context);
-	void VMGetNumMipLevels(FVectorVMExternalFunctionContext& Context);
-	void VMSetFormat(FVectorVMExternalFunctionContext& Context);
+	NIAGARA_API void VMGetSize(FVectorVMExternalFunctionContext& Context);
+	NIAGARA_API void VMSetSize(FVectorVMExternalFunctionContext& Context);
+	NIAGARA_API void VMGetNumMipLevels(FVectorVMExternalFunctionContext& Context);
+	NIAGARA_API void VMSetFormat(FVectorVMExternalFunctionContext& Context);
 
 	UPROPERTY(EditAnywhere, Category = "Render Target", meta = (DisplayPriority = 2, EditCondition = "!bInheritUserParameterSettings", EditConditionHides))
 	FIntVector Size;
@@ -160,7 +160,7 @@ public:
 	FNiagaraUserParameterBinding RenderTargetUserParameter;
 
 protected:
-	static FNiagaraVariableBase ExposedRTVar;
+	static NIAGARA_API FNiagaraVariableBase ExposedRTVar;
 	
 	UPROPERTY(Transient, DuplicateTransient)
 	TMap<uint64, TObjectPtr<UTextureRenderTarget2DArray>> ManagedRenderTargets;
diff --git a/Engine/Plugins/FX/Niagara/Source/Niagara/Classes/NiagaraDataInterfaceRenderTargetCube.h b/Engine/Plugins/FX/Niagara/Source/Niagara/Classes/NiagaraDataInterfaceRenderTargetCube.h
--- a/Engine/Plugins/FX/Niagara/Source/Niagara/Classes/NiagaraDataInterfaceRenderTargetCube.h
+++ b/Engine/Plugins/FX/Niagara/Source/Niagara/Classes/NiagaraDataInterfaceRenderTargetCube.h
@@ -74,53 +74,53 @@ struct FNiagaraDataInterfaceProxyRenderTargetCubeProxy : public FNiagaraDataInte
 	TMap<FNiagaraSystemInstanceID, FRenderTargetCubeRWInstanceData_RenderThread> SystemInstancesToProxyData_RT;
 };
 
-UCLASS(EditInlineNew, Category = "Rendering", CollapseCategories, meta = (DisplayName = "Render Target Cube"), Blueprintable, BlueprintType)
-class NIAGARA_API UNiagaraDataInterfaceRenderTargetCube : public UNiagaraDataInterfaceRWBase
+UCLASS(EditInlineNew, Category = "Rendering", CollapseCategories, meta = (DisplayName = "Render Target Cube"), Blueprintable, BlueprintType, MinimalAPI)
+class UNiagaraDataInterfaceRenderTargetCube : public UNiagaraDataInterfaceRWBase
 {
 	GENERATED_UCLASS_BODY()
 
 public:
-	virtual void PostInitProperties() override;
+	NIAGARA_API virtual void PostInitProperties() override;
 	
 	//~ UNiagaraDataInterface interface
 	// VM functionality
 	virtual bool CanExecuteOnTarget(ENiagaraSimTarget Target)const override { return true; }
-	virtual void GetFunctions(TArray<FNiagaraFunctionSignature>& OutFunctions) override;
+	NIAGARA_API virtual void GetFunctions(TArray<FNiagaraFunctionSignature>& OutFunctions) override;
 #if WITH_EDITORONLY_DATA
-	virtual bool UpgradeFunctionCall(FNiagaraFunctionSignature& FunctionSignature) override;
+	NIAGARA_API virtual bool UpgradeFunctionCall(FNiagaraFunctionSignature& FunctionSignature) override;
 #endif
-	virtual void GetVMExternalFunction(const FVMExternalFunctionBindingInfo& BindingInfo, void* InstanceData, FVMExternalFunction &OutFunc) override;
+	NIAGARA_API virtual void GetVMExternalFunction(const FVMExternalFunctionBindingInfo& BindingInfo, void* InstanceData, FVMExternalFunction &OutFunc) override;
 
-	virtual bool Equals(const UNiagaraDataInterface* Other) const override;
-	virtual bool CopyToInternal(UNiagaraDataInterface* Destination) const override;
+	NIAGARA_API virtual bool Equals(const UNiagaraDataInterface* Other) const override;
+	NIAGARA_API virtual bool CopyToInternal(UNiagaraDataInterface* Destination) const override;
 
 	// GPU sim functionality
 #if WITH_EDITORONLY_DATA
-	virtual bool AppendCompileHash(FNiagaraCompileHashVisitor* InVisitor) const override;
-	virtual void GetParameterDefinitionHLSL(const FNiagaraDataInterfaceGPUParamInfo& ParamInfo, FString& OutHLSL) override;
-	virtual bool GetFunctionHLSL(const FNiagaraDataInterfaceGPUParamInfo& ParamInfo, const FNiagaraDataInterfaceGeneratedFunction& FunctionInfo, int FunctionInstanceIndex, FString& OutHLSL) override;
+	NIAGARA_API virtual bool AppendCompileHash(FNiagaraCompileHashVisitor* InVisitor) const override;
+	NIAGARA_API virtual void GetParameterDefinitionHLSL(const FNiagaraDataInterfaceGPUParamInfo& ParamInfo, FString& OutHLSL) override;
+	NIAGARA_API virtual bool GetFunctionHLSL(const FNiagaraDataInterfaceGPUParamInfo& ParamInfo, const FNiagaraDataInterfaceGeneratedFunction& FunctionInfo, int FunctionInstanceIndex, FString& OutHLSL) override;
 #endif
-	virtual void BuildShaderParameters(FNiagaraShaderParametersBuilder& ShaderParametersBuilder) const override;
-	virtual void SetShaderParameters(const FNiagaraDataInterfaceSetShaderParametersContext& Context) const override;
+	NIAGARA_API virtual void BuildShaderParameters(FNiagaraShaderParametersBuilder& ShaderParametersBuilder) const override;
+	NIAGARA_API virtual void SetShaderParameters(const FNiagaraDataInterfaceSetShaderParametersContext& Context) const override;
 
 	virtual void ProvidePerInstanceDataForRenderThread(void* DataForRenderThread, void* PerInstanceData, const FNiagaraSystemInstanceID& SystemInstance) override {}
-	virtual bool InitPerInstanceData(void* PerInstanceData, FNiagaraSystemInstance* SystemInstance) override;
-	virtual void DestroyPerInstanceData(void* PerInstanceData, FNiagaraSystemInstance* SystemInstance) override;
-	virtual bool PerInstanceTick(void* PerInstanceData, FNiagaraSystemInstance* SystemInstance, float DeltaSeconds) override;
+	NIAGARA_API virtual bool InitPerInstanceData(void* PerInstanceData, FNiagaraSystemInstance* SystemInstance) override;
+	NIAGARA_API virtual void DestroyPerInstanceData(void* PerInstanceData, FNiagaraSystemInstance* SystemInstance) override;
+	NIAGARA_API virtual bool PerInstanceTick(void* PerInstanceData, FNiagaraSystemInstance* SystemInstance, float DeltaSeconds) override;
 	virtual int32 PerInstanceDataSize()const override { return sizeof(FRenderTargetCubeRWInstanceData_GameThread); }
-	virtual bool PerInstanceTickPostSimulate(void* PerInstanceData, FNiagaraSystemInstance* InSystemInstance, float DeltaSeconds) override;
+	NIAGARA_API virtual bool PerInstanceTickPostSimulate(void* PerInstanceData, FNiagaraSystemInstance* InSystemInstance, float DeltaSeconds) override;
 	virtual bool HasPreSimulateTick() const override { return true; }
 	virtual bool HasPostSimulateTick() const override { return true; }
 
 	virtual bool CanExposeVariables() const override { return true;}
-	virtual void GetExposedVariables(TArray<FNiagaraVariableBase>& OutVariables) const override;
-	virtual bool GetExposedVariableValue(const FNiagaraVariableBase& InVariable, void* InPerInstanceData, FNiagaraSystemInstance* InSystemInstance, void* OutData) const override;
+	NIAGARA_API virtual void GetExposedVariables(TArray<FNiagaraVariableBase>& OutVariables) const override;
+	NIAGARA_API virtual bool GetExposedVariableValue(const FNiagaraVariableBase& InVariable, void* InPerInstanceData, FNiagaraSystemInstance* InSystemInstance, void* OutData) const override;
 	//~ UNiagaraDataInterface interface END
 	
-	void VMGetSize(FVectorVMExternalFunctionContext& Context);
-	void VMSetSize(FVectorVMExternalFunctionContext& Context);
-	void VMGetNumMipLevels(FVectorVMExternalFunctionContext& Context);
-	void VMSetFormat(FVectorVMExternalFunctionContext& Context);
+	NIAGARA_API void VMGetSize(FVectorVMExternalFunctionContext& Context);
+	NIAGARA_API void VMSetSize(FVectorVMExternalFunctionContext& Context);
+	NIAGARA_API void VMGetNumMipLevels(FVectorVMExternalFunctionContext& Context);
+	NIAGARA_API void VMSetFormat(FVectorVMExternalFunctionContext& Context);
 
 	UPROPERTY(EditAnywhere, Category = "Render Target", meta = (DisplayPriority = 2, EditCondition = "!bInheritUserParameterSettings", EditConditionHides))
 	int Size;
@@ -155,7 +155,7 @@ public:
 protected:
 	//~ UNiagaraDataInterface interface END
 
-	static FNiagaraVariableBase ExposedRTVar;
+	static NIAGARA_API FNiagaraVariableBase ExposedRTVar;
 	
 	UPROPERTY(Transient, DuplicateTransient)
 	TMap<uint64, TObjectPtr<UTextureRenderTargetCube>> ManagedRenderTargets;
diff --git a/Engine/Plugins/FX/Niagara/Source/Niagara/Classes/NiagaraDataInterfaceRenderTargetVolume.h b/Engine/Plugins/FX/Niagara/Source/Niagara/Classes/NiagaraDataInterfaceRenderTargetVolume.h
--- a/Engine/Plugins/FX/Niagara/Source/Niagara/Classes/NiagaraDataInterfaceRenderTargetVolume.h
+++ b/Engine/Plugins/FX/Niagara/Source/Niagara/Classes/NiagaraDataInterfaceRenderTargetVolume.h
@@ -74,67 +74,67 @@ struct FNiagaraDataInterfaceProxyRenderTargetVolumeProxy : public FNiagaraDataIn
 	TMap<FNiagaraSystemInstanceID, FRenderTargetVolumeRWInstanceData_RenderThread> SystemInstancesToProxyData_RT;
 };
 
-UCLASS(EditInlineNew, Category = "Rendering", CollapseCategories, meta = (DisplayName = "Render Target Volume"), Blueprintable, BlueprintType)
-class NIAGARA_API UNiagaraDataInterfaceRenderTargetVolume : public UNiagaraDataInterfaceRWBase, public INiagaraSimCacheCustomStorageInterface
+UCLASS(EditInlineNew, Category = "Rendering", CollapseCategories, meta = (DisplayName = "Render Target Volume"), Blueprintable, BlueprintType, MinimalAPI)
+class UNiagaraDataInterfaceRenderTargetVolume : public UNiagaraDataInterfaceRWBase, public INiagaraSimCacheCustomStorageInterface
 {
 	GENERATED_UCLASS_BODY()
 
 public:
-	virtual void PostInitProperties() override;
+	NIAGARA_API virtual void PostInitProperties() override;
 	
 	//~ UNiagaraDataInterface interface
 	// VM functionality
 	virtual bool CanExecuteOnTarget(ENiagaraSimTarget Target)const override { return true; }
-	virtual void GetFunctions(TArray<FNiagaraFunctionSignature>& OutFunctions) override;
+	NIAGARA_API virtual void GetFunctions(TArray<FNiagaraFunctionSignature>& OutFunctions) override;
 #if WITH_EDITORONLY_DATA
-	virtual bool UpgradeFunctionCall(FNiagaraFunctionSignature& FunctionSignature) override;
+	NIAGARA_API virtual bool UpgradeFunctionCall(FNiagaraFunctionSignature& FunctionSignature) override;
 #endif
-	virtual void GetVMExternalFunction(const FVMExternalFunctionBindingInfo& BindingInfo, void* InstanceData, FVMExternalFunction &OutFunc) override;
+	NIAGARA_API virtual void GetVMExternalFunction(const FVMExternalFunctionBindingInfo& BindingInfo, void* InstanceData, FVMExternalFunction &OutFunc) override;
 
-	virtual bool Equals(const UNiagaraDataInterface* Other) const override;
-	virtual bool CopyToInternal(UNiagaraDataInterface* Destination) const override;
+	NIAGARA_API virtual bool Equals(const UNiagaraDataInterface* Other) const override;
+	NIAGARA_API virtual bool CopyToInternal(UNiagaraDataInterface* Destination) const override;
 
 	// GPU sim functionality
 #if WITH_EDITOR
-	virtual bool ShouldCompile(EShaderPlatform ShaderPlatform) const override;
+	NIAGARA_API virtual bool ShouldCompile(EShaderPlatform ShaderPlatform) const override;
 #endif
 #if WITH_EDITORONLY_DATA
-	virtual bool AppendCompileHash(FNiagaraCompileHashVisitor* InVisitor) const override;
-	virtual void GetParameterDefinitionHLSL(const FNiagaraDataInterfaceGPUParamInfo& ParamInfo, FString& OutHLSL) override;
-	virtual bool GetFunctionHLSL(const FNiagaraDataInterfaceGPUParamInfo& ParamInfo, const FNiagaraDataInterfaceGeneratedFunction& FunctionInfo, int FunctionInstanceIndex, FString& OutHLSL) override;
+	NIAGARA_API virtual bool AppendCompileHash(FNiagaraCompileHashVisitor* InVisitor) const override;
+	NIAGARA_API virtual void GetParameterDefinitionHLSL(const FNiagaraDataInterfaceGPUParamInfo& ParamInfo, FString& OutHLSL) override;
+	NIAGARA_API virtual bool GetFunctionHLSL(const FNiagaraDataInterfaceGPUParamInfo& ParamInfo, const FNiagaraDataInterfaceGeneratedFunction& FunctionInfo, int FunctionInstanceIndex, FString& OutHLSL) override;
 #endif
-	virtual void BuildShaderParameters(FNiagaraShaderParametersBuilder& ShaderParametersBuilder) const override;
-	virtual void SetShaderParameters(const FNiagaraDataInterfaceSetShaderParametersContext& Context) const override;
+	NIAGARA_API virtual void BuildShaderParameters(FNiagaraShaderParametersBuilder& ShaderParametersBuilder) const override;
+	NIAGARA_API virtual void SetShaderParameters(const FNiagaraDataInterfaceSetShaderParametersContext& Context) const override;
 
 #if WITH_EDITOR
 	virtual ENiagaraGpuDispatchType GetGpuDispatchType() const override { return ENiagaraGpuDispatchType::ThreeD; }
 #endif
 
 	virtual void ProvidePerInstanceDataForRenderThread(void* DataForRenderThread, void* PerInstanceData, const FNiagaraSystemInstanceID& SystemInstance) override {}
-	virtual bool InitPerInstanceData(void* PerInstanceData, FNiagaraSystemInstance* SystemInstance) override;
-	virtual void DestroyPerInstanceData(void* PerInstanceData, FNiagaraSystemInstance* SystemInstance) override;
-	virtual bool PerInstanceTick(void* PerInstanceData, FNiagaraSystemInstance* SystemInstance, float DeltaSeconds) override;
+	NIAGARA_API virtual bool InitPerInstanceData(void* PerInstanceData, FNiagaraSystemInstance* SystemInstance) override;
+	NIAGARA_API virtual void DestroyPerInstanceData(void* PerInstanceData, FNiagaraSystemInstance* SystemInstance) override;
+	NIAGARA_API virtual bool PerInstanceTick(void* PerInstanceData, FNiagaraSystemInstance* SystemInstance, float DeltaSeconds) override;
 	virtual int32 PerInstanceDataSize()const override { return sizeof(FRenderTargetVolumeRWInstanceData_GameThread); }
-	virtual bool PerInstanceTickPostSimulate(void* PerInstanceData, FNiagaraSystemInstance* InSystemInstance, float DeltaSeconds) override;
+	NIAGARA_API virtual bool PerInstanceTickPostSimulate(void* PerInstanceData, FNiagaraSystemInstance* InSystemInstance, float DeltaSeconds) override;
 	virtual bool HasPreSimulateTick() const override { return true; }
 	virtual bool HasPostSimulateTick() const override { return true; }
 
 	virtual bool CanExposeVariables() const override { return true;}
-	virtual void GetExposedVariables(TArray<FNiagaraVariableBase>& OutVariables) const override;
-	virtual bool GetExposedVariableValue(const FNiagaraVariableBase& InVariable, void* InPerInstanceData, FNiagaraSystemInstance* InSystemInstance, void* OutData) const override;
+	NIAGARA_API virtual void GetExposedVariables(TArray<FNiagaraVariableBase>& OutVariables) const override;
+	NIAGARA_API virtual bool GetExposedVariableValue(const FNiagaraVariableBase& InVariable, void* InPerInstanceData, FNiagaraSystemInstance* InSystemInstance, void* OutData) const override;
 	//~ UNiagaraDataInterface interface END
 
 	//~ INiagaraSimCacheCustomStorageInterface interface BEGIN
-	virtual UObject* SimCacheBeginWrite(UObject* SimCache, FNiagaraSystemInstance* NiagaraSystemInstance, const void* OptionalPerInstanceData, FNiagaraSimCacheFeedbackContext& FeedbackContext) const override;
-	virtual bool SimCacheWriteFrame(UObject* StorageObject, int FrameIndex, FNiagaraSystemInstance* SystemInstance, const void* OptionalPerInstanceData, FNiagaraSimCacheFeedbackContext& FeedbackContext) const override;
-	virtual bool SimCacheEndWrite(UObject* StorageObject) const override;
-	virtual bool SimCacheReadFrame(UObject* StorageObject, int FrameA, int FrameB, float Interp, FNiagaraSystemInstance* SystemInstance, void* OptionalPerInstanceData) override;
+	NIAGARA_API virtual UObject* SimCacheBeginWrite(UObject* SimCache, FNiagaraSystemInstance* NiagaraSystemInstance, const void* OptionalPerInstanceData, FNiagaraSimCacheFeedbackContext& FeedbackContext) const override;
+	NIAGARA_API virtual bool SimCacheWriteFrame(UObject* StorageObject, int FrameIndex, FNiagaraSystemInstance* SystemInstance, const void* OptionalPerInstanceData, FNiagaraSimCacheFeedbackContext& FeedbackContext) const override;
+	NIAGARA_API virtual bool SimCacheEndWrite(UObject* StorageObject) const override;
+	NIAGARA_API virtual bool SimCacheReadFrame(UObject* StorageObject, int FrameA, int FrameB, float Interp, FNiagaraSystemInstance* SystemInstance, void* OptionalPerInstanceData) override;
 	//~ UNiagaraDataInterface interface END
 
-	void VMGetSize(FVectorVMExternalFunctionContext& Context);
-	void VMSetSize(FVectorVMExternalFunctionContext& Context);
-	void VMGetNumMipLevels(FVectorVMExternalFunctionContext& Context);
-	void VMSetFormat(FVectorVMExternalFunctionContext& Context);
+	NIAGARA_API void VMGetSize(FVectorVMExternalFunctionContext& Context);
+	NIAGARA_API void VMSetSize(FVectorVMExternalFunctionContext& Context);
+	NIAGARA_API void VMGetNumMipLevels(FVectorVMExternalFunctionContext& Context);
+	NIAGARA_API void VMSetFormat(FVectorVMExternalFunctionContext& Context);
 
 	UPROPERTY(EditAnywhere, Category = "Render Target", meta = (DisplayPriority = 2, EditCondition = "!bInheritUserParameterSettings", EditConditionHides))
 	FIntVector Size;
@@ -169,7 +169,7 @@ public:
 protected:
 	//~ UNiagaraDataInterface interface END
 
-	static FNiagaraVariableBase ExposedRTVar;
+	static NIAGARA_API FNiagaraVariableBase ExposedRTVar;
 	
 	UPROPERTY(Transient, DuplicateTransient)
 	TMap<uint64, TObjectPtr<UTextureRenderTargetVolume>> ManagedRenderTargets;
