[Niagara]

diff --git a/Engine/Plugins/FX/Niagara/Source/Niagara/Classes/NiagaraDataInterfaceSkeletalMesh.h b/Engine/Plugins/FX/Niagara/Source/Niagara/Classes/NiagaraDataInterfaceSkeletalMesh.h
--- a/Engine/Plugins/FX/Niagara/Source/Niagara/Classes/NiagaraDataInterfaceSkeletalMesh.h
+++ b/Engine/Plugins/FX/Niagara/Source/Niagara/Classes/NiagaraDataInterfaceSkeletalMesh.h
@@ -256,31 +256,31 @@ struct FSkeletalMeshConnectivityUsage
 	bool RequiresGpuAccess = false;
 };
 
-struct NIAGARA_API FSkeletalMeshConnectivityHandle
+struct FSkeletalMeshConnectivityHandle
 {
-	FSkeletalMeshConnectivityHandle();
-	FSkeletalMeshConnectivityHandle(FSkeletalMeshConnectivityUsage InUsage, const TSharedPtr<struct FSkeletalMeshConnectivity>& InConnectivityData, bool bNeedsDataImmediately);
+	NIAGARA_API FSkeletalMeshConnectivityHandle();
+	NIAGARA_API FSkeletalMeshConnectivityHandle(FSkeletalMeshConnectivityUsage InUsage, const TSharedPtr<struct FSkeletalMeshConnectivity>& InConnectivityData, bool bNeedsDataImmediately);
 	FSkeletalMeshConnectivityHandle(const FSkeletalMeshConnectivityHandle& Other) = delete;
-	FSkeletalMeshConnectivityHandle(FSkeletalMeshConnectivityHandle&& Other) noexcept;
-	~FSkeletalMeshConnectivityHandle();
+	NIAGARA_API FSkeletalMeshConnectivityHandle(FSkeletalMeshConnectivityHandle&& Other) noexcept;
+	NIAGARA_API ~FSkeletalMeshConnectivityHandle();
 
 	FSkeletalMeshConnectivityHandle& operator=(const FSkeletalMeshConnectivityHandle& Other) = delete;
-	FSkeletalMeshConnectivityHandle& operator=(FSkeletalMeshConnectivityHandle&& Other) noexcept;
-	explicit operator bool() const;
+	NIAGARA_API FSkeletalMeshConnectivityHandle& operator=(FSkeletalMeshConnectivityHandle&& Other) noexcept;
+	NIAGARA_API explicit operator bool() const;
 
 	FSkeletalMeshConnectivityUsage Usage;
 
-	int32 GetAdjacentTriangleIndex(int32 VertexIndex, int32 AdjacencyIndex) const;
+	NIAGARA_API int32 GetAdjacentTriangleIndex(int32 VertexIndex, int32 AdjacencyIndex) const;
 
-	const FSkeletalMeshConnectivityProxy* GetProxy() const;
+	NIAGARA_API const FSkeletalMeshConnectivityProxy* GetProxy() const;
 
-	void PinAndInvalidateHandle();
+	NIAGARA_API void PinAndInvalidateHandle();
 
 private:
 	TSharedPtr<FSkeletalMeshConnectivity> ConnectivityData;
 };
 
-class NIAGARA_API FNDI_SkeletalMesh_GeneratedData : public FNDI_GeneratedData
+class FNDI_SkeletalMesh_GeneratedData : public FNDI_GeneratedData
 {
 	FRWLock CachedSkinningDataGuard;
 	TMap<TWeakObjectPtr<USkeletalMeshComponent>, TSharedPtr<FSkeletalMeshSkinningData> > CachedSkinningData;
@@ -292,11 +292,11 @@ class NIAGARA_API FNDI_SkeletalMesh_GeneratedData : public FNDI_GeneratedData
 	TArray<TSharedPtr<FSkeletalMeshConnectivity>> CachedConnectivity;
 
 public:
-	FSkeletalMeshSkinningDataHandle GetCachedSkinningData(TWeakObjectPtr<USkeletalMeshComponent>& InComponent, FSkeletalMeshSkinningDataUsage Usage, bool bNeedsDataImmediately);
-	FSkeletalMeshUvMappingHandle GetCachedUvMapping(TWeakObjectPtr<USkeletalMesh>& InMeshObject, int32 InLodIndex, int32 InUvSetIndex, FMeshUvMappingUsage Usage, bool bNeedsDataImmediately);
-	FSkeletalMeshConnectivityHandle GetCachedConnectivity(TWeakObjectPtr<USkeletalMesh>& InMeshObject, int32 InLodIndex, FSkeletalMeshConnectivityUsage Usage, bool bNeedsDataImmediately);
+	NIAGARA_API FSkeletalMeshSkinningDataHandle GetCachedSkinningData(TWeakObjectPtr<USkeletalMeshComponent>& InComponent, FSkeletalMeshSkinningDataUsage Usage, bool bNeedsDataImmediately);
+	NIAGARA_API FSkeletalMeshUvMappingHandle GetCachedUvMapping(TWeakObjectPtr<USkeletalMesh>& InMeshObject, int32 InLodIndex, int32 InUvSetIndex, FMeshUvMappingUsage Usage, bool bNeedsDataImmediately);
+	NIAGARA_API FSkeletalMeshConnectivityHandle GetCachedConnectivity(TWeakObjectPtr<USkeletalMesh>& InMeshObject, int32 InLodIndex, FSkeletalMeshConnectivityUsage Usage, bool bNeedsDataImmediately);
 
-	virtual void Tick(ETickingGroup TickGroup, float DeltaSeconds) override;
+	NIAGARA_API virtual void Tick(ETickingGroup TickGroup, float DeltaSeconds) override;
 
 	static TypeHash GetTypeHash()
 	{
@@ -674,8 +674,8 @@ struct FNDISkeletalMesh_InstanceData
 };
 
 /** Data Interface allowing sampling of skeletal meshes. */
-UCLASS(EditInlineNew, Category = "Meshes", meta = (DisplayName = "Skeletal Mesh"))
-class NIAGARA_API UNiagaraDataInterfaceSkeletalMesh : public UNiagaraDataInterface
+UCLASS(EditInlineNew, Category = "Meshes", meta = (DisplayName = "Skeletal Mesh"), MinimalAPI)
+class UNiagaraDataInterfaceSkeletalMesh : public UNiagaraDataInterface
 {
 	GENERATED_UCLASS_BODY()
 
@@ -758,84 +758,84 @@ public:
 	uint32 ChangeId;
 
 	//~ UObject interface
-	virtual void PostInitProperties()override;
-	virtual void PostLoad()override;
+	NIAGARA_API virtual void PostInitProperties()override;
+	NIAGARA_API virtual void PostLoad()override;
 #if WITH_EDITOR
-	virtual void PostEditChangeProperty(struct FPropertyChangedEvent& PropertyChangedEvent) override;
-	virtual bool CanEditChange(const FProperty* InProperty) const override;
+	NIAGARA_API virtual void PostEditChangeProperty(struct FPropertyChangedEvent& PropertyChangedEvent) override;
+	NIAGARA_API virtual bool CanEditChange(const FProperty* InProperty) const override;
 #endif
 	//~ UObject interface END
 
 
 	//~ UNiagaraDataInterface interface
-	virtual bool InitPerInstanceData(void* PerInstanceData, FNiagaraSystemInstance* SystemInstance)override;
-	virtual void DestroyPerInstanceData(void* PerInstanceData, FNiagaraSystemInstance* SystemInstance)override;
-	virtual bool PerInstanceTick(void* PerInstanceData, FNiagaraSystemInstance* SystemInstance, float DeltaSeconds) override;
+	NIAGARA_API virtual bool InitPerInstanceData(void* PerInstanceData, FNiagaraSystemInstance* SystemInstance)override;
+	NIAGARA_API virtual void DestroyPerInstanceData(void* PerInstanceData, FNiagaraSystemInstance* SystemInstance)override;
+	NIAGARA_API virtual bool PerInstanceTick(void* PerInstanceData, FNiagaraSystemInstance* SystemInstance, float DeltaSeconds) override;
 	virtual int32 PerInstanceDataSize()const override { return sizeof(FNDISkeletalMesh_InstanceData); }
 	virtual bool HasPreSimulateTick() const override { return true; }
 
-	virtual void GetFunctions(TArray<FNiagaraFunctionSignature>& OutFunctions)override;
-	virtual void GetVMExternalFunction(const FVMExternalFunctionBindingInfo& BindingInfo, void* InstanceData, FVMExternalFunction &OutFunc)override;
-	virtual bool Equals(const UNiagaraDataInterface* Other) const override;
+	NIAGARA_API virtual void GetFunctions(TArray<FNiagaraFunctionSignature>& OutFunctions)override;
+	NIAGARA_API virtual void GetVMExternalFunction(const FVMExternalFunctionBindingInfo& BindingInfo, void* InstanceData, FVMExternalFunction &OutFunc)override;
+	NIAGARA_API virtual bool Equals(const UNiagaraDataInterface* Other) const override;
 	virtual bool CanExecuteOnTarget(ENiagaraSimTarget Target)const override { return true; }
 #if WITH_NIAGARA_DEBUGGER
-	virtual void DrawDebugHud(FNDIDrawDebugHudContext& DebugHudContext) const override;
+	NIAGARA_API virtual void DrawDebugHud(FNDIDrawDebugHudContext& DebugHudContext) const override;
 #endif
 #if WITH_EDITOR
-	virtual void GetFeedback(UNiagaraSystem* Asset, UNiagaraComponent* Component, TArray<FNiagaraDataInterfaceError>& OutErrors,
+	NIAGARA_API virtual void GetFeedback(UNiagaraSystem* Asset, UNiagaraComponent* Component, TArray<FNiagaraDataInterfaceError>& OutErrors,
 		TArray<FNiagaraDataInterfaceFeedback>& Warnings, TArray<FNiagaraDataInterfaceFeedback>& Info) override;
-	virtual void ValidateFunction(const FNiagaraFunctionSignature& Function, TArray<FText>& OutValidationErrors) override;
+	NIAGARA_API virtual void ValidateFunction(const FNiagaraFunctionSignature& Function, TArray<FText>& OutValidationErrors) override;
 #endif
 	virtual bool HasTickGroupPrereqs() const override { return true; }
-	virtual ETickingGroup CalculateTickGroup(const void* PerInstanceData) const override;
+	NIAGARA_API virtual ETickingGroup CalculateTickGroup(const void* PerInstanceData) const override;
 
 #if WITH_EDITOR
-	virtual void ModifyCompilationEnvironment(EShaderPlatform ShaderPlatform, struct FShaderCompilerEnvironment& OutEnvironment) const override;
+	NIAGARA_API virtual void ModifyCompilationEnvironment(EShaderPlatform ShaderPlatform, struct FShaderCompilerEnvironment& OutEnvironment) const override;
 #endif
 #if WITH_EDITORONLY_DATA
-	virtual bool AppendCompileHash(FNiagaraCompileHashVisitor* InVisitor) const override;
-	virtual void GetCommonHLSL(FString& OutHLSL) override;
-	virtual void GetParameterDefinitionHLSL(const FNiagaraDataInterfaceGPUParamInfo& ParamInfo, FString& OutHLSL) override;
-	virtual bool GetFunctionHLSL(const FNiagaraDataInterfaceGPUParamInfo& ParamInfo, const FNiagaraDataInterfaceGeneratedFunction& FunctionInfo, int FunctionInstanceIndex, FString& OutHLSL) override;
-	virtual bool UpgradeFunctionCall(FNiagaraFunctionSignature& FunctionSignature) override;
+	NIAGARA_API virtual bool AppendCompileHash(FNiagaraCompileHashVisitor* InVisitor) const override;
+	NIAGARA_API virtual void GetCommonHLSL(FString& OutHLSL) override;
+	NIAGARA_API virtual void GetParameterDefinitionHLSL(const FNiagaraDataInterfaceGPUParamInfo& ParamInfo, FString& OutHLSL) override;
+	NIAGARA_API virtual bool GetFunctionHLSL(const FNiagaraDataInterfaceGPUParamInfo& ParamInfo, const FNiagaraDataInterfaceGeneratedFunction& FunctionInfo, int FunctionInstanceIndex, FString& OutHLSL) override;
+	NIAGARA_API virtual bool UpgradeFunctionCall(FNiagaraFunctionSignature& FunctionSignature) override;
 #endif
-	virtual void BuildShaderParameters(FNiagaraShaderParametersBuilder& ShaderParametersBuilder) const override;
-	virtual void SetShaderParameters(const FNiagaraDataInterfaceSetShaderParametersContext& Context) const override;
+	NIAGARA_API virtual void BuildShaderParameters(FNiagaraShaderParametersBuilder& ShaderParametersBuilder) const override;
+	NIAGARA_API virtual void SetShaderParameters(const FNiagaraDataInterfaceSetShaderParametersContext& Context) const override;
 	//~ UNiagaraDataInterface interface END
 
 	/** This overload is for use when initializing per-instance data. It possibly uses the SystemInstance and instance data to initialize a user binding */
-	USkeletalMesh* GetSkeletalMesh(FNiagaraSystemInstance* SystemInstance, USceneComponent* AttachComponent, TWeakObjectPtr<USceneComponent>& SceneComponent, USkeletalMeshComponent*& FoundSkelComp, FNDISkeletalMesh_InstanceData* InstData);
+	NIAGARA_API USkeletalMesh* GetSkeletalMesh(FNiagaraSystemInstance* SystemInstance, USceneComponent* AttachComponent, TWeakObjectPtr<USceneComponent>& SceneComponent, USkeletalMeshComponent*& FoundSkelComp, FNDISkeletalMesh_InstanceData* InstData);
 	/** Finds the skeletal mesh based on settings of the DI and the hierarchy of the object provided */
-	USkeletalMesh* GetSkeletalMesh(UNiagaraComponent* Component);
+	NIAGARA_API USkeletalMesh* GetSkeletalMesh(UNiagaraComponent* Component);
 
 	USkeletalMeshComponent* GetSourceComponent() const { return SourceComponent; }
 	AActor* GetSourceActor() const { return SoftSourceActor.Get(); }
 
-	int32 CalculateLODIndexAndSamplingRegions(USkeletalMesh* InMesh, TArray<int32>& OutSamplingRegionIndices, bool& OutAllRegionsAreAreaWeighting) const;
+	NIAGARA_API int32 CalculateLODIndexAndSamplingRegions(USkeletalMesh* InMesh, TArray<int32>& OutSamplingRegionIndices, bool& OutAllRegionsAreAreaWeighting) const;
 
-	virtual void ProvidePerInstanceDataForRenderThread(void* DataForRenderThread, void* PerInstanceData, const FNiagaraSystemInstanceID& SystemInstance) override;
+	NIAGARA_API virtual void ProvidePerInstanceDataForRenderThread(void* DataForRenderThread, void* PerInstanceData, const FNiagaraSystemInstanceID& SystemInstance) override;
 
 protected:
-	virtual bool CopyToInternal(UNiagaraDataInterface* Destination) const override;
+	NIAGARA_API virtual bool CopyToInternal(UNiagaraDataInterface* Destination) const override;
 
 	// Bind/unbind delegates to release references to the source actor & component.
-	void UnbindSourceDelegates();
-	void BindSourceDelegates();
+	NIAGARA_API void UnbindSourceDelegates();
+	NIAGARA_API void BindSourceDelegates();
 
 	UFUNCTION()
-	void OnSourceEndPlay(AActor* InSource, EEndPlayReason::Type Reason);
+	NIAGARA_API void OnSourceEndPlay(AActor* InSource, EEndPlayReason::Type Reason);
 
 	//////////////////////////////////////////////////////////////////////////
 	// Misc Functions
-	void VMGetPreSkinnedLocalBounds(FVectorVMExternalFunctionContext& Context);
+	NIAGARA_API void VMGetPreSkinnedLocalBounds(FVectorVMExternalFunctionContext& Context);
 
 	//////////////////////////////////////////////////////////////////////////
 	//Triangle sampling
 	//Triangles are sampled a using MeshTriangleCoordinates which are composed of Triangle index and a bary centric coordinate on that triangle.
 public:
 
-	void GetTriangleSamplingFunctions(TArray<FNiagaraFunctionSignature>& OutFunctions);
-	void BindTriangleSamplingFunction(const FVMExternalFunctionBindingInfo& BindingInfo, FNDISkeletalMesh_InstanceData* InstData, FVMExternalFunction &OutFunc);
+	NIAGARA_API void GetTriangleSamplingFunctions(TArray<FNiagaraFunctionSignature>& OutFunctions);
+	NIAGARA_API void BindTriangleSamplingFunction(const FVMExternalFunctionBindingInfo& BindingInfo, FNDISkeletalMesh_InstanceData* InstData, FVMExternalFunction &OutFunc);
 
 	template<typename FilterMode, typename AreaWeightingMode>
 	void GetFilteredTriangleCount(FVectorVMExternalFunctionContext& Context);
@@ -849,7 +849,7 @@ public:
 	template<typename FilterMode, typename AreaWeightingMode>
 	void IsValidTriCoord(FVectorVMExternalFunctionContext& Context);
 
-	void GetTriangleData(FVectorVMExternalFunctionContext& Context);
+	NIAGARA_API void GetTriangleData(FVectorVMExternalFunctionContext& Context);
 
 	template<typename SkinningHandlerType, typename TransformHandlerType, typename VertexAccessorType, typename bInterpolated>
 	void GetTriCoordSkinnedData(FVectorVMExternalFunctionContext& Context);
@@ -857,9 +857,9 @@ public:
 	template<typename TransformHandlerType, typename bInterpolated>
 	void GetTriCoordSkinnedDataFallback(FVectorVMExternalFunctionContext& Context);
 
-	void GetTriCoordColor(FVectorVMExternalFunctionContext& Context);
+	NIAGARA_API void GetTriCoordColor(FVectorVMExternalFunctionContext& Context);
 
-	void GetTriCoordColorFallback(FVectorVMExternalFunctionContext& Context);
+	NIAGARA_API void GetTriCoordColorFallback(FVectorVMExternalFunctionContext& Context);
 
 	template<typename VertexAccessorType>
 	void GetTriCoordUV(FVectorVMExternalFunctionContext& Context);
@@ -877,8 +877,8 @@ private:
 	template<typename FilterMode, typename AreaWeightingMode>
 	FORCEINLINE int32 RandomTriIndex(FNDIRandomHelper& RandHelper, FSkeletalMeshAccessorHelper& Accessor, FNDISkeletalMesh_InstanceData* InstData, int32 InstanceIndex);
 
-	void RandomTriangle(FVectorVMExternalFunctionContext& Context);
-	void GetTriangleCount(FVectorVMExternalFunctionContext& Context);
+	NIAGARA_API void RandomTriangle(FVectorVMExternalFunctionContext& Context);
+	NIAGARA_API void GetTriangleCount(FVectorVMExternalFunctionContext& Context);
 
 	template<typename FilterMode, typename AreaWeightingMode>
 	FORCEINLINE int32 GetFilteredTriangleCount(FSkeletalMeshAccessorHelper& Accessor, FNDISkeletalMesh_InstanceData* InstData);
@@ -893,12 +893,12 @@ private:
 	//Vertex sampling done with direct vertex indices.
 public:
 
-	void GetVertexSamplingFunctions(TArray<FNiagaraFunctionSignature>& OutFunctions);
-	void BindVertexSamplingFunction(const FVMExternalFunctionBindingInfo& BindingInfo, FNDISkeletalMesh_InstanceData* InstData, FVMExternalFunction &OutFunc);
+	NIAGARA_API void GetVertexSamplingFunctions(TArray<FNiagaraFunctionSignature>& OutFunctions);
+	NIAGARA_API void BindVertexSamplingFunction(const FVMExternalFunctionBindingInfo& BindingInfo, FNDISkeletalMesh_InstanceData* InstData, FVMExternalFunction &OutFunc);
 
-	void IsValidVertex(FVectorVMExternalFunctionContext& Context);
-	void RandomVertex(FVectorVMExternalFunctionContext& Context);
-	void GetVertexCount(FVectorVMExternalFunctionContext& Context);
+	NIAGARA_API void IsValidVertex(FVectorVMExternalFunctionContext& Context);
+	NIAGARA_API void RandomVertex(FVectorVMExternalFunctionContext& Context);
+	NIAGARA_API void GetVertexCount(FVectorVMExternalFunctionContext& Context);
 
 	template<typename FilterMode>
 	void IsValidFilteredVertex(FVectorVMExternalFunctionContext& Context);
@@ -915,9 +915,9 @@ public:
 	template<typename SkinningHandlerType, typename TransformHandlerType, typename VertexAccessorType>
 	void GetVertexSkinnedData(FVectorVMExternalFunctionContext& Context);
 
-	void GetVertexColor(FVectorVMExternalFunctionContext& Context);
+	NIAGARA_API void GetVertexColor(FVectorVMExternalFunctionContext& Context);
 
-	void GetVertexColorFallback(FVectorVMExternalFunctionContext& Context);
+	NIAGARA_API void GetVertexColorFallback(FVectorVMExternalFunctionContext& Context);
 
 	template<typename VertexAccessorType>
 	void GetVertexUV(FVectorVMExternalFunctionContext& Context);
@@ -939,41 +939,41 @@ private:
 	// Direct Bone + Socket Sampling
 
 public:
-	void GetSkeletonSamplingFunctions(TArray<FNiagaraFunctionSignature>& OutFunctions);
-	void BindSkeletonSamplingFunction(const FVMExternalFunctionBindingInfo& BindingInfo, FNDISkeletalMesh_InstanceData* InstData, FVMExternalFunction &OutFunc);
+	NIAGARA_API void GetSkeletonSamplingFunctions(TArray<FNiagaraFunctionSignature>& OutFunctions);
+	NIAGARA_API void BindSkeletonSamplingFunction(const FVMExternalFunctionBindingInfo& BindingInfo, FNDISkeletalMesh_InstanceData* InstData, FVMExternalFunction &OutFunc);
 
 	template<typename SkinningHandlerType, typename TransformHandlerType, typename bInterpolated>
 	void GetSkinnedBoneData(FVectorVMExternalFunctionContext& Context);
 
-	void IsValidBone(FVectorVMExternalFunctionContext& Context);
-	void RandomBone(FVectorVMExternalFunctionContext& Context);
-	void GetBoneCount(FVectorVMExternalFunctionContext& Context);
-	void GetParentBone(FVectorVMExternalFunctionContext& Context);
+	NIAGARA_API void IsValidBone(FVectorVMExternalFunctionContext& Context);
+	NIAGARA_API void RandomBone(FVectorVMExternalFunctionContext& Context);
+	NIAGARA_API void GetBoneCount(FVectorVMExternalFunctionContext& Context);
+	NIAGARA_API void GetParentBone(FVectorVMExternalFunctionContext& Context);
 
-	void GetFilteredBoneCount(FVectorVMExternalFunctionContext& Context);
-	void GetFilteredBoneAt(FVectorVMExternalFunctionContext& Context);
-	void RandomFilteredBone(FVectorVMExternalFunctionContext& Context);
+	NIAGARA_API void GetFilteredBoneCount(FVectorVMExternalFunctionContext& Context);
+	NIAGARA_API void GetFilteredBoneAt(FVectorVMExternalFunctionContext& Context);
+	NIAGARA_API void RandomFilteredBone(FVectorVMExternalFunctionContext& Context);
 
-	void GetUnfilteredBoneCount(FVectorVMExternalFunctionContext& Context);
-	void GetUnfilteredBoneAt(FVectorVMExternalFunctionContext& Context);
-	void RandomUnfilteredBone(FVectorVMExternalFunctionContext& Context);
+	NIAGARA_API void GetUnfilteredBoneCount(FVectorVMExternalFunctionContext& Context);
+	NIAGARA_API void GetUnfilteredBoneAt(FVectorVMExternalFunctionContext& Context);
+	NIAGARA_API void RandomUnfilteredBone(FVectorVMExternalFunctionContext& Context);
 
-	void GetFilteredSocketCount(FVectorVMExternalFunctionContext& Context);
-	void GetFilteredSocketBoneAt(FVectorVMExternalFunctionContext& Context);
-	void GetFilteredSocketTransform(FVectorVMExternalFunctionContext& Context);
-	void RandomFilteredSocket(FVectorVMExternalFunctionContext& Context);
+	NIAGARA_API void GetFilteredSocketCount(FVectorVMExternalFunctionContext& Context);
+	NIAGARA_API void GetFilteredSocketBoneAt(FVectorVMExternalFunctionContext& Context);
+	NIAGARA_API void GetFilteredSocketTransform(FVectorVMExternalFunctionContext& Context);
+	NIAGARA_API void RandomFilteredSocket(FVectorVMExternalFunctionContext& Context);
 
-	void RandomFilteredSocketOrBone(FVectorVMExternalFunctionContext& Context);
-	void GetFilteredSocketOrBoneCount(FVectorVMExternalFunctionContext& Context);
-	void GetFilteredSocketOrBoneBoneAt(FVectorVMExternalFunctionContext& Context);
+	NIAGARA_API void RandomFilteredSocketOrBone(FVectorVMExternalFunctionContext& Context);
+	NIAGARA_API void GetFilteredSocketOrBoneCount(FVectorVMExternalFunctionContext& Context);
+	NIAGARA_API void GetFilteredSocketOrBoneBoneAt(FVectorVMExternalFunctionContext& Context);
 	// End of Direct Bone + Socket Sampling
 	//////////////////////////////////////////////////////////////////////////
 
-	void SetSourceComponentFromBlueprints(USkeletalMeshComponent* ComponentToUse);
-	void SetSamplingRegionsFromBlueprints(const TArray<FName>& InSamplingRegions);
-	void SetFilteredBonesFromBlueprints(const TArray<FName>& InFilteredBones);
-	void SetFilteredSocketsFromBlueprints(const TArray<FName>& InFilteredSockets);
-	void SetWholeMeshLODFromBlueprints(int32 MeshLODLevel);
+	NIAGARA_API void SetSourceComponentFromBlueprints(USkeletalMeshComponent* ComponentToUse);
+	NIAGARA_API void SetSamplingRegionsFromBlueprints(const TArray<FName>& InSamplingRegions);
+	NIAGARA_API void SetFilteredBonesFromBlueprints(const TArray<FName>& InFilteredBones);
+	NIAGARA_API void SetFilteredSocketsFromBlueprints(const TArray<FName>& InFilteredSockets);
+	NIAGARA_API void SetWholeMeshLODFromBlueprints(int32 MeshLODLevel);
 };
 
 
diff --git a/Engine/Plugins/FX/Niagara/Source/Niagara/Classes/NiagaraDataInterfaceSparseVolumeTexture.h b/Engine/Plugins/FX/Niagara/Source/Niagara/Classes/NiagaraDataInterfaceSparseVolumeTexture.h
--- a/Engine/Plugins/FX/Niagara/Source/Niagara/Classes/NiagaraDataInterfaceSparseVolumeTexture.h
+++ b/Engine/Plugins/FX/Niagara/Source/Niagara/Classes/NiagaraDataInterfaceSparseVolumeTexture.h
@@ -10,8 +10,8 @@
 class USparseVolumeTexture;
 
 /** Data Interface allowing sampling of a sparse volume texture */
-UCLASS(EditInlineNew, Category = "Texture", CollapseCategories, meta = (DisplayName = "Sparse Volume Texture Sample"))
-class NIAGARA_API UNiagaraDataInterfaceSparseVolumeTexture : public UNiagaraDataInterface
+UCLASS(EditInlineNew, Category = "Texture", CollapseCategories, meta = (DisplayName = "Sparse Volume Texture Sample"), MinimalAPI)
+class UNiagaraDataInterfaceSparseVolumeTexture : public UNiagaraDataInterface
 {
 	GENERATED_UCLASS_BODY()
 
@@ -35,48 +35,48 @@ public:
 	FNiagaraUserParameterBinding SparseVolumeTextureUserParameter;
 
 	//UObject Interface
-	virtual void PostInitProperties() override;
-	virtual void PostLoad() override;
-	virtual void Serialize(FArchive& Ar) override;
+	NIAGARA_API virtual void PostInitProperties() override;
+	NIAGARA_API virtual void PostLoad() override;
+	NIAGARA_API virtual void Serialize(FArchive& Ar) override;
 	//UObject Interface End
 
 	//UNiagaraDataInterface Interface
-	virtual void GetFunctions(TArray<FNiagaraFunctionSignature>& OutFunctions) override;
-	virtual void GetVMExternalFunction(const FVMExternalFunctionBindingInfo& BindingInfo, void* InstanceData, FVMExternalFunction& OutFunc) override;
+	NIAGARA_API virtual void GetFunctions(TArray<FNiagaraFunctionSignature>& OutFunctions) override;
+	NIAGARA_API virtual void GetVMExternalFunction(const FVMExternalFunctionBindingInfo& BindingInfo, void* InstanceData, FVMExternalFunction& OutFunc) override;
 	virtual bool CanExecuteOnTarget(ENiagaraSimTarget Target) const override { return Target == ENiagaraSimTarget::GPUComputeSim; }
 
-	virtual int32 PerInstanceDataSize() const override;
-	virtual bool InitPerInstanceData(void* PerInstanceData, FNiagaraSystemInstance* SystemInstance) override;
-	virtual void DestroyPerInstanceData(void* PerInstanceData, FNiagaraSystemInstance* SystemInstance) override;
+	NIAGARA_API virtual int32 PerInstanceDataSize() const override;
+	NIAGARA_API virtual bool InitPerInstanceData(void* PerInstanceData, FNiagaraSystemInstance* SystemInstance) override;
+	NIAGARA_API virtual void DestroyPerInstanceData(void* PerInstanceData, FNiagaraSystemInstance* SystemInstance) override;
 
 	virtual bool HasPreSimulateTick() const override { return true; }
-	virtual bool PerInstanceTick(void* PerInstanceData, FNiagaraSystemInstance* SystemInstance, float DeltaSeconds) override;
+	NIAGARA_API virtual bool PerInstanceTick(void* PerInstanceData, FNiagaraSystemInstance* SystemInstance, float DeltaSeconds) override;
 	//UNiagaraDataInterface Interface End
 
-	virtual bool Equals(const UNiagaraDataInterface* Other) const override;
+	NIAGARA_API virtual bool Equals(const UNiagaraDataInterface* Other) const override;
 
 	// GPU sim functionality
 #if WITH_EDITORONLY_DATA
-	virtual bool AppendCompileHash(FNiagaraCompileHashVisitor* InVisitor) const override;
-	virtual void GetParameterDefinitionHLSL(const FNiagaraDataInterfaceGPUParamInfo& ParamInfo, FString& OutHLSL) override;
-	virtual bool GetFunctionHLSL(const FNiagaraDataInterfaceGPUParamInfo& ParamInfo, const FNiagaraDataInterfaceGeneratedFunction& FunctionInfo, int FunctionInstanceIndex, FString& OutHLSL) override;
-	virtual bool UpgradeFunctionCall(FNiagaraFunctionSignature& FunctionSignature) override;
+	NIAGARA_API virtual bool AppendCompileHash(FNiagaraCompileHashVisitor* InVisitor) const override;
+	NIAGARA_API virtual void GetParameterDefinitionHLSL(const FNiagaraDataInterfaceGPUParamInfo& ParamInfo, FString& OutHLSL) override;
+	NIAGARA_API virtual bool GetFunctionHLSL(const FNiagaraDataInterfaceGPUParamInfo& ParamInfo, const FNiagaraDataInterfaceGeneratedFunction& FunctionInfo, int FunctionInstanceIndex, FString& OutHLSL) override;
+	NIAGARA_API virtual bool UpgradeFunctionCall(FNiagaraFunctionSignature& FunctionSignature) override;
 #endif
-	virtual void BuildShaderParameters(FNiagaraShaderParametersBuilder& ShaderParametersBuilder) const override;
-	virtual void SetShaderParameters(const FNiagaraDataInterfaceSetShaderParametersContext& Context) const override;
+	NIAGARA_API virtual void BuildShaderParameters(FNiagaraShaderParametersBuilder& ShaderParametersBuilder) const override;
+	NIAGARA_API virtual void SetShaderParameters(const FNiagaraDataInterfaceSetShaderParametersContext& Context) const override;
 
-	void VMGetTextureDimensions(FVectorVMExternalFunctionContext& Context);
-	void VMGetNumMipLevels(FVectorVMExternalFunctionContext& Context);
+	NIAGARA_API void VMGetTextureDimensions(FVectorVMExternalFunctionContext& Context);
+	NIAGARA_API void VMGetNumMipLevels(FVectorVMExternalFunctionContext& Context);
 
-	void SetTexture(USparseVolumeTexture* InSparseVolumeTexture);
+	NIAGARA_API void SetTexture(USparseVolumeTexture* InSparseVolumeTexture);
 
 protected:
-	virtual bool CopyToInternal(UNiagaraDataInterface* Destination) const override;
+	NIAGARA_API virtual bool CopyToInternal(UNiagaraDataInterface* Destination) const override;
 
 protected:
-	static const TCHAR* TemplateShaderFilePath;
-	static const FName LoadSparseVolumeTextureName;
-	static const FName SampleSparseVolumeTextureName;
-	static const FName GetTextureDimensionsName;
-	static const FName GetNumMipLevelsName;
+	static NIAGARA_API const TCHAR* TemplateShaderFilePath;
+	static NIAGARA_API const FName LoadSparseVolumeTextureName;
+	static NIAGARA_API const FName SampleSparseVolumeTextureName;
+	static NIAGARA_API const FName GetTextureDimensionsName;
+	static NIAGARA_API const FName GetNumMipLevelsName;
 };
diff --git a/Engine/Plugins/FX/Niagara/Source/Niagara/Classes/NiagaraDataInterfaceSpline.h b/Engine/Plugins/FX/Niagara/Source/Niagara/Classes/NiagaraDataInterfaceSpline.h
--- a/Engine/Plugins/FX/Niagara/Source/Niagara/Classes/NiagaraDataInterfaceSpline.h
+++ b/Engine/Plugins/FX/Niagara/Source/Niagara/Classes/NiagaraDataInterfaceSpline.h
@@ -162,8 +162,8 @@ private:
 
 
 /** Data Interface allowing sampling of in-world spline components. Note that this data interface is very experimental. */
-UCLASS(EditInlineNew, Category = "Splines", CollapseCategories, meta = (DisplayName = "Spline"))
-class NIAGARA_API UNiagaraDataInterfaceSpline : public UNiagaraDataInterface
+UCLASS(EditInlineNew, Category = "Splines", CollapseCategories, meta = (DisplayName = "Spline"), MinimalAPI)
+class UNiagaraDataInterfaceSpline : public UNiagaraDataInterface
 {
 	GENERATED_UCLASS_BODY()
 
@@ -183,17 +183,17 @@ public:
 	int32 NumLUTSteps;
 	
 	//UObject Interface
-	virtual void PostInitProperties()override;
+	NIAGARA_API virtual void PostInitProperties()override;
 	//UObject Interface End
 
 	//UNiagaraDataInterface Interface
-	virtual bool InitPerInstanceData(void* PerInstanceData, FNiagaraSystemInstance* SystemInstance) override;
-	virtual void DestroyPerInstanceData(void* PerInstanceData, FNiagaraSystemInstance* SystemInstance)override;
-	virtual bool PerInstanceTick(void* PerInstanceData, FNiagaraSystemInstance* SystemInstance, float DeltaSeconds) override;
-	virtual int32 PerInstanceDataSize()const override;
-	virtual void GetFunctions(TArray<FNiagaraFunctionSignature>& OutFunctions)override;
-	virtual void GetVMExternalFunction(const FVMExternalFunctionBindingInfo& BindingInfo, void* InstanceData, FVMExternalFunction &OutFunc) override;
-	virtual bool Equals(const UNiagaraDataInterface* Other) const override;
+	NIAGARA_API virtual bool InitPerInstanceData(void* PerInstanceData, FNiagaraSystemInstance* SystemInstance) override;
+	NIAGARA_API virtual void DestroyPerInstanceData(void* PerInstanceData, FNiagaraSystemInstance* SystemInstance)override;
+	NIAGARA_API virtual bool PerInstanceTick(void* PerInstanceData, FNiagaraSystemInstance* SystemInstance, float DeltaSeconds) override;
+	NIAGARA_API virtual int32 PerInstanceDataSize()const override;
+	NIAGARA_API virtual void GetFunctions(TArray<FNiagaraFunctionSignature>& OutFunctions)override;
+	NIAGARA_API virtual void GetVMExternalFunction(const FVMExternalFunctionBindingInfo& BindingInfo, void* InstanceData, FVMExternalFunction &OutFunc) override;
+	NIAGARA_API virtual bool Equals(const UNiagaraDataInterface* Other) const override;
 	virtual bool CanExecuteOnTarget(ENiagaraSimTarget Target)const override { return true; }
 	virtual bool HasPreSimulateTick() const override { return true; }
 	//UNiagaraDataInterface Interface End
@@ -213,10 +213,10 @@ public:
 	template<typename UseLUT>
 	void FindClosestUnitDistanceFromPositionWS(FVectorVMExternalFunctionContext& Context);
 	
-	void GetLocalToWorld(FVectorVMExternalFunctionContext& Context);
-	void GetLocalToWorldInverseTransposed(FVectorVMExternalFunctionContext& Context);
+	NIAGARA_API void GetLocalToWorld(FVectorVMExternalFunctionContext& Context);
+	NIAGARA_API void GetLocalToWorldInverseTransposed(FVectorVMExternalFunctionContext& Context);
 
-	void VMGetSplineLength(FVectorVMExternalFunctionContext& Context);
+	NIAGARA_API void VMGetSplineLength(FVectorVMExternalFunctionContext& Context);
 
 	virtual int32 PerInstanceDataPassedToRenderThreadSize() const override
 	{
@@ -224,23 +224,23 @@ public:
 	}
 
 protected:
-	virtual bool CopyToInternal(UNiagaraDataInterface* Destination) const override;
+	NIAGARA_API virtual bool CopyToInternal(UNiagaraDataInterface* Destination) const override;
 
 #if WITH_EDITORONLY_DATA
-	virtual bool UpgradeFunctionCall(FNiagaraFunctionSignature& FunctionSignature) override;
-	virtual void GetCommonHLSL(FString& OutHLSL) override;
-	virtual void GetParameterDefinitionHLSL(const FNiagaraDataInterfaceGPUParamInfo& ParamInfo, FString& OutHLSL) override;
+	NIAGARA_API virtual bool UpgradeFunctionCall(FNiagaraFunctionSignature& FunctionSignature) override;
+	NIAGARA_API virtual void GetCommonHLSL(FString& OutHLSL) override;
+	NIAGARA_API virtual void GetParameterDefinitionHLSL(const FNiagaraDataInterfaceGPUParamInfo& ParamInfo, FString& OutHLSL) override;
 
-	virtual bool GetFunctionHLSL(const FNiagaraDataInterfaceGPUParamInfo& ParamInfo, const FNiagaraDataInterfaceGeneratedFunction& FunctionInfo, int FunctionInstanceIndex, FString& OutHLSL) override;
+	NIAGARA_API virtual bool GetFunctionHLSL(const FNiagaraDataInterfaceGPUParamInfo& ParamInfo, const FNiagaraDataInterfaceGeneratedFunction& FunctionInfo, int FunctionInstanceIndex, FString& OutHLSL) override;
 
-	virtual bool AppendCompileHash(FNiagaraCompileHashVisitor* InVisitor) const override;
+	NIAGARA_API virtual bool AppendCompileHash(FNiagaraCompileHashVisitor* InVisitor) const override;
 #endif
-	virtual void BuildShaderParameters(FNiagaraShaderParametersBuilder& ShaderParametersBuilder) const override;
-	virtual void SetShaderParameters(const FNiagaraDataInterfaceSetShaderParametersContext& Context) const override;
+	NIAGARA_API virtual void BuildShaderParameters(FNiagaraShaderParametersBuilder& ShaderParametersBuilder) const override;
+	NIAGARA_API virtual void SetShaderParameters(const FNiagaraDataInterfaceSetShaderParametersContext& Context) const override;
 
 private:
 	
-	void WriteTransform(const FMatrix44f& ToWrite, FVectorVMExternalFunctionContext& Context);
+	NIAGARA_API void WriteTransform(const FMatrix44f& ToWrite, FVectorVMExternalFunctionContext& Context);
 
 	TMap<FNiagaraSystemInstanceID, FNDISpline_InstanceData*> SystemInstancesToProxyData_GT;
 };
diff --git a/Engine/Plugins/FX/Niagara/Source/Niagara/Classes/NiagaraDataInterfaceSpriteRendererInfo.h b/Engine/Plugins/FX/Niagara/Source/Niagara/Classes/NiagaraDataInterfaceSpriteRendererInfo.h
--- a/Engine/Plugins/FX/Niagara/Source/Niagara/Classes/NiagaraDataInterfaceSpriteRendererInfo.h
+++ b/Engine/Plugins/FX/Niagara/Source/Niagara/Classes/NiagaraDataInterfaceSpriteRendererInfo.h
@@ -10,8 +10,8 @@
 class UNiagaraSpriteRendererProperties;
 
 /** This Data Interface can be used to query information about the sprite renderers of an emitter */
-UCLASS(EditInlineNew, Category = "Sprite Particles", CollapseCategories, meta = (DisplayName = "Sprite Renderer Info"))
-class NIAGARA_API UNiagaraDataInterfaceSpriteRendererInfo : public UNiagaraDataInterface
+UCLASS(EditInlineNew, Category = "Sprite Particles", CollapseCategories, meta = (DisplayName = "Sprite Renderer Info"), MinimalAPI)
+class UNiagaraDataInterfaceSpriteRendererInfo : public UNiagaraDataInterface
 {
 	GENERATED_UCLASS_BODY()
 
@@ -19,44 +19,44 @@ public:
 	UNiagaraSpriteRendererProperties* GetSpriteRenderer() const { return SpriteRenderer; }
 
 	//UObject Interface
-	virtual void PostInitProperties() override;
-	virtual void PostLoad() override;
-	virtual void BeginDestroy() override;
+	NIAGARA_API virtual void PostInitProperties() override;
+	NIAGARA_API virtual void PostLoad() override;
+	NIAGARA_API virtual void BeginDestroy() override;
 #if WITH_EDITOR	
-	virtual void PreEditChange(FProperty* PropertyAboutToChange) override;
-	virtual void PostEditChangeProperty(struct FPropertyChangedEvent& PropertyChangedEvent) override;
+	NIAGARA_API virtual void PreEditChange(FProperty* PropertyAboutToChange) override;
+	NIAGARA_API virtual void PostEditChangeProperty(struct FPropertyChangedEvent& PropertyChangedEvent) override;
 #endif 
 	//UObject Interface End
 
-	void OnSpriteRendererChanged(UNiagaraSpriteRendererProperties* NewSpriteRenderer);
+	NIAGARA_API void OnSpriteRendererChanged(UNiagaraSpriteRendererProperties* NewSpriteRenderer);
 
 	//UNiagaraDataInterface Interface
-	virtual void GetFunctions(TArray<FNiagaraFunctionSignature>& OutFunctions) override;
-	virtual void GetVMExternalFunction(const FVMExternalFunctionBindingInfo& BindingInfo, void* InstanceData, FVMExternalFunction& OutFunc) override;
-	virtual bool Equals(const UNiagaraDataInterface* Other) const override;
+	NIAGARA_API virtual void GetFunctions(TArray<FNiagaraFunctionSignature>& OutFunctions) override;
+	NIAGARA_API virtual void GetVMExternalFunction(const FVMExternalFunctionBindingInfo& BindingInfo, void* InstanceData, FVMExternalFunction& OutFunc) override;
+	NIAGARA_API virtual bool Equals(const UNiagaraDataInterface* Other) const override;
 #if WITH_EDITORONLY_DATA
-	virtual bool AppendCompileHash(FNiagaraCompileHashVisitor* InVisitor) const override;
-	virtual void GetParameterDefinitionHLSL(const FNiagaraDataInterfaceGPUParamInfo& ParamInfo, FString& OutHLSL) override;
-	virtual bool GetFunctionHLSL(const FNiagaraDataInterfaceGPUParamInfo& ParamInfo, const FNiagaraDataInterfaceGeneratedFunction& FunctionInfo, int FunctionInstanceIndex, FString& OutHLSL) override;
+	NIAGARA_API virtual bool AppendCompileHash(FNiagaraCompileHashVisitor* InVisitor) const override;
+	NIAGARA_API virtual void GetParameterDefinitionHLSL(const FNiagaraDataInterfaceGPUParamInfo& ParamInfo, FString& OutHLSL) override;
+	NIAGARA_API virtual bool GetFunctionHLSL(const FNiagaraDataInterfaceGPUParamInfo& ParamInfo, const FNiagaraDataInterfaceGeneratedFunction& FunctionInfo, int FunctionInstanceIndex, FString& OutHLSL) override;
 #endif
-	virtual void BuildShaderParameters(FNiagaraShaderParametersBuilder& ShaderParametersBuilder) const override;
-	virtual void SetShaderParameters(const FNiagaraDataInterfaceSetShaderParametersContext& Context) const override;
+	NIAGARA_API virtual void BuildShaderParameters(FNiagaraShaderParametersBuilder& ShaderParametersBuilder) const override;
+	NIAGARA_API virtual void SetShaderParameters(const FNiagaraDataInterfaceSetShaderParametersContext& Context) const override;
 
 	virtual bool CanExecuteOnTarget(ENiagaraSimTarget Target) const override { return true; }
 #if WITH_EDITOR
-	virtual void GetFeedback(UNiagaraSystem* InAsset, UNiagaraComponent* InComponent, TArray<FNiagaraDataInterfaceError>& OutErrors, TArray<FNiagaraDataInterfaceFeedback>& OutWarnings, TArray<FNiagaraDataInterfaceFeedback>& OutInfo) override;
+	NIAGARA_API virtual void GetFeedback(UNiagaraSystem* InAsset, UNiagaraComponent* InComponent, TArray<FNiagaraDataInterfaceError>& OutErrors, TArray<FNiagaraDataInterfaceFeedback>& OutWarnings, TArray<FNiagaraDataInterfaceFeedback>& OutInfo) override;
 #endif
 	//UNiagaraDataInterface Interface	
 
 protected:
-	virtual bool CopyToInternal(UNiagaraDataInterface* Destination) const override;
-	virtual void PushToRenderThreadImpl() override;
+	NIAGARA_API virtual bool CopyToInternal(UNiagaraDataInterface* Destination) const override;
+	NIAGARA_API virtual void PushToRenderThreadImpl() override;
 
-	void VMIsValid(FVectorVMExternalFunctionContext& Context);
-	void VMGetSourceMode(FVectorVMExternalFunctionContext& Context);
-	void VMGetAlignment(FVectorVMExternalFunctionContext& Context);
-	void VMGetFacingMode(FVectorVMExternalFunctionContext& Context);
-	void VMGetSubUVDetails(FVectorVMExternalFunctionContext& Context);
+	NIAGARA_API void VMIsValid(FVectorVMExternalFunctionContext& Context);
+	NIAGARA_API void VMGetSourceMode(FVectorVMExternalFunctionContext& Context);
+	NIAGARA_API void VMGetAlignment(FVectorVMExternalFunctionContext& Context);
+	NIAGARA_API void VMGetFacingMode(FVectorVMExternalFunctionContext& Context);
+	NIAGARA_API void VMGetSubUVDetails(FVectorVMExternalFunctionContext& Context);
 
 	/** The name of the sprite renderer */
 	UPROPERTY(EditAnywhere, Category = "Source")
diff --git a/Engine/Plugins/FX/Niagara/Source/Niagara/Classes/NiagaraDataInterfaceTexture.h b/Engine/Plugins/FX/Niagara/Source/Niagara/Classes/NiagaraDataInterfaceTexture.h
--- a/Engine/Plugins/FX/Niagara/Source/Niagara/Classes/NiagaraDataInterfaceTexture.h
+++ b/Engine/Plugins/FX/Niagara/Source/Niagara/Classes/NiagaraDataInterfaceTexture.h
@@ -8,8 +8,8 @@
 #include "NiagaraDataInterfaceTexture.generated.h"
 
 /** Data Interface allowing sampling of a texture */
-UCLASS(EditInlineNew, Category = "Texture", CollapseCategories, meta = (DisplayName = "Texture Sample"))
-class NIAGARA_API UNiagaraDataInterfaceTexture : public UNiagaraDataInterface
+UCLASS(EditInlineNew, Category = "Texture", CollapseCategories, meta = (DisplayName = "Texture Sample"), MinimalAPI)
+class UNiagaraDataInterfaceTexture : public UNiagaraDataInterface
 {
 	GENERATED_UCLASS_BODY()
 
@@ -28,52 +28,52 @@ public:
 	FNiagaraUserParameterBinding TextureUserParameter;
 
 	//UObject Interface
-	virtual void PostInitProperties()override;
-	virtual void PostLoad() override;
-	virtual void Serialize(FArchive& Ar) override;
+	NIAGARA_API virtual void PostInitProperties()override;
+	NIAGARA_API virtual void PostLoad() override;
+	NIAGARA_API virtual void Serialize(FArchive& Ar) override;
 	//UObject Interface End
 
 	//UNiagaraDataInterface Interface
-	virtual void GetFunctions(TArray<FNiagaraFunctionSignature>& OutFunctions)override;
-	virtual void GetVMExternalFunction(const FVMExternalFunctionBindingInfo& BindingInfo, void* InstanceData, FVMExternalFunction &OutFunc) override;
+	NIAGARA_API virtual void GetFunctions(TArray<FNiagaraFunctionSignature>& OutFunctions)override;
+	NIAGARA_API virtual void GetVMExternalFunction(const FVMExternalFunctionBindingInfo& BindingInfo, void* InstanceData, FVMExternalFunction &OutFunc) override;
 	virtual bool CanExecuteOnTarget(ENiagaraSimTarget Target)const override { return Target==ENiagaraSimTarget::GPUComputeSim; }
 
-	virtual int32 PerInstanceDataSize() const override;
-	virtual bool InitPerInstanceData(void* PerInstanceData, FNiagaraSystemInstance* SystemInstance) override;
-	virtual void DestroyPerInstanceData(void* PerInstanceData, FNiagaraSystemInstance* SystemInstance) override;
+	NIAGARA_API virtual int32 PerInstanceDataSize() const override;
+	NIAGARA_API virtual bool InitPerInstanceData(void* PerInstanceData, FNiagaraSystemInstance* SystemInstance) override;
+	NIAGARA_API virtual void DestroyPerInstanceData(void* PerInstanceData, FNiagaraSystemInstance* SystemInstance) override;
 
 	virtual bool HasPreSimulateTick() const override { return true; }
-	virtual bool PerInstanceTick(void* PerInstanceData, FNiagaraSystemInstance* SystemInstance, float DeltaSeconds) override;
+	NIAGARA_API virtual bool PerInstanceTick(void* PerInstanceData, FNiagaraSystemInstance* SystemInstance, float DeltaSeconds) override;
 	//UNiagaraDataInterface Interface End
 
-	void VMSampleTexture(FVectorVMExternalFunctionContext& Context);
-	void VMSamplePseudoVolumeTexture(FVectorVMExternalFunctionContext& Context);
-	void VMGetTextureDimensions(FVectorVMExternalFunctionContext& Context);
-	void VMGetNumMipLevels(FVectorVMExternalFunctionContext& Context);
+	NIAGARA_API void VMSampleTexture(FVectorVMExternalFunctionContext& Context);
+	NIAGARA_API void VMSamplePseudoVolumeTexture(FVectorVMExternalFunctionContext& Context);
+	NIAGARA_API void VMGetTextureDimensions(FVectorVMExternalFunctionContext& Context);
+	NIAGARA_API void VMGetNumMipLevels(FVectorVMExternalFunctionContext& Context);
 
-	virtual bool Equals(const UNiagaraDataInterface* Other) const override;
+	NIAGARA_API virtual bool Equals(const UNiagaraDataInterface* Other) const override;
 
 	// GPU sim functionality
 #if WITH_EDITORONLY_DATA
-	virtual bool AppendCompileHash(FNiagaraCompileHashVisitor* InVisitor) const override;
-	virtual void GetParameterDefinitionHLSL(const FNiagaraDataInterfaceGPUParamInfo& ParamInfo, FString& OutHLSL) override;
-	virtual bool GetFunctionHLSL(const FNiagaraDataInterfaceGPUParamInfo& ParamInfo, const FNiagaraDataInterfaceGeneratedFunction& FunctionInfo, int FunctionInstanceIndex, FString& OutHLSL) override;
-	virtual bool UpgradeFunctionCall(FNiagaraFunctionSignature& FunctionSignature) override;
+	NIAGARA_API virtual bool AppendCompileHash(FNiagaraCompileHashVisitor* InVisitor) const override;
+	NIAGARA_API virtual void GetParameterDefinitionHLSL(const FNiagaraDataInterfaceGPUParamInfo& ParamInfo, FString& OutHLSL) override;
+	NIAGARA_API virtual bool GetFunctionHLSL(const FNiagaraDataInterfaceGPUParamInfo& ParamInfo, const FNiagaraDataInterfaceGeneratedFunction& FunctionInfo, int FunctionInstanceIndex, FString& OutHLSL) override;
+	NIAGARA_API virtual bool UpgradeFunctionCall(FNiagaraFunctionSignature& FunctionSignature) override;
 #endif
-	virtual void BuildShaderParameters(FNiagaraShaderParametersBuilder& ShaderParametersBuilder) const override;
-	virtual void SetShaderParameters(const FNiagaraDataInterfaceSetShaderParametersContext& Context) const override;
+	NIAGARA_API virtual void BuildShaderParameters(FNiagaraShaderParametersBuilder& ShaderParametersBuilder) const override;
+	NIAGARA_API virtual void SetShaderParameters(const FNiagaraDataInterfaceSetShaderParametersContext& Context) const override;
 
-	void SetTexture(UTexture* InTexture);
+	NIAGARA_API void SetTexture(UTexture* InTexture);
 
 protected:
-	virtual bool CopyToInternal(UNiagaraDataInterface* Destination) const override;
+	NIAGARA_API virtual bool CopyToInternal(UNiagaraDataInterface* Destination) const override;
 
 protected:
-	static const TCHAR* TemplateShaderFilePath;
-	static const FName LoadTexture2DName;
-	static const FName GatherRedTexture2DName;
-	static const FName SampleTexture2DName;
-	static const FName SamplePseudoVolumeTextureName;
-	static const FName GetTextureDimensionsName;
-	static const FName GetNumMipLevelsName;
+	static NIAGARA_API const TCHAR* TemplateShaderFilePath;
+	static NIAGARA_API const FName LoadTexture2DName;
+	static NIAGARA_API const FName GatherRedTexture2DName;
+	static NIAGARA_API const FName SampleTexture2DName;
+	static NIAGARA_API const FName SamplePseudoVolumeTextureName;
+	static NIAGARA_API const FName GetTextureDimensionsName;
+	static NIAGARA_API const FName GetNumMipLevelsName;
 };
diff --git a/Engine/Plugins/FX/Niagara/Source/Niagara/Classes/NiagaraDataInterfaceVector2DCurve.h b/Engine/Plugins/FX/Niagara/Source/Niagara/Classes/NiagaraDataInterfaceVector2DCurve.h
--- a/Engine/Plugins/FX/Niagara/Source/Niagara/Classes/NiagaraDataInterfaceVector2DCurve.h
+++ b/Engine/Plugins/FX/Niagara/Source/Niagara/Classes/NiagaraDataInterfaceVector2DCurve.h
@@ -11,8 +11,8 @@
 
 
 /** Data Interface allowing sampling of vector curves. */
-UCLASS(EditInlineNew, Category = "Curves", CollapseCategories, meta = (DisplayName = "Curve for Vector2D's"))
-class NIAGARA_API UNiagaraDataInterfaceVector2DCurve : public UNiagaraDataInterfaceCurveBase
+UCLASS(EditInlineNew, Category = "Curves", CollapseCategories, meta = (DisplayName = "Curve for Vector2D's"), MinimalAPI)
+class UNiagaraDataInterfaceVector2DCurve : public UNiagaraDataInterfaceCurveBase
 {
 	GENERATED_UCLASS_BODY()
 
@@ -38,28 +38,28 @@ public:
 	};
 
 	//UObject Interface
-	virtual void PostInitProperties() override;
-	virtual void Serialize(FArchive& Ar) override;
+	NIAGARA_API virtual void PostInitProperties() override;
+	NIAGARA_API virtual void Serialize(FArchive& Ar) override;
 	//UObject Interface End
 	
-	virtual void UpdateTimeRanges() override;
-	virtual TArray<float> BuildLUT(int32 NumEntries) const override;
+	NIAGARA_API virtual void UpdateTimeRanges() override;
+	NIAGARA_API virtual TArray<float> BuildLUT(int32 NumEntries) const override;
 
-	virtual bool Equals(const UNiagaraDataInterface* Other) const override;
+	NIAGARA_API virtual bool Equals(const UNiagaraDataInterface* Other) const override;
 
 
-	virtual void GetFunctions(TArray<FNiagaraFunctionSignature>& OutFunctions)override;
+	NIAGARA_API virtual void GetFunctions(TArray<FNiagaraFunctionSignature>& OutFunctions)override;
 
-	virtual void GetVMExternalFunction(const FVMExternalFunctionBindingInfo& BindingInfo, void* InstanceData, FVMExternalFunction &OutFunc) override;
+	NIAGARA_API virtual void GetVMExternalFunction(const FVMExternalFunctionBindingInfo& BindingInfo, void* InstanceData, FVMExternalFunction &OutFunc) override;
 
 	template<typename UseLUT>
 	void SampleCurve(FVectorVMExternalFunctionContext& Context);
 
 	//~ UNiagaraDataInterfaceCurveBase interface
-	virtual void GetCurveData(TArray<FCurveData>& OutCurveData) override;
+	NIAGARA_API virtual void GetCurveData(TArray<FCurveData>& OutCurveData) override;
 
 #if WITH_EDITORONLY_DATA
-	virtual bool GetFunctionHLSL(const FNiagaraDataInterfaceGPUParamInfo& ParamInfo, const FNiagaraDataInterfaceGeneratedFunction& FunctionInfo, int FunctionInstanceIndex, FString& OutHLSL) override;
+	NIAGARA_API virtual bool GetFunctionHLSL(const FNiagaraDataInterfaceGPUParamInfo& ParamInfo, const FNiagaraDataInterfaceGeneratedFunction& FunctionInfo, int FunctionInstanceIndex, FString& OutHLSL) override;
 #endif
 
 	virtual int32 GetCurveNumElems() const override { return CurveLUTNumElems; }
@@ -68,11 +68,11 @@ public:
 #endif
 
 protected:
-	virtual bool CopyToInternal(UNiagaraDataInterface* Destination) const override;
+	NIAGARA_API virtual bool CopyToInternal(UNiagaraDataInterface* Destination) const override;
 
 private:
 	template<typename UseLUT>
 	FORCEINLINE_DEBUGGABLE FVector2f SampleCurveInternal(float X);
 
-	static const FName SampleCurveName;
+	static NIAGARA_API const FName SampleCurveName;
 };
diff --git a/Engine/Plugins/FX/Niagara/Source/Niagara/Classes/NiagaraDataInterfaceVector4Curve.h b/Engine/Plugins/FX/Niagara/Source/Niagara/Classes/NiagaraDataInterfaceVector4Curve.h
--- a/Engine/Plugins/FX/Niagara/Source/Niagara/Classes/NiagaraDataInterfaceVector4Curve.h
+++ b/Engine/Plugins/FX/Niagara/Source/Niagara/Classes/NiagaraDataInterfaceVector4Curve.h
@@ -16,8 +16,8 @@ class UCurveFloat;
 class FNiagaraSystemInstance;
 
 /** Data Interface allowing sampling of vector curves. */
-UCLASS(EditInlineNew, Category = "Curves", CollapseCategories, meta = (DisplayName = "Curve for Vector4's"))
-class NIAGARA_API UNiagaraDataInterfaceVector4Curve : public UNiagaraDataInterfaceCurveBase
+UCLASS(EditInlineNew, Category = "Curves", CollapseCategories, meta = (DisplayName = "Curve for Vector4's"), MinimalAPI)
+class UNiagaraDataInterfaceVector4Curve : public UNiagaraDataInterfaceCurveBase
 {
 	GENERATED_UCLASS_BODY()
 
@@ -53,25 +53,25 @@ public:
 	};
 
 	//UObject Interface
-	virtual void PostInitProperties() override;
-	virtual void Serialize(FArchive& Ar) override;
+	NIAGARA_API virtual void PostInitProperties() override;
+	NIAGARA_API virtual void Serialize(FArchive& Ar) override;
 	//UObject Interface End
 
-	virtual void UpdateTimeRanges() override;
-	virtual TArray<float> BuildLUT(int32 NumEntries) const override;
+	NIAGARA_API virtual void UpdateTimeRanges() override;
+	NIAGARA_API virtual TArray<float> BuildLUT(int32 NumEntries) const override;
 
-	virtual void GetFunctions(TArray<FNiagaraFunctionSignature>& OutFunctions)override;
-	virtual void GetVMExternalFunction(const FVMExternalFunctionBindingInfo& BindingInfo, void* InstanceData, FVMExternalFunction &OutFunc) override;
+	NIAGARA_API virtual void GetFunctions(TArray<FNiagaraFunctionSignature>& OutFunctions)override;
+	NIAGARA_API virtual void GetVMExternalFunction(const FVMExternalFunctionBindingInfo& BindingInfo, void* InstanceData, FVMExternalFunction &OutFunc) override;
 
 	template<typename UseLUT>
 	void SampleCurve(FVectorVMExternalFunctionContext& Context);
 
-	virtual bool Equals(const UNiagaraDataInterface* Other) const override;
+	NIAGARA_API virtual bool Equals(const UNiagaraDataInterface* Other) const override;
 
 	//~ UNiagaraDataInterfaceCurveBase interface
-	virtual void GetCurveData(TArray<FCurveData>& OutCurveData) override;
+	NIAGARA_API virtual void GetCurveData(TArray<FCurveData>& OutCurveData) override;
 #if WITH_EDITORONLY_DATA
-	virtual bool GetFunctionHLSL(const FNiagaraDataInterfaceGPUParamInfo& ParamInfo, const FNiagaraDataInterfaceGeneratedFunction& FunctionInfo, int FunctionInstanceIndex, FString& OutHLSL) override;
+	NIAGARA_API virtual bool GetFunctionHLSL(const FNiagaraDataInterfaceGPUParamInfo& ParamInfo, const FNiagaraDataInterfaceGeneratedFunction& FunctionInfo, int FunctionInstanceIndex, FString& OutHLSL) override;
 #endif
 
 	virtual int32 GetCurveNumElems() const override { return CurveLUTNumElems; }
@@ -80,11 +80,11 @@ public:
 #endif
 
 protected:
-	virtual bool CopyToInternal(UNiagaraDataInterface* Destination) const override;
+	NIAGARA_API virtual bool CopyToInternal(UNiagaraDataInterface* Destination) const override;
 
 private: 
 	template<typename UseLUT>
 	FORCEINLINE_DEBUGGABLE FVector4f SampleCurveInternal(float X);
 
-	static const FName SampleCurveName;
+	static NIAGARA_API const FName SampleCurveName;
 };
diff --git a/Engine/Plugins/FX/Niagara/Source/Niagara/Classes/NiagaraDataInterfaceVectorCurve.h b/Engine/Plugins/FX/Niagara/Source/Niagara/Classes/NiagaraDataInterfaceVectorCurve.h
--- a/Engine/Plugins/FX/Niagara/Source/Niagara/Classes/NiagaraDataInterfaceVectorCurve.h
+++ b/Engine/Plugins/FX/Niagara/Source/Niagara/Classes/NiagaraDataInterfaceVectorCurve.h
@@ -17,8 +17,8 @@ class FNiagaraSystemInstance;
 
 
 /** Data Interface allowing sampling of vector curves. */
-UCLASS(EditInlineNew, Category = "Curves", CollapseCategories, meta = (DisplayName = "Curve for Vector3's"))
-class NIAGARA_API UNiagaraDataInterfaceVectorCurve : public UNiagaraDataInterfaceCurveBase
+UCLASS(EditInlineNew, Category = "Curves", CollapseCategories, meta = (DisplayName = "Curve for Vector3's"), MinimalAPI)
+class UNiagaraDataInterfaceVectorCurve : public UNiagaraDataInterfaceCurveBase
 {
 	GENERATED_UCLASS_BODY()
 public:
@@ -48,23 +48,23 @@ public:
 	};
 
 	//UObject Interface
-	virtual void PostInitProperties() override;
-	virtual void Serialize(FArchive& Ar) override;
+	NIAGARA_API virtual void PostInitProperties() override;
+	NIAGARA_API virtual void Serialize(FArchive& Ar) override;
 	//UObject Interface End
 
-	virtual void UpdateTimeRanges() override;
-	virtual TArray<float> BuildLUT(int32 NumEntries) const override;
+	NIAGARA_API virtual void UpdateTimeRanges() override;
+	NIAGARA_API virtual TArray<float> BuildLUT(int32 NumEntries) const override;
 
-	virtual void GetFunctions(TArray<FNiagaraFunctionSignature>& OutFunctions)override;
-	virtual void GetVMExternalFunction(const FVMExternalFunctionBindingInfo& BindingInfo, void* InstanceData, FVMExternalFunction &OutFunc) override;
+	NIAGARA_API virtual void GetFunctions(TArray<FNiagaraFunctionSignature>& OutFunctions)override;
+	NIAGARA_API virtual void GetVMExternalFunction(const FVMExternalFunctionBindingInfo& BindingInfo, void* InstanceData, FVMExternalFunction &OutFunc) override;
 
 	template<typename UseLUT>
 	void SampleCurve(FVectorVMExternalFunctionContext& Context);
 
-	virtual bool Equals(const UNiagaraDataInterface* Other) const override;
+	NIAGARA_API virtual bool Equals(const UNiagaraDataInterface* Other) const override;
 
 	//~ UNiagaraDataInterfaceCurveBase interface
-	virtual void GetCurveData(TArray<FCurveData>& OutCurveData) override;
+	NIAGARA_API virtual void GetCurveData(TArray<FCurveData>& OutCurveData) override;
 	
 	virtual int32 GetCurveNumElems() const override { return CurveLUTNumElems; }
 #if WITH_EDITORONLY_DATA
@@ -72,16 +72,16 @@ public:
 #endif
 
 #if WITH_EDITORONLY_DATA
-	virtual void SyncCurvesToAsset() override;
-	virtual bool GetFunctionHLSL(const FNiagaraDataInterfaceGPUParamInfo& ParamInfo, const FNiagaraDataInterfaceGeneratedFunction& FunctionInfo, int FunctionInstanceIndex, FString& OutHLSL) override;
+	NIAGARA_API virtual void SyncCurvesToAsset() override;
+	NIAGARA_API virtual bool GetFunctionHLSL(const FNiagaraDataInterfaceGPUParamInfo& ParamInfo, const FNiagaraDataInterfaceGeneratedFunction& FunctionInfo, int FunctionInstanceIndex, FString& OutHLSL) override;
 #endif
 
 protected:
-	virtual bool CopyToInternal(UNiagaraDataInterface* Destination) const override;
+	NIAGARA_API virtual bool CopyToInternal(UNiagaraDataInterface* Destination) const override;
 
 private:
 	template<typename UseLUT>
 	FORCEINLINE_DEBUGGABLE FVector3f SampleCurveInternal(float X);
 
-	static const FName SampleCurveName;
+	static NIAGARA_API const FName SampleCurveName;
 };
diff --git a/Engine/Plugins/FX/Niagara/Source/Niagara/Classes/NiagaraDataInterfaceVectorField.h b/Engine/Plugins/FX/Niagara/Source/Niagara/Classes/NiagaraDataInterfaceVectorField.h
--- a/Engine/Plugins/FX/Niagara/Source/Niagara/Classes/NiagaraDataInterfaceVectorField.h
+++ b/Engine/Plugins/FX/Niagara/Source/Niagara/Classes/NiagaraDataInterfaceVectorField.h
@@ -7,8 +7,8 @@
 
 class FNiagaraSystemInstance;
 
-UCLASS(EditInlineNew, Category = "Vector Field", CollapseCategories, meta = (DisplayName = "Vector Field"))
-class NIAGARA_API UNiagaraDataInterfaceVectorField : public UNiagaraDataInterface
+UCLASS(EditInlineNew, Category = "Vector Field", CollapseCategories, meta = (DisplayName = "Vector Field"), MinimalAPI)
+class UNiagaraDataInterfaceVectorField : public UNiagaraDataInterface
 {
 	GENERATED_UCLASS_BODY()
 
@@ -35,55 +35,55 @@ public:
 
 public:
 	//~ UObject interface
-	virtual void PostInitProperties() override;
-	virtual void PostLoad() override; 
+	NIAGARA_API virtual void PostInitProperties() override;
+	NIAGARA_API virtual void PostLoad() override; 
 #if WITH_EDITOR
-	virtual void PostEditChangeProperty(struct FPropertyChangedEvent& PropertyChangedEvent) override;
-	virtual void PreEditChange(FProperty* PropertyAboutToChange) override;
+	NIAGARA_API virtual void PostEditChangeProperty(struct FPropertyChangedEvent& PropertyChangedEvent) override;
+	NIAGARA_API virtual void PreEditChange(FProperty* PropertyAboutToChange) override;
 #endif
 	//~ UObject interface END
 
 	//~ UNiagaraDataInterface interface
 	// VM functionality
-	virtual void GetFunctions(TArray<FNiagaraFunctionSignature>& OutFunctions) override;
-	virtual void GetVMExternalFunction(const FVMExternalFunctionBindingInfo& BindingInfo, void* InstanceData, FVMExternalFunction &OutFunc) override;
-	virtual bool Equals(const UNiagaraDataInterface* Other) const override;
-	virtual bool CanExecuteOnTarget(ENiagaraSimTarget Target) const override;
+	NIAGARA_API virtual void GetFunctions(TArray<FNiagaraFunctionSignature>& OutFunctions) override;
+	NIAGARA_API virtual void GetVMExternalFunction(const FVMExternalFunctionBindingInfo& BindingInfo, void* InstanceData, FVMExternalFunction &OutFunc) override;
+	NIAGARA_API virtual bool Equals(const UNiagaraDataInterface* Other) const override;
+	NIAGARA_API virtual bool CanExecuteOnTarget(ENiagaraSimTarget Target) const override;
 
 #if WITH_EDITOR	
 	// Editor functionality
-	virtual void GetFeedback(UNiagaraSystem* InAsset, UNiagaraComponent* InComponent, TArray<FNiagaraDataInterfaceError>& OutErrors, TArray<FNiagaraDataInterfaceFeedback>& OutWarnings, TArray<FNiagaraDataInterfaceFeedback>& OutInfo) override;
+	NIAGARA_API virtual void GetFeedback(UNiagaraSystem* InAsset, UNiagaraComponent* InComponent, TArray<FNiagaraDataInterfaceError>& OutErrors, TArray<FNiagaraDataInterfaceFeedback>& OutWarnings, TArray<FNiagaraDataInterfaceFeedback>& OutInfo) override;
 
 #endif
 
 	// GPU sim functionality
 #if WITH_EDITORONLY_DATA
-	virtual bool AppendCompileHash(FNiagaraCompileHashVisitor* InVisitor) const override;
-	virtual void GetParameterDefinitionHLSL(const FNiagaraDataInterfaceGPUParamInfo& ParamInfo, FString& OutHLSL) override;
-	virtual bool GetFunctionHLSL(const FNiagaraDataInterfaceGPUParamInfo& ParamInfo, const FNiagaraDataInterfaceGeneratedFunction& FunctionInfo, int FunctionInstanceIndex, FString& OutHLSL) override;
+	NIAGARA_API virtual bool AppendCompileHash(FNiagaraCompileHashVisitor* InVisitor) const override;
+	NIAGARA_API virtual void GetParameterDefinitionHLSL(const FNiagaraDataInterfaceGPUParamInfo& ParamInfo, FString& OutHLSL) override;
+	NIAGARA_API virtual bool GetFunctionHLSL(const FNiagaraDataInterfaceGPUParamInfo& ParamInfo, const FNiagaraDataInterfaceGeneratedFunction& FunctionInfo, int FunctionInstanceIndex, FString& OutHLSL) override;
 #endif
-	virtual void BuildShaderParameters(FNiagaraShaderParametersBuilder& ShaderParametersBuilder) const override;
-	virtual void SetShaderParameters(const FNiagaraDataInterfaceSetShaderParametersContext& Context) const override;
+	NIAGARA_API virtual void BuildShaderParameters(FNiagaraShaderParametersBuilder& ShaderParametersBuilder) const override;
+	NIAGARA_API virtual void SetShaderParameters(const FNiagaraDataInterfaceSetShaderParametersContext& Context) const override;
 	//~ UNiagaraDataInterface interface END
 
 	// VM functions
-	void GetFieldDimensions(FVectorVMExternalFunctionContext& Context);
-	void GetFieldBounds(FVectorVMExternalFunctionContext& Context); 
-	void GetFieldTilingAxes(FVectorVMExternalFunctionContext& Context);
-	void SampleVectorField(FVectorVMExternalFunctionContext& Context);
-	void LoadVectorField(FVectorVMExternalFunctionContext& Context);
+	NIAGARA_API void GetFieldDimensions(FVectorVMExternalFunctionContext& Context);
+	NIAGARA_API void GetFieldBounds(FVectorVMExternalFunctionContext& Context); 
+	NIAGARA_API void GetFieldTilingAxes(FVectorVMExternalFunctionContext& Context);
+	NIAGARA_API void SampleVectorField(FVectorVMExternalFunctionContext& Context);
+	NIAGARA_API void LoadVectorField(FVectorVMExternalFunctionContext& Context);
 	
 	//	
-	FVector GetTilingAxes() const;
-	FVector GetDimensions() const;
-	FVector GetMinBounds() const;
-	FVector GetMaxBounds() const;
+	NIAGARA_API FVector GetTilingAxes() const;
+	NIAGARA_API FVector GetDimensions() const;
+	NIAGARA_API FVector GetMinBounds() const;
+	NIAGARA_API FVector GetMaxBounds() const;
 protected:
 	//~ UNiagaraDataInterface interface
-	virtual bool CopyToInternal(UNiagaraDataInterface* Destination) const override;
+	NIAGARA_API virtual bool CopyToInternal(UNiagaraDataInterface* Destination) const override;
 	//~ UNiagaraDataInterface interface END
 
-	virtual void PushToRenderThreadImpl() override;
+	NIAGARA_API virtual void PushToRenderThreadImpl() override;
 };
 
 struct FNiagaraDataInterfaceProxyVectorField : public FNiagaraDataInterfaceProxy
diff --git a/Engine/Plugins/FX/Niagara/Source/Niagara/Classes/NiagaraDataInterfaceVolumeCache.h b/Engine/Plugins/FX/Niagara/Source/Niagara/Classes/NiagaraDataInterfaceVolumeCache.h
--- a/Engine/Plugins/FX/Niagara/Source/Niagara/Classes/NiagaraDataInterfaceVolumeCache.h
+++ b/Engine/Plugins/FX/Niagara/Source/Niagara/Classes/NiagaraDataInterfaceVolumeCache.h
@@ -12,8 +12,8 @@
 
 struct FVolumeCacheInstanceData_GameThread;
 
-UCLASS(EditInlineNew, Category = "Grid", CollapseCategories, meta = (DisplayName = "Volume Cache"))
-class NIAGARA_API UNiagaraDataInterfaceVolumeCache : public UNiagaraDataInterface
+UCLASS(EditInlineNew, Category = "Grid", CollapseCategories, meta = (DisplayName = "Volume Cache"), MinimalAPI)
+class UNiagaraDataInterfaceVolumeCache : public UNiagaraDataInterface
 {
 	GENERATED_UCLASS_BODY()
 
@@ -28,50 +28,50 @@ public:
 	UPROPERTY(EditAnywhere, Category = "File")
 	TObjectPtr<UVolumeCache> VolumeCache;
 
-	virtual void PostInitProperties() override;
+	NIAGARA_API virtual void PostInitProperties() override;
 
 	// UNiagaraDataInterface Interface Begin
 	virtual bool CanExecuteOnTarget(ENiagaraSimTarget Target) const override { return true; }
 
-	virtual bool Equals(const UNiagaraDataInterface* Other) const override;
-	virtual int32 PerInstanceDataSize()const override;
-	virtual bool InitPerInstanceData(void* PerInstanceData, FNiagaraSystemInstance* SystemInstance) override;	
-	virtual void GetFunctions(TArray<FNiagaraFunctionSignature>& OutFunctions) override;
-	virtual bool PerInstanceTick(void* PerInstanceData, FNiagaraSystemInstance* SystemInstance, float DeltaSeconds) override;
+	NIAGARA_API virtual bool Equals(const UNiagaraDataInterface* Other) const override;
+	NIAGARA_API virtual int32 PerInstanceDataSize()const override;
+	NIAGARA_API virtual bool InitPerInstanceData(void* PerInstanceData, FNiagaraSystemInstance* SystemInstance) override;	
+	NIAGARA_API virtual void GetFunctions(TArray<FNiagaraFunctionSignature>& OutFunctions) override;
+	NIAGARA_API virtual bool PerInstanceTick(void* PerInstanceData, FNiagaraSystemInstance* SystemInstance, float DeltaSeconds) override;
 	virtual bool HasPreSimulateTick() const override { return true; }
-	virtual void GetVMExternalFunction(const FVMExternalFunctionBindingInfo& BindingInfo, void* InstanceData, FVMExternalFunction& OutFunc) override;
-	virtual void DestroyPerInstanceData(void* PerInstanceData, FNiagaraSystemInstance* SystemInstance) override;
+	NIAGARA_API virtual void GetVMExternalFunction(const FVMExternalFunctionBindingInfo& BindingInfo, void* InstanceData, FVMExternalFunction& OutFunc) override;
+	NIAGARA_API virtual void DestroyPerInstanceData(void* PerInstanceData, FNiagaraSystemInstance* SystemInstance) override;
 	
-	virtual void ProvidePerInstanceDataForRenderThread(void* InDataFromGT, void* InInstanceData, const FNiagaraSystemInstanceID& SystemInstance) override;
+	NIAGARA_API virtual void ProvidePerInstanceDataForRenderThread(void* InDataFromGT, void* InInstanceData, const FNiagaraSystemInstanceID& SystemInstance) override;
 
 #if WITH_EDITORONLY_DATA
-	virtual bool AppendCompileHash(FNiagaraCompileHashVisitor* InVisitor) const override;
-	virtual void GetParameterDefinitionHLSL(const FNiagaraDataInterfaceGPUParamInfo& ParamInfo, FString& OutHLSL) override;
-	virtual bool GetFunctionHLSL(const FNiagaraDataInterfaceGPUParamInfo& ParamInfo, const FNiagaraDataInterfaceGeneratedFunction& FunctionInfo, int FunctionInstanceIndex, FString& OutHLSL) override;
+	NIAGARA_API virtual bool AppendCompileHash(FNiagaraCompileHashVisitor* InVisitor) const override;
+	NIAGARA_API virtual void GetParameterDefinitionHLSL(const FNiagaraDataInterfaceGPUParamInfo& ParamInfo, FString& OutHLSL) override;
+	NIAGARA_API virtual bool GetFunctionHLSL(const FNiagaraDataInterfaceGPUParamInfo& ParamInfo, const FNiagaraDataInterfaceGeneratedFunction& FunctionInfo, int FunctionInstanceIndex, FString& OutHLSL) override;
 #endif
 
-	virtual void BuildShaderParameters(FNiagaraShaderParametersBuilder& ShaderParametersBuilder) const override;
-	virtual void SetShaderParameters(const FNiagaraDataInterfaceSetShaderParametersContext& Context) const override;
+	NIAGARA_API virtual void BuildShaderParameters(FNiagaraShaderParametersBuilder& ShaderParametersBuilder) const override;
+	NIAGARA_API virtual void SetShaderParameters(const FNiagaraDataInterfaceSetShaderParametersContext& Context) const override;
 
-	void SetFrame(FVectorVMExternalFunctionContext& Context);
-	void ReadFile(FVectorVMExternalFunctionContext& Context);	
-	void GetNumCells(FVectorVMExternalFunctionContext& Context);
+	NIAGARA_API void SetFrame(FVectorVMExternalFunctionContext& Context);
+	NIAGARA_API void ReadFile(FVectorVMExternalFunctionContext& Context);	
+	NIAGARA_API void GetNumCells(FVectorVMExternalFunctionContext& Context);
 
-	FString GetAssetPath(FString PathFormat, int32 FrameIndex) const;
+	NIAGARA_API FString GetAssetPath(FString PathFormat, int32 FrameIndex) const;
 
 protected:
 	//~ UNiagaraDataInterface interface
-	virtual bool CopyToInternal(UNiagaraDataInterface* Destination) const override;
+	NIAGARA_API virtual bool CopyToInternal(UNiagaraDataInterface* Destination) const override;
 	//~ UNiagaraDataInterface interface END
 
 	TMap<FNiagaraSystemInstanceID, FVolumeCacheInstanceData_GameThread*> SystemInstancesToProxyData_GT;	
 
-	static const FName SetFrameName;
-	static const FName ReadFileName;
-	static const FName GetNumCellsName;
-	static const FName IndexToUnitName;
-	static const FName SampleCurrentFrameValueName;
-	static const FName GetCurrentFrameValue;
-	static const FName GetCurrentFrameNumCells;
-	static const TCHAR* TemplateShaderFilePath;
-};
\ No newline at end of file
+	static NIAGARA_API const FName SetFrameName;
+	static NIAGARA_API const FName ReadFileName;
+	static NIAGARA_API const FName GetNumCellsName;
+	static NIAGARA_API const FName IndexToUnitName;
+	static NIAGARA_API const FName SampleCurrentFrameValueName;
+	static NIAGARA_API const FName GetCurrentFrameValue;
+	static NIAGARA_API const FName GetCurrentFrameNumCells;
+	static NIAGARA_API const TCHAR* TemplateShaderFilePath;
+};
diff --git a/Engine/Plugins/FX/Niagara/Source/Niagara/Classes/NiagaraDataInterfaceVolumeTexture.h b/Engine/Plugins/FX/Niagara/Source/Niagara/Classes/NiagaraDataInterfaceVolumeTexture.h
--- a/Engine/Plugins/FX/Niagara/Source/Niagara/Classes/NiagaraDataInterfaceVolumeTexture.h
+++ b/Engine/Plugins/FX/Niagara/Source/Niagara/Classes/NiagaraDataInterfaceVolumeTexture.h
@@ -10,8 +10,8 @@
 class UVolumeTexture;
 
 /** Data Interface allowing sampling of a texture */
-UCLASS(EditInlineNew, Category = "Texture", CollapseCategories, meta = (DisplayName = "Volume Texture Sample"))
-class NIAGARA_API UNiagaraDataInterfaceVolumeTexture : public UNiagaraDataInterface
+UCLASS(EditInlineNew, Category = "Texture", CollapseCategories, meta = (DisplayName = "Volume Texture Sample"), MinimalAPI)
+class UNiagaraDataInterfaceVolumeTexture : public UNiagaraDataInterface
 {
 	GENERATED_UCLASS_BODY()
 
@@ -29,44 +29,44 @@ public:
 	FNiagaraUserParameterBinding TextureUserParameter;
 
 	//UObject Interface
-	virtual void PostInitProperties()override;
+	NIAGARA_API virtual void PostInitProperties()override;
 	//UObject Interface End
 
 	//UNiagaraDataInterface Interface
-	virtual void GetFunctions(TArray<FNiagaraFunctionSignature>& OutFunctions)override;
-	virtual void GetVMExternalFunction(const FVMExternalFunctionBindingInfo& BindingInfo, void* InstanceData, FVMExternalFunction &OutFunc) override;
+	NIAGARA_API virtual void GetFunctions(TArray<FNiagaraFunctionSignature>& OutFunctions)override;
+	NIAGARA_API virtual void GetVMExternalFunction(const FVMExternalFunctionBindingInfo& BindingInfo, void* InstanceData, FVMExternalFunction &OutFunc) override;
 	virtual bool CanExecuteOnTarget(ENiagaraSimTarget Target)const override { return Target == ENiagaraSimTarget::GPUComputeSim; }
 
-	virtual int32 PerInstanceDataSize() const override;
-	virtual bool InitPerInstanceData(void* PerInstanceData, FNiagaraSystemInstance* SystemInstance) override;
-	virtual void DestroyPerInstanceData(void* PerInstanceData, FNiagaraSystemInstance* SystemInstance) override;
+	NIAGARA_API virtual int32 PerInstanceDataSize() const override;
+	NIAGARA_API virtual bool InitPerInstanceData(void* PerInstanceData, FNiagaraSystemInstance* SystemInstance) override;
+	NIAGARA_API virtual void DestroyPerInstanceData(void* PerInstanceData, FNiagaraSystemInstance* SystemInstance) override;
 
 	virtual bool HasPreSimulateTick() const override { return true; }
-	virtual bool PerInstanceTick(void* PerInstanceData, FNiagaraSystemInstance* SystemInstance, float DeltaSeconds) override;
+	NIAGARA_API virtual bool PerInstanceTick(void* PerInstanceData, FNiagaraSystemInstance* SystemInstance, float DeltaSeconds) override;
 	//UNiagaraDataInterface Interface End
 
-	void SampleVolumeTexture(FVectorVMExternalFunctionContext& Context);
-	void GetTextureDimensions(FVectorVMExternalFunctionContext& Context);
+	NIAGARA_API void SampleVolumeTexture(FVectorVMExternalFunctionContext& Context);
+	NIAGARA_API void GetTextureDimensions(FVectorVMExternalFunctionContext& Context);
 
-	virtual bool Equals(const UNiagaraDataInterface* Other) const override;
+	NIAGARA_API virtual bool Equals(const UNiagaraDataInterface* Other) const override;
 
 	// GPU sim functionality
 #if WITH_EDITORONLY_DATA
-	virtual bool AppendCompileHash(FNiagaraCompileHashVisitor* InVisitor) const override;
-	virtual void GetParameterDefinitionHLSL(const FNiagaraDataInterfaceGPUParamInfo& ParamInfo, FString& OutHLSL) override;
-	virtual bool GetFunctionHLSL(const FNiagaraDataInterfaceGPUParamInfo& ParamInfo, const FNiagaraDataInterfaceGeneratedFunction& FunctionInfo, int FunctionInstanceIndex, FString& OutHLSL) override;
+	NIAGARA_API virtual bool AppendCompileHash(FNiagaraCompileHashVisitor* InVisitor) const override;
+	NIAGARA_API virtual void GetParameterDefinitionHLSL(const FNiagaraDataInterfaceGPUParamInfo& ParamInfo, FString& OutHLSL) override;
+	NIAGARA_API virtual bool GetFunctionHLSL(const FNiagaraDataInterfaceGPUParamInfo& ParamInfo, const FNiagaraDataInterfaceGeneratedFunction& FunctionInfo, int FunctionInstanceIndex, FString& OutHLSL) override;
 #endif
-	virtual void BuildShaderParameters(FNiagaraShaderParametersBuilder& ShaderParametersBuilder) const override;
-	virtual void SetShaderParameters(const FNiagaraDataInterfaceSetShaderParametersContext& Context) const override;
+	NIAGARA_API virtual void BuildShaderParameters(FNiagaraShaderParametersBuilder& ShaderParametersBuilder) const override;
+	NIAGARA_API virtual void SetShaderParameters(const FNiagaraDataInterfaceSetShaderParametersContext& Context) const override;
 
-	void SetTexture(UTexture* InTexture);
+	NIAGARA_API void SetTexture(UTexture* InTexture);
 
 protected:
-	virtual bool CopyToInternal(UNiagaraDataInterface* Destination) const override;
+	NIAGARA_API virtual bool CopyToInternal(UNiagaraDataInterface* Destination) const override;
 
 protected:
-	static const TCHAR* TemplateShaderFilePath;
-	static const FName LoadVolumeTextureName;
-	static const FName SampleVolumeTextureName;
-	static const FName TextureDimsName;
+	static NIAGARA_API const TCHAR* TemplateShaderFilePath;
+	static NIAGARA_API const FName LoadVolumeTextureName;
+	static NIAGARA_API const FName SampleVolumeTextureName;
+	static NIAGARA_API const FName TextureDimsName;
 };
diff --git a/Engine/Plugins/FX/Niagara/Source/Niagara/Classes/NiagaraDataSet.h b/Engine/Plugins/FX/Niagara/Source/Niagara/Classes/NiagaraDataSet.h
--- a/Engine/Plugins/FX/Niagara/Source/Niagara/Classes/NiagaraDataSet.h
+++ b/Engine/Plugins/FX/Niagara/Source/Niagara/Classes/NiagaraDataSet.h
@@ -75,30 +75,30 @@ protected:
 };
 
 /** Buffer containing one frame of Niagara simulation data. */
-class NIAGARA_API FNiagaraDataBuffer : public FNiagaraSharedObject
+class FNiagaraDataBuffer : public FNiagaraSharedObject
 {
 	friend class FScopedNiagaraDataSetGPUReadback;
 	//friend class FNiagaraGpuComputeDispatchInterface;
 	
 protected:
-	virtual ~FNiagaraDataBuffer();
+	NIAGARA_API virtual ~FNiagaraDataBuffer();
 
 public:
-	FNiagaraDataBuffer(FNiagaraDataSet* InOwner);
-	void Allocate(uint32 NumInstances, bool bMaintainExisting = false);
-	void ReleaseCPU();
-
-	void AllocateGPU(FRHICommandList& RHICmdList, uint32 InNumInstances, ERHIFeatureLevel::Type FeatureLevel, const TCHAR* DebugSimName);
-	void SwapGPU(FNiagaraDataBuffer* BufferToSwap);
-	void ReleaseGPU();
-
-	void SwapInstances(uint32 OldIndex, uint32 NewIndex);
-	void KillInstance(uint32 InstanceIdx);
-	void CopyTo(FNiagaraDataBuffer& DestBuffer, int32 SrcStartIdx, int32 DestStartIdx, int32 NumInstances)const;
-	void CopyToUnrelated(FNiagaraDataBuffer& DestBuffer, int32 SrcStartIdx, int32 DestStartIdx, int32 NumInstances)const;
-	void GPUCopyFrom(const float* GPUReadBackFloat, const int* GPUReadBackInt, const FFloat16* GPUReadBackHalf, int32 StartIdx, int32 NumInstances, uint32 InSrcFloatStride, uint32 InSrcIntStride, uint32 InSrcHalfStride);
-	void PushCPUBuffersToGPU(const TArray<FNiagaraDataBufferRef>& SourceBuffers, bool bReleaseRef, FRHICommandList& RHICmdList, ERHIFeatureLevel::Type FeatureLevel, const TCHAR* DebugSimName);
-	void Dump(int32 StartIndex, int32 NumInstances, const FString& Label, const FName& SortParameterKey = FName())const;
+	NIAGARA_API FNiagaraDataBuffer(FNiagaraDataSet* InOwner);
+	NIAGARA_API void Allocate(uint32 NumInstances, bool bMaintainExisting = false);
+	NIAGARA_API void ReleaseCPU();
+
+	NIAGARA_API void AllocateGPU(FRHICommandList& RHICmdList, uint32 InNumInstances, ERHIFeatureLevel::Type FeatureLevel, const TCHAR* DebugSimName);
+	NIAGARA_API void SwapGPU(FNiagaraDataBuffer* BufferToSwap);
+	NIAGARA_API void ReleaseGPU();
+
+	NIAGARA_API void SwapInstances(uint32 OldIndex, uint32 NewIndex);
+	NIAGARA_API void KillInstance(uint32 InstanceIdx);
+	NIAGARA_API void CopyTo(FNiagaraDataBuffer& DestBuffer, int32 SrcStartIdx, int32 DestStartIdx, int32 NumInstances)const;
+	NIAGARA_API void CopyToUnrelated(FNiagaraDataBuffer& DestBuffer, int32 SrcStartIdx, int32 DestStartIdx, int32 NumInstances)const;
+	NIAGARA_API void GPUCopyFrom(const float* GPUReadBackFloat, const int* GPUReadBackInt, const FFloat16* GPUReadBackHalf, int32 StartIdx, int32 NumInstances, uint32 InSrcFloatStride, uint32 InSrcIntStride, uint32 InSrcHalfStride);
+	NIAGARA_API void PushCPUBuffersToGPU(const TArray<FNiagaraDataBufferRef>& SourceBuffers, bool bReleaseRef, FRHICommandList& RHICmdList, ERHIFeatureLevel::Type FeatureLevel, const TCHAR* DebugSimName);
+	NIAGARA_API void Dump(int32 StartIndex, int32 NumInstances, const FString& Label, const FName& SortParameterKey = FName())const;
 
 	FORCEINLINE TArrayView<uint8 const* RESTRICT const> GetRegisterTable() const { return TArrayView<uint8 const* RESTRICT const>(RegisterTable); }
 	
@@ -159,19 +159,19 @@ public:
 
 	FORCEINLINE FNiagaraDataSet* GetOwner()const { return Owner; }
 
-	int32 TransferInstance(FNiagaraDataBuffer& SourceBuffer, int32 InstanceIndex, bool bRemoveFromSource=true);
+	NIAGARA_API int32 TransferInstance(FNiagaraDataBuffer& SourceBuffer, int32 InstanceIndex, bool bRemoveFromSource=true);
 
-	bool CheckForNaNs()const;
+	NIAGARA_API bool CheckForNaNs()const;
 
 	FORCEINLINE TArray<int32>& GetIDTable() { return IDToIndexTable; }
 	FORCEINLINE const TArray<int32>& GetIDTable() const { return IDToIndexTable; }
 
 	FORCEINLINE void ClearGPUInstanceCount() { GPUInstanceCountBufferOffset = INDEX_NONE; }
 
-	void BuildRegisterTable();
+	NIAGARA_API void BuildRegisterTable();
 
 private:
-	FORCEINLINE void CheckUsage(bool bReadOnly)const;
+	NIAGARA_API FORCEINLINE void CheckUsage(bool bReadOnly)const;
 
 	FORCEINLINE int32 GetSafeComponentBufferSize(int32 RequiredSize) const
 	{
@@ -240,37 +240,37 @@ private:
 /**
 General storage class for all per instance simulation data in Niagara.
 */
-class NIAGARA_API FNiagaraDataSet
+class FNiagaraDataSet
 {
 	friend FNiagaraDataBuffer;
 	friend class FNiagaraGpuComputeDispatch;
 
 public:
 
-	FNiagaraDataSet();
-	~FNiagaraDataSet();
+	NIAGARA_API FNiagaraDataSet();
+	NIAGARA_API ~FNiagaraDataSet();
 	FNiagaraDataSet& operator=(const FNiagaraDataSet&) = delete;
 
 	/** Initialize the data set with the compiled data */
-	void Init(const FNiagaraDataSetCompiledData* InDataSetCompiledData);
+	NIAGARA_API void Init(const FNiagaraDataSetCompiledData* InDataSetCompiledData);
 
 	/** Resets current data but leaves variable/layout information etc intact. */
-	void ResetBuffers();
+	NIAGARA_API void ResetBuffers();
 
 	/** Begins a new simulation pass and grabs a destination buffer. Returns the new destination data buffer. */
-	FNiagaraDataBuffer& BeginSimulate(bool bResetDestinationData = true);
+	NIAGARA_API FNiagaraDataBuffer& BeginSimulate(bool bResetDestinationData = true);
 
 	/** Ends a simulation pass and sets the current simulation state. */
-	void EndSimulate(bool SetCurrentData = true);
+	NIAGARA_API void EndSimulate(bool SetCurrentData = true);
 
 	/** Set current data directly, you can not call this while inside a BeginSimulate block */
-	void SetCurrentData(FNiagaraDataBuffer* CurrentData);
+	NIAGARA_API void SetCurrentData(FNiagaraDataBuffer* CurrentData);
 
 	/** Allocates space for NumInstances in the current destination buffer. */
-	void Allocate(int32 NumInstances, bool bMaintainExisting = false);
+	NIAGARA_API void Allocate(int32 NumInstances, bool bMaintainExisting = false);
 
 	/** Returns size in bytes for all data buffers currently allocated by this dataset. */
-	int64 GetSizeBytes() const;
+	NIAGARA_API int64 GetSizeBytes() const;
 
 	FORCEINLINE bool IsInitialized() const { return bInitialized; }
 	FORCEINLINE ENiagaraSimTarget GetSimTarget() const { return CompiledData->SimTarget; }
@@ -297,16 +297,16 @@ public:
 	FORCEINLINE uint32 GetNumHalfComponents() const { return CompiledData->TotalHalfComponents; }
 
 	const TArray<FNiagaraVariableLayoutInfo>& GetVariableLayouts() const { return CompiledData->VariableLayouts; }
-	const FNiagaraVariableLayoutInfo* GetVariableLayout(const FNiagaraVariableBase& Var) const;
-	bool GetVariableComponentOffsets(const FNiagaraVariableBase& Var, int32 &FloatStart, int32 &IntStart, int32& HalfStart) const;
+	NIAGARA_API const FNiagaraVariableLayoutInfo* GetVariableLayout(const FNiagaraVariableBase& Var) const;
+	NIAGARA_API bool GetVariableComponentOffsets(const FNiagaraVariableBase& Var, int32 &FloatStart, int32 &IntStart, int32& HalfStart) const;
 
-	void CopyTo(FNiagaraDataSet& Other, int32 StartIdx = 0, int32 NumInstances = INDEX_NONE, bool bResetOther=true)const;
+	NIAGARA_API void CopyTo(FNiagaraDataSet& Other, int32 StartIdx = 0, int32 NumInstances = INDEX_NONE, bool bResetOther=true)const;
 
-	void CopyFromGPUReadback(const float* GPUReadBackFloat, const int* GPUReadBackInt, const FFloat16* GPUReadBackHalf, int32 StartIdx = 0, int32 NumInstances = INDEX_NONE, uint32 FloatStride = 0, uint32 IntStride = 0, uint32 HalfStride = 0);
+	NIAGARA_API void CopyFromGPUReadback(const float* GPUReadBackFloat, const int* GPUReadBackInt, const FFloat16* GPUReadBackHalf, int32 StartIdx = 0, int32 NumInstances = INDEX_NONE, uint32 FloatStride = 0, uint32 IntStride = 0, uint32 HalfStride = 0);
 
-	void CheckForNaNs() const;
+	NIAGARA_API void CheckForNaNs() const;
 
-	void Dump(int32 StartIndex, int32 NumInstances, const FString& Label, const FName& SortParameterKey = FName()) const;
+	NIAGARA_API void Dump(int32 StartIndex, int32 NumInstances, const FString& Label, const FName& SortParameterKey = FName()) const;
 
 	FORCEINLINE bool IsCurrentDataValid()const { return CurrentData != nullptr; }
 	FORCEINLINE FNiagaraDataBuffer* GetCurrentData()const {	return CurrentData; }
@@ -324,7 +324,7 @@ public:
 		return *DestinationData;
 	}
 
-	void AllocateGPUFreeIDs(uint32 InNumInstances, FRHICommandList& RHICmdList, ERHIFeatureLevel::Type FeatureLevel, const TCHAR* DebugSimName);
+	NIAGARA_API void AllocateGPUFreeIDs(uint32 InNumInstances, FRHICommandList& RHICmdList, ERHIFeatureLevel::Type FeatureLevel, const TCHAR* DebugSimName);
 
 	void SetMaxInstanceCount(uint32 InMaxInstanceCount) { MaxInstanceCount = InMaxInstanceCount; }
 	void SetMaxAllocationCount(uint32 InMaxAllocationCount) { MaxAllocationCount = InMaxAllocationCount; }
@@ -336,9 +336,9 @@ public:
 	int NumSpawnedIDs;
 private:
 
-	void Reset();
+	NIAGARA_API void Reset();
 
-	void ResetBuffersInternal();
+	NIAGARA_API void ResetBuffersInternal();
 
 	FORCEINLINE void CheckCorrectThread()const
 	{
@@ -470,7 +470,7 @@ Allows immediate access to GPU data on the CPU, you can then use FNiagaraDataSet
 This will make a copy of the GPU data and will stall the CPU until the data is ready from the GPU,
 therefore it should only be used for tools / debugging.  For async readback see FNiagaraSystemInstance::RequestCapture.
 */
-class NIAGARA_API FScopedNiagaraDataSetGPUReadback
+class FScopedNiagaraDataSetGPUReadback
 {
 public:
 	FScopedNiagaraDataSetGPUReadback() {}
@@ -483,7 +483,7 @@ public:
 		}
 	}
 
-	void ReadbackData(FNiagaraGpuComputeDispatchInterface* ComputeDispatchInterface, FNiagaraDataSet* InDataSet);
+	NIAGARA_API void ReadbackData(FNiagaraGpuComputeDispatchInterface* ComputeDispatchInterface, FNiagaraDataSet* InDataSet);
 	uint32 GetNumInstances() const { check(DataSet != nullptr); return NumInstances; }
 
 private:
diff --git a/Engine/Plugins/FX/Niagara/Source/Niagara/Classes/NiagaraDataSetCompiledData.h b/Engine/Plugins/FX/Niagara/Source/Niagara/Classes/NiagaraDataSetCompiledData.h
--- a/Engine/Plugins/FX/Niagara/Source/Niagara/Classes/NiagaraDataSetCompiledData.h
+++ b/Engine/Plugins/FX/Niagara/Source/Niagara/Classes/NiagaraDataSetCompiledData.h
@@ -42,7 +42,7 @@ public:
 };
 
 USTRUCT()
-struct NIAGARA_API FNiagaraDataSetCompiledData
+struct FNiagaraDataSetCompiledData
 {
 	GENERATED_BODY()
 
@@ -76,12 +76,12 @@ struct NIAGARA_API FNiagaraDataSetCompiledData
 	UPROPERTY()
 	ENiagaraSimTarget SimTarget;
 
-	FNiagaraDataSetCompiledData();
-	const FNiagaraVariableLayoutInfo* FindVariableLayoutInfo(const FNiagaraVariableBase& VariableDef) const;
-	void BuildLayout();
-	void Empty();
+	NIAGARA_API FNiagaraDataSetCompiledData();
+	NIAGARA_API const FNiagaraVariableLayoutInfo* FindVariableLayoutInfo(const FNiagaraVariableBase& VariableDef) const;
+	NIAGARA_API void BuildLayout();
+	NIAGARA_API void Empty();
 
-	static FNiagaraDataSetCompiledData DummyCompiledData;
+	static NIAGARA_API FNiagaraDataSetCompiledData DummyCompiledData;
 
 	uint32 GetLayoutHash()const
 	{
diff --git a/Engine/Plugins/FX/Niagara/Source/Niagara/Classes/NiagaraEffectType.h b/Engine/Plugins/FX/Niagara/Source/Niagara/Classes/NiagaraEffectType.h
--- a/Engine/Plugins/FX/Niagara/Source/Niagara/Classes/NiagaraEffectType.h
+++ b/Engine/Plugins/FX/Niagara/Source/Niagara/Classes/NiagaraEffectType.h
@@ -181,7 +181,7 @@ struct FNiagaraSystemVisibilityCullingSettings
 
 /** Scalability settings for Niagara Systems for a particular platform set (unless overridden). */
 USTRUCT()
-struct NIAGARA_API FNiagaraSystemScalabilitySettings
+struct FNiagaraSystemScalabilitySettings
 {
 	GENERATED_USTRUCT_BODY()
 
@@ -247,15 +247,15 @@ struct NIAGARA_API FNiagaraSystemScalabilitySettings
 	UPROPERTY(EditAnywhere, Category = "Scalability", meta = (DisplayInScalabilityValuesBar))
 	FNiagaraGlobalBudgetScaling BudgetScaling;
 	
-	FNiagaraSystemScalabilitySettings();
+	NIAGARA_API FNiagaraSystemScalabilitySettings();
 
-	void Clear();
-	void PostLoad(int32 Version);
+	NIAGARA_API void Clear();
+	NIAGARA_API void PostLoad(int32 Version);
 };
 
 /** Container struct for an array of system scalability settings. Enables details customization and data validation. */
 USTRUCT()
-struct NIAGARA_API FNiagaraSystemScalabilitySettingsArray
+struct FNiagaraSystemScalabilitySettingsArray
 {
 	GENERATED_USTRUCT_BODY()
 
@@ -264,11 +264,11 @@ struct NIAGARA_API FNiagaraSystemScalabilitySettingsArray
 };
 
 USTRUCT()
-struct NIAGARA_API FNiagaraSystemScalabilityOverride : public FNiagaraSystemScalabilitySettings
+struct FNiagaraSystemScalabilityOverride : public FNiagaraSystemScalabilitySettings
 {
 	GENERATED_USTRUCT_BODY()
 
-	FNiagaraSystemScalabilityOverride();
+	NIAGARA_API FNiagaraSystemScalabilityOverride();
 
 	/** Controls whether we override the distance culling settings. */
 	UPROPERTY(EditAnywhere, Category = "Override")
@@ -292,7 +292,7 @@ struct NIAGARA_API FNiagaraSystemScalabilityOverride : public FNiagaraSystemScal
 
 /** Scalability settings for Niagara Emitters on a particular platform set. */
 USTRUCT()
-struct NIAGARA_API FNiagaraEmitterScalabilitySettings
+struct FNiagaraEmitterScalabilitySettings
 {
 	GENERATED_USTRUCT_BODY()
 
@@ -308,13 +308,13 @@ struct NIAGARA_API FNiagaraEmitterScalabilitySettings
 	UPROPERTY(EditAnywhere, Category = "Scalability", meta = (EditCondition = "bScaleSpawnCount", DisplayInScalabilityValuesBar))
 	float SpawnCountScale;
 
-	FNiagaraEmitterScalabilitySettings();
-	void Clear();
+	NIAGARA_API FNiagaraEmitterScalabilitySettings();
+	NIAGARA_API void Clear();
 };
 
 /** Container struct for an array of emitter scalability settings. Enables details customization and data validation. */
 USTRUCT()
-struct NIAGARA_API FNiagaraEmitterScalabilitySettingsArray
+struct FNiagaraEmitterScalabilitySettingsArray
 {
 	GENERATED_USTRUCT_BODY()
 
@@ -323,11 +323,11 @@ struct NIAGARA_API FNiagaraEmitterScalabilitySettingsArray
 };
 
 USTRUCT()
-struct NIAGARA_API FNiagaraEmitterScalabilityOverride : public FNiagaraEmitterScalabilitySettings
+struct FNiagaraEmitterScalabilityOverride : public FNiagaraEmitterScalabilitySettings
 {
 	GENERATED_USTRUCT_BODY()
 
-	FNiagaraEmitterScalabilityOverride();
+	NIAGARA_API FNiagaraEmitterScalabilityOverride();
 
 	//Controls whether spawn count scale should be overridden.
 	UPROPERTY(EditAnywhere, Category = "Override")
@@ -336,7 +336,7 @@ struct NIAGARA_API FNiagaraEmitterScalabilityOverride : public FNiagaraEmitterSc
 
 /** Container struct for an array of emitter scalability overrides. Enables details customization and data validation. */
 USTRUCT()
-struct NIAGARA_API FNiagaraEmitterScalabilityOverrides
+struct FNiagaraEmitterScalabilityOverrides
 {
 	GENERATED_USTRUCT_BODY()
 
@@ -356,53 +356,53 @@ These allow Niagara's scalability system to determine the relative significance
 Some basic ones are provided but projects are free to implement their own more complex determinations of significance.
 For example, FX attached to the player character could be given higher priority.
 */
-UCLASS(abstract, EditInlineNew)
-class NIAGARA_API UNiagaraSignificanceHandler : public UObject
+UCLASS(abstract, EditInlineNew, MinimalAPI)
+class UNiagaraSignificanceHandler : public UObject
 {
 	GENERATED_BODY()
 
 public:
-	virtual void CalculateSignificance(TConstArrayView<UNiagaraComponent*> Components, TArrayView<FNiagaraScalabilityState> OutState, TConstArrayView<FNiagaraScalabilitySystemData> SystemData, TArray<int32>& OutIndices)PURE_VIRTUAL(CalculateSignificance, );
+	NIAGARA_API virtual void CalculateSignificance(TConstArrayView<UNiagaraComponent*> Components, TArrayView<FNiagaraScalabilityState> OutState, TConstArrayView<FNiagaraScalabilitySystemData> SystemData, TArray<int32>& OutIndices)PURE_VIRTUAL(CalculateSignificance, );
 };
 
 /** Significance is determined by the system's distance to the nearest camera. Closer systems are more significant. */
-UCLASS(EditInlineNew, meta = (DisplayName = "Distance"))
-class NIAGARA_API UNiagaraSignificanceHandlerDistance : public UNiagaraSignificanceHandler
+UCLASS(EditInlineNew, meta = (DisplayName = "Distance"), MinimalAPI)
+class UNiagaraSignificanceHandlerDistance : public UNiagaraSignificanceHandler
 {
 	GENERATED_BODY()
 
 public:
-	virtual void CalculateSignificance(TConstArrayView<UNiagaraComponent*> Components, TArrayView<FNiagaraScalabilityState> OutState, TConstArrayView<FNiagaraScalabilitySystemData> SystemData, TArray<int32>& OutIndices) override;
+	NIAGARA_API virtual void CalculateSignificance(TConstArrayView<UNiagaraComponent*> Components, TArrayView<FNiagaraScalabilityState> OutState, TConstArrayView<FNiagaraScalabilitySystemData> SystemData, TArray<int32>& OutIndices) override;
 };
 
 /** Significance is determined by the system's age. Newer systems are more significant. */
-UCLASS(EditInlineNew, meta = (DisplayName = "Age"))
-class NIAGARA_API UNiagaraSignificanceHandlerAge : public UNiagaraSignificanceHandler
+UCLASS(EditInlineNew, meta = (DisplayName = "Age"), MinimalAPI)
+class UNiagaraSignificanceHandlerAge : public UNiagaraSignificanceHandler
 {
 	GENERATED_BODY()
 
 public:
-	virtual void CalculateSignificance(TConstArrayView<UNiagaraComponent*> Components, TArrayView<FNiagaraScalabilityState> OutState, TConstArrayView<FNiagaraScalabilitySystemData> SystemData, TArray<int32>& OutIndices) override;
+	NIAGARA_API virtual void CalculateSignificance(TConstArrayView<UNiagaraComponent*> Components, TArrayView<FNiagaraScalabilityState> OutState, TConstArrayView<FNiagaraScalabilitySystemData> SystemData, TArray<int32>& OutIndices) override;
 };
 
 //////////////////////////////////////////////////////////////////////////
 
 /** Contains settings and working data shared among many NiagaraSystems that share some commonality of type. For example ImpactFX vs EnvironmentalFX. */
-UCLASS(config = Niagara, perObjectConfig)
-class NIAGARA_API UNiagaraEffectType : public UObject
+UCLASS(config = Niagara, perObjectConfig, MinimalAPI)
+class UNiagaraEffectType : public UObject
 {
 	GENERATED_UCLASS_BODY()
 
 	//UObject Interface
-	virtual void BeginDestroy()override;
-	virtual bool IsReadyForFinishDestroy()override;
-	virtual void Serialize(FArchive& Ar)override;
-	virtual void PostLoad()override;
+	NIAGARA_API virtual void BeginDestroy()override;
+	NIAGARA_API virtual bool IsReadyForFinishDestroy()override;
+	NIAGARA_API virtual void Serialize(FArchive& Ar)override;
+	NIAGARA_API virtual void PostLoad()override;
 #if WITH_EDITORONLY_DATA
-	static void DeclareConstructClasses(TArray<FTopLevelAssetPath>& OutConstructClasses, const UClass* SpecificSubclass);
+	static NIAGARA_API void DeclareConstructClasses(TArray<FTopLevelAssetPath>& OutConstructClasses, const UClass* SpecificSubclass);
 #endif
 #if WITH_EDITOR
-	virtual void PostEditChangeProperty(struct FPropertyChangedEvent& PropertyChangedEvent);
+	NIAGARA_API virtual void PostEditChangeProperty(struct FPropertyChangedEvent& PropertyChangedEvent);
 #endif
 	//UObject Interface END
 
@@ -445,8 +445,8 @@ class NIAGARA_API UNiagaraEffectType : public UObject
 	FORCEINLINE const FNiagaraSystemScalabilitySettingsArray& GetSystemScalabilitySettings()const { return SystemScalabilitySettings; }
 	FORCEINLINE const FNiagaraEmitterScalabilitySettingsArray& GetEmitterScalabilitySettings()const { return EmitterScalabilitySettings; }
 
-	const FNiagaraSystemScalabilitySettings& GetActiveSystemScalabilitySettings()const;
-	const FNiagaraEmitterScalabilitySettings& GetActiveEmitterScalabilitySettings()const;
+	NIAGARA_API const FNiagaraSystemScalabilitySettings& GetActiveSystemScalabilitySettings()const;
+	NIAGARA_API const FNiagaraEmitterScalabilitySettings& GetActiveEmitterScalabilitySettings()const;
 
 	UNiagaraSignificanceHandler* GetSignificanceHandler()const { return SignificanceHandler; }
 
@@ -460,11 +460,11 @@ class NIAGARA_API UNiagaraEffectType : public UObject
 	UNiagaraBaselineController* GetPerfBaselineController() { return PerformanceBaselineController; }
 	FNiagaraPerfBaselineStats& GetPerfBaselineStats()const { return PerfBaselineStats; }
 	FORCEINLINE bool IsPerfBaselineValid()const { return PerfBaselineVersion == CurrentPerfBaselineVersion; }
-	void UpdatePerfBaselineStats(FNiagaraPerfBaselineStats& NewBaselineStats);
+	NIAGARA_API void UpdatePerfBaselineStats(FNiagaraPerfBaselineStats& NewBaselineStats);
 
-	void InvalidatePerfBaseline();
+	NIAGARA_API void InvalidatePerfBaseline();
 
-	void SpawnBaselineActor(UWorld* World);
+	NIAGARA_API void SpawnBaselineActor(UWorld* World);
 #endif
 	/** Performs the passed action for all FNiagaraPlatformSets in this system. */
 	template<typename TAction>
@@ -490,15 +490,15 @@ private:
 
 #if NIAGARA_PERF_BASELINES
 	/** The current version for perf baselines. Regenerate this if there are significant performance improvements that would invalidate existing baseline data. */
-	static const FGuid CurrentPerfBaselineVersion;
+	static NIAGARA_API const FGuid CurrentPerfBaselineVersion;
 
 	DECLARE_DELEGATE_OneParam(FGeneratePerfBaselines, TArray<UNiagaraEffectType*>&/**BaselinesToGenerate*/);
 
 	/** Delegate allowing us to call into editor code to generate performance baselines. */
-	static FGeneratePerfBaselines GeneratePerfBaselinesDelegate;
+	static NIAGARA_API FGeneratePerfBaselines GeneratePerfBaselinesDelegate;
 public:
 	static FGeneratePerfBaselines& OnGeneratePerfBaselines(){ return GeneratePerfBaselinesDelegate; }
-	static void GeneratePerfBaselines();
+	static NIAGARA_API void GeneratePerfBaselines();
 #endif
 };
 
