[Niagara]

diff --git a/Engine/Plugins/FX/Niagara/Source/Niagara/Classes/NiagaraEmitter.h b/Engine/Plugins/FX/Niagara/Source/Niagara/Classes/NiagaraEmitter.h
--- a/Engine/Plugins/FX/Niagara/Source/Niagara/Classes/NiagaraEmitter.h
+++ b/Engine/Plugins/FX/Niagara/Source/Niagara/Classes/NiagaraEmitter.h
@@ -221,25 +221,25 @@ struct FNiagaraDetailsLevelScaleOverrides
 	float Cine;
 };
 
-struct NIAGARA_API FMemoryRuntimeEstimation
+struct FMemoryRuntimeEstimation
 {
 	TMap<uint64, int32> RuntimeAllocations;
 	bool IsEstimationDirty = false;
 	int32 AllocationEstimate = 0;
 
-	FCriticalSection* GetCriticalSection();
-	void Init();
+	NIAGARA_API FCriticalSection* GetCriticalSection();
+	NIAGARA_API void Init();
 private:
 	TSharedPtr<FCriticalSection> EstimationCriticalSection;
 };
 
 /** Struct containing all of the data that can be different between different emitter versions.*/
 USTRUCT(BlueprintInternalUseOnly)
-struct NIAGARA_API FVersionedNiagaraEmitterData
+struct FVersionedNiagaraEmitterData
 {
 	GENERATED_BODY()
 
-	FVersionedNiagaraEmitterData();
+	NIAGARA_API FVersionedNiagaraEmitterData();
 	
 	UPROPERTY()
 	FNiagaraAssetVersion Version;
@@ -359,23 +359,23 @@ struct NIAGARA_API FVersionedNiagaraEmitterData
 	UPROPERTY()
 	TMap<FNiagaraVariableBase, FNiagaraVariableBase> ResolvedDIBindings;
 
-	void CopyFrom(const FVersionedNiagaraEmitterData& Source);
-	void PostLoad(UNiagaraEmitter& Emitter, int32 NiagaraVer);
+	NIAGARA_API void CopyFrom(const FVersionedNiagaraEmitterData& Source);
+	NIAGARA_API void PostLoad(UNiagaraEmitter& Emitter, int32 NiagaraVer);
 
-	void PostInitProperties(UNiagaraEmitter* Outer);
-	bool UsesCollection(const UNiagaraParameterCollection* Collection) const;
-	bool UsesScript(const UNiagaraScript* Script) const;
-	bool IsReadyToRun() const;
+	NIAGARA_API void PostInitProperties(UNiagaraEmitter* Outer);
+	NIAGARA_API bool UsesCollection(const UNiagaraParameterCollection* Collection) const;
+	NIAGARA_API bool UsesScript(const UNiagaraScript* Script) const;
+	NIAGARA_API bool IsReadyToRun() const;
 	const TArray<UNiagaraRendererProperties*>& GetRenderers() const { return RendererProperties; }
-	void GetScripts(TArray<UNiagaraScript*>& OutScripts, bool bCompilableOnly = true, bool bEnabledOnly = false) const;
-	UNiagaraScript* GetScript(ENiagaraScriptUsage Usage, FGuid UsageId);
+	NIAGARA_API void GetScripts(TArray<UNiagaraScript*>& OutScripts, bool bCompilableOnly = true, bool bEnabledOnly = false) const;
+	NIAGARA_API UNiagaraScript* GetScript(ENiagaraScriptUsage Usage, FGuid UsageId);
 	UNiagaraScript* GetGPUComputeScript() { return GPUComputeScript; }
 	const UNiagaraScript* GetGPUComputeScript() const { return GPUComputeScript; }
 	FORCEINLINE const TArray<FNiagaraEventScriptProperties>& GetEventHandlers() const { return EventHandlerScriptProps; }
-	void CacheFromCompiledData(const FNiagaraDataSetCompiledData* CompiledData, const UNiagaraEmitter& Emitter);
-	void CacheFromShaderCompiled();
+	NIAGARA_API void CacheFromCompiledData(const FNiagaraDataSetCompiledData* CompiledData, const UNiagaraEmitter& Emitter);
+	NIAGARA_API void CacheFromShaderCompiled();
 
-	FGraphEventArray PrecacheComputePSOs(const UNiagaraEmitter& NiagaraEmitter);
+	NIAGARA_API FGraphEventArray PrecacheComputePSOs(const UNiagaraEmitter& NiagaraEmitter);
 	bool DidPSOPrecacheFail() const { return PSOPrecacheResult == EPSOPrecacheResult::NotSupported; }
 
 	bool RequiresViewUniformBuffer() const { return bRequiresViewUniformBuffer; }
@@ -383,23 +383,23 @@ struct NIAGARA_API FVersionedNiagaraEmitterData
 	uint32 GetMaxInstanceCount() const { return MaxInstanceCount; }
 	uint32 GetMaxAllocationCount() const { return MaxAllocationCount; }
 	TConstArrayView<TSharedPtr<FNiagaraBoundsCalculator>> GetBoundsCalculators() const { return MakeArrayView(BoundsCalculators); }
-	bool RequiresPersistentIDs() const;
+	NIAGARA_API bool RequiresPersistentIDs() const;
 	const TArray<UNiagaraSimulationStageBase*>& GetSimulationStages() const { return SimulationStages; }
 	FORCEINLINE const FNiagaraEmitterScalabilitySettings& GetScalabilitySettings()const { return CurrentScalabilitySettings; }
-	const FNiagaraEmitterScalabilityOverride& GetCurrentOverrideSettings() const;
-	UNiagaraSimulationStageBase* GetSimulationStageById(FGuid ScriptUsageId) const;
-	bool BuildParameterStoreRendererBindings(FNiagaraParameterStore& ParameterStore) const;
-	void RebuildRendererBindings(const UNiagaraEmitter& Emitter);
+	NIAGARA_API const FNiagaraEmitterScalabilityOverride& GetCurrentOverrideSettings() const;
+	NIAGARA_API UNiagaraSimulationStageBase* GetSimulationStageById(FGuid ScriptUsageId) const;
+	NIAGARA_API bool BuildParameterStoreRendererBindings(FNiagaraParameterStore& ParameterStore) const;
+	NIAGARA_API void RebuildRendererBindings(const UNiagaraEmitter& Emitter);
 	FORCEINLINE static FBox GetDefaultFixedBounds() { return FBox(FVector(-100), FVector(100)); }
 	
-	FVersionedNiagaraEmitter GetParent() const;
-	FVersionedNiagaraEmitter GetParentAtLastMerge() const;
-	void RemoveParent();
-	void Reparent(const FVersionedNiagaraEmitter& InParent);
+	NIAGARA_API FVersionedNiagaraEmitter GetParent() const;
+	NIAGARA_API FVersionedNiagaraEmitter GetParentAtLastMerge() const;
+	NIAGARA_API void RemoveParent();
+	NIAGARA_API void Reparent(const FVersionedNiagaraEmitter& InParent);
 	
-	bool IsValid() const;
-	void UpdateDebugName(const UNiagaraEmitter& Emitter, const FNiagaraDataSetCompiledData* CompiledData);
-	void SyncEmitterAlias(const FString& InOldName, const UNiagaraEmitter& InEmitter);
+	NIAGARA_API bool IsValid() const;
+	NIAGARA_API void UpdateDebugName(const UNiagaraEmitter& Emitter, const FNiagaraDataSetCompiledData* CompiledData);
+	NIAGARA_API void SyncEmitterAlias(const FString& InOldName, const UNiagaraEmitter& InEmitter);
 
 	template<typename TAction>
 	void ForEachRenderer(TAction Func) const;
@@ -414,21 +414,21 @@ struct NIAGARA_API FVersionedNiagaraEmitterData
 	void ForEachPlatformSet(TAction Func);
 
 	/** Returns true if this emitter's platform filter allows it on this platform and quality level. */
-	bool IsAllowedByScalability() const;
+	NIAGARA_API bool IsAllowedByScalability() const;
 
 	/* Returns the number of max expected particles for memory allocations. */
-	int32 GetMaxParticleCountEstimate();
+	NIAGARA_API int32 GetMaxParticleCountEstimate();
 
 	/* Gets whether or not the supplied event generator id matches an event generator which is shared between the particle spawn and update scrips. */
-	bool IsEventGeneratorShared(FName EventGeneratorId) const;
+	NIAGARA_API bool IsEventGeneratorShared(FName EventGeneratorId) const;
 
 	/* This is used by the emitter instances to report runtime allocations to reduce reallocation in future simulation runs. */
-	int32 AddRuntimeAllocation(uint64 ReporterHandle, int32 AllocationCount);
-	void ClearRuntimeAllocationEstimate(uint64 ReportHandle = INDEX_NONE);
+	NIAGARA_API int32 AddRuntimeAllocation(uint64 ReporterHandle, int32 AllocationCount);
+	NIAGARA_API void ClearRuntimeAllocationEstimate(uint64 ReportHandle = INDEX_NONE);
 
 	/* Gets a pointer to an event handler by script usage id.  This method is potentially unsafe because modifications to
 	the event handler array can make this pointer become invalid without warning. */
-	FNiagaraEventScriptProperties* GetEventHandlerByIdUnsafe(FGuid ScriptUsageId);
+	NIAGARA_API FNiagaraEventScriptProperties* GetEventHandlerByIdUnsafe(FGuid ScriptUsageId);
  
 #if WITH_EDITORONLY_DATA
 	/** An allow list of Particle attributes (e.g. "Particle.Position" or "Particle.Age") that will not be removed from the DataSet  even if they aren't read by the VM.
@@ -462,14 +462,14 @@ struct NIAGARA_API FVersionedNiagaraEmitterData
 	UPROPERTY()
 	FVersionedNiagaraEmitter VersionedParentAtLastMerge;
 
-	bool AreAllScriptAndSourcesSynchronized() const;
-	void InvalidateCompileResults();
-	bool IsSynchronizedWithParent() const;
-	bool UsesEmitter(const UNiagaraEmitter& InEmitter) const;
-	void GatherStaticVariables(TArray<FNiagaraVariable>& OutVars) const;
+	NIAGARA_API bool AreAllScriptAndSourcesSynchronized() const;
+	NIAGARA_API void InvalidateCompileResults();
+	NIAGARA_API bool IsSynchronizedWithParent() const;
+	NIAGARA_API bool UsesEmitter(const UNiagaraEmitter& InEmitter) const;
+	NIAGARA_API void GatherStaticVariables(TArray<FNiagaraVariable>& OutVars) const;
 	
-	UNiagaraEditorDataBase* GetEditorData() const;
-	UNiagaraEditorParametersAdapterBase* GetEditorParameters();
+	NIAGARA_API UNiagaraEditorDataBase* GetEditorData() const;
+	NIAGARA_API UNiagaraEditorParametersAdapterBase* GetEditorParameters();
 #endif
 
 #if STATS
@@ -482,7 +482,7 @@ struct NIAGARA_API FVersionedNiagaraEmitterData
 	const TCHAR* GetDebugSimName() const { return TEXT(""); }
 #endif
 
-	void GatherCompiledParticleAttributes(TArray<FNiagaraVariableBase>& OutVariables) const;
+	NIAGARA_API void GatherCompiledParticleAttributes(TArray<FNiagaraVariableBase>& OutVariables) const;
 	
 private:
 	UPROPERTY()
@@ -538,8 +538,8 @@ private:
 
 	friend UNiagaraEmitter;
 
-	void EnsureScriptsPostLoaded();
-	void OnPostCompile(const UNiagaraEmitter& InEmitter);
+	NIAGARA_API void EnsureScriptsPostLoaded();
+	NIAGARA_API void OnPostCompile(const UNiagaraEmitter& InEmitter);
 
 	EPSOPrecacheResult PSOPrecacheResult = EPSOPrecacheResult::Unknown;
 };
@@ -562,9 +562,9 @@ public:
 	DECLARE_MULTICAST_DELEGATE(FOnEventHandlersChanged);
 	DECLARE_MULTICAST_DELEGATE_OneParam(FOnEmitterCompiled, FVersionedNiagaraEmitter);
 
-	struct NIAGARA_API PrivateMemberNames
+	struct PrivateMemberNames
 	{
-		static const FName EventHandlerScriptProps;
+		static NIAGARA_API const FName EventHandlerScriptProps;
 	};
 #endif
 
@@ -649,7 +649,7 @@ public:
 	NIAGARA_API virtual TSharedPtr<FNiagaraVersionDataAccessor> GetVersionDataAccessor(const FGuid& Version) override; 
 
 	/** If true then this script asset uses active version control to track changes. */
-	NIAGARA_API virtual bool IsVersioningEnabled() const override { return bVersioningEnabled; }
+	virtual bool IsVersioningEnabled() const override { return bVersioningEnabled; }
 
 	/** Returns the version of the exposed version data (i.e. the version used when adding an emitter to a system) */
 	NIAGARA_API virtual FNiagaraAssetVersion GetExposedVersion() const override;
@@ -853,7 +853,7 @@ public:
 	NIAGARA_API FOnEmitterCompiled& OnEmitterGPUCompiled();
 
 	/** Callback issued whenever a GPU compilation successfully happened (even if the results are a script that cannot be executed due to errors)*/
-	NIAGARA_API FOnEmitterCompiled& OnGPUCompilationComplete()
+	FOnEmitterCompiled& OnGPUCompilationComplete()
 	{
 		return OnGPUScriptCompiledDelegate;
 	}
@@ -898,7 +898,7 @@ public:
 	NIAGARA_API void ChangeParentVersion(const FGuid& NewParentVersion, const FGuid& EmitterVersion);
 
 	NIAGARA_API void NotifyScratchPadScriptsChanged();
-	NIAGARA_API FNiagaraMessageStore& GetMessageStore() { return MessageStore; }
+	FNiagaraMessageStore& GetMessageStore() { return MessageStore; }
 #endif
 
 protected:
diff --git a/Engine/Plugins/FX/Niagara/Source/Niagara/Classes/NiagaraEmitterHandle.h b/Engine/Plugins/FX/Niagara/Source/Niagara/Classes/NiagaraEmitterHandle.h
--- a/Engine/Plugins/FX/Niagara/Source/Niagara/Classes/NiagaraEmitterHandle.h
+++ b/Engine/Plugins/FX/Niagara/Source/Niagara/Classes/NiagaraEmitterHandle.h
@@ -19,42 +19,42 @@ class UNiagaraSystem;
  * Stores emitter information within the context of a System.
  */
 USTRUCT()
-struct NIAGARA_API FNiagaraEmitterHandle
+struct FNiagaraEmitterHandle
 {
 	GENERATED_USTRUCT_BODY()
 public:
 	/** Creates a new invalid emitter handle. */
-	FNiagaraEmitterHandle();
+	NIAGARA_API FNiagaraEmitterHandle();
 
 #if WITH_EDITORONLY_DATA
 	/** Creates a new emitter handle from an emitter. */
-	FNiagaraEmitterHandle(UNiagaraEmitter& InEmitter, const FGuid& Version);
+	NIAGARA_API FNiagaraEmitterHandle(UNiagaraEmitter& InEmitter, const FGuid& Version);
 	
 	/** Creates a new emitter handle from an emitter. */
-	FNiagaraEmitterHandle(const FVersionedNiagaraEmitter& InEmitter);
+	NIAGARA_API FNiagaraEmitterHandle(const FVersionedNiagaraEmitter& InEmitter);
 #endif
 
 	/** Whether or not this is a valid emitter handle. */
-	bool IsValid() const;
+	NIAGARA_API bool IsValid() const;
 
 	/** Gets the unique id for this handle. */
-	FGuid GetId() const;
+	NIAGARA_API FGuid GetId() const;
 
 	// HACK!  Data sets used to use the emitter name, but this isn't guaranteed to be unique.  This is a temporary hack
 	// to allow the data sets to continue work with using names, but that code needs to be refactored to use the id defined here.
-	FName GetIdName() const;
+	NIAGARA_API FName GetIdName() const;
 
 	/** Gets the display name for this emitter in the System. */
-	FName GetName() const;
+	NIAGARA_API FName GetName() const;
 
 	/** Sets the display name for this emitter in the System. The system is needed here in order to ensure uniqueness of the name. */
-	void SetName(FName InName, UNiagaraSystem& InOwnerSystem);
+	NIAGARA_API void SetName(FName InName, UNiagaraSystem& InOwnerSystem);
 
 	/** Gets whether or not this emitter is enabled within the System.  Disabled emitters aren't simulated. */
-	bool GetIsEnabled() const;
+	NIAGARA_API bool GetIsEnabled() const;
 
 	/** Sets whether or not this emitter is enabled within the System.  Disabled emitters aren't simulated. Returns whether or not the enabled state changed. */
-	bool SetIsEnabled(bool bInIsEnabled, UNiagaraSystem& InOwnerSystem, bool bRecompileIfChanged);
+	NIAGARA_API bool SetIsEnabled(bool bInIsEnabled, UNiagaraSystem& InOwnerSystem, bool bRecompileIfChanged);
 
 #if WITH_EDITORONLY_DATA
 	bool IsIsolated() const {	return bIsolated; }
@@ -62,33 +62,33 @@ public:
 #endif
 
 	/** Gets the copied instance of the emitter this handle references. */
-	FVersionedNiagaraEmitter GetInstance() const;
-	FVersionedNiagaraEmitter& GetInstance();
+	NIAGARA_API FVersionedNiagaraEmitter GetInstance() const;
+	NIAGARA_API FVersionedNiagaraEmitter& GetInstance();
 
-	FVersionedNiagaraEmitterData* GetEmitterData() const;
+	NIAGARA_API FVersionedNiagaraEmitterData* GetEmitterData() const;
 
 	/** Gets a unique name for this emitter instance for use in scripts and parameter stores etc.*/
-	FString GetUniqueInstanceName()const;
+	NIAGARA_API FString GetUniqueInstanceName()const;
 
 #if WITH_EDITORONLY_DATA
 	/** Determine whether or not the Instance script is in synch with its graph.*/
-	bool NeedsRecompile() const;
+	NIAGARA_API bool NeedsRecompile() const;
 
 	/** Calls conditional post load on all sub-objects this handle references. */
-	void ConditionalPostLoad(int32 NiagaraCustomVersion);
+	NIAGARA_API void ConditionalPostLoad(int32 NiagaraCustomVersion);
 
 	/** Whether or not this handle uses the supplied emitter. */
-	bool UsesEmitter(const FVersionedNiagaraEmitter& InEmitter) const;
-	bool UsesEmitter(const UNiagaraEmitter& InEmitter) const;
+	NIAGARA_API bool UsesEmitter(const FVersionedNiagaraEmitter& InEmitter) const;
+	NIAGARA_API bool UsesEmitter(const UNiagaraEmitter& InEmitter) const;
 
-	void ClearEmitter();
+	NIAGARA_API void ClearEmitter();
 
 	bool GetDebugShowBounds() const { return bDebugShowBounds; }
 	void SetDebugShowBounds(bool bShowBounds) { bDebugShowBounds = bShowBounds; }
 #endif
 public:
 	/** A static const invalid handle. */
-	static const FNiagaraEmitterHandle InvalidHandle;
+	static NIAGARA_API const FNiagaraEmitterHandle InvalidHandle;
 
 private:
 	/** The id of this emitter handle. */
diff --git a/Engine/Plugins/FX/Niagara/Source/Niagara/Classes/NiagaraEmitterInstance.h b/Engine/Plugins/FX/Niagara/Source/Niagara/Classes/NiagaraEmitterInstance.h
--- a/Engine/Plugins/FX/Niagara/Source/Niagara/Classes/NiagaraEmitterInstance.h
+++ b/Engine/Plugins/FX/Niagara/Source/Niagara/Classes/NiagaraEmitterInstance.h
@@ -121,7 +121,7 @@ public:
 	float NIAGARA_API GetTotalCPUTimeMS();
 	int64 NIAGARA_API GetTotalBytesUsed();
 
-	ENiagaraExecutionState NIAGARA_API GetExecutionState() { return ExecutionState; }
+	ENiagaraExecutionState GetExecutionState() { return ExecutionState; }
 	void NIAGARA_API SetExecutionState(ENiagaraExecutionState InState);
 
 	bool AreBoundsDynamic() const { return bCachedBoundsDynamic; }
diff --git a/Engine/Plugins/FX/Niagara/Source/Niagara/Classes/NiagaraGPUProfilerInterface.h b/Engine/Plugins/FX/Niagara/Source/Niagara/Classes/NiagaraGPUProfilerInterface.h
--- a/Engine/Plugins/FX/Niagara/Source/Niagara/Classes/NiagaraGPUProfilerInterface.h
+++ b/Engine/Plugins/FX/Niagara/Source/Niagara/Classes/NiagaraGPUProfilerInterface.h
@@ -11,13 +11,13 @@ struct FNiagaraSimStageData;
 
 //////////////////////////////////////////////////////////////////////////
 // Public API for tracking GPU time when the profiler is enabled
-struct NIAGARA_API FNiagaraGpuProfileEvent
+struct FNiagaraGpuProfileEvent
 {
 #if WITH_NIAGARA_GPU_PROFILER
 	friend class FNiagaraGPUProfiler;
 
-	explicit FNiagaraGpuProfileEvent(const FNiagaraComputeInstanceData& InstanceData, const FNiagaraSimStageData& SimStageData, const bool bFirstInstanceData);
-	explicit FNiagaraGpuProfileEvent(const FNiagaraComputeInstanceData& InstanceData, FName CustomStageName);
+	NIAGARA_API explicit FNiagaraGpuProfileEvent(const FNiagaraComputeInstanceData& InstanceData, const FNiagaraSimStageData& SimStageData, const bool bFirstInstanceData);
+	NIAGARA_API explicit FNiagaraGpuProfileEvent(const FNiagaraComputeInstanceData& InstanceData, FName CustomStageName);
 
 private:
 	uint32									bUniqueInstance : 1;
@@ -30,11 +30,11 @@ private:
 #endif
 };
 
-struct NIAGARA_API FNiagaraGpuProfileScope
+struct FNiagaraGpuProfileScope
 {
 #if WITH_NIAGARA_GPU_PROFILER
-	explicit FNiagaraGpuProfileScope(FRHICommandList& RHICmdList, const class FNiagaraGpuComputeDispatchInterface* ComputeDispatchInterface, const FNiagaraGpuProfileEvent& Event);
-	~FNiagaraGpuProfileScope();
+	NIAGARA_API explicit FNiagaraGpuProfileScope(FRHICommandList& RHICmdList, const class FNiagaraGpuComputeDispatchInterface* ComputeDispatchInterface, const FNiagaraGpuProfileEvent& Event);
+	NIAGARA_API ~FNiagaraGpuProfileScope();
 
 private:
 	FRHICommandList& RHICmdList;
@@ -70,13 +70,13 @@ using FNiagaraGpuFrameResultsPtr = TSharedPtr<FNiagaraGpuFrameResults, ESPMode::
 
 //////////////////////////////////////////////////////////////////////////
 // Allows various systems to listen to profiler results
-struct NIAGARA_API FNiagaraGpuProfilerListener
+struct FNiagaraGpuProfilerListener
 {
-	FNiagaraGpuProfilerListener();
-	~FNiagaraGpuProfilerListener();
+	NIAGARA_API FNiagaraGpuProfilerListener();
+	NIAGARA_API ~FNiagaraGpuProfilerListener();
 
-	void SetEnabled(bool bEnabled);
-	void SetHandler(TFunction<void(const FNiagaraGpuFrameResultsPtr&)> Function);
+	NIAGARA_API void SetEnabled(bool bEnabled);
+	NIAGARA_API void SetHandler(TFunction<void(const FNiagaraGpuFrameResultsPtr&)> Function);
 	bool IsEnabled() const { return bEnabled; }
 
 private:
diff --git a/Engine/Plugins/FX/Niagara/Source/Niagara/Classes/NiagaraLensEffectBase.h b/Engine/Plugins/FX/Niagara/Source/Niagara/Classes/NiagaraLensEffectBase.h
--- a/Engine/Plugins/FX/Niagara/Source/Niagara/Classes/NiagaraLensEffectBase.h
+++ b/Engine/Plugins/FX/Niagara/Source/Niagara/Classes/NiagaraLensEffectBase.h
@@ -14,14 +14,14 @@ class APlayerCameraManager;
  * Niagara equivalent of AEmitterCameraLensEffectBase.
  * This class is supported by APlayerCameraManager (via ICameraLensEffectInterface) as a camera lens effect.
  */
-UCLASS(abstract, Blueprintable, hideCategories=(NiagaraActorActivation, "Components|Activation", Input, Collision, "Game|Damage"))
-class NIAGARA_API ANiagaraLensEffectBase : public ANiagaraActor, public ICameraLensEffectInterface
+UCLASS(abstract, Blueprintable, hideCategories=(NiagaraActorActivation, "Components|Activation", Input, Collision, "Game|Damage"), MinimalAPI)
+class ANiagaraLensEffectBase : public ANiagaraActor, public ICameraLensEffectInterface
 {
 	GENERATED_BODY()
 
 
 protected:
-	ANiagaraLensEffectBase(const FObjectInitializer& ObjectInitializer);
+	NIAGARA_API ANiagaraLensEffectBase(const FObjectInitializer& ObjectInitializer);
 
 	/**
 	 * Relative offset from the camera (where X is out from the camera)
@@ -64,19 +64,19 @@ protected:
 protected:
 
 	//~ Begin ICameraLensEffectInterface interface
-	virtual const FTransform& GetRelativeTransform() const override;
-	virtual float GetBaseFOV() const override;
-	virtual bool ShouldAllowMultipleInstances() const override;
-	virtual bool ResetWhenTriggered() const override;
-	virtual bool ShouldTreatEmitterAsSame(TSubclassOf<AActor> OtherEmitter) const override;
-	virtual void NotifyWillBePooled() override;
-	virtual void AdjustBaseFOV(float NewFOV);
-	virtual void RegisterCamera(APlayerCameraManager* CameraManager) override;
-	virtual void NotifyRetriggered() override;
-	virtual void ActivateLensEffect() override;
-	virtual void DeactivateLensEffect() override;
-	virtual bool IsLooping() const override;
+	NIAGARA_API virtual const FTransform& GetRelativeTransform() const override;
+	NIAGARA_API virtual float GetBaseFOV() const override;
+	NIAGARA_API virtual bool ShouldAllowMultipleInstances() const override;
+	NIAGARA_API virtual bool ResetWhenTriggered() const override;
+	NIAGARA_API virtual bool ShouldTreatEmitterAsSame(TSubclassOf<AActor> OtherEmitter) const override;
+	NIAGARA_API virtual void NotifyWillBePooled() override;
+	NIAGARA_API virtual void AdjustBaseFOV(float NewFOV);
+	NIAGARA_API virtual void RegisterCamera(APlayerCameraManager* CameraManager) override;
+	NIAGARA_API virtual void NotifyRetriggered() override;
+	NIAGARA_API virtual void ActivateLensEffect() override;
+	NIAGARA_API virtual void DeactivateLensEffect() override;
+	NIAGARA_API virtual bool IsLooping() const override;
 	//~ End ICameraLensEffectInterface interface
 
-	virtual void EndPlay(const EEndPlayReason::Type EndPlayReason) override;
-};
\ No newline at end of file
+	NIAGARA_API virtual void EndPlay(const EEndPlayReason::Type EndPlayReason) override;
+};
diff --git a/Engine/Plugins/FX/Niagara/Source/Niagara/Classes/NiagaraParameterCollection.h b/Engine/Plugins/FX/Niagara/Source/Niagara/Classes/NiagaraParameterCollection.h
--- a/Engine/Plugins/FX/Niagara/Source/Niagara/Classes/NiagaraParameterCollection.h
+++ b/Engine/Plugins/FX/Niagara/Source/Niagara/Classes/NiagaraParameterCollection.h
@@ -10,33 +10,33 @@
 class UMaterialParameterCollection;
 class UNiagaraParameterCollection;
 
-UCLASS()
-class NIAGARA_API UNiagaraParameterCollectionInstance : public UObject
+UCLASS(MinimalAPI)
+class UNiagaraParameterCollectionInstance : public UObject
 {
 	GENERATED_UCLASS_BODY()
 
 public:
-	virtual ~UNiagaraParameterCollectionInstance();
+	NIAGARA_API virtual ~UNiagaraParameterCollectionInstance();
 
 	//~UObject interface
-	virtual void PostLoad()override;
+	NIAGARA_API virtual void PostLoad()override;
 	//~UObject interface
 	
-	bool IsDefaultInstance()const;
+	NIAGARA_API bool IsDefaultInstance()const;
 	
-	void SetParent(UNiagaraParameterCollection* InParent);
+	NIAGARA_API void SetParent(UNiagaraParameterCollection* InParent);
 	UNiagaraParameterCollection* GetParent()const { return Collection; }
 
 	FNiagaraParameterStore& GetParameterStore() { return ParameterStorage; }
 	const FNiagaraParameterStore& GetParameterStore() const { return ParameterStorage; }
 
-	bool AddParameter(const FNiagaraVariable& Parameter);
-	bool RemoveParameter(const FNiagaraVariable& Parameter);
-	void RenameParameter(const FNiagaraVariable& Parameter, FName NewName);
-	void Empty();
-	void GetParameters(TArray<FNiagaraVariable>& OutParameters);
+	NIAGARA_API bool AddParameter(const FNiagaraVariable& Parameter);
+	NIAGARA_API bool RemoveParameter(const FNiagaraVariable& Parameter);
+	NIAGARA_API void RenameParameter(const FNiagaraVariable& Parameter, FName NewName);
+	NIAGARA_API void Empty();
+	NIAGARA_API void GetParameters(TArray<FNiagaraVariable>& OutParameters);
 
-	void Tick(UWorld* World);
+	NIAGARA_API void Tick(UWorld* World);
 
 	//TODO: Abstract to some interface to allow a hierarchy like UMaterialInstance?
 	UPROPERTY(EditAnywhere, Category=Instance)
@@ -47,17 +47,17 @@ public:
 	*/
 	UPROPERTY()
 	TArray<FNiagaraVariable> OverridenParameters;
-	bool OverridesParameter(const FNiagaraVariable& Parameter)const;
-	void SetOverridesParameter(const FNiagaraVariable& Parameter, bool bOverrides);
+	NIAGARA_API bool OverridesParameter(const FNiagaraVariable& Parameter)const;
+	NIAGARA_API void SetOverridesParameter(const FNiagaraVariable& Parameter, bool bOverrides);
 
 	/** Synchronizes this instance with any changes with it's parent collection. */
-	void SyncWithCollection();
+	NIAGARA_API void SyncWithCollection();
 
-	void Bind(UWorld* World);
+	NIAGARA_API void Bind(UWorld* World);
 
 private:
-	void RefreshSourceParameters(UWorld* World, const TArray<TPair<FName, float>>& ScalarParameters, const TArray<TPair<FName, FLinearColor>>& VectorParameters);
-	bool EnsureNotBoundToMaterialParameterCollection(FName InVariableName, FString CallingFunction) const;
+	NIAGARA_API void RefreshSourceParameters(UWorld* World, const TArray<TPair<FName, float>>& ScalarParameters, const TArray<TPair<FName, FLinearColor>>& VectorParameters);
+	NIAGARA_API bool EnsureNotBoundToMaterialParameterCollection(FName InVariableName, FString CallingFunction) const;
 
 	UPROPERTY()
 	FNiagaraParameterStore ParameterStorage;
@@ -72,80 +72,80 @@ private:
 
 	//~UObject interface
 #if WITH_EDITORONLY_DATA
-	virtual void PostEditChangeProperty(struct FPropertyChangedEvent& PropertyChangedEvent)override;
+	NIAGARA_API virtual void PostEditChangeProperty(struct FPropertyChangedEvent& PropertyChangedEvent)override;
 #endif
 	//~UObject interface
 public:
 	//Accessors from Blueprint. For now just exposing common types but ideally we can expose any somehow in future.
 	
 	UFUNCTION(BlueprintCallable, Category = Niagara, meta = (DisplayName = "Get Bool Parameter"))
-	bool GetBoolParameter(const FString& InVariableName);
+	NIAGARA_API bool GetBoolParameter(const FString& InVariableName);
 
 	UFUNCTION(BlueprintCallable, Category = Niagara, meta = (DisplayName = "Get Float Parameter"))
-	float GetFloatParameter(const FString& InVariableName);
+	NIAGARA_API float GetFloatParameter(const FString& InVariableName);
 
 	UFUNCTION(BlueprintCallable, Category = Niagara, meta = (DisplayName = "Get Int Parameter"))
-	int32 GetIntParameter(const FString& InVariableName);
+	NIAGARA_API int32 GetIntParameter(const FString& InVariableName);
 
 	UFUNCTION(BlueprintCallable, Category = Niagara, meta = (DisplayName = "Get Vector2D Parameter"))
-	FVector2D GetVector2DParameter(const FString& InVariableName);
+	NIAGARA_API FVector2D GetVector2DParameter(const FString& InVariableName);
 
 	UFUNCTION(BlueprintCallable, Category = Niagara, meta = (DisplayName = "Get Vector Parameter"))
-	FVector GetVectorParameter(const FString& InVariableName);
+	NIAGARA_API FVector GetVectorParameter(const FString& InVariableName);
 
 	UFUNCTION(BlueprintCallable, Category = Niagara, meta = (DisplayName = "Get Vector4 Parameter"))
-	FVector4 GetVector4Parameter(const FString& InVariableName);
+	NIAGARA_API FVector4 GetVector4Parameter(const FString& InVariableName);
 
 	UFUNCTION(BlueprintCallable, Category = Niagara, meta = (DisplayName = "Get Quaternion Parameter"))
-	FQuat GetQuatParameter(const FString& InVariableName);
+	NIAGARA_API FQuat GetQuatParameter(const FString& InVariableName);
 
 	UFUNCTION(BlueprintCallable, Category = Niagara, meta = (DisplayName = "Get Color Parameter"))
-	FLinearColor GetColorParameter(const FString& InVariableName);
+	NIAGARA_API FLinearColor GetColorParameter(const FString& InVariableName);
 
 	UFUNCTION(BlueprintCallable, Category = Niagara, meta = (DisplayName = "Set Bool Parameter"))
-	void SetBoolParameter(const FString& InVariableName, bool InValue);
+	NIAGARA_API void SetBoolParameter(const FString& InVariableName, bool InValue);
 
 	UFUNCTION(BlueprintCallable, Category = Niagara, meta = (DisplayName = "Set Float Parameter"))
-	void SetFloatParameter(const FString& InVariableName, float InValue);
+	NIAGARA_API void SetFloatParameter(const FString& InVariableName, float InValue);
 
 	UFUNCTION(BlueprintCallable, Category = Niagara, meta = (DisplayName = "Set Int Parameter"))
-	void SetIntParameter(const FString& InVariableName, int32 InValue);
+	NIAGARA_API void SetIntParameter(const FString& InVariableName, int32 InValue);
 
 	UFUNCTION(BlueprintCallable, Category = Niagara, meta = (DisplayName = "Set Vector2D Parameter"))
-	void SetVector2DParameter(const FString& InVariableName, FVector2D InValue);
+	NIAGARA_API void SetVector2DParameter(const FString& InVariableName, FVector2D InValue);
 
 	UFUNCTION(BlueprintCallable, Category = Niagara, meta = (DisplayName = "Set Vector Parameter"))
-	void SetVectorParameter(const FString& InVariableName, FVector InValue); // TODO[mg]: add position setter for LWC
+	NIAGARA_API void SetVectorParameter(const FString& InVariableName, FVector InValue); // TODO[mg]: add position setter for LWC
 
 	UFUNCTION(BlueprintCallable, Category = Niagara, meta = (DisplayName = "Set Vector4 Parameter"))
-	void SetVector4Parameter(const FString& InVariableName, const FVector4& InValue);
+	NIAGARA_API void SetVector4Parameter(const FString& InVariableName, const FVector4& InValue);
 
 	UFUNCTION(BlueprintCallable, Category = Niagara, meta = (DisplayName = "Set Color Parameter"))
-	void SetColorParameter(const FString& InVariableName, FLinearColor InValue);
+	NIAGARA_API void SetColorParameter(const FString& InVariableName, FLinearColor InValue);
 	
 	UFUNCTION(BlueprintCallable, Category = Niagara, meta = (DisplayName = "Set Quaternion Parameter"))
-	void SetQuatParameter(const FString& InVariableName, const FQuat& InValue);
+	NIAGARA_API void SetQuatParameter(const FString& InVariableName, const FQuat& InValue);
 };
 
 /** Asset containing a collection of global parameters usable by Niagara. */
-UCLASS()
-class NIAGARA_API UNiagaraParameterCollection : public UObject
+UCLASS(MinimalAPI)
+class UNiagaraParameterCollection : public UObject
 {
 	GENERATED_UCLASS_BODY()
 public:
 
 	//~UObject interface
 #if WITH_EDITORONLY_DATA
-	virtual void PostEditChangeProperty(struct FPropertyChangedEvent& PropertyChangedEvent)override;
+	NIAGARA_API virtual void PostEditChangeProperty(struct FPropertyChangedEvent& PropertyChangedEvent)override;
 #endif
 	//~UObject interface
 
-	int32 IndexOfParameter(const FNiagaraVariable& Var);
+	NIAGARA_API int32 IndexOfParameter(const FNiagaraVariable& Var);
 
-	int32 AddParameter(const FNiagaraVariable& Parameter);
-	int32 AddParameter(FName Name, FNiagaraTypeDefinition Type);
-	void RemoveParameter(const FNiagaraVariable& Parameter);
-	void RenameParameter(FNiagaraVariable& Parameter, FName NewName);
+	NIAGARA_API int32 AddParameter(const FNiagaraVariable& Parameter);
+	NIAGARA_API int32 AddParameter(FName Name, FNiagaraTypeDefinition Type);
+	NIAGARA_API void RemoveParameter(const FNiagaraVariable& Parameter);
+	NIAGARA_API void RenameParameter(FNiagaraVariable& Parameter, FName NewName);
 
 	TArray<FNiagaraVariable>& GetParameters() { return Parameters; }
 	const TArray<FNiagaraVariable>& GetParameters() const { return Parameters; }
@@ -162,26 +162,26 @@ public:
 	Takes the friendly name presented to the UI and converts to the real parameter name used under the hood.
 	Converts from "ParameterName" to "CollectionUniqueName_ParameterName".
 	*/
-	FString ParameterNameFromFriendlyName(const FString& FriendlyName)const;
+	NIAGARA_API FString ParameterNameFromFriendlyName(const FString& FriendlyName)const;
 	/**
 	Takes the real parameter name used under the hood and converts to the friendly name for use in the UI.
 	Converts from "CollectionUniqueName_ParameterName" to "ParameterName".
 	*/
 
-	FNiagaraVariable CollectionParameterFromFriendlyParameter(const FNiagaraVariable& FriendlyParameter)const;
-	FNiagaraVariable FriendlyParameterFromCollectionParameter(const FNiagaraVariable& CollectionParameter)const;
+	NIAGARA_API FNiagaraVariable CollectionParameterFromFriendlyParameter(const FNiagaraVariable& FriendlyParameter)const;
+	NIAGARA_API FNiagaraVariable FriendlyParameterFromCollectionParameter(const FNiagaraVariable& CollectionParameter)const;
 
-	FString FriendlyNameFromParameterName(FString ParameterName)const;
-	FString GetFullNamespace()const;
+	NIAGARA_API FString FriendlyNameFromParameterName(FString ParameterName)const;
+	NIAGARA_API FString GetFullNamespace()const;
 
 	/** The compile Id is an indicator to any compiled scripts that reference this collection that contents may have changed and a recompile is recommended to be safe.*/
-	FNiagaraCompileHash GetCompileHash() const;
+	NIAGARA_API FNiagaraCompileHash GetCompileHash() const;
 
 	/** If for any reason this data has changed externally and needs recompilation and it isn't autodetected, use this method. */
-	void RefreshCompileId();
+	NIAGARA_API void RefreshCompileId();
 
 	//~UObject interface
-	virtual void PostLoad()override;
+	NIAGARA_API virtual void PostLoad()override;
 	//~UObject interface
 
 	FName GetNamespace()const { return Namespace; }
@@ -195,8 +195,8 @@ public:
 protected:
 	
 #if WITH_EDITORONLY_DATA
-	void MakeNamespaceNameUnique();
-	void AddDefaultSourceParameters();
+	NIAGARA_API void MakeNamespaceNameUnique();
+	NIAGARA_API void AddDefaultSourceParameters();
 #endif
 
 	/** Namespace for this parameter collection. Is enforced to be unique across all parameter collections. */
@@ -216,4 +216,4 @@ protected:
 	/** Used to track whenever something of note changes in this parameter collection that might invalidate a compilation downstream of a script/emitter/system.*/
 	UPROPERTY()
 	FGuid CompileId;
-};
\ No newline at end of file
+};
diff --git a/Engine/Plugins/FX/Niagara/Source/Niagara/Classes/NiagaraParameterDefinitionsSubscriber.h b/Engine/Plugins/FX/Niagara/Source/Niagara/Classes/NiagaraParameterDefinitionsSubscriber.h
--- a/Engine/Plugins/FX/Niagara/Source/Niagara/Classes/NiagaraParameterDefinitionsSubscriber.h
+++ b/Engine/Plugins/FX/Niagara/Source/Niagara/Classes/NiagaraParameterDefinitionsSubscriber.h
@@ -39,14 +39,14 @@ public:
 };
 
 /** Interface for classes which subscribe to parameter definitions. */
-class NIAGARA_API INiagaraParameterDefinitionsSubscriber
+class INiagaraParameterDefinitionsSubscriber
 {
 #if WITH_EDITORONLY_DATA
 public:
-	virtual ~INiagaraParameterDefinitionsSubscriber();
+	NIAGARA_API virtual ~INiagaraParameterDefinitionsSubscriber();
 
-	void PostLoadDefinitionsSubscriptions();
-	void CleanupDefinitionsSubscriptions();
+	NIAGARA_API void PostLoadDefinitionsSubscriptions();
+	NIAGARA_API void CleanupDefinitionsSubscriptions();
 
 	//~ Begin Pure Virtual Methods
 	virtual const TArray<FParameterDefinitionsSubscription>& GetParameterDefinitionsSubscriptions() const = 0;
@@ -67,13 +67,13 @@ public:
 	 */
 	virtual TArray<INiagaraParameterDefinitionsSubscriber*> GetOwnedParameterDefinitionsSubscribers() { return TArray<INiagaraParameterDefinitionsSubscriber*>(); };
 
-	TArray<UNiagaraParameterDefinitionsBase*> GetSubscribedParameterDefinitions() const;
-	bool GetIsSubscribedToParameterDefinitions(const UNiagaraParameterDefinitionsBase* Definition) const;
-	UNiagaraParameterDefinitionsBase* FindSubscribedParameterDefinitionsById(const FGuid& DefinitionsId) const;
+	NIAGARA_API TArray<UNiagaraParameterDefinitionsBase*> GetSubscribedParameterDefinitions() const;
+	NIAGARA_API bool GetIsSubscribedToParameterDefinitions(const UNiagaraParameterDefinitionsBase* Definition) const;
+	NIAGARA_API UNiagaraParameterDefinitionsBase* FindSubscribedParameterDefinitionsById(const FGuid& DefinitionsId) const;
 
-	void SubscribeToParameterDefinitions(UNiagaraParameterDefinitionsBase* NewParameterDefinitions, bool bDoNotAssertIfAlreadySubscribed = false);
-	void UnsubscribeFromParameterDefinitions(const FGuid& ParameterDefinitionsToRemoveId);
-	void SynchronizeWithParameterDefinitions(const FSynchronizeWithParameterDefinitionsArgs Args = FSynchronizeWithParameterDefinitionsArgs());
+	NIAGARA_API void SubscribeToParameterDefinitions(UNiagaraParameterDefinitionsBase* NewParameterDefinitions, bool bDoNotAssertIfAlreadySubscribed = false);
+	NIAGARA_API void UnsubscribeFromParameterDefinitions(const FGuid& ParameterDefinitionsToRemoveId);
+	NIAGARA_API void SynchronizeWithParameterDefinitions(const FSynchronizeWithParameterDefinitionsArgs Args = FSynchronizeWithParameterDefinitionsArgs());
 
 	FOnSubscribedParameterDefinitionsChanged& GetOnSubscribedParameterDefinitionsChangedDelegate() { return OnSubscribedParameterDefinitionsChangedDelegate; };
 
@@ -84,6 +84,6 @@ private:
 
 private:
 	/** Update the cached change Id of specified Synchronized Subscribed Parameter Definition subscriptions so that they are not marked pending sync. */
-	void MarkParameterDefinitionSubscriptionsSynchronized(TArray<FGuid> SynchronizedParameterDefinitionsIds /*= TArray<FGuid>()*/);
+	NIAGARA_API void MarkParameterDefinitionSubscriptionsSynchronized(TArray<FGuid> SynchronizedParameterDefinitionsIds /*= TArray<FGuid>()*/);
 #endif
 };
diff --git a/Engine/Plugins/FX/Niagara/Source/Niagara/Classes/NiagaraPerfBaseline.h b/Engine/Plugins/FX/Niagara/Source/Niagara/Classes/NiagaraPerfBaseline.h
--- a/Engine/Plugins/FX/Niagara/Source/Niagara/Classes/NiagaraPerfBaseline.h
+++ b/Engine/Plugins/FX/Niagara/Source/Niagara/Classes/NiagaraPerfBaseline.h
@@ -16,7 +16,7 @@ class UTextRenderComponent;
 class FCommonViewportClient;
 
 USTRUCT(BlueprintType)
-struct NIAGARA_API FNiagaraPerfBaselineStats
+struct FNiagaraPerfBaselineStats
 {
 	GENERATED_BODY()
 
@@ -51,25 +51,25 @@ struct NIAGARA_API FNiagaraPerfBaselineStats
 	@param Stats	Accumulated stats we take values from.
 	@param bSyncRT	Whether to sync with the RT. If true we flush the RT and can get the current stat values. If false we get the current GT stats but have to just use the RT stats from the last completed frame.
 	*/
-	FNiagaraPerfBaselineStats(FAccumulatedParticlePerfStats& Stats, bool bSyncRT);
+	NIAGARA_API FNiagaraPerfBaselineStats(FAccumulatedParticlePerfStats& Stats, bool bSyncRT);
 
 	FORCEINLINE bool IsValid()const { return PerInstanceAvg_GT > 0.0f && PerInstanceAvg_RT > 0.0f && PerInstanceMax_GT > 0.0f && PerInstanceMax_RT > 0.0f; }
 	/**
 	Compares a system's performance stats with that of a baseline.
 	*/
-	static EComparisonResult Compare(const FNiagaraPerfBaselineStats& Stat, const FNiagaraPerfBaselineStats& Baseline, FNiagaraPerfBaselineStats& OutRatio);
-	static EComparisonResult Compare(const FNiagaraPerfBaselineStats& Stat, const FNiagaraPerfBaselineStats& Baseline, FNiagaraPerfBaselineStats& OutRatio, EComparisonResult& OutGTAvgResult, EComparisonResult& OutGTMaxResult, EComparisonResult& OutRTAvgResult, EComparisonResult& OutRTMaxResult);
-	static EComparisonResult Compare(float Stat, float Baseline, float& OutRatio);	
-	static FLinearColor GetComparisonResultColor(EComparisonResult Result);
-	static FText GetComparisonResultText(EComparisonResult Result);
+	static NIAGARA_API EComparisonResult Compare(const FNiagaraPerfBaselineStats& Stat, const FNiagaraPerfBaselineStats& Baseline, FNiagaraPerfBaselineStats& OutRatio);
+	static NIAGARA_API EComparisonResult Compare(const FNiagaraPerfBaselineStats& Stat, const FNiagaraPerfBaselineStats& Baseline, FNiagaraPerfBaselineStats& OutRatio, EComparisonResult& OutGTAvgResult, EComparisonResult& OutGTMaxResult, EComparisonResult& OutRTAvgResult, EComparisonResult& OutRTMaxResult);
+	static NIAGARA_API EComparisonResult Compare(float Stat, float Baseline, float& OutRatio);	
+	static NIAGARA_API FLinearColor GetComparisonResultColor(EComparisonResult Result);
+	static NIAGARA_API FText GetComparisonResultText(EComparisonResult Result);
 #endif
 };
 
 /**
 * Base class for baseline controllers. These can are responsible for spawning and manipulating the FX needed for the baseline perf tests.
 */
-UCLASS(abstract, EditInlineNew, BlueprintType, Blueprintable)
-class NIAGARA_API UNiagaraBaselineController : public UObject
+UCLASS(abstract, EditInlineNew, BlueprintType, Blueprintable, MinimalAPI)
+class UNiagaraBaselineController : public UObject
 {
 	GENERATED_BODY()
 	
@@ -77,23 +77,23 @@ public:
 
 	/** Called from the stats system when we begin gathering stats for the given System asset.*/
 	UFUNCTION(BlueprintNativeEvent)
-	void OnBeginTest();
+	NIAGARA_API void OnBeginTest();
 
 	/** Returns whether the baseline test is complete. */
 	UFUNCTION(BlueprintNativeEvent)
-	bool OnTickTest();
+	NIAGARA_API bool OnTickTest();
 
 	/** Called from the stats system on completion of the test with the final stats for the given system asset. */
 	UFUNCTION(BlueprintNativeEvent)
-	void OnEndTest(FNiagaraPerfBaselineStats Stats);
+	NIAGARA_API void OnEndTest(FNiagaraPerfBaselineStats Stats);
 
 	/** Called when the owning actor is ticked. */
 	UFUNCTION(BlueprintNativeEvent)
-	void OnOwnerTick(float DeltaTime);
+	NIAGARA_API void OnOwnerTick(float DeltaTime);
 
 	/** Returns the System for this baseline. Will synchronously load the system if needed. */
 	UFUNCTION(BlueprintCallable, Category="Baseline")
-	UNiagaraSystem* GetSystem();
+	NIAGARA_API UNiagaraSystem* GetSystem();
 	
 	/** Duration to gather performance stats for the given system. */
 	UPROPERTY(EditAnywhere, Category = "Baseline", BlueprintReadWrite)
@@ -117,15 +117,15 @@ private:
 Simple controller that will just spawn the given system N times. If any instance completes, it will spawn a new one to replace it.
 Can handle simple burst or looping systems.
 */
-UCLASS(EditInlineNew)
-class NIAGARA_API UNiagaraBaselineController_Basic : public UNiagaraBaselineController
+UCLASS(EditInlineNew, MinimalAPI)
+class UNiagaraBaselineController_Basic : public UNiagaraBaselineController
 {
 	GENERATED_BODY()
 	
-	virtual void OnBeginTest_Implementation();
-	virtual bool OnTickTest_Implementation();
-	virtual void OnEndTest_Implementation(FNiagaraPerfBaselineStats Stats);
-	virtual void OnOwnerTick_Implementation(float DeltaTime);
+	NIAGARA_API virtual void OnBeginTest_Implementation();
+	NIAGARA_API virtual bool OnTickTest_Implementation();
+	NIAGARA_API virtual void OnEndTest_Implementation(FNiagaraPerfBaselineStats Stats);
+	NIAGARA_API virtual void OnOwnerTick_Implementation(float DeltaTime);
 
 	UPROPERTY(EditAnywhere, Category = "Baseline")
 	int32 NumInstances = 1;
@@ -135,8 +135,8 @@ class NIAGARA_API UNiagaraBaselineController_Basic : public UNiagaraBaselineCont
 };
 
 /** Actor that controls how the baseline system behaves and also controls stats gathering for. */
-UCLASS()
-class NIAGARA_API ANiagaraPerfBaselineActor : public AActor
+UCLASS(MinimalAPI)
+class ANiagaraPerfBaselineActor : public AActor
 {
 	GENERATED_UCLASS_BODY()
 
@@ -152,8 +152,8 @@ public:
 #if NIAGARA_PERF_BASELINES
 
 	//AActor Interface
-	virtual void BeginPlay();
-	virtual void TickActor(float DeltaTime, enum ELevelTick TickType, FActorTickFunction& ThisTickFunction);
+	NIAGARA_API virtual void BeginPlay();
+	NIAGARA_API virtual void TickActor(float DeltaTime, enum ELevelTick TickType, FActorTickFunction& ThisTickFunction);
 	//AActor Interface END
 #endif
 };
@@ -164,14 +164,14 @@ public:
 Listener that accumulates a short run of stats for a particular baseline system.
 These are then stored for comparison against WIP systems in the editor to give a baseline for performance.
 */
-class NIAGARA_API FNiagaraPerfBaselineStatsListener : public FParticlePerfStatsListener
+class FNiagaraPerfBaselineStatsListener : public FParticlePerfStatsListener
 {
 public:
-	FNiagaraPerfBaselineStatsListener(UNiagaraBaselineController* OwnerBaseline);
-	virtual void Begin()override;
-	virtual void End()override;
-	virtual bool Tick()override;
-	virtual void TickRT()override;
+	NIAGARA_API FNiagaraPerfBaselineStatsListener(UNiagaraBaselineController* OwnerBaseline);
+	NIAGARA_API virtual void Begin()override;
+	NIAGARA_API virtual void End()override;
+	NIAGARA_API virtual bool Tick()override;
+	NIAGARA_API virtual void TickRT()override;
 
 	virtual bool NeedsWorldStats()const override { return false; }
 	virtual bool NeedsSystemStats()const override { return true; }
@@ -189,16 +189,16 @@ private:
 /**
 This listener gathers perf stats for all systems and will render them to the editor viewport with a comparison to their perf baselines.
 */
-class NIAGARA_API FParticlePerfStatsListener_NiagaraBaselineComparisonRender : public FParticlePerfStatsListener_GatherAll
+class FParticlePerfStatsListener_NiagaraBaselineComparisonRender : public FParticlePerfStatsListener_GatherAll
 {
 public:
-	FParticlePerfStatsListener_NiagaraBaselineComparisonRender();
-	virtual bool Tick()override;
-	virtual void TickRT()override;
+	NIAGARA_API FParticlePerfStatsListener_NiagaraBaselineComparisonRender();
+	NIAGARA_API virtual bool Tick()override;
+	NIAGARA_API virtual void TickRT()override;
 	virtual bool NeedsWorldStats()const { return false; }
 	virtual bool NeedsSystemStats()const { return true; }
 
-	int32 RenderStats(UWorld* World, class FViewport* Viewport, class FCanvas* Canvas, int32 X, int32 Y, const FVector* ViewLocation, const FRotator* ViewRotation);
+	NIAGARA_API int32 RenderStats(UWorld* World, class FViewport* Viewport, class FCanvas* Canvas, int32 X, int32 Y, const FVector* ViewLocation, const FRotator* ViewRotation);
 
 	int32 NumFrames = 0;
 	int32 NumFramesRT = 0;
@@ -211,19 +211,19 @@ public:
 /**
 This listener gathers perf stats for all systems and will give less intrusive reports on systems that exceed their baseline cost.
 */
-class NIAGARA_API FParticlePerfStatsListener_NiagaraPerformanceReporter : public FParticlePerfStatsListener_GatherAll
+class FParticlePerfStatsListener_NiagaraPerformanceReporter : public FParticlePerfStatsListener_GatherAll
 {
 public:
-	FParticlePerfStatsListener_NiagaraPerformanceReporter(UWorld* InWorld);
+	NIAGARA_API FParticlePerfStatsListener_NiagaraPerformanceReporter(UWorld* InWorld);
 
-	virtual bool Tick()override;
-	virtual void TickRT()override;
+	NIAGARA_API virtual bool Tick()override;
+	NIAGARA_API virtual void TickRT()override;
 	virtual bool NeedsWorldStats()const override { return false; }
 	virtual bool NeedsSystemStats()const override { return true; }
 
-	void HandleTestResults();
-	void ReportToScreen();
-	void ReportToLog();
+	NIAGARA_API void HandleTestResults();
+	NIAGARA_API void ReportToScreen();
+	NIAGARA_API void ReportToLog();
 
 	TWeakObjectPtr<UWorld> World;
 
@@ -262,7 +262,7 @@ public:
 	FRenderCommandFence ResultsFence;
 	bool bResultsTrigger = false;
 
-	static const int32 TestDebugMessageID;
+	static NIAGARA_API const int32 TestDebugMessageID;
 };
 
 // Helper class for managing the generation and tracking of Niagara Performance Baselines.
@@ -300,4 +300,4 @@ private:
 	TWeakObjectPtr<UWorld> BaselineGenWorld;
 };
 
-#endif
\ No newline at end of file
+#endif
diff --git a/Engine/Plugins/FX/Niagara/Source/Niagara/Classes/NiagaraPlatformSet.h b/Engine/Plugins/FX/Niagara/Source/Niagara/Classes/NiagaraPlatformSet.h
--- a/Engine/Plugins/FX/Niagara/Source/Niagara/Classes/NiagaraPlatformSet.h
+++ b/Engine/Plugins/FX/Niagara/Source/Niagara/Classes/NiagaraPlatformSet.h
@@ -185,19 +185,19 @@ enum class ENiagaraCVarConditionResponse : uint8
 
 /** Imposes a condition that a CVar must contain a set value or range of values for a platform set to be enabled. */
 USTRUCT()
-struct NIAGARA_API FNiagaraPlatformSetCVarCondition
+struct FNiagaraPlatformSetCVarCondition
 {
 	GENERATED_BODY();
 	
-	FNiagaraPlatformSetCVarCondition();
+	NIAGARA_API FNiagaraPlatformSetCVarCondition();
 
 	/** Return true if this condition is met for the given device profile. */
-	bool Evaluate(const UDeviceProfile* DeviceProfile, int32 QualityLevel, bool bCheckCurrentStateOnly, bool& bOutCanEverPass, bool& bOutCanEverFail)const;
+	NIAGARA_API bool Evaluate(const UDeviceProfile* DeviceProfile, int32 QualityLevel, bool bCheckCurrentStateOnly, bool& bOutCanEverPass, bool& bOutCanEverFail)const;
 
 	/** Returns the CVar for this condition. Can be null if the name give is not a valid CVar or the CVar is removed. */
-	IConsoleVariable* GetCVar()const;
+	NIAGARA_API IConsoleVariable* GetCVar()const;
 
-	void SetCVar(FName InCVarName);
+	NIAGARA_API void SetCVar(FName InCVarName);
 
 	/** The name of the CVar we're testing the value of. */
 	UPROPERTY(EditAnywhere, Category = "CVar")
@@ -305,76 +305,76 @@ struct FNiagaraPlatformSetRedirect
 };
 
 USTRUCT()
-struct NIAGARA_API FNiagaraPlatformSet
+struct FNiagaraPlatformSet
 {
 	GENERATED_BODY();
 
-	static FORCEINLINE int32 QualityLevelFromMask(int32 QLMask);
-	static FORCEINLINE int32 CreateQualityLevelMask(int32 QL);
-	static FText GetQualityLevelText(int32 QualityLevel);
-	static FText GetQualityLevelMaskText(int32 QualityLevelMask);
+	static NIAGARA_API FORCEINLINE int32 QualityLevelFromMask(int32 QLMask);
+	static NIAGARA_API FORCEINLINE int32 CreateQualityLevelMask(int32 QL);
+	static NIAGARA_API FText GetQualityLevelText(int32 QualityLevel);
+	static NIAGARA_API FText GetQualityLevelMaskText(int32 QualityLevelMask);
 
-	static int32 GetQualityLevel();
-	static int32 GetMinQualityLevel();
-	static int32 GetMaxQualityLevel();
-	static int32 GetAvailableQualityLevelMask();
-	static int32 GetFullQualityLevelMask(int32 NumQualityLevels);
+	static NIAGARA_API int32 GetQualityLevel();
+	static NIAGARA_API int32 GetMinQualityLevel();
+	static NIAGARA_API int32 GetMaxQualityLevel();
+	static NIAGARA_API int32 GetAvailableQualityLevelMask();
+	static NIAGARA_API int32 GetFullQualityLevelMask(int32 NumQualityLevels);
 	
 public:
 
 	//Runtime public API
 
-	FNiagaraPlatformSet(int32 QLMask);
-	FNiagaraPlatformSet();
+	NIAGARA_API FNiagaraPlatformSet(int32 QLMask);
+	NIAGARA_API FNiagaraPlatformSet();
 
-	bool operator==(const FNiagaraPlatformSet& Other)const;
+	NIAGARA_API bool operator==(const FNiagaraPlatformSet& Other)const;
 
 	/** Is this set active right now. i.e. enabled for the current device profile and quality level. */
-	bool IsActive()const;
+	NIAGARA_API bool IsActive()const;
 	
 	/** Is this platform set enabled for the give device profile. */
-	int32 IsEnabledForDeviceProfile(const UDeviceProfile* DeviceProfile)const;
+	NIAGARA_API int32 IsEnabledForDeviceProfile(const UDeviceProfile* DeviceProfile)const;
 
 	/** Is this platform set enabled on any quality level for the passed device profile. Returns the QualityLevelMask for all enabled effects qualities for this profile. */
-	int32 GetEnabledMaskForDeviceProfile(const UDeviceProfile* DeviceProfile)const;
+	NIAGARA_API int32 GetEnabledMaskForDeviceProfile(const UDeviceProfile* DeviceProfile)const;
 
 	/** Is this platform set enabled at this quality level on any device profile. */
-	bool IsEnabledForQualityLevel(int32 QualityLevel)const;
+	NIAGARA_API bool IsEnabledForQualityLevel(int32 QualityLevel)const;
 
 	/** Fill OutProfiles with all device profiles that have been overridden at the passed QualityLevel. */
-	void GetOverridenDeviceProfiles(int32 QualityLevel, TArray<UDeviceProfile*>& OutEnabledProfiles, TArray<UDeviceProfile*>& OutDisabledProfiles)const;
+	NIAGARA_API void GetOverridenDeviceProfiles(int32 QualityLevel, TArray<UDeviceProfile*>& OutEnabledProfiles, TArray<UDeviceProfile*>& OutDisabledProfiles)const;
 
 	/** Returns true if this set is enabled for any profiles on the specified platform. */
-	bool IsEnabledForPlatform(const FString& PlatformName)const;
+	NIAGARA_API bool IsEnabledForPlatform(const FString& PlatformName)const;
 
 	/** Returns true if the current platform can modify it's Niagara scalability settings at runtime. */
-	static bool CanChangeScalabilityAtRuntime();
+	static NIAGARA_API bool CanChangeScalabilityAtRuntime();
 
 	/**Will force all platform sets to regenerate their cached data next time they are used.*/
-	static void InvalidateCachedData();
+	static NIAGARA_API void InvalidateCachedData();
 
 	/** Returns the currenlty(or default) active quality mask for a profile. */
-	static int32 GetActiveQualityMaskForDeviceProfile(const UDeviceProfile* Profile);
+	static NIAGARA_API int32 GetActiveQualityMaskForDeviceProfile(const UDeviceProfile* Profile);
 
 	/** Returns the mask of all available quality levels for a device profile. */
-	static int32 GetAvailableQualityMaskForDeviceProfile(const UDeviceProfile* Profile);
+	static NIAGARA_API int32 GetAvailableQualityMaskForDeviceProfile(const UDeviceProfile* Profile);
 
 	/** Returns true if the passed platform should prune emitters on cook. */
-	static bool ShouldPruneEmittersOnCook(const FString& PlatformName);
+	static NIAGARA_API bool ShouldPruneEmittersOnCook(const FString& PlatformName);
 
 	/** Returns true if the passed platform can modify it's niagara scalability settings at runtime. */
-	static bool CanChangeScalabilityAtRuntime(const UDeviceProfile* DeviceProfile);
+	static NIAGARA_API bool CanChangeScalabilityAtRuntime(const UDeviceProfile* DeviceProfile);
 
-	FNiagaraPlatformSetEnabledState IsEnabled(const UDeviceProfile* Profile, int32 QualityLevel, bool bCheckCurrentStateOnly, FNiagaraPlatformSetEnabledStateDetails* OutDetails=nullptr)const;
+	NIAGARA_API FNiagaraPlatformSetEnabledState IsEnabled(const UDeviceProfile* Profile, int32 QualityLevel, bool bCheckCurrentStateOnly, FNiagaraPlatformSetEnabledStateDetails* OutDetails=nullptr)const;
 
-	bool CanConsiderDeviceProfile(const UDeviceProfile* Profile)const;
+	NIAGARA_API bool CanConsiderDeviceProfile(const UDeviceProfile* Profile)const;
 
-	bool ApplyRedirects();
+	NIAGARA_API bool ApplyRedirects();
 
-	static void RefreshScalability();
-	static void OnCVarChanged(IConsoleVariable* CVar);
-	static void OnCVarUnregistered(IConsoleVariable* CVar);
-	static IConsoleVariable* GetCVar(FName CVarName);
+	static NIAGARA_API void RefreshScalability();
+	static NIAGARA_API void OnCVarChanged(IConsoleVariable* CVar);
+	static NIAGARA_API void OnCVarUnregistered(IConsoleVariable* CVar);
+	static NIAGARA_API IConsoleVariable* GetCVar(FName CVarName);
 
 	//Editor only public API
 #if WITH_EDITOR
@@ -385,16 +385,16 @@ public:
 	//bool Conflicts(const FNiagaraPlatformSet& Other, TArray<const UDeviceProfile*>& OutConflictingProfiles, bool bIncludeProfilesWithVariableQualityLevel=false)const;
 	
 	/** Direct state accessors for the UI. */
-	bool IsEffectQualityEnabled(int32 EffectQuality)const;
-	void SetEnabledForEffectQuality(int32 EffectQuality, bool bEnabled);
-	void SetDeviceProfileState(UDeviceProfile* Profile, int32 QualityLevel, ENiagaraPlatformSelectionState NewState);
-	ENiagaraPlatformSelectionState GetDeviceProfileState(UDeviceProfile* Profile, int32 QualityLevel)const;
+	NIAGARA_API bool IsEffectQualityEnabled(int32 EffectQuality)const;
+	NIAGARA_API void SetEnabledForEffectQuality(int32 EffectQuality, bool bEnabled);
+	NIAGARA_API void SetDeviceProfileState(UDeviceProfile* Profile, int32 QualityLevel, ENiagaraPlatformSelectionState NewState);
+	NIAGARA_API ENiagaraPlatformSelectionState GetDeviceProfileState(UDeviceProfile* Profile, int32 QualityLevel)const;
 	
 	/** Invalidates any cached data on this platform set when something has changed. */
-	void OnChanged();
+	NIAGARA_API void OnChanged();
 
 	/** Inspects the passed sets and generates an array of all conflicts between these sets. Used to keep arrays of platform sets orthogonal. */
-	static bool GatherConflicts(TConstArrayView<const FNiagaraPlatformSet*> PlatformSets, TArray<FNiagaraPlatformSetConflictInfo>& OutConflicts);
+	static NIAGARA_API bool GatherConflicts(TConstArrayView<const FNiagaraPlatformSet*> PlatformSets, TArray<FNiagaraPlatformSetConflictInfo>& OutConflicts);
 	
 	DECLARE_DELEGATE_RetVal(int32, FOnOverrideQualityLevel);
 	FOnOverrideQualityLevel OnOverrideQualityLevelDelegate;
@@ -429,15 +429,15 @@ public:
 
 		TArray<int32> QualityLevelsPerEffectsQuality;
 	};
-	static FPlatformIniSettings& GetPlatformIniSettings(const FString& PlatformName);
+	static NIAGARA_API FPlatformIniSettings& GetPlatformIniSettings(const FString& PlatformName);
 
-	static int32 GetEffectQualityMaskForPlatform(const FString& PlatformName);
+	static NIAGARA_API int32 GetEffectQualityMaskForPlatform(const FString& PlatformName);
 #endif
 
 	static uint32 GetLastDirtiedFrame(){ return LastDirtiedFrame; }
 
-	static void SetNiagaraQualityLevelOverride(int32 Override);
-	static void ClearNiagaraQualityLevelOverride();
+	static NIAGARA_API void SetNiagaraQualityLevelOverride(int32 Override);
+	static NIAGARA_API void ClearNiagaraQualityLevelOverride();
 
 private:
 
@@ -447,21 +447,21 @@ private:
 	mutable uint32 LastBuiltFrame;
 
 	//Set from outside when we need to force all cached values to be regenerated. For example on CVar changes.
-	static uint32 LastDirtiedFrame;
+	static NIAGARA_API uint32 LastDirtiedFrame;
 
 	/** Cached value of the CVar fx.Niagara.QualityLevel. Allows us to detect when it has changed so we can properly apply scalability changes. */
-	static int32 CachedQualityLevel;
-	static int32 CachedAvailableQualityLevelMask;
+	static NIAGARA_API int32 CachedQualityLevel;
+	static NIAGARA_API int32 CachedAvailableQualityLevelMask;
 
 	/** An override for Niagara quality level, driven from the Niagara UI. */
-	static int32 QualityLevelOverride;
+	static NIAGARA_API int32 QualityLevelOverride;
 
 #if WITH_EDITOR
 	//Cached data read from platform ini files.
-	static TMap<FName, FPlatformIniSettings> CachedPlatformIniSettings;
+	static NIAGARA_API TMap<FName, FPlatformIniSettings> CachedPlatformIniSettings;
 
 	//Cached final QualityLevel setting for each device profile.
-	static TMap<const UDeviceProfile*, int32> CachedQLMasksPerDeviceProfile;
+	static NIAGARA_API TMap<const UDeviceProfile*, int32> CachedQLMasksPerDeviceProfile;
 #endif
 
 	/**
@@ -472,10 +472,10 @@ private:
 		IConsoleVariable* CVar = nullptr;
 		FDelegateHandle ChangedHandle;
 	};
-	static TMap<FName, FCachedCVarInfo> CachedCVarInfo;
-	static FCriticalSection CachedCVarInfoCritSec;
+	static NIAGARA_API TMap<FName, FCachedCVarInfo> CachedCVarInfo;
+	static NIAGARA_API FCriticalSection CachedCVarInfoCritSec;
 
-	static TArray<FName> ChangedCVars;
+	static NIAGARA_API TArray<FName> ChangedCVars;
 };
 
 
diff --git a/Engine/Plugins/FX/Niagara/Source/Niagara/Classes/NiagaraScratchPadContainer.h b/Engine/Plugins/FX/Niagara/Source/Niagara/Classes/NiagaraScratchPadContainer.h
--- a/Engine/Plugins/FX/Niagara/Source/Niagara/Classes/NiagaraScratchPadContainer.h
+++ b/Engine/Plugins/FX/Niagara/Source/Niagara/Classes/NiagaraScratchPadContainer.h
@@ -10,22 +10,22 @@
 /**
  * Wrapper class for scratch pad scripts in an emitter. This is needed because each emitter version can have it's own copy of a scratch pad with the same name and they can't all be outered to the emitter that way.
  */
-UCLASS()
-class NIAGARA_API UNiagaraScratchPadContainer : public UObject
+UCLASS(MinimalAPI)
+class UNiagaraScratchPadContainer : public UObject
 {
 	GENERATED_BODY()
 	
 public:
-	virtual void PostLoad() override;
-	void CheckConsistency();
-	void SetScripts(const TArray<TObjectPtr<UNiagaraScript>>& InScripts);
-	void AppendScripts(const TArray<TObjectPtr<UNiagaraScript>>& InScripts);
-	void AppendScripts(TObjectPtr<UNiagaraScratchPadContainer> InScripts);
+	NIAGARA_API virtual void PostLoad() override;
+	NIAGARA_API void CheckConsistency();
+	NIAGARA_API void SetScripts(const TArray<TObjectPtr<UNiagaraScript>>& InScripts);
+	NIAGARA_API void AppendScripts(const TArray<TObjectPtr<UNiagaraScript>>& InScripts);
+	NIAGARA_API void AppendScripts(TObjectPtr<UNiagaraScratchPadContainer> InScripts);
 	
 	// The scripts get added in specific ordering through inheritance. These accessors just allow us to deal with checks 
 	// for known scripts 
-	int32 FindIndexForScript(UNiagaraScript* InScript) const;
-	UNiagaraScript* FindScriptAtIndex(int32 Index) const;
+	NIAGARA_API int32 FindIndexForScript(UNiagaraScript* InScript) const;
+	NIAGARA_API UNiagaraScript* FindScriptAtIndex(int32 Index) const;
 
 #if WITH_EDITORONLY_DATA
 
diff --git a/Engine/Plugins/FX/Niagara/Source/Niagara/Classes/NiagaraScript.h b/Engine/Plugins/FX/Niagara/Source/Niagara/Classes/NiagaraScript.h
--- a/Engine/Plugins/FX/Niagara/Source/Niagara/Classes/NiagaraScript.h
+++ b/Engine/Plugins/FX/Niagara/Source/Niagara/Classes/NiagaraScript.h
@@ -179,7 +179,7 @@ private:
 };
 
 USTRUCT()
-struct NIAGARA_API FNiagaraCompilerTag
+struct FNiagaraCompilerTag
 {
 	GENERATED_USTRUCT_BODY()
 public:
@@ -191,8 +191,8 @@ public:
 	UPROPERTY()
 	FString StringValue;
 
-	static FNiagaraCompilerTag* FindTag(TArray< FNiagaraCompilerTag>& InTags,  const FNiagaraVariableBase& InSearchVar);
-	static const FNiagaraCompilerTag* FindTag(const TArray< FNiagaraCompilerTag>& InTags, const FNiagaraVariableBase& InSearchVar);
+	static NIAGARA_API FNiagaraCompilerTag* FindTag(TArray< FNiagaraCompilerTag>& InTags,  const FNiagaraVariableBase& InSearchVar);
+	static NIAGARA_API const FNiagaraCompilerTag* FindTag(const TArray< FNiagaraCompilerTag>& InTags, const FNiagaraVariableBase& InSearchVar);
 
 };
 
@@ -221,7 +221,7 @@ struct FNiagaraScriptDebuggerInfo
 
 /** Struct containing all of the data necessary to look up a NiagaraScript's VM executable results from the Derived Data Cache.*/
 USTRUCT()
-struct NIAGARA_API FNiagaraVMExecutableDataId
+struct FNiagaraVMExecutableDataId
 {
 	GENERATED_USTRUCT_BODY()
 public:
@@ -245,7 +245,7 @@ public:
 	UPROPERTY()
 	TArray<FNiagaraVariableBase> AdditionalVariables;
 
-	TArray<FString> GetAdditionalVariableStrings();
+	NIAGARA_API TArray<FString> GetAdditionalVariableStrings();
 #endif
 
 	/** Whether or not we need to bake Rapid Iteration params. True to keep params, false to bake.*/
@@ -303,8 +303,8 @@ public:
 	~FNiagaraVMExecutableDataId()
 	{ }
 
-	bool IsValid() const;
-	void Invalidate();
+	NIAGARA_API bool IsValid() const;
+	NIAGARA_API void Invalidate();
 	
 	friend uint32 GetTypeHash(const FNiagaraVMExecutableDataId& Ref)
 	{
@@ -316,11 +316,11 @@ public:
 		return sizeof(*this);
 	}
 
-	bool HasInterpolatedParameters() const;
-	bool RequiresPersistentIDs() const;
+	NIAGARA_API bool HasInterpolatedParameters() const;
+	NIAGARA_API bool RequiresPersistentIDs() const;
 #if 0
 	/** Hashes the script-specific part of this shader map Id. */
-	void GetScriptHash(FSHAHash& OutHash) const;
+	NIAGARA_API void GetScriptHash(FSHAHash& OutHash) const;
 #endif
 
 	/**
@@ -329,7 +329,7 @@ public:
 	* @param ReferenceSet	The set to compare against
 	* @return				true if the sets are equal
 	*/
-	bool operator==(const FNiagaraVMExecutableDataId& ReferenceSet) const;
+	NIAGARA_API bool operator==(const FNiagaraVMExecutableDataId& ReferenceSet) const;
 
 	bool operator!=(const FNiagaraVMExecutableDataId& ReferenceSet) const
 	{
@@ -338,13 +338,13 @@ public:
 
 #if WITH_EDITORONLY_DATA
 	/** Appends string representations of this Id to a key string. */
-	void AppendKeyString(FString& KeyString, const FString& Delimiter = TEXT("_"), bool bAppendObjectForDebugging = false) const;
+	NIAGARA_API void AppendKeyString(FString& KeyString, const FString& Delimiter = TEXT("_"), bool bAppendObjectForDebugging = false) const;
 
 #endif
 };
 
 USTRUCT()
-struct NIAGARA_API FNiagaraVMExecutableByteCode
+struct FNiagaraVMExecutableByteCode
 {
 	GENERATED_USTRUCT_BODY()
 private:
@@ -355,15 +355,15 @@ private:
 	int32 UncompressedSize = INDEX_NONE;
 
 public:
-	bool HasByteCode() const;
-	bool IsCompressed() const;
-	bool Compress();
-	bool Uncompress();
+	NIAGARA_API bool HasByteCode() const;
+	NIAGARA_API bool IsCompressed() const;
+	NIAGARA_API bool Compress();
+	NIAGARA_API bool Uncompress();
 
-	void Reset();
+	NIAGARA_API void Reset();
 
-	void SetData(const TArray<uint8>& InData);
-	void SetData(TArray<uint8>&& InData);
+	NIAGARA_API void SetData(const TArray<uint8>& InData);
+	NIAGARA_API void SetData(TArray<uint8>&& InData);
 
 	const TArray<uint8>& GetData() const { return Data; }
 	TArray<uint8>& GetData() { return Data; }
@@ -374,7 +374,7 @@ public:
 	void Shrink() { Data.Shrink(); }
 
 	/** Used to upgrade a serialized property to our own struct */
-	bool SerializeFromMismatchedTag(const struct FPropertyTag& Tag, FStructuredArchive::FSlot Slot);
+	NIAGARA_API bool SerializeFromMismatchedTag(const struct FPropertyTag& Tag, FStructuredArchive::FSlot Slot);
 };
 
 template<>
@@ -387,7 +387,7 @@ struct TStructOpsTypeTraits<FNiagaraVMExecutableByteCode> : public TStructOpsTyp
 };
 
 // Carrier for uncompressed, and optimmized bytecode returned from optimization task
-struct NIAGARA_API FNiagaraScriptAsyncOptimizeTaskState
+struct FNiagaraScriptAsyncOptimizeTaskState
 {
 	FNiagaraVMExecutableDataId CachedScriptVMId;
 	FNiagaraVMExecutableByteCode SourceByteCode;
@@ -399,18 +399,18 @@ struct NIAGARA_API FNiagaraScriptAsyncOptimizeTaskState
 	bool bShouldFreeSourceByteCodeOnCooked = false;
 	bool bOptimizationComplete = false;
 		
-	void OptimizeByteCode();
+	NIAGARA_API void OptimizeByteCode();
 };
 using FNiagaraScriptAsyncOptimizeTaskStatePtr = TSharedPtr<FNiagaraScriptAsyncOptimizeTaskState, ESPMode::ThreadSafe>;
 
 
 /** Struct containing all of the data needed to run a Niagara VM executable script.*/
 USTRUCT()
-struct NIAGARA_API FNiagaraVMExecutableData
+struct FNiagaraVMExecutableData
 {
 	GENERATED_USTRUCT_BODY()
 public:
-	FNiagaraVMExecutableData();
+	NIAGARA_API FNiagaraVMExecutableData();
 
 	/** Byte code to execute for this system. */
 	UPROPERTY()
@@ -568,26 +568,26 @@ public:
 	UPROPERTY()
 	uint32 bNeedsGPUContextInit : 1;
 
-	void ApplyFinishedOptimization(const FNiagaraVMExecutableDataId& CachedScriptVMId, const FNiagaraScriptAsyncOptimizeTaskStatePtr& Result);
+	NIAGARA_API void ApplyFinishedOptimization(const FNiagaraVMExecutableDataId& CachedScriptVMId, const FNiagaraScriptAsyncOptimizeTaskStatePtr& Result);
 
-	void SerializeData(FArchive& Ar, bool bDDCData);
+	NIAGARA_API void SerializeData(FArchive& Ar, bool bDDCData);
 	
-	bool IsValid() const;
+	NIAGARA_API bool IsValid() const;
 
-	void Reset();
+	NIAGARA_API void Reset();
 
 #if WITH_EDITORONLY_DATA
-	void BakeScriptLiterals(TArray<uint8>& OutLiterals) const;
+	NIAGARA_API void BakeScriptLiterals(TArray<uint8>& OutLiterals) const;
 #endif
 
-	bool HasByteCode() const;
+	NIAGARA_API bool HasByteCode() const;
 
 #if VECTORVM_SUPPORTS_EXPERIMENTAL
-	bool SupportsExperimentalVM() const;
-	FVectorVMOptimizeContext BuildExperimentalContext() const;
+	NIAGARA_API bool SupportsExperimentalVM() const;
+	NIAGARA_API FVectorVMOptimizeContext BuildExperimentalContext() const;
 #endif //VECTORVM_SUPPORTS_EXPERIMENTAL
 
-	void PostSerialize(const FArchive& Ar);
+	NIAGARA_API void PostSerialize(const FArchive& Ar);
 };
 
 template<>
@@ -599,7 +599,7 @@ struct TStructOpsTypeTraits<FNiagaraVMExecutableData> : public TStructOpsTypeTra
 	};
 };
 
-struct NIAGARA_API FNiagaraGraphCachedDataBase
+struct FNiagaraGraphCachedDataBase
 {
 	virtual ~FNiagaraGraphCachedDataBase() = default;
 	virtual void GetStaticVariables(TArray<FNiagaraVariable>& OutVars) {}
@@ -626,7 +626,7 @@ enum class ENiagaraInlineDynamicInputFormatTokenUsage
 /** Defines a single token in a format for displaying dynamic input trees inline in the UI.  These tokens can represent the
 	inputs to the dynamic input or additional text and layout options. */
 USTRUCT()
-struct NIAGARA_API FNiagaraInlineDynamicInputFormatToken
+struct FNiagaraInlineDynamicInputFormatToken
 {
 	GENERATED_BODY();
 #if WITH_EDITORONLY_DATA
@@ -655,13 +655,13 @@ struct NIAGARA_API FNiagaraInlineDynamicInputFormatToken
 
 /** Struct containing all of the data that can be different between different script versions.*/
 USTRUCT()
-struct NIAGARA_API FVersionedNiagaraScriptData
+struct FVersionedNiagaraScriptData
 {
 	GENERATED_USTRUCT_BODY()
 
 #if WITH_EDITORONLY_DATA
 public:
-	FVersionedNiagaraScriptData();
+	NIAGARA_API FVersionedNiagaraScriptData();
 
 	UPROPERTY()
 	FNiagaraAssetVersion Version;
@@ -790,7 +790,7 @@ public:
 	/** Subscriptions to parameter definitions for this script version */
 	UPROPERTY()
 	TArray<FParameterDefinitionsSubscription> ParameterDefinitionsSubscriptions;
-	TArray<ENiagaraScriptUsage> GetSupportedUsageContexts() const;
+	NIAGARA_API TArray<ENiagaraScriptUsage> GetSupportedUsageContexts() const;
 
 private:
 	friend class UNiagaraScript;
@@ -811,7 +811,7 @@ public:
 
 #if WITH_EDITORONLY_DATA
 	/** If true then this script asset uses active version control to track changes. */
-	NIAGARA_API virtual bool IsVersioningEnabled() const override { return bVersioningEnabled; }
+	virtual bool IsVersioningEnabled() const override { return bVersioningEnabled; }
 
 	/** Returns the script data for latest exposed version. */
 	NIAGARA_API FVersionedNiagaraScriptData* GetLatestScriptData();
@@ -990,7 +990,7 @@ public:
 #endif
 
 	NIAGARA_API void ComputeVMCompilationId(FNiagaraVMExecutableDataId& Id, const FGuid& VersionGuid) const;
-	NIAGARA_API const FNiagaraVMExecutableDataId& GetComputedVMCompilationId() const
+	const FNiagaraVMExecutableDataId& GetComputedVMCompilationId() const
 	{
 #if WITH_EDITORONLY_DATA
 		if (!IsCooked)
@@ -1262,11 +1262,11 @@ public:
 	void RaiseOnGPUCompilationComplete();
 
 
-	NIAGARA_API FORCEINLINE FNiagaraVMExecutableData& GetVMExecutableData() { return CachedScriptVM; }
-	NIAGARA_API FORCEINLINE const FNiagaraVMExecutableData& GetVMExecutableData() const { return CachedScriptVM; }
-	NIAGARA_API FORCEINLINE const FNiagaraVMExecutableDataId& GetVMExecutableDataCompilationId() const { return CachedScriptVMId; }
+	FORCEINLINE FNiagaraVMExecutableData& GetVMExecutableData() { return CachedScriptVM; }
+	FORCEINLINE const FNiagaraVMExecutableData& GetVMExecutableData() const { return CachedScriptVM; }
+	FORCEINLINE const FNiagaraVMExecutableDataId& GetVMExecutableDataCompilationId() const { return CachedScriptVMId; }
 
-	NIAGARA_API TConstArrayView<FNiagaraDataInterfaceGPUParamInfo> GetDataInterfaceGPUParamInfos() const { return CachedScriptVM.ShaderScriptParametersMetadata.DataInterfaceParamInfo; }
+	TConstArrayView<FNiagaraDataInterfaceGPUParamInfo> GetDataInterfaceGPUParamInfos() const { return CachedScriptVM.ShaderScriptParametersMetadata.DataInterfaceParamInfo; }
 
 	NIAGARA_API TArray<TObjectPtr<UNiagaraParameterCollection>>& GetCachedParameterCollectionReferences();
 
@@ -1446,7 +1446,7 @@ private:
 struct FVersionedNiagaraScriptWeakPtr;
 
 /** Struct combining a script with a specific version.*/
-struct NIAGARA_API FVersionedNiagaraScript : public INiagaraParameterDefinitionsSubscriber
+struct FVersionedNiagaraScript : public INiagaraParameterDefinitionsSubscriber
 {
 #if WITH_EDITORONLY_DATA
 public:
@@ -1472,14 +1472,14 @@ public:
 	virtual TArray<FParameterDefinitionsSubscription>& GetParameterDefinitionsSubscriptions() override { return GetScriptData()->ParameterDefinitionsSubscriptions; };
 
 	/** Get all UNiagaraScriptSourceBase of this subscriber. */
-	virtual TArray<UNiagaraScriptSourceBase*> GetAllSourceScripts() override;
+	NIAGARA_API virtual TArray<UNiagaraScriptSourceBase*> GetAllSourceScripts() override;
 
 	/** Get the path to the UObject of this subscriber. */
-	virtual FString GetSourceObjectPathName() const override;
+	NIAGARA_API virtual FString GetSourceObjectPathName() const override;
 	//~ End INiagaraParameterDefinitionsSubscriber Interface
 
-	FVersionedNiagaraScriptWeakPtr ToWeakPtr();
-	FVersionedNiagaraScriptData* GetScriptData() const;
+	NIAGARA_API FVersionedNiagaraScriptWeakPtr ToWeakPtr();
+	NIAGARA_API FVersionedNiagaraScriptData* GetScriptData() const;
 
 public:
 	TObjectPtr<UNiagaraScript> Script = nullptr;
@@ -1489,7 +1489,7 @@ public:
 };
 
 /** Struct combining a script with a specific version.*/
-struct NIAGARA_API FVersionedNiagaraScriptWeakPtr : public INiagaraParameterDefinitionsSubscriber
+struct FVersionedNiagaraScriptWeakPtr : public INiagaraParameterDefinitionsSubscriber
 {
 #if WITH_EDITORONLY_DATA
 public:
@@ -1504,14 +1504,14 @@ public:
 	virtual TArray<FParameterDefinitionsSubscription>& GetParameterDefinitionsSubscriptions() override { return Pin().GetScriptData()->ParameterDefinitionsSubscriptions; };
 
 	/** Get all UNiagaraScriptSourceBase of this subscriber. */
-	virtual TArray<UNiagaraScriptSourceBase*> GetAllSourceScripts() override;
+	NIAGARA_API virtual TArray<UNiagaraScriptSourceBase*> GetAllSourceScripts() override;
 
 	/** Get the path to the UObject of this subscriber. */
-	virtual FString GetSourceObjectPathName() const override;
+	NIAGARA_API virtual FString GetSourceObjectPathName() const override;
 	//~ End INiagaraParameterDefinitionsSubscriber Interface
 
-	const FVersionedNiagaraScript Pin() const;
-	FVersionedNiagaraScript Pin();
+	NIAGARA_API const FVersionedNiagaraScript Pin() const;
+	NIAGARA_API FVersionedNiagaraScript Pin();
 
 public:
 	TWeakObjectPtr<UNiagaraScript> Script;
diff --git a/Engine/Plugins/FX/Niagara/Source/Niagara/Classes/NiagaraScriptHighlight.h b/Engine/Plugins/FX/Niagara/Source/Niagara/Classes/NiagaraScriptHighlight.h
--- a/Engine/Plugins/FX/Niagara/Source/Niagara/Classes/NiagaraScriptHighlight.h
+++ b/Engine/Plugins/FX/Niagara/Source/Niagara/Classes/NiagaraScriptHighlight.h
@@ -5,13 +5,13 @@
 #include "NiagaraScriptHighlight.generated.h"
 
 USTRUCT()
-struct NIAGARA_API FNiagaraScriptHighlight
+struct FNiagaraScriptHighlight
 {
 	GENERATED_BODY()
 
-	FNiagaraScriptHighlight();
+	NIAGARA_API FNiagaraScriptHighlight();
 
-	bool IsValid() const;
+	NIAGARA_API bool IsValid() const;
 
 	UPROPERTY(EditAnywhere, Category = Highlight)
 	FLinearColor Color;
@@ -19,10 +19,10 @@ struct NIAGARA_API FNiagaraScriptHighlight
 	UPROPERTY(EditAnywhere, Category = Highlight)
 	FText DisplayName;
 
-	bool operator==(const FNiagaraScriptHighlight& Other) const;
+	NIAGARA_API bool operator==(const FNiagaraScriptHighlight& Other) const;
 
 #if WITH_EDITORONLY_DATA
-	static void ArrayToJson(const TArray<FNiagaraScriptHighlight>& InHighlights, FString& OutJson);
-	static void JsonToArray(const FString& InJson, TArray<FNiagaraScriptHighlight>& OutHighlights);
+	static NIAGARA_API void ArrayToJson(const TArray<FNiagaraScriptHighlight>& InHighlights, FString& OutJson);
+	static NIAGARA_API void JsonToArray(const FString& InJson, TArray<FNiagaraScriptHighlight>& OutHighlights);
 #endif
-};
\ No newline at end of file
+};
diff --git a/Engine/Plugins/FX/Niagara/Source/Niagara/Classes/NiagaraScriptSourceBase.h b/Engine/Plugins/FX/Niagara/Source/Niagara/Classes/NiagaraScriptSourceBase.h
--- a/Engine/Plugins/FX/Niagara/Source/Niagara/Classes/NiagaraScriptSourceBase.h
+++ b/Engine/Plugins/FX/Niagara/Source/Niagara/Classes/NiagaraScriptSourceBase.h
@@ -142,7 +142,7 @@ class UNiagaraScriptSourceBase : public UObject
 	virtual void PostLoadFromEmitter(FVersionedNiagaraEmitter OwningEmitter) { }
 
 	/** Adds a module if it isn't already in the graph. If the module isn't found bOutFoundModule will be false. If it is found and it did need to be added, the function returns true. If it already exists, it returns false. */
-	NIAGARA_API virtual bool AddModuleIfMissing(FString ModulePath, ENiagaraScriptUsage Usage, bool& bOutFoundModule) { bOutFoundModule = false; return false; }
+	virtual bool AddModuleIfMissing(FString ModulePath, ENiagaraScriptUsage Usage, bool& bOutFoundModule) { bOutFoundModule = false; return false; }
 
 #if WITH_EDITOR
 	virtual void CleanUpOldAndInitializeNewRapidIterationParameters(const FVersionedNiagaraEmitter& Emitter, ENiagaraScriptUsage ScriptUsage, FGuid ScriptUsageId, FNiagaraParameterStore& RapidIterationParameters) const { checkf(false, TEXT("Not implemented")); }
diff --git a/Engine/Plugins/FX/Niagara/Source/Niagara/Classes/NiagaraSimCache.h b/Engine/Plugins/FX/Niagara/Source/Niagara/Classes/NiagaraSimCache.h
--- a/Engine/Plugins/FX/Niagara/Source/Niagara/Classes/NiagaraSimCache.h
+++ b/Engine/Plugins/FX/Niagara/Source/Niagara/Classes/NiagaraSimCache.h
@@ -363,8 +363,8 @@ struct FNiagaraSimCacheFeedbackContext
 	TArray<FString> Warnings;
 };
 
-UCLASS(Experimental, BlueprintType)
-class NIAGARA_API UNiagaraSimCache : public UObject
+UCLASS(Experimental, BlueprintType, MinimalAPI)
+class UNiagaraSimCache : public UObject
 {
 	friend struct FNiagaraSimCacheAttributeReaderHelper;
 	friend struct FNiagaraSimCacheHelper;
@@ -380,7 +380,7 @@ public:
 	static constexpr float CacheAgeResolution = 10000.0f;
 
 	// UObject Interface
-	virtual bool IsReadyForFinishDestroy() override;
+	NIAGARA_API virtual bool IsReadyForFinishDestroy() override;
 	// UObject Interface
 
 	/** Get the caches assigned GUID, this can be set from a user or will be auto created on first write of the cache. */
@@ -400,15 +400,15 @@ public:
 	UFUNCTION(BlueprintCallable, Category=NiagaraSimCache)
 	ENiagaraSimCacheAttributeCaptureMode GetAttributeCaptureMode() const { return CreateParameters.AttributeCaptureMode; }
 
-	bool BeginWrite(FNiagaraSimCacheCreateParameters InCreateParameters, UNiagaraComponent * NiagaraComponent);
-	bool BeginWrite(FNiagaraSimCacheCreateParameters InCreateParameters, UNiagaraComponent* NiagaraComponent, FNiagaraSimCacheFeedbackContext& FeedbackContext);
-	bool WriteFrame(UNiagaraComponent* NiagaraComponent);
-	bool WriteFrame(UNiagaraComponent* NiagaraComponent, FNiagaraSimCacheFeedbackContext& FeedbackContext);
-	bool EndWrite();
+	NIAGARA_API bool BeginWrite(FNiagaraSimCacheCreateParameters InCreateParameters, UNiagaraComponent * NiagaraComponent);
+	NIAGARA_API bool BeginWrite(FNiagaraSimCacheCreateParameters InCreateParameters, UNiagaraComponent* NiagaraComponent, FNiagaraSimCacheFeedbackContext& FeedbackContext);
+	NIAGARA_API bool WriteFrame(UNiagaraComponent* NiagaraComponent);
+	NIAGARA_API bool WriteFrame(UNiagaraComponent* NiagaraComponent, FNiagaraSimCacheFeedbackContext& FeedbackContext);
+	NIAGARA_API bool EndWrite();
 
-	bool CanRead(UNiagaraSystem* NiagaraSystem);
-	bool Read(float TimeSeconds, FNiagaraSystemInstance* SystemInstance) const;
-	bool ReadFrame(int32 FrameIndex, float FrameFraction, FNiagaraSystemInstance* SystemInstance) const;
+	NIAGARA_API bool CanRead(UNiagaraSystem* NiagaraSystem);
+	NIAGARA_API bool Read(float TimeSeconds, FNiagaraSystemInstance* SystemInstance) const;
+	NIAGARA_API bool ReadFrame(int32 FrameIndex, float FrameFraction, FNiagaraSystemInstance* SystemInstance) const;
 
 	/** Get the time the simulation was at when recorded. */
 	UFUNCTION(BlueprintCallable, Category = NiagaraSimCache)
@@ -426,10 +426,10 @@ public:
 	int GetNumEmitters() const { return CacheLayout.EmitterLayouts.Num(); }
 
 	/** Get the emitter index from a name */
-	int GetEmitterIndex(FName EmitterName) const;
+	NIAGARA_API int GetEmitterIndex(FName EmitterName) const;
 
 	/** Get the system this cache is based on */
-	UNiagaraSystem* GetSystem (bool bLoadSynchronous = false);
+	NIAGARA_API UNiagaraSystem* GetSystem (bool bLoadSynchronous = false);
 	
 	/** Get the emitter name at the provided index. */
 	UFUNCTION(BlueprintCallable, Category = NiagaraSimCache)
@@ -437,76 +437,76 @@ public:
 
 	/** Returns a list of emitters we have captured in the SimCache. */
 	UFUNCTION(BlueprintCallable, Category = NiagaraSimCache)
-	TArray<FName> GetEmitterNames() const;
+	NIAGARA_API TArray<FName> GetEmitterNames() const;
 
 	/**
 	Get number of active instances for the emitter at the given frame.
 	An EmitterIndex or INDEX_NONE will return information about the system instance.
 	*/
-	int GetEmitterNumInstances(int32 EmitterIndex, int32 FrameIndex) const;
+	NIAGARA_API int GetEmitterNumInstances(int32 EmitterIndex, int32 FrameIndex) const;
 
 	/**
 	Runs a function for each attribute for the provided emitter index.
 	An EmitterIndex or INDEX_NONE will return information about the system instance.
 	Return true to continue iterating or false to stop
 	*/
-	void ForEachEmitterAttribute(int32 EmitterIndex, TFunction<bool(const FNiagaraSimCacheVariable&)> Function) const;
+	NIAGARA_API void ForEachEmitterAttribute(int32 EmitterIndex, TFunction<bool(const FNiagaraSimCacheVariable&)> Function) const;
 
 	/**
 	Reads Niagara attributes by name from the cache frame and appends them into the relevant arrays.
 	When reading using this method the attributes for a FVector would come out as AoS rather than SoA so each component is wrote one by one
 	EmitterName - If left blank will return the system simulation attributes.
 	*/
-	void ReadAttribute(TArray<float>& OutFloats, TArray<FFloat16>& OutHalfs, TArray<int32>& OutInts, FName AttributeName, FName EmitterName, int FrameIndex = 0) const;
+	NIAGARA_API void ReadAttribute(TArray<float>& OutFloats, TArray<FFloat16>& OutHalfs, TArray<int32>& OutInts, FName AttributeName, FName EmitterName, int FrameIndex = 0) const;
 
 	/**
 	Reads Niagara int attributes by name from the cache frame and appends them into the OutValues array.
 	EmitterName - If left blank will return the system simulation attributes.
 	*/
 	UFUNCTION(BlueprintCallable, Category = NiagaraSimCache)
-	void ReadIntAttribute(TArray<int32>& OutValues, FName AttributeName, FName EmitterName, int FrameIndex = 0) const;
+	NIAGARA_API void ReadIntAttribute(TArray<int32>& OutValues, FName AttributeName, FName EmitterName, int FrameIndex = 0) const;
 
 	/**
 	Reads Niagara float attributes by name from the cache frame and appends them into the OutValues array.
 	EmitterName - If left blank will return the system simulation attributes.
 	*/
 	UFUNCTION(BlueprintCallable, Category = NiagaraSimCache)
-	void ReadFloatAttribute(TArray<float>& OutValues, FName AttributeName, FName EmitterName, int FrameIndex = 0) const;
+	NIAGARA_API void ReadFloatAttribute(TArray<float>& OutValues, FName AttributeName, FName EmitterName, int FrameIndex = 0) const;
 
 	/**
 	Reads Niagara Vec2 attributes by name from the cache frame and appends them into the OutValues array.
 	EmitterName - If left blank will return the system simulation attributes.
 	*/
 	UFUNCTION(BlueprintCallable, Category = NiagaraSimCache)
-	void ReadVector2Attribute(TArray<FVector2D>& OutValues, FName AttributeName, FName EmitterName, int FrameIndex = 0) const;
+	NIAGARA_API void ReadVector2Attribute(TArray<FVector2D>& OutValues, FName AttributeName, FName EmitterName, int FrameIndex = 0) const;
 
 	/**
 	Reads Niagara Vec3 attributes by name from the cache frame and appends them into the OutValues array.
 	EmitterName - If left blank will return the system simulation attributes.
 	*/
 	UFUNCTION(BlueprintCallable, Category = NiagaraSimCache)
-	void ReadVectorAttribute(TArray<FVector>& OutValues, FName AttributeName, FName EmitterName, int FrameIndex = 0) const;
+	NIAGARA_API void ReadVectorAttribute(TArray<FVector>& OutValues, FName AttributeName, FName EmitterName, int FrameIndex = 0) const;
 
 	/**
 	Reads Niagara Vec4 attributes by name from the cache frame and appends them into the OutValues array.
 	EmitterName - If left blank will return the system simulation attributes.
 	*/
 	UFUNCTION(BlueprintCallable, Category = NiagaraSimCache)
-	void ReadVector4Attribute(TArray<FVector4>& OutValues, FName AttributeName, FName EmitterName, int FrameIndex = 0) const;
+	NIAGARA_API void ReadVector4Attribute(TArray<FVector4>& OutValues, FName AttributeName, FName EmitterName, int FrameIndex = 0) const;
 
 	/**
 	Reads Niagara Color attributes by name from the cache frame and appends them into the OutValues array.
 	EmitterName - If left blank will return the system simulation attributes.
 	*/
 	UFUNCTION(BlueprintCallable, Category = NiagaraSimCache)
-	void ReadColorAttribute(TArray<FLinearColor>& OutValues, FName AttributeName = FName("Color"), FName EmitterName = NAME_None, int FrameIndex = 0) const;
+	NIAGARA_API void ReadColorAttribute(TArray<FLinearColor>& OutValues, FName AttributeName = FName("Color"), FName EmitterName = NAME_None, int FrameIndex = 0) const;
 
 	/**
 	Reads Niagara ID attributes by name from the cache frame and appends them into the OutValues array.
 	EmitterName - If left blank will return the system simulation attributes.
 	*/
 	UFUNCTION(BlueprintCallable, Category = NiagaraSimCache)
-	void ReadIDAttribute(TArray<FNiagaraID>& OutValues, FName AttributeName, FName EmitterName, int FrameIndex = 0) const;
+	NIAGARA_API void ReadIDAttribute(TArray<FNiagaraID>& OutValues, FName AttributeName, FName EmitterName, int FrameIndex = 0) const;
 
 	/**
 	Reads Niagara Position attributes by name from the cache frame and appends them into the OutValues array.
@@ -515,7 +515,7 @@ public:
 	LocalSpaceToWorld - Caches are always stored in the emitters space, i.e. local or world space.  You can set this to false if you want the local position rather than the world position.
 	*/
 	UFUNCTION(BlueprintCallable, Category = NiagaraSimCache, meta=(AdvancedDisplay="bLocalSpaceToWorld"))
-	void ReadPositionAttribute(TArray<FVector>& OutValues, FName AttributeName = FName("Position"), FName EmitterName = NAME_None, bool bLocalSpaceToWorld = true, int FrameIndex = 0) const;
+	NIAGARA_API void ReadPositionAttribute(TArray<FVector>& OutValues, FName AttributeName = FName("Position"), FName EmitterName = NAME_None, bool bLocalSpaceToWorld = true, int FrameIndex = 0) const;
 
 	/**
 	Reads Niagara Position attributes by name from the cache frame and appends them into the OutValues array.
@@ -524,7 +524,7 @@ public:
 	EmitterName - If left blank will return the system simulation attributes.
 	*/
 	UFUNCTION(BlueprintCallable, Category = NiagaraSimCache)
-	void ReadPositionAttributeWithRebase(TArray<FVector>& OutValues, FTransform Transform, FName AttributeName = FName("Position"), FName EmitterName = NAME_None, int FrameIndex = 0) const;
+	NIAGARA_API void ReadPositionAttributeWithRebase(TArray<FVector>& OutValues, FTransform Transform, FName AttributeName = FName("Position"), FName EmitterName = NAME_None, int FrameIndex = 0) const;
 
 	/**
 	Reads Niagara Quaternion attributes by name from the cache frame and appends them into the OutValues array.
@@ -533,7 +533,7 @@ public:
 	LocalSpaceToWorld - Caches are always stored in the emitters space, i.e. local or world space.  You can set this to false if you want the local Quat rather than the world Quat.
 	*/
 	UFUNCTION(BlueprintCallable, Category = NiagaraSimCache, meta=(AdvancedDisplay="bLocalSpaceToWorld"))
-	void ReadQuatAttribute(TArray<FQuat>& OutValues, FName AttributeName = FName("MeshOrientation"), FName EmitterName = NAME_None, bool bLocalSpaceToWorld = true, int FrameIndex = 0) const;
+	NIAGARA_API void ReadQuatAttribute(TArray<FQuat>& OutValues, FName AttributeName = FName("MeshOrientation"), FName EmitterName = NAME_None, bool bLocalSpaceToWorld = true, int FrameIndex = 0) const;
 
 	/**
 	Reads Niagara Quaternion attributes by name from the cache frame and appends them into the OutValues array.
@@ -542,11 +542,11 @@ public:
 	EmitterName - If left blank will return the system simulation attributes.
 	*/
 	UFUNCTION(BlueprintCallable, Category = NiagaraSimCache)
-	void ReadQuatAttributeWithRebase(TArray<FQuat>& OutValues, FQuat Quat, FName AttributeName = FName("MeshOrientation"), FName EmitterName = NAME_None, int FrameIndex = 0) const;
+	NIAGARA_API void ReadQuatAttributeWithRebase(TArray<FQuat>& OutValues, FQuat Quat, FName AttributeName = FName("MeshOrientation"), FName EmitterName = NAME_None, int FrameIndex = 0) const;
 
 #if WITH_EDITORONLY_DATA
 	UE_DEPRECATED(5.3, "For test purposes only, don't use in production code.")
-	bool IsDataEqual(const UNiagaraSimCache& OtherCache, float ErrorTolerance, FString& OutDifference) const;
+	NIAGARA_API bool IsDataEqual(const UNiagaraSimCache& OtherCache, float ErrorTolerance, FString& OutDifference) const;
 #endif
 	
 private:
@@ -587,6 +587,6 @@ private:
 	mutable std::atomic<int32> PendingCommandsInFlight;
 
 public:
-	static FOnCacheBeginWrite	OnCacheBeginWrite;
-	static FOnCacheEndWrite		OnCacheEndWrite;
+	static NIAGARA_API FOnCacheBeginWrite	OnCacheBeginWrite;
+	static NIAGARA_API FOnCacheEndWrite		OnCacheEndWrite;
 };
diff --git a/Engine/Plugins/FX/Niagara/Source/Niagara/Classes/NiagaraSimCacheCapture.h b/Engine/Plugins/FX/Niagara/Source/Niagara/Classes/NiagaraSimCacheCapture.h
--- a/Engine/Plugins/FX/Niagara/Source/Niagara/Classes/NiagaraSimCacheCapture.h
+++ b/Engine/Plugins/FX/Niagara/Source/Niagara/Classes/NiagaraSimCacheCapture.h
@@ -41,7 +41,7 @@ struct FNiagaraSimCacheCaptureParameters
 	float AdvanceDeltaTime = 0.01666f;
 };
 
-class NIAGARA_API FNiagaraSimCacheCapture
+class FNiagaraSimCacheCapture
 {
 public:
 	DECLARE_MULTICAST_DELEGATE_OneParam(FOnCaptureComplete, UNiagaraSimCache*);
@@ -50,7 +50,7 @@ public:
 	
 	
 	// Captures a niagara sim cache. The caller is responsible for managing the lifetime of the provided component and sim cache. 
-	void CaptureNiagaraSimCache(UNiagaraSimCache* SimCache, FNiagaraSimCacheCreateParameters CreateParameters, UNiagaraComponent* NiagaraComponent, FNiagaraSimCacheCaptureParameters CaptureParameters);
+	NIAGARA_API void CaptureNiagaraSimCache(UNiagaraSimCache* SimCache, FNiagaraSimCacheCreateParameters CreateParameters, UNiagaraComponent* NiagaraComponent, FNiagaraSimCacheCaptureParameters CaptureParameters);
 
 	/**
 	Captures the simulations current frame data into the SimCache.
@@ -58,7 +58,7 @@ public:
 	The return can be invalid if the component can not be captured for some reason (i.e. not active).
 	When AdvanceSimulation is true we will manually advance the simulation one frame using the provided AdvanceDeltaTime before capturing.
 	*/
-	bool CaptureCurrentFrameImmediate(UNiagaraSimCache* SimCache, FNiagaraSimCacheCreateParameters CreateParameters, UNiagaraComponent* NiagaraComponent, UNiagaraSimCache*& OutSimCache, bool bAdvanceSimulation=false, float AdvanceDeltaTime=0.01666f);
+	NIAGARA_API bool CaptureCurrentFrameImmediate(UNiagaraSimCache* SimCache, FNiagaraSimCacheCreateParameters CreateParameters, UNiagaraComponent* NiagaraComponent, UNiagaraSimCache*& OutSimCache, bool bAdvanceSimulation=false, float AdvanceDeltaTime=0.01666f);
 
 private:
 
@@ -76,7 +76,7 @@ private:
 
 	int32 TimeOutCounter = 0;
 	
-	bool OnFrameTick(float DeltaTime);
+	NIAGARA_API bool OnFrameTick(float DeltaTime);
 
-	void FinishCapture();
+	NIAGARA_API void FinishCapture();
 };
diff --git a/Engine/Plugins/FX/Niagara/Source/Niagara/Classes/NiagaraSimCacheFunctionLibrary.h b/Engine/Plugins/FX/Niagara/Source/Niagara/Classes/NiagaraSimCacheFunctionLibrary.h
--- a/Engine/Plugins/FX/Niagara/Source/Niagara/Classes/NiagaraSimCacheFunctionLibrary.h
+++ b/Engine/Plugins/FX/Niagara/Source/Niagara/Classes/NiagaraSimCacheFunctionLibrary.h
@@ -60,8 +60,8 @@ private:
 	void CaptureFinished(UNiagaraSimCache* CapturedSimCache);
 };
 
-UCLASS()
-class NIAGARA_API UNiagaraSimCacheFunctionLibrary : public UBlueprintFunctionLibrary
+UCLASS(MinimalAPI)
+class UNiagaraSimCacheFunctionLibrary : public UBlueprintFunctionLibrary
 {
 	GENERATED_UCLASS_BODY()
 
@@ -73,11 +73,11 @@ public:
 	When AdvanceSimulation is true we will manually advance the simulation one frame using the provided AdvanceDeltaTime before capturing.
 	*/
 	UFUNCTION(BlueprintCallable, Category=NiagaraSimCache, meta=(ReturnDisplayName="Success"))
-	static bool CaptureNiagaraSimCacheImmediate(UNiagaraSimCache* SimCache, FNiagaraSimCacheCreateParameters CreateParameters, UNiagaraComponent* NiagaraComponent, UNiagaraSimCache*& OutSimCache, bool bAdvanceSimulation=false, float AdvanceDeltaTime=0.01666f);
+	static NIAGARA_API bool CaptureNiagaraSimCacheImmediate(UNiagaraSimCache* SimCache, FNiagaraSimCacheCreateParameters CreateParameters, UNiagaraComponent* NiagaraComponent, UNiagaraSimCache*& OutSimCache, bool bAdvanceSimulation=false, float AdvanceDeltaTime=0.01666f);
 
 	/**
 	Captures the simulation cache object that can be captured into using the various API calls.
 	*/
 	UFUNCTION(BlueprintCallable, Category=NiagaraSimCache, meta=(WorldContext="WorldContextObject", ReturnDisplayName="SimCache"))
-	static UNiagaraSimCache* CreateNiagaraSimCache(UObject* WorldContextObject);
+	static NIAGARA_API UNiagaraSimCache* CreateNiagaraSimCache(UObject* WorldContextObject);
 };
