[Niagara]

diff --git a/Engine/Plugins/FX/Niagara/Source/Niagara/Classes/NiagaraSystem.h b/Engine/Plugins/FX/Niagara/Source/Niagara/Classes/NiagaraSystem.h
--- a/Engine/Plugins/FX/Niagara/Source/Niagara/Classes/NiagaraSystem.h
+++ b/Engine/Plugins/FX/Niagara/Source/Niagara/Classes/NiagaraSystem.h
@@ -193,12 +193,12 @@ struct FNiagaraRendererExecutionIndex
 
 /** Container struct for an array of system scalability overrides. Enables details customization and data validation. */
 USTRUCT()
-struct NIAGARA_API FNiagaraSystemScalabilityOverrides
+struct FNiagaraSystemScalabilityOverrides
 {
 	GENERATED_USTRUCT_BODY()
 
-	FNiagaraSystemScalabilityOverrides();
-	~FNiagaraSystemScalabilityOverrides();
+	NIAGARA_API FNiagaraSystemScalabilityOverrides();
+	NIAGARA_API ~FNiagaraSystemScalabilityOverrides();
 
 	UPROPERTY(EditAnywhere, Category = "Override")
 	TArray<FNiagaraSystemScalabilityOverride> Overrides;
@@ -207,8 +207,8 @@ struct NIAGARA_API FNiagaraSystemScalabilityOverrides
 
 /** A Niagara System contains multiple Niagara Emitters to create various effects.
  * Niagara Systems can be placed in the world, unlike Emitters, and expose User Parameters to configure an effect at runtime.*/
-UCLASS(BlueprintType, meta= (LoadBehavior = "LazyOnDemand"))
-class NIAGARA_API UNiagaraSystem : public UFXSystemAsset, public INiagaraParameterDefinitionsSubscriber
+UCLASS(BlueprintType, meta= (LoadBehavior = "LazyOnDemand"), MinimalAPI)
+class UNiagaraSystem : public UFXSystemAsset, public INiagaraParameterDefinitionsSubscriber
 {
 	GENERATED_UCLASS_BODY()
 
@@ -220,35 +220,35 @@ public:
 #endif
 	//TestChange
 
-	UNiagaraSystem(FVTableHelper& Helper);
-	~UNiagaraSystem();
+	NIAGARA_API UNiagaraSystem(FVTableHelper& Helper);
+	NIAGARA_API ~UNiagaraSystem();
 
 	//~ UObject interface
-	void PostInitProperties();
-	virtual void Serialize(FArchive& Ar) override;
-	virtual void PostLoad() override; 
+	NIAGARA_API void PostInitProperties();
+	NIAGARA_API virtual void Serialize(FArchive& Ar) override;
+	NIAGARA_API virtual void PostLoad() override; 
 	virtual bool CanBeClusterRoot() const override { return true; }
-	virtual void GetResourceSizeEx(FResourceSizeEx& CumulativeResourceSize) override;
+	NIAGARA_API virtual void GetResourceSizeEx(FResourceSizeEx& CumulativeResourceSize) override;
 #if WITH_EDITORONLY_DATA
-	static void DeclareConstructClasses(TArray<FTopLevelAssetPath>& OutConstructClasses, const UClass* SpecificSubclass);
+	static NIAGARA_API void DeclareConstructClasses(TArray<FTopLevelAssetPath>& OutConstructClasses, const UClass* SpecificSubclass);
 #endif
-	virtual void BeginDestroy() override;
+	NIAGARA_API virtual void BeginDestroy() override;
 	PRAGMA_DISABLE_DEPRECATION_WARNINGS // Suppress compiler warning on override of deprecated function
 	UE_DEPRECATED(5.0, "Use version that takes FObjectPreSaveContext instead.")
-	virtual void PreSave(const class ITargetPlatform* TargetPlatform) override;
-	PRAGMA_ENABLE_DEPRECATION_WARNINGS
+	NIAGARA_API virtual void PreSave(const class ITargetPlatform* TargetPlatform) override;
+	NIAGARA_API PRAGMA_ENABLE_DEPRECATION_WARNINGS
 	virtual void PreSave(FObjectPreSaveContext ObjectSaveContext) override;
 #if WITH_EDITOR
-	virtual void PreEditChange(FProperty* PropertyThatWillChange)override;
-	virtual void PostEditChangeProperty(struct FPropertyChangedEvent& PropertyChangedEvent) override; 
-	virtual void BeginCacheForCookedPlatformData(const ITargetPlatform *TargetPlatform) override;
-	virtual bool IsCachedCookedPlatformDataLoaded(const ITargetPlatform* TargetPlatform) override;
+	NIAGARA_API virtual void PreEditChange(FProperty* PropertyThatWillChange)override;
+	NIAGARA_API virtual void PostEditChangeProperty(struct FPropertyChangedEvent& PropertyChangedEvent) override; 
+	NIAGARA_API virtual void BeginCacheForCookedPlatformData(const ITargetPlatform *TargetPlatform) override;
+	NIAGARA_API virtual bool IsCachedCookedPlatformDataLoaded(const ITargetPlatform* TargetPlatform) override;
 	//~ End UObject interface
 
 	/** Helper method to handle when an internal variable has been renamed. Renames any downstream dependencies in the emitters or exposed variables.*/
-	void HandleVariableRenamed(const FNiagaraVariable& InOldVariable, const FNiagaraVariable& InNewVariable, bool bUpdateContexts);
+	NIAGARA_API void HandleVariableRenamed(const FNiagaraVariable& InOldVariable, const FNiagaraVariable& InNewVariable, bool bUpdateContexts);
 	/** Helper method to handle when an internal variable has been removed. Resets any downstream dependencies in the emitters or exposed variables.*/
-	void HandleVariableRemoved(const FNiagaraVariable& InOldVariable, bool bUpdateContexts);
+	NIAGARA_API void HandleVariableRemoved(const FNiagaraVariable& InOldVariable, bool bUpdateContexts);
 #endif
 
 #if WITH_EDITORONLY_DATA
@@ -258,34 +258,34 @@ public:
 	virtual TArray<FParameterDefinitionsSubscription>& GetParameterDefinitionsSubscriptions() override { return ParameterDefinitionsSubscriptions; };
 
 	/** Get all UNiagaraScriptSourceBase of this subscriber. */
-	virtual TArray<UNiagaraScriptSourceBase*> GetAllSourceScripts() override;
+	NIAGARA_API virtual TArray<UNiagaraScriptSourceBase*> GetAllSourceScripts() override;
 
 	/** Get the path to the UObject of this subscriber. */
-	virtual FString GetSourceObjectPathName() const override;
+	NIAGARA_API virtual FString GetSourceObjectPathName() const override;
 
 	/** Get All adapters to editor only script vars owned directly by this subscriber. */
-	virtual TArray<UNiagaraEditorParametersAdapterBase*> GetEditorOnlyParametersAdapters() override;
+	NIAGARA_API virtual TArray<UNiagaraEditorParametersAdapterBase*> GetEditorOnlyParametersAdapters() override;
 
 	/** Get all subscribers that are owned by this subscriber.
 	 *  Note: Implemented for synchronizing UNiagaraSystem. UNiagaraSystem returns all UNiagaraEmitters it owns to call SynchronizeWithParameterDefinitions for each.
 	 */
-	virtual TArray<INiagaraParameterDefinitionsSubscriber*> GetOwnedParameterDefinitionsSubscribers() override;
+	NIAGARA_API virtual TArray<INiagaraParameterDefinitionsSubscriber*> GetOwnedParameterDefinitionsSubscribers() override;
 	//~ End INiagaraParameterDefinitionsSubscriber interface
 
-	virtual bool ChangeEmitterVersion(const FVersionedNiagaraEmitter& Emitter, const FGuid& NewVersion);
+	NIAGARA_API virtual bool ChangeEmitterVersion(const FVersionedNiagaraEmitter& Emitter, const FGuid& NewVersion);
 #endif 
 
 	/** Gets an array of the emitter handles. */
-	TArray<FNiagaraEmitterHandle>& GetEmitterHandles();
-	const TArray<FNiagaraEmitterHandle>& GetEmitterHandles()const;
+	NIAGARA_API TArray<FNiagaraEmitterHandle>& GetEmitterHandles();
+	NIAGARA_API const TArray<FNiagaraEmitterHandle>& GetEmitterHandles()const;
 	
 	FNiagaraSystemScalabilityOverrides& GetScalabilityOverrides(){return SystemScalabilityOverrides; }
 
 private:
-	bool IsValidInternal() const;
+	NIAGARA_API bool IsValidInternal() const;
 
 #if WITH_EDITORONLY_DATA
-	void FixupPositionUserParameters();
+	NIAGARA_API void FixupPositionUserParameters();
 #endif
 	
 public:
@@ -295,20 +295,20 @@ public:
 #if WITH_EDITORONLY_DATA
 	/** Adds a new emitter handle to this System.  The new handle exposes an Instance value which is a copy of the
 		original asset. */
-	FNiagaraEmitterHandle AddEmitterHandle(UNiagaraEmitter& SourceEmitter, FName EmitterName, FGuid EmitterVersion);
+	NIAGARA_API FNiagaraEmitterHandle AddEmitterHandle(UNiagaraEmitter& SourceEmitter, FName EmitterName, FGuid EmitterVersion);
 
 	/** Adds a new emitter handle to this system without copying the original asset. This should only be used for temporary systems and never for live assets. */
-	void AddEmitterHandleDirect(FNiagaraEmitterHandle& EmitterHandleToAdd);
+	NIAGARA_API void AddEmitterHandleDirect(FNiagaraEmitterHandle& EmitterHandleToAdd);
 
 	/** Duplicates an existing emitter handle and adds it to the System.  The new handle will reference the same source asset,
 		but will have a copy of the duplicated Instance value. */
-	FNiagaraEmitterHandle DuplicateEmitterHandle(const FNiagaraEmitterHandle& EmitterHandleToDuplicate, FName EmitterName);
+	NIAGARA_API FNiagaraEmitterHandle DuplicateEmitterHandle(const FNiagaraEmitterHandle& EmitterHandleToDuplicate, FName EmitterName);
 
 	/** Removes the provided emitter handle. */
-	void RemoveEmitterHandle(const FNiagaraEmitterHandle& EmitterHandleToDelete);
+	NIAGARA_API void RemoveEmitterHandle(const FNiagaraEmitterHandle& EmitterHandleToDelete);
 
 	/** Removes the emitter handles which have an Id in the supplied set. */
-	void RemoveEmitterHandlesById(const TSet<FGuid>& HandlesToRemove);
+	NIAGARA_API void RemoveEmitterHandlesById(const TSet<FGuid>& HandlesToRemove);
 
 #endif
 
@@ -335,10 +335,10 @@ public:
 	FNiagaraUserRedirectionParameterStore& GetExposedParameters()  { return ExposedParameters; }
 
 	/** Gets the System script which is used to populate the System parameters and parameter bindings. */
-	UNiagaraScript* GetSystemSpawnScript();
-	UNiagaraScript* GetSystemUpdateScript();
-	const UNiagaraScript* GetSystemSpawnScript() const;
-	const UNiagaraScript* GetSystemUpdateScript() const;
+	NIAGARA_API UNiagaraScript* GetSystemSpawnScript();
+	NIAGARA_API UNiagaraScript* GetSystemUpdateScript();
+	NIAGARA_API const UNiagaraScript* GetSystemSpawnScript() const;
+	NIAGARA_API const UNiagaraScript* GetSystemUpdateScript() const;
 
 	TOptional<float> GetMaxDeltaTime() const { return MaxDeltaTime; }
 	const FNiagaraDataSetAccessor<ENiagaraExecutionState>& GetSystemExecutionStateAccessor() const { return SystemExecutionStateAccessor; }
@@ -355,15 +355,15 @@ public:
 	template<typename TAction>
 	inline void ForEachPlatformSet(TAction Func);
 
-	bool AllowScalabilityForLocalPlayerFX()const;
+	NIAGARA_API bool AllowScalabilityForLocalPlayerFX()const;
 
-	void PrecachePSOs();
+	NIAGARA_API void PrecachePSOs();
 
 private:
-	bool IsReadyToRunInternal() const;
+	NIAGARA_API bool IsReadyToRunInternal() const;
 
 public:
-	bool IsReadyToRun() const;
+	NIAGARA_API bool IsReadyToRun() const;
 
 	FORCEINLINE bool NeedsWarmup()const { return WarmupTickCount > 0 && WarmupTickDelta > SMALL_NUMBER; }
 	FORCEINLINE float GetWarmupTime()const { return WarmupTime; }
@@ -371,14 +371,14 @@ public:
 	FORCEINLINE float GetWarmupTickDelta()const { return WarmupTickDelta; }
 	FORCEINLINE bool HasFixedTickDelta() const { return bFixedTickDelta; }
 	FORCEINLINE float GetFixedTickDeltaTime()const { return FixedTickDeltaTime; }
-	virtual void GetAssetRegistryTags(TArray<FAssetRegistryTag>& OutTags)  const override;
+	NIAGARA_API virtual void GetAssetRegistryTags(TArray<FAssetRegistryTag>& OutTags)  const override;
 
 	FORCEINLINE bool NeedsDeterminism() const { return bDeterminism; }
 	FORCEINLINE int32 GetRandomSeed() const { return RandomSeed; }
 
 	FORCEINLINE void SetWarmupTime(float InWarmupTime) { WarmupTime = InWarmupTime; ResolveWarmupTickCount(); }
 	FORCEINLINE void SetWarmupTickDelta(float InWarmupTickDelta) { WarmupTickDelta = InWarmupTickDelta; ResolveWarmupTickCount(); }
-	void ResolveWarmupTickCount();
+	NIAGARA_API void ResolveWarmupTickCount();
 
 #if STATS
 	FNiagaraStatDatabase& GetStatData() { return StatDatabase; }
@@ -386,49 +386,49 @@ public:
 
 #if WITH_EDITORONLY_DATA
 	/** Are there any pending compile requests?*/
-	bool HasOutstandingCompilationRequests(bool bIncludingGPUShaders = false) const;
+	NIAGARA_API bool HasOutstandingCompilationRequests(bool bIncludingGPUShaders = false) const;
 
 	/** Determines if this system has the supplied emitter as an editable and simulating emitter instance. */
-	bool ReferencesInstanceEmitter(const FVersionedNiagaraEmitter& Emitter) const;
+	NIAGARA_API bool ReferencesInstanceEmitter(const FVersionedNiagaraEmitter& Emitter) const;
 
 	/** Updates the system's rapid iteration parameters from a specific emitter. */
-	void RefreshSystemParametersFromEmitter(const FNiagaraEmitterHandle& EmitterHandle);
+	NIAGARA_API void RefreshSystemParametersFromEmitter(const FNiagaraEmitterHandle& EmitterHandle);
 
 	/** Removes the system's rapid iteration parameters for a specific emitter. */
-	void RemoveSystemParametersForEmitter(const FNiagaraEmitterHandle& EmitterHandle);
+	NIAGARA_API void RemoveSystemParametersForEmitter(const FNiagaraEmitterHandle& EmitterHandle);
 
 	/** Request that any dirty scripts referenced by this system be compiled.*/
-	bool RequestCompile(bool bForce, FNiagaraSystemUpdateContext* OptionalUpdateContext = nullptr);
+	NIAGARA_API bool RequestCompile(bool bForce, FNiagaraSystemUpdateContext* OptionalUpdateContext = nullptr);
 
 	/** If we have a pending compile request, is it done with yet? */
-	bool PollForCompilationComplete(bool bFlushRequestCompile = true);
+	NIAGARA_API bool PollForCompilationComplete(bool bFlushRequestCompile = true);
 
 	/** Blocks until all active compile jobs have finished */
-	void WaitForCompilationComplete(bool bIncludingGPUShaders = false, bool bShowProgress = true);
+	NIAGARA_API void WaitForCompilationComplete(bool bIncludingGPUShaders = false, bool bShowProgress = true);
 
 	/** Tries to abort all running shader compilations */
-	void KillAllActiveCompilations();
+	NIAGARA_API void KillAllActiveCompilations();
 
 	/** Invalidates any active compilation requests which will ignore their results. */
-	void InvalidateActiveCompiles();
+	NIAGARA_API void InvalidateActiveCompiles();
 
 	/** Delegate called when the system's dependencies have all been compiled.*/
-	FOnSystemCompiled& OnSystemCompiled();
+	NIAGARA_API FOnSystemCompiled& OnSystemCompiled();
 
 	/** Delegate called on PostEditChange.*/
-	FOnSystemPostEditChange& OnSystemPostEditChange();
+	NIAGARA_API FOnSystemPostEditChange& OnSystemPostEditChange();
 
 	/** Delegate called on effect type or effect type value change */
-	FOnScalabilityChanged& OnScalabilityChanged();
+	NIAGARA_API FOnScalabilityChanged& OnScalabilityChanged();
 
 	/** Gets editor specific data stored with this system. */
-	UNiagaraEditorDataBase* GetEditorData();
+	NIAGARA_API UNiagaraEditorDataBase* GetEditorData();
 
 	/** Gets editor specific parameters stored with this system */
-	UNiagaraEditorParametersAdapterBase* GetEditorParameters();
+	NIAGARA_API UNiagaraEditorParametersAdapterBase* GetEditorParameters();
 
 	/** Gets editor specific data stored with this system. */
-	const UNiagaraEditorDataBase* GetEditorData() const;
+	NIAGARA_API const UNiagaraEditorDataBase* GetEditorData() const;
 
 	/** Internal: The thumbnail image.*/
 	UPROPERTY()
@@ -463,40 +463,40 @@ public:
 	UPROPERTY(transient)
 	FNiagaraParameterStore EditorOnlyAddedParameters;
 
-	bool GetIsolateEnabled() const;
-	void SetIsolateEnabled(bool bIsolate);
+	NIAGARA_API bool GetIsolateEnabled() const;
+	NIAGARA_API void SetIsolateEnabled(bool bIsolate);
 	
 	UPROPERTY(transient)
 	FNiagaraSystemUpdateContext UpdateContext;
 
-	const static FGuid ResolveDIsMessageId;
-	const static FGuid ComputeEmitterExecutionOrderMessageId;
+	NIAGARA_API const static FGuid ResolveDIsMessageId;
+	NIAGARA_API const static FGuid ComputeEmitterExecutionOrderMessageId;
 #endif
 
-	void UpdateSystemAfterLoad();
-	void EnsureFullyLoaded() const;
+	NIAGARA_API void UpdateSystemAfterLoad();
+	NIAGARA_API void EnsureFullyLoaded() const;
 
 	bool ShouldAutoDeactivate() const { return bAutoDeactivate; }
-	bool IsLooping() const;
+	NIAGARA_API bool IsLooping() const;
 
 	const TArray<TSharedRef<const FNiagaraEmitterCompiledData>>& GetEmitterCompiledData() const { return EmitterCompiledData; };
 
 	const FNiagaraSystemCompiledData& GetSystemCompiledData() const { return SystemCompiledData; };
 
-	bool UsesCollection(const UNiagaraParameterCollection* Collection)const;
+	NIAGARA_API bool UsesCollection(const UNiagaraParameterCollection* Collection)const;
 
 	bool SupportsLargeWorldCoordinates() const { return bSupportLargeWorldCoordinates && bLwcEnabledSettingCached; }
 	FORCEINLINE bool ShouldDisableExperimentalVM() const { return bDisableExperimentalVM; }
 
 #if WITH_EDITORONLY_DATA
-	bool UsesEmitter(UNiagaraEmitter* Emitter) const;
-	bool UsesEmitter(const FVersionedNiagaraEmitter& VersionedEmitter) const;
-	bool UsesScript(const UNiagaraScript* Script)const; 
-	void ForceGraphToRecompileOnNextCheck();
+	NIAGARA_API bool UsesEmitter(UNiagaraEmitter* Emitter) const;
+	NIAGARA_API bool UsesEmitter(const FVersionedNiagaraEmitter& VersionedEmitter) const;
+	NIAGARA_API bool UsesScript(const UNiagaraScript* Script)const; 
+	NIAGARA_API void ForceGraphToRecompileOnNextCheck();
 
-	static void RequestCompileForEmitter(const FVersionedNiagaraEmitter& InEmitter);
-	static void RecomputeExecutionOrderForEmitter(const FVersionedNiagaraEmitter& InEmitter);
-	static void RecomputeExecutionOrderForDataInterface(class UNiagaraDataInterface* DataInterface);
+	static NIAGARA_API void RequestCompileForEmitter(const FVersionedNiagaraEmitter& InEmitter);
+	static NIAGARA_API void RecomputeExecutionOrderForEmitter(const FVersionedNiagaraEmitter& InEmitter);
+	static NIAGARA_API void RecomputeExecutionOrderForDataInterface(class UNiagaraDataInterface* DataInterface);
 
 	FORCEINLINE bool ShouldUseRapidIterationParameters() const { return bCompileForEdit ? !bBakeOutRapidIteration : !bBakeOutRapidIterationOnCook; }
 	FORCEINLINE bool ShouldTrimAttributes() const { return bCompileForEdit ? bTrimAttributes : bTrimAttributesOnCook; }
@@ -639,22 +639,22 @@ public:
 	//TArray<TArray<uint8>> CachedDataInterfaceEmitterDependencies;
 
 	/** Computes emitter priorities based on the dependency information. */
-	bool ComputeEmitterPriority(int32 EmitterIdx, TArray<int32, TInlineAllocator<32>>& EmitterPriorities, const TBitArray<TInlineAllocator<32>>& EmitterDependencyGraph);
+	NIAGARA_API bool ComputeEmitterPriority(int32 EmitterIdx, TArray<int32, TInlineAllocator<32>>& EmitterPriorities, const TBitArray<TInlineAllocator<32>>& EmitterDependencyGraph);
 
 	/** Queries all the data interfaces in the array for emitter dependencies. */
-	void FindDataInterfaceDependencies(FVersionedNiagaraEmitterData* EmitterData, UNiagaraScript* Script, TArray<FVersionedNiagaraEmitter>& Dependencies);
+	NIAGARA_API void FindDataInterfaceDependencies(FVersionedNiagaraEmitterData* EmitterData, UNiagaraScript* Script, TArray<FVersionedNiagaraEmitter>& Dependencies);
 
 	/** Looks at all the event handlers in the emitter to determine which other emitters it depends on. */
-	void FindEventDependencies(FVersionedNiagaraEmitterData* EmitterData, TArray<FVersionedNiagaraEmitter>& Dependencies);
+	NIAGARA_API void FindEventDependencies(FVersionedNiagaraEmitterData* EmitterData, TArray<FVersionedNiagaraEmitter>& Dependencies);
 
 	/** Computes the order in which the emitters in the Emitters array will be ticked and stores the results in EmitterExecutionOrder. */
-	void ComputeEmittersExecutionOrder();
+	NIAGARA_API void ComputeEmittersExecutionOrder();
 
 	/** Computes the order in which renderers will render */
-	void ComputeRenderersDrawOrder();
+	NIAGARA_API void ComputeRenderersDrawOrder();
 
 	/** Cache data & accessors from the compiled data, allows us to avoid per instance. */
-	void CacheFromCompiledData();
+	NIAGARA_API void CacheFromCompiledData();
 
 #if WITH_EDITORONLY_DATA
 	bool NeedsRequestCompile() const { return bNeedsRequestCompile; }
@@ -693,77 +693,77 @@ public:
 	UPROPERTY(EditAnywhere, Category = "Performance")
 	bool bRequireCurrentFrameData = true;
 
-	bool HasSystemScriptDIsWithPerInstanceData() const;
+	NIAGARA_API bool HasSystemScriptDIsWithPerInstanceData() const;
 	FORCEINLINE bool HasDIsWithPostSimulateTick() const { return bHasDIsWithPostSimulateTick; }
 	FORCEINLINE bool AllDIsPostSimulateCanOverlapFrames() const { return bAllDIsPostSimulateCanOverlapFrames; }
 	FORCEINLINE bool AsyncWorkCanOverlapTickGroups() const { return bAllDIsPostStageCanOverlapTickGroups; }
 	FORCEINLINE bool HasAnyGPUEmitters()const{ return bHasAnyGPUEmitters; }
 	FORCEINLINE bool NeedsGPUContextInitForDataInterfaces() const { return bNeedsGPUContextInitForDataInterfaces; }
 
-	FBox GetFixedBounds() const;
+	NIAGARA_API FBox GetFixedBounds() const;
 	FORCEINLINE void SetFixedBounds(const FBox& Box) { FixedBounds = Box;  }
 
 #if WITH_EDITOR
-	void SetEffectType(UNiagaraEffectType* EffectType);
+	NIAGARA_API void SetEffectType(UNiagaraEffectType* EffectType);
 
 	FNiagaraSystemScalabilityOverrides& GetSystemScalabilityOverrides() { return SystemScalabilityOverrides; }
 	FORCEINLINE bool GetOverrideScalabilitySettings()const { return bOverrideScalabilitySettings; }
 	FORCEINLINE void SetOverrideScalabilitySettings(bool bOverride) { bOverrideScalabilitySettings = bOverride; }
 
 
-	void GatherStaticVariables(TArray<FNiagaraVariable>& OutVars, TArray<FNiagaraVariable>& OutEmitterVars) const;
+	NIAGARA_API void GatherStaticVariables(TArray<FNiagaraVariable>& OutVars, TArray<FNiagaraVariable>& OutEmitterVars) const;
 #endif
-	UNiagaraEffectType* GetEffectType()const;
+	NIAGARA_API UNiagaraEffectType* GetEffectType()const;
 	FORCEINLINE const FNiagaraSystemScalabilitySettings& GetScalabilitySettings()const { return CurrentScalabilitySettings; }
-	const FNiagaraSystemScalabilityOverride& GetCurrentOverrideSettings() const;
+	NIAGARA_API const FNiagaraSystemScalabilityOverride& GetCurrentOverrideSettings() const;
 	FORCEINLINE bool NeedsSortedSignificanceCull()const{ return bNeedsSortedSignificanceCull; }
 	
-	void UpdateScalability();
+	NIAGARA_API void UpdateScalability();
 
-	ENiagaraCullProxyMode GetCullProxyMode()const;
+	NIAGARA_API ENiagaraCullProxyMode GetCullProxyMode()const;
 
 	/** Whether or not fixed bounds are enabled. */
 	UPROPERTY(EditAnywhere, Category = "System", meta = (SkipSystemResetOnChange = "true", InlineEditConditionToggle))
 	uint32 bFixedBounds : 1;
 
-	TStatId GetStatID(bool bGameThread, bool bConcurrent)const;
-	void AddToInstanceCountStat(int32 NumInstances, bool bSolo)const;
+	NIAGARA_API TStatId GetStatID(bool bGameThread, bool bConcurrent)const;
+	NIAGARA_API void AddToInstanceCountStat(int32 NumInstances, bool bSolo)const;
 
-	const FString& GetCrashReporterTag()const;
-	bool CanObtainEmitterAttribute(const FNiagaraVariableBase& InVarWithUniqueNameNamespace, FNiagaraTypeDefinition& OutBoundType) const;
-	bool CanObtainSystemAttribute(const FNiagaraVariableBase& InVar, FNiagaraTypeDefinition& OutBoundType) const;
-	bool CanObtainUserVariable(const FNiagaraVariableBase& InVar) const;
+	NIAGARA_API const FString& GetCrashReporterTag()const;
+	NIAGARA_API bool CanObtainEmitterAttribute(const FNiagaraVariableBase& InVarWithUniqueNameNamespace, FNiagaraTypeDefinition& OutBoundType) const;
+	NIAGARA_API bool CanObtainSystemAttribute(const FNiagaraVariableBase& InVar, FNiagaraTypeDefinition& OutBoundType) const;
+	NIAGARA_API bool CanObtainUserVariable(const FNiagaraVariableBase& InVar) const;
 
 #if WITH_EDITORONLY_DATA
 	FNiagaraMessageStore& GetMessageStore() { return MessageStore; }
 	const FGuid& GetAssetGuid() const {return AssetGuid;};
 #endif
 
-	FORCEINLINE void RegisterActiveInstance();
-	FORCEINLINE void UnregisterActiveInstance();
+	NIAGARA_API FORCEINLINE void RegisterActiveInstance();
+	NIAGARA_API FORCEINLINE void UnregisterActiveInstance();
 	FORCEINLINE int32& GetActiveInstancesCount() { return ActiveInstances; }
 
 #if WITH_EDITORONLY_DATA
-	UNiagaraBakerSettings* GetBakerSettings();
+	NIAGARA_API UNiagaraBakerSettings* GetBakerSettings();
 	const UNiagaraBakerSettings* GetBakerGeneratedSettings() const { return BakerGeneratedSettings; }
-	void SetBakerGeneratedSettings(UNiagaraBakerSettings* Settings);
+	NIAGARA_API void SetBakerGeneratedSettings(UNiagaraBakerSettings* Settings);
 
 	/** Get the cached parameter map traversal for this emitter.  */
-	const TSharedPtr<FNiagaraGraphCachedDataBase, ESPMode::ThreadSafe>& GetCachedTraversalData() const;
-	void InvalidateCachedData();
-	void GraphSourceChanged();
+	NIAGARA_API const TSharedPtr<FNiagaraGraphCachedDataBase, ESPMode::ThreadSafe>& GetCachedTraversalData() const;
+	NIAGARA_API void InvalidateCachedData();
+	NIAGARA_API void GraphSourceChanged();
 
 	/** Resets internal data leaving it in a state which would have minimal cost to exist in headless builds (servers) */
-	void ResetToEmptySystem();
+	NIAGARA_API void ResetToEmptySystem();
 
 	/** Updates any post compile data based upon data interfaces. */
-	void OnCompiledDataInterfaceChanged();
+	NIAGARA_API void OnCompiledDataInterfaceChanged();
 
 	/** Updates the system post UObject change. */
-	void OnCompiledUObjectChanged();
+	NIAGARA_API void OnCompiledUObjectChanged();
 
 	/** Updates the rapid iteration parameters for all scripts referenced by the system. */
-	void PrepareRapidIterationParametersForCompilation();
+	NIAGARA_API void PrepareRapidIterationParametersForCompilation();
 #endif
 
 private:
@@ -774,32 +774,32 @@ private:
 	*   If bWait is true then this *blocks* the game thread (and ui) until all shader compilations are finished.
 	*	Results from the compilation will be applied to the system unless there's already another compilation queued up
 	*/
-	bool QueryCompileComplete(bool bWait);
+	NIAGARA_API bool QueryCompileComplete(bool bWait);
 
-	void BroadcastOnSystemCompiled();
+	NIAGARA_API void BroadcastOnSystemCompiled();
 
-	void EvaluateCompileResultDependencies() const;
+	NIAGARA_API void EvaluateCompileResultDependencies() const;
 
-	void InitEmitterCompiledData();
+	NIAGARA_API void InitEmitterCompiledData();
 
-	void InitSystemCompiledData();
+	NIAGARA_API void InitSystemCompiledData();
 
 	/** Helper for filling in precomputed variable names per emitter. Converts an emitter paramter "Emitter.XXXX" into it's real parameter name. */
-	void InitEmitterVariableAliasNames(FNiagaraEmitterCompiledData& EmitterCompiledDataToInit, const UNiagaraEmitter* InAssociatedEmitter);
+	NIAGARA_API void InitEmitterVariableAliasNames(FNiagaraEmitterCompiledData& EmitterCompiledDataToInit, const UNiagaraEmitter* InAssociatedEmitter);
 
 	/** Helper for generating aliased FNiagaraVariable names for the Emitter they are associated with. */
-	const FName GetEmitterVariableAliasName(const FNiagaraVariableBase& InEmitterVar, const UNiagaraEmitter* InEmitter) const;
+	NIAGARA_API const FName GetEmitterVariableAliasName(const FNiagaraVariableBase& InEmitterVar, const UNiagaraEmitter* InEmitter) const;
 
 	/** Helper for filling in attribute datasets per emitter. */
-	void InitEmitterDataSetCompiledData(FNiagaraDataSetCompiledData& DataSetToInit, const FNiagaraEmitterHandle& InAssociatedEmitterHandle);
+	NIAGARA_API void InitEmitterDataSetCompiledData(FNiagaraDataSetCompiledData& DataSetToInit, const FNiagaraEmitterHandle& InAssociatedEmitterHandle);
 
-	void ResolveParameterStoreBindings();
+	NIAGARA_API void ResolveParameterStoreBindings();
 #endif
 
-	void ResolveScalabilitySettings();
-	void UpdatePostCompileDIInfo();
-	void UpdateDITickFlags();
-	void UpdateHasGPUEmitters();
+	NIAGARA_API void ResolveScalabilitySettings();
+	NIAGARA_API void UpdatePostCompileDIInfo();
+	NIAGARA_API void UpdateDITickFlags();
+	NIAGARA_API void UpdateHasGPUEmitters();
 
 protected:
 	/** An effect types defines settings shared between systems, for example scalability and validation rules.
@@ -965,16 +965,16 @@ protected:
 	FGraphEventRef PSOPrecacheCompletionEvent;
 
 public:
-	void AsyncOptimizeAllScripts();
+	NIAGARA_API void AsyncOptimizeAllScripts();
 
-	FGraphEventRef GetScriptOptimizationCompletionEvent();
+	NIAGARA_API FGraphEventRef GetScriptOptimizationCompletionEvent();
 
 	bool AllowCullingForLocalPlayers() const { return bAllowCullingForLocalPlayers; }
 
 	const FNiagaraSystemStaticBuffers* GetStaticBuffers() const { return StaticBuffers.Get(); }
 
 protected:
-	void GenerateStatID()const;
+	NIAGARA_API void GenerateStatID()const;
 #if STATS
 	mutable TStatId StatID_GT;
 	mutable TStatId StatID_GT_CNC;
diff --git a/Engine/Plugins/FX/Niagara/Source/Niagara/Classes/NiagaraValidationRule.h b/Engine/Plugins/FX/Niagara/Source/Niagara/Classes/NiagaraValidationRule.h
--- a/Engine/Plugins/FX/Niagara/Source/Niagara/Classes/NiagaraValidationRule.h
+++ b/Engine/Plugins/FX/Niagara/Source/Niagara/Classes/NiagaraValidationRule.h
@@ -64,14 +64,14 @@ struct FNiagaraValidationContext
 Base class for system validation logic. 
 These allow Niagara systems to be inspected for content validation either at save time or from a commandlet.
 */
-UCLASS(abstract, EditInlineNew)
-class NIAGARA_API UNiagaraValidationRule : public UObject
+UCLASS(abstract, EditInlineNew, MinimalAPI)
+class UNiagaraValidationRule : public UObject
 {
 	GENERATED_BODY()
 public:
 
 #if WITH_EDITOR
-	virtual void CheckValidity(const FNiagaraValidationContext& Context, TArray<FNiagaraValidationResult>& OutResults) const;
+	NIAGARA_API virtual void CheckValidity(const FNiagaraValidationContext& Context, TArray<FNiagaraValidationResult>& OutResults) const;
 #endif
 };
 
diff --git a/Engine/Plugins/FX/Niagara/Source/Niagara/Classes/NiagaraValidationRuleSet.h b/Engine/Plugins/FX/Niagara/Source/Niagara/Classes/NiagaraValidationRuleSet.h
--- a/Engine/Plugins/FX/Niagara/Source/Niagara/Classes/NiagaraValidationRuleSet.h
+++ b/Engine/Plugins/FX/Niagara/Source/Niagara/Classes/NiagaraValidationRuleSet.h
@@ -6,8 +6,8 @@
 #include "NiagaraValidationRuleSet.generated.h"
 
 /** A set of reusable validation rules. */
-UCLASS(BlueprintType)
-class NIAGARA_API UNiagaraValidationRuleSet : public UObject
+UCLASS(BlueprintType, MinimalAPI)
+class UNiagaraValidationRuleSet : public UObject
 {
 	GENERATED_BODY()
 public:
diff --git a/Engine/Plugins/FX/Niagara/Source/Niagara/Classes/NiagaraVersionedObject.h b/Engine/Plugins/FX/Niagara/Source/Niagara/Classes/NiagaraVersionedObject.h
--- a/Engine/Plugins/FX/Niagara/Source/Niagara/Classes/NiagaraVersionedObject.h
+++ b/Engine/Plugins/FX/Niagara/Source/Niagara/Classes/NiagaraVersionedObject.h
@@ -5,52 +5,52 @@
 #include "CoreMinimal.h"
 #include "NiagaraTypes.h"
 
-struct NIAGARA_API FNiagaraVersionDataAccessor
+struct FNiagaraVersionDataAccessor
 {
 	virtual ~FNiagaraVersionDataAccessor() = default;
 
-	virtual FNiagaraAssetVersion& GetObjectVersion();
-	virtual FText& GetVersionChangeDescription();
-	virtual bool& IsDeprecated();
-	virtual FText& GetDeprecationMessage();
-	virtual ENiagaraPythonUpdateScriptReference& GetUpdateScriptExecutionType();
-	virtual FString& GetPythonUpdateScript();
-	virtual FFilePath& GetScriptAsset();
+	NIAGARA_API virtual FNiagaraAssetVersion& GetObjectVersion();
+	NIAGARA_API virtual FText& GetVersionChangeDescription();
+	NIAGARA_API virtual bool& IsDeprecated();
+	NIAGARA_API virtual FText& GetDeprecationMessage();
+	NIAGARA_API virtual ENiagaraPythonUpdateScriptReference& GetUpdateScriptExecutionType();
+	NIAGARA_API virtual FString& GetPythonUpdateScript();
+	NIAGARA_API virtual FFilePath& GetScriptAsset();
 };
 
-class NIAGARA_API FNiagaraVersionedObject
+class FNiagaraVersionedObject
 {
 public:
 	virtual ~FNiagaraVersionedObject() = default;
 
 	/** Returns all available versions for this object. */
-	virtual TArray<FNiagaraAssetVersion> GetAllAvailableVersions() const;
+	NIAGARA_API virtual TArray<FNiagaraAssetVersion> GetAllAvailableVersions() const;
 
 #if WITH_EDITORONLY_DATA
-	virtual TSharedPtr<FNiagaraVersionDataAccessor> GetVersionDataAccessor(const FGuid& Version);
+	NIAGARA_API virtual TSharedPtr<FNiagaraVersionDataAccessor> GetVersionDataAccessor(const FGuid& Version);
 
 	/** If true then this script asset uses active version control to track changes. */
-	virtual bool IsVersioningEnabled() const;
+	NIAGARA_API virtual bool IsVersioningEnabled() const;
 	
 	/** Returns the version of the exposed version data (i.e. the version used when adding a module to the stack) */
-	virtual FNiagaraAssetVersion GetExposedVersion() const;
+	NIAGARA_API virtual FNiagaraAssetVersion GetExposedVersion() const;
 
 	/** Returns the version data for the given guid, if it exists. Otherwise returns nullptr. */
-	virtual FNiagaraAssetVersion const* FindVersionData(const FGuid& VersionGuid) const;
+	NIAGARA_API virtual FNiagaraAssetVersion const* FindVersionData(const FGuid& VersionGuid) const;
 	
 	/** Creates a new data entry for the given version number. The version must be > 1.0 and must not collide with an already existing version. The data will be a copy of the previous minor version. */
-	virtual FGuid AddNewVersion(int32 MajorVersion, int32 MinorVersion);
+	NIAGARA_API virtual FGuid AddNewVersion(int32 MajorVersion, int32 MinorVersion);
 
 	/** Deletes the version data for an existing version. The exposed version cannot be deleted and will result in an error. Does nothing if the guid does not exist in the object's version data. */
-	virtual void DeleteVersion(const FGuid& VersionGuid);
+	NIAGARA_API virtual void DeleteVersion(const FGuid& VersionGuid);
 
 	/** Changes the exposed version. Does nothing if the guid does not exist in the object's version data. */
-	virtual void ExposeVersion(const FGuid& VersionGuid);
+	NIAGARA_API virtual void ExposeVersion(const FGuid& VersionGuid);
 
 	/** Enables versioning for this object. */
-	virtual void EnableVersioning();
+	NIAGARA_API virtual void EnableVersioning();
 
 	/** Disables versioning and keeps only the data from the given version guid. Note that this breaks ALL references from existing assets and should only be used when creating a copy of an object, as the effect is very destructive.  */
-	virtual void DisableVersioning(const FGuid& VersionGuidToUse);
+	NIAGARA_API virtual void DisableVersioning(const FGuid& VersionGuidToUse);
 #endif //WITH_EDITORONLY_DATA
 };
diff --git a/Engine/Plugins/FX/Niagara/Source/Niagara/Classes/VolumeCache.h b/Engine/Plugins/FX/Niagara/Source/Niagara/Classes/VolumeCache.h
--- a/Engine/Plugins/FX/Niagara/Source/Niagara/Classes/VolumeCache.h
+++ b/Engine/Plugins/FX/Niagara/Source/Niagara/Classes/VolumeCache.h
@@ -18,8 +18,8 @@ enum class EVolumeCacheType : uint8
 	OpenVDB
 };
 
-UCLASS(Experimental)
-class NIAGARA_API UVolumeCache : public UObject
+UCLASS(Experimental, MinimalAPI)
+class UVolumeCache : public UObject
 {
 	GENERATED_UCLASS_BODY()
 
@@ -40,13 +40,13 @@ public:
 	UPROPERTY(EditAnywhere, Category = File, meta = (DisplayName = "Frame Range End"))
 	int32 FrameRangeEnd;
 		
-	void InitData();
+	NIAGARA_API void InitData();
 
-	bool LoadFile(int frame);
-	bool UnloadFile(int frame);		
-	bool LoadRange();
-	void UnloadAll();
-	FString GetAssetPath(int frame);
+	NIAGARA_API bool LoadFile(int frame);
+	NIAGARA_API bool UnloadFile(int frame);		
+	NIAGARA_API bool LoadRange();
+	NIAGARA_API void UnloadAll();
+	NIAGARA_API FString GetAssetPath(int frame);
 	
 	TSharedPtr<FVolumeCacheData> GetData() { return CachedVolumeFiles;  }
 
@@ -57,13 +57,13 @@ private:
 	TSharedPtr<FVolumeCacheData> CachedVolumeFiles;
 };
 
-class NIAGARA_API FVolumeCacheData
+class FVolumeCacheData
 {
 public:		
 	FVolumeCacheData() : DenseResolution(-1, -1, -1) {}
 	virtual ~FVolumeCacheData() {}
 
-	FString GetAssetPath(FString PathFormat, int32 FrameIndex) const;
+	NIAGARA_API FString GetAssetPath(FString PathFormat, int32 FrameIndex) const;
 
 	FIntVector GetDenseResolution() { return DenseResolution;  }
 
diff --git a/Engine/Plugins/FX/Niagara/Source/Niagara/Internal/DataInterface/NiagaraDataInterfaceActorComponent.h b/Engine/Plugins/FX/Niagara/Source/Niagara/Internal/DataInterface/NiagaraDataInterfaceActorComponent.h
--- a/Engine/Plugins/FX/Niagara/Source/Niagara/Internal/DataInterface/NiagaraDataInterfaceActorComponent.h
+++ b/Engine/Plugins/FX/Niagara/Source/Niagara/Internal/DataInterface/NiagaraDataInterfaceActorComponent.h
@@ -31,8 +31,8 @@ enum class ENDIActorComponentSourceMode : uint8
 /**
 Data interface that gives you access to actor & component information.
 */
-UCLASS(EditInlineNew, Category = "Actor", CollapseCategories, meta = (DisplayName = "Actor Component Interface"))
-class NIAGARA_API UNiagaraDataInterfaceActorComponent : public UNiagaraDataInterface
+UCLASS(EditInlineNew, Category = "Actor", CollapseCategories, meta = (DisplayName = "Actor Component Interface"), MinimalAPI)
+class UNiagaraDataInterfaceActorComponent : public UNiagaraDataInterface
 {
 	GENERATED_UCLASS_BODY()
 
@@ -66,43 +66,43 @@ public:
 	bool bRequireCurrentFrameData = true;
 	
 	//UObject Interface
-	virtual void PostInitProperties() override;
+	NIAGARA_API virtual void PostInitProperties() override;
 	//UObject Interface End
 
-	class UActorComponent* ResolveComponent(FNiagaraSystemInstance* SystemInstance, const void* PerInstanceData) const;
+	NIAGARA_API class UActorComponent* ResolveComponent(FNiagaraSystemInstance* SystemInstance, const void* PerInstanceData) const;
 
 	//UNiagaraDataInterface Interface
 	virtual bool CanExecuteOnTarget(ENiagaraSimTarget Target) const override { return true; }
 
-	virtual void GetFunctions(TArray<FNiagaraFunctionSignature>& OutFunctions) override;
-	virtual void GetVMExternalFunction(const FVMExternalFunctionBindingInfo& BindingInfo, void* InstanceData, FVMExternalFunction &OutFunc) override;
+	NIAGARA_API virtual void GetFunctions(TArray<FNiagaraFunctionSignature>& OutFunctions) override;
+	NIAGARA_API virtual void GetVMExternalFunction(const FVMExternalFunctionBindingInfo& BindingInfo, void* InstanceData, FVMExternalFunction &OutFunc) override;
 #if WITH_EDITORONLY_DATA
-	virtual bool AppendCompileHash(FNiagaraCompileHashVisitor* InVisitor) const override;
-	virtual void GetParameterDefinitionHLSL(const FNiagaraDataInterfaceGPUParamInfo& ParamInfo, FString& OutHLSL) override;
-	virtual bool GetFunctionHLSL(const FNiagaraDataInterfaceGPUParamInfo& ParamInfo, const FNiagaraDataInterfaceGeneratedFunction& FunctionInfo, int FunctionInstanceIndex, FString& OutHLSL) override;
-	virtual bool UpgradeFunctionCall(FNiagaraFunctionSignature& FunctionSignature) override;
+	NIAGARA_API virtual bool AppendCompileHash(FNiagaraCompileHashVisitor* InVisitor) const override;
+	NIAGARA_API virtual void GetParameterDefinitionHLSL(const FNiagaraDataInterfaceGPUParamInfo& ParamInfo, FString& OutHLSL) override;
+	NIAGARA_API virtual bool GetFunctionHLSL(const FNiagaraDataInterfaceGPUParamInfo& ParamInfo, const FNiagaraDataInterfaceGeneratedFunction& FunctionInfo, int FunctionInstanceIndex, FString& OutHLSL) override;
+	NIAGARA_API virtual bool UpgradeFunctionCall(FNiagaraFunctionSignature& FunctionSignature) override;
 #endif
-	virtual void BuildShaderParameters(FNiagaraShaderParametersBuilder& ShaderParametersBuilder) const override;
-	virtual void SetShaderParameters(const FNiagaraDataInterfaceSetShaderParametersContext& Context) const override;
+	NIAGARA_API virtual void BuildShaderParameters(FNiagaraShaderParametersBuilder& ShaderParametersBuilder) const override;
+	NIAGARA_API virtual void SetShaderParameters(const FNiagaraDataInterfaceSetShaderParametersContext& Context) const override;
 
-	virtual bool InitPerInstanceData(void* PerInstanceData, FNiagaraSystemInstance* SystemInstance) override;
-	virtual void DestroyPerInstanceData(void* PerInstanceData, FNiagaraSystemInstance* SystemInstance) override;
-	virtual int32 PerInstanceDataSize() const override;
-	virtual bool PerInstanceTick(void* PerInstanceData, FNiagaraSystemInstance* SystemInstance, float DeltaSeconds) override;
+	NIAGARA_API virtual bool InitPerInstanceData(void* PerInstanceData, FNiagaraSystemInstance* SystemInstance) override;
+	NIAGARA_API virtual void DestroyPerInstanceData(void* PerInstanceData, FNiagaraSystemInstance* SystemInstance) override;
+	NIAGARA_API virtual int32 PerInstanceDataSize() const override;
+	NIAGARA_API virtual bool PerInstanceTick(void* PerInstanceData, FNiagaraSystemInstance* SystemInstance, float DeltaSeconds) override;
 
-	virtual void ProvidePerInstanceDataForRenderThread(void* DataForRenderThread, void* PerInstanceData, const FNiagaraSystemInstanceID& SystemInstance) override;
+	NIAGARA_API virtual void ProvidePerInstanceDataForRenderThread(void* DataForRenderThread, void* PerInstanceData, const FNiagaraSystemInstanceID& SystemInstance) override;
 
 	virtual bool HasTickGroupPrereqs() const override { return true; }
-	virtual ETickingGroup CalculateTickGroup(const void* PerInstanceData) const override;
+	NIAGARA_API virtual ETickingGroup CalculateTickGroup(const void* PerInstanceData) const override;
 
-	virtual bool Equals(const UNiagaraDataInterface* Other) const override;
-	virtual bool CopyToInternal(UNiagaraDataInterface* Destination) const override;
+	NIAGARA_API virtual bool Equals(const UNiagaraDataInterface* Other) const override;
+	NIAGARA_API virtual bool CopyToInternal(UNiagaraDataInterface* Destination) const override;
 
 	virtual bool HasPreSimulateTick() const override { return true; }
 	//UNiagaraDataInterface Interface
 
 private:
-	void VMGetMatrix(FVectorVMExternalFunctionContext& Context);
-	void VMGetTransform(FVectorVMExternalFunctionContext& Context);
-	void VMGetPhysicsVelocity(FVectorVMExternalFunctionContext& Context);
+	NIAGARA_API void VMGetMatrix(FVectorVMExternalFunctionContext& Context);
+	NIAGARA_API void VMGetTransform(FVectorVMExternalFunctionContext& Context);
+	NIAGARA_API void VMGetPhysicsVelocity(FVectorVMExternalFunctionContext& Context);
 };
diff --git a/Engine/Plugins/FX/Niagara/Source/Niagara/Internal/DataInterface/NiagaraDataInterfaceDataChannelRead.h b/Engine/Plugins/FX/Niagara/Source/Niagara/Internal/DataInterface/NiagaraDataInterfaceDataChannelRead.h
--- a/Engine/Plugins/FX/Niagara/Source/Niagara/Internal/DataInterface/NiagaraDataInterfaceDataChannelRead.h
+++ b/Engine/Plugins/FX/Niagara/Source/Niagara/Internal/DataInterface/NiagaraDataInterfaceDataChannelRead.h
@@ -51,8 +51,8 @@ enum class ENDIDataChannelSpawnMode
 	Max UMETA(Hidden),
 };
 
-UCLASS(Experimental, EditInlineNew, Category = "Data Channels", CollapseCategories, meta = (DisplayName = "Data Channel Reader"))
-class NIAGARA_API UNiagaraDataInterfaceDataChannelRead : public UNiagaraDataInterfaceRWBase
+UCLASS(Experimental, EditInlineNew, Category = "Data Channels", CollapseCategories, meta = (DisplayName = "Data Channel Reader"), MinimalAPI)
+class UNiagaraDataInterfaceDataChannelRead : public UNiagaraDataInterfaceRWBase
 {
 	GENERATED_UCLASS_BODY()
 protected:
@@ -102,63 +102,63 @@ public:
 	FName SpawnInfoName;
 	
 	//UObject Interface
-	virtual void PostInitProperties() override;
-	virtual void BeginDestroy() override;
+	NIAGARA_API virtual void PostInitProperties() override;
+	NIAGARA_API virtual void BeginDestroy() override;
 	//UObject Interface End
 
 	//UNiagaraDataInterface Interface
-	virtual void GetFunctions(TArray<FNiagaraFunctionSignature>& OutFunctions) override;
-	virtual void GetVMExternalFunction(const FVMExternalFunctionBindingInfo& BindingInfo, void* InstanceData, FVMExternalFunction &OutFunc) override;
+	NIAGARA_API virtual void GetFunctions(TArray<FNiagaraFunctionSignature>& OutFunctions) override;
+	NIAGARA_API virtual void GetVMExternalFunction(const FVMExternalFunctionBindingInfo& BindingInfo, void* InstanceData, FVMExternalFunction &OutFunc) override;
 	virtual bool CanExecuteOnTarget(ENiagaraSimTarget Target) const override { return true; }
 #if WITH_EDITORONLY_DATA
-	virtual bool AppendCompileHash(FNiagaraCompileHashVisitor* InVisitor) const override;
-	virtual void GetCommonHLSL(FString& OutHLSL)override;
-	virtual bool GetFunctionHLSL(FNiagaraDataInterfaceHlslGenerationContext& HlslGenContext, FString& OutHLSL) override;
-	virtual void GetParameterDefinitionHLSL(FNiagaraDataInterfaceHlslGenerationContext& HlslGenContext, FString& OutHLSL) override;
+	NIAGARA_API virtual bool AppendCompileHash(FNiagaraCompileHashVisitor* InVisitor) const override;
+	NIAGARA_API virtual void GetCommonHLSL(FString& OutHLSL)override;
+	NIAGARA_API virtual bool GetFunctionHLSL(FNiagaraDataInterfaceHlslGenerationContext& HlslGenContext, FString& OutHLSL) override;
+	NIAGARA_API virtual void GetParameterDefinitionHLSL(FNiagaraDataInterfaceHlslGenerationContext& HlslGenContext, FString& OutHLSL) override;
 
-	virtual void PostCompile()override;
+	NIAGARA_API virtual void PostCompile()override;
 #endif	
 
 #if WITH_EDITOR	
-	virtual void GetFeedback(UNiagaraSystem* InAsset, UNiagaraComponent* InComponent, TArray<FNiagaraDataInterfaceError>& OutErrors, TArray<FNiagaraDataInterfaceFeedback>& OutWarnings, TArray<FNiagaraDataInterfaceFeedback>& OutInfo) override;
-	virtual void ValidateFunction(const FNiagaraFunctionSignature& Function, TArray<FText>& OutValidationErrors) override;
+	NIAGARA_API virtual void GetFeedback(UNiagaraSystem* InAsset, UNiagaraComponent* InComponent, TArray<FNiagaraDataInterfaceError>& OutErrors, TArray<FNiagaraDataInterfaceFeedback>& OutWarnings, TArray<FNiagaraDataInterfaceFeedback>& OutInfo) override;
+	NIAGARA_API virtual void ValidateFunction(const FNiagaraFunctionSignature& Function, TArray<FText>& OutValidationErrors) override;
 #endif
 
-	virtual bool Equals(const UNiagaraDataInterface* Other) const override;
+	NIAGARA_API virtual bool Equals(const UNiagaraDataInterface* Other) const override;
 
-	virtual void BuildShaderParameters(FNiagaraShaderParametersBuilder& ShaderParametersBuilder) const override;
-	virtual void SetShaderParameters(const FNiagaraDataInterfaceSetShaderParametersContext& Context) const override;
+	NIAGARA_API virtual void BuildShaderParameters(FNiagaraShaderParametersBuilder& ShaderParametersBuilder) const override;
+	NIAGARA_API virtual void SetShaderParameters(const FNiagaraDataInterfaceSetShaderParametersContext& Context) const override;
 
-	virtual bool InitPerInstanceData(void* PerInstanceData, FNiagaraSystemInstance* SystemInstance) override;
-	virtual void DestroyPerInstanceData(void* PerInstanceData, FNiagaraSystemInstance* SystemInstance) override;
-	virtual int32 PerInstanceDataSize() const override;
+	NIAGARA_API virtual bool InitPerInstanceData(void* PerInstanceData, FNiagaraSystemInstance* SystemInstance) override;
+	NIAGARA_API virtual void DestroyPerInstanceData(void* PerInstanceData, FNiagaraSystemInstance* SystemInstance) override;
+	NIAGARA_API virtual int32 PerInstanceDataSize() const override;
 	virtual bool HasPreSimulateTick() const override { return true; }
 	virtual bool HasPostSimulateTick() const override { return true; }
 	virtual bool HasPostStageTick(ENiagaraScriptUsage Usage) const override { return Usage == ENiagaraScriptUsage::EmitterUpdateScript || Usage == ENiagaraScriptUsage::SystemUpdateScript; }
-	virtual bool PerInstanceTick(void* PerInstanceData, FNiagaraSystemInstance* SystemInstance, float DeltaSeconds) override;
-	virtual bool PerInstanceTickPostSimulate(void* PerInstanceData, FNiagaraSystemInstance* SystemInstance, float DeltaSeconds) override;
-	virtual void PostStageTick(FNDICpuPostStageContext& Context)override;
-	virtual void ProvidePerInstanceDataForRenderThread(void* DataForRenderThread, void* PerInstanceData, const FNiagaraSystemInstanceID& SystemInstance) override;
-	virtual void GetEmitterDependencies(UNiagaraSystem* Asset, TArray<FVersionedNiagaraEmitter>& Dependencies) const;
+	NIAGARA_API virtual bool PerInstanceTick(void* PerInstanceData, FNiagaraSystemInstance* SystemInstance, float DeltaSeconds) override;
+	NIAGARA_API virtual bool PerInstanceTickPostSimulate(void* PerInstanceData, FNiagaraSystemInstance* SystemInstance, float DeltaSeconds) override;
+	NIAGARA_API virtual void PostStageTick(FNDICpuPostStageContext& Context)override;
+	NIAGARA_API virtual void ProvidePerInstanceDataForRenderThread(void* DataForRenderThread, void* PerInstanceData, const FNiagaraSystemInstanceID& SystemInstance) override;
+	NIAGARA_API virtual void GetEmitterDependencies(UNiagaraSystem* Asset, TArray<FVersionedNiagaraEmitter>& Dependencies) const;
 	
 	//We cannot overlap frames as we must correctly sync up with the data channel manager on Begin/End frame etc.
 	virtual bool PostSimulateCanOverlapFrames() const { return false; }
 	//UNiagaraDataInterface Interface
 
 	//Functions usable anywhere.
-	void Num(FVectorVMExternalFunctionContext& Context);
-	void Read(FVectorVMExternalFunctionContext& Context, int32 FuncIdx);
-	void Consume(FVectorVMExternalFunctionContext& Context, int32 FuncIdx);
+	NIAGARA_API void Num(FVectorVMExternalFunctionContext& Context);
+	NIAGARA_API void Read(FVectorVMExternalFunctionContext& Context, int32 FuncIdx);
+	NIAGARA_API void Consume(FVectorVMExternalFunctionContext& Context, int32 FuncIdx);
 
 	//Emitter only functions.
-	void SpawnFromSpawnInfo(FVectorVMExternalFunctionContext& Context);
-	void SpawnConditional(FVectorVMExternalFunctionContext& Context, int32 FuncIndex);
+	NIAGARA_API void SpawnFromSpawnInfo(FVectorVMExternalFunctionContext& Context);
+	NIAGARA_API void SpawnConditional(FVectorVMExternalFunctionContext& Context, int32 FuncIndex);
 
 	FNDIDataChannelCompiledData& GetCompiledData() { return CompiledData; }
 
 protected:
-	virtual bool CopyToInternal(UNiagaraDataInterface* Destination) const override;
-	UNiagaraDataInterfaceDataChannelWrite* FindSourceDI()const;
+	NIAGARA_API virtual bool CopyToInternal(UNiagaraDataInterface* Destination) const override;
+	NIAGARA_API UNiagaraDataInterfaceDataChannelWrite* FindSourceDI()const;
 
  	UPROPERTY()
  	FNDIDataChannelCompiledData CompiledData;
@@ -251,4 +251,4 @@ struct FNiagaraDataInterfaceProxy_DataChannelRead : public FNiagaraDataInterface
 	};
 
 	TMap<FNiagaraSystemInstanceID, FInstanceData> SystemInstancesToProxyData_RT;
-};
\ No newline at end of file
+};
diff --git a/Engine/Plugins/FX/Niagara/Source/Niagara/Internal/DataInterface/NiagaraDataInterfaceDataChannelWrite.h b/Engine/Plugins/FX/Niagara/Source/Niagara/Internal/DataInterface/NiagaraDataInterfaceDataChannelWrite.h
--- a/Engine/Plugins/FX/Niagara/Source/Niagara/Internal/DataInterface/NiagaraDataInterfaceDataChannelWrite.h
+++ b/Engine/Plugins/FX/Niagara/Source/Niagara/Internal/DataInterface/NiagaraDataInterfaceDataChannelWrite.h
@@ -34,8 +34,8 @@ struct FNDIDataChannelWriteCompiledData : public FNDIDataChannelCompiledData
 	bool Init(UNiagaraSystem* System, UNiagaraDataInterfaceDataChannelWrite* OwnerDI);
 };
 
-UCLASS(Experimental, EditInlineNew, Category = "Data Channels", CollapseCategories, meta = (DisplayName = "Data Channel Writer"))
-class NIAGARA_API UNiagaraDataInterfaceDataChannelWrite : public UNiagaraDataInterface
+UCLASS(Experimental, EditInlineNew, Category = "Data Channels", CollapseCategories, meta = (DisplayName = "Data Channel Writer"), MinimalAPI)
+class UNiagaraDataInterfaceDataChannelWrite : public UNiagaraDataInterface
 {
 	GENERATED_UCLASS_BODY()
 
@@ -75,42 +75,42 @@ public:
 	TObjectPtr<UNiagaraDataChannelAsset> Channel;
 
 	//UObject Interface
-	virtual void PostInitProperties() override;
+	NIAGARA_API virtual void PostInitProperties() override;
 	//UObject Interface End
 
 	//UNiagaraDataInterface Interface
-	virtual void GetFunctions(TArray<FNiagaraFunctionSignature>& OutFunctions) override;
-	virtual void GetVMExternalFunction(const FVMExternalFunctionBindingInfo& BindingInfo, void* InstanceData, FVMExternalFunction& OutFunc) override;
+	NIAGARA_API virtual void GetFunctions(TArray<FNiagaraFunctionSignature>& OutFunctions) override;
+	NIAGARA_API virtual void GetVMExternalFunction(const FVMExternalFunctionBindingInfo& BindingInfo, void* InstanceData, FVMExternalFunction& OutFunc) override;
 	virtual bool CanExecuteOnTarget(ENiagaraSimTarget Target) const override { return true; }
 #if WITH_EDITORONLY_DATA
-	virtual bool AppendCompileHash(FNiagaraCompileHashVisitor* InVisitor) const override;
-	virtual void GetCommonHLSL(FString& OutHLSL)override;
-	virtual bool GetFunctionHLSL(FNiagaraDataInterfaceHlslGenerationContext& HlslGenContext, FString& OutHLSL) override;
-	virtual void GetParameterDefinitionHLSL(FNiagaraDataInterfaceHlslGenerationContext& HlslGenContext, FString& OutHLSL) override;
+	NIAGARA_API virtual bool AppendCompileHash(FNiagaraCompileHashVisitor* InVisitor) const override;
+	NIAGARA_API virtual void GetCommonHLSL(FString& OutHLSL)override;
+	NIAGARA_API virtual bool GetFunctionHLSL(FNiagaraDataInterfaceHlslGenerationContext& HlslGenContext, FString& OutHLSL) override;
+	NIAGARA_API virtual void GetParameterDefinitionHLSL(FNiagaraDataInterfaceHlslGenerationContext& HlslGenContext, FString& OutHLSL) override;
 
-	virtual void PostCompile()override;
+	NIAGARA_API virtual void PostCompile()override;
 #endif
-	virtual void BuildShaderParameters(FNiagaraShaderParametersBuilder& ShaderParametersBuilder) const override;
-	virtual void SetShaderParameters(const FNiagaraDataInterfaceSetShaderParametersContext& Context) const override;
+	NIAGARA_API virtual void BuildShaderParameters(FNiagaraShaderParametersBuilder& ShaderParametersBuilder) const override;
+	NIAGARA_API virtual void SetShaderParameters(const FNiagaraDataInterfaceSetShaderParametersContext& Context) const override;
 
 #if WITH_EDITOR
-	virtual void GetFeedback(UNiagaraSystem* InAsset, UNiagaraComponent* InComponent, TArray<FNiagaraDataInterfaceError>& OutErrors, TArray<FNiagaraDataInterfaceFeedback>& OutWarnings, TArray<FNiagaraDataInterfaceFeedback>& OutInfo);
-	virtual void ValidateFunction(const FNiagaraFunctionSignature& Function, TArray<FText>& OutValidationErrors);
+	NIAGARA_API virtual void GetFeedback(UNiagaraSystem* InAsset, UNiagaraComponent* InComponent, TArray<FNiagaraDataInterfaceError>& OutErrors, TArray<FNiagaraDataInterfaceFeedback>& OutWarnings, TArray<FNiagaraDataInterfaceFeedback>& OutInfo);
+	NIAGARA_API virtual void ValidateFunction(const FNiagaraFunctionSignature& Function, TArray<FText>& OutValidationErrors);
 #endif
 
-	virtual bool Equals(const UNiagaraDataInterface* Other) const override;
+	NIAGARA_API virtual bool Equals(const UNiagaraDataInterface* Other) const override;
 
-	virtual bool InitPerInstanceData(void* PerInstanceData, FNiagaraSystemInstance* SystemInstance) override;
-	virtual void DestroyPerInstanceData(void* PerInstanceData, FNiagaraSystemInstance* SystemInstance) override;
-	virtual int32 PerInstanceDataSize() const override;
+	NIAGARA_API virtual bool InitPerInstanceData(void* PerInstanceData, FNiagaraSystemInstance* SystemInstance) override;
+	NIAGARA_API virtual void DestroyPerInstanceData(void* PerInstanceData, FNiagaraSystemInstance* SystemInstance) override;
+	NIAGARA_API virtual int32 PerInstanceDataSize() const override;
 	virtual bool HasPreSimulateTick() const override { return true; }
 	virtual bool HasPreStageTick(ENiagaraScriptUsage Usage) const override { return true; }
 	virtual bool HasPostStageTick(ENiagaraScriptUsage Usage) const override { return true; }
-	virtual bool PerInstanceTick(void* PerInstanceData, FNiagaraSystemInstance* SystemInstance, float DeltaSeconds) override;
-	virtual bool PerInstanceTickPostSimulate(void* PerInstanceData, FNiagaraSystemInstance* SystemInstance, float DeltaSeconds) override;
-	virtual void PreStageTick(FNDICpuPreStageContext& Context) override;
-	virtual void PostStageTick(FNDICpuPostStageContext& Context) override;
-	virtual void ProvidePerInstanceDataForRenderThread(void* DataForRenderThread, void* PerInstanceData, const FNiagaraSystemInstanceID& SystemInstance) override;
+	NIAGARA_API virtual bool PerInstanceTick(void* PerInstanceData, FNiagaraSystemInstance* SystemInstance, float DeltaSeconds) override;
+	NIAGARA_API virtual bool PerInstanceTickPostSimulate(void* PerInstanceData, FNiagaraSystemInstance* SystemInstance, float DeltaSeconds) override;
+	NIAGARA_API virtual void PreStageTick(FNDICpuPreStageContext& Context) override;
+	NIAGARA_API virtual void PostStageTick(FNDICpuPostStageContext& Context) override;
+	NIAGARA_API virtual void ProvidePerInstanceDataForRenderThread(void* DataForRenderThread, void* PerInstanceData, const FNiagaraSystemInstanceID& SystemInstance) override;
 
 	//We cannot overlap frames as we must correctly sync up with the data channel manager on Begin/End frame etc.
 	virtual bool PostSimulateCanOverlapFrames() const { return false; }
@@ -118,17 +118,17 @@ public:
 	virtual bool PostStageCanOverlapTickGroups() const { return false; }
 	//UNiagaraDataInterface Interface
 
-	void Num(FVectorVMExternalFunctionContext& Context);
-	void Write(FVectorVMExternalFunctionContext& Context, int32 FuncIdx);
-	void Append(FVectorVMExternalFunctionContext& Context, int32 FuncIdx);
+	NIAGARA_API void Num(FVectorVMExternalFunctionContext& Context);
+	NIAGARA_API void Write(FVectorVMExternalFunctionContext& Context, int32 FuncIdx);
+	NIAGARA_API void Append(FVectorVMExternalFunctionContext& Context, int32 FuncIdx);
 
 	const FNDIDataChannelWriteCompiledData& GetCompiledData()const { return CompiledData; }
 
 	bool ShouldPublish()const { return bPublishToGame || bPublishToCPU || bPublishToGPU; }
 
 protected:
-	virtual bool CopyToInternal(UNiagaraDataInterface* Destination) const override;
+	NIAGARA_API virtual bool CopyToInternal(UNiagaraDataInterface* Destination) const override;
 
 	UPROPERTY()
 	FNDIDataChannelWriteCompiledData CompiledData;
-};
\ No newline at end of file
+};
diff --git a/Engine/Plugins/FX/Niagara/Source/Niagara/Internal/DataInterface/NiagaraDataInterfaceStaticMesh.h b/Engine/Plugins/FX/Niagara/Source/Niagara/Internal/DataInterface/NiagaraDataInterfaceStaticMesh.h
--- a/Engine/Plugins/FX/Niagara/Source/Niagara/Internal/DataInterface/NiagaraDataInterfaceStaticMesh.h
+++ b/Engine/Plugins/FX/Niagara/Source/Niagara/Internal/DataInterface/NiagaraDataInterfaceStaticMesh.h
@@ -51,8 +51,8 @@ struct FNDIStaticMeshSectionFilter
 };
 
 /** Data Interface allowing sampling of static meshes. */
-UCLASS(EditInlineNew, Category = "Meshes", meta = (DisplayName = "Static Mesh"))
-class NIAGARA_API UNiagaraDataInterfaceStaticMesh : public UNiagaraDataInterface
+UCLASS(EditInlineNew, Category = "Meshes", meta = (DisplayName = "Static Mesh"), MinimalAPI)
+class UNiagaraDataInterfaceStaticMesh : public UNiagaraDataInterface
 {
 	GENERATED_UCLASS_BODY()
 
@@ -126,123 +126,123 @@ public:
 	uint32 ChangeId = 0;
 
 	//~ UObject interface
-	virtual void PostInitProperties() override;
-	virtual void PostLoad() override;
+	NIAGARA_API virtual void PostInitProperties() override;
+	NIAGARA_API virtual void PostLoad() override;
 #if WITH_EDITOR
-	virtual void PostEditChangeProperty(struct FPropertyChangedEvent& PropertyChangedEvent) override;
-	virtual bool CanEditChange(const FProperty* InProperty) const override;
+	NIAGARA_API virtual void PostEditChangeProperty(struct FPropertyChangedEvent& PropertyChangedEvent) override;
+	NIAGARA_API virtual bool CanEditChange(const FProperty* InProperty) const override;
 #endif
 	//~ UObject interface
 
 	//~ UNiagaraDataInterface interface
-	virtual bool InitPerInstanceData(void* PerInstanceData, FNiagaraSystemInstance* SystemInstance) override;
-	virtual void DestroyPerInstanceData(void* PerInstanceData, FNiagaraSystemInstance* SystemInstance) override;
-	virtual bool PerInstanceTick(void* PerInstanceData, FNiagaraSystemInstance* SystemInstance, float DeltaSeconds) override;
-	virtual int32 PerInstanceDataSize() const override;
-
-	virtual void ProvidePerInstanceDataForRenderThread(void* DataForRenderThread, void* PerInstanceData, const FNiagaraSystemInstanceID& SystemInstance) override;
-
-	virtual void GetFunctions(TArray<FNiagaraFunctionSignature>& OutFunctions) override;
-	void GetVertexSamplingFunctions(TArray<FNiagaraFunctionSignature>& OutFunctions, const FNiagaraFunctionSignature& BaseSignature) const;
-	void GetTriangleSamplingFunctions(TArray<FNiagaraFunctionSignature>& OutFunctions, const FNiagaraFunctionSignature& BaseSignature) const;
-	void GetSocketSamplingFunctions(TArray<FNiagaraFunctionSignature>& OutFunctions, const FNiagaraFunctionSignature& BaseSignature) const;
-	void GetSectionFunctions(TArray<FNiagaraFunctionSignature>& OutFunctions, const FNiagaraFunctionSignature& BaseSignature) const;
-	void GetMiscFunctions(TArray<FNiagaraFunctionSignature>& OutFunctions, const FNiagaraFunctionSignature& BaseSignature) const;
-	void GetUVMappingFunctions(TArray<FNiagaraFunctionSignature>& OutFunctions, const FNiagaraFunctionSignature& BaseSignature) const;
-	void GetDistanceFieldFunctions(TArray<FNiagaraFunctionSignature>& OutFunctions, const FNiagaraFunctionSignature& BaseSignature) const;
-	void GetDeprecatedFunctions(TArray<FNiagaraFunctionSignature>& OutFunctions, const FNiagaraFunctionSignature& BaseSignature) const;
-	static bool FunctionNeedsCpuAccess(FName InName);
+	NIAGARA_API virtual bool InitPerInstanceData(void* PerInstanceData, FNiagaraSystemInstance* SystemInstance) override;
+	NIAGARA_API virtual void DestroyPerInstanceData(void* PerInstanceData, FNiagaraSystemInstance* SystemInstance) override;
+	NIAGARA_API virtual bool PerInstanceTick(void* PerInstanceData, FNiagaraSystemInstance* SystemInstance, float DeltaSeconds) override;
+	NIAGARA_API virtual int32 PerInstanceDataSize() const override;
+
+	NIAGARA_API virtual void ProvidePerInstanceDataForRenderThread(void* DataForRenderThread, void* PerInstanceData, const FNiagaraSystemInstanceID& SystemInstance) override;
+
+	NIAGARA_API virtual void GetFunctions(TArray<FNiagaraFunctionSignature>& OutFunctions) override;
+	NIAGARA_API void GetVertexSamplingFunctions(TArray<FNiagaraFunctionSignature>& OutFunctions, const FNiagaraFunctionSignature& BaseSignature) const;
+	NIAGARA_API void GetTriangleSamplingFunctions(TArray<FNiagaraFunctionSignature>& OutFunctions, const FNiagaraFunctionSignature& BaseSignature) const;
+	NIAGARA_API void GetSocketSamplingFunctions(TArray<FNiagaraFunctionSignature>& OutFunctions, const FNiagaraFunctionSignature& BaseSignature) const;
+	NIAGARA_API void GetSectionFunctions(TArray<FNiagaraFunctionSignature>& OutFunctions, const FNiagaraFunctionSignature& BaseSignature) const;
+	NIAGARA_API void GetMiscFunctions(TArray<FNiagaraFunctionSignature>& OutFunctions, const FNiagaraFunctionSignature& BaseSignature) const;
+	NIAGARA_API void GetUVMappingFunctions(TArray<FNiagaraFunctionSignature>& OutFunctions, const FNiagaraFunctionSignature& BaseSignature) const;
+	NIAGARA_API void GetDistanceFieldFunctions(TArray<FNiagaraFunctionSignature>& OutFunctions, const FNiagaraFunctionSignature& BaseSignature) const;
+	NIAGARA_API void GetDeprecatedFunctions(TArray<FNiagaraFunctionSignature>& OutFunctions, const FNiagaraFunctionSignature& BaseSignature) const;
+	static NIAGARA_API bool FunctionNeedsCpuAccess(FName InName);
 
 #if WITH_EDITORONLY_DATA
-	virtual void GetCommonHLSL(FString& OutHLSL) override;
-	virtual bool UpgradeFunctionCall(FNiagaraFunctionSignature& FunctionSignature) override;
+	NIAGARA_API virtual void GetCommonHLSL(FString& OutHLSL) override;
+	NIAGARA_API virtual bool UpgradeFunctionCall(FNiagaraFunctionSignature& FunctionSignature) override;
 #endif
-	virtual void GetVMExternalFunction(const FVMExternalFunctionBindingInfo& BindingInfo, void* InstanceData, FVMExternalFunction &OutFunc) override;
+	NIAGARA_API virtual void GetVMExternalFunction(const FVMExternalFunctionBindingInfo& BindingInfo, void* InstanceData, FVMExternalFunction &OutFunc) override;
 	virtual bool CanExecuteOnTarget(ENiagaraSimTarget Target) const override { return true; }
-	virtual bool RequiresDistanceFieldData() const override;
+	NIAGARA_API virtual bool RequiresDistanceFieldData() const override;
 	virtual bool HasPreSimulateTick() const override { return true; }
 
 #if WITH_EDITORONLY_DATA
-	virtual bool AppendCompileHash(FNiagaraCompileHashVisitor* InVisitor) const override;
-	virtual void ModifyCompilationEnvironment(EShaderPlatform ShaderPlatform, struct FShaderCompilerEnvironment& OutEnvironment) const override;
-	virtual void GetParameterDefinitionHLSL(const FNiagaraDataInterfaceGPUParamInfo& ParamInfo, FString& OutHLSL) override;
-	virtual bool GetFunctionHLSL(const FNiagaraDataInterfaceGPUParamInfo& ParamInfo, const FNiagaraDataInterfaceGeneratedFunction& FunctionInfo, int FunctionInstanceIndex, FString& OutHLSL) override;
+	NIAGARA_API virtual bool AppendCompileHash(FNiagaraCompileHashVisitor* InVisitor) const override;
+	NIAGARA_API virtual void ModifyCompilationEnvironment(EShaderPlatform ShaderPlatform, struct FShaderCompilerEnvironment& OutEnvironment) const override;
+	NIAGARA_API virtual void GetParameterDefinitionHLSL(const FNiagaraDataInterfaceGPUParamInfo& ParamInfo, FString& OutHLSL) override;
+	NIAGARA_API virtual bool GetFunctionHLSL(const FNiagaraDataInterfaceGPUParamInfo& ParamInfo, const FNiagaraDataInterfaceGeneratedFunction& FunctionInfo, int FunctionInstanceIndex, FString& OutHLSL) override;
 #endif
-	virtual void BuildShaderParameters(FNiagaraShaderParametersBuilder& ShaderParametersBuilder) const override;
-	virtual void SetShaderParameters(const FNiagaraDataInterfaceSetShaderParametersContext& Context) const override;
+	NIAGARA_API virtual void BuildShaderParameters(FNiagaraShaderParametersBuilder& ShaderParametersBuilder) const override;
+	NIAGARA_API virtual void SetShaderParameters(const FNiagaraDataInterfaceSetShaderParametersContext& Context) const override;
 
-	virtual bool Equals(const UNiagaraDataInterface* Other) const override;
+	NIAGARA_API virtual bool Equals(const UNiagaraDataInterface* Other) const override;
 protected:
-	virtual bool CopyToInternal(UNiagaraDataInterface* Destination) const override;
+	NIAGARA_API virtual bool CopyToInternal(UNiagaraDataInterface* Destination) const override;
 public:
 #if WITH_NIAGARA_DEBUGGER
-	virtual void DrawDebugHud(FNDIDrawDebugHudContext& DebugHudContext) const override;
+	NIAGARA_API virtual void DrawDebugHud(FNDIDrawDebugHudContext& DebugHudContext) const override;
 #endif
 #if WITH_EDITOR
-	virtual void GetFeedback(UNiagaraSystem* Asset, UNiagaraComponent* Component, TArray<FNiagaraDataInterfaceError>& OutErrors,
+	NIAGARA_API virtual void GetFeedback(UNiagaraSystem* Asset, UNiagaraComponent* Component, TArray<FNiagaraDataInterfaceError>& OutErrors,
 		TArray<FNiagaraDataInterfaceFeedback>& OutWarnings, TArray<FNiagaraDataInterfaceFeedback>& OutInfo) override;
 #endif
 	//~ UNiagaraDataInterface interface
 
-	UStaticMesh* GetStaticMesh(USceneComponent*& OutComponent, class FNiagaraSystemInstance* SystemInstance = nullptr);
-	void SetSourceComponentFromBlueprints(UStaticMeshComponent* ComponentToUse);
-	void SetDefaultMeshFromBlueprints(UStaticMesh* MeshToUse);
+	NIAGARA_API UStaticMesh* GetStaticMesh(USceneComponent*& OutComponent, class FNiagaraSystemInstance* SystemInstance = nullptr);
+	NIAGARA_API void SetSourceComponentFromBlueprints(UStaticMeshComponent* ComponentToUse);
+	NIAGARA_API void SetDefaultMeshFromBlueprints(UStaticMesh* MeshToUse);
 
 	UFUNCTION(BlueprintCallable, Category = Niagara, meta = (DisplayName = "Set Niagara Static Mesh DI Instance Index"))
-	static void SetNiagaraStaticMeshDIInstanceIndex(UNiagaraComponent* NiagaraSystem, const FName UserParameterName, int32 NewInstanceIndex);
+	static NIAGARA_API void SetNiagaraStaticMeshDIInstanceIndex(UNiagaraComponent* NiagaraSystem, const FName UserParameterName, int32 NewInstanceIndex);
 
 protected:
 	// Bind/unbind delegates to release references to the source actor & component.
-	void UnbindSourceDelegates();
-	void BindSourceDelegates();
+	NIAGARA_API void UnbindSourceDelegates();
+	NIAGARA_API void BindSourceDelegates();
 
 	UFUNCTION()
-	void OnSourceEndPlay(AActor* InSource, EEndPlayReason::Type Reason);
+	NIAGARA_API void OnSourceEndPlay(AActor* InSource, EEndPlayReason::Type Reason);
 
 	// VM Vertex Sampling
-	void VMIsValidVertex(FVectorVMExternalFunctionContext& Context);
-	void VMRandomVertex(FVectorVMExternalFunctionContext& Context);
-	void VMGetVertexCount(FVectorVMExternalFunctionContext& Context);
+	NIAGARA_API void VMIsValidVertex(FVectorVMExternalFunctionContext& Context);
+	NIAGARA_API void VMRandomVertex(FVectorVMExternalFunctionContext& Context);
+	NIAGARA_API void VMGetVertexCount(FVectorVMExternalFunctionContext& Context);
 
 	template<typename TTransformHandler>
 	void VMGetVertex(FVectorVMExternalFunctionContext& Context);
 	template<typename TTransformHandler>
 	void VMGetVertexInterpolated(FVectorVMExternalFunctionContext& Context);
-	void VMGetVertexColor(FVectorVMExternalFunctionContext& Context);
-	void VMGetVertexUV(FVectorVMExternalFunctionContext& Context);
+	NIAGARA_API void VMGetVertexColor(FVectorVMExternalFunctionContext& Context);
+	NIAGARA_API void VMGetVertexUV(FVectorVMExternalFunctionContext& Context);
 
 	// VM Triangle Sampling
-	void VMIsValidTriangle(FVectorVMExternalFunctionContext& Context);
+	NIAGARA_API void VMIsValidTriangle(FVectorVMExternalFunctionContext& Context);
 	template<typename TRandomHelper>
 	void VMRandomTriangle(FVectorVMExternalFunctionContext& Context);
-	void VMGetTriangleCount(FVectorVMExternalFunctionContext& Context);
+	NIAGARA_API void VMGetTriangleCount(FVectorVMExternalFunctionContext& Context);
 
-	void VMIsValidFilteredTriangle(FVectorVMExternalFunctionContext& Context);
-	void VMRandomFilteredTriangle(FVectorVMExternalFunctionContext& Context);
-	void VMGetFilteredTriangleCount(FVectorVMExternalFunctionContext& Context);
-	void VMGetFilteredTriangleAt(FVectorVMExternalFunctionContext& Context);
+	NIAGARA_API void VMIsValidFilteredTriangle(FVectorVMExternalFunctionContext& Context);
+	NIAGARA_API void VMRandomFilteredTriangle(FVectorVMExternalFunctionContext& Context);
+	NIAGARA_API void VMGetFilteredTriangleCount(FVectorVMExternalFunctionContext& Context);
+	NIAGARA_API void VMGetFilteredTriangleAt(FVectorVMExternalFunctionContext& Context);
 
-	void VMIsValidUnfilteredTriangle(FVectorVMExternalFunctionContext& Context);
-	void VMRandomUnfilteredTriangle(FVectorVMExternalFunctionContext& Context);
-	void VMGetUnfilteredTriangleCount(FVectorVMExternalFunctionContext& Context);
-	void VMGetUnfilteredTriangleAt(FVectorVMExternalFunctionContext& Context);
+	NIAGARA_API void VMIsValidUnfilteredTriangle(FVectorVMExternalFunctionContext& Context);
+	NIAGARA_API void VMRandomUnfilteredTriangle(FVectorVMExternalFunctionContext& Context);
+	NIAGARA_API void VMGetUnfilteredTriangleCount(FVectorVMExternalFunctionContext& Context);
+	NIAGARA_API void VMGetUnfilteredTriangleAt(FVectorVMExternalFunctionContext& Context);
 
 	template<typename TTransformHandler>
 	void VMGetTriangle(FVectorVMExternalFunctionContext& Context);
 	template<typename TTransformHandler>
 	void VMGetTriangleInterpolated(FVectorVMExternalFunctionContext& Context);
-	void VMGetTriangleColor(FVectorVMExternalFunctionContext& Context);
-	void VMGetTriangleUV(FVectorVMExternalFunctionContext& Context);
+	NIAGARA_API void VMGetTriangleColor(FVectorVMExternalFunctionContext& Context);
+	NIAGARA_API void VMGetTriangleUV(FVectorVMExternalFunctionContext& Context);
 
-	void VMGetTriangleIndices(FVectorVMExternalFunctionContext& Context);
+	NIAGARA_API void VMGetTriangleIndices(FVectorVMExternalFunctionContext& Context);
 
 	// Socket Functions
-	void VMGetSocketCount(FVectorVMExternalFunctionContext& Context);
-	void VMGetFilteredSocketCount(FVectorVMExternalFunctionContext& Context);
-	void VMGetUnfilteredSocketCount(FVectorVMExternalFunctionContext& Context);
-	void VMRandomSocket(FVectorVMExternalFunctionContext& Context);
-	void VMRandomFilteredSocket(FVectorVMExternalFunctionContext& Context);
-	void VMRandomUnfilteredSocket(FVectorVMExternalFunctionContext& Context);
+	NIAGARA_API void VMGetSocketCount(FVectorVMExternalFunctionContext& Context);
+	NIAGARA_API void VMGetFilteredSocketCount(FVectorVMExternalFunctionContext& Context);
+	NIAGARA_API void VMGetUnfilteredSocketCount(FVectorVMExternalFunctionContext& Context);
+	NIAGARA_API void VMRandomSocket(FVectorVMExternalFunctionContext& Context);
+	NIAGARA_API void VMRandomFilteredSocket(FVectorVMExternalFunctionContext& Context);
+	NIAGARA_API void VMRandomUnfilteredSocket(FVectorVMExternalFunctionContext& Context);
 	template<typename TTransformHandler>
 	void VMGetSocketTransform(FVectorVMExternalFunctionContext& Context);
 	template<typename TTransformHandler>
@@ -255,46 +255,46 @@ protected:
 	void VMGetUnfilteredSocketTransform(FVectorVMExternalFunctionContext& Context);
 	template<typename TTransformHandler>
 	void VMGetUnfilteredSocketTransformInterpolated(FVectorVMExternalFunctionContext& Context);
-	void VMGetFilteredSocket(FVectorVMExternalFunctionContext& Context);
-	void VMGetUnfilteredSocket(FVectorVMExternalFunctionContext& Context);
+	NIAGARA_API void VMGetFilteredSocket(FVectorVMExternalFunctionContext& Context);
+	NIAGARA_API void VMGetUnfilteredSocket(FVectorVMExternalFunctionContext& Context);
 
 	// Section functions
-	void VMIsValidSection(FVectorVMExternalFunctionContext& Context);
-	void VMGetSectionTriangleCount(FVectorVMExternalFunctionContext& Context);
+	NIAGARA_API void VMIsValidSection(FVectorVMExternalFunctionContext& Context);
+	NIAGARA_API void VMGetSectionTriangleCount(FVectorVMExternalFunctionContext& Context);
 	template<typename TRandomHandler>
 	void VMRandomSectionTriangle(FVectorVMExternalFunctionContext& Context);
-	void VMGetSectionTriangleAt(FVectorVMExternalFunctionContext& Context);
+	NIAGARA_API void VMGetSectionTriangleAt(FVectorVMExternalFunctionContext& Context);
 
-	void VMGetFilteredSectionAt(FVectorVMExternalFunctionContext& Context);
-	void VMGetUnfilteredSectionAt(FVectorVMExternalFunctionContext& Context);
+	NIAGARA_API void VMGetFilteredSectionAt(FVectorVMExternalFunctionContext& Context);
+	NIAGARA_API void VMGetUnfilteredSectionAt(FVectorVMExternalFunctionContext& Context);
 
-	void VMGetSectionCount(FVectorVMExternalFunctionContext& Context);
-	void VMGetFilteredSectionCount(FVectorVMExternalFunctionContext& Context);
-	void VMGetUnfilteredSectionCount(FVectorVMExternalFunctionContext& Context);
+	NIAGARA_API void VMGetSectionCount(FVectorVMExternalFunctionContext& Context);
+	NIAGARA_API void VMGetFilteredSectionCount(FVectorVMExternalFunctionContext& Context);
+	NIAGARA_API void VMGetUnfilteredSectionCount(FVectorVMExternalFunctionContext& Context);
 
 	template<typename TRandomHandler>
 	void VMRandomSection(FVectorVMExternalFunctionContext& Context);
 	template<typename TRandomHandler>
 	void VMRandomFilteredSection(FVectorVMExternalFunctionContext& Context);
-	void VMRandomUnfilteredSection(FVectorVMExternalFunctionContext& Context);
+	NIAGARA_API void VMRandomUnfilteredSection(FVectorVMExternalFunctionContext& Context);
 
 	// VM Misc Functions
-	void VMIsValid(FVectorVMExternalFunctionContext& Context);
+	NIAGARA_API void VMIsValid(FVectorVMExternalFunctionContext& Context);
 
-	void VMGetPreSkinnedLocalBounds(FVectorVMExternalFunctionContext& Context);
+	NIAGARA_API void VMGetPreSkinnedLocalBounds(FVectorVMExternalFunctionContext& Context);
 	template<bool bLocalSpace>
 	void VMGetMeshBounds(FVectorVMExternalFunctionContext& Context);
 
-	void VMGetLocalToWorld(FVectorVMExternalFunctionContext& Context);
-	void VMGetLocalToWorldInverseTransposed(FVectorVMExternalFunctionContext& Context);
-	void VMGetWorldVelocity(FVectorVMExternalFunctionContext& Context);
-	void VMGetInstanceIndex(FVectorVMExternalFunctionContext& Context);
-	void VMSetInstanceIndex(FVectorVMExternalFunctionContext& Context);
+	NIAGARA_API void VMGetLocalToWorld(FVectorVMExternalFunctionContext& Context);
+	NIAGARA_API void VMGetLocalToWorldInverseTransposed(FVectorVMExternalFunctionContext& Context);
+	NIAGARA_API void VMGetWorldVelocity(FVectorVMExternalFunctionContext& Context);
+	NIAGARA_API void VMGetInstanceIndex(FVectorVMExternalFunctionContext& Context);
+	NIAGARA_API void VMSetInstanceIndex(FVectorVMExternalFunctionContext& Context);
 
 	// VM UV mapping functions
-	void VMGetTriangleCoordAtUV(FVectorVMExternalFunctionContext& Context);
-	void VMGetTriangleCoordInAabb(FVectorVMExternalFunctionContext& Context);
-	void VMBuildUvMapping(FVectorVMExternalFunctionContext& Context);;
+	NIAGARA_API void VMGetTriangleCoordAtUV(FVectorVMExternalFunctionContext& Context);
+	NIAGARA_API void VMGetTriangleCoordInAabb(FVectorVMExternalFunctionContext& Context);
+	NIAGARA_API void VMBuildUvMapping(FVectorVMExternalFunctionContext& Context);;
 
 	// Deprecated VM Functions
 	template<typename TTransformHandler>
@@ -309,15 +309,15 @@ protected:
 	void VMGetTriangleNormal_Deprecated(FVectorVMExternalFunctionContext& Context);
 };
 
-class NIAGARA_API FNDI_StaticMesh_GeneratedData : public FNDI_GeneratedData
+class FNDI_StaticMesh_GeneratedData : public FNDI_GeneratedData
 {
 	FRWLock CachedUvMappingGuard;
 	TArray<TSharedPtr<FStaticMeshUvMapping>> CachedUvMapping;
 
 public:
-	FStaticMeshUvMappingHandle GetCachedUvMapping(TWeakObjectPtr<UStaticMesh>& InMeshObject, int32 InLodIndex, int32 InUvSetIndex, FMeshUvMappingUsage Usage, bool bNeedsDataImmediately);
+	NIAGARA_API FStaticMeshUvMappingHandle GetCachedUvMapping(TWeakObjectPtr<UStaticMesh>& InMeshObject, int32 InLodIndex, int32 InUvSetIndex, FMeshUvMappingUsage Usage, bool bNeedsDataImmediately);
 
-	virtual void Tick(ETickingGroup TickGroup, float DeltaSeconds) override;
+	NIAGARA_API virtual void Tick(ETickingGroup TickGroup, float DeltaSeconds) override;
 
 	static TypeHash GetTypeHash()
 	{
diff --git a/Engine/Plugins/FX/Niagara/Source/Niagara/Internal/NiagaraDataChannelCommon.h b/Engine/Plugins/FX/Niagara/Source/Niagara/Internal/NiagaraDataChannelCommon.h
--- a/Engine/Plugins/FX/Niagara/Source/Niagara/Internal/NiagaraDataChannelCommon.h
+++ b/Engine/Plugins/FX/Niagara/Source/Niagara/Internal/NiagaraDataChannelCommon.h
@@ -61,25 +61,25 @@ Some data channels will have many of these and can distribute them as needed to
 For example, some data channel handlers may subdivide the scene such that distant systems are not interacting.
 In this case, each subdivision would have it's own FNiagaraDataChannelData and distribute these to the relevant NiagaraSystems.
 */
-struct NIAGARA_API FNiagaraDataChannelData final
+struct FNiagaraDataChannelData final
 {
 	UE_NONCOPYABLE(FNiagaraDataChannelData)
-	explicit FNiagaraDataChannelData(UNiagaraDataChannelHandler* Owner);
-	~FNiagaraDataChannelData();
+	NIAGARA_API explicit FNiagaraDataChannelData(UNiagaraDataChannelHandler* Owner);
+	NIAGARA_API ~FNiagaraDataChannelData();
 
-	void Reset();
+	NIAGARA_API void Reset();
 
-	void BeginFrame(UNiagaraDataChannelHandler* Owner);
-	void EndFrame(UNiagaraDataChannelHandler* Owner);
-	void Tick(UNiagaraDataChannelHandler* Owner);
-	void ConsumePublishRequests(UNiagaraDataChannelHandler* Owner);
+	NIAGARA_API void BeginFrame(UNiagaraDataChannelHandler* Owner);
+	NIAGARA_API void EndFrame(UNiagaraDataChannelHandler* Owner);
+	NIAGARA_API void Tick(UNiagaraDataChannelHandler* Owner);
+	NIAGARA_API void ConsumePublishRequests(UNiagaraDataChannelHandler* Owner);
 
-	FNiagaraDataChannelGameData* GetGameData();
-	FNiagaraDataBufferRef GetCPUData(bool bPreviousFrame);
+	NIAGARA_API FNiagaraDataChannelGameData* GetGameData();
+	NIAGARA_API FNiagaraDataBufferRef GetCPUData(bool bPreviousFrame);
 	FNiagaraDataChannelDataProxy* GetRTProxy(){ return RTProxy.Get(); }
 	
 	/** Adds a request to publish some data into the channel on the next tick. */
-	void Publish(const FNiagaraDataChannelPublishRequest& Request);
+	NIAGARA_API void Publish(const FNiagaraDataChannelPublishRequest& Request);
 		
 	/**
 	 *Removes all publish requests involving the given dataset.
@@ -87,9 +87,9 @@ struct NIAGARA_API FNiagaraDataChannelData final
 	 *This is a hack to get around lifetime issues wrt data buffers/datasets and their compiled data.
 	 *We should rework things such that data buffers can exist beyond their owning dataset, including the compiled data detailing their layout.
 	 **/
-	void RemovePublishRequests(const FNiagaraDataSet* DataSet);
+	NIAGARA_API void RemovePublishRequests(const FNiagaraDataSet* DataSet);
 
-	const FNiagaraDataSetCompiledData& GetCompiledData(ENiagaraSimTarget SimTarget);
+	NIAGARA_API const FNiagaraDataSetCompiledData& GetCompiledData(ENiagaraSimTarget SimTarget);
 
 	void SetLwcTile(FVector3f InLwcTile){ LwcTile = InLwcTile; }
 	FVector3f GetLwcTile()const { return LwcTile; }
@@ -137,4 +137,4 @@ private:
 
 	/** Critical section protecting shared state for multiple writers publishing from different threads. */
 	FCriticalSection PublishCritSec;
-};
\ No newline at end of file
+};
diff --git a/Engine/Plugins/FX/Niagara/Source/Niagara/Internal/NiagaraPrecompileContainer.h b/Engine/Plugins/FX/Niagara/Source/Niagara/Internal/NiagaraPrecompileContainer.h
--- a/Engine/Plugins/FX/Niagara/Source/Niagara/Internal/NiagaraPrecompileContainer.h
+++ b/Engine/Plugins/FX/Niagara/Source/Niagara/Internal/NiagaraPrecompileContainer.h
@@ -6,8 +6,8 @@
 #include "NiagaraPrecompileContainer.generated.h"
 
 
-UCLASS()
-class NIAGARA_API UNiagaraPrecompileContainer : public UObject
+UCLASS(MinimalAPI)
+class UNiagaraPrecompileContainer : public UObject
 {
 	GENERATED_UCLASS_BODY()
 
@@ -17,4 +17,4 @@ public:
 
 	UPROPERTY()
 	TObjectPtr<UNiagaraSystem> System;
-};
\ No newline at end of file
+};
diff --git a/Engine/Plugins/FX/Niagara/Source/Niagara/Private/DataInterface/NiagaraDataInterfaceDebugDraw.h b/Engine/Plugins/FX/Niagara/Source/Niagara/Private/DataInterface/NiagaraDataInterfaceDebugDraw.h
--- a/Engine/Plugins/FX/Niagara/Source/Niagara/Private/DataInterface/NiagaraDataInterfaceDebugDraw.h
+++ b/Engine/Plugins/FX/Niagara/Source/Niagara/Private/DataInterface/NiagaraDataInterfaceDebugDraw.h
@@ -12,8 +12,8 @@ struct FNiagaraFunctionSignature;
 struct FNiagaraScriptDataInterfaceCompileInfo;
 struct FVMExternalFunctionBindingInfo;
 
-UCLASS(EditInlineNew, Category = "Debug", CollapseCategories, meta = (DisplayName = "DebugDraw"))
-class NIAGARA_API UNiagaraDataInterfaceDebugDraw : public UNiagaraDataInterface
+UCLASS(EditInlineNew, Category = "Debug", CollapseCategories, meta = (DisplayName = "DebugDraw"), MinimalAPI)
+class UNiagaraDataInterfaceDebugDraw : public UNiagaraDataInterface
 {
 	GENERATED_UCLASS_BODY()
 
@@ -23,13 +23,13 @@ public:
 	uint32 OverrideMaxLineInstances = 0;
 
 	//UObject Interface
-	virtual void PostInitProperties() override;
+	NIAGARA_API virtual void PostInitProperties() override;
 	//UObject Interface End
 
 	//UNiagaraDataInterface Interface
 	virtual bool CanExecuteOnTarget(ENiagaraSimTarget Target) const override { return true; }
-	virtual void GetFunctions(TArray<FNiagaraFunctionSignature>& OutFunctions) override;
-	virtual void GetVMExternalFunction(const FVMExternalFunctionBindingInfo& BindingInfo, void* InstanceData, FVMExternalFunction& OutFunc) override;
+	NIAGARA_API virtual void GetFunctions(TArray<FNiagaraFunctionSignature>& OutFunctions) override;
+	NIAGARA_API virtual void GetVMExternalFunction(const FVMExternalFunctionBindingInfo& BindingInfo, void* InstanceData, FVMExternalFunction& OutFunc) override;
 
 	enum ShapeId
 	{
@@ -48,33 +48,33 @@ public:
 	};
 
 #if WITH_EDITORONLY_DATA
-	virtual bool AppendCompileHash(FNiagaraCompileHashVisitor* InVisitor) const override;
-	virtual void GetCommonHLSL(FString& OutHLSL) override;
-	virtual void GetParameterDefinitionHLSL(const FNiagaraDataInterfaceGPUParamInfo& ParamInfo, FString& OutHLSL) override;
-	virtual bool GetFunctionHLSL(const FNiagaraDataInterfaceGPUParamInfo& ParamInfo, const FNiagaraDataInterfaceGeneratedFunction& FunctionInfo, int FunctionInstanceIndex, FString& OutHLSL) override;
-	virtual bool GenerateCompilerTagPrefix(const FNiagaraFunctionSignature& InSignature, FString& OutPrefix) const override;	
-	virtual bool UpgradeFunctionCall(FNiagaraFunctionSignature& FunctionSignature) override;
+	NIAGARA_API virtual bool AppendCompileHash(FNiagaraCompileHashVisitor* InVisitor) const override;
+	NIAGARA_API virtual void GetCommonHLSL(FString& OutHLSL) override;
+	NIAGARA_API virtual void GetParameterDefinitionHLSL(const FNiagaraDataInterfaceGPUParamInfo& ParamInfo, FString& OutHLSL) override;
+	NIAGARA_API virtual bool GetFunctionHLSL(const FNiagaraDataInterfaceGPUParamInfo& ParamInfo, const FNiagaraDataInterfaceGeneratedFunction& FunctionInfo, int FunctionInstanceIndex, FString& OutHLSL) override;
+	NIAGARA_API virtual bool GenerateCompilerTagPrefix(const FNiagaraFunctionSignature& InSignature, FString& OutPrefix) const override;	
+	NIAGARA_API virtual bool UpgradeFunctionCall(FNiagaraFunctionSignature& FunctionSignature) override;
 #endif
-	virtual void BuildShaderParameters(FNiagaraShaderParametersBuilder& ShaderParametersBuilder) const override;
-	virtual void SetShaderParameters(const FNiagaraDataInterfaceSetShaderParametersContext& Context) const override;
+	NIAGARA_API virtual void BuildShaderParameters(FNiagaraShaderParametersBuilder& ShaderParametersBuilder) const override;
+	NIAGARA_API virtual void SetShaderParameters(const FNiagaraDataInterfaceSetShaderParametersContext& Context) const override;
 
-	virtual bool PerInstanceTick(void* PerInstanceData, FNiagaraSystemInstance* SystemInstance, float DeltaSeconds) override;
-	virtual bool PerInstanceTickPostSimulate(void* PerInstanceData, FNiagaraSystemInstance* SystemInstance, float DeltaSeconds) override;
+	NIAGARA_API virtual bool PerInstanceTick(void* PerInstanceData, FNiagaraSystemInstance* SystemInstance, float DeltaSeconds) override;
+	NIAGARA_API virtual bool PerInstanceTickPostSimulate(void* PerInstanceData, FNiagaraSystemInstance* SystemInstance, float DeltaSeconds) override;
 
 	virtual bool HasPreSimulateTick() const override { return true; }
 	virtual bool HasPostSimulateTick() const override { return true; }
 
-	int32 PerInstanceDataSize() const;
-	bool InitPerInstanceData(void* PerInstanceData, FNiagaraSystemInstance* SystemInstance);
-	void DestroyPerInstanceData(void* PerInstanceData, FNiagaraSystemInstance* SystemInstance);
+	NIAGARA_API int32 PerInstanceDataSize() const;
+	NIAGARA_API bool InitPerInstanceData(void* PerInstanceData, FNiagaraSystemInstance* SystemInstance);
+	NIAGARA_API void DestroyPerInstanceData(void* PerInstanceData, FNiagaraSystemInstance* SystemInstance);
 
 	virtual bool NeedsGPUContextInit() const override { return true; }
-	virtual bool GPUContextInit(const FNiagaraScriptDataInterfaceCompileInfo& InInfo, void* PerInstanceData, FNiagaraSystemInstance* SystemInstance) const override;
+	NIAGARA_API virtual bool GPUContextInit(const FNiagaraScriptDataInterfaceCompileInfo& InInfo, void* PerInstanceData, FNiagaraSystemInstance* SystemInstance) const override;
 
-	virtual bool Equals(const UNiagaraDataInterface* Other) const override;
+	NIAGARA_API virtual bool Equals(const UNiagaraDataInterface* Other) const override;
 
 protected:
 	/** Copy one niagara DI to this */
-	virtual bool CopyToInternal(UNiagaraDataInterface* Destination) const override;
+	NIAGARA_API virtual bool CopyToInternal(UNiagaraDataInterface* Destination) const override;
 	//UNiagaraDataInterface Interface
 };
diff --git a/Engine/Plugins/FX/Niagara/Source/Niagara/Private/DataInterface/NiagaraDataInterfaceDynamicMesh.h b/Engine/Plugins/FX/Niagara/Source/Niagara/Private/DataInterface/NiagaraDataInterfaceDynamicMesh.h
--- a/Engine/Plugins/FX/Niagara/Source/Niagara/Private/DataInterface/NiagaraDataInterfaceDynamicMesh.h
+++ b/Engine/Plugins/FX/Niagara/Source/Niagara/Private/DataInterface/NiagaraDataInterfaceDynamicMesh.h
@@ -49,8 +49,8 @@ struct FNiagaraDynamicMeshMaterial
 };
 
 /** Data Interface allowing sampling of a texture */
-UCLASS(Experimental, EditInlineNew, Category = "Mesh", CollapseCategories, meta = (DisplayName = "Dynamic Mesh"))
-class NIAGARA_API UNiagaraDataInterfaceDynamicMesh : public UNiagaraDataInterface, public INiagaraRenderableMeshInterface
+UCLASS(Experimental, EditInlineNew, Category = "Mesh", CollapseCategories, meta = (DisplayName = "Dynamic Mesh"), MinimalAPI)
+class UNiagaraDataInterfaceDynamicMesh : public UNiagaraDataInterface, public INiagaraRenderableMeshInterface
 {
 	GENERATED_UCLASS_BODY()
 
@@ -91,69 +91,69 @@ public:
 	FBox DefaultBounds = FBox(FVector(-100.0f), FVector(100.0f));
 
 	//UObject Interface
-	virtual void PostInitProperties() override;
+	NIAGARA_API virtual void PostInitProperties() override;
 	//UObject Interface End
 
 	//UNiagaraDataInterface Interface
-	virtual void GetFunctions(TArray<FNiagaraFunctionSignature>& OutFunctions) override;
-	virtual void GetVMExternalFunction(const FVMExternalFunctionBindingInfo& BindingInfo, void* InstanceData, FVMExternalFunction &OutFunc) override;
+	NIAGARA_API virtual void GetFunctions(TArray<FNiagaraFunctionSignature>& OutFunctions) override;
+	NIAGARA_API virtual void GetVMExternalFunction(const FVMExternalFunctionBindingInfo& BindingInfo, void* InstanceData, FVMExternalFunction &OutFunc) override;
 	virtual bool CanExecuteOnTarget(ENiagaraSimTarget Target) const override { return true; }
 
-	virtual int32 PerInstanceDataSize() const override;
-	virtual bool InitPerInstanceData(void* PerInstanceData, FNiagaraSystemInstance* SystemInstance) override;
-	virtual void DestroyPerInstanceData(void* PerInstanceData, FNiagaraSystemInstance* SystemInstance) override;
+	NIAGARA_API virtual int32 PerInstanceDataSize() const override;
+	NIAGARA_API virtual bool InitPerInstanceData(void* PerInstanceData, FNiagaraSystemInstance* SystemInstance) override;
+	NIAGARA_API virtual void DestroyPerInstanceData(void* PerInstanceData, FNiagaraSystemInstance* SystemInstance) override;
 
 	virtual bool HasPreSimulateTick() const override { return true; }
 	virtual bool HasPostSimulateTick() const override { return true; }
-	virtual bool PerInstanceTick(void* PerInstanceData, FNiagaraSystemInstance* SystemInstance, float DeltaSeconds) override;
-	virtual bool PerInstanceTickPostSimulate(void* PerInstanceData, FNiagaraSystemInstance* SystemInstance, float DeltaSeconds) override;
+	NIAGARA_API virtual bool PerInstanceTick(void* PerInstanceData, FNiagaraSystemInstance* SystemInstance, float DeltaSeconds) override;
+	NIAGARA_API virtual bool PerInstanceTickPostSimulate(void* PerInstanceData, FNiagaraSystemInstance* SystemInstance, float DeltaSeconds) override;
 
-	virtual bool Equals(const UNiagaraDataInterface* Other) const override;
-	virtual bool CopyToInternal(UNiagaraDataInterface* Destination) const override;
+	NIAGARA_API virtual bool Equals(const UNiagaraDataInterface* Other) const override;
+	NIAGARA_API virtual bool CopyToInternal(UNiagaraDataInterface* Destination) const override;
 
 	// GPU sim functionality
 #if WITH_EDITORONLY_DATA
-	virtual bool AppendCompileHash(FNiagaraCompileHashVisitor* InVisitor) const override;
-	virtual void GetParameterDefinitionHLSL(const FNiagaraDataInterfaceGPUParamInfo& ParamInfo, FString& OutHLSL) override;
-	virtual bool GetFunctionHLSL(const FNiagaraDataInterfaceGPUParamInfo& ParamInfo, const FNiagaraDataInterfaceGeneratedFunction& FunctionInfo, int FunctionInstanceIndex, FString& OutHLSL) override;
+	NIAGARA_API virtual bool AppendCompileHash(FNiagaraCompileHashVisitor* InVisitor) const override;
+	NIAGARA_API virtual void GetParameterDefinitionHLSL(const FNiagaraDataInterfaceGPUParamInfo& ParamInfo, FString& OutHLSL) override;
+	NIAGARA_API virtual bool GetFunctionHLSL(const FNiagaraDataInterfaceGPUParamInfo& ParamInfo, const FNiagaraDataInterfaceGeneratedFunction& FunctionInfo, int FunctionInstanceIndex, FString& OutHLSL) override;
 #endif
-	virtual void BuildShaderParameters(FNiagaraShaderParametersBuilder& ShaderParametersBuilder) const override;
-	virtual void SetShaderParameters(const FNiagaraDataInterfaceSetShaderParametersContext& Context) const override;
+	NIAGARA_API virtual void BuildShaderParameters(FNiagaraShaderParametersBuilder& ShaderParametersBuilder) const override;
+	NIAGARA_API virtual void SetShaderParameters(const FNiagaraDataInterfaceSetShaderParametersContext& Context) const override;
 	//UNiagaraDataInterface Interface End
 
 	//INiagaraRenderableMeshInterface Interface Begin
-	virtual FNiagaraRenderableMeshPtr GetRenderableMesh(FNiagaraSystemInstanceID SystemInstanceID) override;
-	virtual void GetUsedMaterials(FNiagaraSystemInstanceID SystemInstanceID, TArray<UMaterialInterface*>& OutMaterials) const override;
+	NIAGARA_API virtual FNiagaraRenderableMeshPtr GetRenderableMesh(FNiagaraSystemInstanceID SystemInstanceID) override;
+	NIAGARA_API virtual void GetUsedMaterials(FNiagaraSystemInstanceID SystemInstanceID, TArray<UMaterialInterface*>& OutMaterials) const override;
 	//INiagaraRenderableMeshInterface Interface End
 
 	// VM Functions
-	static void VMGetMeshProperties(FVectorVMExternalFunctionContext& Context);
-	static void VMGetSectionCount(FVectorVMExternalFunctionContext& Context);
-	static void VMGetSectionData(FVectorVMExternalFunctionContext& Context);
-	static void VMGetLocalBounds(FVectorVMExternalFunctionContext& Context);
-
-	static void VMSetMeshProperties(FVectorVMExternalFunctionContext& Context);
-	static void VMSetSectionCount(FVectorVMExternalFunctionContext& Context);
-	static void VMSetSectionData(FVectorVMExternalFunctionContext& Context);
-	static void VMSetLocalBounds(FVectorVMExternalFunctionContext& Context);
-
-	static void VMClearAllSectionTriangles(FVectorVMExternalFunctionContext& Context);
-	static void VMClearSectionTriangles(FVectorVMExternalFunctionContext& Context);
-	static void VMAllocateSectionTriangles(FVectorVMExternalFunctionContext& Context);
-
-	static void VMGetTriangleVertices(FVectorVMExternalFunctionContext& Context);
-	static void VMGetVertexPosition(FVectorVMExternalFunctionContext& Context);
-	static void VMGetVertexTangentBasis(FVectorVMExternalFunctionContext& Context);
-	static void VMGetVertexTexCoord(FVectorVMExternalFunctionContext& Context);
-	static void VMGetVertexColor(FVectorVMExternalFunctionContext& Context);
-	static void VMGetVertexData(FVectorVMExternalFunctionContext& Context);
-
-	static void VMSetTriangleVertices(FVectorVMExternalFunctionContext& Context);
-	static void VMSetVertexPosition(FVectorVMExternalFunctionContext& Context);
-	static void VMSetVertexTangentBasis(FVectorVMExternalFunctionContext& Context);
-	static void VMSetVertexTexCoord(FVectorVMExternalFunctionContext& Context);
-	static void VMSetVertexColor(FVectorVMExternalFunctionContext& Context);
-	static void VMSetVertexData(FVectorVMExternalFunctionContext& Context);
-
-	static void VMAppendTriangle(FVectorVMExternalFunctionContext& Context);
+	static NIAGARA_API void VMGetMeshProperties(FVectorVMExternalFunctionContext& Context);
+	static NIAGARA_API void VMGetSectionCount(FVectorVMExternalFunctionContext& Context);
+	static NIAGARA_API void VMGetSectionData(FVectorVMExternalFunctionContext& Context);
+	static NIAGARA_API void VMGetLocalBounds(FVectorVMExternalFunctionContext& Context);
+
+	static NIAGARA_API void VMSetMeshProperties(FVectorVMExternalFunctionContext& Context);
+	static NIAGARA_API void VMSetSectionCount(FVectorVMExternalFunctionContext& Context);
+	static NIAGARA_API void VMSetSectionData(FVectorVMExternalFunctionContext& Context);
+	static NIAGARA_API void VMSetLocalBounds(FVectorVMExternalFunctionContext& Context);
+
+	static NIAGARA_API void VMClearAllSectionTriangles(FVectorVMExternalFunctionContext& Context);
+	static NIAGARA_API void VMClearSectionTriangles(FVectorVMExternalFunctionContext& Context);
+	static NIAGARA_API void VMAllocateSectionTriangles(FVectorVMExternalFunctionContext& Context);
+
+	static NIAGARA_API void VMGetTriangleVertices(FVectorVMExternalFunctionContext& Context);
+	static NIAGARA_API void VMGetVertexPosition(FVectorVMExternalFunctionContext& Context);
+	static NIAGARA_API void VMGetVertexTangentBasis(FVectorVMExternalFunctionContext& Context);
+	static NIAGARA_API void VMGetVertexTexCoord(FVectorVMExternalFunctionContext& Context);
+	static NIAGARA_API void VMGetVertexColor(FVectorVMExternalFunctionContext& Context);
+	static NIAGARA_API void VMGetVertexData(FVectorVMExternalFunctionContext& Context);
+
+	static NIAGARA_API void VMSetTriangleVertices(FVectorVMExternalFunctionContext& Context);
+	static NIAGARA_API void VMSetVertexPosition(FVectorVMExternalFunctionContext& Context);
+	static NIAGARA_API void VMSetVertexTangentBasis(FVectorVMExternalFunctionContext& Context);
+	static NIAGARA_API void VMSetVertexTexCoord(FVectorVMExternalFunctionContext& Context);
+	static NIAGARA_API void VMSetVertexColor(FVectorVMExternalFunctionContext& Context);
+	static NIAGARA_API void VMSetVertexData(FVectorVMExternalFunctionContext& Context);
+
+	static NIAGARA_API void VMAppendTriangle(FVectorVMExternalFunctionContext& Context);
 };
diff --git a/Engine/Plugins/FX/Niagara/Source/Niagara/Private/DataInterface/NiagaraDataInterfaceGBuffer.h b/Engine/Plugins/FX/Niagara/Source/Niagara/Private/DataInterface/NiagaraDataInterfaceGBuffer.h
--- a/Engine/Plugins/FX/Niagara/Source/Niagara/Private/DataInterface/NiagaraDataInterfaceGBuffer.h
+++ b/Engine/Plugins/FX/Niagara/Source/Niagara/Private/DataInterface/NiagaraDataInterfaceGBuffer.h
@@ -8,27 +8,27 @@
 
 struct FNiagaraDataInterfaceGeneratedFunction;
 
-UCLASS(EditInlineNew, Category = "Rendering", meta = (DisplayName = "GBuffer"))
-class NIAGARA_API UNiagaraDataInterfaceGBuffer : public UNiagaraDataInterface
+UCLASS(EditInlineNew, Category = "Rendering", meta = (DisplayName = "GBuffer"), MinimalAPI)
+class UNiagaraDataInterfaceGBuffer : public UNiagaraDataInterface
 {
 	GENERATED_UCLASS_BODY()
 
 public:
 	//UObject Interface
-	virtual void PostInitProperties() override;
+	NIAGARA_API virtual void PostInitProperties() override;
 	//UObject Interface End
 
 	//UNiagaraDataInterface Interface
-	virtual void GetFunctions(TArray<FNiagaraFunctionSignature>& OutFunctions)override;
+	NIAGARA_API virtual void GetFunctions(TArray<FNiagaraFunctionSignature>& OutFunctions)override;
 
-	virtual void BuildShaderParameters(FNiagaraShaderParametersBuilder& ShaderParametersBuilder) const override;
-	virtual void SetShaderParameters(const FNiagaraDataInterfaceSetShaderParametersContext& Context) const override;
+	NIAGARA_API virtual void BuildShaderParameters(FNiagaraShaderParametersBuilder& ShaderParametersBuilder) const override;
+	NIAGARA_API virtual void SetShaderParameters(const FNiagaraDataInterfaceSetShaderParametersContext& Context) const override;
 
 #if WITH_EDITORONLY_DATA
-	virtual bool AppendCompileHash(FNiagaraCompileHashVisitor* InVisitor) const override;
-	virtual void GetParameterDefinitionHLSL(const FNiagaraDataInterfaceGPUParamInfo& ParamInfo, FString& OutHLSL) override;
-	virtual bool GetFunctionHLSL(const FNiagaraDataInterfaceGPUParamInfo& ParamInfo, const FNiagaraDataInterfaceGeneratedFunction& FunctionInfo, int FunctionInstanceIndex, FString& OutHLSL) override;
-	virtual bool UpgradeFunctionCall(FNiagaraFunctionSignature& FunctionSignature);
+	NIAGARA_API virtual bool AppendCompileHash(FNiagaraCompileHashVisitor* InVisitor) const override;
+	NIAGARA_API virtual void GetParameterDefinitionHLSL(const FNiagaraDataInterfaceGPUParamInfo& ParamInfo, FString& OutHLSL) override;
+	NIAGARA_API virtual bool GetFunctionHLSL(const FNiagaraDataInterfaceGPUParamInfo& ParamInfo, const FNiagaraDataInterfaceGeneratedFunction& FunctionInfo, int FunctionInstanceIndex, FString& OutHLSL) override;
+	NIAGARA_API virtual bool UpgradeFunctionCall(FNiagaraFunctionSignature& FunctionSignature);
 #endif
 
 	virtual bool CanExecuteOnTarget(ENiagaraSimTarget Target) const override { return Target == ENiagaraSimTarget::GPUComputeSim; }
