[Niagara]

diff --git a/Engine/Plugins/FX/Niagara/Source/Niagara/Private/DataInterface/NiagaraDataInterfaceSceneCapture2D.h b/Engine/Plugins/FX/Niagara/Source/Niagara/Private/DataInterface/NiagaraDataInterfaceSceneCapture2D.h
--- a/Engine/Plugins/FX/Niagara/Source/Niagara/Private/DataInterface/NiagaraDataInterfaceSceneCapture2D.h
+++ b/Engine/Plugins/FX/Niagara/Source/Niagara/Private/DataInterface/NiagaraDataInterfaceSceneCapture2D.h
@@ -38,8 +38,8 @@ enum class ENDISceneCapture2DOffsetMode : uint8
 };
 
 /** Data Interface which can control or read from a scene capture component. */
-UCLASS(Experimental, EditInlineNew, Category = "SceneCapture", meta = (DisplayName = "Scene Capture 2D"))
-class NIAGARA_API UNiagaraDataInterfaceSceneCapture2D : public UNiagaraDataInterface
+UCLASS(Experimental, EditInlineNew, Category = "SceneCapture", meta = (DisplayName = "Scene Capture 2D"), MinimalAPI)
+class UNiagaraDataInterfaceSceneCapture2D : public UNiagaraDataInterface
 {
 	GENERATED_UCLASS_BODY()
 
@@ -101,50 +101,50 @@ public:
 	uint32 ChangeId = 0;
 
 	//UObject Interface
-	virtual void PostInitProperties() override;
+	NIAGARA_API virtual void PostInitProperties() override;
 	//UObject Interface End
 
 	//UNiagaraDataInterface Interface
-	virtual void GetFunctions(TArray<FNiagaraFunctionSignature>& OutFunctions) override;
-	virtual void GetVMExternalFunction(const FVMExternalFunctionBindingInfo& BindingInfo, void* InstanceData, FVMExternalFunction &OutFunc) override;
+	NIAGARA_API virtual void GetFunctions(TArray<FNiagaraFunctionSignature>& OutFunctions) override;
+	NIAGARA_API virtual void GetVMExternalFunction(const FVMExternalFunctionBindingInfo& BindingInfo, void* InstanceData, FVMExternalFunction &OutFunc) override;
 	virtual bool CanExecuteOnTarget(ENiagaraSimTarget Target) const override { return Target==ENiagaraSimTarget::GPUComputeSim; }
 
-	virtual int32 PerInstanceDataSize() const override;
-	virtual bool InitPerInstanceData(void* PerInstanceData, FNiagaraSystemInstance* SystemInstance) override;
-	virtual void DestroyPerInstanceData(void* PerInstanceData, FNiagaraSystemInstance* SystemInstance) override;
+	NIAGARA_API virtual int32 PerInstanceDataSize() const override;
+	NIAGARA_API virtual bool InitPerInstanceData(void* PerInstanceData, FNiagaraSystemInstance* SystemInstance) override;
+	NIAGARA_API virtual void DestroyPerInstanceData(void* PerInstanceData, FNiagaraSystemInstance* SystemInstance) override;
 
 	virtual bool HasPreSimulateTick() const override { return true; }
 	virtual bool HasPostSimulateTick() const override { return true; }
-	virtual bool PerInstanceTick(void* PerInstanceData, FNiagaraSystemInstance* SystemInstance, float DeltaSeconds) override;
-	virtual bool PerInstanceTickPostSimulate(void* PerInstanceData, FNiagaraSystemInstance* SystemInstance, float DeltaSeconds) override;
+	NIAGARA_API virtual bool PerInstanceTick(void* PerInstanceData, FNiagaraSystemInstance* SystemInstance, float DeltaSeconds) override;
+	NIAGARA_API virtual bool PerInstanceTickPostSimulate(void* PerInstanceData, FNiagaraSystemInstance* SystemInstance, float DeltaSeconds) override;
 
-	virtual void ProvidePerInstanceDataForRenderThread(void* DataForRenderThread, void* PerInstanceData, const FNiagaraSystemInstanceID& SystemInstance) override;
+	NIAGARA_API virtual void ProvidePerInstanceDataForRenderThread(void* DataForRenderThread, void* PerInstanceData, const FNiagaraSystemInstanceID& SystemInstance) override;
 
 #if WITH_NIAGARA_DEBUGGER
-	virtual void DrawDebugHud(FNDIDrawDebugHudContext& DebugHudContext) const override;
+	NIAGARA_API virtual void DrawDebugHud(FNDIDrawDebugHudContext& DebugHudContext) const override;
 #endif
 
 	// GPU sim functionality
 #if WITH_EDITORONLY_DATA
-	virtual bool AppendCompileHash(FNiagaraCompileHashVisitor* InVisitor) const override;
-	virtual void GetParameterDefinitionHLSL(const FNiagaraDataInterfaceGPUParamInfo& ParamInfo, FString& OutHLSL) override;
-	virtual bool GetFunctionHLSL(const FNiagaraDataInterfaceGPUParamInfo& ParamInfo, const FNiagaraDataInterfaceGeneratedFunction& FunctionInfo, int FunctionInstanceIndex, FString& OutHLSL) override;
+	NIAGARA_API virtual bool AppendCompileHash(FNiagaraCompileHashVisitor* InVisitor) const override;
+	NIAGARA_API virtual void GetParameterDefinitionHLSL(const FNiagaraDataInterfaceGPUParamInfo& ParamInfo, FString& OutHLSL) override;
+	NIAGARA_API virtual bool GetFunctionHLSL(const FNiagaraDataInterfaceGPUParamInfo& ParamInfo, const FNiagaraDataInterfaceGeneratedFunction& FunctionInfo, int FunctionInstanceIndex, FString& OutHLSL) override;
 #endif
-	virtual void BuildShaderParameters(FNiagaraShaderParametersBuilder& ShaderParametersBuilder) const override;
-	virtual void SetShaderParameters(const FNiagaraDataInterfaceSetShaderParametersContext& Context) const override;
+	NIAGARA_API virtual void BuildShaderParameters(FNiagaraShaderParametersBuilder& ShaderParametersBuilder) const override;
+	NIAGARA_API virtual void SetShaderParameters(const FNiagaraDataInterfaceSetShaderParametersContext& Context) const override;
 
 protected:
-	virtual bool Equals(const UNiagaraDataInterface* Other) const override;
-	virtual bool CopyToInternal(UNiagaraDataInterface* Destination) const override;
+	NIAGARA_API virtual bool Equals(const UNiagaraDataInterface* Other) const override;
+	NIAGARA_API virtual bool CopyToInternal(UNiagaraDataInterface* Destination) const override;
 
 	virtual bool CanExposeVariables() const override { return true; }
-	virtual void GetExposedVariables(TArray<FNiagaraVariableBase>& OutVariables) const override;
-	virtual bool GetExposedVariableValue(const FNiagaraVariableBase& InVariable, void* InPerInstanceData, FNiagaraSystemInstance* InSystemInstance, void* OutData) const override;
+	NIAGARA_API virtual void GetExposedVariables(TArray<FNiagaraVariableBase>& OutVariables) const override;
+	NIAGARA_API virtual bool GetExposedVariableValue(const FNiagaraVariableBase& InVariable, void* InPerInstanceData, FNiagaraSystemInstance* InSystemInstance, void* OutData) const override;
 	//UNiagaraDataInterface Interface End
 
 	/** Allows you to set the show only actors when Niagara manages the component.  If Niagara does not manage the component use the regular BP methods. */
 	UFUNCTION(BlueprintCallable, Category = Niagara)
-	static void SetSceneCapture2DManagedShowOnlyActors(UNiagaraComponent* NiagaraSystem, const FName ParameterName, TArray<AActor*> ShowOnlyActors);
+	static NIAGARA_API void SetSceneCapture2DManagedShowOnlyActors(UNiagaraComponent* NiagaraSystem, const FName ParameterName, TArray<AActor*> ShowOnlyActors);
 
 	UPROPERTY(Transient, DuplicateTransient)
 	TMap<uint64, TObjectPtr<USceneCaptureComponent2D>>	ManagedCaptureComponents;
diff --git a/Engine/Plugins/FX/Niagara/Source/Niagara/Private/DataInterface/NiagaraDataInterfaceSimpleCounter.h b/Engine/Plugins/FX/Niagara/Source/Niagara/Private/DataInterface/NiagaraDataInterfaceSimpleCounter.h
--- a/Engine/Plugins/FX/Niagara/Source/Niagara/Private/DataInterface/NiagaraDataInterfaceSimpleCounter.h
+++ b/Engine/Plugins/FX/Niagara/Source/Niagara/Private/DataInterface/NiagaraDataInterfaceSimpleCounter.h
@@ -14,8 +14,8 @@ struct FNiagaraDataInterfaceGeneratedFunction;
 Thread safe counter starts at the initial value on start / reset.
 When operating between CPU & GPU ensure you set the appropriate sync mode.
 */
-UCLASS(EditInlineNew, Category = "Counting", CollapseCategories, meta = (DisplayName = "Simple Counter"))
-class NIAGARA_API UNiagaraDataInterfaceSimpleCounter : public UNiagaraDataInterfaceRWBase
+UCLASS(EditInlineNew, Category = "Counting", CollapseCategories, meta = (DisplayName = "Simple Counter"), MinimalAPI)
+class UNiagaraDataInterfaceSimpleCounter : public UNiagaraDataInterfaceRWBase
 {
 	GENERATED_UCLASS_BODY()
 
@@ -25,46 +25,46 @@ class NIAGARA_API UNiagaraDataInterfaceSimpleCounter : public UNiagaraDataInterf
 
 public:
 	// UObject Interface
-	virtual void PostInitProperties() override;
+	NIAGARA_API virtual void PostInitProperties() override;
 #if WITH_EDITOR
-	virtual void PostEditChangeProperty(struct FPropertyChangedEvent& PropertyChangedEvent) override;
+	NIAGARA_API virtual void PostEditChangeProperty(struct FPropertyChangedEvent& PropertyChangedEvent) override;
 #endif
 	// UObject Interface End
 
 	// UNiagaraDataInterface Interface Begin
 	virtual bool CanExecuteOnTarget(ENiagaraSimTarget Target) const override { return true; }
 
-	virtual bool Equals(const UNiagaraDataInterface* Other) const override;
-	virtual bool CopyToInternal(UNiagaraDataInterface* Destination) const override;
+	NIAGARA_API virtual bool Equals(const UNiagaraDataInterface* Other) const override;
+	NIAGARA_API virtual bool CopyToInternal(UNiagaraDataInterface* Destination) const override;
 
-	virtual bool InitPerInstanceData(void* PerInstanceData, FNiagaraSystemInstance* SystemInstance) override;
-	virtual void DestroyPerInstanceData(void* PerInstanceData, FNiagaraSystemInstance* SystemInstance) override;
-	virtual int32 PerInstanceDataSize() const override;
+	NIAGARA_API virtual bool InitPerInstanceData(void* PerInstanceData, FNiagaraSystemInstance* SystemInstance) override;
+	NIAGARA_API virtual void DestroyPerInstanceData(void* PerInstanceData, FNiagaraSystemInstance* SystemInstance) override;
+	NIAGARA_API virtual int32 PerInstanceDataSize() const override;
 
-	virtual void GetFunctions(TArray<FNiagaraFunctionSignature>& OutFunctions) override;
-	virtual void GetVMExternalFunction(const FVMExternalFunctionBindingInfo& BindingInfo, void* InstanceData, FVMExternalFunction &OutFunc) override;
+	NIAGARA_API virtual void GetFunctions(TArray<FNiagaraFunctionSignature>& OutFunctions) override;
+	NIAGARA_API virtual void GetVMExternalFunction(const FVMExternalFunctionBindingInfo& BindingInfo, void* InstanceData, FVMExternalFunction &OutFunc) override;
 #if WITH_EDITORONLY_DATA
-	virtual void GetParameterDefinitionHLSL(const FNiagaraDataInterfaceGPUParamInfo& ParamInfo, FString& OutHLSL) override;
-	virtual bool GetFunctionHLSL(const FNiagaraDataInterfaceGPUParamInfo& ParamInfo, const FNiagaraDataInterfaceGeneratedFunction& FunctionInfo, int FunctionInstanceIndex, FString& OutHLSL) override;
-	virtual bool AppendCompileHash(FNiagaraCompileHashVisitor* InVisitor) const override;
+	NIAGARA_API virtual void GetParameterDefinitionHLSL(const FNiagaraDataInterfaceGPUParamInfo& ParamInfo, FString& OutHLSL) override;
+	NIAGARA_API virtual bool GetFunctionHLSL(const FNiagaraDataInterfaceGPUParamInfo& ParamInfo, const FNiagaraDataInterfaceGeneratedFunction& FunctionInfo, int FunctionInstanceIndex, FString& OutHLSL) override;
+	NIAGARA_API virtual bool AppendCompileHash(FNiagaraCompileHashVisitor* InVisitor) const override;
 #endif
 #if WITH_EDITOR
 	virtual bool GetGpuUseIndirectDispatch() const override { return true; }
 #endif
-	virtual void BuildShaderParameters(FNiagaraShaderParametersBuilder& ShaderParametersBuilder) const override;
-	virtual void SetShaderParameters(const FNiagaraDataInterfaceSetShaderParametersContext& Context) const override;
-	virtual void PushToRenderThreadImpl() override;
+	NIAGARA_API virtual void BuildShaderParameters(FNiagaraShaderParametersBuilder& ShaderParametersBuilder) const override;
+	NIAGARA_API virtual void SetShaderParameters(const FNiagaraDataInterfaceSetShaderParametersContext& Context) const override;
+	NIAGARA_API virtual void PushToRenderThreadImpl() override;
 	// UNiagaraDataInterface Interface End
 
 	// VM functions
-	void VMGet(FVectorVMExternalFunctionContext& Context);
-	void VMSet(FVectorVMExternalFunctionContext& Context);
-	void VMExchange(FVectorVMExternalFunctionContext& Context);
-	void VMAdd(FVectorVMExternalFunctionContext& Context);
-	void VMIncrement(FVectorVMExternalFunctionContext& Context);
-	void VMDecrement(FVectorVMExternalFunctionContext& Context);
+	NIAGARA_API void VMGet(FVectorVMExternalFunctionContext& Context);
+	NIAGARA_API void VMSet(FVectorVMExternalFunctionContext& Context);
+	NIAGARA_API void VMExchange(FVectorVMExternalFunctionContext& Context);
+	NIAGARA_API void VMAdd(FVectorVMExternalFunctionContext& Context);
+	NIAGARA_API void VMIncrement(FVectorVMExternalFunctionContext& Context);
+	NIAGARA_API void VMDecrement(FVectorVMExternalFunctionContext& Context);
 
-	void GetNextValue_Deprecated(FVectorVMExternalFunctionContext& Context);
+	NIAGARA_API void GetNextValue_Deprecated(FVectorVMExternalFunctionContext& Context);
 
 	/** Select how we should synchronize the counter between Cpu & Gpu */
 	UPROPERTY(EditAnywhere, Category = "Simple Counter")
diff --git a/Engine/Plugins/FX/Niagara/Source/Niagara/Private/DataInterface/NiagaraDataInterfaceVirtualTexture.h b/Engine/Plugins/FX/Niagara/Source/Niagara/Private/DataInterface/NiagaraDataInterfaceVirtualTexture.h
--- a/Engine/Plugins/FX/Niagara/Source/Niagara/Private/DataInterface/NiagaraDataInterfaceVirtualTexture.h
+++ b/Engine/Plugins/FX/Niagara/Source/Niagara/Private/DataInterface/NiagaraDataInterfaceVirtualTexture.h
@@ -12,8 +12,8 @@ struct FNiagaraDataInterfaceGeneratedFunction;
 class URuntimeVirtualTexture;
 
 /** Data Interface allowing sampling of a texture */
-UCLASS(Experimental, EditInlineNew, Category = "Texture", CollapseCategories, meta = (DisplayName = "Runtime Virtual Texture Sample"))
-class NIAGARA_API UNiagaraDataInterfaceVirtualTexture : public UNiagaraDataInterface
+UCLASS(Experimental, EditInlineNew, Category = "Texture", CollapseCategories, meta = (DisplayName = "Runtime Virtual Texture Sample"), MinimalAPI)
+class UNiagaraDataInterfaceVirtualTexture : public UNiagaraDataInterface
 {
 	GENERATED_UCLASS_BODY()
 
@@ -25,42 +25,42 @@ public:
 	FNiagaraUserParameterBinding TextureUserParameter;
 
 	//UObject Interface
-	virtual void PostInitProperties() override;
-	virtual void PostLoad() override;
-	virtual void Serialize(FArchive& Ar) override;
+	NIAGARA_API virtual void PostInitProperties() override;
+	NIAGARA_API virtual void PostLoad() override;
+	NIAGARA_API virtual void Serialize(FArchive& Ar) override;
 	//UObject Interface End
 
 	//UNiagaraDataInterface Interface
-	virtual void GetFunctions(TArray<FNiagaraFunctionSignature>& OutFunctions) override;
-	virtual void GetVMExternalFunction(const FVMExternalFunctionBindingInfo& BindingInfo, void* InstanceData, FVMExternalFunction &OutFunc) override;
+	NIAGARA_API virtual void GetFunctions(TArray<FNiagaraFunctionSignature>& OutFunctions) override;
+	NIAGARA_API virtual void GetVMExternalFunction(const FVMExternalFunctionBindingInfo& BindingInfo, void* InstanceData, FVMExternalFunction &OutFunc) override;
 	virtual bool CanExecuteOnTarget(ENiagaraSimTarget Target) const override { return Target==ENiagaraSimTarget::GPUComputeSim; }
 
-	virtual int32 PerInstanceDataSize() const override;
-	virtual bool InitPerInstanceData(void* PerInstanceData, FNiagaraSystemInstance* SystemInstance) override;
-	virtual void DestroyPerInstanceData(void* PerInstanceData, FNiagaraSystemInstance* SystemInstance) override;
+	NIAGARA_API virtual int32 PerInstanceDataSize() const override;
+	NIAGARA_API virtual bool InitPerInstanceData(void* PerInstanceData, FNiagaraSystemInstance* SystemInstance) override;
+	NIAGARA_API virtual void DestroyPerInstanceData(void* PerInstanceData, FNiagaraSystemInstance* SystemInstance) override;
 
 	virtual bool HasPreSimulateTick() const override { return true; }
-	virtual bool PerInstanceTick(void* PerInstanceData, FNiagaraSystemInstance* SystemInstance, float DeltaSeconds) override;
+	NIAGARA_API virtual bool PerInstanceTick(void* PerInstanceData, FNiagaraSystemInstance* SystemInstance, float DeltaSeconds) override;
 	//UNiagaraDataInterface Interface End
 
 	//void SampleTexture(FVectorVMExternalFunctionContext& Context);
 	//void GetTextureDimensions(FVectorVMExternalFunctionContext& Context);
 	//void SamplePseudoVolumeTexture(FVectorVMExternalFunctionContext& Context);
 
-	virtual bool Equals(const UNiagaraDataInterface* Other) const override;
+	NIAGARA_API virtual bool Equals(const UNiagaraDataInterface* Other) const override;
 
 	// GPU sim functionality
 #if WITH_EDITORONLY_DATA
-	virtual bool AppendCompileHash(FNiagaraCompileHashVisitor* InVisitor) const override;
-	virtual void GetCommonHLSL(FString& OutHlsl) override;
-	virtual void GetParameterDefinitionHLSL(const FNiagaraDataInterfaceGPUParamInfo& ParamInfo, FString& OutHLSL) override;
-	virtual bool GetFunctionHLSL(const FNiagaraDataInterfaceGPUParamInfo& ParamInfo, const FNiagaraDataInterfaceGeneratedFunction& FunctionInfo, int FunctionInstanceIndex, FString& OutHLSL) override;
+	NIAGARA_API virtual bool AppendCompileHash(FNiagaraCompileHashVisitor* InVisitor) const override;
+	NIAGARA_API virtual void GetCommonHLSL(FString& OutHlsl) override;
+	NIAGARA_API virtual void GetParameterDefinitionHLSL(const FNiagaraDataInterfaceGPUParamInfo& ParamInfo, FString& OutHLSL) override;
+	NIAGARA_API virtual bool GetFunctionHLSL(const FNiagaraDataInterfaceGPUParamInfo& ParamInfo, const FNiagaraDataInterfaceGeneratedFunction& FunctionInfo, int FunctionInstanceIndex, FString& OutHLSL) override;
 #endif
-	virtual void BuildShaderParameters(FNiagaraShaderParametersBuilder& ShaderParametersBuilder) const override;
-	virtual void SetShaderParameters(const FNiagaraDataInterfaceSetShaderParametersContext& Context) const override;
+	NIAGARA_API virtual void BuildShaderParameters(FNiagaraShaderParametersBuilder& ShaderParametersBuilder) const override;
+	NIAGARA_API virtual void SetShaderParameters(const FNiagaraDataInterfaceSetShaderParametersContext& Context) const override;
 
-	void SetTexture(URuntimeVirtualTexture* InTexture);
+	NIAGARA_API void SetTexture(URuntimeVirtualTexture* InTexture);
 
 protected:
-	virtual bool CopyToInternal(UNiagaraDataInterface* Destination) const override;
+	NIAGARA_API virtual bool CopyToInternal(UNiagaraDataInterface* Destination) const override;
 };
diff --git a/Engine/Plugins/FX/Niagara/Source/Niagara/Private/NiagaraDebuggerClient.h b/Engine/Plugins/FX/Niagara/Source/Niagara/Private/NiagaraDebuggerClient.h
--- a/Engine/Plugins/FX/Niagara/Source/Niagara/Private/NiagaraDebuggerClient.h
+++ b/Engine/Plugins/FX/Niagara/Source/Niagara/Private/NiagaraDebuggerClient.h
@@ -23,7 +23,7 @@ class FMessageEndpoint;
 /**
 Listener that accumulates a short run of stats for all systems and components in the scene and reports those stats to the outliner.
 */
-class NIAGARA_API FNiagaraOutlinerPerfListener : public FParticlePerfStatsListener_GatherAll
+class FNiagaraOutlinerPerfListener : public FParticlePerfStatsListener_GatherAll
 {
 public:
 	FNiagaraOutlinerPerfListener() : FParticlePerfStatsListener_GatherAll(true, true, true)
diff --git a/Engine/Plugins/FX/Niagara/Source/Niagara/Private/NiagaraGpuScratchPad.h b/Engine/Plugins/FX/Niagara/Source/Niagara/Private/NiagaraGpuScratchPad.h
--- a/Engine/Plugins/FX/Niagara/Source/Niagara/Private/NiagaraGpuScratchPad.h
+++ b/Engine/Plugins/FX/Niagara/Source/Niagara/Private/NiagaraGpuScratchPad.h
@@ -26,7 +26,7 @@ struct FNiagaraDataInterfaceProxy;
 *	Useful for transient working data generated in one dispatch and consumed in others.
 *	All buffers in the same scratch pad are transitioned as one.
 */
-class NIAGARA_API FNiagaraGpuScratchPad
+class FNiagaraGpuScratchPad
 {
 public:
 	struct FAllocation
@@ -275,7 +275,7 @@ FORCEINLINE FNiagaraGpuScratchPad::FAllocation FNiagaraGpuScratchPad::Alloc<uint
 /** Behaves similar to FNiagaraGpuScratchPad but for use with RWStructuredBuffers
 */
 template<typename T>
-class NIAGARA_API FNiagaraGpuScratchPadStructured
+class FNiagaraGpuScratchPadStructured
 {
 public:
 	struct FAllocation
diff --git a/Engine/Plugins/FX/Niagara/Source/Niagara/Public/Experimental/NiagaraMeshUvMappingHandle.h b/Engine/Plugins/FX/Niagara/Source/Niagara/Public/Experimental/NiagaraMeshUvMappingHandle.h
--- a/Engine/Plugins/FX/Niagara/Source/Niagara/Public/Experimental/NiagaraMeshUvMappingHandle.h
+++ b/Engine/Plugins/FX/Niagara/Source/Niagara/Public/Experimental/NiagaraMeshUvMappingHandle.h
@@ -22,32 +22,32 @@ struct FMeshUvMappingUsage
 	bool RequiresGpuAccess = false;
 };
 
-struct NIAGARA_API FMeshUvMappingHandleBase
+struct FMeshUvMappingHandleBase
 {
 	FMeshUvMappingHandleBase() = default;
 	FMeshUvMappingHandleBase(const FMeshUvMappingHandleBase& Other) = delete;
-	~FMeshUvMappingHandleBase();
+	NIAGARA_API ~FMeshUvMappingHandleBase();
 	FMeshUvMappingHandleBase& operator=(const FMeshUvMappingHandleBase& Other) = delete;
-	explicit operator bool() const;
+	NIAGARA_API explicit operator bool() const;
 
-	const FMeshUvMappingBufferProxy* GetQuadTreeProxy() const;
-	int32 GetUvSetIndex() const;
-	int32 GetLodIndex() const;
-	void PinAndInvalidateHandle();
+	NIAGARA_API const FMeshUvMappingBufferProxy* GetQuadTreeProxy() const;
+	NIAGARA_API int32 GetUvSetIndex() const;
+	NIAGARA_API int32 GetLodIndex() const;
+	NIAGARA_API void PinAndInvalidateHandle();
 
 	FMeshUvMappingUsage Usage;
 
 protected:
-	FMeshUvMappingHandleBase(FMeshUvMappingUsage InUsage, const TSharedPtr<FMeshUvMapping>& InUvMappingData, bool bNeedsDataImmediately);
-	FMeshUvMappingHandleBase(FMeshUvMappingHandleBase&& Other) noexcept;
-	FMeshUvMappingHandleBase& operator=(FMeshUvMappingHandleBase&& Other) noexcept;
+	NIAGARA_API FMeshUvMappingHandleBase(FMeshUvMappingUsage InUsage, const TSharedPtr<FMeshUvMapping>& InUvMappingData, bool bNeedsDataImmediately);
+	NIAGARA_API FMeshUvMappingHandleBase(FMeshUvMappingHandleBase&& Other) noexcept;
+	NIAGARA_API FMeshUvMappingHandleBase& operator=(FMeshUvMappingHandleBase&& Other) noexcept;
 
 	TSharedPtr<FMeshUvMapping> UvMappingData;
 };
 
 
 template<typename MappingType>
-struct NIAGARA_API TTypedMeshUvMappingHandle : public FMeshUvMappingHandleBase
+struct TTypedMeshUvMappingHandle : public FMeshUvMappingHandleBase
 {
 	TTypedMeshUvMappingHandle() = default;
 	TTypedMeshUvMappingHandle(FMeshUvMappingUsage InUsage, const TSharedPtr<MappingType>& InUvMappingData, bool bNeedsDataImmediately)
@@ -74,4 +74,4 @@ struct NIAGARA_API TTypedMeshUvMappingHandle : public FMeshUvMappingHandleBase
 };
 
 using FSkeletalMeshUvMappingHandle = TTypedMeshUvMappingHandle<FSkeletalMeshUvMapping>;
-using FStaticMeshUvMappingHandle = TTypedMeshUvMappingHandle<FStaticMeshUvMapping>;
\ No newline at end of file
+using FStaticMeshUvMappingHandle = TTypedMeshUvMappingHandle<FStaticMeshUvMapping>;
diff --git a/Engine/Plugins/FX/Niagara/Source/Niagara/Public/MovieScene/Parameters/MovieSceneNiagaraParameterTrack.h b/Engine/Plugins/FX/Niagara/Source/Niagara/Public/MovieScene/Parameters/MovieSceneNiagaraParameterTrack.h
--- a/Engine/Plugins/FX/Niagara/Source/Niagara/Public/MovieScene/Parameters/MovieSceneNiagaraParameterTrack.h
+++ b/Engine/Plugins/FX/Niagara/Source/Niagara/Public/MovieScene/Parameters/MovieSceneNiagaraParameterTrack.h
@@ -20,7 +20,7 @@ public:
 	/** Sets the parameter for this parameter track .*/
 	NIAGARA_API void SetParameter(FNiagaraVariable InParameter);
 
-	NIAGARA_API virtual void SetSectionChannelDefaults(UMovieSceneSection* Section, const TArray<uint8>& DefaultValueData) const { }
+	virtual void SetSectionChannelDefaults(UMovieSceneSection* Section, const TArray<uint8>& DefaultValueData) const { }
 
 protected:
 	template<class ChannelType>
@@ -48,4 +48,4 @@ private:
 	/** The parameter for the parameter this track animates. */
 	UPROPERTY()
 	FNiagaraVariable Parameter;
-};
\ No newline at end of file
+};
diff --git a/Engine/Plugins/FX/Niagara/Source/Niagara/Public/NiagaraActor.h b/Engine/Plugins/FX/Niagara/Source/Niagara/Public/NiagaraActor.h
--- a/Engine/Plugins/FX/Niagara/Source/Niagara/Public/NiagaraActor.h
+++ b/Engine/Plugins/FX/Niagara/Source/Niagara/Public/NiagaraActor.h
@@ -7,21 +7,21 @@
 #include "GameFramework/Actor.h"
 #include "NiagaraActor.generated.h"
 
-UCLASS(hideCategories = (Activation, "Components|Activation", Input, Collision, "Game|Damage"), ComponentWrapperClass)
-class NIAGARA_API ANiagaraActor : public AActor
+UCLASS(hideCategories = (Activation, "Components|Activation", Input, Collision, "Game|Damage"), ComponentWrapperClass, MinimalAPI)
+class ANiagaraActor : public AActor
 {
 	GENERATED_BODY()
 protected:
 
-	ANiagaraActor(const FObjectInitializer& ObjectInitializer);
+	NIAGARA_API ANiagaraActor(const FObjectInitializer& ObjectInitializer);
 	
 public:
 
-	virtual void PostRegisterAllComponents() override;
+	NIAGARA_API virtual void PostRegisterAllComponents() override;
 
 	/** Set true for this actor to self-destruct when the Niagara system finishes, false otherwise */
 	UFUNCTION(BlueprintCallable, Category=NiagaraActor)
-	void SetDestroyOnSystemFinish(bool bShouldDestroyOnSystemFinish);
+	NIAGARA_API void SetDestroyOnSystemFinish(bool bShouldDestroyOnSystemFinish);
 
 	/** Returns true if the system will destroy on finish */
 	UFUNCTION(BlueprintCallable, Category=NiagaraActor)
@@ -49,7 +49,7 @@ private:
 
 	/** Callback when Niagara system finishes. */
 	UFUNCTION(CallInEditor)
-	void OnNiagaraSystemFinished(UNiagaraComponent* FinishedComponent);
+	NIAGARA_API void OnNiagaraSystemFinished(UNiagaraComponent* FinishedComponent);
 
 public:
 	/** Returns NiagaraComponent subobject **/
@@ -63,11 +63,11 @@ public:
 
 #if WITH_EDITOR
 	// AActor interface
-	virtual bool GetReferencedContentObjects(TArray<UObject*>& Objects) const override;
+	NIAGARA_API virtual bool GetReferencedContentObjects(TArray<UObject*>& Objects) const override;
 	// End of AActor interface
 
 	/** Reset this actor in the level.*/
-	void ResetInLevel();
+	NIAGARA_API void ResetInLevel();
 #endif // WITH_EDITOR
 
 };
diff --git a/Engine/Plugins/FX/Niagara/Source/Niagara/Public/NiagaraCommon.h b/Engine/Plugins/FX/Niagara/Source/Niagara/Public/NiagaraCommon.h
--- a/Engine/Plugins/FX/Niagara/Source/Niagara/Public/NiagaraCommon.h
+++ b/Engine/Plugins/FX/Niagara/Source/Niagara/Public/NiagaraCommon.h
@@ -343,7 +343,7 @@ struct FNiagaraScriptDataUsageInfo
 };
 
 USTRUCT()
-struct NIAGARA_API FNiagaraFunctionSignature
+struct FNiagaraFunctionSignature
 {
 	GENERATED_BODY()
 
@@ -665,12 +665,12 @@ struct NIAGARA_API FNiagaraFunctionSignature
 	int32 NumRequiredOutputs()const { return RequiredOutputs == INDEX_NONE ? Outputs.Num() : RequiredOutputs; }
 	int32 NumOptionalOutputs()const { return Outputs.Num() - NumRequiredOutputs(); }
 
-	void GetVariadicInputs(TArray<FNiagaraVariableBase>& OutVariadicInputs, bool bStripNonExecution = true)const;
-	void GetVariadicOutputs(TArray<FNiagaraVariableBase>& OutVariadicOutputs, bool bStripNonExecution = true)const;
+	NIAGARA_API void GetVariadicInputs(TArray<FNiagaraVariableBase>& OutVariadicInputs, bool bStripNonExecution = true)const;
+	NIAGARA_API void GetVariadicOutputs(TArray<FNiagaraVariableBase>& OutVariadicOutputs, bool bStripNonExecution = true)const;
 };
 
 USTRUCT()
-struct NIAGARA_API FNiagaraScriptUObjectCompileInfo
+struct FNiagaraScriptUObjectCompileInfo
 {
 	GENERATED_BODY()
 
@@ -691,7 +691,7 @@ struct NIAGARA_API FNiagaraScriptUObjectCompileInfo
 };
 
 USTRUCT()
-struct NIAGARA_API FNiagaraResolvedUObjectInfo
+struct FNiagaraResolvedUObjectInfo
 {
 	GENERATED_BODY()
 
@@ -709,7 +709,7 @@ struct NIAGARA_API FNiagaraResolvedUObjectInfo
 };
 
 USTRUCT()
-struct NIAGARA_API FNiagaraExternalUObjectInfo
+struct FNiagaraExternalUObjectInfo
 {
 	GENERATED_BODY()
 
@@ -723,7 +723,7 @@ struct NIAGARA_API FNiagaraExternalUObjectInfo
 };
 
 USTRUCT()
-struct NIAGARA_API FNiagaraScriptDataInterfaceInfo
+struct FNiagaraScriptDataInterfaceInfo
 {
 	GENERATED_USTRUCT_BODY()
 public:
@@ -760,11 +760,11 @@ public:
 	UPROPERTY()
 	FString SourceEmitterName;
 
-	bool IsUserDataInterface() const;
+	NIAGARA_API bool IsUserDataInterface() const;
 };
 
 USTRUCT()
-struct NIAGARA_API FNiagaraResolvedUserDataInterfaceBinding
+struct FNiagaraResolvedUserDataInterfaceBinding
 {
 	GENERATED_BODY()
 
@@ -785,7 +785,7 @@ struct NIAGARA_API FNiagaraResolvedUserDataInterfaceBinding
 };
 
 USTRUCT()
-struct NIAGARA_API FNiagaraScriptResolvedDataInterfaceInfo
+struct FNiagaraScriptResolvedDataInterfaceInfo
 {
 	GENERATED_USTRUCT_BODY();
 
@@ -820,11 +820,11 @@ struct NIAGARA_API FNiagaraScriptResolvedDataInterfaceInfo
 	UPROPERTY()
 	int32 UserPtrIdx;
 
-	bool NeedsPerInstanceBinding() const;
+	NIAGARA_API bool NeedsPerInstanceBinding() const;
 };
 
 USTRUCT()
-struct NIAGARA_API FNiagaraScriptDataInterfaceCompileInfo
+struct FNiagaraScriptDataInterfaceCompileInfo
 {
 	GENERATED_USTRUCT_BODY()
 public:
@@ -865,14 +865,14 @@ public:
 	FString SourceEmitterName;
 
 	/** Would this data interface work on the target execution type? Only call this on the game thread.*/
-	bool CanExecuteOnTarget(ENiagaraSimTarget SimTarget) const;
+	NIAGARA_API bool CanExecuteOnTarget(ENiagaraSimTarget SimTarget) const;
 
 	/** Note that this is the CDO for this type of data interface, as we often cannot guarantee that the same instance of the data interface we compiled with is the one the user ultimately executes.  Only call this on the game thread.*/
-	UNiagaraDataInterface* GetDefaultDataInterface() const;
+	NIAGARA_API UNiagaraDataInterface* GetDefaultDataInterface() const;
 
-	bool NeedsPerInstanceBinding() const;
+	NIAGARA_API bool NeedsPerInstanceBinding() const;
 
-	bool MatchesClass(const UClass* InClass) const;
+	NIAGARA_API bool MatchesClass(const UClass* InClass) const;
 };
 
 USTRUCT()
@@ -966,7 +966,7 @@ For example, this can be split between PreEditChange and PostEditChange to ensur
 This can be made a UPROPERTY() to ensure safey in cases where a GC could be possible between Add() and CommitUpdate().
 */
 USTRUCT()
-struct NIAGARA_API FNiagaraSystemUpdateContext
+struct FNiagaraSystemUpdateContext
 {
 	GENERATED_BODY()
 
@@ -979,32 +979,32 @@ struct NIAGARA_API FNiagaraSystemUpdateContext
 	FNiagaraSystemUpdateContext() :bDestroyOnAdd(false), bOnlyActive(false), bDestroySystemSim(true) { }
 	FNiagaraSystemUpdateContext(FNiagaraSystemUpdateContext& Other) = delete;
 
-	~FNiagaraSystemUpdateContext();
+	NIAGARA_API ~FNiagaraSystemUpdateContext();
 
 	void SetDestroyOnAdd(bool bInDestroyOnAdd) { bDestroyOnAdd = bInDestroyOnAdd; }
 	void SetOnlyActive(bool bInOnlyActive) { bOnlyActive = bInOnlyActive; }
 	void SetDestroySystemSim(bool bInDestroySystemSim) { bDestroySystemSim = bInDestroySystemSim; }
 
-	void AddSoloComponent(UNiagaraComponent* Component, bool bReInit);
-	void Add(const UNiagaraSystem* System, bool bReInit);
+	NIAGARA_API void AddSoloComponent(UNiagaraComponent* Component, bool bReInit);
+	NIAGARA_API void Add(const UNiagaraSystem* System, bool bReInit);
 #if WITH_EDITORONLY_DATA
-	void Add(const FVersionedNiagaraEmitter& Emitter, bool bReInit);
-	void Add(const UNiagaraScript* Script, bool bReInit);
-	void Add(const UNiagaraParameterCollection* Collection, bool bReInit);
+	NIAGARA_API void Add(const FVersionedNiagaraEmitter& Emitter, bool bReInit);
+	NIAGARA_API void Add(const UNiagaraScript* Script, bool bReInit);
+	NIAGARA_API void Add(const UNiagaraParameterCollection* Collection, bool bReInit);
 #endif
 
 	/** Adds all currently active systems.*/
-	void AddAll(bool bReInit);
+	NIAGARA_API void AddAll(bool bReInit);
 	
 	/** Handles any pending reinits or resets of system instances in this update context. */
-	void CommitUpdate();
+	NIAGARA_API void CommitUpdate();
 
 	DECLARE_DELEGATE_OneParam(FCustomWorkDelegate, UNiagaraComponent*);
 	FCustomWorkDelegate& GetPreWork(){ return PreWork; }
 	FCustomWorkDelegate& GetPostWork() { return PostWork; }
 
 private:
-	void AddInternal(class UNiagaraComponent* Comp, bool bReInit, bool bAllowDestroySystemSim);
+	NIAGARA_API void AddInternal(class UNiagaraComponent* Comp, bool bReInit, bool bAllowDestroySystemSim);
 
 	UPROPERTY(transient)
 	TArray<TObjectPtr<UNiagaraComponent>> ComponentsToReset;
@@ -1165,26 +1165,26 @@ private:
 /** Combines all stat reporting and evaluation of niagara instances (emitter or system).
  * This is then used by the SNiagaraStackRowPerfWidget to display the data in the ui.
  */
-struct NIAGARA_API FNiagaraStatDatabase
+struct FNiagaraStatDatabase
 {	
 	typedef TTuple<uint64, ENiagaraScriptUsage> FStatReportKey;
 
 	/* Used by emitter and system instances to add the recorded data of a frame to this emitter's data store. */
-	void AddStatCapture(FStatReportKey ReportKey, TMap<TStatIdData const*, float> CapturedData);
+	NIAGARA_API void AddStatCapture(FStatReportKey ReportKey, TMap<TStatIdData const*, float> CapturedData);
 
 	/* Removes all captured stats. */
-	void ClearStatCaptures();
+	NIAGARA_API void ClearStatCaptures();
 
 	/* Returns the average runtime cost of a specific module call inside the script for the given usage. Returns 0 if no data was found. */
-	float GetRuntimeStat(FName StatName, ENiagaraScriptUsage Usage, ENiagaraStatEvaluationType EvaluationType);
+	NIAGARA_API float GetRuntimeStat(FName StatName, ENiagaraScriptUsage Usage, ENiagaraStatEvaluationType EvaluationType);
 
 	/* Returns the average runtime cost of a script for the given usage. Returns 0 if no data was recorded for that usage. */
-	float GetRuntimeStat(ENiagaraScriptUsage Usage, ENiagaraStatEvaluationType EvaluationType);
+	NIAGARA_API float GetRuntimeStat(ENiagaraScriptUsage Usage, ENiagaraStatEvaluationType EvaluationType);
 
 	/* Returns the names of all captures stat data points. Useful for debugging and to dump the stat data. */
-	TMap<ENiagaraScriptUsage, TSet<FName>> GetAvailableStatNames();
+	NIAGARA_API TMap<ENiagaraScriptUsage, TSet<FName>> GetAvailableStatNames();
 
-	void Init();
+	NIAGARA_API void Init();
 
 private:
 	/** The captured runtime stat data. The first key is a combination of reporter handle and script usage, the second key is the stat id which correlates to a single recorded scope. */
@@ -1230,9 +1230,9 @@ struct FNiagaraVariableAttributeBinding
 
 	FNiagaraVariableAttributeBinding() : BindingSourceMode(ENiagaraBindingSource::ImplicitFromSource), bBindingExistsOnSource(false), bIsCachedParticleValue(true) {}
 
-	bool NIAGARA_API IsParticleBinding() const { return bIsCachedParticleValue; }
-	bool NIAGARA_API DoesBindingExistOnSource() const { return bBindingExistsOnSource; }
-	bool NIAGARA_API CanBindToHostParameterMap() const { return bBindingExistsOnSource && !bIsCachedParticleValue; }
+	bool IsParticleBinding() const { return bIsCachedParticleValue; }
+	bool DoesBindingExistOnSource() const { return bBindingExistsOnSource; }
+	bool CanBindToHostParameterMap() const { return bBindingExistsOnSource && !bIsCachedParticleValue; }
 	void NIAGARA_API SetValue(const FName& InValue, const FVersionedNiagaraEmitter& InEmitter, ENiagaraRendererSourceDataMode InSourceMode);
 	void NIAGARA_API SetAsPreviousValue(const FNiagaraVariableBase& Src, const FVersionedNiagaraEmitter& InEmitter, ENiagaraRendererSourceDataMode InSourceMode);
 	void NIAGARA_API SetAsPreviousValue(const FNiagaraVariableAttributeBinding& Src, const FVersionedNiagaraEmitter& InEmitter, ENiagaraRendererSourceDataMode InSourceMode);
@@ -1241,15 +1241,15 @@ struct FNiagaraVariableAttributeBinding
 	bool NIAGARA_API Matches(const FNiagaraVariableBase& OldVariable, const FVersionedNiagaraEmitter& InEmitter, ENiagaraRendererSourceDataMode InSourceMode);
 
 #if WITH_EDITORONLY_DATA
-	NIAGARA_API const FName& GetName() const { return CachedDisplayName; }
+	const FName& GetName() const { return CachedDisplayName; }
 	[[nodiscard]] NIAGARA_API FString GetDefaultValueString() const;
 #endif
-	NIAGARA_API const FNiagaraVariableBase& GetParamMapBindableVariable() const { return ParamMapVariable; }
-	[[nodiscard]] NIAGARA_API FNiagaraVariableBase GetDataSetBindableVariable() const { return FNiagaraVariableBase(ParamMapVariable.GetType(), DataSetName); }
-	NIAGARA_API const FNiagaraTypeDefinition& GetType() const { return ParamMapVariable.GetType(); }
-	NIAGARA_API ENiagaraBindingSource GetBindingSourceMode() const { return BindingSourceMode; }
+	const FNiagaraVariableBase& GetParamMapBindableVariable() const { return ParamMapVariable; }
+	[[nodiscard]] FNiagaraVariableBase GetDataSetBindableVariable() const { return FNiagaraVariableBase(ParamMapVariable.GetType(), DataSetName); }
+	const FNiagaraTypeDefinition& GetType() const { return ParamMapVariable.GetType(); }
+	ENiagaraBindingSource GetBindingSourceMode() const { return BindingSourceMode; }
 
-	NIAGARA_API bool IsValid() const { return !DataSetName.IsNone() && ParamMapVariable.GetType().IsValid(); }
+	bool IsValid() const { return !DataSetName.IsNone() && ParamMapVariable.GetType().IsValid(); }
 	
 	template<typename T>
 	T GetDefaultValue() const
@@ -1372,7 +1372,7 @@ struct FNiagaraScriptVariableBinding
 	bool IsValid() const { return Name != NAME_None; }
 };
 
-struct NIAGARA_API FNiagaraAliasContext
+struct FNiagaraAliasContext
 {
 	/** Defines different modes which can be used to split rapid iteration parameter constant names. */
 	enum class ERapidIterationParameterMode
@@ -1397,25 +1397,25 @@ struct NIAGARA_API FNiagaraAliasContext
 	}
 
 	/** Configures the context to replace the unaliased emitter namespace with an emitter name namespace in parameter map parameters and rapid iteration parameters. */
-	FNiagaraAliasContext& ChangeEmitterToEmitterName(const FString& InEmitterName);
+	NIAGARA_API FNiagaraAliasContext& ChangeEmitterToEmitterName(const FString& InEmitterName);
 
 	/** Configures the context to replace an emitter name namespace with the unaliased emitter namespace in parameter map parameters and rapid iteration parameters. */
-	FNiagaraAliasContext& ChangeEmitterNameToEmitter(const FString& InEmitterName);
+	NIAGARA_API FNiagaraAliasContext& ChangeEmitterNameToEmitter(const FString& InEmitterName);
 
 	/** Configures the context to replace an old emitter name namespace with a new emitter name namespace in parameter map parameters and rapid iteration parameters. */
-	FNiagaraAliasContext& ChangeEmitterName(const FString& InOldEmitterName, const FString& InNewEmitterName);
+	NIAGARA_API FNiagaraAliasContext& ChangeEmitterName(const FString& InOldEmitterName, const FString& InNewEmitterName);
 
 	/** Configures the context to replace the unaliased module namespace with a module name namespace in parameter map parameters and rapid iteration parameters. */
-	FNiagaraAliasContext& ChangeModuleToModuleName(const FString& InModuleName);
+	NIAGARA_API FNiagaraAliasContext& ChangeModuleToModuleName(const FString& InModuleName);
 
 	/** Configures the context to replace a module name namespace with the unaliased module namespace in parameter map parameters and rapid iteration parameters. */
-	FNiagaraAliasContext& ChangeModuleNameToModule(const FString& InModuleName);
+	NIAGARA_API FNiagaraAliasContext& ChangeModuleNameToModule(const FString& InModuleName);
 
 	/** Configures the context to replace an old module name namespace with a new module name namespace in parameter map parameters and rapid iteration parameters. */
-	FNiagaraAliasContext& ChangeModuleName(const FString& InOldModuleName, const FString& InNewModuleName);
+	NIAGARA_API FNiagaraAliasContext& ChangeModuleName(const FString& InOldModuleName, const FString& InNewModuleName);
 
 	/** Configures the context to replace the stack context namespace with the specified named stack context in parameter map parameters. */
-	FNiagaraAliasContext& ChangeStackContext(const FString& InStackContextName);
+	NIAGARA_API FNiagaraAliasContext& ChangeStackContext(const FString& InStackContextName);
 
 	ERapidIterationParameterMode GetRapidIterationParameterMode() const { return RapidIterationParameterMode; }
 
@@ -1530,13 +1530,13 @@ namespace FNiagaraUtilities
 };
 
 USTRUCT()
-struct NIAGARA_API FNiagaraUserParameterBinding
+struct FNiagaraUserParameterBinding
 {
 	GENERATED_USTRUCT_BODY()
 
-	FNiagaraUserParameterBinding();
+	NIAGARA_API FNiagaraUserParameterBinding();
 
-	FNiagaraUserParameterBinding(const FNiagaraTypeDefinition& InMaterialDef);
+	NIAGARA_API FNiagaraUserParameterBinding(const FNiagaraTypeDefinition& InMaterialDef);
 
 	UPROPERTY(EditAnywhere, Category = "User Parameter")
 	FNiagaraVariable Parameter;
@@ -1719,9 +1719,9 @@ enum class ENiagaraFunctionDebugState : uint8
 };
 
 /** Args struct for INiagaraParameterDefinitionsSubscriberViewModel::SynchronizeWithParameterDefinitions(...). */
-struct NIAGARA_API FSynchronizeWithParameterDefinitionsArgs
+struct FSynchronizeWithParameterDefinitionsArgs
 {
-	FSynchronizeWithParameterDefinitionsArgs();
+	NIAGARA_API FSynchronizeWithParameterDefinitionsArgs();
 
 	/** If set, instead of gathering all available parameter libraries, only consider subscribed parameter definitions that have a matching Id. */
 	TArray<FGuid> SpecificDefinitionsUniqueIds;
@@ -1804,4 +1804,4 @@ bool EvalConditional(ENiagaraConditionalOperator Op, const T& A, const T& B)
 	};
 
 	return false;
-};
\ No newline at end of file
+};
diff --git a/Engine/Plugins/FX/Niagara/Source/Niagara/Public/NiagaraCompileHashVisitor.h b/Engine/Plugins/FX/Niagara/Source/Niagara/Public/NiagaraCompileHashVisitor.h
--- a/Engine/Plugins/FX/Niagara/Source/Niagara/Public/NiagaraCompileHashVisitor.h
+++ b/Engine/Plugins/FX/Niagara/Source/Niagara/Public/NiagaraCompileHashVisitor.h
@@ -6,7 +6,7 @@
 #include "NiagaraCompileHashVisitor.generated.h"
 
 USTRUCT()
-struct NIAGARA_API FNiagaraCompileHashVisitorDebugInfo
+struct FNiagaraCompileHashVisitorDebugInfo
 {
 	GENERATED_USTRUCT_BODY()
 public:
@@ -23,7 +23,7 @@ public:
 /**
 Used to store the state of a graph when deciding if it has been dirtied for recompile.
 */
-struct NIAGARA_API FNiagaraCompileHashVisitor
+struct FNiagaraCompileHashVisitor
 {
 public:
 	FNiagaraCompileHashVisitor(FSHA1& InHashState) : HashState(InHashState) {}
@@ -31,7 +31,7 @@ public:
 	FSHA1& HashState;
 	TArray<const void*> ObjectList;
 
-	static int LogCompileIdGeneration;
+	static NIAGARA_API int LogCompileIdGeneration;
 
 #if WITH_EDITORONLY_DATA
 	FNiagaraCompileHashVisitorDebugInfo* AddDebugInfo()
@@ -144,10 +144,10 @@ public:
 		return UpdatePOD(ShaderParametersMetadata->GetStructTypeName(), ShaderParametersMetadata->GetLayoutHash());
 	}
 
-	bool UpdateShaderFile(const TCHAR* ShaderFilePath);
+	NIAGARA_API bool UpdateShaderFile(const TCHAR* ShaderFilePath);
 
 	/**
 	Adds an string value to the hash.
 	*/
-	bool UpdateString(const TCHAR* InDebugName, FStringView InData);
+	NIAGARA_API bool UpdateString(const TCHAR* InDebugName, FStringView InData);
 };
diff --git a/Engine/Plugins/FX/Niagara/Source/Niagara/Public/NiagaraComponent.h b/Engine/Plugins/FX/Niagara/Source/Niagara/Public/NiagaraComponent.h
--- a/Engine/Plugins/FX/Niagara/Source/Niagara/Public/NiagaraComponent.h
+++ b/Engine/Plugins/FX/Niagara/Source/Niagara/Public/NiagaraComponent.h
@@ -44,8 +44,8 @@ DECLARE_DYNAMIC_MULTICAST_DELEGATE_OneParam(FOnNiagaraSystemFinished, class UNia
 * @see ANiagaraActor
 * @see UNiagaraSystem
 */
-UCLASS(ClassGroup = (Rendering, Common), Blueprintable, hidecategories = Object, hidecategories = Physics, hidecategories = Collision, showcategories = Trigger, editinlinenew, meta = (BlueprintSpawnableComponent, DisplayName = "Niagara Particle System Component"))
-class NIAGARA_API UNiagaraComponent : public UFXSystemComponent
+UCLASS(ClassGroup = (Rendering, Common), Blueprintable, hidecategories = Object, hidecategories = Physics, hidecategories = Collision, showcategories = Trigger, editinlinenew, meta = (BlueprintSpawnableComponent, DisplayName = "Niagara Particle System Component"), MinimalAPI)
+class UNiagaraComponent : public UFXSystemComponent
 {
 	friend struct FNiagaraScalabilityManager;
 	GENERATED_UCLASS_BODY()
@@ -58,21 +58,21 @@ class NIAGARA_API UNiagaraComponent : public UFXSystemComponent
 public:
 
 	/********* UFXSystemComponent *********/
-	void SetBoolParameter(FName ParameterName, bool Param) override;
-	void SetIntParameter(FName ParameterName, int Param) override;
-	void SetFloatParameter(FName ParameterName, float Param) override;
-	void SetVectorParameter(FName ParameterName, FVector Param) override;
-	void SetColorParameter(FName ParameterName, FLinearColor Param) override;
-	void SetActorParameter(FName ParameterName, class AActor* Param) override;
-
-	virtual UFXSystemAsset* GetFXSystemAsset() const override;
-	void SetEmitterEnable(FName EmitterName, bool bNewEnableState) override;
-	void ReleaseToPool() override;
-	uint32 GetApproxMemoryUsage() const override;
-	virtual void ActivateSystem(bool bFlagAsJustAttached = false) override;
+	NIAGARA_API void SetBoolParameter(FName ParameterName, bool Param) override;
+	NIAGARA_API void SetIntParameter(FName ParameterName, int Param) override;
+	NIAGARA_API void SetFloatParameter(FName ParameterName, float Param) override;
+	NIAGARA_API void SetVectorParameter(FName ParameterName, FVector Param) override;
+	NIAGARA_API void SetColorParameter(FName ParameterName, FLinearColor Param) override;
+	NIAGARA_API void SetActorParameter(FName ParameterName, class AActor* Param) override;
+
+	NIAGARA_API virtual UFXSystemAsset* GetFXSystemAsset() const override;
+	NIAGARA_API void SetEmitterEnable(FName EmitterName, bool bNewEnableState) override;
+	NIAGARA_API void ReleaseToPool() override;
+	NIAGARA_API uint32 GetApproxMemoryUsage() const override;
+	NIAGARA_API virtual void ActivateSystem(bool bFlagAsJustAttached = false) override;
 	/********* UFXSystemComponent *********/
 
-	virtual FName GetFNameForStatID() const override;
+	NIAGARA_API virtual FName GetFNameForStatID() const override;
 
 private:
 	UPROPERTY(EditAnywhere, Category="Niagara", meta = (DisplayName = "Niagara System Asset"))
@@ -164,14 +164,14 @@ private:
 
 	//~ Begin UActorComponent Interface.
 protected:
-	virtual void OnRegister() override;
-	virtual void OnUnregister() override; 
-	virtual void ApplyWorldOffset(const FVector& InOffset, bool bWorldShift) override;
-	virtual void OnEndOfFrameUpdateDuringTick() override;
-	virtual void CreateRenderState_Concurrent(FRegisterComponentContext* Context) override;
-	virtual void DestroyRenderState_Concurrent() override;
-	virtual void SendRenderDynamicData_Concurrent() override;
-	virtual void BeginDestroy() override;
+	NIAGARA_API virtual void OnRegister() override;
+	NIAGARA_API virtual void OnUnregister() override; 
+	NIAGARA_API virtual void ApplyWorldOffset(const FVector& InOffset, bool bWorldShift) override;
+	NIAGARA_API virtual void OnEndOfFrameUpdateDuringTick() override;
+	NIAGARA_API virtual void CreateRenderState_Concurrent(FRegisterComponentContext* Context) override;
+	NIAGARA_API virtual void DestroyRenderState_Concurrent() override;
+	NIAGARA_API virtual void SendRenderDynamicData_Concurrent() override;
+	NIAGARA_API virtual void BeginDestroy() override;
 	//virtual void OnAttachmentChanged() override;
 
 public:
@@ -202,66 +202,66 @@ public:
 	/** How to handle pooling for this component instance. */
 	ENCPoolMethod PoolingMethod;
 
-	virtual void Activate(bool bReset = false) override;
-	virtual void Deactivate() override;
-	virtual void DeactivateImmediate() override;
-	virtual void GetResourceSizeEx(FResourceSizeEx& CumulativeResourceSize) override;
+	NIAGARA_API virtual void Activate(bool bReset = false) override;
+	NIAGARA_API virtual void Deactivate() override;
+	NIAGARA_API virtual void DeactivateImmediate() override;
+	NIAGARA_API virtual void GetResourceSizeEx(FResourceSizeEx& CumulativeResourceSize) override;
 
-	ENiagaraExecutionState GetRequestedExecutionState() const;
-	ENiagaraExecutionState GetExecutionState() const;
+	NIAGARA_API ENiagaraExecutionState GetRequestedExecutionState() const;
+	NIAGARA_API ENiagaraExecutionState GetExecutionState() const;
 
-	bool IsComplete() const;
+	NIAGARA_API bool IsComplete() const;
 
 private:
 
 	//Internal versions that can be called from the scalability code.
 	//These will behave as expected but will keep the component registered with the scalability manager.
-	void ActivateInternal(bool bReset, bool bIsScalabilityCull);
-	void DeactivateInternal(bool bIsScalabilityCull);
-	void DeactivateImmediateInternal(bool bIsScalabilityCull);
-	void SetPausedInternal(bool bInPaused, bool bIsScalabilityCull);	
+	NIAGARA_API void ActivateInternal(bool bReset, bool bIsScalabilityCull);
+	NIAGARA_API void DeactivateInternal(bool bIsScalabilityCull);
+	NIAGARA_API void DeactivateImmediateInternal(bool bIsScalabilityCull);
+	NIAGARA_API void SetPausedInternal(bool bInPaused, bool bIsScalabilityCull);	
 
 
-	bool ShouldPreCull();
-	void RegisterWithScalabilityManager();
-	void UnregisterWithScalabilityManager();
+	NIAGARA_API bool ShouldPreCull();
+	NIAGARA_API void RegisterWithScalabilityManager();
+	NIAGARA_API void UnregisterWithScalabilityManager();
 
-	void PostSystemTick_GameThread();
-	void OnSystemComplete(bool bExternalCompletion);
+	NIAGARA_API void PostSystemTick_GameThread();
+	NIAGARA_API void OnSystemComplete(bool bExternalCompletion);
 
 public:
 
-	virtual void TickComponent(float DeltaTime, enum ELevelTick TickType, FActorComponentTickFunction* ThisTickFunction) override;
-	virtual const UObject* AdditionalStatObject() const override;
-	virtual bool IsReadyForOwnerToAutoDestroy() const override;
-	virtual void OnComponentDestroyed(bool bDestroyingHierarchy) override;
-	virtual void OnComponentCreated() override;
-	virtual void PostApplyToComponent() override;
+	NIAGARA_API virtual void TickComponent(float DeltaTime, enum ELevelTick TickType, FActorComponentTickFunction* ThisTickFunction) override;
+	NIAGARA_API virtual const UObject* AdditionalStatObject() const override;
+	NIAGARA_API virtual bool IsReadyForOwnerToAutoDestroy() const override;
+	NIAGARA_API virtual void OnComponentDestroyed(bool bDestroyingHierarchy) override;
+	NIAGARA_API virtual void OnComponentCreated() override;
+	NIAGARA_API virtual void PostApplyToComponent() override;
 	//~ End UActorComponent Interface.
 
 	//~ Begin UPrimitiveComponent Interface
-	virtual int32 GetNumMaterials() const override;
-	virtual FBoxSphereBounds CalcBounds(const FTransform& LocalToWorld) const override;
-	virtual FPrimitiveSceneProxy* CreateSceneProxy() override;
-	virtual void GetUsedMaterials(TArray<UMaterialInterface*>& OutMaterials, bool bGetDebugMaterials = false) const override;
-	virtual void GetStreamingRenderAssetInfo(FStreamingTextureLevelContext& LevelContext, TArray<FStreamingRenderAssetPrimitiveInfo>& OutStreamingRenderAssets) const override;
-	virtual void OnAttachmentChanged() override;
+	NIAGARA_API virtual int32 GetNumMaterials() const override;
+	NIAGARA_API virtual FBoxSphereBounds CalcBounds(const FTransform& LocalToWorld) const override;
+	NIAGARA_API virtual FPrimitiveSceneProxy* CreateSceneProxy() override;
+	NIAGARA_API virtual void GetUsedMaterials(TArray<UMaterialInterface*>& OutMaterials, bool bGetDebugMaterials = false) const override;
+	NIAGARA_API virtual void GetStreamingRenderAssetInfo(FStreamingTextureLevelContext& LevelContext, TArray<FStreamingRenderAssetPrimitiveInfo>& OutStreamingRenderAssets) const override;
+	NIAGARA_API virtual void OnAttachmentChanged() override;
 	//~ End UPrimitiveComponent Interface
 
 	//~ Begin USceneComponent Interface
-	virtual void OnChildAttached(USceneComponent* ChildComponent) override;
-	virtual void OnChildDetached(USceneComponent* ChildComponent) override;
-	virtual bool IsVisible() const override;
+	NIAGARA_API virtual void OnChildAttached(USceneComponent* ChildComponent) override;
+	NIAGARA_API virtual void OnChildDetached(USceneComponent* ChildComponent) override;
+	NIAGARA_API virtual bool IsVisible() const override;
 	//~ Begin USceneComponent Interface
 
 	UE_DEPRECATED(5.0, "This interface is no longer safe to access directly. Use the interface provided by GetSystemInstanceController instead.")
-	TSharedPtr<FNiagaraSystemSimulation, ESPMode::ThreadSafe> GetSystemSimulation();
+	NIAGARA_API TSharedPtr<FNiagaraSystemSimulation, ESPMode::ThreadSafe> GetSystemSimulation();
 
-	bool InitializeSystem();
-	void DestroyInstance();
-	void DestroyInstanceNotComponent();
+	NIAGARA_API bool InitializeSystem();
+	NIAGARA_API void DestroyInstance();
+	NIAGARA_API void DestroyInstanceNotComponent();
 
-	void OnPooledReuse(UWorld* NewWorld);
+	NIAGARA_API void OnPooledReuse(UWorld* NewWorld);
 
 	/*
 	Switch which asset the component is using.
@@ -269,80 +269,80 @@ public:
 	By default existing parameters are reset when we call SetAsset, modify bResetExistingOverrideParameters to leave existing parameter data as is.
 	*/
 	UFUNCTION(BlueprintCallable, Category = Niagara, meta = (DisplayName = "Set Niagara System Asset"))
-	void SetAsset(UNiagaraSystem* InAsset, bool bResetExistingOverrideParameters = true);
+	NIAGARA_API void SetAsset(UNiagaraSystem* InAsset, bool bResetExistingOverrideParameters = true);
 
 	UFUNCTION(BlueprintCallable, Category = Niagara, meta = (DisplayName = "Get Niagara System Asset"))
 	UNiagaraSystem* GetAsset() const { return Asset; }
 
 	UFUNCTION(BlueprintCallable, Category = Niagara, meta = (DisplayName = "Set Forced Solo Mode"))
-	void SetForceSolo(bool bInForceSolo);
+	NIAGARA_API void SetForceSolo(bool bInForceSolo);
 
 	UFUNCTION(BlueprintCallable, Category = Niagara, meta = (DisplayName = "Is In Forced Solo Mode"))
 	bool GetForceSolo()const { return bForceSolo; }
 
 private:
-	bool RequiresSoloMode() const;
-	void UpdateInstanceSoloMode();
+	NIAGARA_API bool RequiresSoloMode() const;
+	NIAGARA_API void UpdateInstanceSoloMode();
 public:
 
 	UFUNCTION(BlueprintCallable, Category = Niagara)
-	void SetGpuComputeDebug(bool bEnableDebug);
+	NIAGARA_API void SetGpuComputeDebug(bool bEnableDebug);
 
 	/**
 	Sets the custom time dilation value for the component.
 	Note: This is only available on components that are in solo mode currently.
 	*/
 	UFUNCTION(BlueprintCallable, Category=Niagara)
-	void SetCustomTimeDilation(float Dilation = 1.0f);
+	NIAGARA_API void SetCustomTimeDilation(float Dilation = 1.0f);
 
 	UFUNCTION(BlueprintCallable, Category = Niagara)
 	float GetCustomTimeDilation() const { return CustomTimeDilation; }
 
 	UFUNCTION(BlueprintCallable, Category = Niagara, meta = (DisplayName = "Get Age Update Mode"))
-	ENiagaraAgeUpdateMode GetAgeUpdateMode() const;
+	NIAGARA_API ENiagaraAgeUpdateMode GetAgeUpdateMode() const;
 
 	/** Sets the age update mode for the System instance. */
 	UFUNCTION(BlueprintCallable, Category = Niagara, meta = (DisplayName = "Set Age Update Mode"))
-	void SetAgeUpdateMode(ENiagaraAgeUpdateMode InAgeUpdateMode);
+	NIAGARA_API void SetAgeUpdateMode(ENiagaraAgeUpdateMode InAgeUpdateMode);
 
 	/** Gets the desired age of the System instance.  This is only relevant when using the DesiredAge age update mode. */
 	UFUNCTION(BlueprintCallable, Category = Niagara, meta = (DisplayName = "Get Desired Age"))
-	float GetDesiredAge() const;
+	NIAGARA_API float GetDesiredAge() const;
 
 	/** Sets the desired age of the System instance.  This is only relevant when using the DesiredAge age update mode. */
 	UFUNCTION(BlueprintCallable, Category = Niagara, meta = (DisplayName = "Set Desired Age"))
-	void SetDesiredAge(float InDesiredAge);
+	NIAGARA_API void SetDesiredAge(float InDesiredAge);
 
 	/** Sets the desired age of the System instance and designates that this change is a seek.  When seeking to a desired age the
 	    The component can optionally prevent rendering. */
 	UFUNCTION(BlueprintCallable, Category = Niagara, meta = (DisplayName = "Seek to Desired Age"))
-	void SeekToDesiredAge(float InDesiredAge);
+	NIAGARA_API void SeekToDesiredAge(float InDesiredAge);
 
 	/** Sets whether or not the system can render while seeking. */
 	UFUNCTION(BlueprintCallable, Category = Niagara, meta = (DisplayName = "Set Can Render While Seeking"))
-	void SetCanRenderWhileSeeking(bool bInCanRenderWhileSeeking);
+	NIAGARA_API void SetCanRenderWhileSeeking(bool bInCanRenderWhileSeeking);
 
 	/** Gets the delta value which is used when seeking from the current age, to the desired age.  This is only relevant
 	when using the DesiredAge age update mode. */
 	UFUNCTION(BlueprintCallable, Category = Niagara, meta = (DisplayName = "Get Desired Age Seek Delta"))
-	float GetSeekDelta() const;
+	NIAGARA_API float GetSeekDelta() const;
 
 	/** Sets the delta value which is used when seeking from the current age, to the desired age.  This is only relevant
 	when using the DesiredAge age update mode. */
 	UFUNCTION(BlueprintCallable, Category = Niagara, meta = (DisplayName = "Set Desired Age Seek Delta"))
-	void SetSeekDelta(float InSeekDelta);
+	NIAGARA_API void SetSeekDelta(float InSeekDelta);
 
 	/** Gets whether or not the delta time used to tick the system instance when using desired age is locked to the seek delta.  When true, the system instance
 	will only be ticked when the desired age has changed by more than the seek delta.  When false the system instance will be ticked by the change in desired
 	age when not seeking. */
 	UFUNCTION(BlueprintCallable, Category = Niagara, meta = (DisplayName = "Get whether or not to lock the desired age delta time to the seek delta."))
-	bool GetLockDesiredAgeDeltaTimeToSeekDelta() const;
+	NIAGARA_API bool GetLockDesiredAgeDeltaTimeToSeekDelta() const;
 
 	/** Sets whether or not the delta time used to tick the system instance when using desired age is locked to the seek delta.  When true, the system instance
 	will only be ticked when the desired age has changed by more than the seek delta.  When false the system instance will be ticked by the change in desired
 	age when not seeking. */
 	UFUNCTION(BlueprintCallable, Category = Niagara, meta = (DisplayName = "Set whether or not to lock the desired age delta time to the seek delta."))
-	void SetLockDesiredAgeDeltaTimeToSeekDelta(bool bLock);
+	NIAGARA_API void SetLockDesiredAgeDeltaTimeToSeekDelta(bool bLock);
 
 	/**
 	Get the maximum CPU time in seconds we will simulate to the desired age, when we go beyond this limit ticks will be processed in the next frame.
@@ -350,7 +350,7 @@ public:
 	Note: The componet will not be visibile if we have ticks to process and SetCanRenderWhileSeeking is set to true
 	*/
 	UFUNCTION(BlueprintCallable, Category = Niagara, meta = (DisplayName = "Get Max Desired Age Tick Delta"))
-	float GetMaxSimTime() const;
+	NIAGARA_API float GetMaxSimTime() const;
 
 	/**
 	Sets the maximum CPU time in seconds we will simulate to the desired age, when we go beyond this limit ticks will be processed in the next frame.
@@ -358,19 +358,19 @@ public:
 	Note: The componet will not be visibile if we have ticks to process and SetCanRenderWhileSeeking is set to true
 	*/
 	UFUNCTION(BlueprintCallable, Category = Niagara, meta = (DisplayName = "Set Max Desired Age Tick Delta"))
-	void SetMaxSimTime(float InMaxTime);
+	NIAGARA_API void SetMaxSimTime(float InMaxTime);
 
 	UFUNCTION(BlueprintCallable, Category = Niagara, meta = (DisplayName = "Set Auto Destroy"))
-	void SetAutoDestroy(bool bInAutoDestroy);
+	NIAGARA_API void SetAutoDestroy(bool bInAutoDestroy);
 
 	UE_DEPRECATED(5.0, "This interface is no longer safe to access directly. Use the interface provided by GetSystemInstanceController instead.")
-	FNiagaraSystemInstance* GetSystemInstance() const;
+	NIAGARA_API FNiagaraSystemInstance* GetSystemInstance() const;
 
 	FNiagaraSystemInstanceControllerPtr GetSystemInstanceController() { return SystemInstanceController; }
 	FNiagaraSystemInstanceControllerConstPtr GetSystemInstanceController() const { return SystemInstanceController; }
 
 	UFUNCTION(BlueprintCallable, Category = Niagara, meta = (DisplayName = "Set Tick Behavior"))
-	void SetTickBehavior(ENiagaraTickBehavior NewTickBehavior);
+	NIAGARA_API void SetTickBehavior(ENiagaraTickBehavior NewTickBehavior);
 
 	UFUNCTION(BlueprintCallable, Category = Niagara, meta = (DisplayName = "Get Tick Behavior"))
 	ENiagaraTickBehavior GetTickBehavior() const { return TickBehavior; }
@@ -381,7 +381,7 @@ public:
 	A default uninitialized box will clear the fixed bounds and revert back to system fixed / dynamic bounds.
 	*/
 	UFUNCTION(BlueprintCallable, Category = Niagara)
-	void SetSystemFixedBounds(FBox LocalBounds);
+	NIAGARA_API void SetSystemFixedBounds(FBox LocalBounds);
 
 	/**
 	Gets the fixed bounds for the system instance.
@@ -389,11 +389,11 @@ public:
 	Note: The returned box may be invalid if no fixed bounds exist
 	*/
 	UFUNCTION(BlueprintCallable, Category = Niagara)
-	FBox GetSystemFixedBounds() const;
+	NIAGARA_API FBox GetSystemFixedBounds() const;
 
 	/** Clear any previously set fixed bounds for the system instance. */
 	UFUNCTION(BlueprintCallable, Category = Niagara)
-	void ClearSystemFixedBounds();
+	NIAGARA_API void ClearSystemFixedBounds();
 
 	/**
 	Sets the fixed bounds for an emitter instance, this overrides all other bounds.
@@ -401,7 +401,7 @@ public:
 	A default uninitialized box will clear the fixed bounds and revert back to emitter fixed / dynamic bounds.
 	*/
 	UFUNCTION(BlueprintCallable, Category = Niagara)
-	void SetEmitterFixedBounds(FName EmitterName, FBox LocalBounds);
+	NIAGARA_API void SetEmitterFixedBounds(FName EmitterName, FBox LocalBounds);
 
 	/**
 	Gets the fixed bounds for an emitter instance.
@@ -409,14 +409,14 @@ public:
 	Note: The returned box may be invalid if no fixed bounds exist
 	*/
 	UFUNCTION(BlueprintCallable, Category = Niagara)
-	FBox GetEmitterFixedBounds(FName EmitterName) const;
+	NIAGARA_API FBox GetEmitterFixedBounds(FName EmitterName) const;
 
 	/** Clear any previously set fixed bounds for the emitter instance. */
 	UFUNCTION(BlueprintCallable, Category = Niagara)
-	void ClearEmitterFixedBounds(FName EmitterName);
+	NIAGARA_API void ClearEmitterFixedBounds(FName EmitterName);
 
 	UFUNCTION(BlueprintCallable, Category = Niagara, meta = (DisplayName = "Set Random Seed Offset"))
-	void SetRandomSeedOffset(int32 NewRandomSeedOffset);
+	NIAGARA_API void SetRandomSeedOffset(int32 NewRandomSeedOffset);
 
 	UFUNCTION(BlueprintCallable, Category = Niagara, meta = (DisplayName = "Get Random Seed Offset"))
 	int32 GetRandomSeedOffset() const { return RandomSeedOffset; }
@@ -424,180 +424,180 @@ public:
 	/** Sets a Niagara FLinearColor parameter by name, overriding locally if necessary.*/
 	UE_DEPRECATED(5.3, "This method will be removed in a future release.  Please update to use the FName variant")
 	UFUNCTION(BlueprintCallable, Category = Niagara, meta = (DisplayName = "Set Niagara Variable By String (LinearColor)", Keywords="user parameter variable color"))
-	void SetNiagaraVariableLinearColor(const FString& InVariableName, const FLinearColor& InValue);
+	NIAGARA_API void SetNiagaraVariableLinearColor(const FString& InVariableName, const FLinearColor& InValue);
 
 	/** Sets a Niagara FLinearColor parameter by name, overriding locally if necessary.*/
 	UFUNCTION(BlueprintCallable, Category = Niagara, meta = (DisplayName = "Set Niagara Variable (LinearColor)", Keywords="user parameter variable color"))
-	void SetVariableLinearColor(FName InVariableName, const FLinearColor& InValue);
+	NIAGARA_API void SetVariableLinearColor(FName InVariableName, const FLinearColor& InValue);
 
 	/** Sets a Niagara Vector4 parameter by name, overriding locally if necessary.*/
 	UE_DEPRECATED(5.3, "This method will be removed in a future release.  Please update to use the FName variant")
 	UFUNCTION(BlueprintCallable, Category = Niagara, meta = (DisplayName = "Set Niagara Variable By String (Vector4)", Keywords="user parameter variable vector"))
-	void SetNiagaraVariableVec4(const FString& InVariableName, const FVector4& InValue);
+	NIAGARA_API void SetNiagaraVariableVec4(const FString& InVariableName, const FVector4& InValue);
 
 	/** Sets a Niagara Vector4 parameter by name, overriding locally if necessary.*/
 	UFUNCTION(BlueprintCallable, Category = Niagara, meta = (DisplayName = "Set Niagara Variable (Vector4)", Keywords="user parameter variable vector"))
-	void SetVariableVec4(FName InVariableName, const FVector4& InValue);
+	NIAGARA_API void SetVariableVec4(FName InVariableName, const FVector4& InValue);
 
 	/** Sets a Niagara quaternion parameter by name, overriding locally if necessary.*/
 	UE_DEPRECATED(5.3, "This method will be removed in a future release.  Please update to use the FName variant")
 	UFUNCTION(BlueprintCallable, Category = Niagara, meta = (DisplayName = "Set Niagara Variable By String (Quaternion)", Keywords="user parameter variable quaternion rotation"))
-	void SetNiagaraVariableQuat(const FString& InVariableName, const FQuat& InValue);
+	NIAGARA_API void SetNiagaraVariableQuat(const FString& InVariableName, const FQuat& InValue);
 
 	/** Sets a Niagara quaternion parameter by name, overriding locally if necessary.*/
 	UFUNCTION(BlueprintCallable, Category = Niagara, meta = (DisplayName = "Set Niagara Variable (Quaternion)", Keywords="user parameter variable quaternion rotation"))
-	void SetVariableQuat(FName InVariableName, const FQuat& InValue);
+	NIAGARA_API void SetVariableQuat(FName InVariableName, const FQuat& InValue);
 
 	/** Sets a Niagara matrix parameter by name, overriding locally if necessary.*/
 	UE_DEPRECATED(5.3, "This method will be removed in a future release.  Please update to use the FName variant")
 	UFUNCTION(BlueprintCallable, Category = Niagara, meta = (DisplayName = "Set Niagara Variable By String (Matrix)", Keywords="user parameter variable matrix"))
-	void SetNiagaraVariableMatrix(const FString& InVariableName, const FMatrix& InValue);
+	NIAGARA_API void SetNiagaraVariableMatrix(const FString& InVariableName, const FMatrix& InValue);
 
 	/** Sets a Niagara matrix parameter by name, overriding locally if necessary.*/
 	UFUNCTION(BlueprintCallable, Category = Niagara, meta = (DisplayName = "Set Niagara Variable (Matrix)", Keywords="user parameter variable matrix"))
-	void SetVariableMatrix(FName InVariableName, const FMatrix& InValue);
+	NIAGARA_API void SetVariableMatrix(FName InVariableName, const FMatrix& InValue);
 
 	/** Sets a Niagara Vector3 parameter by name, overriding locally if necessary.*/
 	UE_DEPRECATED(5.3, "This method will be removed in a future release.  Please update to use the FName variant")
 	UFUNCTION(BlueprintCallable, Category = Niagara, meta = (DisplayName = "Set Niagara Variable By String (Vector3)", Keywords="user parameter variable vector"))
-	void SetNiagaraVariableVec3(const FString& InVariableName, FVector InValue);
+	NIAGARA_API void SetNiagaraVariableVec3(const FString& InVariableName, FVector InValue);
 
 	/** Sets a Niagara Vector3 parameter by name, overriding locally if necessary.*/
 	UFUNCTION(BlueprintCallable, Category = Niagara, meta = (DisplayName = "Set Niagara Variable (Vector3)", Keywords="user parameter variable vector"))
-	void SetVariableVec3(FName InVariableName, FVector InValue);
+	NIAGARA_API void SetVariableVec3(FName InVariableName, FVector InValue);
 
 	/** Sets a Niagara Position parameter by name, overriding locally if necessary.*/
 	UE_DEPRECATED(5.3, "This method will be removed in a future release.  Please update to use the FName variant")
 	UFUNCTION(BlueprintCallable, Category = Niagara, meta = (DisplayName = "Set Niagara Variable By String (Position)", Keywords="user parameter variable vector position lwc"))
-	void SetNiagaraVariablePosition(const FString& InVariableName, FVector InValue);
+	NIAGARA_API void SetNiagaraVariablePosition(const FString& InVariableName, FVector InValue);
 
 	/** Sets a Niagara Position parameter by name, overriding locally if necessary.*/
 	UFUNCTION(BlueprintCallable, Category = Niagara, meta = (DisplayName = "Set Niagara Variable (Position)", Keywords="user parameter variable vector position lwc"))
-	void SetVariablePosition(FName InVariableName, FVector InValue);
+	NIAGARA_API void SetVariablePosition(FName InVariableName, FVector InValue);
 
 	/** Sets a Niagara Vector2 parameter by name, overriding locally if necessary.*/
 	UE_DEPRECATED(5.3, "This method will be removed in a future release.  Please update to use the FName variant")
 	UFUNCTION(BlueprintCallable, Category = Niagara, meta = (DisplayName = "Set Niagara Variable By String (Vector2)", Keywords="user parameter variable vector"))
-	void SetNiagaraVariableVec2(const FString& InVariableName, FVector2D InValue);
+	NIAGARA_API void SetNiagaraVariableVec2(const FString& InVariableName, FVector2D InValue);
 
 	/** Sets a Niagara Vector2 parameter by name, overriding locally if necessary.*/
 	UFUNCTION(BlueprintCallable, Category = Niagara, meta = (DisplayName = "Set Niagara Variable (Vector2)", Keywords="user parameter variable vector"))
-	void SetVariableVec2(FName InVariableName, FVector2D InValue);
+	NIAGARA_API void SetVariableVec2(FName InVariableName, FVector2D InValue);
 
 	/** Sets a Niagara float parameter by name, overriding locally if necessary.*/
 	UE_DEPRECATED(5.3, "This method will be removed in a future release.  Please update to use the FName variant")
 	UFUNCTION(BlueprintCallable, Category = Niagara, meta = (DisplayName = "Set Niagara Variable By String (Float)", Keywords="user parameter variable float"))
-	void SetNiagaraVariableFloat(const FString& InVariableName, float InValue);
+	NIAGARA_API void SetNiagaraVariableFloat(const FString& InVariableName, float InValue);
 
 	/** Sets a Niagara float parameter by name, overriding locally if necessary.*/
 	UFUNCTION(BlueprintCallable, Category = Niagara, meta = (DisplayName = "Set Niagara Variable (Float)", Keywords="user parameter variable float"))
-	void SetVariableFloat(FName InVariableName, float InValue);
+	NIAGARA_API void SetVariableFloat(FName InVariableName, float InValue);
 
 	/** Sets a Niagara int parameter by name, overriding locally if necessary.*/
 	UE_DEPRECATED(5.3, "This method will be removed in a future release.  Please update to use the FName variant")
 	UFUNCTION(BlueprintCallable, Category = Niagara, meta = (DisplayName = "Set Niagara Variable By String (Int32)", Keywords="user parameter variable int"))
-	void SetNiagaraVariableInt(const FString& InVariableName, int32 InValue);
+	NIAGARA_API void SetNiagaraVariableInt(const FString& InVariableName, int32 InValue);
 
 	/** Sets a Niagara int parameter by name, overriding locally if necessary.*/
 	UFUNCTION(BlueprintCallable, Category = Niagara, meta = (DisplayName = "Set Niagara Variable (Int32)", Keywords="user parameter variable int"))
-	void SetVariableInt(FName InVariableName, int32 InValue);
+	NIAGARA_API void SetVariableInt(FName InVariableName, int32 InValue);
 
 	/** Sets a Niagara bool parameter by name, overriding locally if necessary.*/
 	UE_DEPRECATED(5.3, "This method will be removed in a future release.  Please update to use the FName variant")
 	UFUNCTION(BlueprintCallable, Category = Niagara, meta = (DisplayName = "Set Niagara Variable By String (Bool)", Keywords="user parameter variable bool"))
-	void SetNiagaraVariableBool(const FString& InVariableName, bool InValue);
+	NIAGARA_API void SetNiagaraVariableBool(const FString& InVariableName, bool InValue);
 
 	/** Sets a Niagara bool parameter by name, overriding locally if necessary.*/
 	UFUNCTION(BlueprintCallable, Category = Niagara, meta = (DisplayName = "Set Niagara Variable (Bool)", Keywords="user parameter variable bool"))
-	void SetVariableBool(FName InVariableName, bool InValue);
+	NIAGARA_API void SetVariableBool(FName InVariableName, bool InValue);
 
 	UE_DEPRECATED(5.3, "This method will be removed in a future release.  Please update to use the FName variant")
 	UFUNCTION(BlueprintCallable, Category = Niagara, meta = (DisplayName = "Set Niagara Variable By String (Actor)", Keywords="user parameter variable actor"))
-	void SetNiagaraVariableActor(const FString& InVariableName, AActor* Actor);
+	NIAGARA_API void SetNiagaraVariableActor(const FString& InVariableName, AActor* Actor);
 
 	UFUNCTION(BlueprintCallable, Category = Niagara, meta = (DisplayName = "Set Niagara Variable (Actor)", Keywords="user parameter variable actor"))
-	void SetVariableActor(FName InVariableName, AActor* Actor);
+	NIAGARA_API void SetVariableActor(FName InVariableName, AActor* Actor);
 
 	UE_DEPRECATED(5.3, "This method will be removed in a future release.  Please update to use the FName variant")
 	UFUNCTION(BlueprintCallable, Category = Niagara, meta = (DisplayName = "Set Niagara Variable By String (Object)", Keywords="user parameter variable object"))
-	void SetNiagaraVariableObject(const FString& InVariableName, UObject* Object);
+	NIAGARA_API void SetNiagaraVariableObject(const FString& InVariableName, UObject* Object);
 
 	UFUNCTION(BlueprintCallable, Category = Niagara, meta = (DisplayName = "Set Niagara Variable (Object)", Keywords="user parameter variable object"))
-	void SetVariableObject(FName InVariableName, UObject* Object);
+	NIAGARA_API void SetVariableObject(FName InVariableName, UObject* Object);
 
 	UFUNCTION(BlueprintCallable, Category = Niagara, meta = (DisplayName = "Set Niagara Variable (Material)", Keywords="user parameter variable material"))
-	void SetVariableMaterial(FName InVariableName, UMaterialInterface* Object);
+	NIAGARA_API void SetVariableMaterial(FName InVariableName, UMaterialInterface* Object);
 
 	UFUNCTION(BlueprintCallable, Category = Niagara, meta = (DisplayName = "Set Niagara Variable (Static Mesh)", Keywords="user parameter variable mesh"))
-	void SetVariableStaticMesh(FName InVariableName, UStaticMesh* InValue);
+	NIAGARA_API void SetVariableStaticMesh(FName InVariableName, UStaticMesh* InValue);
 
 	UFUNCTION(BlueprintCallable, Category = Niagara, meta = (DisplayName = "Set Niagara Variable (Texture)", Keywords="user parameter variable texture"))
-	void SetVariableTexture(FName InVariableName, class UTexture* Texture);
+	NIAGARA_API void SetVariableTexture(FName InVariableName, class UTexture* Texture);
 
 	UFUNCTION(BlueprintCallable, Category = Niagara, meta = (DisplayName = "Set Niagara Variable (TextureRenderTarget)"))
-	void SetVariableTextureRenderTarget(FName InVariableName, class UTextureRenderTarget* TextureRenderTarget);
+	NIAGARA_API void SetVariableTextureRenderTarget(FName InVariableName, class UTextureRenderTarget* TextureRenderTarget);
 
 	/** Debug accessors for getting positions in blueprints. */
 	UFUNCTION(BlueprintCallable, Category = Niagara, meta = (DisplayName = "Get Niagara Emitter Positions", DeprecatedFunction, DeprecationMessage = "Get Niagara Emitter Positions is deprecated, use the particle export DI inside your emitter instead."))
-	TArray<FVector> GetNiagaraParticlePositions_DebugOnly(const FString& InEmitterName);
+	NIAGARA_API TArray<FVector> GetNiagaraParticlePositions_DebugOnly(const FString& InEmitterName);
 
 	/** Debug accessors for getting a float attribute array in blueprints.  The attribute name should be without namespaces. For example for "Particles.Position", send "Position". */
 	UFUNCTION(BlueprintCallable, Category = Niagara, meta = (DisplayName = "Get Niagara Emitter Float Attrib", DeprecatedFunction, DeprecationMessage = "Get Niagara Emitter Float Attrib is deprecated, use the particle export DI inside your emitter instead."))
-	TArray<float> GetNiagaraParticleValues_DebugOnly(const FString& InEmitterName, const FString& InValueName);
+	NIAGARA_API TArray<float> GetNiagaraParticleValues_DebugOnly(const FString& InEmitterName, const FString& InValueName);
 
 	/** Debug accessors for getting a FVector attribute array in blueprints. The attribute name should be without namespaces. For example for "Particles.Position", send "Position". */
 	UFUNCTION(BlueprintCallable, Category = Niagara, meta = (DisplayName = "Get Niagara Emitter Vec3 Attrib", DeprecatedFunction, DeprecationMessage = "Get Niagara Emitter Vec3 Attrib is deprecated, use the particle export DI inside your emitter instead."))
-	TArray<FVector> GetNiagaraParticleValueVec3_DebugOnly(const FString& InEmitterName, const FString& InValueName);
+	NIAGARA_API TArray<FVector> GetNiagaraParticleValueVec3_DebugOnly(const FString& InEmitterName, const FString& InValueName);
 
 	/** Resets the System to it's initial pre-simulated state. */
 	UFUNCTION(BlueprintCallable, Category = Niagara, meta = (DisplayName = "Reset System"))
-	void ResetSystem();
+	NIAGARA_API void ResetSystem();
 
 	/** Called on when an external object wishes to force this System to reinitialize itself from the System data.*/
 	UFUNCTION(BlueprintCallable, Category = Niagara, meta = (DisplayName = "Reinitialize System"))
-	void ReinitializeSystem();
+	NIAGARA_API void ReinitializeSystem();
 
 	/** Gets whether or not rendering is enabled for this component. */
-	bool GetRenderingEnabled() const;
+	NIAGARA_API bool GetRenderingEnabled() const;
 
 	/** Sets whether or not rendering is enabled for this component. */
 	UFUNCTION(BlueprintCallable, Category = Niagara, meta = (DisplayName = "Set Rendering Enabled"))
-	void SetRenderingEnabled(bool bInRenderingEnabled);
+	NIAGARA_API void SetRenderingEnabled(bool bInRenderingEnabled);
 
 	/** Advances this system's simulation by the specified number of ticks and delta time. */
 	UFUNCTION(BlueprintCallable, Category = Niagara)
-	void AdvanceSimulation(int32 TickCount, float TickDeltaSeconds);
+	NIAGARA_API void AdvanceSimulation(int32 TickCount, float TickDeltaSeconds);
 
 	/** Advances this system's simulation by the specified time in seconds and delta time. Advancement is done in whole ticks of TickDeltaSeconds so actual simulated time will be the nearest lower multiple of TickDeltaSeconds. */
 	UFUNCTION(BlueprintCallable, Category = Niagara)
-	void AdvanceSimulationByTime(float SimulateTime, float TickDeltaSeconds);
+	NIAGARA_API void AdvanceSimulationByTime(float SimulateTime, float TickDeltaSeconds);
 
 	UFUNCTION(BlueprintCallable, Category = Niagara)
-	void SetPaused(bool bInPaused);
+	NIAGARA_API void SetPaused(bool bInPaused);
 
 	UFUNCTION(BlueprintCallable, Category = Niagara)
-	bool IsPaused() const;
+	NIAGARA_API bool IsPaused() const;
 
 	UE_DEPRECATED(4.27, "This method will be removed in a future release. Use the method provided by the Niagara Function Library instead.")
 	UFUNCTION(BlueprintCallable, Category = Niagara)
-	UNiagaraDataInterface* GetDataInterface(const FString &Name);
+	NIAGARA_API UNiagaraDataInterface* GetDataInterface(const FString &Name);
 
 	/**
 		The significant index for this component. i.e. this is the Nth most significant instance of it's system in the scene.
 		Passed to the script to allow us to scale down internally for less significant systems instances.
 	*/
-	void SetSystemSignificanceIndex(int32 InIndex);
+	NIAGARA_API void SetSystemSignificanceIndex(int32 InIndex);
 
 	//~ Begin UObject Interface.
-	virtual void Serialize(FStructuredArchive::FRecord Record) override;
-	virtual void PostLoad() override;
+	NIAGARA_API virtual void Serialize(FStructuredArchive::FRecord Record) override;
+	NIAGARA_API virtual void PostLoad() override;
 #if WITH_EDITORONLY_DATA
-	static void DeclareConstructClasses(TArray<FTopLevelAssetPath>& OutConstructClasses, const UClass* SpecificSubclass);
+	static NIAGARA_API void DeclareConstructClasses(TArray<FTopLevelAssetPath>& OutConstructClasses, const UClass* SpecificSubclass);
 #endif
 
 #if WITH_EDITOR
-	virtual void PreEditChange(FProperty* PropertyAboutToChange) override;
-	virtual void PostEditChangeProperty(FPropertyChangedEvent& PropertyChangedEvent) override;
+	NIAGARA_API virtual void PreEditChange(FProperty* PropertyAboutToChange) override;
+	NIAGARA_API virtual void PostEditChangeProperty(FPropertyChangedEvent& PropertyChangedEvent) override;
 
 	/**
 	  * Find the value of an overridden parameter.  The value returned may not be the current value being used by the simulation but
@@ -605,43 +605,43 @@ public:
 	  * blueprint editor.
 	  * Returns null if the parameter isn't overridden by this component.
 	  */
-	FNiagaraVariant FindParameterOverride(const FNiagaraVariableBase& InKey) const;
+	NIAGARA_API FNiagaraVariant FindParameterOverride(const FNiagaraVariableBase& InKey) const;
 
 	/** Gets the current value of a parameter which is being used by the simulation.  This value will reflect runtime changes such as
 	  * those made by sequencer, or in PIE through code or blueprint. */
-	FNiagaraVariant GetCurrentParameterValue(const FNiagaraVariableBase& InKey) const;
+	NIAGARA_API FNiagaraVariant GetCurrentParameterValue(const FNiagaraVariableBase& InKey) const;
 
-	bool HasParameterOverride(const FNiagaraVariableBase& InKey) const;
-	void SetParameterOverride(const FNiagaraVariableBase& InKey, const FNiagaraVariant& InValue);
+	NIAGARA_API bool HasParameterOverride(const FNiagaraVariableBase& InKey) const;
+	NIAGARA_API void SetParameterOverride(const FNiagaraVariableBase& InKey, const FNiagaraVariant& InValue);
 
 	/** Remove an override for a given parameter if one exists. */
-	void RemoveParameterOverride(const FNiagaraVariableBase& InKey);
+	NIAGARA_API void RemoveParameterOverride(const FNiagaraVariableBase& InKey);
 
-	void UpgradeDeprecatedParameterOverrides();
-	void EnsureOverrideParametersConsistent() const;
+	NIAGARA_API void UpgradeDeprecatedParameterOverrides();
+	NIAGARA_API void EnsureOverrideParametersConsistent() const;
 #endif
 	//~ End UObject Interface
 
 	UFUNCTION(BlueprintCallable, Category = Preview, meta = (Keywords = "preview LOD Distance scalability"))
-	void SetPreviewLODDistance(bool bEnablePreviewLODDistance, float PreviewLODDistance, float PreviewMaxDistance);
+	NIAGARA_API void SetPreviewLODDistance(bool bEnablePreviewLODDistance, float PreviewLODDistance, float PreviewMaxDistance);
 
 	UFUNCTION(BlueprintCallable, Category = Preview, meta = (Keywords = "preview LOD Distance scalability"))
-	FORCEINLINE bool GetPreviewLODDistanceEnabled()const;
+	NIAGARA_API FORCEINLINE bool GetPreviewLODDistanceEnabled()const;
 
 	UFUNCTION(BlueprintCallable, Category = Preview, meta = (Keywords = "preview LOD Distance scalability"))
-	FORCEINLINE float GetPreviewLODDistance()const;
+	NIAGARA_API FORCEINLINE float GetPreviewLODDistance()const;
 
 	/**
 	Initializes this component for capturing a performance baseline.
 	This will do things such as disabling distance culling and setting a LODDistance of 0 to ensure the effect is at it's maximum cost.
 	*/
 	UFUNCTION(BlueprintCallable, Category = Performance, meta = (Keywords = "Niagara Performance"))
-	void InitForPerformanceBaseline();
+	NIAGARA_API void InitForPerformanceBaseline();
 
-	void SetLODDistance(float InLODDistance, float InMaxLODDistance);
+	NIAGARA_API void SetLODDistance(float InLODDistance, float InMaxLODDistance);
 
 #if WITH_EDITOR
-	void PostLoadNormalizeOverrideNames();
+	NIAGARA_API void PostLoadNormalizeOverrideNames();
 	FOnSystemInstanceChanged& OnSystemInstanceChanged() { return OnSystemInstanceChangedDelegate; }
 
 	FOnSynchronizedWithAssetParameters& OnSynchronizedWithAssetParameters() { return OnSynchronizedWithAssetParametersDelegate; }
@@ -651,7 +651,7 @@ public:
 
 	const FNiagaraParameterStore& GetOverrideParameters() const { return OverrideParameters; }
 
-	bool IsWorldReadyToRun() const;
+	NIAGARA_API bool IsWorldReadyToRun() const;
 
 	//~ End UObject Interface.
 
@@ -660,31 +660,31 @@ public:
 	FOnNiagaraSystemFinished OnSystemFinished;
 
 	/** Removes all local overrides and replaces them with the values from the source System - note: this also removes the editor overrides from the component as it is used by the pooling mechanism to prevent values leaking between different instances. */
-	void SetUserParametersToDefaultValues();
+	NIAGARA_API void SetUserParametersToDefaultValues();
 
 	/** Is this an effect on or linked to the local player. */
-	bool IsLocalPlayerEffect()const;
+	NIAGARA_API bool IsLocalPlayerEffect()const;
 
 private:
 	/** Compare local overrides with the source System. Remove any that have mismatched types or no longer exist on the System.*/
-	void SynchronizeWithSourceSystem();
+	NIAGARA_API void SynchronizeWithSourceSystem();
 
-	void FixInvalidUserParameterOverrideData();
+	NIAGARA_API void FixInvalidUserParameterOverrideData();
 
-	void AssetExposedParametersChanged();
+	NIAGARA_API void AssetExposedParametersChanged();
 
-	void CopyParametersFromAsset(bool bResetExistingOverrideParameters = true);
+	NIAGARA_API void CopyParametersFromAsset(bool bResetExistingOverrideParameters = true);
 
 #if WITH_EDITOR
-	void SetOverrideParameterStoreValue(const FNiagaraVariableBase& InKey, const FNiagaraVariant& InValue);
-	void ApplyOverridesToParameterStore();
+	NIAGARA_API void SetOverrideParameterStoreValue(const FNiagaraVariableBase& InKey, const FNiagaraVariant& InValue);
+	NIAGARA_API void ApplyOverridesToParameterStore();
 
 	/**
 	 * Ensures that the data interfaces in the template parameter overrides, instance parameter overrides, and in the override parameters
 	 * parameter store have this component as their outer and have their archetype cleared.  This is required to hack around issues with
 	 * subobject handling on components when they're spawned from templates which can cause invalid cross package references.
  	 */
-	void FixDataInterfaceOuters();
+	NIAGARA_API void FixDataInterfaceOuters();
 #endif
 
 public:
@@ -734,7 +734,7 @@ public:
 	 * @param  ScaleRule		Option for how we handle our scale when we attach to Parent.
 	 * @see bAutoManageAttachment, AutoAttachParent, AutoAttachSocketName, AutoAttachLocationRule, AutoAttachRotationRule, AutoAttachScaleRule
 	 */
-	void SetAutoAttachmentParameters(USceneComponent* Parent, FName SocketName, EAttachmentRule LocationRule, EAttachmentRule RotationRule, EAttachmentRule ScaleRule) override;
+	NIAGARA_API void SetAutoAttachmentParameters(USceneComponent* Parent, FName SocketName, EAttachmentRule LocationRule, EAttachmentRule RotationRule, EAttachmentRule ScaleRule) override;
 
 	virtual void SetUseAutoManageAttachment(bool bAutoManage) override { bAutoManageAttachment = bAutoManage; }
 
@@ -755,13 +755,13 @@ public:
 	simulation due to missing simulation data.
 	*/
 	UFUNCTION(BlueprintCallable, Category = SimCache)
-	void SetSimCache(UNiagaraSimCache* SimCache, bool bResetSystem = false);
+	NIAGARA_API void SetSimCache(UNiagaraSimCache* SimCache, bool bResetSystem = false);
 
 	/**
 	Get the active simulation cache, will return null if we do not have an active one.
 	*/
 	UFUNCTION(BlueprintCallable, Category = SimCache)
-	UNiagaraSimCache* GetSimCache() const;
+	NIAGARA_API UNiagaraSimCache* GetSimCache() const;
 
 	/**
 	Clear any active simulation cache.
@@ -774,16 +774,16 @@ public:
 
 	/** Set whether this component is allowed to perform scalability checks and potentially be culled etc. Occasionally it is useful to disable this for specific components. E.g. Effects on the local player. */
 	UFUNCTION(BlueprintSetter, Category = Scalability, meta = (Keywords = "LOD scalability"))
-	void SetAllowScalability(bool bAllow);
+	NIAGARA_API void SetAllowScalability(bool bAllow);
 	
 	UFUNCTION(BlueprintGetter)
-	bool GetAllowScalability()const;
+	NIAGARA_API bool GetAllowScalability()const;
 	
 	UFUNCTION(BlueprintSetter, Category = Scalability, meta = (Keywords = "LOD scalability"))
-	void SetForceLocalPlayerEffect(bool bIsPlayerEffect);
+	NIAGARA_API void SetForceLocalPlayerEffect(bool bIsPlayerEffect);
 	
 	UFUNCTION(BlueprintGetter)
-	bool GetForceLocalPlayerEffect()const;
+	NIAGARA_API bool GetForceLocalPlayerEffect()const;
 
 	FORCEINLINE bool IsRegisteredWithScalabilityManager()const { return ScalabilityManagerHandle != INDEX_NONE; }
 	FORCEINLINE int32 GetScalabilityManagerHandle()const { return ScalabilityManagerHandle; }
@@ -798,7 +798,7 @@ public:
 
 	FORCEINLINE bool IsUsingCullProxy()const { return CullProxy != nullptr; }
 
-	bool ResolveOwnerAllowsScalability(bool bRegister=true);
+	NIAGARA_API bool ResolveOwnerAllowsScalability(bool bRegister=true);
 
 private:
 	/** Did we try and activate but fail due to the asset being not yet ready. Keep looping.*/
@@ -833,7 +833,7 @@ private:
 	uint32 bDesiredPauseState : 1;
 
 	/** Restore relative transform from auto attachment and optionally detach from parent (regardless of whether it was an auto attachment). */
-	void CancelAutoAttachment(bool bDetachFromParent);
+	NIAGARA_API void CancelAutoAttachment(bool bDetachFromParent);
 
 
 	/** Saved relative transform before auto attachment. Used during detachment to restore the transform if we had automatically attached. */
@@ -861,11 +861,11 @@ private:
 	UPROPERTY(transient)
 	TObjectPtr<class UNiagaraCullProxyComponent> CullProxy;
 
-	void CreateCullProxy(bool bForce = false);
-	void DestroyCullProxy();
+	NIAGARA_API void CreateCullProxy(bool bForce = false);
+	NIAGARA_API void DestroyCullProxy();
 
 public:
-	FParticlePerfStatsContext GetPerfStatsContext();
+	NIAGARA_API FParticlePerfStatsContext GetPerfStatsContext();
 };
 
 FORCEINLINE bool UNiagaraComponent::GetPreviewLODDistanceEnabled()const { return bEnablePreviewLODDistance; }
diff --git a/Engine/Plugins/FX/Niagara/Source/Niagara/Public/NiagaraComponentPool.h b/Engine/Plugins/FX/Niagara/Source/Niagara/Public/NiagaraComponentPool.h
--- a/Engine/Plugins/FX/Niagara/Source/Niagara/Public/NiagaraComponentPool.h
+++ b/Engine/Plugins/FX/Niagara/Source/Niagara/Public/NiagaraComponentPool.h
@@ -65,8 +65,8 @@ public:
 	int32 NumComponents() { return FreeElements.Num(); }
 };
 
-UCLASS(Transient)
-class NIAGARA_API UNiagaraComponentPool : public UObject
+UCLASS(Transient, MinimalAPI)
+class UNiagaraComponentPool : public UObject
 {
 	GENERATED_UCLASS_BODY()
 
@@ -77,29 +77,29 @@ private:
 	double LastParticleSytemPoolCleanTime;
 
 public:
-	static bool Enabled();
+	static NIAGARA_API bool Enabled();
 
-	~UNiagaraComponentPool();
+	NIAGARA_API ~UNiagaraComponentPool();
 
-	void Cleanup(UWorld* World);
+	NIAGARA_API void Cleanup(UWorld* World);
 
 	/** Clear all free entires of the specified system. */
-	void ClearPool(UNiagaraSystem* System);
+	NIAGARA_API void ClearPool(UNiagaraSystem* System);
 
-	void PrimePool(UNiagaraSystem* Template, UWorld* World);
-	UNiagaraComponent* CreateWorldParticleSystem(UNiagaraSystem* Template, UWorld* World, ENCPoolMethod PoolingMethod);
+	NIAGARA_API void PrimePool(UNiagaraSystem* Template, UWorld* World);
+	NIAGARA_API UNiagaraComponent* CreateWorldParticleSystem(UNiagaraSystem* Template, UWorld* World, ENCPoolMethod PoolingMethod);
 
 	/** Called when an in-use particle component is finished and wishes to be returned to the pool. */
-	void ReclaimWorldParticleSystem(UNiagaraComponent* Component);
+	NIAGARA_API void ReclaimWorldParticleSystem(UNiagaraComponent* Component);
 
 	/** Notification that the component is being destroyed but has relevance to the component pool. */
-	void PooledComponentDestroyed(UNiagaraComponent* Component);
+	NIAGARA_API void PooledComponentDestroyed(UNiagaraComponent* Component);
 
 	/** Remove any components that are using this system. */
-	void RemoveComponentsBySystem(UNiagaraSystem* System);
+	NIAGARA_API void RemoveComponentsBySystem(UNiagaraSystem* System);
 
 	/** Dumps the current state of the pool to the log. */
-	void Dump();
+	NIAGARA_API void Dump();
 
 #if ENABLE_NC_POOL_DEBUGGING
 	//Array of currently in flight components that will auto release.
diff --git a/Engine/Plugins/FX/Niagara/Source/Niagara/Public/NiagaraConvertInPlaceUtilityBase.h b/Engine/Plugins/FX/Niagara/Source/Niagara/Public/NiagaraConvertInPlaceUtilityBase.h
--- a/Engine/Plugins/FX/Niagara/Source/Niagara/Public/NiagaraConvertInPlaceUtilityBase.h
+++ b/Engine/Plugins/FX/Niagara/Source/Niagara/Public/NiagaraConvertInPlaceUtilityBase.h
@@ -13,10 +13,10 @@ class UNiagaraNodeFunctionCall;
 class UNiagaraClipboardContent;
 class UNiagaraStackFunctionInputCollection;
 
-UCLASS(Abstract)
-class NIAGARA_API UNiagaraConvertInPlaceUtilityBase : public UObject
+UCLASS(Abstract, MinimalAPI)
+class UNiagaraConvertInPlaceUtilityBase : public UObject
 {
 	GENERATED_BODY()
 public:
 	virtual bool Convert(UNiagaraScript* InOldScript, UNiagaraClipboardContent* InOldClipboardContent, UNiagaraScript* InNewScript, UNiagaraStackFunctionInputCollection* InInputCollection, UNiagaraClipboardContent* InNewClipboardContent, UNiagaraNodeFunctionCall* InCallingNode, FText& OutMessage) { return true; };
-};
\ No newline at end of file
+};
diff --git a/Engine/Plugins/FX/Niagara/Source/Niagara/Public/NiagaraDataChannel.h b/Engine/Plugins/FX/Niagara/Source/Niagara/Public/NiagaraDataChannel.h
--- a/Engine/Plugins/FX/Niagara/Source/Niagara/Public/NiagaraDataChannel.h
+++ b/Engine/Plugins/FX/Niagara/Source/Niagara/Public/NiagaraDataChannel.h
@@ -81,19 +81,19 @@ struct FNiagaraDataChannelDataProxy
 
 DECLARE_MULTICAST_DELEGATE_OneParam(FOnDataChannelCreated, const UNiagaraDataChannel*);
 
-UCLASS(Experimental, abstract, EditInlineNew)
-class NIAGARA_API UNiagaraDataChannel : public UObject
+UCLASS(Experimental, abstract, EditInlineNew, MinimalAPI)
+class UNiagaraDataChannel : public UObject
 {
 public:
 	GENERATED_BODY()
 
 	//UObject Interface
-	virtual void PostInitProperties()override;
-	virtual void PostLoad();
-	virtual void BeginDestroy();
+	NIAGARA_API virtual void PostInitProperties()override;
+	NIAGARA_API virtual void PostLoad();
+	NIAGARA_API virtual void BeginDestroy();
 #if WITH_EDITOR
-	virtual void PreEditChange(FProperty* PropertyAboutToChange)override;
-	virtual void PostEditChangeProperty(FPropertyChangedEvent& PropertyChangedDataChannel)override;
+	NIAGARA_API virtual void PreEditChange(FProperty* PropertyAboutToChange)override;
+	NIAGARA_API virtual void PostEditChangeProperty(FPropertyChangedEvent& PropertyChangedDataChannel)override;
 #endif
 	//UObject Interface End.
 
@@ -104,14 +104,14 @@ public:
 	bool KeepPreviousFrameData()const { return bKeepPreviousFrameData; }
 
 	/** Returns the compiled data describing the data layout for DataChannels in this channel. */
-	const FNiagaraDataSetCompiledData& GetCompiledData(ENiagaraSimTarget SimTarget)const;
+	NIAGARA_API const FNiagaraDataSetCompiledData& GetCompiledData(ENiagaraSimTarget SimTarget)const;
 
 	/** Create the appropriate handler object for this data channel. */
-	virtual UNiagaraDataChannelHandler* CreateHandler(UWorld* OwningWorld) const PURE_VIRTUAL(UNiagaraDataChannel::CreateHandler, {return nullptr;} );
+	NIAGARA_API virtual UNiagaraDataChannelHandler* CreateHandler(UWorld* OwningWorld) const PURE_VIRTUAL(UNiagaraDataChannel::CreateHandler, {return nullptr;} );
 	
 	const FNiagaraDataChannelGameDataLayout& GetGameDataLayout()const { return GameDataLayout; }
 
-	FNiagaraDataChannelGameDataPtr CreateGameData()const;
+	NIAGARA_API FNiagaraDataChannelGameDataPtr CreateGameData()const;
 
 	bool IsValid()const;
 
