[Niagara]

diff --git a/Engine/Plugins/FX/Niagara/Source/Niagara/Public/NiagaraDataChannelAccessor.h b/Engine/Plugins/FX/Niagara/Source/Niagara/Public/NiagaraDataChannelAccessor.h
--- a/Engine/Plugins/FX/Niagara/Source/Niagara/Public/NiagaraDataChannelAccessor.h
+++ b/Engine/Plugins/FX/Niagara/Source/Niagara/Public/NiagaraDataChannelAccessor.h
@@ -16,8 +16,8 @@ Likely to be replaced in the near future with a custom BP node and a helper stru
 */
 
 
-UCLASS(Experimental, BlueprintType)
-class NIAGARA_API UNiagaraDataChannelReader : public UObject
+UCLASS(Experimental, BlueprintType, MinimalAPI)
+class UNiagaraDataChannelReader : public UObject
 {
 	GENERATED_BODY()
 private:
@@ -35,41 +35,41 @@ public:
 	
 	/** Call before each access to the data channel to grab the correct data to read. */
 	UFUNCTION(BlueprintCallable, Category = NiagaraDataChannel, meta = (Keywords = "Niagara DataChannel"))
-	bool InitAccess(FNiagaraDataChannelSearchParameters SearchParams, bool bReadPrevFrameData);
+	NIAGARA_API bool InitAccess(FNiagaraDataChannelSearchParameters SearchParams, bool bReadPrevFrameData);
 
 	UFUNCTION(BlueprintCallable, Category = NiagaraDataChannel, meta = (Keywords = "Niagara DataChannel"))
-	int32 Num()const;
+	NIAGARA_API int32 Num()const;
 
 	UFUNCTION(BlueprintCallable, Category = NiagaraDataChannel, meta = (Keywords = "Niagara DataChannel"))
-	double ReadFloat(FName VarName, int32 Index)const;
+	NIAGARA_API double ReadFloat(FName VarName, int32 Index)const;
 
 	UFUNCTION(BlueprintCallable, Category = NiagaraDataChannel, meta = (Keywords = "Niagara DataChannel"))
-	FVector2D ReadVector2D(FName VarName, int32 Index)const;
+	NIAGARA_API FVector2D ReadVector2D(FName VarName, int32 Index)const;
 
 	UFUNCTION(BlueprintCallable, Category = NiagaraDataChannel, meta = (Keywords = "Niagara DataChannel"))
-	FVector ReadVector(FName VarName, int32 Index)const;
+	NIAGARA_API FVector ReadVector(FName VarName, int32 Index)const;
 
 	UFUNCTION(BlueprintCallable, Category = NiagaraDataChannel, meta = (Keywords = "Niagara DataChannel"))
-	FVector4 ReadVector4(FName VarName, int32 Index)const;
+	NIAGARA_API FVector4 ReadVector4(FName VarName, int32 Index)const;
 
 	UFUNCTION(BlueprintCallable, Category = NiagaraDataChannel, meta = (Keywords = "Niagara DataChannel"))
-	FQuat ReadQuat(FName VarName, int32 Index)const;
+	NIAGARA_API FQuat ReadQuat(FName VarName, int32 Index)const;
 
 	UFUNCTION(BlueprintCallable, Category = NiagaraDataChannel, meta = (Keywords = "Niagara DataChannel"))
-	FLinearColor ReadLinearColor(FName VarName, int32 Index)const;
+	NIAGARA_API FLinearColor ReadLinearColor(FName VarName, int32 Index)const;
 
 	UFUNCTION(BlueprintCallable, Category = NiagaraDataChannel, meta = (Keywords = "Niagara DataChannel"))
-	int32 ReadInt(FName VarName, int32 Index)const;
+	NIAGARA_API int32 ReadInt(FName VarName, int32 Index)const;
 
 	UFUNCTION(BlueprintCallable, Category = NiagaraDataChannel, meta = (Keywords = "Niagara DataChannel"))
-	bool ReadBool(FName VarName, int32 Index)const;
+	NIAGARA_API bool ReadBool(FName VarName, int32 Index)const;
 
 	UFUNCTION(BlueprintCallable, Category = NiagaraDataChannel, meta = (Keywords = "Niagara DataChannel"))
-	FVector ReadPosition(FName VarName, int32 Index)const;
+	NIAGARA_API FVector ReadPosition(FName VarName, int32 Index)const;
 };
 
-UCLASS(Experimental, BlueprintType)
-class NIAGARA_API UNiagaraDataChannelWriter : public UObject
+UCLASS(Experimental, BlueprintType, MinimalAPI)
+class UNiagaraDataChannelWriter : public UObject
 {
 	GENERATED_BODY()
 private:
@@ -87,46 +87,46 @@ public:
 	
 	/** Call before each batch of writes to allocate the data we'll be writing to. */
 	UFUNCTION(BlueprintCallable, Category = NiagaraDataChannel, meta = (Keywords = "niagara DataChannel"))
-	bool InitWrite(FNiagaraDataChannelSearchParameters SearchParams, int32 Count, bool bVisibleToGame=true, bool bVisibleToCPU=true, bool bVisibleToGPU=true);
+	NIAGARA_API bool InitWrite(FNiagaraDataChannelSearchParameters SearchParams, int32 Count, bool bVisibleToGame=true, bool bVisibleToCPU=true, bool bVisibleToGPU=true);
 
 	UFUNCTION(BlueprintCallable, Category = NiagaraDataChannel, meta = (Keywords = "niagara DataChannel"))
-	int32 Num()const;
+	NIAGARA_API int32 Num()const;
 
 	UFUNCTION(BlueprintCallable, Category = NiagaraDataChannel, meta = (Keywords = "niagara DataChannel"))
-	void WriteFloat(FName VarName, int32 Index, double InData);
+	NIAGARA_API void WriteFloat(FName VarName, int32 Index, double InData);
 	
 	UFUNCTION(BlueprintCallable, Category = NiagaraDataChannel, meta = (Keywords = "niagara DataChannel"))
-	void WriteVector2D(FName VarName, int32 Index, FVector2D InData);
+	NIAGARA_API void WriteVector2D(FName VarName, int32 Index, FVector2D InData);
 
 	UFUNCTION(BlueprintCallable, Category = NiagaraDataChannel, meta = (Keywords = "niagara DataChannel"))
-	void WriteVector(FName VarName, int32 Index, FVector InData);
+	NIAGARA_API void WriteVector(FName VarName, int32 Index, FVector InData);
 
 	UFUNCTION(BlueprintCallable, Category = NiagaraDataChannel, meta = (Keywords = "niagara DataChannel"))
-	void WriteVector4(FName VarName, int32 Index, FVector4 InData);
+	NIAGARA_API void WriteVector4(FName VarName, int32 Index, FVector4 InData);
 
 	UFUNCTION(BlueprintCallable, Category = NiagaraDataChannel, meta = (Keywords = "niagara DataChannel"))
-	void WriteQuat(FName VarName, int32 Index, FQuat InData);
+	NIAGARA_API void WriteQuat(FName VarName, int32 Index, FQuat InData);
 
 	UFUNCTION(BlueprintCallable, Category = NiagaraDataChannel, meta = (Keywords = "niagara DataChannel"))
-	void WriteLinearColor(FName VarName, int32 Index, FLinearColor InData);
+	NIAGARA_API void WriteLinearColor(FName VarName, int32 Index, FLinearColor InData);
 	
 	UFUNCTION(BlueprintCallable, Category = NiagaraDataChannel, meta = (Keywords = "niagara DataChannel"))
-	void WriteInt(FName VarName, int32 Index, int32 InData);
+	NIAGARA_API void WriteInt(FName VarName, int32 Index, int32 InData);
 
 	UFUNCTION(BlueprintCallable, Category = NiagaraDataChannel, meta = (Keywords = "niagara DataChannel"))
-	void WriteBool(FName VarName, int32 Index, bool InData);
+	NIAGARA_API void WriteBool(FName VarName, int32 Index, bool InData);
 
 	UFUNCTION(BlueprintCallable, Category = NiagaraDataChannel, meta = (Keywords = "niagara DataChannel"))
-	void WriteSpawnInfo(FName VarName, int32 Index, FNiagaraSpawnInfo InData);
+	NIAGARA_API void WriteSpawnInfo(FName VarName, int32 Index, FNiagaraSpawnInfo InData);
 
 	UFUNCTION(BlueprintCallable, Category = NiagaraDataChannel, meta = (Keywords = "niagara DataChannel"))
-	void WritePosition(FName VarName, int32 Index, FVector InData);
+	NIAGARA_API void WritePosition(FName VarName, int32 Index, FVector InData);
 };
 
 
 
 /** Base clas for template data channel writer objects. Provides common, non-template functionality. */
-struct NIAGARA_API FNiagaraDataChannelGameDataWriterBase
+struct FNiagaraDataChannelGameDataWriterBase
 {
 	FNiagaraDataChannelGameDataWriterBase(FNiagaraDataChannelGameDataPtr InData)
 	: Data(InData)
@@ -134,22 +134,22 @@ struct NIAGARA_API FNiagaraDataChannelGameDataWriterBase
 	}
 
 	/** Call before we begin writing to this data. */
-	void BeginWrite();
+	NIAGARA_API void BeginWrite();
 
 	/** Publish the contents of this writer's Data buffer to the given data channel data. */
-	void Publish(FNiagaraDataChannelDataPtr& Destination);
+	NIAGARA_API void Publish(FNiagaraDataChannelDataPtr& Destination);
 
 	//Binds a specific variable to an index.
-	bool BindVariable(int32 VariableIndex, const FNiagaraVariableBase& Var);
+	NIAGARA_API bool BindVariable(int32 VariableIndex, const FNiagaraVariableBase& Var);
 
 	// Set the size of buffers for all variables.
-	void SetNum(int32 Num);
+	NIAGARA_API void SetNum(int32 Num);
 
 	// Reserve space in buffers for all variables.
-	void Reserve(int32 Num);
+	NIAGARA_API void Reserve(int32 Num);
 
 	/** Adds count elements to all variable buffers. */
-	int32 Add(int32 Count = 1);
+	NIAGARA_API int32 Add(int32 Count = 1);
 
 	FNiagaraDataChannelGameDataPtr& GetData(){ return Data; }
 
@@ -249,9 +249,9 @@ struct NIAGARA_API FNiagaraDataChannelGameDataWriter<FExampleNDCData> : public F
 // Results will be published to the NDC when the scoped writer is destroyed.
 // Call Append(), passing each item's GroupIndex and the data to write.
 
-struct NIAGARA_API FNiagaraDataChannelGameDataGroupedWriterBase
+struct FNiagaraDataChannelGameDataGroupedWriterBase
 {
-	FNiagaraDataChannelDataPtr FindDataChannelData(UWorld* World, const UNiagaraDataChannel* NDC, const FNiagaraDataChannelSearchParameters& ItemSearchParams);
+	NIAGARA_API FNiagaraDataChannelDataPtr FindDataChannelData(UWorld* World, const UNiagaraDataChannel* NDC, const FNiagaraDataChannelSearchParameters& ItemSearchParams);
 };
 
 template<typename T>
diff --git a/Engine/Plugins/FX/Niagara/Source/Niagara/Public/NiagaraDataChannelHandler.h b/Engine/Plugins/FX/Niagara/Source/Niagara/Public/NiagaraDataChannelHandler.h
--- a/Engine/Plugins/FX/Niagara/Source/Niagara/Public/NiagaraDataChannelHandler.h
+++ b/Engine/Plugins/FX/Niagara/Source/Niagara/Public/NiagaraDataChannelHandler.h
@@ -14,41 +14,41 @@ Some more complex handlers may want to divide up the scene in various different
 #include "NiagaraDataChannelPublic.h"
 #include "NiagaraDataChannelHandler.generated.h"
 
-UCLASS(Experimental, abstract, BlueprintType)
-class NIAGARA_API UNiagaraDataChannelHandler : public UObject
+UCLASS(Experimental, abstract, BlueprintType, MinimalAPI)
+class UNiagaraDataChannelHandler : public UObject
 {
 public:
 
 	GENERATED_BODY()
 
-	virtual ~UNiagaraDataChannelHandler();
+	NIAGARA_API virtual ~UNiagaraDataChannelHandler();
 
-	virtual void Init(const UNiagaraDataChannel* InChannel);
+	NIAGARA_API virtual void Init(const UNiagaraDataChannel* InChannel);
 
-	virtual void BeginFrame(float DeltaTime, FNiagaraWorldManager* OwningWorld);
+	NIAGARA_API virtual void BeginFrame(float DeltaTime, FNiagaraWorldManager* OwningWorld);
 
-	virtual void EndFrame(float DeltaTime, FNiagaraWorldManager* OwningWorld);
+	NIAGARA_API virtual void EndFrame(float DeltaTime, FNiagaraWorldManager* OwningWorld);
 
-	virtual void Tick(float DeltaTime, ETickingGroup TickGroup, FNiagaraWorldManager* OwningWorld);
+	NIAGARA_API virtual void Tick(float DeltaTime, ETickingGroup TickGroup, FNiagaraWorldManager* OwningWorld);
 
 	/** 
 	Finds the correct internal data for this data channel and the given search parameters. For example in some cases this may require a search of several elements of data that correspond to different spacial regions.
 	This shared ptr provides access to Game level data, CPU simulation data and a render thread proxy that can be given to the RT and provides access to GPUSimulaiton data.
 	*/
-	virtual FNiagaraDataChannelDataPtr FindData(FNiagaraDataChannelSearchParameters SearchParams, ENiagaraResourceAccess AccessType)  PURE_VIRTUAL(UNiagaraDataChannelHandler::FindData, return nullptr;);
+	NIAGARA_API virtual FNiagaraDataChannelDataPtr FindData(FNiagaraDataChannelSearchParameters SearchParams, ENiagaraResourceAccess AccessType)  PURE_VIRTUAL(UNiagaraDataChannelHandler::FindData, return nullptr;);
 
 	const UNiagaraDataChannel* GetDataChannel()const { return DataChannel; }
 
 	UFUNCTION(BlueprintCallable, Category="Data Channel")
-	UNiagaraDataChannelWriter* GetDataChannelWriter();
+	NIAGARA_API UNiagaraDataChannelWriter* GetDataChannelWriter();
 
 	UFUNCTION(BlueprintCallable, Category = "Data Channel")
-	UNiagaraDataChannelReader* GetDataChannelReader();
+	NIAGARA_API UNiagaraDataChannelReader* GetDataChannelReader();
 
 	template<typename T> 
 	T* GetChannelTyped()const{ return Cast<T>(DataChannel); }
 
-	FNiagaraDataChannelDataPtr CreateData();
+	NIAGARA_API FNiagaraDataChannelDataPtr CreateData();
 
 protected:
 
@@ -62,4 +62,4 @@ protected:
 	/** Helper object allowing BP to read data in this channel. */
 	UPROPERTY()
 	TObjectPtr<UNiagaraDataChannelReader> Reader;
-};
\ No newline at end of file
+};
diff --git a/Engine/Plugins/FX/Niagara/Source/Niagara/Public/NiagaraDataChannelPublic.h b/Engine/Plugins/FX/Niagara/Source/Niagara/Public/NiagaraDataChannelPublic.h
--- a/Engine/Plugins/FX/Niagara/Source/Niagara/Public/NiagaraDataChannelPublic.h
+++ b/Engine/Plugins/FX/Niagara/Source/Niagara/Public/NiagaraDataChannelPublic.h
@@ -16,8 +16,8 @@ using FNiagaraDataChannelGameDataPtr = TSharedPtr<FNiagaraDataChannelGameData>;
 using FNiagaraDataChannelDataPtr = TSharedPtr<FNiagaraDataChannelData>;
 
 /** Wrapper asset class for UNiagaraDataChannel which is instanced. */
-UCLASS(Experimental, DisplayName = "Niagara Data Channel")
-class NIAGARA_API UNiagaraDataChannelAsset : public UObject
+UCLASS(Experimental, DisplayName = "Niagara Data Channel", MinimalAPI)
+class UNiagaraDataChannelAsset : public UObject
 {
 	GENERATED_BODY()
 
@@ -33,8 +33,8 @@ class NIAGARA_API UNiagaraDataChannelAsset : public UObject
 public:
 
 #if WITH_EDITOR
-	virtual void PreEditChange(class FProperty* PropertyAboutToChange)override;
-	virtual void PostEditChangeProperty(FPropertyChangedEvent& PropertyChangedEvent)override;
+	NIAGARA_API virtual void PreEditChange(class FProperty* PropertyAboutToChange)override;
+	NIAGARA_API virtual void PostEditChangeProperty(FPropertyChangedEvent& PropertyChangedEvent)override;
 #endif
 
 	UNiagaraDataChannel* Get() const { return DataChannel; }
@@ -50,7 +50,7 @@ Some data channels will sub divide their data internally in various ways, e.g.,
 These parameters allow users to search for the correct internal data when reading and writing.
 */
 USTRUCT(BlueprintType)
-struct NIAGARA_API FNiagaraDataChannelSearchParameters
+struct FNiagaraDataChannelSearchParameters
 {
 	GENERATED_BODY()
 
@@ -62,7 +62,7 @@ struct NIAGARA_API FNiagaraDataChannelSearchParameters
 	UPROPERTY(EditAnywhere, BlueprintReadWrite, Category = "Parameters")
 	FVector Location = FVector::ZeroVector;
 
-	FVector GetLocation()const;
+	NIAGARA_API FVector GetLocation()const;
 };
 
 USTRUCT()
@@ -85,15 +85,15 @@ struct FNiagaraDataChannelGameDataLayout
 #if !UE_BUILD_SHIPPING
 
 /** Hooks into internal NiagaraDataChannels code for debugging and testing purposes. */
-class NIAGARA_API FNiagaraDataChannelDebugUtilities
+class FNiagaraDataChannelDebugUtilities
 {
 public: 
 
-	static void BeginFrame(FNiagaraWorldManager* WorldMan, float DeltaSeconds);
-	static void EndFrame(FNiagaraWorldManager* WorldMan, float DeltaSeconds);
-	static void Tick(FNiagaraWorldManager* WorldMan, float DeltaSeconds, ETickingGroup TickGroup);
+	static NIAGARA_API void BeginFrame(FNiagaraWorldManager* WorldMan, float DeltaSeconds);
+	static NIAGARA_API void EndFrame(FNiagaraWorldManager* WorldMan, float DeltaSeconds);
+	static NIAGARA_API void Tick(FNiagaraWorldManager* WorldMan, float DeltaSeconds, ETickingGroup TickGroup);
 	
-	static UNiagaraDataChannelHandler* FindDataChannelHandler(FNiagaraWorldManager* WorldMan, UNiagaraDataChannel* DataChannel);
+	static NIAGARA_API UNiagaraDataChannelHandler* FindDataChannelHandler(FNiagaraWorldManager* WorldMan, UNiagaraDataChannel* DataChannel);
 };
 
 
@@ -184,7 +184,7 @@ struct FNiagaraDataChannelVariableBuffer
 
 
 /** Storage for game level DataChannels generated by BP / C++ */
-struct NIAGARA_API FNiagaraDataChannelGameData : public TSharedFromThis<FNiagaraDataChannelGameData>
+struct FNiagaraDataChannelGameData : public TSharedFromThis<FNiagaraDataChannelGameData>
 {
 private:
 
@@ -198,11 +198,11 @@ private:
 
 public:
 
-	void Init(const UNiagaraDataChannel* InDataChannel);
-	void Empty();
-	void BeginFrame();
-	void SetNum(int32 NewNum);
-	void Reserve(int32 NewNum);
+	NIAGARA_API void Init(const UNiagaraDataChannel* InDataChannel);
+	NIAGARA_API void Empty();
+	NIAGARA_API void BeginFrame();
+	NIAGARA_API void SetNum(int32 NewNum);
+	NIAGARA_API void Reserve(int32 NewNum);
 	int32 Num() const { return NumElements; }
 	int32 PrevNum()const { return PrevNumElements; }
 
@@ -213,13 +213,13 @@ public:
 		return Ret;
 	}
 
-	FNiagaraDataChannelVariableBuffer* FindVariableBuffer(const FNiagaraVariableBase& Var);
+	NIAGARA_API FNiagaraDataChannelVariableBuffer* FindVariableBuffer(const FNiagaraVariableBase& Var);
 
-	void WriteToDataSet(FNiagaraDataBuffer* DestBuffer, int32 DestStartIdx, FVector3f SimulationLwcTile);
+	NIAGARA_API void WriteToDataSet(FNiagaraDataBuffer* DestBuffer, int32 DestStartIdx, FVector3f SimulationLwcTile);
 
-	void AppendFromGameData(const FNiagaraDataChannelGameData& GameData);
+	NIAGARA_API void AppendFromGameData(const FNiagaraDataChannelGameData& GameData);
 
-	void AppendFromDataSet(const FNiagaraDataBuffer* SrcBuffer, FVector3f SimulationLwcTile);
+	NIAGARA_API void AppendFromDataSet(const FNiagaraDataBuffer* SrcBuffer, FVector3f SimulationLwcTile);
 
 	const UNiagaraDataChannel* GetDataChannel()const { return DataChannel.Get(); }
 
diff --git a/Engine/Plugins/FX/Niagara/Source/Niagara/Public/NiagaraDataChannel_Global.h b/Engine/Plugins/FX/Niagara/Source/Niagara/Public/NiagaraDataChannel_Global.h
--- a/Engine/Plugins/FX/Niagara/Source/Niagara/Public/NiagaraDataChannel_Global.h
+++ b/Engine/Plugins/FX/Niagara/Source/Niagara/Public/NiagaraDataChannel_Global.h
@@ -7,33 +7,33 @@
 #include "NiagaraDataChannel_Global.generated.h"
 
 
-UCLASS(Experimental)
-class NIAGARA_API UNiagaraDataChannel_Global : public UNiagaraDataChannel
+UCLASS(Experimental, MinimalAPI)
+class UNiagaraDataChannel_Global : public UNiagaraDataChannel
 {
 	GENERATED_BODY()
 
-	virtual UNiagaraDataChannelHandler* CreateHandler(UWorld* OwningWorld)const override;
+	NIAGARA_API virtual UNiagaraDataChannelHandler* CreateHandler(UWorld* OwningWorld)const override;
 };
 
 /**
 Basic DataChannel handler that makes all data visible globally.
 */
-UCLASS(Experimental, BlueprintType)
-class NIAGARA_API UNiagaraDataChannelHandler_Global : public UNiagaraDataChannelHandler
+UCLASS(Experimental, BlueprintType, MinimalAPI)
+class UNiagaraDataChannelHandler_Global : public UNiagaraDataChannelHandler
 {
 	GENERATED_UCLASS_BODY()
 
 	FNiagaraDataChannelDataPtr Data;
 
 	//UObject Interface
-	virtual void BeginDestroy()override;
+	NIAGARA_API virtual void BeginDestroy()override;
 	//UObject Interface End
 
-	virtual void Init(const UNiagaraDataChannel* InChannel) override;
-	virtual void BeginFrame(float DeltaTime, FNiagaraWorldManager* OwningWorld)override;
-	virtual void EndFrame(float DeltaTime, FNiagaraWorldManager* OwningWorld)override;
-	virtual void Tick(float DeltaTime, ETickingGroup TickGroup, FNiagaraWorldManager* OwningWorld) override;
-	virtual FNiagaraDataChannelDataPtr FindData(FNiagaraDataChannelSearchParameters SearchParams, ENiagaraResourceAccess AccessType) override;
+	NIAGARA_API virtual void Init(const UNiagaraDataChannel* InChannel) override;
+	NIAGARA_API virtual void BeginFrame(float DeltaTime, FNiagaraWorldManager* OwningWorld)override;
+	NIAGARA_API virtual void EndFrame(float DeltaTime, FNiagaraWorldManager* OwningWorld)override;
+	NIAGARA_API virtual void Tick(float DeltaTime, ETickingGroup TickGroup, FNiagaraWorldManager* OwningWorld) override;
+	NIAGARA_API virtual FNiagaraDataChannelDataPtr FindData(FNiagaraDataChannelSearchParameters SearchParams, ENiagaraResourceAccess AccessType) override;
 };
 
 /**
diff --git a/Engine/Plugins/FX/Niagara/Source/Niagara/Public/NiagaraDataChannel_Islands.h b/Engine/Plugins/FX/Niagara/Source/Niagara/Public/NiagaraDataChannel_Islands.h
--- a/Engine/Plugins/FX/Niagara/Source/Niagara/Public/NiagaraDataChannel_Islands.h
+++ b/Engine/Plugins/FX/Niagara/Source/Niagara/Public/NiagaraDataChannel_Islands.h
@@ -106,14 +106,14 @@ struct FNDCIslandDebugDrawSettings
 	bool ShowBounds()const { return bEnabled && bShowIslandBounds; }
 };
 
-UCLASS(Experimental)
-class NIAGARA_API UNiagaraDataChannel_Islands : public UNiagaraDataChannel
+UCLASS(Experimental, MinimalAPI)
+class UNiagaraDataChannel_Islands : public UNiagaraDataChannel
 {
 	GENERATED_BODY()
 
 public: 
 
-	virtual UNiagaraDataChannelHandler* CreateHandler(UWorld* OwningWorld)const override;
+	NIAGARA_API virtual UNiagaraDataChannelHandler* CreateHandler(UWorld* OwningWorld)const override;
 
 	/** Controls how islands are placed and sized. */
 	UPROPERTY(EditAnywhere, Category = "Islands")
@@ -160,20 +160,20 @@ public:
 	FNDCIslandDebugDrawSettings DebugDrawSettings;
 };
 
-UCLASS(Experimental, BlueprintType)
-class NIAGARA_API UNiagaraDataChannelHandler_Islands : public UNiagaraDataChannelHandler
+UCLASS(Experimental, BlueprintType, MinimalAPI)
+class UNiagaraDataChannelHandler_Islands : public UNiagaraDataChannelHandler
 {
 	GENERATED_UCLASS_BODY()
 
 	//UObject Interface
-	virtual void BeginDestroy()override;
+	NIAGARA_API virtual void BeginDestroy()override;
 	//UObject Interface End
 
-	virtual void Init(const UNiagaraDataChannel* InChannel) override;
-	virtual void BeginFrame(float DeltaTime, FNiagaraWorldManager* OwningWorld)override;
-	virtual void EndFrame(float DeltaTime, FNiagaraWorldManager* OwningWorld)override;
-	virtual void Tick(float DeltaTime, ETickingGroup TickGroup, FNiagaraWorldManager* OwningWorld) override;
-	virtual FNiagaraDataChannelDataPtr FindData(FNiagaraDataChannelSearchParameters SearchParams, ENiagaraResourceAccess AccessType) override;
+	NIAGARA_API virtual void Init(const UNiagaraDataChannel* InChannel) override;
+	NIAGARA_API virtual void BeginFrame(float DeltaTime, FNiagaraWorldManager* OwningWorld)override;
+	NIAGARA_API virtual void EndFrame(float DeltaTime, FNiagaraWorldManager* OwningWorld)override;
+	NIAGARA_API virtual void Tick(float DeltaTime, ETickingGroup TickGroup, FNiagaraWorldManager* OwningWorld) override;
+	NIAGARA_API virtual FNiagaraDataChannelDataPtr FindData(FNiagaraDataChannelSearchParameters SearchParams, ENiagaraResourceAccess AccessType) override;
 
 private:
 
@@ -194,4 +194,4 @@ private:
 
 	/** Initializes and adds a new island to the ActiveIslands list. Either retrieving from the free pool of existing inactive islands or creating a new one. */
 	int32 ActivateNewIsland(FVector Location);
-};
\ No newline at end of file
+};
diff --git a/Engine/Plugins/FX/Niagara/Source/Niagara/Public/NiagaraDataInterfacePhysicsAsset.h b/Engine/Plugins/FX/Niagara/Source/Niagara/Public/NiagaraDataInterfacePhysicsAsset.h
--- a/Engine/Plugins/FX/Niagara/Source/Niagara/Public/NiagaraDataInterfacePhysicsAsset.h
+++ b/Engine/Plugins/FX/Niagara/Source/Niagara/Public/NiagaraDataInterfacePhysicsAsset.h
@@ -100,8 +100,8 @@ struct FNDIPhysicsAssetData
 };
 
 /** Data Interface for interacting with PhysicsAssets */
-UCLASS(EditInlineNew, Category = "Physics", CollapseCategories, meta = (DisplayName = "Physics Asset"))
-class NIAGARA_API UNiagaraDataInterfacePhysicsAsset : public UNiagaraDataInterface
+UCLASS(EditInlineNew, Category = "Physics", CollapseCategories, meta = (DisplayName = "Physics Asset"), MinimalAPI)
+class UNiagaraDataInterfacePhysicsAsset : public UNiagaraDataInterface
 {
 	GENERATED_UCLASS_BODY()
 
@@ -130,70 +130,70 @@ public:
 	TArray<TWeakObjectPtr<class UPhysicsAsset>> PhysicsAssets;
 
 	/** UObject Interface */
-	virtual void PostInitProperties() override;
-	virtual void PostLoad() override;
+	NIAGARA_API virtual void PostInitProperties() override;
+	NIAGARA_API virtual void PostLoad() override;
 
 	/** UNiagaraDataInterface Interface */
-	virtual void GetFunctions(TArray<FNiagaraFunctionSignature>& OutFunctions) override;
-	virtual void GetVMExternalFunction(const FVMExternalFunctionBindingInfo& BindingInfo, void* InstanceData, FVMExternalFunction& OutFunc) override;
+	NIAGARA_API virtual void GetFunctions(TArray<FNiagaraFunctionSignature>& OutFunctions) override;
+	NIAGARA_API virtual void GetVMExternalFunction(const FVMExternalFunctionBindingInfo& BindingInfo, void* InstanceData, FVMExternalFunction& OutFunc) override;
 	virtual bool CanExecuteOnTarget(ENiagaraSimTarget Target) const override { return Target == ENiagaraSimTarget::GPUComputeSim; }
-	virtual bool InitPerInstanceData(void* PerInstanceData, FNiagaraSystemInstance* SystemInstance) override;
-	virtual void DestroyPerInstanceData(void* PerInstanceData, FNiagaraSystemInstance* SystemInstance)override;
-	virtual bool PerInstanceTick(void* PerInstanceData, FNiagaraSystemInstance* SystemInstance, float DeltaSeconds) override;
+	NIAGARA_API virtual bool InitPerInstanceData(void* PerInstanceData, FNiagaraSystemInstance* SystemInstance) override;
+	NIAGARA_API virtual void DestroyPerInstanceData(void* PerInstanceData, FNiagaraSystemInstance* SystemInstance)override;
+	NIAGARA_API virtual bool PerInstanceTick(void* PerInstanceData, FNiagaraSystemInstance* SystemInstance, float DeltaSeconds) override;
 	virtual int32 PerInstanceDataSize()const override { return sizeof(FNDIPhysicsAssetData); }
-	virtual bool Equals(const UNiagaraDataInterface* Other) const override;
+	NIAGARA_API virtual bool Equals(const UNiagaraDataInterface* Other) const override;
 	virtual bool HasPreSimulateTick() const override { return true; }
 	virtual bool HasTickGroupPrereqs() const override { return true; }
-	virtual ETickingGroup CalculateTickGroup(const void* PerInstanceData) const override;
+	NIAGARA_API virtual ETickingGroup CalculateTickGroup(const void* PerInstanceData) const override;
 
 	/** GPU simulation  functionality */
 #if WITH_EDITORONLY_DATA
-	virtual bool AppendCompileHash(FNiagaraCompileHashVisitor* InVisitor) const override;
-	virtual void GetCommonHLSL(FString& OutHLSL) override;
-	virtual void GetParameterDefinitionHLSL(const FNiagaraDataInterfaceGPUParamInfo& ParamInfo, FString& OutHLSL) override;
-	virtual bool GetFunctionHLSL(const FNiagaraDataInterfaceGPUParamInfo& ParamInfo, const FNiagaraDataInterfaceGeneratedFunction& FunctionInfo, int FunctionInstanceIndex, FString& OutHLSL) override;
+	NIAGARA_API virtual bool AppendCompileHash(FNiagaraCompileHashVisitor* InVisitor) const override;
+	NIAGARA_API virtual void GetCommonHLSL(FString& OutHLSL) override;
+	NIAGARA_API virtual void GetParameterDefinitionHLSL(const FNiagaraDataInterfaceGPUParamInfo& ParamInfo, FString& OutHLSL) override;
+	NIAGARA_API virtual bool GetFunctionHLSL(const FNiagaraDataInterfaceGPUParamInfo& ParamInfo, const FNiagaraDataInterfaceGeneratedFunction& FunctionInfo, int FunctionInstanceIndex, FString& OutHLSL) override;
 #endif
-	virtual void BuildShaderParameters(FNiagaraShaderParametersBuilder& ShaderParametersBuilder) const override;
-	virtual void SetShaderParameters(const FNiagaraDataInterfaceSetShaderParametersContext& Context) const override;
+	NIAGARA_API virtual void BuildShaderParameters(FNiagaraShaderParametersBuilder& ShaderParametersBuilder) const override;
+	NIAGARA_API virtual void SetShaderParameters(const FNiagaraDataInterfaceSetShaderParametersContext& Context) const override;
 
-	virtual void ProvidePerInstanceDataForRenderThread(void* DataForRenderThread, void* PerInstanceData, const FNiagaraSystemInstanceID& SystemInstance) override;
+	NIAGARA_API virtual void ProvidePerInstanceDataForRenderThread(void* DataForRenderThread, void* PerInstanceData, const FNiagaraSystemInstanceID& SystemInstance) override;
 
 	/** Extract the source component */
-	void ExtractSourceComponent(FNiagaraSystemInstance* SystemInstance, FTransform& LocalTransform);
+	NIAGARA_API void ExtractSourceComponent(FNiagaraSystemInstance* SystemInstance, FTransform& LocalTransform);
 
 	/** Get the number of boxes*/
-	void GetNumBoxes(FVectorVMExternalFunctionContext& Context);
+	NIAGARA_API void GetNumBoxes(FVectorVMExternalFunctionContext& Context);
 
 	/** Get the number of spheres  */
-	void GetNumSpheres(FVectorVMExternalFunctionContext& Context);
+	NIAGARA_API void GetNumSpheres(FVectorVMExternalFunctionContext& Context);
 
 	/** Get the number of capsules */
-	void GetNumCapsules(FVectorVMExternalFunctionContext& Context);
+	NIAGARA_API void GetNumCapsules(FVectorVMExternalFunctionContext& Context);
 
 	/** Get the element point */
-	void GetElementPoint(FVectorVMExternalFunctionContext& Context);
+	NIAGARA_API void GetElementPoint(FVectorVMExternalFunctionContext& Context);
 
 	/** Get the element distance */
-	void GetElementDistance(FVectorVMExternalFunctionContext& Context);
+	NIAGARA_API void GetElementDistance(FVectorVMExternalFunctionContext& Context);
 
 	/** Get the closest element */
-	void GetClosestElement(FVectorVMExternalFunctionContext& Context);
+	NIAGARA_API void GetClosestElement(FVectorVMExternalFunctionContext& Context);
 
 	/** Get the closest point */
-	void GetClosestPoint(FVectorVMExternalFunctionContext& Context);
+	NIAGARA_API void GetClosestPoint(FVectorVMExternalFunctionContext& Context);
 
 	/** Get the closest distance */
-	void GetClosestDistance(FVectorVMExternalFunctionContext& Context);
+	NIAGARA_API void GetClosestDistance(FVectorVMExternalFunctionContext& Context);
 
 	/** Get the closest texture point */
-	void GetTexturePoint(FVectorVMExternalFunctionContext& Context);
+	NIAGARA_API void GetTexturePoint(FVectorVMExternalFunctionContext& Context);
 
 	/** Get the projection point */
-	void GetProjectionPoint(FVectorVMExternalFunctionContext& Context);
+	NIAGARA_API void GetProjectionPoint(FVectorVMExternalFunctionContext& Context);
 
 protected:
 	/** Copy one niagara DI to this */
-	virtual bool CopyToInternal(UNiagaraDataInterface* Destination) const override;
+	NIAGARA_API virtual bool CopyToInternal(UNiagaraDataInterface* Destination) const override;
 };
 
 /** Proxy to send data to gpu */
@@ -218,13 +218,13 @@ struct FNDIPhysicsAssetProxy : public FNiagaraDataInterfaceProxy
 	TMap<FNiagaraSystemInstanceID, FNDIPhysicsAssetData> SystemInstancesToProxyData;
 };
 
-UINTERFACE()
-class NIAGARA_API UNiagaraPhysicsAssetDICollectorInterface : public UInterface
+UINTERFACE(MinimalAPI)
+class UNiagaraPhysicsAssetDICollectorInterface : public UInterface
 {
 	GENERATED_UINTERFACE_BODY()
 };
 
-class NIAGARA_API INiagaraPhysicsAssetDICollectorInterface
+class INiagaraPhysicsAssetDICollectorInterface
 {
 	GENERATED_IINTERFACE_BODY()
 
@@ -232,4 +232,4 @@ class NIAGARA_API INiagaraPhysicsAssetDICollectorInterface
 		FTransform& BoneTransform,
 		TArray<TWeakObjectPtr<USkeletalMeshComponent>>& SourceComponents,
 		TArray<TWeakObjectPtr<UPhysicsAsset>>& PhysicsAssets) const = 0;
-};
\ No newline at end of file
+};
diff --git a/Engine/Plugins/FX/Niagara/Source/Niagara/Public/NiagaraDataInterfaceRigidMeshCollisionQuery.h b/Engine/Plugins/FX/Niagara/Source/Niagara/Public/NiagaraDataInterfaceRigidMeshCollisionQuery.h
--- a/Engine/Plugins/FX/Niagara/Source/Niagara/Public/NiagaraDataInterfaceRigidMeshCollisionQuery.h
+++ b/Engine/Plugins/FX/Niagara/Source/Niagara/Public/NiagaraDataInterfaceRigidMeshCollisionQuery.h
@@ -285,13 +285,13 @@ protected:
 /**
 * C++ and Blueprint library for accessing array types
 */
-UCLASS()
-class NIAGARA_API UNiagaraDIRigidMeshCollisionFunctionLibrary : public UBlueprintFunctionLibrary
+UCLASS(MinimalAPI)
+class UNiagaraDIRigidMeshCollisionFunctionLibrary : public UBlueprintFunctionLibrary
 {
 	GENERATED_BODY()
 
 public:
 
 	UFUNCTION(BlueprintCallable, Category = Niagara, meta = (DisplayName = "Niagara Set Source Actors"))
-	static void SetSourceActors(UNiagaraComponent* NiagaraSystem, FName OverrideName, const TArray<AActor*>& SourceActors);
-};
\ No newline at end of file
+	static NIAGARA_API void SetSourceActors(UNiagaraComponent* NiagaraSystem, FName OverrideName, const TArray<AActor*>& SourceActors);
+};
diff --git a/Engine/Plugins/FX/Niagara/Source/Niagara/Public/NiagaraDebuggerCommon.h b/Engine/Plugins/FX/Niagara/Source/Niagara/Public/NiagaraDebuggerCommon.h
--- a/Engine/Plugins/FX/Niagara/Source/Niagara/Public/NiagaraDebuggerCommon.h
+++ b/Engine/Plugins/FX/Niagara/Source/Niagara/Public/NiagaraDebuggerCommon.h
@@ -214,7 +214,7 @@ Messaged broadcast from debugger to request a connection to a particular session
 If any matching client is found and it accepts, it will return a FNiagaraDebuggerAcceptConnection message to the sender. 
 */
 USTRUCT()
-struct NIAGARA_API FNiagaraDebuggerRequestConnection
+struct FNiagaraDebuggerRequestConnection
 {
 	GENERATED_BODY()
 	
@@ -234,7 +234,7 @@ struct NIAGARA_API FNiagaraDebuggerRequestConnection
 
 /** Response message from the a debugger client accepting a connection requested by a FNiagaraDebuggerRequestConnection message. */
 USTRUCT()
-struct NIAGARA_API FNiagaraDebuggerAcceptConnection
+struct FNiagaraDebuggerAcceptConnection
 {
 	GENERATED_BODY()
 		
@@ -254,7 +254,7 @@ struct NIAGARA_API FNiagaraDebuggerAcceptConnection
 
 /** Empty message informing a debugger client that the debugger is closing the connection. */
 USTRUCT()
-struct NIAGARA_API FNiagaraDebuggerConnectionClosed
+struct FNiagaraDebuggerConnectionClosed
 {
 	GENERATED_BODY()
 	
@@ -274,7 +274,7 @@ struct NIAGARA_API FNiagaraDebuggerConnectionClosed
 
 /** Command that will execute a console command on the debugger client. */
 USTRUCT() 
-struct NIAGARA_API FNiagaraDebuggerExecuteConsoleCommand
+struct FNiagaraDebuggerExecuteConsoleCommand
 {
 	GENERATED_BODY()
 	
@@ -294,7 +294,7 @@ struct NIAGARA_API FNiagaraDebuggerExecuteConsoleCommand
 
 /** Message containing updated outliner information sent from the client to the debugger. */
 USTRUCT()
-struct NIAGARA_API FNiagaraDebuggerOutlinerUpdate
+struct FNiagaraDebuggerOutlinerUpdate
 {
 	GENERATED_BODY()
 
@@ -366,7 +366,7 @@ struct FNiagaraDebugHudTextOptions
 };
 
 USTRUCT()
-struct NIAGARA_API FNiagaraDebugHUDVariable
+struct FNiagaraDebugHUDVariable
 {
 	GENERATED_BODY()
 		
@@ -377,8 +377,8 @@ struct NIAGARA_API FNiagaraDebugHUDVariable
 	UPROPERTY(EditAnywhere, Category = "Attribute")
 	FString Name;
 
-	static FString BuildVariableString(const TArray<FNiagaraDebugHUDVariable>& Variables);
-	static void InitFromString(const FString& VariablesString, TArray<FNiagaraDebugHUDVariable>& OutVariables);
+	static NIAGARA_API FString BuildVariableString(const TArray<FNiagaraDebugHUDVariable>& Variables);
+	static NIAGARA_API void InitFromString(const FString& VariablesString, TArray<FNiagaraDebugHUDVariable>& OutVariables);
 };
 
 UENUM()
@@ -408,11 +408,11 @@ enum class ENiagaraDebugHUDPerfSampleMode
 
 /** Settings for Niagara debug HUD. Contained in it's own struct so that we can pass it whole in a message to the debugger client. */
 USTRUCT()
-struct NIAGARA_API FNiagaraDebugHUDSettingsData
+struct FNiagaraDebugHUDSettingsData
 {
 	GENERATED_BODY()
 
-	FNiagaraDebugHUDSettingsData();
+	NIAGARA_API FNiagaraDebugHUDSettingsData();
 
 	bool IsEnabled() const
 	{
@@ -730,13 +730,13 @@ struct NIAGARA_API FNiagaraDebugHUDSettingsData
 
 /** Message passed from debugger to client when it needs updated simple client info. */
 USTRUCT()
-struct NIAGARA_API FNiagaraRequestSimpleClientInfoMessage
+struct FNiagaraRequestSimpleClientInfoMessage
 {
 	GENERATED_BODY()
 };
 
-UCLASS(config = EditorPerProjectUserSettings, defaultconfig)
-class NIAGARA_API UNiagaraDebugHUDSettings : public UObject, public FNotifyHook
+UCLASS(config = EditorPerProjectUserSettings, defaultconfig, MinimalAPI)
+class UNiagaraDebugHUDSettings : public UObject, public FNotifyHook
 {
 	GENERATED_BODY()
 
@@ -747,7 +747,7 @@ public:
 	UPROPERTY(Config, EditAnywhere, Category = "Settings", meta=(ShowOnlyInnerProperties))
 	FNiagaraDebugHUDSettingsData Data;
 
-	void NotifyPropertyChanged();
+	NIAGARA_API void NotifyPropertyChanged();
 	virtual void NotifyPreChange(FProperty*) {}
 	virtual void NotifyPostChange(const FPropertyChangedEvent&, FProperty*) { NotifyPropertyChanged(); }
 	virtual void NotifyPreChange(class FEditPropertyChain*) {}
@@ -756,7 +756,7 @@ public:
 
 
 USTRUCT()
-struct NIAGARA_API FNiagaraOutlinerCaptureSettings
+struct FNiagaraOutlinerCaptureSettings
 {
 	GENERATED_BODY()
 
@@ -778,7 +778,7 @@ struct NIAGARA_API FNiagaraOutlinerCaptureSettings
 
 /** Simple information on the connected client for use in continuous or immediate response UI elements. */
 USTRUCT()
-struct NIAGARA_API FNiagaraSimpleClientInfo
+struct FNiagaraSimpleClientInfo
 {
 	GENERATED_BODY()
 
@@ -801,7 +801,7 @@ struct NIAGARA_API FNiagaraSimpleClientInfo
 
 /** Message sent from the debugger to a client to request a sim cache capture for a particular component. */
 USTRUCT()
-struct NIAGARA_API FNiagaraSystemSimCacheCaptureRequest
+struct FNiagaraSystemSimCacheCaptureRequest
 {
 	GENERATED_BODY()
 
@@ -821,7 +821,7 @@ struct NIAGARA_API FNiagaraSystemSimCacheCaptureRequest
 
 /** Message sent from a debugger client to a connected debugger containing the results of a sim cache capture. */
 USTRUCT()
-struct NIAGARA_API FNiagaraSystemSimCacheCaptureReply
+struct FNiagaraSystemSimCacheCaptureReply
 {
 	GENERATED_BODY()
 
@@ -853,4 +853,4 @@ struct FNiagaraDebugMessage
 		, Message(InMessage)
 		, Lifetime(InLifetime)
 	{}
-};
\ No newline at end of file
+};
diff --git a/Engine/Plugins/FX/Niagara/Source/Niagara/Public/NiagaraEditorDataBase.h b/Engine/Plugins/FX/Niagara/Source/Niagara/Public/NiagaraEditorDataBase.h
--- a/Engine/Plugins/FX/Niagara/Source/Niagara/Public/NiagaraEditorDataBase.h
+++ b/Engine/Plugins/FX/Niagara/Source/Niagara/Public/NiagaraEditorDataBase.h
@@ -53,7 +53,7 @@ public:
 #if WITH_EDITORONLY_DATA
 	virtual void PostLoadFromOwner(UObject* InOwner) { }
 
-	NIAGARA_API FSimpleMulticastDelegate& OnPersistentDataChanged() { return PersistentDataChangedDelegate; }
+	FSimpleMulticastDelegate& OnPersistentDataChanged() { return PersistentDataChangedDelegate; }
 
 private:
 	FSimpleMulticastDelegate PersistentDataChangedDelegate;
diff --git a/Engine/Plugins/FX/Niagara/Source/Niagara/Public/NiagaraEmptyUAVPool.h b/Engine/Plugins/FX/Niagara/Source/Niagara/Public/NiagaraEmptyUAVPool.h
--- a/Engine/Plugins/FX/Niagara/Source/Niagara/Public/NiagaraEmptyUAVPool.h
+++ b/Engine/Plugins/FX/Niagara/Source/Niagara/Public/NiagaraEmptyUAVPool.h
@@ -34,7 +34,7 @@ enum class ENiagaraEmptyUAVType
 };
 
 // Empty UAV pool used for ensuring we bind a buffer when one does not exist
-class NIAGARA_API FNiagaraEmptyUAVPool
+class FNiagaraEmptyUAVPool
 {
 	friend struct FNiagaraEmptyUAVPoolScopedAccess;
 
@@ -43,12 +43,12 @@ public:
 	Grab a temporary empty RW buffer from the pool.
 	Note: When doing this outside of Niagara you must be within a FNiagaraUAVPoolAccessScope.
 	*/
-	FRHIUnorderedAccessView* GetEmptyUAVFromPool(FRHICommandList& RHICmdList, EPixelFormat Format, ENiagaraEmptyUAVType Type);
+	NIAGARA_API FRHIUnorderedAccessView* GetEmptyUAVFromPool(FRHICommandList& RHICmdList, EPixelFormat Format, ENiagaraEmptyUAVType Type);
 
 	/**
 	Returns all used UAVs back to the pool
 	*/
-	void ResetEmptyUAVPools();
+	NIAGARA_API void ResetEmptyUAVPools();
 
 protected:
 	struct FEmptyUAV
diff --git a/Engine/Plugins/FX/Niagara/Source/Niagara/Public/NiagaraFunctionLibrary.h b/Engine/Plugins/FX/Niagara/Source/Niagara/Public/NiagaraFunctionLibrary.h
--- a/Engine/Plugins/FX/Niagara/Source/Niagara/Public/NiagaraFunctionLibrary.h
+++ b/Engine/Plugins/FX/Niagara/Source/Niagara/Public/NiagaraFunctionLibrary.h
@@ -21,71 +21,71 @@ class UVolumeTexture;
 * A C++ and Blueprint accessible library of utility functions for accessing Niagara simulations
 * All positions & orientations are returned in Unreal reference frame & units, assuming the Leap device is located at the origin.
 */
-UCLASS()
-class NIAGARA_API UNiagaraFunctionLibrary : public UBlueprintFunctionLibrary
+UCLASS(MinimalAPI)
+class UNiagaraFunctionLibrary : public UBlueprintFunctionLibrary
 {
 	GENERATED_UCLASS_BODY()
 
 public:
 
 	UFUNCTION(BlueprintCallable, Category = Niagara, meta = (Keywords = "niagara System", WorldContext = "WorldContextObject", UnsafeDuringActorConstruction = "true"))
-	static UNiagaraComponent* SpawnSystemAtLocationWithParams(const FFXSystemSpawnParameters& SpawnParams);
+	static NIAGARA_API UNiagaraComponent* SpawnSystemAtLocationWithParams(const FFXSystemSpawnParameters& SpawnParams);
 
 	UFUNCTION(BlueprintCallable, Category = Niagara, meta = (Keywords = "niagara System", UnsafeDuringActorConstruction = "true"))
-	static UNiagaraComponent* SpawnSystemAttachedWithParams(const FFXSystemSpawnParameters& SpawnParams);
+	static NIAGARA_API UNiagaraComponent* SpawnSystemAttachedWithParams(const FFXSystemSpawnParameters& SpawnParams);
 
 	/**
 	* Spawns a Niagara System at the specified world location/rotation
 	* @return			The spawned UNiagaraComponent
 	*/
 	UFUNCTION(BlueprintCallable, Category = Niagara, meta = (Keywords = "niagara System", WorldContext = "WorldContextObject", UnsafeDuringActorConstruction = "true"))
-	static UNiagaraComponent* SpawnSystemAtLocation(const UObject* WorldContextObject, class UNiagaraSystem* SystemTemplate, FVector Location, FRotator Rotation = FRotator::ZeroRotator, FVector Scale = FVector(1.f), bool bAutoDestroy = true, bool bAutoActivate = true, ENCPoolMethod PoolingMethod = ENCPoolMethod::None, bool bPreCullCheck = true);
+	static NIAGARA_API UNiagaraComponent* SpawnSystemAtLocation(const UObject* WorldContextObject, class UNiagaraSystem* SystemTemplate, FVector Location, FRotator Rotation = FRotator::ZeroRotator, FVector Scale = FVector(1.f), bool bAutoDestroy = true, bool bAutoActivate = true, ENCPoolMethod PoolingMethod = ENCPoolMethod::None, bool bPreCullCheck = true);
 
 	UFUNCTION(BlueprintCallable, Category = Niagara, meta = (Keywords = "niagara System", UnsafeDuringActorConstruction = "true"))
-	static UNiagaraComponent* SpawnSystemAttached(UNiagaraSystem* SystemTemplate, USceneComponent* AttachToComponent, FName AttachPointName, FVector Location, FRotator Rotation, EAttachLocation::Type LocationType, bool bAutoDestroy, bool bAutoActivate = true, ENCPoolMethod PoolingMethod = ENCPoolMethod::None, bool bPreCullCheck = true);
+	static NIAGARA_API UNiagaraComponent* SpawnSystemAttached(UNiagaraSystem* SystemTemplate, USceneComponent* AttachToComponent, FName AttachPointName, FVector Location, FRotator Rotation, EAttachLocation::Type LocationType, bool bAutoDestroy, bool bAutoActivate = true, ENCPoolMethod PoolingMethod = ENCPoolMethod::None, bool bPreCullCheck = true);
 
-	static UNiagaraComponent* SpawnSystemAttached(UNiagaraSystem* SystemTemplate, USceneComponent* AttachToComponent, FName AttachPointName, FVector Location, FRotator Rotation, FVector Scale, EAttachLocation::Type LocationType, bool bAutoDestroy, ENCPoolMethod PoolingMethod, bool bAutoActivate = true, bool bPreCullCheck = true);
+	static NIAGARA_API UNiagaraComponent* SpawnSystemAttached(UNiagaraSystem* SystemTemplate, USceneComponent* AttachToComponent, FName AttachPointName, FVector Location, FRotator Rotation, FVector Scale, EAttachLocation::Type LocationType, bool bAutoDestroy, ENCPoolMethod PoolingMethod, bool bAutoActivate = true, bool bPreCullCheck = true);
 
 	/** Sets a Niagara StaticMesh parameter by name, overriding locally if necessary.*/
 	UFUNCTION(BlueprintCallable, Category = Niagara, meta = (DisplayName = "Set Niagara Static Mesh Component"))
-	static void OverrideSystemUserVariableStaticMeshComponent(UNiagaraComponent* NiagaraSystem, const FString& OverrideName, UStaticMeshComponent* StaticMeshComponent);
+	static NIAGARA_API void OverrideSystemUserVariableStaticMeshComponent(UNiagaraComponent* NiagaraSystem, const FString& OverrideName, UStaticMeshComponent* StaticMeshComponent);
 
 	UFUNCTION(BlueprintCallable, Category = Niagara, meta = (DisplayName = "Set Niagara Static Mesh Directly"))
-	static void OverrideSystemUserVariableStaticMesh(UNiagaraComponent* NiagaraSystem, const FString& OverrideName, UStaticMesh* StaticMesh);
+	static NIAGARA_API void OverrideSystemUserVariableStaticMesh(UNiagaraComponent* NiagaraSystem, const FString& OverrideName, UStaticMesh* StaticMesh);
 
 	/** Get the skeletal mesh data interface by name .*/
-	static class UNiagaraDataInterfaceSkeletalMesh* GetSkeletalMeshDataInterface(UNiagaraComponent* NiagaraSystem, const FString& OverrideName);
+	static NIAGARA_API class UNiagaraDataInterfaceSkeletalMesh* GetSkeletalMeshDataInterface(UNiagaraComponent* NiagaraSystem, const FString& OverrideName);
 
 	/** Sets a Niagara StaticMesh parameter by name, overriding locally if necessary.*/
 	UFUNCTION(BlueprintCallable, Category = Niagara, meta = (DisplayName = "Set Niagara Skeletal Mesh Component"))
-	static void OverrideSystemUserVariableSkeletalMeshComponent(UNiagaraComponent* NiagaraSystem, const FString& OverrideName, USkeletalMeshComponent* SkeletalMeshComponent);
+	static NIAGARA_API void OverrideSystemUserVariableSkeletalMeshComponent(UNiagaraComponent* NiagaraSystem, const FString& OverrideName, USkeletalMeshComponent* SkeletalMeshComponent);
 
 	/** Sets the SamplingRegion to use on the skeletal mesh data interface, this is destructive as it modifies the data interface. */
 	UFUNCTION(BlueprintCallable, Category = Niagara)
-	static void SetSkeletalMeshDataInterfaceSamplingRegions(UNiagaraComponent* NiagaraSystem, const FString& OverrideName, const TArray<FName>& SamplingRegions);
+	static NIAGARA_API void SetSkeletalMeshDataInterfaceSamplingRegions(UNiagaraComponent* NiagaraSystem, const FString& OverrideName, const TArray<FName>& SamplingRegions);
 
 	/** Sets the Filtered Bones to use on the skeletal mesh data interface, this is destructive as it modifies the data interface. */
 	UFUNCTION(BlueprintCallable, Category = Niagara)
-	static void SetSkeletalMeshDataInterfaceFilteredBones(UNiagaraComponent* NiagaraSystem, const FString& OverrideName, const TArray<FName>& FilteredBones);
+	static NIAGARA_API void SetSkeletalMeshDataInterfaceFilteredBones(UNiagaraComponent* NiagaraSystem, const FString& OverrideName, const TArray<FName>& FilteredBones);
 
 	/** Sets the Filtered Sockets to use on the skeletal mesh data interface, this is destructive as it modifies the data interface. */
 	UFUNCTION(BlueprintCallable, Category = Niagara)
-	static void SetSkeletalMeshDataInterfaceFilteredSockets(UNiagaraComponent* NiagaraSystem, const FString& OverrideName, const TArray<FName>& FilteredSockets);
+	static NIAGARA_API void SetSkeletalMeshDataInterfaceFilteredSockets(UNiagaraComponent* NiagaraSystem, const FString& OverrideName, const TArray<FName>& FilteredSockets);
 
 	/** Overrides the Texture Object for a Niagara Texture Data Interface User Parameter.*/
 	UFUNCTION(BlueprintCallable, Category = Niagara)
-	static void SetTextureObject(UNiagaraComponent* NiagaraSystem, const FString& OverrideName, UTexture* Texture);
+	static NIAGARA_API void SetTextureObject(UNiagaraComponent* NiagaraSystem, const FString& OverrideName, UTexture* Texture);
 
 	/** Overrides the 2D Array Texture for a Niagara 2D Array Texture Data Interface User Parameter.*/
 	UFUNCTION(BlueprintCallable, Category = Niagara)
-	static void SetTexture2DArrayObject(UNiagaraComponent* NiagaraSystem, const FString& OverrideName, class UTexture2DArray* Texture);
+	static NIAGARA_API void SetTexture2DArrayObject(UNiagaraComponent* NiagaraSystem, const FString& OverrideName, class UTexture2DArray* Texture);
 
 	/** Overrides the Volume Texture for a Niagara Volume Texture Data Interface User Parameter.*/
 	UFUNCTION(BlueprintCallable, Category = Niagara)
-	static void SetVolumeTextureObject(UNiagaraComponent* NiagaraSystem, const FString& OverrideName, UVolumeTexture* Texture);
+	static NIAGARA_API void SetVolumeTextureObject(UNiagaraComponent* NiagaraSystem, const FString& OverrideName, UVolumeTexture* Texture);
 	
 	/** Finds an array interface of the given class. */
-	static UNiagaraDataInterface* GetDataInterface(UClass* DIClass, UNiagaraComponent* NiagaraSystem, FName OverrideName);
+	static NIAGARA_API UNiagaraDataInterface* GetDataInterface(UClass* DIClass, UNiagaraComponent* NiagaraSystem, FName OverrideName);
 
 	/** Finds an array interface of the given class. */
 	template<class TDIType>
@@ -99,33 +99,33 @@ public:
 // 	static void SetUpdateScriptConstant(UNiagaraComponent* Component, FName EmitterName, FName ConstantName, FVector Value);
 
 	UFUNCTION(BlueprintCallable, Category = Niagara, meta = (Keywords = "niagara parameter collection", WorldContext = "WorldContextObject"))
-	static UNiagaraParameterCollectionInstance* GetNiagaraParameterCollection(UObject* WorldContextObject, UNiagaraParameterCollection* Collection);
+	static NIAGARA_API UNiagaraParameterCollectionInstance* GetNiagaraParameterCollection(UObject* WorldContextObject, UNiagaraParameterCollection* Collection);
 
-	static const TArray<FNiagaraFunctionSignature>& GetVectorVMFastPathOps(bool bIgnoreConsoleVariable = false);
-	static bool DefineFunctionHLSL(const FNiagaraFunctionSignature& FunctionSignature, FString& HlslOutput);
+	static NIAGARA_API const TArray<FNiagaraFunctionSignature>& GetVectorVMFastPathOps(bool bIgnoreConsoleVariable = false);
+	static NIAGARA_API bool DefineFunctionHLSL(const FNiagaraFunctionSignature& FunctionSignature, FString& HlslOutput);
 
-	static bool GetVectorVMFastPathExternalFunction(const FVMExternalFunctionBindingInfo& BindingInfo, FVMExternalFunction &OutFunc);
+	static NIAGARA_API bool GetVectorVMFastPathExternalFunction(const FVMExternalFunctionBindingInfo& BindingInfo, FVMExternalFunction &OutFunc);
 
 
 	//Functions providing access to HWRT collision specific features
 
 	/** Sets the Niagara GPU ray traced collision group for the give primitive component. */
 	UFUNCTION(BlueprintCallable, Category = Niagara, meta = (Keywords = "niagara collision ray tracing", WorldContext = "WorldContextObject"))
-	static void SetComponentNiagaraGPURayTracedCollisionGroup(UObject* WorldContextObject, UPrimitiveComponent* Primitive, int32 CollisionGroup);
+	static NIAGARA_API void SetComponentNiagaraGPURayTracedCollisionGroup(UObject* WorldContextObject, UPrimitiveComponent* Primitive, int32 CollisionGroup);
 	
 	/** Sets the Niagara GPU ray traced collision group for all primitive components on the given actor. */
 	UFUNCTION(BlueprintCallable, Category = Niagara, meta = (Keywords = "niagara collision ray tracing", WorldContext = "WorldContextObject"))
-	static void SetActorNiagaraGPURayTracedCollisionGroup(UObject* WorldContextObject, AActor* Actor, int32 CollisionGroup);
+	static NIAGARA_API void SetActorNiagaraGPURayTracedCollisionGroup(UObject* WorldContextObject, AActor* Actor, int32 CollisionGroup);
 
 	/** Returns a free collision group for use in HWRT collision group filtering. Returns -1 on failure. */
 	UFUNCTION(BlueprintCallable, Category = Niagara, meta = (Keywords = "niagara collision ray tracing", WorldContext = "WorldContextObject"))
-	static int32 AcquireNiagaraGPURayTracedCollisionGroup(UObject* WorldContextObject);
+	static NIAGARA_API int32 AcquireNiagaraGPURayTracedCollisionGroup(UObject* WorldContextObject);
 
 	/** Releases a collision group back to the system for use by ohers. */
 	UFUNCTION(BlueprintCallable, Category = Niagara, meta = (Keywords = "niagara collision ray tracing", WorldContext = "WorldContextObject"))
-	static void ReleaseNiagaraGPURayTracedCollisionGroup(UObject* WorldContextObject, int32 CollisionGroup);
+	static NIAGARA_API void ReleaseNiagaraGPURayTracedCollisionGroup(UObject* WorldContextObject, int32 CollisionGroup);
 private:
-	static void InitVectorVMFastPathOps();
-	static TArray<FNiagaraFunctionSignature> VectorVMOps;
-	static TArray<FString> VectorVMOpsHLSL;
+	static NIAGARA_API void InitVectorVMFastPathOps();
+	static NIAGARA_API TArray<FNiagaraFunctionSignature> VectorVMOps;
+	static NIAGARA_API TArray<FString> VectorVMOpsHLSL;
 };
diff --git a/Engine/Plugins/FX/Niagara/Source/Niagara/Public/NiagaraGpuComputeDataManager.h b/Engine/Plugins/FX/Niagara/Source/Niagara/Public/NiagaraGpuComputeDataManager.h
--- a/Engine/Plugins/FX/Niagara/Source/Niagara/Public/NiagaraGpuComputeDataManager.h
+++ b/Engine/Plugins/FX/Niagara/Source/Niagara/Public/NiagaraGpuComputeDataManager.h
@@ -9,7 +9,7 @@ class FNiagaraGpuComputeDispatchInterface;
 // Abstract class for managing GPU data for Niagara
 // Once the manager is created it will last for the lifetime of the owner dispatch interface
 // The manager will be destroyed when the dispatch interface is also destroyed
-class NIAGARA_API FNiagaraGpuComputeDataManager
+class FNiagaraGpuComputeDataManager
 {
 	UE_NONCOPYABLE(FNiagaraGpuComputeDataManager);
 
diff --git a/Engine/Plugins/FX/Niagara/Source/Niagara/Public/NiagaraGpuComputeDebugInterface.h b/Engine/Plugins/FX/Niagara/Source/Niagara/Public/NiagaraGpuComputeDebugInterface.h
--- a/Engine/Plugins/FX/Niagara/Source/Niagara/Public/NiagaraGpuComputeDebugInterface.h
+++ b/Engine/Plugins/FX/Niagara/Source/Niagara/Public/NiagaraGpuComputeDebugInterface.h
@@ -14,7 +14,7 @@ Class used help realtime debug Gpu Compute simulations
 
 class FNiagaraGpuComputeDebug;
 
-class NIAGARA_API FNiagaraGpuComputeDebugInterface
+class FNiagaraGpuComputeDebugInterface
 {
 	friend class FNiagaraGpuComputeDispatchInterface;
 
@@ -26,17 +26,17 @@ private:
 
 public:
 	// Add a texture to visualize
-	void AddTexture(FRDGBuilder& GraphBuilder, FNiagaraSystemInstanceID SystemInstanceID, FName SourceName, FRDGTextureRef Texture, FVector2D PreviewDisplayRange = FVector2D::ZeroVector);
+	NIAGARA_API void AddTexture(FRDGBuilder& GraphBuilder, FNiagaraSystemInstanceID SystemInstanceID, FName SourceName, FRDGTextureRef Texture, FVector2D PreviewDisplayRange = FVector2D::ZeroVector);
 
 	// Add a texture to visualize that contains a number of attributes and select which attributes to push into RGBA where -1 means ignore that channel
 	// The first -1 in the attribute indices list will also limit the number of attributes we attempt to read.
 	// NumTextureAttributes in this version is meant for a 2D atlas
-	void AddAttributeTexture(FRDGBuilder& GraphBuilder, FNiagaraSystemInstanceID SystemInstanceID, FName SourceName, FRDGTextureRef Texture, FIntPoint NumTextureAttributes, FIntVector4 AttributeIndices, FVector2D PreviewDisplayRange = FVector2D::ZeroVector);
+	NIAGARA_API void AddAttributeTexture(FRDGBuilder& GraphBuilder, FNiagaraSystemInstanceID SystemInstanceID, FName SourceName, FRDGTextureRef Texture, FIntPoint NumTextureAttributes, FIntVector4 AttributeIndices, FVector2D PreviewDisplayRange = FVector2D::ZeroVector);
 
 	// Add a texture to visualize that contains a number of attributes and select which attributes to push into RGBA where -1 means ignore that channel
 	// The first -1 in the attribute indices list will also limit the number of attributes we attempt to read
 	// NumTextureAttributes in this version is meant for a 3D atlas
-	void AddAttributeTexture(FRDGBuilder& GraphBuilder, FNiagaraSystemInstanceID SystemInstanceID, FName SourceName, FRDGTextureRef Texture, FIntVector4 NumTextureAttributes, FIntVector4 AttributeIndices, FVector2D PreviewDisplayRange = FVector2D::ZeroVector);
+	NIAGARA_API void AddAttributeTexture(FRDGBuilder& GraphBuilder, FNiagaraSystemInstanceID SystemInstanceID, FName SourceName, FRDGTextureRef Texture, FIntVector4 NumTextureAttributes, FIntVector4 AttributeIndices, FVector2D PreviewDisplayRange = FVector2D::ZeroVector);
 
 private:
 	FNiagaraGpuComputeDebug* ComputeDebug = nullptr;
diff --git a/Engine/Plugins/FX/Niagara/Source/Niagara/Public/NiagaraGpuComputeDispatchInterface.h b/Engine/Plugins/FX/Niagara/Source/Niagara/Public/NiagaraGpuComputeDispatchInterface.h
--- a/Engine/Plugins/FX/Niagara/Source/Niagara/Public/NiagaraGpuComputeDispatchInterface.h
+++ b/Engine/Plugins/FX/Niagara/Source/Niagara/Public/NiagaraGpuComputeDispatchInterface.h
@@ -25,15 +25,15 @@ class FNiagaraSystemGpuComputeProxy;
 
 // Public API for Niagara's Compute Dispatcher
 // This is generally used with DataInterfaces or Custom Renderers
-class NIAGARA_API FNiagaraGpuComputeDispatchInterface : public FFXSystemInterface
+class FNiagaraGpuComputeDispatchInterface : public FFXSystemInterface
 {
 public:
-	static FNiagaraGpuComputeDispatchInterface* Get(class UWorld* World);
-	static FNiagaraGpuComputeDispatchInterface* Get(class FSceneInterface* Scene);
-	static FNiagaraGpuComputeDispatchInterface* Get(class FFXSystemInterface* FXSceneInterface);
+	static NIAGARA_API FNiagaraGpuComputeDispatchInterface* Get(class UWorld* World);
+	static NIAGARA_API FNiagaraGpuComputeDispatchInterface* Get(class FSceneInterface* Scene);
+	static NIAGARA_API FNiagaraGpuComputeDispatchInterface* Get(class FFXSystemInterface* FXSceneInterface);
 
-	explicit FNiagaraGpuComputeDispatchInterface(EShaderPlatform InShaderPlatform, ERHIFeatureLevel::Type InFeatureLevel);
-	virtual ~FNiagaraGpuComputeDispatchInterface();
+	NIAGARA_API explicit FNiagaraGpuComputeDispatchInterface(EShaderPlatform InShaderPlatform, ERHIFeatureLevel::Type InFeatureLevel);
+	NIAGARA_API virtual ~FNiagaraGpuComputeDispatchInterface();
 
 	/** Get ShaderPlatform the batcher is bound to */
 	EShaderPlatform GetShaderPlatform() const { return ShaderPlatform; }
@@ -86,7 +86,7 @@ public:
 
 #if NIAGARA_COMPUTEDEBUG_ENABLED
 	/** Public interface to Niagara compute debugging. */
-	FNiagaraGpuComputeDebugInterface GetGpuComputeDebugInterface() const;
+	NIAGARA_API FNiagaraGpuComputeDebugInterface GetGpuComputeDebugInterface() const;
 
 	/** Get access to Niagara's GpuComputeDebug this is for internal use */
 	FNiagaraGpuComputeDebug* GetGpuComputeDebugPrivate() const { return GpuComputeDebugPtr.Get(); }
@@ -109,19 +109,19 @@ public:
 	FRHIUnorderedAccessView* GetEmptyUAVFromPool(FRHICommandList& RHICmdList, EPixelFormat Format, ENiagaraEmptyUAVType Type) const { return EmptyUAVPoolPtr->GetEmptyUAVFromPool(RHICmdList, Format, Type); }
 
 	/** Helper function to return an RDG Texture where the texture contains 0 for all channels. */
-	FRDGTextureRef GetBlackTexture(FRDGBuilder& GraphBuilder, ETextureDimension TextureDimension) const;
+	NIAGARA_API FRDGTextureRef GetBlackTexture(FRDGBuilder& GraphBuilder, ETextureDimension TextureDimension) const;
 
 	/** Helper function to return a RDG Texture SRV where the texture contains 0 for all channels. */
-	FRDGTextureSRVRef GetBlackTextureSRV(FRDGBuilder& GraphBuilder, ETextureDimension TextureDimension) const;
+	NIAGARA_API FRDGTextureSRVRef GetBlackTextureSRV(FRDGBuilder& GraphBuilder, ETextureDimension TextureDimension) const;
 
 	/** Helper function to return a RDG Texture UAV you don't care about the contents of or the results, i.e. to use as a dummy binding. */
-	FRDGTextureUAVRef GetEmptyTextureUAV(FRDGBuilder& GraphBuilder, EPixelFormat Format, ETextureDimension TextureDimension) const;
+	NIAGARA_API FRDGTextureUAVRef GetEmptyTextureUAV(FRDGBuilder& GraphBuilder, EPixelFormat Format, ETextureDimension TextureDimension) const;
 
 	/** Helper function to return a Buffer UAV you don't care about the contents of or the results, i.e. to use as a dummy binding. */
-	FRDGBufferUAVRef GetEmptyBufferUAV(FRDGBuilder& GraphBuilder, EPixelFormat Format) const;
+	NIAGARA_API FRDGBufferUAVRef GetEmptyBufferUAV(FRDGBuilder& GraphBuilder, EPixelFormat Format) const;
 
 	/** Helper function to return a Buffer SRV which will contain 1 element of 0 value, i.e. to use as a dummy binding. */
-	FRDGBufferSRVRef GetEmptyBufferSRV(FRDGBuilder& GraphBuilder, EPixelFormat Format) const;
+	NIAGARA_API FRDGBufferSRVRef GetEmptyBufferSRV(FRDGBuilder& GraphBuilder, EPixelFormat Format) const;
 
 	/**
 	Call this to force all pending ticks to be flushed from the batcher.
diff --git a/Engine/Plugins/FX/Niagara/Source/Niagara/Public/NiagaraGpuReadbackManager.h b/Engine/Plugins/FX/Niagara/Source/Niagara/Public/NiagaraGpuReadbackManager.h
--- a/Engine/Plugins/FX/Niagara/Source/Niagara/Public/NiagaraGpuReadbackManager.h
+++ b/Engine/Plugins/FX/Niagara/Source/Niagara/Public/NiagaraGpuReadbackManager.h
@@ -7,7 +7,7 @@
 
 class FRDGBuilder;
 
-class NIAGARA_API FNiagaraGpuReadbackManager
+class FNiagaraGpuReadbackManager
 {
 protected:
 	typedef TFunction<void(TConstArrayView<TPair<void*, uint32>>)> FCompletionCallback;
@@ -41,11 +41,11 @@ public:
 	};
 
 public:
-	FNiagaraGpuReadbackManager();
-	~FNiagaraGpuReadbackManager();
+	NIAGARA_API FNiagaraGpuReadbackManager();
+	NIAGARA_API ~FNiagaraGpuReadbackManager();
 
 	// Tick call which polls for completed readbacks
-	void Tick();
+	NIAGARA_API void Tick();
 	
 private:
 	// Internal tick impl
@@ -53,27 +53,27 @@ private:
 
 public:
 	// Wait for all pending readbacks to complete
-	void WaitCompletion(FRHICommandListImmediate& RHICmdList);
+	NIAGARA_API void WaitCompletion(FRHICommandListImmediate& RHICmdList);
 
 	// Enqueue a readback of a single buffer
-	void EnqueueReadback(FRDGBuilder& GraphBuilder, FRDGBufferRef Buffer, FCompletionCallback Callback);
-	void EnqueueReadback(FRDGBuilder& GraphBuilder, FRDGBufferRef Buffer, uint32 Offset, uint32 NumBytes, FCompletionCallback Callback);
+	NIAGARA_API void EnqueueReadback(FRDGBuilder& GraphBuilder, FRDGBufferRef Buffer, FCompletionCallback Callback);
+	NIAGARA_API void EnqueueReadback(FRDGBuilder& GraphBuilder, FRDGBufferRef Buffer, uint32 Offset, uint32 NumBytes, FCompletionCallback Callback);
 
 	// Enqueue a readback of multiple buffers
-	void EnqueueReadbacks(FRDGBuilder& GraphBuilder, TConstArrayView<FRDGBufferRef> Buffers, FCompletionCallback Callback);
-	void EnqueueReadbacks(FRDGBuilder& GraphBuilder, TConstArrayView<FRDGBufferRequest> BufferRequest, FCompletionCallback Callback);
+	NIAGARA_API void EnqueueReadbacks(FRDGBuilder& GraphBuilder, TConstArrayView<FRDGBufferRef> Buffers, FCompletionCallback Callback);
+	NIAGARA_API void EnqueueReadbacks(FRDGBuilder& GraphBuilder, TConstArrayView<FRDGBufferRequest> BufferRequest, FCompletionCallback Callback);
 
 	//-TODO:RDG:Deprecate these calls
 	//UE_DEPRECATED(5.1, "Support for legacy shaders bindings will be removed in a future release, please upgrade your data interface.")
-	void EnqueueReadback(FRHICommandList& RHICmdList, FRHIBuffer* Buffer, FCompletionCallback Callback);
+	NIAGARA_API void EnqueueReadback(FRHICommandList& RHICmdList, FRHIBuffer* Buffer, FCompletionCallback Callback);
 	//UE_DEPRECATED(5.1, "Support for legacy shaders bindings will be removed in a future release, please upgrade your data interface.")
-	void EnqueueReadback(FRHICommandList& RHICmdList, FRHIBuffer* Buffer, uint32 Offset, uint32 NumBytes, FCompletionCallback Callback);
+	NIAGARA_API void EnqueueReadback(FRHICommandList& RHICmdList, FRHIBuffer* Buffer, uint32 Offset, uint32 NumBytes, FCompletionCallback Callback);
 
 	//-TODO:RDG:Deprecate these calls
 	//UE_DEPRECATED(5.1, "Support for legacy shaders bindings will be removed in a future release, please upgrade your data interface.")
-	void EnqueueReadbacks(FRHICommandList& RHICmdList, TConstArrayView<FRHIBuffer*> Buffers, FCompletionCallback Callback);
+	NIAGARA_API void EnqueueReadbacks(FRHICommandList& RHICmdList, TConstArrayView<FRHIBuffer*> Buffers, FCompletionCallback Callback);
 	//UE_DEPRECATED(5.1, "Support for legacy shaders bindings will be removed in a future release, please upgrade your data interface.")
-	void EnqueueReadbacks(FRHICommandList& RHICmdList, TConstArrayView<FBufferRequest> BufferRequest, FCompletionCallback Callback);
+	NIAGARA_API void EnqueueReadbacks(FRHICommandList& RHICmdList, TConstArrayView<FBufferRequest> BufferRequest, FCompletionCallback Callback);
 
 private:
 	TQueue<FPendingReadback> PendingReadbacks;
diff --git a/Engine/Plugins/FX/Niagara/Source/Niagara/Public/NiagaraMeshRendererProperties.h b/Engine/Plugins/FX/Niagara/Source/Niagara/Public/NiagaraMeshRendererProperties.h
--- a/Engine/Plugins/FX/Niagara/Source/Niagara/Public/NiagaraMeshRendererProperties.h
+++ b/Engine/Plugins/FX/Niagara/Source/Niagara/Public/NiagaraMeshRendererProperties.h
@@ -55,7 +55,7 @@ enum class ENiagaraMeshLockedAxisSpace : uint8
 };
 
 USTRUCT()
-struct NIAGARA_API FNiagaraMeshMICOverride
+struct FNiagaraMeshMICOverride
 {
 	GENERATED_BODY()
 
@@ -67,14 +67,14 @@ struct NIAGARA_API FNiagaraMeshMICOverride
 };
 
 USTRUCT()
-struct NIAGARA_API FNiagaraMeshMaterialOverride
+struct FNiagaraMeshMaterialOverride
 {
 	GENERATED_USTRUCT_BODY()
 public:
-	FNiagaraMeshMaterialOverride();
+	NIAGARA_API FNiagaraMeshMaterialOverride();
 
 	/** Used to upgrade a serialized FNiagaraParameterStore property to our own struct */
-	bool SerializeFromMismatchedTag(const struct FPropertyTag& Tag, FStructuredArchive::FSlot Slot);
+	NIAGARA_API bool SerializeFromMismatchedTag(const struct FPropertyTag& Tag, FStructuredArchive::FSlot Slot);
 
 	/** Use this UMaterialInterface if set to a valid value. This will be subordinate to UserParamBinding if it is set to a valid user variable.*/
 	UPROPERTY(EditAnywhere, Category = "Material", meta = (EditCondition = "bOverrideMaterials"))
@@ -133,11 +133,11 @@ namespace ENiagaraMeshVFLayout
 };
 
 USTRUCT()
-struct NIAGARA_API FNiagaraMeshRendererMeshProperties
+struct FNiagaraMeshRendererMeshProperties
 {
 	GENERATED_BODY()
 
-	FNiagaraMeshRendererMeshProperties();
+	NIAGARA_API FNiagaraMeshRendererMeshProperties();
 
 	/** The mesh to use when rendering this slot */
 	UPROPERTY(EditAnywhere, Category = "Mesh")
@@ -169,67 +169,67 @@ struct NIAGARA_API FNiagaraMeshRendererMeshProperties
 	ENiagaraMeshPivotOffsetSpace PivotOffsetSpace;
 
 	/** Resolve renderable mesh. */
-	FNiagaraRenderableMeshPtr ResolveRenderableMesh(const FNiagaraEmitterInstance* EmitterInstance) const;
+	NIAGARA_API FNiagaraRenderableMeshPtr ResolveRenderableMesh(const FNiagaraEmitterInstance* EmitterInstance) const;
 
 	/** Is the renderable mesh potentially valid or not. */
-	bool HasValidRenderableMesh() const;
+	NIAGARA_API bool HasValidRenderableMesh() const;
 };
 
-UCLASS(editinlinenew, meta = (DisplayName = "Mesh Renderer"))
-class NIAGARA_API UNiagaraMeshRendererProperties : public UNiagaraRendererProperties
+UCLASS(editinlinenew, meta = (DisplayName = "Mesh Renderer"), MinimalAPI)
+class UNiagaraMeshRendererProperties : public UNiagaraRendererProperties
 {
 public:
 	GENERATED_BODY()
 
-	UNiagaraMeshRendererProperties();
+	NIAGARA_API UNiagaraMeshRendererProperties();
 
 	//UObject Interface
-	virtual void PostLoad() override;
-	virtual void PostInitProperties() override;
-	virtual void Serialize(FArchive& Ar) override;
-	virtual void GetResourceSizeEx(FResourceSizeEx& CumulativeResourceSize) override;
+	NIAGARA_API virtual void PostLoad() override;
+	NIAGARA_API virtual void PostInitProperties() override;
+	NIAGARA_API virtual void Serialize(FArchive& Ar) override;
+	NIAGARA_API virtual void GetResourceSizeEx(FResourceSizeEx& CumulativeResourceSize) override;
 #if WITH_EDITORONLY_DATA
-	virtual void BeginDestroy() override;
-	virtual void PreEditChange(class FProperty* PropertyThatWillChange) override;
-	virtual void PostEditChangeProperty(struct FPropertyChangedEvent& PropertyChangedEvent) override;
-	virtual void RenameVariable(const FNiagaraVariableBase& OldVariable, const FNiagaraVariableBase& NewVariable, const FVersionedNiagaraEmitter& InEmitter) override;
-	virtual void RemoveVariable(const FNiagaraVariableBase& OldVariable, const FVersionedNiagaraEmitter& InEmitter) override;
+	NIAGARA_API virtual void BeginDestroy() override;
+	NIAGARA_API virtual void PreEditChange(class FProperty* PropertyThatWillChange) override;
+	NIAGARA_API virtual void PostEditChangeProperty(struct FPropertyChangedEvent& PropertyChangedEvent) override;
+	NIAGARA_API virtual void RenameVariable(const FNiagaraVariableBase& OldVariable, const FNiagaraVariableBase& NewVariable, const FVersionedNiagaraEmitter& InEmitter) override;
+	NIAGARA_API virtual void RemoveVariable(const FNiagaraVariableBase& OldVariable, const FVersionedNiagaraEmitter& InEmitter) override;
 #endif// WITH_EDITORONLY_DATA
 	//UObject Interface END
 
-	static void InitCDOPropertiesAfterModuleStartup();
+	static NIAGARA_API void InitCDOPropertiesAfterModuleStartup();
 
 	//~ UNiagaraRendererProperties interface
-	virtual FNiagaraRenderer* CreateEmitterRenderer(ERHIFeatureLevel::Type FeatureLevel, const FNiagaraEmitterInstance* Emitter, const FNiagaraSystemInstanceController& InController) override;
-	virtual class FNiagaraBoundsCalculator* CreateBoundsCalculator() override;
-	virtual void GetUsedMaterials(const FNiagaraEmitterInstance* InEmitter, TArray<UMaterialInterface*>& OutMaterials) const override;
-	virtual void GetStreamingMeshInfo(const FBoxSphereBounds& OwnerBounds, const FNiagaraEmitterInstance* InEmitter, TArray<FStreamingRenderAssetPrimitiveInfo>& OutStreamingRenderAssets) const override;
-	virtual const FVertexFactoryType* GetVertexFactoryType() const override;
+	NIAGARA_API virtual FNiagaraRenderer* CreateEmitterRenderer(ERHIFeatureLevel::Type FeatureLevel, const FNiagaraEmitterInstance* Emitter, const FNiagaraSystemInstanceController& InController) override;
+	NIAGARA_API virtual class FNiagaraBoundsCalculator* CreateBoundsCalculator() override;
+	NIAGARA_API virtual void GetUsedMaterials(const FNiagaraEmitterInstance* InEmitter, TArray<UMaterialInterface*>& OutMaterials) const override;
+	NIAGARA_API virtual void GetStreamingMeshInfo(const FBoxSphereBounds& OwnerBounds, const FNiagaraEmitterInstance* InEmitter, TArray<FStreamingRenderAssetPrimitiveInfo>& OutStreamingRenderAssets) const override;
+	NIAGARA_API virtual const FVertexFactoryType* GetVertexFactoryType() const override;
 	virtual bool IsSimTargetSupported(ENiagaraSimTarget InSimTarget) const override { return true; };
-	virtual bool PopulateRequiredBindings(FNiagaraParameterStore& InParameterStore) override;
-	virtual void CollectPSOPrecacheData(FPSOPrecacheParamsList& OutParams) override;
+	NIAGARA_API virtual bool PopulateRequiredBindings(FNiagaraParameterStore& InParameterStore) override;
+	NIAGARA_API virtual void CollectPSOPrecacheData(FPSOPrecacheParamsList& OutParams) override;
 
 #if WITH_EDITORONLY_DATA
-	virtual const TArray<FNiagaraVariable>& GetOptionalAttributes() override;
-	virtual void GetAdditionalVariables(TArray<FNiagaraVariableBase>& OutArray) const override;
-	virtual	void GetRendererWidgets(const FNiagaraEmitterInstance* InEmitter, TArray<TSharedPtr<SWidget>>& OutWidgets, TSharedPtr<FAssetThumbnailPool> InThumbnailPool) const override;
-	virtual	void GetRendererTooltipWidgets(const FNiagaraEmitterInstance* InEmitter, TArray<TSharedPtr<SWidget>>& OutWidgets, TSharedPtr<FAssetThumbnailPool> InThumbnailPool) const override;
-	virtual void GetRendererFeedback(const FVersionedNiagaraEmitter& InEmitter, TArray<FNiagaraRendererFeedback>& OutErrors, TArray<FNiagaraRendererFeedback>& OutWarnings, TArray<FNiagaraRendererFeedback>& OutInfo) const override;
-	void OnMeshChanged();
-	void OnMeshPostBuild(UStaticMesh*);
-	void OnAssetReimported(UObject*);
-	void CheckMaterialUsage();
-	virtual TArray<FNiagaraVariable> GetBoundAttributes() const override;
+	NIAGARA_API virtual const TArray<FNiagaraVariable>& GetOptionalAttributes() override;
+	NIAGARA_API virtual void GetAdditionalVariables(TArray<FNiagaraVariableBase>& OutArray) const override;
+	NIAGARA_API virtual	void GetRendererWidgets(const FNiagaraEmitterInstance* InEmitter, TArray<TSharedPtr<SWidget>>& OutWidgets, TSharedPtr<FAssetThumbnailPool> InThumbnailPool) const override;
+	NIAGARA_API virtual	void GetRendererTooltipWidgets(const FNiagaraEmitterInstance* InEmitter, TArray<TSharedPtr<SWidget>>& OutWidgets, TSharedPtr<FAssetThumbnailPool> InThumbnailPool) const override;
+	NIAGARA_API virtual void GetRendererFeedback(const FVersionedNiagaraEmitter& InEmitter, TArray<FNiagaraRendererFeedback>& OutErrors, TArray<FNiagaraRendererFeedback>& OutWarnings, TArray<FNiagaraRendererFeedback>& OutInfo) const override;
+	NIAGARA_API void OnMeshChanged();
+	NIAGARA_API void OnMeshPostBuild(UStaticMesh*);
+	NIAGARA_API void OnAssetReimported(UObject*);
+	NIAGARA_API void CheckMaterialUsage();
+	NIAGARA_API virtual TArray<FNiagaraVariable> GetBoundAttributes() const override;
 #endif // WITH_EDITORONLY_DATA
-	virtual void CacheFromCompiledData(const FNiagaraDataSetCompiledData* CompiledData) override;
+	NIAGARA_API virtual void CacheFromCompiledData(const FNiagaraDataSetCompiledData* CompiledData) override;
 	virtual bool UseHeterogeneousVolumes() const override { return bUseHeterogeneousVolumes; }
 
-	void UpdateMICs();
+	NIAGARA_API void UpdateMICs();
 
 	virtual ENiagaraRendererSourceDataMode GetCurrentSourceMode() const override { return SourceMode; }
 	//UNiagaraRendererProperties Interface END
 
-	void ApplyMaterialOverrides(const FNiagaraEmitterInstance* EmitterInstance, TArray<UMaterialInterface*>& InOutMaterials) const;
+	NIAGARA_API void ApplyMaterialOverrides(const FNiagaraEmitterInstance* EmitterInstance, TArray<UMaterialInterface*>& InOutMaterials) const;
 
 	/**
 	 * The static mesh(es) to be instanced when rendering mesh particles.
@@ -459,19 +459,19 @@ public:
 	FNiagaraRendererLayout RendererLayoutWithoutCustomSorting;
 
 protected:
-	void InitBindings();
-	void SetPreviousBindings(const FVersionedNiagaraEmitter& SrcEmitter, ENiagaraRendererSourceDataMode InSourceMode);
-	virtual void UpdateSourceModeDerivates(ENiagaraRendererSourceDataMode InSourceMode, bool bFromPropertyEdit = false) override;
+	NIAGARA_API void InitBindings();
+	NIAGARA_API void SetPreviousBindings(const FVersionedNiagaraEmitter& SrcEmitter, ENiagaraRendererSourceDataMode InSourceMode);
+	NIAGARA_API virtual void UpdateSourceModeDerivates(ENiagaraRendererSourceDataMode InSourceMode, bool bFromPropertyEdit = false) override;
 	virtual bool NeedsMIDsForMaterials() const override { return MaterialParameters.HasAnyBindings(); }
 #if WITH_EDITORONLY_DATA
-	bool ChangeRequiresMeshListRebuild(const FProperty* Property);
-	void RebuildMeshList();
+	NIAGARA_API bool ChangeRequiresMeshListRebuild(const FProperty* Property);
+	NIAGARA_API void RebuildMeshList();
 
-	virtual FNiagaraVariable GetBoundAttribute(const FNiagaraVariableAttributeBinding* Binding) const override;
+	NIAGARA_API virtual FNiagaraVariable GetBoundAttribute(const FNiagaraVariableAttributeBinding* Binding) const override;
 #endif
 
 private:
-	static TArray<TWeakObjectPtr<UNiagaraMeshRendererProperties>> MeshRendererPropertiesToDeferredInit;
+	static NIAGARA_API TArray<TWeakObjectPtr<UNiagaraMeshRendererProperties>> MeshRendererPropertiesToDeferredInit;
 
 #if WITH_EDITORONLY_DATA
 	// These properties are deprecated and moved to FNiagaraMeshRendererMeshProperties
diff --git a/Engine/Plugins/FX/Niagara/Source/Niagara/Public/NiagaraMessageDataBase.h b/Engine/Plugins/FX/Niagara/Source/Niagara/Public/NiagaraMessageDataBase.h
--- a/Engine/Plugins/FX/Niagara/Source/Niagara/Public/NiagaraMessageDataBase.h
+++ b/Engine/Plugins/FX/Niagara/Source/Niagara/Public/NiagaraMessageDataBase.h
@@ -5,8 +5,8 @@
 #include "UObject/Object.h"
 #include "NiagaraMessageDataBase.generated.h"
 
-UCLASS()
-class NIAGARA_API UNiagaraMessageDataBase : public UObject
+UCLASS(MinimalAPI)
+class UNiagaraMessageDataBase : public UObject
 {
 	GENERATED_BODY()
 
@@ -16,6 +16,6 @@ public:
 	virtual bool GetAllowDismissal() const { return false; }
 
 #if WITH_EDITORONLY_DATA
-	virtual bool Equals(const UNiagaraMessageDataBase* Other) const;
+	NIAGARA_API virtual bool Equals(const UNiagaraMessageDataBase* Other) const;
 #endif
 };
diff --git a/Engine/Plugins/FX/Niagara/Source/Niagara/Public/NiagaraMessageStore.h b/Engine/Plugins/FX/Niagara/Source/Niagara/Public/NiagaraMessageStore.h
--- a/Engine/Plugins/FX/Niagara/Source/Niagara/Public/NiagaraMessageStore.h
+++ b/Engine/Plugins/FX/Niagara/Source/Niagara/Public/NiagaraMessageStore.h
@@ -7,19 +7,19 @@
 class UNiagaraMessageDataBase;
 
 USTRUCT()
-struct NIAGARA_API FNiagaraMessageStore
+struct FNiagaraMessageStore
 {
 	GENERATED_BODY();
 
 #if WITH_EDITORONLY_DATA
-	const TMap<FGuid, TObjectPtr<UNiagaraMessageDataBase>>& GetMessages() const;
-	void SetMessages(const TMap<FGuid, TObjectPtr<UNiagaraMessageDataBase>>& InMessageKeyToMessageMap);
-	void AddMessage(const FGuid& MessageKey, UNiagaraMessageDataBase* NewMessage);
-	void RemoveMessage(const FGuid& MessageKey);
-	void DismissMessage(const FGuid& MessageKey);
-	bool IsMessageDismissed(const FGuid& MessageKey);
-	bool HasDismissedMessages() const;
-	void ClearDismissedMessages();
+	NIAGARA_API const TMap<FGuid, TObjectPtr<UNiagaraMessageDataBase>>& GetMessages() const;
+	NIAGARA_API void SetMessages(const TMap<FGuid, TObjectPtr<UNiagaraMessageDataBase>>& InMessageKeyToMessageMap);
+	NIAGARA_API void AddMessage(const FGuid& MessageKey, UNiagaraMessageDataBase* NewMessage);
+	NIAGARA_API void RemoveMessage(const FGuid& MessageKey);
+	NIAGARA_API void DismissMessage(const FGuid& MessageKey);
+	NIAGARA_API bool IsMessageDismissed(const FGuid& MessageKey);
+	NIAGARA_API bool HasDismissedMessages() const;
+	NIAGARA_API void ClearDismissedMessages();
 
 private:
 	UPROPERTY()
@@ -32,14 +32,14 @@ private:
 
 #if WITH_EDITORONLY_DATA
 
-struct NIAGARA_API FNiagaraMessageSourceAndStore
+struct FNiagaraMessageSourceAndStore
 {
 public:
-	FNiagaraMessageSourceAndStore(UObject& InSource, FNiagaraMessageStore& InStore);
+	NIAGARA_API FNiagaraMessageSourceAndStore(UObject& InSource, FNiagaraMessageStore& InStore);
 
-	UObject* GetSource() const;
+	NIAGARA_API UObject* GetSource() const;
 
-	FNiagaraMessageStore* GetStore() const;
+	NIAGARA_API FNiagaraMessageStore* GetStore() const;
 
 private:
 	TWeakObjectPtr<UObject> SourceWeak;
diff --git a/Engine/Plugins/FX/Niagara/Source/Niagara/Public/NiagaraModule.h b/Engine/Plugins/FX/Niagara/Source/Niagara/Public/NiagaraModule.h
--- a/Engine/Plugins/FX/Niagara/Source/Niagara/Public/NiagaraModule.h
+++ b/Engine/Plugins/FX/Niagara/Source/Niagara/Public/NiagaraModule.h
@@ -31,7 +31,7 @@ extern NIAGARA_API int32 GEnableVerboseNiagaraChangeIdLogging;
 /**
 * Niagara module interface
 */
-class NIAGARA_API INiagaraModule : public IModuleInterface
+class INiagaraModule : public IModuleInterface
 {
 public:
 #if WITH_EDITOR
@@ -53,98 +53,98 @@ public:
 	DECLARE_DELEGATE_RetVal(void, FOnProcessQueue);
 
 public:
-	virtual void StartupModule()override;
-	virtual void ShutdownModule()override;
-	void ShutdownRenderingResources();
+	NIAGARA_API virtual void StartupModule()override;
+	NIAGARA_API virtual void ShutdownModule()override;
+	NIAGARA_API void ShutdownRenderingResources();
 	
-	void OnPostEngineInit();
-	void OnPreExit();
+	NIAGARA_API void OnPostEngineInit();
+	NIAGARA_API void OnPreExit();
 
-	void OnWorldTickStart(UWorld* World, ELevelTick TickType, float DeltaSeconds);
-	void OnBeginFrame();
-	void OnPostGarbageCollect();
-	void OnWorldBeginTearDown(UWorld* World);
+	NIAGARA_API void OnWorldTickStart(UWorld* World, ELevelTick TickType, float DeltaSeconds);
+	NIAGARA_API void OnBeginFrame();
+	NIAGARA_API void OnPostGarbageCollect();
+	NIAGARA_API void OnWorldBeginTearDown(UWorld* World);
 
-	FDelegateHandle SetOnProcessShaderCompilationQueue(FOnProcessQueue InOnProcessQueue);
-	void ResetOnProcessShaderCompilationQueue(FDelegateHandle DelegateHandle);
-	void ProcessShaderCompilationQueue();
+	NIAGARA_API FDelegateHandle SetOnProcessShaderCompilationQueue(FOnProcessQueue InOnProcessQueue);
+	NIAGARA_API void ResetOnProcessShaderCompilationQueue(FDelegateHandle DelegateHandle);
+	NIAGARA_API void ProcessShaderCompilationQueue();
 
 #if WITH_NIAGARA_DEBUGGER
 	FNiagaraDebuggerClient* GetDebuggerClient() { return DebuggerClient.Get(); }
 #endif
 
 #if WITH_EDITOR
-	const INiagaraMergeManager& GetMergeManager() const;
+	NIAGARA_API const INiagaraMergeManager& GetMergeManager() const;
 
-	void RegisterMergeManager(TSharedRef<INiagaraMergeManager> InMergeManager);
+	NIAGARA_API void RegisterMergeManager(TSharedRef<INiagaraMergeManager> InMergeManager);
 
-	void UnregisterMergeManager(TSharedRef<INiagaraMergeManager> InMergeManager);
+	NIAGARA_API void UnregisterMergeManager(TSharedRef<INiagaraMergeManager> InMergeManager);
 
-	const INiagaraEditorOnlyDataUtilities& GetEditorOnlyDataUtilities() const;
+	NIAGARA_API const INiagaraEditorOnlyDataUtilities& GetEditorOnlyDataUtilities() const;
 
-	void RegisterEditorOnlyDataUtilities(TSharedRef<INiagaraEditorOnlyDataUtilities> InEditorOnlyDataUtilities);
+	NIAGARA_API void RegisterEditorOnlyDataUtilities(TSharedRef<INiagaraEditorOnlyDataUtilities> InEditorOnlyDataUtilities);
 
-	void UnregisterEditorOnlyDataUtilities(TSharedRef<INiagaraEditorOnlyDataUtilities> InEditorOnlyDataUtilities);
+	NIAGARA_API void UnregisterEditorOnlyDataUtilities(TSharedRef<INiagaraEditorOnlyDataUtilities> InEditorOnlyDataUtilities);
 
-	int32 StartScriptCompileJob(const FNiagaraCompileRequestDataBase* InCompileData, const FNiagaraCompileRequestDuplicateDataBase* InCompileDuplicateData, const FNiagaraCompileOptions& InCompileOptions);
-	TSharedPtr<FNiagaraVMExecutableData> GetCompileJobResult(int32 JobID, bool bWait, FNiagaraScriptCompileMetrics& Metrics);
+	NIAGARA_API int32 StartScriptCompileJob(const FNiagaraCompileRequestDataBase* InCompileData, const FNiagaraCompileRequestDuplicateDataBase* InCompileDuplicateData, const FNiagaraCompileOptions& InCompileOptions);
+	NIAGARA_API TSharedPtr<FNiagaraVMExecutableData> GetCompileJobResult(int32 JobID, bool bWait, FNiagaraScriptCompileMetrics& Metrics);
 
-	FDelegateHandle RegisterScriptCompiler(FScriptCompiler ScriptCompiler);
-	void UnregisterScriptCompiler(FDelegateHandle DelegateHandle);
+	NIAGARA_API FDelegateHandle RegisterScriptCompiler(FScriptCompiler ScriptCompiler);
+	NIAGARA_API void UnregisterScriptCompiler(FDelegateHandle DelegateHandle);
 
-	FDelegateHandle RegisterCompileResultDelegate(FCheckCompilationResult ResultDelegate);
-	void UnregisterCompileResultDelegate(FDelegateHandle DelegateHandle);
+	NIAGARA_API FDelegateHandle RegisterCompileResultDelegate(FCheckCompilationResult ResultDelegate);
+	NIAGARA_API void UnregisterCompileResultDelegate(FDelegateHandle DelegateHandle);
 
-	TSharedPtr<FNiagaraCompileRequestDataBase, ESPMode::ThreadSafe> Precompile(UObject* InObj, FGuid Version);
-	TSharedPtr<FNiagaraCompileRequestDuplicateDataBase, ESPMode::ThreadSafe> PrecompileDuplicate(
+	NIAGARA_API TSharedPtr<FNiagaraCompileRequestDataBase, ESPMode::ThreadSafe> Precompile(UObject* InObj, FGuid Version);
+	NIAGARA_API TSharedPtr<FNiagaraCompileRequestDuplicateDataBase, ESPMode::ThreadSafe> PrecompileDuplicate(
 		const FNiagaraCompileRequestDataBase* OwningSystemRequestData,
 		UNiagaraSystem* OwningSystem,
 		UNiagaraEmitter* OwningEmitter,
 		UNiagaraScript* TargetScript,
 		FGuid TargetVersion);
-	FDelegateHandle RegisterPrecompiler(FOnPrecompile PreCompiler);
-	void UnregisterPrecompiler(FDelegateHandle DelegateHandle);
-	FDelegateHandle RegisterPrecompileDuplicator(FOnPrecompileDuplicate PreCompileDuplicator);
-	void UnregisterPrecompileDuplicator(FDelegateHandle DelegateHandle);
+	NIAGARA_API FDelegateHandle RegisterPrecompiler(FOnPrecompile PreCompiler);
+	NIAGARA_API void UnregisterPrecompiler(FDelegateHandle DelegateHandle);
+	NIAGARA_API FDelegateHandle RegisterPrecompileDuplicator(FOnPrecompileDuplicate PreCompileDuplicator);
+	NIAGARA_API void UnregisterPrecompileDuplicator(FDelegateHandle DelegateHandle);
 
-	TSharedPtr<FNiagaraGraphCachedDataBase, ESPMode::ThreadSafe> CacheGraphTraversal(const UObject* InObj, FGuid Version);
-	FDelegateHandle RegisterGraphTraversalCacher(FOnCacheGraphTraversal PreCompiler);
-	void UnregisterGraphTraversalCacher(FDelegateHandle DelegateHandle);
+	NIAGARA_API TSharedPtr<FNiagaraGraphCachedDataBase, ESPMode::ThreadSafe> CacheGraphTraversal(const UObject* InObj, FGuid Version);
+	NIAGARA_API FDelegateHandle RegisterGraphTraversalCacher(FOnCacheGraphTraversal PreCompiler);
+	NIAGARA_API void UnregisterGraphTraversalCacher(FDelegateHandle DelegateHandle);
 
-	FNiagaraCompilationTaskHandle RequestCompileSystem(UNiagaraSystem* System, bool bForce);
-	FDelegateHandle RegisterRequestCompileSystem(FOnRequestCompileSystem RequestCompileSystemCallback);
-	void UnregisterRequestCompileSystem(FDelegateHandle DelegateHandle);
+	NIAGARA_API FNiagaraCompilationTaskHandle RequestCompileSystem(UNiagaraSystem* System, bool bForce);
+	NIAGARA_API FDelegateHandle RegisterRequestCompileSystem(FOnRequestCompileSystem RequestCompileSystemCallback);
+	NIAGARA_API void UnregisterRequestCompileSystem(FDelegateHandle DelegateHandle);
 
-	bool PollSystemCompile(FNiagaraCompilationTaskHandle, FNiagaraSystemAsyncCompileResults&, bool /*bWait*/, bool /*bPeek*/);
-	FDelegateHandle RegisterPollSystemCompile(FOnPollSystemCompile PollSystemCompileCallback);
-	void UnregisterPollSystemCompile(FDelegateHandle DelegateHandle);
+	NIAGARA_API bool PollSystemCompile(FNiagaraCompilationTaskHandle, FNiagaraSystemAsyncCompileResults&, bool /*bWait*/, bool /*bPeek*/);
+	NIAGARA_API FDelegateHandle RegisterPollSystemCompile(FOnPollSystemCompile PollSystemCompileCallback);
+	NIAGARA_API void UnregisterPollSystemCompile(FDelegateHandle DelegateHandle);
 
-	void AbortSystemCompile(FNiagaraCompilationTaskHandle);
-	FDelegateHandle RegisterAbortSystemCompile(FOnAbortSystemCompile AbortSystemCompileCallback);
-	void UnregisterAbortSystemCompile(FDelegateHandle DelegateHandle);
+	NIAGARA_API void AbortSystemCompile(FNiagaraCompilationTaskHandle);
+	NIAGARA_API FDelegateHandle RegisterAbortSystemCompile(FOnAbortSystemCompile AbortSystemCompileCallback);
+	NIAGARA_API void UnregisterAbortSystemCompile(FDelegateHandle DelegateHandle);
 
-	void OnAssetLoaded(UObject* Asset);
+	NIAGARA_API void OnAssetLoaded(UObject* Asset);
 
 #endif
 
 #if NIAGARA_PERF_BASELINES
-	void GeneratePerfBaselines(TArray<UNiagaraEffectType*>& BaselinesToGenerate);
+	NIAGARA_API void GeneratePerfBaselines(TArray<UNiagaraEffectType*>& BaselinesToGenerate);
 
-	bool ToggleStatPerfBaselines(UWorld* World, FCommonViewportClient* ViewportClient, const TCHAR* Stream = nullptr);
-	int32 RenderStatPerfBaselines(UWorld* World, FViewport* Viewport, FCanvas* Canvas, int32 X, int32 Y, const FVector* ViewLocation = nullptr, const FRotator* ViewRotation = nullptr);
+	NIAGARA_API bool ToggleStatPerfBaselines(UWorld* World, FCommonViewportClient* ViewportClient, const TCHAR* Stream = nullptr);
+	NIAGARA_API int32 RenderStatPerfBaselines(UWorld* World, FViewport* Viewport, FCanvas* Canvas, int32 X, int32 Y, const FVector* ViewLocation = nullptr, const FRotator* ViewRotation = nullptr);
 #endif
 
 	FORCEINLINE static bool UseGlobalFXBudget() { return bUseGlobalFXBudget; }
-	static void OnUseGlobalFXBudgetChanged(IConsoleVariable* Variable);
+	static NIAGARA_API void OnUseGlobalFXBudgetChanged(IConsoleVariable* Variable);
 
 	FORCEINLINE static bool DataChannelsEnabled() { return bDataChannelsEnabled; }
-	static void OnDataChannelsEnabledChanged(IConsoleVariable* Variable);
+	static NIAGARA_API void OnDataChannelsEnabledChanged(IConsoleVariable* Variable);
 
 	FORCEINLINE static float GetGlobalSpawnCountScale() { return EngineGlobalSpawnCountScale; }
 	FORCEINLINE static float GetGlobalSystemCountScale() { return EngineGlobalSystemCountScale; }
 
-	static float EngineGlobalSpawnCountScale;
-	static float EngineGlobalSystemCountScale;
+	static NIAGARA_API float EngineGlobalSpawnCountScale;
+	static NIAGARA_API float EngineGlobalSystemCountScale;
 
 	FORCEINLINE static const FNiagaraVariable&  GetVar_Engine_WorldDeltaTime() { return Engine_WorldDeltaTime; }
 	FORCEINLINE static const FNiagaraVariable&  GetVar_Engine_DeltaTime() { return Engine_DeltaTime; }
@@ -273,122 +273,122 @@ public:
 	FOnAbortSystemCompile AbortSystemCompileDelegate;
 #endif
 
-	static int32 EngineEffectsQuality;
+	static NIAGARA_API int32 EngineEffectsQuality;
 
-	static bool bUseGlobalFXBudget;
-	static bool bDataChannelsEnabled;
+	static NIAGARA_API bool bUseGlobalFXBudget;
+	static NIAGARA_API bool bDataChannelsEnabled;
 
 private:
-	static FNiagaraVariable Engine_WorldDeltaTime;
-	static FNiagaraVariable Engine_DeltaTime;
-	static FNiagaraVariable Engine_InvDeltaTime;
-	static FNiagaraVariable Engine_Time; 
-	static FNiagaraVariable Engine_RealTime; 
-	static FNiagaraVariable Engine_QualityLevel; 
-
-	static FNiagaraVariable Engine_Owner_Position;
-	static FNiagaraVariable Engine_Owner_Velocity;
-	static FNiagaraVariable Engine_Owner_XAxis;
-	static FNiagaraVariable Engine_Owner_YAxis;
-	static FNiagaraVariable Engine_Owner_ZAxis;
-	static FNiagaraVariable Engine_Owner_Scale;
-	static FNiagaraVariable Engine_Owner_Rotation;
-	static FNiagaraVariable Engine_Owner_LWC_Tile;
-
-	static FNiagaraVariable Engine_Owner_SystemLocalToWorld;
-	static FNiagaraVariable Engine_Owner_SystemWorldToLocal;
-	static FNiagaraVariable Engine_Owner_SystemLocalToWorldTransposed;
-	static FNiagaraVariable Engine_Owner_SystemWorldToLocalTransposed;
-	static FNiagaraVariable Engine_Owner_SystemLocalToWorldNoScale;
-	static FNiagaraVariable Engine_Owner_SystemWorldToLocalNoScale;
-
-	static FNiagaraVariable Engine_Owner_TimeSinceRendered;
-	static FNiagaraVariable Engine_Owner_LODDistance;
-	static FNiagaraVariable Engine_Owner_LODDistanceFraction;
+	static NIAGARA_API FNiagaraVariable Engine_WorldDeltaTime;
+	static NIAGARA_API FNiagaraVariable Engine_DeltaTime;
+	static NIAGARA_API FNiagaraVariable Engine_InvDeltaTime;
+	static NIAGARA_API FNiagaraVariable Engine_Time; 
+	static NIAGARA_API FNiagaraVariable Engine_RealTime; 
+	static NIAGARA_API FNiagaraVariable Engine_QualityLevel; 
+
+	static NIAGARA_API FNiagaraVariable Engine_Owner_Position;
+	static NIAGARA_API FNiagaraVariable Engine_Owner_Velocity;
+	static NIAGARA_API FNiagaraVariable Engine_Owner_XAxis;
+	static NIAGARA_API FNiagaraVariable Engine_Owner_YAxis;
+	static NIAGARA_API FNiagaraVariable Engine_Owner_ZAxis;
+	static NIAGARA_API FNiagaraVariable Engine_Owner_Scale;
+	static NIAGARA_API FNiagaraVariable Engine_Owner_Rotation;
+	static NIAGARA_API FNiagaraVariable Engine_Owner_LWC_Tile;
+
+	static NIAGARA_API FNiagaraVariable Engine_Owner_SystemLocalToWorld;
+	static NIAGARA_API FNiagaraVariable Engine_Owner_SystemWorldToLocal;
+	static NIAGARA_API FNiagaraVariable Engine_Owner_SystemLocalToWorldTransposed;
+	static NIAGARA_API FNiagaraVariable Engine_Owner_SystemWorldToLocalTransposed;
+	static NIAGARA_API FNiagaraVariable Engine_Owner_SystemLocalToWorldNoScale;
+	static NIAGARA_API FNiagaraVariable Engine_Owner_SystemWorldToLocalNoScale;
+
+	static NIAGARA_API FNiagaraVariable Engine_Owner_TimeSinceRendered;
+	static NIAGARA_API FNiagaraVariable Engine_Owner_LODDistance;
+	static NIAGARA_API FNiagaraVariable Engine_Owner_LODDistanceFraction;
 	
-	static FNiagaraVariable Engine_Owner_ExecutionState;
-
-	static FNiagaraVariable Engine_ExecutionCount;
-	static FNiagaraVariable Engine_Emitter_NumParticles;
-	static FNiagaraVariable Engine_Emitter_SimulationPosition;
-	static FNiagaraVariable Engine_Emitter_TotalSpawnedParticles;
-	static FNiagaraVariable Engine_Emitter_SpawnCountScale;
-	static FNiagaraVariable Engine_System_TickCount;
-	static FNiagaraVariable Engine_System_NumEmittersAlive;
-	static FNiagaraVariable Engine_System_SignificanceIndex;
-	static FNiagaraVariable Engine_System_RandomSeed;
+	static NIAGARA_API FNiagaraVariable Engine_Owner_ExecutionState;
+
+	static NIAGARA_API FNiagaraVariable Engine_ExecutionCount;
+	static NIAGARA_API FNiagaraVariable Engine_Emitter_NumParticles;
+	static NIAGARA_API FNiagaraVariable Engine_Emitter_SimulationPosition;
+	static NIAGARA_API FNiagaraVariable Engine_Emitter_TotalSpawnedParticles;
+	static NIAGARA_API FNiagaraVariable Engine_Emitter_SpawnCountScale;
+	static NIAGARA_API FNiagaraVariable Engine_System_TickCount;
+	static NIAGARA_API FNiagaraVariable Engine_System_NumEmittersAlive;
+	static NIAGARA_API FNiagaraVariable Engine_System_SignificanceIndex;
+	static NIAGARA_API FNiagaraVariable Engine_System_RandomSeed;
 	
-	static FNiagaraVariable Engine_System_CurrentTimeStep;
-	static FNiagaraVariable Engine_System_NumTimeSteps;
-	static FNiagaraVariable Engine_System_TimeStepFraction;
-	static FNiagaraVariable Engine_System_NumParticles;
+	static NIAGARA_API FNiagaraVariable Engine_System_CurrentTimeStep;
+	static NIAGARA_API FNiagaraVariable Engine_System_NumTimeSteps;
+	static NIAGARA_API FNiagaraVariable Engine_System_TimeStepFraction;
+	static NIAGARA_API FNiagaraVariable Engine_System_NumParticles;
 	
-	static FNiagaraVariable Engine_System_NumEmitters;
-	static FNiagaraVariable Engine_NumSystemInstances;
-
-	static FNiagaraVariable Engine_GlobalSpawnCountScale;
-	static FNiagaraVariable Engine_GlobalSystemScale;
-
-	static FNiagaraVariable Engine_System_Age;
-	static FNiagaraVariable Emitter_Age;
-	static FNiagaraVariable Emitter_LocalSpace;
-	static FNiagaraVariable Emitter_Determinism;
-	static FNiagaraVariable Emitter_InterpolatedSpawn;
-	static FNiagaraVariable Emitter_OverrideGlobalSpawnCountScale;
-	static FNiagaraVariable Emitter_SimulationTarget;
-	static FNiagaraVariable Emitter_RandomSeed;
-	static FNiagaraVariable Engine_Emitter_InstanceSeed;
-	static FNiagaraVariable Emitter_SpawnRate;
-	static FNiagaraVariable Emitter_SpawnInterval;
-	static FNiagaraVariable Emitter_InterpSpawnStartDt;
-	static FNiagaraVariable Emitter_SpawnGroup;
-
-	static FNiagaraVariable Particles_UniqueID;
-	static FNiagaraVariable Particles_ID;
-	static FNiagaraVariable Particles_Position;
-	static FNiagaraVariable Particles_Velocity;
-	static FNiagaraVariable Particles_Color;
-	static FNiagaraVariable Particles_SpriteRotation;
-	static FNiagaraVariable Particles_NormalizedAge;
-	static FNiagaraVariable Particles_SpriteSize;
-	static FNiagaraVariable Particles_SpriteFacing;
-	static FNiagaraVariable Particles_SpriteAlignment;
-	static FNiagaraVariable Particles_SubImageIndex;
-	static FNiagaraVariable Particles_DynamicMaterialParameter;
-	static FNiagaraVariable Particles_DynamicMaterialParameter1;
-	static FNiagaraVariable Particles_DynamicMaterialParameter2;
-	static FNiagaraVariable Particles_DynamicMaterialParameter3;
-	static FNiagaraVariable Particles_Scale;
-	static FNiagaraVariable Particles_Lifetime;
-	static FNiagaraVariable Particles_MeshOrientation;
-	static FNiagaraVariable Particles_VisibilityTag;
-	static FNiagaraVariable Particles_MeshIndex;
-	static FNiagaraVariable Particles_UVScale;
-	static FNiagaraVariable Particles_PivotOffset;
-	static FNiagaraVariable Particles_CameraOffset;
-	static FNiagaraVariable Particles_MaterialRandom;
-	static FNiagaraVariable Particles_LightRadius;
-	static FNiagaraVariable Particles_LightExponent;
-	static FNiagaraVariable Particles_LightEnabled;
-	static FNiagaraVariable Particles_LightVolumetricScattering;
-	static FNiagaraVariable Particles_RibbonID;
-	static FNiagaraVariable Particles_RibbonWidth;
-	static FNiagaraVariable Particles_RibbonTwist;
-	static FNiagaraVariable Particles_RibbonFacing;
-	static FNiagaraVariable Particles_RibbonLinkOrder;
-	static FNiagaraVariable Particles_ComponentsEnabled;
-	static FNiagaraVariable Particles_RibbonUVDistance;
-	static FNiagaraVariable Particles_RibbonU0Override;
-	static FNiagaraVariable Particles_RibbonV0RangeOverride;
-	static FNiagaraVariable Particles_RibbonU1Override;
-	static FNiagaraVariable Particles_RibbonV1RangeOverride;
-
-	static FNiagaraVariable ScriptUsage;
-	static FNiagaraVariable ScriptContext;
-	static FNiagaraVariable FunctionDebugState;
-	static FNiagaraVariable DataInstance_Alive;
-	static FNiagaraVariable Translator_BeginDefaults;
-	static FNiagaraVariable Translator_CallID;
+	static NIAGARA_API FNiagaraVariable Engine_System_NumEmitters;
+	static NIAGARA_API FNiagaraVariable Engine_NumSystemInstances;
+
+	static NIAGARA_API FNiagaraVariable Engine_GlobalSpawnCountScale;
+	static NIAGARA_API FNiagaraVariable Engine_GlobalSystemScale;
+
+	static NIAGARA_API FNiagaraVariable Engine_System_Age;
+	static NIAGARA_API FNiagaraVariable Emitter_Age;
+	static NIAGARA_API FNiagaraVariable Emitter_LocalSpace;
+	static NIAGARA_API FNiagaraVariable Emitter_Determinism;
+	static NIAGARA_API FNiagaraVariable Emitter_InterpolatedSpawn;
+	static NIAGARA_API FNiagaraVariable Emitter_OverrideGlobalSpawnCountScale;
+	static NIAGARA_API FNiagaraVariable Emitter_SimulationTarget;
+	static NIAGARA_API FNiagaraVariable Emitter_RandomSeed;
+	static NIAGARA_API FNiagaraVariable Engine_Emitter_InstanceSeed;
+	static NIAGARA_API FNiagaraVariable Emitter_SpawnRate;
+	static NIAGARA_API FNiagaraVariable Emitter_SpawnInterval;
+	static NIAGARA_API FNiagaraVariable Emitter_InterpSpawnStartDt;
+	static NIAGARA_API FNiagaraVariable Emitter_SpawnGroup;
+
+	static NIAGARA_API FNiagaraVariable Particles_UniqueID;
+	static NIAGARA_API FNiagaraVariable Particles_ID;
+	static NIAGARA_API FNiagaraVariable Particles_Position;
+	static NIAGARA_API FNiagaraVariable Particles_Velocity;
+	static NIAGARA_API FNiagaraVariable Particles_Color;
+	static NIAGARA_API FNiagaraVariable Particles_SpriteRotation;
+	static NIAGARA_API FNiagaraVariable Particles_NormalizedAge;
+	static NIAGARA_API FNiagaraVariable Particles_SpriteSize;
+	static NIAGARA_API FNiagaraVariable Particles_SpriteFacing;
+	static NIAGARA_API FNiagaraVariable Particles_SpriteAlignment;
+	static NIAGARA_API FNiagaraVariable Particles_SubImageIndex;
+	static NIAGARA_API FNiagaraVariable Particles_DynamicMaterialParameter;
+	static NIAGARA_API FNiagaraVariable Particles_DynamicMaterialParameter1;
+	static NIAGARA_API FNiagaraVariable Particles_DynamicMaterialParameter2;
+	static NIAGARA_API FNiagaraVariable Particles_DynamicMaterialParameter3;
+	static NIAGARA_API FNiagaraVariable Particles_Scale;
+	static NIAGARA_API FNiagaraVariable Particles_Lifetime;
+	static NIAGARA_API FNiagaraVariable Particles_MeshOrientation;
+	static NIAGARA_API FNiagaraVariable Particles_VisibilityTag;
+	static NIAGARA_API FNiagaraVariable Particles_MeshIndex;
+	static NIAGARA_API FNiagaraVariable Particles_UVScale;
+	static NIAGARA_API FNiagaraVariable Particles_PivotOffset;
+	static NIAGARA_API FNiagaraVariable Particles_CameraOffset;
+	static NIAGARA_API FNiagaraVariable Particles_MaterialRandom;
+	static NIAGARA_API FNiagaraVariable Particles_LightRadius;
+	static NIAGARA_API FNiagaraVariable Particles_LightExponent;
+	static NIAGARA_API FNiagaraVariable Particles_LightEnabled;
+	static NIAGARA_API FNiagaraVariable Particles_LightVolumetricScattering;
+	static NIAGARA_API FNiagaraVariable Particles_RibbonID;
+	static NIAGARA_API FNiagaraVariable Particles_RibbonWidth;
+	static NIAGARA_API FNiagaraVariable Particles_RibbonTwist;
+	static NIAGARA_API FNiagaraVariable Particles_RibbonFacing;
+	static NIAGARA_API FNiagaraVariable Particles_RibbonLinkOrder;
+	static NIAGARA_API FNiagaraVariable Particles_ComponentsEnabled;
+	static NIAGARA_API FNiagaraVariable Particles_RibbonUVDistance;
+	static NIAGARA_API FNiagaraVariable Particles_RibbonU0Override;
+	static NIAGARA_API FNiagaraVariable Particles_RibbonV0RangeOverride;
+	static NIAGARA_API FNiagaraVariable Particles_RibbonU1Override;
+	static NIAGARA_API FNiagaraVariable Particles_RibbonV1RangeOverride;
+
+	static NIAGARA_API FNiagaraVariable ScriptUsage;
+	static NIAGARA_API FNiagaraVariable ScriptContext;
+	static NIAGARA_API FNiagaraVariable FunctionDebugState;
+	static NIAGARA_API FNiagaraVariable DataInstance_Alive;
+	static NIAGARA_API FNiagaraVariable Translator_BeginDefaults;
+	static NIAGARA_API FNiagaraVariable Translator_CallID;
 
 #if NIAGARA_PERF_BASELINES
 	TUniquePtr<FNiagaraPerfBaselineHandler> BaselineHandler;
