[Niagara]

diff --git a/Engine/Plugins/FX/Niagara/Source/Niagara/Public/NiagaraParameterBinding.h b/Engine/Plugins/FX/Niagara/Source/Niagara/Public/NiagaraParameterBinding.h
--- a/Engine/Plugins/FX/Niagara/Source/Niagara/Public/NiagaraParameterBinding.h
+++ b/Engine/Plugins/FX/Niagara/Source/Niagara/Public/NiagaraParameterBinding.h
@@ -24,19 +24,19 @@ enum class ENiagaraParameterBindingUsage
 ENUM_CLASS_FLAGS(ENiagaraParameterBindingUsage)
 
 USTRUCT()
-struct NIAGARA_API FNiagaraParameterBinding
+struct FNiagaraParameterBinding
 {
 	GENERATED_USTRUCT_BODY()
 
 #if WITH_EDITORONLY_DATA
 	virtual ~FNiagaraParameterBinding() = default;
 
-	bool CanBindTo(FNiagaraTypeDefinition TypeDefinition) const;
-	bool CanBindTo(FNiagaraVariableBase InVariable, FNiagaraVariableBase& OutAliasedVariable, FStringView EmitterName) const;
+	NIAGARA_API bool CanBindTo(FNiagaraTypeDefinition TypeDefinition) const;
+	NIAGARA_API bool CanBindTo(FNiagaraVariableBase InVariable, FNiagaraVariableBase& OutAliasedVariable, FStringView EmitterName) const;
 
-	void OnRenameEmitter(FStringView EmitterName);
-	void OnRenameVariable(const FNiagaraVariableBase& OldVariable, const FNiagaraVariableBase& NewVariable, FStringView EmitterName);
-	void OnRemoveVariable(const FNiagaraVariableBase& OldVariable, FStringView EmitterName);
+	NIAGARA_API void OnRenameEmitter(FStringView EmitterName);
+	NIAGARA_API void OnRenameVariable(const FNiagaraVariableBase& OldVariable, const FNiagaraVariableBase& NewVariable, FStringView EmitterName);
+	NIAGARA_API void OnRemoveVariable(const FNiagaraVariableBase& OldVariable, FStringView EmitterName);
 
 	void SetUsage(ENiagaraParameterBindingUsage InUsage) { BindingUsage = InUsage; }
 	void SetAllowedDataInterfaces(TArray<UClass*> InClasses) { AllowedDataInterfaces = InClasses; }
@@ -44,7 +44,7 @@ struct NIAGARA_API FNiagaraParameterBinding
 	void SetAllowedInterfaces(TArray<UClass*> InClasses) { AllowedInterfaces = InClasses; }
 	void SetAllowedTypeDefinitions(TArray<FNiagaraTypeDefinition> InTypeDefs) { AllowedTypeDefinitions = InTypeDefs; }
 
-	void SetDefaultParameter(const FNiagaraVariable& InVariable);
+	NIAGARA_API void SetDefaultParameter(const FNiagaraVariable& InVariable);
 	void SetDefaultParameter(const FName& InName, const FNiagaraTypeDefinition& InTypeDef) { SetDefaultParameter(FNiagaraVariable(InTypeDef, InName)); }
 
 	const FNiagaraVariable& GetDefaultAliasedParameter() { return DefaultAliasedParameter; }
@@ -61,14 +61,14 @@ struct NIAGARA_API FNiagaraParameterBinding
 	TConstArrayView<UClass*> GetAllowedInterfaces() const { return MakeArrayView(AllowedInterfaces); }
 	TConstArrayView<FNiagaraTypeDefinition> GetAllowedTypeDefinitions() const { return MakeArrayView(AllowedTypeDefinitions); }
 
-	bool IsSetoToDefault() const;
-	void SetToDefault();
+	NIAGARA_API bool IsSetoToDefault() const;
+	NIAGARA_API void SetToDefault();
 
-	FString ToString() const;
+	NIAGARA_API FString ToString() const;
 
-	static void ForEachRenameEmitter(UObject* InObject, FStringView EmitterName);
-	static void ForEachRenameVariable(UObject* InObject, const FNiagaraVariableBase& OldVariable, const FNiagaraVariableBase& NewVariable, FStringView EmitterName);
-	static void ForEachRemoveVariable(UObject* InObject, const FNiagaraVariableBase& OldVariable, FStringView EmitterName);
+	static NIAGARA_API void ForEachRenameEmitter(UObject* InObject, FStringView EmitterName);
+	static NIAGARA_API void ForEachRenameVariable(UObject* InObject, const FNiagaraVariableBase& OldVariable, const FNiagaraVariableBase& NewVariable, FStringView EmitterName);
+	static NIAGARA_API void ForEachRemoveVariable(UObject* InObject, const FNiagaraVariableBase& OldVariable, FStringView EmitterName);
 
 	virtual bool HasDefaultValueEditorOnly() const { return false; }
 	virtual TConstArrayView<uint8> GetDefaultValueEditorOnly() const { checkNoEntry(); return MakeArrayView<uint8>(nullptr, 0); }
@@ -120,7 +120,7 @@ protected:
 };
 
 USTRUCT()
-struct NIAGARA_API FNiagaraParameterBindingWithValue : public FNiagaraParameterBinding
+struct FNiagaraParameterBindingWithValue : public FNiagaraParameterBinding
 {
 	GENERATED_USTRUCT_BODY()
 
@@ -129,10 +129,10 @@ struct NIAGARA_API FNiagaraParameterBindingWithValue : public FNiagaraParameterB
 	template<typename T> void SetDefaultParameter(const FNiagaraTypeDefinition& InTypeDef, const T& InDefaultValue) { SetDefaultParameter(FNiagaraVariableBase(InTypeDef, NAME_None), InDefaultValue); }
 	template<typename T> void SetDefaultParameter(const FName& InName, const FNiagaraTypeDefinition& InTypeDef, const T& InDefaultValue) { SetDefaultParameter(FNiagaraVariableBase(InTypeDef, InName), InDefaultValue); }
 
-	virtual bool HasDefaultValueEditorOnly() const override;
-	virtual TConstArrayView<uint8> GetDefaultValueEditorOnly() const override;
-	virtual void SetDefaultValueEditorOnly(TConstArrayView<uint8> Memory) override;
-	virtual void SetDefaultValueEditorOnly(const uint8* Memory) override;
+	NIAGARA_API virtual bool HasDefaultValueEditorOnly() const override;
+	NIAGARA_API virtual TConstArrayView<uint8> GetDefaultValueEditorOnly() const override;
+	NIAGARA_API virtual void SetDefaultValueEditorOnly(TConstArrayView<uint8> Memory) override;
+	NIAGARA_API virtual void SetDefaultValueEditorOnly(const uint8* Memory) override;
 	template<typename T> void SetDefaultValueEditorOnly(const T& InDefaultValue) { check(sizeof(T) == DefaultAliasedParameter.GetType().GetSize()); SetDefaultValueEditorOnly(MakeArrayView(reinterpret_cast<const uint8*>(&InDefaultValue), sizeof(InDefaultValue))); }
 #endif
 
diff --git a/Engine/Plugins/FX/Niagara/Source/Niagara/Public/NiagaraParameterStore.h b/Engine/Plugins/FX/Niagara/Source/Niagara/Public/NiagaraParameterStore.h
--- a/Engine/Plugins/FX/Niagara/Source/Niagara/Public/NiagaraParameterStore.h
+++ b/Engine/Plugins/FX/Niagara/Source/Niagara/Public/NiagaraParameterStore.h
@@ -162,7 +162,7 @@ struct TStructOpsTypeTraits<FNiagaraVariableWithOffset> : public TStructOpsTypeT
 
 /** Base storage class for Niagara parameter values. */
 USTRUCT()
-struct NIAGARA_API FNiagaraParameterStore
+struct FNiagaraParameterStore
 {
 #if WITH_EDITOR
 	DECLARE_MULTICAST_DELEGATE(FOnChanged);
@@ -243,17 +243,17 @@ private:
 	FOnParameterRenamed OnParameterRenamedDelegate;
 #endif
 
-	void SetPositionData(const FName& Name, const FVector& Position);
-	bool HasPositionData(const FName& Name) const;
-	const FVector* GetPositionData(const FName& Name) const;
-	void RemovePositionData(const FName& Name);
+	NIAGARA_API void SetPositionData(const FName& Name, const FVector& Position);
+	NIAGARA_API bool HasPositionData(const FName& Name) const;
+	NIAGARA_API const FVector* GetPositionData(const FName& Name) const;
+	NIAGARA_API void RemovePositionData(const FName& Name);
 
 public:
-	FNiagaraParameterStore();
-	FNiagaraParameterStore(const FNiagaraParameterStore& Other);
-	FNiagaraParameterStore& operator=(const FNiagaraParameterStore& Other);
+	NIAGARA_API FNiagaraParameterStore();
+	NIAGARA_API FNiagaraParameterStore(const FNiagaraParameterStore& Other);
+	NIAGARA_API FNiagaraParameterStore& operator=(const FNiagaraParameterStore& Other);
 
-	virtual ~FNiagaraParameterStore();
+	NIAGARA_API virtual ~FNiagaraParameterStore();
 	
 #if WITH_EDITORONLY_DATA
 	UPROPERTY()
@@ -264,11 +264,11 @@ public:
 	TMap<FNiagaraVariable, FGuid> ParameterGuidMapping;
 #endif
 
-	void SetOwner(UObject* InOwner);
+	NIAGARA_API void SetOwner(UObject* InOwner);
 	UObject* GetOwner() const { return Owner.Get(); }
 
-	void Dump() const;
-	void DumpParameters(bool bDumpBindings = false)const;
+	NIAGARA_API void Dump() const;
+	NIAGARA_API void DumpParameters(bool bDumpBindings = false)const;
 
 	SIZE_T GetResourceSize() const
 	{
@@ -296,33 +296,33 @@ public:
 	FORCEINLINE uint32 GetLayoutVersion() const { return LayoutVersion; }
 
 	/** Binds this parameter store to another, by default if we find no matching parameters we will not maintain a pointer to the store. */
-	void Bind(FNiagaraParameterStore* DestStore, const FNiagaraBoundParameterArray* BoundParameters = nullptr);
+	NIAGARA_API void Bind(FNiagaraParameterStore* DestStore, const FNiagaraBoundParameterArray* BoundParameters = nullptr);
 	/** Unbinds this store form one it's bound to. */
-	void Unbind(FNiagaraParameterStore* DestStore);
+	NIAGARA_API void Unbind(FNiagaraParameterStore* DestStore);
 	/** Unbinds this store from all source and destination stores. */
-	void UnbindAll();
+	NIAGARA_API void UnbindAll();
 	/** Recreates any bindings to reflect a layout change etc. */
-	void Rebind();
+	NIAGARA_API void Rebind();
 	/** Recreates any bindings to reflect a layout change etc. */
-	void TransferBindings(FNiagaraParameterStore& OtherStore);
+	NIAGARA_API void TransferBindings(FNiagaraParameterStore& OtherStore);
 	/** Handles any update such as pushing parameters to bound stores etc. */
-	void Tick();
+	NIAGARA_API void Tick();
 	/** Unbinds this store from all stores it's being driven by. */
-	void UnbindFromSourceStores();
+	NIAGARA_API void UnbindFromSourceStores();
 	
-	bool VerifyBinding(const FNiagaraParameterStore* InDestStore) const;
+	NIAGARA_API bool VerifyBinding(const FNiagaraParameterStore* InDestStore) const;
 
-	void CheckForNaNs() const;
+	NIAGARA_API void CheckForNaNs() const;
 
 	/**
 	Adds the passed parameter to this store.
 	Does nothing if this parameter is already present.
 	Returns true if we added a new parameter.
 	*/
-	virtual bool AddParameter(const FNiagaraVariable& Param, bool bInitialize = true, bool bTriggerRebind = true, int32* OutOffset = nullptr);
+	NIAGARA_API virtual bool AddParameter(const FNiagaraVariable& Param, bool bInitialize = true, bool bTriggerRebind = true, int32* OutOffset = nullptr);
 
 #if WITH_EDITORONLY_DATA
-	virtual void ConvertParameterType(const FNiagaraVariable& ExistingParam, const FNiagaraTypeDefinition& NewType);
+	NIAGARA_API virtual void ConvertParameterType(const FNiagaraVariable& ExistingParam, const FNiagaraTypeDefinition& NewType);
 	
 	template<typename BufferType>
 	void AddConstantBuffer()
@@ -335,19 +335,19 @@ public:
 #endif
 
 	/** Removes the passed parameter if it exists in the store. */
-	virtual bool RemoveParameter(const FNiagaraVariableBase& Param);
+	NIAGARA_API virtual bool RemoveParameter(const FNiagaraVariableBase& Param);
 
 	/** Renames the passed parameter. */
-	virtual void RenameParameter(const FNiagaraVariableBase& Param, FName NewName);
+	NIAGARA_API virtual void RenameParameter(const FNiagaraVariableBase& Param, FName NewName);
 
 	/** Changes the type of the passed parameter. */
-	virtual void ChangeParameterType(const FNiagaraVariableBase& Param, const FNiagaraTypeDefinition& NewType);
+	NIAGARA_API virtual void ChangeParameterType(const FNiagaraVariableBase& Param, const FNiagaraTypeDefinition& NewType);
 
 	/** Removes all parameters from this store and releases any data. */
-	virtual void Empty(bool bClearBindings = true);
+	NIAGARA_API virtual void Empty(bool bClearBindings = true);
 
 	/** Removes all parameters from this store but doesn't change memory allocations. */
-	virtual void Reset(bool bClearBindings = true);
+	NIAGARA_API virtual void Reset(bool bClearBindings = true);
 
 	FORCEINLINE void GetParameters(TArray<FNiagaraVariable>& OutParameters) const
 	{
@@ -369,10 +369,10 @@ public:
 	FORCEINLINE int32 Num() const {return SortedParameterOffsets.Num(); }
 	FORCEINLINE bool IsEmpty() const { return SortedParameterOffsets.Num() == 0; }
 
-	virtual void SanityCheckData(bool bInitInterfaces = true);
+	NIAGARA_API virtual void SanityCheckData(bool bInitInterfaces = true);
 
 	// Called to initially set up the parameter store to *exactly* match the input store (other than any bindings and the internal name of it).
-	virtual void InitFromSource(const FNiagaraParameterStore* SrcStore, bool bNotifyAsDirty);
+	NIAGARA_API virtual void InitFromSource(const FNiagaraParameterStore* SrcStore, bool bNotifyAsDirty);
 
 	/** Gets the index of the passed parameter. If it is a data interface, this is an offset into the data interface table, otherwise a byte offset into the parameter data buffer. */
 	FORCEINLINE_DEBUGGABLE int32 IndexOf(const FNiagaraVariableBase& Parameter) const
@@ -381,7 +381,7 @@ public:
 		return Off ? *Off : (int32)INDEX_NONE;
 	}
 	
-	virtual const FNiagaraVariableWithOffset* FindParameterVariable(const FNiagaraVariable& Parameter, bool IgnoreType = false) const;
+	NIAGARA_API virtual const FNiagaraVariableWithOffset* FindParameterVariable(const FNiagaraVariable& Parameter, bool IgnoreType = false) const;
 
 	/** Gets the typed parameter data. */
 	template<typename T>
@@ -448,7 +448,7 @@ public:
 	 * Just using the raw parameter store pointer via IndexOf() will be wrong if the target struct contains a lwc type like FVector.
 	 * Returns true if the data was copied, false if the parameter could not be found in the store.
 	 */
-	bool CopyParameterData(const FNiagaraVariable& Parameter, uint8* DestinationData) const;
+	NIAGARA_API bool CopyParameterData(const FNiagaraVariable& Parameter, uint8* DestinationData) const;
 
 	/** Returns the data interface at the passed offset. */
 	FORCEINLINE UNiagaraDataInterface* GetDataInterface(int32 Offset)const
@@ -462,18 +462,18 @@ public:
 	}
 
 	/** Returns the data interface for the passed parameter if it exists in this store. */
-	UNiagaraDataInterface* GetDataInterface(const FNiagaraVariable& Parameter) const;
+	NIAGARA_API UNiagaraDataInterface* GetDataInterface(const FNiagaraVariable& Parameter) const;
 
 	/** Returns a struct converter for the given variable, if the store contains the variable and it's a LWC type. */
-	FNiagaraLwcStructConverter GetStructConverter(const FNiagaraVariable& Parameter) const;
+	NIAGARA_API FNiagaraLwcStructConverter GetStructConverter(const FNiagaraVariable& Parameter) const;
 
 	/** Returns the associated FNiagaraVariable for the passed data interface if it exists in the store. Null if not.*/
-	const FNiagaraVariableBase* FindVariable(const UNiagaraDataInterface* Interface) const;
+	NIAGARA_API const FNiagaraVariableBase* FindVariable(const UNiagaraDataInterface* Interface) const;
 
-	virtual const int32* FindParameterOffset(const FNiagaraVariableBase& Parameter, bool IgnoreType = false) const;
+	NIAGARA_API virtual const int32* FindParameterOffset(const FNiagaraVariableBase& Parameter, bool IgnoreType = false) const;
 
-	void PostLoad();
-	void SortParameters();
+	NIAGARA_API void PostLoad();
+	NIAGARA_API void SortParameters();
 
 	/** Returns the UObject at the passed offset. */
 	FORCEINLINE UObject* GetUObject(int32 Offset)const
@@ -495,8 +495,8 @@ public:
 	}
 
 	/** Copies the passed parameter from this parameter store into another. */
-	void CopyParameterData(FNiagaraParameterStore& DestStore, const FNiagaraVariable& Parameter) const;
-	void CopyParameterData(FNiagaraParameterStore& DestStore, const FNiagaraVariable& SourceParameter, const FNiagaraVariable& TargetParameter) const;
+	NIAGARA_API void CopyParameterData(FNiagaraParameterStore& DestStore, const FNiagaraVariable& Parameter) const;
+	NIAGARA_API void CopyParameterData(FNiagaraParameterStore& DestStore, const FNiagaraVariable& SourceParameter, const FNiagaraVariable& TargetParameter) const;
 	
 	enum class EDataInterfaceCopyMethod
 	{
@@ -510,16 +510,16 @@ public:
 	};
 
 	/** Copies all parameters from this parameter store into another.*/
-	void CopyParametersTo(FNiagaraParameterStore& DestStore, bool bOnlyAdd, EDataInterfaceCopyMethod DataInterfaceCopyMethod = EDataInterfaceCopyMethod::None) const;
+	NIAGARA_API void CopyParametersTo(FNiagaraParameterStore& DestStore, bool bOnlyAdd, EDataInterfaceCopyMethod DataInterfaceCopyMethod = EDataInterfaceCopyMethod::None) const;
 
 	/** Remove all parameters from this parameter store from another.*/
-	void RemoveParameters(FNiagaraParameterStore& DestStore);
+	NIAGARA_API void RemoveParameters(FNiagaraParameterStore& DestStore);
 
-	FString ToString() const;
+	NIAGARA_API FString ToString() const;
 
-	virtual bool SetPositionParameterValue(const FVector& InValue, const FName& ParamName, bool bAdd=false);
-	virtual const FVector* GetPositionParameterValue(const FName& ParamName) const;
-	void ResolvePositions(FNiagaraLWCConverter LwcConverter);
+	NIAGARA_API virtual bool SetPositionParameterValue(const FVector& InValue, const FName& ParamName, bool bAdd=false);
+	NIAGARA_API virtual const FVector* GetPositionParameterValue(const FName& ParamName) const;
+	NIAGARA_API void ResolvePositions(FNiagaraLWCConverter LwcConverter);
 
 	template<typename T>
 	FORCEINLINE_DEBUGGABLE bool SetParameterValue(const T& InValue, const FNiagaraVariable& Param, bool bAdd=false)
@@ -573,7 +573,7 @@ public:
 		OnParameterChange();
 	}
 
-	bool SetParameterData(const uint8* Data, FNiagaraVariable Param, bool bAdd = false);
+	NIAGARA_API bool SetParameterData(const uint8* Data, FNiagaraVariable Param, bool bAdd = false);
 
 	FORCEINLINE_DEBUGGABLE void SetDataInterface(UNiagaraDataInterface* InInterface, int32 Offset)
 	{
@@ -654,9 +654,9 @@ public:
 	}
 
 #if WITH_EDITOR
-	FDelegateHandle AddOnChangedHandler(FOnChanged::FDelegate InOnChanged);
-	void RemoveOnChangedHandler(FDelegateHandle DelegateHandle);
-	void RemoveAllOnChangedHandlers(const void* InUserObject);
+	NIAGARA_API FDelegateHandle AddOnChangedHandler(FOnChanged::FDelegate InOnChanged);
+	NIAGARA_API void RemoveOnChangedHandler(FDelegateHandle DelegateHandle);
+	NIAGARA_API void RemoveAllOnChangedHandlers(const void* InUserObject);
 
 	FOnStructureChanged& OnStructureChanged() { return OnStructureChangedDelegate; }
 	FOnParameterRenamed& OnParameterRenamed() { return OnParameterRenamedDelegate; }
@@ -665,11 +665,11 @@ public:
 	void TriggerOnLayoutChanged() { OnLayoutChange(); }
 
 protected:
-	void TickBindings() const;
-	void OnLayoutChange();
-	void CopySortedParameterOffsets(TArrayView<const FNiagaraVariableWithOffset> Src);
-	void AssignParameterData(TConstArrayView<uint8> SourceParameterData);
-	static int32 PaddedParameterSize(int32 ParameterSize);
+	NIAGARA_API void TickBindings() const;
+	NIAGARA_API void OnLayoutChange();
+	NIAGARA_API void CopySortedParameterOffsets(TArrayView<const FNiagaraVariableWithOffset> Src);
+	NIAGARA_API void AssignParameterData(TConstArrayView<uint8> SourceParameterData);
+	static NIAGARA_API int32 PaddedParameterSize(int32 ParameterSize);
 
 	/** Returns the parameter data at the passed offset. */
 	FORCEINLINE uint8* GetParameterData_Internal(int32 Offset) 
@@ -685,10 +685,10 @@ protected:
 		//SetParameterData((const uint8*)&Param, ParamOffset, sizeof(ParamType)); // TODO why aren't we using this path instead of SetParametersByOffset?
 	}
 
-	void SetParameterDataArray(const TArray<uint8>& InParameterDataArray, bool bNotifyAsDirty = true);
-	void SetDataInterfaces(const TArray<UNiagaraDataInterface*>& InDataInterfaces, bool bNotifyAsDirty = true);
-	void SetUObjects(const TArray<UObject*>& InUObjects, bool bNotifyAsDirty = true);
-	void SetOriginalPositionData(const TArray<FNiagaraPositionSource>& InOriginalPositionData);
+	NIAGARA_API void SetParameterDataArray(const TArray<uint8>& InParameterDataArray, bool bNotifyAsDirty = true);
+	NIAGARA_API void SetDataInterfaces(const TArray<UNiagaraDataInterface*>& InDataInterfaces, bool bNotifyAsDirty = true);
+	NIAGARA_API void SetUObjects(const TArray<UObject*>& InUObjects, bool bNotifyAsDirty = true);
+	NIAGARA_API void SetOriginalPositionData(const TArray<FNiagaraPositionSource>& InOriginalPositionData);
 
 	friend struct FNiagaraParameterStoreToDataSetBinding;    // this should be the only class calling SetParameterByOffset
 };
diff --git a/Engine/Plugins/FX/Niagara/Source/Niagara/Public/NiagaraRenderGraphUtils.h b/Engine/Plugins/FX/Niagara/Source/Niagara/Public/NiagaraRenderGraphUtils.h
--- a/Engine/Plugins/FX/Niagara/Source/Niagara/Public/NiagaraRenderGraphUtils.h
+++ b/Engine/Plugins/FX/Niagara/Source/Niagara/Public/NiagaraRenderGraphUtils.h
@@ -5,23 +5,23 @@
 
 //////////////////////////////////////////////////////////////////////////
 // Helper for wrapping a persistent RW buffer
-struct NIAGARA_API FNiagaraPooledRWBuffer
+struct FNiagaraPooledRWBuffer
 {
 	~FNiagaraPooledRWBuffer() { Release(); }
 
-	void Initialize(FRDGBuilder& GraphBuilder, const TCHAR* ResourceName, EPixelFormat PixelFormat, const FRDGBufferDesc& BufferDesc);
-	void Initialize(FRDGBuilder& GraphBuilder, const TCHAR* ResourceName, EPixelFormat PixelFormat, const uint32 BytesPerElemenet, const uint32 NumElements, EBufferUsageFlags UsageFlags = EBufferUsageFlags::None);
-	void Release();
+	NIAGARA_API void Initialize(FRDGBuilder& GraphBuilder, const TCHAR* ResourceName, EPixelFormat PixelFormat, const FRDGBufferDesc& BufferDesc);
+	NIAGARA_API void Initialize(FRDGBuilder& GraphBuilder, const TCHAR* ResourceName, EPixelFormat PixelFormat, const uint32 BytesPerElemenet, const uint32 NumElements, EBufferUsageFlags UsageFlags = EBufferUsageFlags::None);
+	NIAGARA_API void Release();
 
 	bool IsValid() const { return Buffer.IsValid(); }
 
 
 	const TRefCountPtr<FRDGPooledBuffer>& GetPooledBuffer() const { return Buffer; }
-	FRDGBufferRef GetOrCreateBuffer(FRDGBuilder& GraphBuilder);
-	FRDGBufferSRVRef GetOrCreateSRV(FRDGBuilder& GraphBuilder);
-	FRDGBufferUAVRef GetOrCreateUAV(FRDGBuilder& GraphBuilder);
+	NIAGARA_API FRDGBufferRef GetOrCreateBuffer(FRDGBuilder& GraphBuilder);
+	NIAGARA_API FRDGBufferSRVRef GetOrCreateSRV(FRDGBuilder& GraphBuilder);
+	NIAGARA_API FRDGBufferUAVRef GetOrCreateUAV(FRDGBuilder& GraphBuilder);
 
-	void EndGraphUsage();
+	NIAGARA_API void EndGraphUsage();
 
 private:
 	TRefCountPtr<FRDGPooledBuffer>	Buffer;
@@ -38,23 +38,23 @@ private:
 
 //////////////////////////////////////////////////////////////////////////
 // Helper for wrapping a persistent RW texture
-struct NIAGARA_API FNiagaraPooledRWTexture
+struct FNiagaraPooledRWTexture
 {
 	~FNiagaraPooledRWTexture() { Release(); }
 
-	void Initialize(FRDGBuilder& GraphBuilder, const TCHAR* ResourceName, const FRDGTextureDesc& TextureDesc);
-	void Release();
+	NIAGARA_API void Initialize(FRDGBuilder& GraphBuilder, const TCHAR* ResourceName, const FRDGTextureDesc& TextureDesc);
+	NIAGARA_API void Release();
 
 	bool IsValid() const { return Texture.IsValid(); }
 
 	const TRefCountPtr<IPooledRenderTarget>& GetPooledTexture() const { return Texture; }
-	FRDGTextureRef GetOrCreateTexture(FRDGBuilder& GraphBuilder);
-	FRDGTextureSRVRef GetOrCreateSRV(FRDGBuilder& GraphBuilder);
-	FRDGTextureUAVRef GetOrCreateUAV(FRDGBuilder& GraphBuilder);
+	NIAGARA_API FRDGTextureRef GetOrCreateTexture(FRDGBuilder& GraphBuilder);
+	NIAGARA_API FRDGTextureSRVRef GetOrCreateSRV(FRDGBuilder& GraphBuilder);
+	NIAGARA_API FRDGTextureUAVRef GetOrCreateUAV(FRDGBuilder& GraphBuilder);
 
-	void CopyToTexture(FRDGBuilder& GraphBuilder, FRHITexture* DestinationTexture, const TCHAR* NameIfNotRegistered);
+	NIAGARA_API void CopyToTexture(FRDGBuilder& GraphBuilder, FRHITexture* DestinationTexture, const TCHAR* NameIfNotRegistered);
 
-	void EndGraphUsage();
+	NIAGARA_API void EndGraphUsage();
 
 private:
 	TRefCountPtr<IPooledRenderTarget>	Texture;
diff --git a/Engine/Plugins/FX/Niagara/Source/Niagara/Public/NiagaraRenderer.h b/Engine/Plugins/FX/Niagara/Source/Niagara/Public/NiagaraRenderer.h
--- a/Engine/Plugins/FX/Niagara/Source/Niagara/Public/NiagaraRenderer.h
+++ b/Engine/Plugins/FX/Niagara/Source/Niagara/Public/NiagaraRenderer.h
@@ -25,17 +25,17 @@ class FNiagaraGPURendererCount;
 class FNiagaraSystemInstanceController;
 
 /** Struct used to pass dynamic data from game thread to render thread */
-struct NIAGARA_API FNiagaraDynamicDataBase
+struct FNiagaraDynamicDataBase
 {
-	explicit FNiagaraDynamicDataBase(const FNiagaraEmitterInstance* InEmitter);
-	virtual ~FNiagaraDynamicDataBase();
+	NIAGARA_API explicit FNiagaraDynamicDataBase(const FNiagaraEmitterInstance* InEmitter);
+	NIAGARA_API virtual ~FNiagaraDynamicDataBase();
 
 	FNiagaraDynamicDataBase() = delete;
 	FNiagaraDynamicDataBase(FNiagaraDynamicDataBase& Other) = delete;
 	FNiagaraDynamicDataBase& operator=(const FNiagaraDynamicDataBase& Other) = delete;
 
-	bool IsGpuLowLatencyTranslucencyEnabled() const;
-	FNiagaraDataBuffer* GetParticleDataToRender(bool bIsLowLatencyTranslucent = false) const;
+	NIAGARA_API bool IsGpuLowLatencyTranslucencyEnabled() const;
+	NIAGARA_API FNiagaraDataBuffer* GetParticleDataToRender(bool bIsLowLatencyTranslucent = false) const;
 	FORCEINLINE ENiagaraSimTarget GetSimTarget() const { return SimTarget; }
 	FORCEINLINE FMaterialRelevance GetMaterialRelevance() const { return MaterialRelevance; }
 
@@ -69,22 +69,22 @@ struct FParticleRenderData
 /**
 * Base class for Niagara System renderers.
 */
-class NIAGARA_API FNiagaraRenderer
+class FNiagaraRenderer
 {
 public:
 
-	FNiagaraRenderer(ERHIFeatureLevel::Type FeatureLevel, const UNiagaraRendererProperties *InProps, const FNiagaraEmitterInstance* Emitter);
-	virtual ~FNiagaraRenderer();
+	NIAGARA_API FNiagaraRenderer(ERHIFeatureLevel::Type FeatureLevel, const UNiagaraRendererProperties *InProps, const FNiagaraEmitterInstance* Emitter);
+	NIAGARA_API virtual ~FNiagaraRenderer();
 
 	FNiagaraRenderer(const FNiagaraRenderer& Other) = delete;
 	FNiagaraRenderer& operator=(const FNiagaraRenderer& Other) = delete;
 
-	virtual void Initialize(const UNiagaraRendererProperties *InProps, const FNiagaraEmitterInstance* Emitter, const FNiagaraSystemInstanceController& InComponent);
+	NIAGARA_API virtual void Initialize(const UNiagaraRendererProperties *InProps, const FNiagaraEmitterInstance* Emitter, const FNiagaraSystemInstanceController& InComponent);
 	virtual void CreateRenderThreadResources() {}
 	virtual void ReleaseRenderThreadResources() {}
 	virtual void DestroyRenderState_Concurrent() {}
 
-	virtual FPrimitiveViewRelevance GetViewRelevance(const FSceneView* View, const FNiagaraSceneProxy *SceneProxy)const;
+	NIAGARA_API virtual FPrimitiveViewRelevance GetViewRelevance(const FSceneView* View, const FNiagaraSceneProxy *SceneProxy)const;
 	virtual void GetDynamicMeshElements(const TArray<const FSceneView*>& Views, const FSceneViewFamily& ViewFamily, uint32 VisibilityMap, FMeshElementCollector& Collector, const FNiagaraSceneProxy *SceneProxy) const {}
 	virtual FNiagaraDynamicDataBase* GenerateDynamicData(const FNiagaraSceneProxy* Proxy, const UNiagaraRendererProperties* InProperties, const FNiagaraEmitterInstance* Emitteride) const { return nullptr; }
 
@@ -92,12 +92,12 @@ public:
 	virtual int32 GetDynamicDataSize()const { return 0; }
 	virtual bool IsMaterialValid(const UMaterialInterface* Mat)const { return Mat != nullptr; }
 
-	static void SortIndices(const struct FNiagaraGPUSortInfo& SortInfo, const FNiagaraRendererVariableInfo& SortVariable, const FNiagaraDataBuffer& Buffer, FGlobalDynamicReadBuffer::FAllocation& OutIndices);
-	static int32 SortAndCullIndices(const FNiagaraGPUSortInfo& SortInfo, const FNiagaraDataBuffer& Buffer, FGlobalDynamicReadBuffer::FAllocation& OutIndices);
+	static NIAGARA_API void SortIndices(const struct FNiagaraGPUSortInfo& SortInfo, const FNiagaraRendererVariableInfo& SortVariable, const FNiagaraDataBuffer& Buffer, FGlobalDynamicReadBuffer::FAllocation& OutIndices);
+	static NIAGARA_API int32 SortAndCullIndices(const FNiagaraGPUSortInfo& SortInfo, const FNiagaraDataBuffer& Buffer, FGlobalDynamicReadBuffer::FAllocation& OutIndices);
 
-	static FVector4f CalcMacroUVParameters(const FSceneView& View, FVector MacroUVPosition, float MacroUVRadius);
+	static NIAGARA_API FVector4f CalcMacroUVParameters(const FSceneView& View, FVector MacroUVPosition, float MacroUVRadius);
 
-	void SetDynamicData_RenderThread(FNiagaraDynamicDataBase* NewDynamicData);
+	NIAGARA_API void SetDynamicData_RenderThread(FNiagaraDynamicDataBase* NewDynamicData);
 	FORCEINLINE FNiagaraDynamicDataBase *GetDynamicData() const { return DynamicDataRender; }
 	FORCEINLINE bool HasDynamicData() const { return DynamicDataRender != nullptr; }
 	FORCEINLINE bool HasLights() const { return bHasLights; }
@@ -123,18 +123,18 @@ public:
 	FORCEINLINE static FRHIShaderResourceView* GetSrvOrDefaultUInt(FRHIShaderResourceView* InSRV) { return InSRV ? InSRV : GetDummyUIntBuffer(); }
 	FORCEINLINE static FRHIShaderResourceView* GetSrvOrDefaultInt(FRHIShaderResourceView* InSRV) { return InSRV ? InSRV : GetDummyIntBuffer(); }
 
-	static FRHIShaderResourceView* GetDummyFloatBuffer();
-	static FRHIShaderResourceView* GetDummyFloat2Buffer();
-	static FRHIShaderResourceView* GetDummyFloat4Buffer();
-	static FRHIShaderResourceView* GetDummyWhiteColorBuffer();
-	static FRHIShaderResourceView* GetDummyIntBuffer();
-	static FRHIShaderResourceView* GetDummyUIntBuffer();
-	static FRHIShaderResourceView* GetDummyUInt2Buffer();
-	static FRHIShaderResourceView* GetDummyUInt4Buffer();
-	static FRHIShaderResourceView* GetDummyTextureReadBuffer2D();
-	static FRHIShaderResourceView* GetDummyTextureReadBuffer2DArray();
-	static FRHIShaderResourceView* GetDummyTextureReadBuffer3D();
-	static FRHIShaderResourceView* GetDummyHalfBuffer();
+	static NIAGARA_API FRHIShaderResourceView* GetDummyFloatBuffer();
+	static NIAGARA_API FRHIShaderResourceView* GetDummyFloat2Buffer();
+	static NIAGARA_API FRHIShaderResourceView* GetDummyFloat4Buffer();
+	static NIAGARA_API FRHIShaderResourceView* GetDummyWhiteColorBuffer();
+	static NIAGARA_API FRHIShaderResourceView* GetDummyIntBuffer();
+	static NIAGARA_API FRHIShaderResourceView* GetDummyUIntBuffer();
+	static NIAGARA_API FRHIShaderResourceView* GetDummyUInt2Buffer();
+	static NIAGARA_API FRHIShaderResourceView* GetDummyUInt4Buffer();
+	static NIAGARA_API FRHIShaderResourceView* GetDummyTextureReadBuffer2D();
+	static NIAGARA_API FRHIShaderResourceView* GetDummyTextureReadBuffer2DArray();
+	static NIAGARA_API FRHIShaderResourceView* GetDummyTextureReadBuffer3D();
+	static NIAGARA_API FRHIShaderResourceView* GetDummyHalfBuffer();
 
 	FORCEINLINE ENiagaraSimTarget GetSimTarget() const { return SimTarget; }
 
@@ -145,12 +145,12 @@ public:
 
 protected:
 
-	virtual void ProcessMaterialParameterBindings(const FNiagaraRendererMaterialParameters& MaterialParameters, const FNiagaraEmitterInstance* InEmitter, TConstArrayView<UMaterialInterface*> InMaterials) const;
+	NIAGARA_API virtual void ProcessMaterialParameterBindings(const FNiagaraRendererMaterialParameters& MaterialParameters, const FNiagaraEmitterInstance* InEmitter, TConstArrayView<UMaterialInterface*> InMaterials) const;
 
-	bool IsRendererEnabled(const UNiagaraRendererProperties* InProperties, const FNiagaraEmitterInstance* Emitter) const;
-	bool UseLocalSpace(const FNiagaraSceneProxy* Proxy) const;
+	NIAGARA_API bool IsRendererEnabled(const UNiagaraRendererProperties* InProperties, const FNiagaraEmitterInstance* Emitter) const;
+	NIAGARA_API bool UseLocalSpace(const FNiagaraSceneProxy* Proxy) const;
 
-	static bool ViewFamilySupportLowLatencyTranslucency(const FSceneViewFamily& ViewFamily);
+	static NIAGARA_API bool ViewFamilySupportLowLatencyTranslucency(const FSceneViewFamily& ViewFamily);
 
 	struct FNiagaraDynamicDataBase *DynamicDataRender;
 
@@ -172,7 +172,7 @@ protected:
 #endif
 
 
-	static FParticleRenderData TransferDataToGPU(FGlobalDynamicReadBuffer& DynamicReadBuffer, const FNiagaraRendererLayout* RendererLayout, TConstArrayView<uint32> IntComponents, const FNiagaraDataBuffer* SrcData);
+	static NIAGARA_API FParticleRenderData TransferDataToGPU(FGlobalDynamicReadBuffer& DynamicReadBuffer, const FNiagaraRendererLayout* RendererLayout, TConstArrayView<uint32> IntComponents, const FNiagaraDataBuffer* SrcData);
 
 	/** Cached array of materials used from the properties data. Validated with usage flags etc. */
 	TArray<UMaterialInterface*> BaseMaterials_GT;
diff --git a/Engine/Plugins/FX/Niagara/Source/Niagara/Public/NiagaraRendererComponents.h b/Engine/Plugins/FX/Niagara/Source/Niagara/Public/NiagaraRendererComponents.h
--- a/Engine/Plugins/FX/Niagara/Source/Niagara/Public/NiagaraRendererComponents.h
+++ b/Engine/Plugins/FX/Niagara/Source/Niagara/Public/NiagaraRendererComponents.h
@@ -15,18 +15,18 @@ class FNiagaraDataSet;
 /**
 * NiagaraRendererComponents renders an FNiagaraEmitterInstance as scene components
 */
-class NIAGARA_API FNiagaraRendererComponents : public FNiagaraRenderer
+class FNiagaraRendererComponents : public FNiagaraRenderer
 {
 	friend struct FNiagaraRendererComponentsOnObjectsReplacedHelper;
 public:
 
-	FNiagaraRendererComponents(ERHIFeatureLevel::Type FeatureLevel, const UNiagaraRendererProperties *InProps, const FNiagaraEmitterInstance* Emitter);
-	virtual ~FNiagaraRendererComponents();
+	NIAGARA_API FNiagaraRendererComponents(ERHIFeatureLevel::Type FeatureLevel, const UNiagaraRendererProperties *InProps, const FNiagaraEmitterInstance* Emitter);
+	NIAGARA_API virtual ~FNiagaraRendererComponents();
 
 	//FNiagaraRenderer interface
-	virtual void DestroyRenderState_Concurrent() override;
-	virtual void PostSystemTick_GameThread(const UNiagaraRendererProperties* InProperties, const FNiagaraEmitterInstance* Emitter) override;
-	virtual void OnSystemComplete_GameThread(const UNiagaraRendererProperties* InProperties, const FNiagaraEmitterInstance* Emitter) override;
+	NIAGARA_API virtual void DestroyRenderState_Concurrent() override;
+	NIAGARA_API virtual void PostSystemTick_GameThread(const UNiagaraRendererProperties* InProperties, const FNiagaraEmitterInstance* Emitter) override;
+	NIAGARA_API virtual void OnSystemComplete_GameThread(const UNiagaraRendererProperties* InProperties, const FNiagaraEmitterInstance* Emitter) override;
 	//FNiagaraRenderer interface END
 
 private:
@@ -57,12 +57,12 @@ private:
 		int32 LastAssignedToParticleID = -1;
 	};
 
-	void ResetComponentPool(bool bResetOwner);
+	NIAGARA_API void ResetComponentPool(bool bResetOwner);
 #if WITH_EDITOR
-	void OnObjectsReplacedCallback(const TMap<UObject*, UObject*>& ReplacementsMap);
+	NIAGARA_API void OnObjectsReplacedCallback(const TMap<UObject*, UObject*>& ReplacementsMap);
 #endif
 	
-	static void TickPropertyBindings(
+	static NIAGARA_API void TickPropertyBindings(
 		const UNiagaraComponentRendererProperties* Properties,
 		USceneComponent* Component,
 		FNiagaraDataSet& Data,
@@ -71,8 +71,8 @@ private:
 		const FNiagaraLWCConverter& LwcConverter);
 
 	// These property accessor methods are largely copied over from MovieSceneCommonHelpers.h
-	static FComponentPropertyAddress FindPropertyRecursive(void* BasePointer, UStruct* InStruct, TArray<FString>& InPropertyNames, uint32 Index);
-	static FComponentPropertyAddress FindProperty(const UObject& Object, const FString& InPropertyPath);
+	static NIAGARA_API FComponentPropertyAddress FindPropertyRecursive(void* BasePointer, UStruct* InStruct, TArray<FString>& InPropertyNames, uint32 Index);
+	static NIAGARA_API FComponentPropertyAddress FindProperty(const UObject& Object, const FString& InPropertyPath);
 
 	// this key is used to check if the template object was changed, e.g. a blueprint compilation
 	TObjectKey<USceneComponent> TemplateKey;
@@ -86,4 +86,4 @@ private:
 #if WITH_EDITORONLY_DATA
 	struct FNiagaraRendererComponentsOnObjectsReplacedHelper* OnObjectsReplacedHandler = nullptr;
 #endif
-};
\ No newline at end of file
+};
diff --git a/Engine/Plugins/FX/Niagara/Source/Niagara/Public/NiagaraRendererDecals.h b/Engine/Plugins/FX/Niagara/Source/Niagara/Public/NiagaraRendererDecals.h
--- a/Engine/Plugins/FX/Niagara/Source/Niagara/Public/NiagaraRendererDecals.h
+++ b/Engine/Plugins/FX/Niagara/Source/Niagara/Public/NiagaraRendererDecals.h
@@ -7,18 +7,18 @@
 class FDeferredDecalProxy;
 class USceneComponent;
 
-class NIAGARA_API FNiagaraRendererDecals : public FNiagaraRenderer
+class FNiagaraRendererDecals : public FNiagaraRenderer
 {
 public:
-	explicit FNiagaraRendererDecals(ERHIFeatureLevel::Type FeatureLevel, const UNiagaraRendererProperties *InProps, const FNiagaraEmitterInstance* Emitter);
-	~FNiagaraRendererDecals();
+	NIAGARA_API explicit FNiagaraRendererDecals(ERHIFeatureLevel::Type FeatureLevel, const UNiagaraRendererProperties *InProps, const FNiagaraEmitterInstance* Emitter);
+	NIAGARA_API ~FNiagaraRendererDecals();
 
-	void ReleaseAllDecals() const;
+	NIAGARA_API void ReleaseAllDecals() const;
 
 	//FNiagaraRenderer interface
-	virtual void DestroyRenderState_Concurrent() override;
-	virtual FPrimitiveViewRelevance GetViewRelevance(const FSceneView* View, const FNiagaraSceneProxy *SceneProxy)const override;
-	virtual FNiagaraDynamicDataBase *GenerateDynamicData(const FNiagaraSceneProxy* Proxy, const UNiagaraRendererProperties* InProperties, const FNiagaraEmitterInstance* Emitter) const override;
+	NIAGARA_API virtual void DestroyRenderState_Concurrent() override;
+	NIAGARA_API virtual FPrimitiveViewRelevance GetViewRelevance(const FSceneView* View, const FNiagaraSceneProxy *SceneProxy)const override;
+	NIAGARA_API virtual FNiagaraDynamicDataBase *GenerateDynamicData(const FNiagaraSceneProxy* Proxy, const UNiagaraRendererProperties* InProperties, const FNiagaraEmitterInstance* Emitter) const override;
 	//FNiagaraRenderer interface END
 
 	mutable TWeakObjectPtr<USceneComponent>		WeakOwnerComponent;
diff --git a/Engine/Plugins/FX/Niagara/Source/Niagara/Public/NiagaraRendererLights.h b/Engine/Plugins/FX/Niagara/Source/Niagara/Public/NiagaraRendererLights.h
--- a/Engine/Plugins/FX/Niagara/Source/Niagara/Public/NiagaraRendererLights.h
+++ b/Engine/Plugins/FX/Niagara/Source/Niagara/Public/NiagaraRendererLights.h
@@ -10,7 +10,7 @@ NiagaraRendererLights.h: Renderer for rendering Niagara particles as Lights.
 /**
 * NiagaraRendererLights renders an FNiagaraEmitterInstance as simple lights
 */
-class NIAGARA_API FNiagaraRendererLights : public FNiagaraRenderer
+class FNiagaraRendererLights : public FNiagaraRenderer
 {
 public:
 	struct SimpleLightData
@@ -19,12 +19,12 @@ public:
 		FSimpleLightPerViewEntry PerViewEntry;
 	};
 
-	FNiagaraRendererLights(ERHIFeatureLevel::Type FeatureLevel, const UNiagaraRendererProperties *InProps, const FNiagaraEmitterInstance* Emitter);
+	NIAGARA_API FNiagaraRendererLights(ERHIFeatureLevel::Type FeatureLevel, const UNiagaraRendererProperties *InProps, const FNiagaraEmitterInstance* Emitter);
 
 	//FNiagaraRenderer interface
-	virtual FPrimitiveViewRelevance GetViewRelevance(const FSceneView* View, const FNiagaraSceneProxy *SceneProxy)const override;
-	virtual FNiagaraDynamicDataBase *GenerateDynamicData(const FNiagaraSceneProxy* Proxy, const UNiagaraRendererProperties* InProperties, const FNiagaraEmitterInstance* Emitter) const override;
-	virtual void GatherSimpleLights(FSimpleLightArray& OutParticleLights)const override;
+	NIAGARA_API virtual FPrimitiveViewRelevance GetViewRelevance(const FSceneView* View, const FNiagaraSceneProxy *SceneProxy)const override;
+	NIAGARA_API virtual FNiagaraDynamicDataBase *GenerateDynamicData(const FNiagaraSceneProxy* Proxy, const UNiagaraRendererProperties* InProperties, const FNiagaraEmitterInstance* Emitter) const override;
+	NIAGARA_API virtual void GatherSimpleLights(FSimpleLightArray& OutParticleLights)const override;
 	//FNiagaraRenderer interface END
 };
 
diff --git a/Engine/Plugins/FX/Niagara/Source/Niagara/Public/NiagaraRendererMeshes.h b/Engine/Plugins/FX/Niagara/Source/Niagara/Public/NiagaraRendererMeshes.h
--- a/Engine/Plugins/FX/Niagara/Source/Niagara/Public/NiagaraRendererMeshes.h
+++ b/Engine/Plugins/FX/Niagara/Source/Niagara/Public/NiagaraRendererMeshes.h
@@ -19,22 +19,22 @@ struct FNiagaraDynamicDataMesh;
 /**
 * NiagaraRendererSprites renders an FNiagaraEmitterInstance as sprite particles
 */
-class NIAGARA_API FNiagaraRendererMeshes : public FNiagaraRenderer
+class FNiagaraRendererMeshes : public FNiagaraRenderer
 {
 public:
-	FNiagaraRendererMeshes(ERHIFeatureLevel::Type FeatureLevel, const UNiagaraRendererProperties *InProps, const FNiagaraEmitterInstance* Emitter);
-	virtual ~FNiagaraRendererMeshes() override;
+	NIAGARA_API FNiagaraRendererMeshes(ERHIFeatureLevel::Type FeatureLevel, const UNiagaraRendererProperties *InProps, const FNiagaraEmitterInstance* Emitter);
+	NIAGARA_API virtual ~FNiagaraRendererMeshes() override;
 
 	//FNiagaraRenderer Interface
-	virtual void Initialize(const UNiagaraRendererProperties* InProps, const FNiagaraEmitterInstance* Emitter, const FNiagaraSystemInstanceController& InController) override;
-	virtual void ReleaseRenderThreadResources() override;
+	NIAGARA_API virtual void Initialize(const UNiagaraRendererProperties* InProps, const FNiagaraEmitterInstance* Emitter, const FNiagaraSystemInstanceController& InController) override;
+	NIAGARA_API virtual void ReleaseRenderThreadResources() override;
 
-	virtual void GetDynamicMeshElements(const TArray<const FSceneView*>& Views, const FSceneViewFamily& ViewFamily, uint32 VisibilityMap, FMeshElementCollector& Collector, const FNiagaraSceneProxy* SceneProxy) const override;
-	virtual FNiagaraDynamicDataBase* GenerateDynamicData(const FNiagaraSceneProxy* Proxy, const UNiagaraRendererProperties* InProperties, const FNiagaraEmitterInstance* Emitter) const override;
-	virtual int32 GetDynamicDataSize()const override;
-	virtual bool IsMaterialValid(const UMaterialInterface* Mat)const override;
+	NIAGARA_API virtual void GetDynamicMeshElements(const TArray<const FSceneView*>& Views, const FSceneViewFamily& ViewFamily, uint32 VisibilityMap, FMeshElementCollector& Collector, const FNiagaraSceneProxy* SceneProxy) const override;
+	NIAGARA_API virtual FNiagaraDynamicDataBase* GenerateDynamicData(const FNiagaraSceneProxy* Proxy, const UNiagaraRendererProperties* InProperties, const FNiagaraEmitterInstance* Emitter) const override;
+	NIAGARA_API virtual int32 GetDynamicDataSize()const override;
+	NIAGARA_API virtual bool IsMaterialValid(const UMaterialInterface* Mat)const override;
 #if RHI_RAYTRACING
-	virtual void GetDynamicRayTracingInstances(FRayTracingMaterialGatheringContext& Context, TArray<FRayTracingInstance>& OutRayTracingInstances, const FNiagaraSceneProxy* Proxy) final override;
+	NIAGARA_API virtual void GetDynamicRayTracingInstances(FRayTracingMaterialGatheringContext& Context, TArray<FRayTracingInstance>& OutRayTracingInstances, const FNiagaraSceneProxy* Proxy) final override;
 #endif
 	//FNiagaraRenderer Interface END
 
@@ -119,15 +119,15 @@ protected:
 		virtual ~FGPUSceneUpdateResource() {}
 	};
 
-	void PrepareParticleMeshRenderData(FParticleMeshRenderData& ParticleMeshRenderData, const FSceneViewFamily& ViewFamily, FMeshElementCollector& Collector, FNiagaraDynamicDataBase* InDynamicData, const FNiagaraSceneProxy* SceneProxy, bool bRayTracing, ENiagaraGpuComputeTickStage::Type GpuReadyTickStage) const;
-	void PrepareParticleRenderBuffers(FParticleMeshRenderData& ParticleMeshRenderData, FGlobalDynamicReadBuffer& DynamicReadBuffer) const;
-	void InitializeSortInfo(const FParticleMeshRenderData& ParticleMeshRenderData, const FNiagaraSceneProxy& SceneProxy, const FSceneView& View, int32 ViewIndex, bool bIsInstancedStereo, FNiagaraGPUSortInfo& OutSortInfo) const;
-	void PreparePerMeshData(FParticleMeshRenderData& ParticleMeshRenderData, const FNiagaraMeshVertexFactory& VertexFactory, const FNiagaraSceneProxy& SceneProxy, const FMeshData& MeshData) const;
-	uint32 PerformSortAndCull(FParticleMeshRenderData& ParticleMeshRenderData, FGlobalDynamicReadBuffer& ReadBuffer, FNiagaraGPUSortInfo& SortInfo, class FNiagaraGpuComputeDispatchInterface* ComputeDispatchInterface, int32 MeshIndex) const;
-	FNiagaraMeshCommonParameters CreateCommonShaderParams(const FParticleMeshRenderData& ParticleMeshRenderData, const FSceneView& View, const FMeshData& MeshData, const FNiagaraSceneProxy& SceneProxy) const;
-	FNiagaraMeshUniformBufferRef CreateVFUniformBuffer(const FParticleMeshRenderData& ParticleMeshRenderData, const FNiagaraMeshCommonParameters& CommonParams) const;
+	NIAGARA_API void PrepareParticleMeshRenderData(FParticleMeshRenderData& ParticleMeshRenderData, const FSceneViewFamily& ViewFamily, FMeshElementCollector& Collector, FNiagaraDynamicDataBase* InDynamicData, const FNiagaraSceneProxy* SceneProxy, bool bRayTracing, ENiagaraGpuComputeTickStage::Type GpuReadyTickStage) const;
+	NIAGARA_API void PrepareParticleRenderBuffers(FParticleMeshRenderData& ParticleMeshRenderData, FGlobalDynamicReadBuffer& DynamicReadBuffer) const;
+	NIAGARA_API void InitializeSortInfo(const FParticleMeshRenderData& ParticleMeshRenderData, const FNiagaraSceneProxy& SceneProxy, const FSceneView& View, int32 ViewIndex, bool bIsInstancedStereo, FNiagaraGPUSortInfo& OutSortInfo) const;
+	NIAGARA_API void PreparePerMeshData(FParticleMeshRenderData& ParticleMeshRenderData, const FNiagaraMeshVertexFactory& VertexFactory, const FNiagaraSceneProxy& SceneProxy, const FMeshData& MeshData) const;
+	NIAGARA_API uint32 PerformSortAndCull(FParticleMeshRenderData& ParticleMeshRenderData, FGlobalDynamicReadBuffer& ReadBuffer, FNiagaraGPUSortInfo& SortInfo, class FNiagaraGpuComputeDispatchInterface* ComputeDispatchInterface, int32 MeshIndex) const;
+	NIAGARA_API FNiagaraMeshCommonParameters CreateCommonShaderParams(const FParticleMeshRenderData& ParticleMeshRenderData, const FSceneView& View, const FMeshData& MeshData, const FNiagaraSceneProxy& SceneProxy) const;
+	NIAGARA_API FNiagaraMeshUniformBufferRef CreateVFUniformBuffer(const FParticleMeshRenderData& ParticleMeshRenderData, const FNiagaraMeshCommonParameters& CommonParams) const;
 
-	void SetupElementForGPUScene(
+	NIAGARA_API void SetupElementForGPUScene(
 		const FParticleMeshRenderData& ParticleMeshRenderData,
 		const FNiagaraMeshCommonParameters& CommonParameters,
 		const FNiagaraSceneProxy& SceneProxy,
@@ -138,7 +138,7 @@ protected:
 		FMeshBatchElement& OutMeshBatchElement
 	) const;
 
-	void CreateMeshBatchForSection(
+	NIAGARA_API void CreateMeshBatchForSection(
 		const FParticleMeshRenderData& ParticleMeshRenderData,
 		const FNiagaraMeshCommonParameters& CommonParams,
 		FMeshBatch& MeshBatch,
diff --git a/Engine/Plugins/FX/Niagara/Source/Niagara/Public/NiagaraRendererProperties.h b/Engine/Plugins/FX/Niagara/Source/Niagara/Public/NiagaraRendererProperties.h
--- a/Engine/Plugins/FX/Niagara/Source/Niagara/Public/NiagaraRendererProperties.h
+++ b/Engine/Plugins/FX/Niagara/Source/Niagara/Public/NiagaraRendererProperties.h
@@ -138,12 +138,12 @@ protected:
 };
 
 /** Used for building renderer layouts for vertex factories */
-struct NIAGARA_API FNiagaraRendererLayout
+struct FNiagaraRendererLayout
 {
-	void Initialize(int32 NumVariables);
-	bool SetVariable(const FNiagaraDataSetCompiledData* CompiledData, const FNiagaraVariableBase& Variable, int32 VFVarOffset);
-	bool SetVariableFromBinding(const FNiagaraDataSetCompiledData* CompiledData, const FNiagaraVariableAttributeBinding& VariableBinding, int32 VFVarOffset);
-	void Finalize();
+	NIAGARA_API void Initialize(int32 NumVariables);
+	NIAGARA_API bool SetVariable(const FNiagaraDataSetCompiledData* CompiledData, const FNiagaraVariableBase& Variable, int32 VFVarOffset);
+	NIAGARA_API bool SetVariableFromBinding(const FNiagaraDataSetCompiledData* CompiledData, const FNiagaraVariableAttributeBinding& VariableBinding, int32 VFVarOffset);
+	NIAGARA_API void Finalize();
 
 	TConstArrayView<FNiagaraRendererVariableInfo> GetVFVariables_GameThread() const { check(IsInGameThread() || IsInParallelGameThread()); return MakeArrayView(VFVariables_GT); }
 	TConstArrayView<FNiagaraRendererVariableInfo> GetVFVariables_RenderThread() const { check(IsInRenderingThread()); return MakeArrayView(VFVariables_RT); }
@@ -276,8 +276,8 @@ struct FNiagaraRendererMaterialParameters
 * of those specific properties is stored on UNiagaraEmitter (on the System) for serialization
 * and handed back to the System renderer on load.
 */
-UCLASS(ABSTRACT)
-class NIAGARA_API UNiagaraRendererProperties : public UNiagaraMergeable
+UCLASS(ABSTRACT, MinimalAPI)
+class UNiagaraRendererProperties : public UNiagaraMergeable
 {
 	GENERATED_BODY()
 	
@@ -293,17 +293,17 @@ public:
 	}
 
 	//UObject Interface Begin
-	virtual void PostInitProperties() override;
-	virtual void PostLoad() override;
+	NIAGARA_API virtual void PostInitProperties() override;
+	NIAGARA_API virtual void PostLoad() override;
 #if WITH_EDITORONLY_DATA
-	virtual void PostEditChangeProperty(struct FPropertyChangedEvent& PropertyChangedEvent) override;
+	NIAGARA_API virtual void PostEditChangeProperty(struct FPropertyChangedEvent& PropertyChangedEvent) override;
 #endif
 
 	//UObject Interface End
 	
-	virtual FNiagaraRenderer* CreateEmitterRenderer(ERHIFeatureLevel::Type FeatureLevel, const FNiagaraEmitterInstance* Emitter, const FNiagaraSystemInstanceController& InController) PURE_VIRTUAL ( UNiagaraRendererProperties::CreateEmitterRenderer, return nullptr;);
-	virtual class FNiagaraBoundsCalculator* CreateBoundsCalculator() PURE_VIRTUAL(UNiagaraRendererProperties::CreateBoundsCalculator, return nullptr;);
-	virtual void GetUsedMaterials(const FNiagaraEmitterInstance* InEmitter, TArray<UMaterialInterface*>& OutMaterials) const PURE_VIRTUAL(UNiagaraRendererProperties::GetUsedMaterials,);
+	NIAGARA_API virtual FNiagaraRenderer* CreateEmitterRenderer(ERHIFeatureLevel::Type FeatureLevel, const FNiagaraEmitterInstance* Emitter, const FNiagaraSystemInstanceController& InController) PURE_VIRTUAL ( UNiagaraRendererProperties::CreateEmitterRenderer, return nullptr;);
+	NIAGARA_API virtual class FNiagaraBoundsCalculator* CreateBoundsCalculator() PURE_VIRTUAL(UNiagaraRendererProperties::CreateBoundsCalculator, return nullptr;);
+	NIAGARA_API virtual void GetUsedMaterials(const FNiagaraEmitterInstance* InEmitter, TArray<UMaterialInterface*>& OutMaterials) const PURE_VIRTUAL(UNiagaraRendererProperties::GetUsedMaterials,);
 	virtual const FVertexFactoryType* GetVertexFactoryType() const { return nullptr; }
 	virtual bool IsBackfaceCullingDisabled() const { return false; }
 
@@ -312,15 +312,15 @@ public:
 	virtual bool IsSimTargetSupported(ENiagaraSimTarget InSimTarget) const { return false; };
 
 	const TArray<const FNiagaraVariableAttributeBinding*>& GetAttributeBindings() const { return AttributeBindings; }
-	uint32 ComputeMaxUsedComponents(const FNiagaraDataSetCompiledData* CompiledDataSetData) const;
+	NIAGARA_API uint32 ComputeMaxUsedComponents(const FNiagaraDataSetCompiledData* CompiledDataSetData) const;
 
-	virtual bool NeedsLoadForTargetPlatform(const class ITargetPlatform* TargetPlatform) const override;
+	NIAGARA_API virtual bool NeedsLoadForTargetPlatform(const class ITargetPlatform* TargetPlatform) const override;
 
 	/** Method to add asset tags that are specific to this renderer. By default we add in how many instances of this class exist in the list.*/
-	virtual void GetAssetTagsForContext(const UObject* InAsset, FGuid AssetVersion, const TArray<const UNiagaraRendererProperties*>& InProperties, TMap<FName, uint32>& NumericKeys, TMap<FName, FString>& StringKeys) const;
+	NIAGARA_API virtual void GetAssetTagsForContext(const UObject* InAsset, FGuid AssetVersion, const TArray<const UNiagaraRendererProperties*>& InProperties, TMap<FName, uint32>& NumericKeys, TMap<FName, FString>& StringKeys) const;
 
 	/** In the case that we need parameters bound in that aren't Particle variables, these should be set up here so that the data is appropriately populated after the simulation.*/
-	virtual bool PopulateRequiredBindings(FNiagaraParameterStore& InParameterStore);
+	NIAGARA_API virtual bool PopulateRequiredBindings(FNiagaraParameterStore& InParameterStore);
 
 	
 	/**
@@ -336,17 +336,17 @@ public:
 
 #if WITH_EDITORONLY_DATA
 
-	virtual bool IsSupportedVariableForBinding(const FNiagaraVariableBase& InSourceForBinding, const FName& InTargetBindingName) const;
+	NIAGARA_API virtual bool IsSupportedVariableForBinding(const FNiagaraVariableBase& InSourceForBinding, const FName& InTargetBindingName) const;
 
 	/** Internal handling of any emitter variable renames. Note that this doesn't modify the renderer, the caller will need to do that if it is desired.*/
-	virtual void RenameEmitter(const FName& InOldName, const UNiagaraEmitter* InRenamedEmitter);
-	virtual void RenameVariable(const FNiagaraVariableBase& OldVariable, const FNiagaraVariableBase& NewVariable, const FVersionedNiagaraEmitter& InEmitter);
-	virtual void RemoveVariable(const FNiagaraVariableBase& OldVariable, const FVersionedNiagaraEmitter& InEmitter);
+	NIAGARA_API virtual void RenameEmitter(const FName& InOldName, const UNiagaraEmitter* InRenamedEmitter);
+	NIAGARA_API virtual void RenameVariable(const FNiagaraVariableBase& OldVariable, const FNiagaraVariableBase& NewVariable, const FVersionedNiagaraEmitter& InEmitter);
+	NIAGARA_API virtual void RemoveVariable(const FNiagaraVariableBase& OldVariable, const FVersionedNiagaraEmitter& InEmitter);
 	virtual bool IsMaterialValidForRenderer(UMaterial* Material, FText& InvalidMessage) { return true; }
 
 	virtual void FixMaterial(UMaterial* Material) { }
 
-	virtual TArray<FNiagaraVariable> GetBoundAttributes() const;
+	NIAGARA_API virtual TArray<FNiagaraVariable> GetBoundAttributes() const;
 
 	virtual const TArray<FNiagaraVariable>& GetRequiredAttributes() { static TArray<FNiagaraVariable> Vars; return Vars; };
 	virtual const TArray<FNiagaraVariable>& GetOptionalAttributes() { static TArray<FNiagaraVariable> Vars; return Vars; };
@@ -358,26 +358,26 @@ public:
 		return CastChecked<UNiagaraRendererProperties>(Super::StaticDuplicateWithNewMergeIdInternal(InOuter));
 	}
 
-	virtual void GetRendererWidgets(const FNiagaraEmitterInstance* InEmitter, TArray<TSharedPtr<SWidget>>& OutWidgets, TSharedPtr<FAssetThumbnailPool> InThumbnailPool) const PURE_VIRTUAL(UNiagaraRendererProperties::GetRendererWidgets, );
-	virtual void GetRendererTooltipWidgets(const FNiagaraEmitterInstance* InEmitter, TArray<TSharedPtr<SWidget>>& OutWidgets, TSharedPtr<FAssetThumbnailPool> InThumbnailPool) const PURE_VIRTUAL(UNiagaraRendererProperties::GetRendererTooltipWidgets, );
+	NIAGARA_API virtual void GetRendererWidgets(const FNiagaraEmitterInstance* InEmitter, TArray<TSharedPtr<SWidget>>& OutWidgets, TSharedPtr<FAssetThumbnailPool> InThumbnailPool) const PURE_VIRTUAL(UNiagaraRendererProperties::GetRendererWidgets, );
+	NIAGARA_API virtual void GetRendererTooltipWidgets(const FNiagaraEmitterInstance* InEmitter, TArray<TSharedPtr<SWidget>>& OutWidgets, TSharedPtr<FAssetThumbnailPool> InThumbnailPool) const PURE_VIRTUAL(UNiagaraRendererProperties::GetRendererTooltipWidgets, );
 	virtual void GetRendererFeedback(const FVersionedNiagaraEmitter& InEmitter, TArray<FText>& OutErrors, TArray<FText>& OutWarnings, TArray<FText>& OutInfo) const {};
-	virtual void GetRendererFeedback(const FVersionedNiagaraEmitter& InEmitter, TArray<FNiagaraRendererFeedback>& OutErrors, TArray<FNiagaraRendererFeedback>& OutWarnings, TArray<FNiagaraRendererFeedback>& OutInfo) const;
+	NIAGARA_API virtual void GetRendererFeedback(const FVersionedNiagaraEmitter& InEmitter, TArray<FNiagaraRendererFeedback>& OutErrors, TArray<FNiagaraRendererFeedback>& OutWarnings, TArray<FNiagaraRendererFeedback>& OutInfo) const;
 
 	// The icon to display in the niagara stack widget under the renderer section
-	virtual const FSlateBrush* GetStackIcon() const;
+	NIAGARA_API virtual const FSlateBrush* GetStackIcon() const;
 
 	// The text to display in the niagara stack widget under the renderer section
-	virtual FText GetWidgetDisplayName() const;
+	NIAGARA_API virtual FText GetWidgetDisplayName() const;
 #endif // WITH_EDITORONLY_DATA
 
 #if WITH_EDITOR
-	FOnPropertiesChanged& OnPropertiesChanged();
+	NIAGARA_API FOnPropertiesChanged& OnPropertiesChanged();
 #endif
 	virtual ENiagaraRendererSourceDataMode GetCurrentSourceMode() const {	return ENiagaraRendererSourceDataMode::Particles;}
 
-	virtual bool GetIsActive() const;
+	NIAGARA_API virtual bool GetIsActive() const;
 	virtual bool GetIsEnabled() const { return bIsEnabled; }
-	virtual void SetIsEnabled(bool bInIsEnabled);
+	NIAGARA_API virtual void SetIsEnabled(bool bInIsEnabled);
 
 	virtual void CacheFromCompiledData(const FNiagaraDataSetCompiledData* CompiledData) {}
 
@@ -388,21 +388,21 @@ public:
 	/** When true, the renderer will be opted in to get its OnSystemComplete_GameThread overload called */
 	virtual bool NeedsSystemCompletion() const { return false; }
 
-	bool NeedsPreciseMotionVectors() const;
+	NIAGARA_API bool NeedsPreciseMotionVectors() const;
 	virtual bool UseHeterogeneousVolumes() const { return false; }
 
-	static bool IsSortHighPrecision(ENiagaraRendererSortPrecision SortPrecision);
+	static NIAGARA_API bool IsSortHighPrecision(ENiagaraRendererSortPrecision SortPrecision);
 
 	/** Should the Gpu translucent data be this frame or not? */
-	static bool ShouldGpuTranslucentThisFrame(ENiagaraRendererGpuTranslucentLatency Latency);
+	static NIAGARA_API bool ShouldGpuTranslucentThisFrame(ENiagaraRendererGpuTranslucentLatency Latency);
 	/** Is the Gpu translucent data going to be this frame, this can be restricted by things like feature level. */
-	static bool IsGpuTranslucentThisFrame(ERHIFeatureLevel::Type FeatureLevel, ENiagaraRendererGpuTranslucentLatency Latency);
+	static NIAGARA_API bool IsGpuTranslucentThisFrame(ERHIFeatureLevel::Type FeatureLevel, ENiagaraRendererGpuTranslucentLatency Latency);
 
 	template<typename TAction>
 	void ForEachPlatformSet(TAction Func);
 
-	FVersionedNiagaraEmitterData* GetEmitterData() const;
-	FVersionedNiagaraEmitter GetOuterEmitter() const;
+	NIAGARA_API FVersionedNiagaraEmitterData* GetEmitterData() const;
+	NIAGARA_API FVersionedNiagaraEmitter GetOuterEmitter() const;
 
 	/** Platforms on which this renderer is enabled. */
 	UPROPERTY(EditAnywhere, Category = "Scalability", meta=(DisplayInScalabilityContext))
@@ -442,24 +442,24 @@ protected:
 
 	TArray<const FNiagaraVariableAttributeBinding*> AttributeBindings;
 
-	virtual void PostLoadBindings(ENiagaraRendererSourceDataMode InSourceMode);
-	virtual void UpdateSourceModeDerivates(ENiagaraRendererSourceDataMode InSourceMode, bool bFromPropertyEdit = false);
+	NIAGARA_API virtual void PostLoadBindings(ENiagaraRendererSourceDataMode InSourceMode);
+	NIAGARA_API virtual void UpdateSourceModeDerivates(ENiagaraRendererSourceDataMode InSourceMode, bool bFromPropertyEdit = false);
 
 #if WITH_EDITORONLY_DATA
 	/** returns the variable associated with the supplied binding if it should be bound given the current settings of the RendererProperties. */
-	virtual FNiagaraVariable GetBoundAttribute(const FNiagaraVariableAttributeBinding* Binding) const;
+	NIAGARA_API virtual FNiagaraVariable GetBoundAttribute(const FNiagaraVariableAttributeBinding* Binding) const;
 
 	/** utility function that can be used to fix up old vec3 bindings into position bindings. */
-	static void ChangeToPositionBinding(FNiagaraVariableAttributeBinding& Binding);
+	static NIAGARA_API void ChangeToPositionBinding(FNiagaraVariableAttributeBinding& Binding);
 
 	/** Update MIC Static Parameters. */
-	bool UpdateMaterialStaticParameters(const FNiagaraRendererMaterialParameters& MaterialParameters, UMaterialInstanceConstant* MIC);
+	NIAGARA_API bool UpdateMaterialStaticParameters(const FNiagaraRendererMaterialParameters& MaterialParameters, UMaterialInstanceConstant* MIC);
 
 	/** Utility function to updates MICs. */
-	void UpdateMaterialParametersMIC(const FNiagaraRendererMaterialParameters& MaterialParameters, TObjectPtr<UMaterialInterface>& InOutMaterial, TObjectPtr<UMaterialInstanceConstant>& InOutMIC);
-	void UpdateMaterialParametersMIC(const FNiagaraRendererMaterialParameters& MaterialParameters, TArrayView<UMaterialInterface*> Materials, TArray<TObjectPtr<UMaterialInstanceConstant>>& InOutMICs);
+	NIAGARA_API void UpdateMaterialParametersMIC(const FNiagaraRendererMaterialParameters& MaterialParameters, TObjectPtr<UMaterialInterface>& InOutMaterial, TObjectPtr<UMaterialInstanceConstant>& InOutMIC);
+	NIAGARA_API void UpdateMaterialParametersMIC(const FNiagaraRendererMaterialParameters& MaterialParameters, TArrayView<UMaterialInterface*> Materials, TArray<TObjectPtr<UMaterialInstanceConstant>>& InOutMICs);
 
-	int32 GetDynamicParameterChannelMask(const FVersionedNiagaraEmitterData* EmitterData, FName BindingName, int32 DefaultChannelMask) const;
+	NIAGARA_API int32 GetDynamicParameterChannelMask(const FVersionedNiagaraEmitterData* EmitterData, FName BindingName, int32 DefaultChannelMask) const;
 #endif
 
 #if WITH_EDITOR
diff --git a/Engine/Plugins/FX/Niagara/Source/Niagara/Public/NiagaraRendererRibbons.h b/Engine/Plugins/FX/Niagara/Source/Niagara/Public/NiagaraRendererRibbons.h
--- a/Engine/Plugins/FX/Niagara/Source/Niagara/Public/NiagaraRendererRibbons.h
+++ b/Engine/Plugins/FX/Niagara/Source/Niagara/Public/NiagaraRendererRibbons.h
@@ -297,33 +297,33 @@ struct FNiagaraRibbonShapeGeometryData
 * NiagaraRendererRibbons renders an FNiagaraEmitterInstance as a ribbon connecting all particles
 * in order by particle age.
 */
-class NIAGARA_API FNiagaraRendererRibbons : public FNiagaraRenderer
+class FNiagaraRendererRibbons : public FNiagaraRenderer
 {
 public:
-	FNiagaraRendererRibbons(ERHIFeatureLevel::Type FeatureLevel, const UNiagaraRendererProperties *InProps, const FNiagaraEmitterInstance* Emitter);	// FNiagaraRenderer Interface 
-	~FNiagaraRendererRibbons();
+	NIAGARA_API FNiagaraRendererRibbons(ERHIFeatureLevel::Type FeatureLevel, const UNiagaraRendererProperties *InProps, const FNiagaraEmitterInstance* Emitter);	// FNiagaraRenderer Interface 
+	NIAGARA_API ~FNiagaraRendererRibbons();
 
 	// FNiagaraRenderer Interface 
-	virtual void CreateRenderThreadResources() override;
-	virtual void ReleaseRenderThreadResources() override;
+	NIAGARA_API virtual void CreateRenderThreadResources() override;
+	NIAGARA_API virtual void ReleaseRenderThreadResources() override;
 
-	virtual void GetDynamicMeshElements(const TArray<const FSceneView*>& Views, const FSceneViewFamily& ViewFamily, uint32 VisibilityMap, FMeshElementCollector& Collector, const FNiagaraSceneProxy *SceneProxy) const override;
-	virtual FNiagaraDynamicDataBase *GenerateDynamicData(const FNiagaraSceneProxy* Proxy, const UNiagaraRendererProperties* InProperties, const FNiagaraEmitterInstance* Emitteride) const override;
-	virtual int32 GetDynamicDataSize()const override;
-	virtual bool IsMaterialValid(const UMaterialInterface* Mat)const override;
+	NIAGARA_API virtual void GetDynamicMeshElements(const TArray<const FSceneView*>& Views, const FSceneViewFamily& ViewFamily, uint32 VisibilityMap, FMeshElementCollector& Collector, const FNiagaraSceneProxy *SceneProxy) const override;
+	NIAGARA_API virtual FNiagaraDynamicDataBase *GenerateDynamicData(const FNiagaraSceneProxy* Proxy, const UNiagaraRendererProperties* InProperties, const FNiagaraEmitterInstance* Emitteride) const override;
+	NIAGARA_API virtual int32 GetDynamicDataSize()const override;
+	NIAGARA_API virtual bool IsMaterialValid(const UMaterialInterface* Mat)const override;
 #if RHI_RAYTRACING
-	virtual void GetDynamicRayTracingInstances(FRayTracingMaterialGatheringContext& Context, TArray<FRayTracingInstance>& OutRayTracingInstances, const FNiagaraSceneProxy* Proxy) final override;
+	NIAGARA_API virtual void GetDynamicRayTracingInstances(FRayTracingMaterialGatheringContext& Context, TArray<FRayTracingInstance>& OutRayTracingInstances, const FNiagaraSceneProxy* Proxy) final override;
 #endif
 
 protected:
 	
-	static void GenerateShapeStateMultiPlane(FNiagaraRibbonShapeGeometryData& State, int32 MultiPlaneCount, int32 WidthSegmentationCount, bool bEnableAccurateGeometry);
-	static void GenerateShapeStateTube(FNiagaraRibbonShapeGeometryData& State, int32 TubeSubdivisions);
-	static void GenerateShapeStateCustom(FNiagaraRibbonShapeGeometryData& State, const TArray<FNiagaraRibbonShapeCustomVertex>& CustomVertices);
-	static void GenerateShapeStatePlane(FNiagaraRibbonShapeGeometryData& State, int32 WidthSegmentationCount);	
-	void InitializeShape(const UNiagaraRibbonRendererProperties* Properties);
+	static NIAGARA_API void GenerateShapeStateMultiPlane(FNiagaraRibbonShapeGeometryData& State, int32 MultiPlaneCount, int32 WidthSegmentationCount, bool bEnableAccurateGeometry);
+	static NIAGARA_API void GenerateShapeStateTube(FNiagaraRibbonShapeGeometryData& State, int32 TubeSubdivisions);
+	static NIAGARA_API void GenerateShapeStateCustom(FNiagaraRibbonShapeGeometryData& State, const TArray<FNiagaraRibbonShapeCustomVertex>& CustomVertices);
+	static NIAGARA_API void GenerateShapeStatePlane(FNiagaraRibbonShapeGeometryData& State, int32 WidthSegmentationCount);	
+	NIAGARA_API void InitializeShape(const UNiagaraRibbonRendererProperties* Properties);
 	
-	void InitializeTessellation(const UNiagaraRibbonRendererProperties* Properties);
+	NIAGARA_API void InitializeTessellation(const UNiagaraRibbonRendererProperties* Properties);
 	
 	
 	template<typename IntType>
@@ -339,14 +339,14 @@ protected:
 	template<typename IDType, typename ReaderType>
 	void GenerateVertexBufferForMultiRibbon(const FNiagaraGenerationInputDataCPUAccessors& CPUData, const ReaderType& IDReader, FNiagaraRibbonCPUGeneratedVertexData& OutputData) const;
 	
-	void GenerateVertexBufferCPU(const FNiagaraGenerationInputDataCPUAccessors& CPUData, FNiagaraRibbonCPUGeneratedVertexData& OutputData) const;
+	NIAGARA_API void GenerateVertexBufferCPU(const FNiagaraGenerationInputDataCPUAccessors& CPUData, FNiagaraRibbonCPUGeneratedVertexData& OutputData) const;
 
 
-	int32 CalculateTessellationFactor(const FNiagaraSceneProxy* SceneProxy, const FSceneView* View, const FVector& ViewOriginForDistanceCulling) const;
-	FNiagaraIndexGenerationInput CalculateIndexBufferConfiguration(const TSharedPtr<FNiagaraRibbonCPUGeneratedVertexData>& GeneratedVertices, const FNiagaraDataBuffer* SourceParticleData,
+	NIAGARA_API int32 CalculateTessellationFactor(const FNiagaraSceneProxy* SceneProxy, const FSceneView* View, const FVector& ViewOriginForDistanceCulling) const;
+	NIAGARA_API FNiagaraIndexGenerationInput CalculateIndexBufferConfiguration(const TSharedPtr<FNiagaraRibbonCPUGeneratedVertexData>& GeneratedVertices, const FNiagaraDataBuffer* SourceParticleData,
 		const FNiagaraSceneProxy* SceneProxy, const FSceneView* View, const FVector& ViewOriginForDistanceCulling, bool bShouldUseGPUInitIndices, bool bIsGPUSim) const;
 	
-	void GenerateIndexBufferForView(
+	NIAGARA_API void GenerateIndexBufferForView(
 		FNiagaraGpuRibbonsDataManager& GpuRibbonsDataManager, FMeshElementCollector& Collector,
 		FNiagaraIndexGenerationInput& GeneratedData, FNiagaraDynamicDataRibbon* DynamicDataRibbon,
 		const TSharedPtr<FNiagaraRibbonRenderingFrameViewResources>& RenderingViewResources, const FSceneView* View, const FVector& ViewOriginForDistanceCulling
@@ -361,24 +361,24 @@ protected:
 	template <typename TValue>
 	static TValue* AppendToIndexBufferCPU(TValue* OutIndices, const FNiagaraIndexGenerationInput& GeneratedData, const FNiagaraRibbonShapeGeometryData& ShapeState, const TArrayView<uint32>& SegmentData, bool bInvertOrder);
 	
-	void SetupPerViewUniformBuffer(FNiagaraIndexGenerationInput& GeneratedData,
+	NIAGARA_API void SetupPerViewUniformBuffer(FNiagaraIndexGenerationInput& GeneratedData,
 	                               const FSceneView* View, const FSceneViewFamily& ViewFamily, const FNiagaraSceneProxy* SceneProxy, FNiagaraRibbonUniformBufferRef& OutUniformBuffer) const;
 
-	void SetupMeshBatchAndCollectorResourceForView(FRHICommandListBase& RHICmdList, const FNiagaraIndexGenerationInput& GeneratedData, FNiagaraDynamicDataRibbon* DynamicDataRibbon,
+	NIAGARA_API void SetupMeshBatchAndCollectorResourceForView(FRHICommandListBase& RHICmdList, const FNiagaraIndexGenerationInput& GeneratedData, FNiagaraDynamicDataRibbon* DynamicDataRibbon,
 	                                               const FNiagaraDataBuffer* SourceParticleData, const FSceneView* View, const FSceneViewFamily& ViewFamily, const FNiagaraSceneProxy* SceneProxy,
 	                                               const TSharedPtr<FNiagaraRibbonRenderingFrameResources>& RenderingResources, const TSharedPtr<FNiagaraRibbonRenderingFrameViewResources>& RenderingViewResources, FMeshBatch& OutMeshBatch, bool bShouldUseGPUInitIndices) const;
 
 
-	void InitializeViewIndexBuffersGPU(FRHICommandListImmediate& RHICmdList, FNiagaraGpuComputeDispatchInterface* ComputeDispatchInterface, const FNiagaraRibbonGPUInitParameters& GpuInitParameters,
+	NIAGARA_API void InitializeViewIndexBuffersGPU(FRHICommandListImmediate& RHICmdList, FNiagaraGpuComputeDispatchInterface* ComputeDispatchInterface, const FNiagaraRibbonGPUInitParameters& GpuInitParameters,
 		const TSharedPtr<FNiagaraRibbonRenderingFrameViewResources>& RenderingViewResources) const;
 
-	void InitializeVertexBuffersResources(const FNiagaraDynamicDataRibbon* DynamicDataRibbon, FNiagaraDataBuffer* SourceParticleData,
+	NIAGARA_API void InitializeVertexBuffersResources(const FNiagaraDynamicDataRibbon* DynamicDataRibbon, FNiagaraDataBuffer* SourceParticleData,
 	                                      FGlobalDynamicReadBuffer& DynamicReadBuffer, const TSharedPtr<FNiagaraRibbonRenderingFrameResources>& RenderingResources, bool bShouldUseGPUInit) const;
 	
-	void InitializeVertexBuffersGPU(FRHICommandListImmediate& RHICmdList, FNiagaraGpuComputeDispatchInterface* ComputeDispatchInterface, const FNiagaraRibbonGPUInitParameters& GpuInitParameters,
+	NIAGARA_API void InitializeVertexBuffersGPU(FRHICommandListImmediate& RHICmdList, FNiagaraGpuComputeDispatchInterface* ComputeDispatchInterface, const FNiagaraRibbonGPUInitParameters& GpuInitParameters,
 		FNiagaraRibbonGPUInitComputeBuffers& TempBuffers, const TSharedPtr<FNiagaraRibbonRenderingFrameResources>& RenderingResources) const;
 
-	FRibbonComputeUniformParameters SetupComputeVertexGenParams(FNiagaraGpuComputeDispatchInterface* ComputeDispatchInterface, const TSharedPtr<FNiagaraRibbonRenderingFrameResources>& RenderingResources, const FNiagaraRibbonGPUInitParameters& GpuInitParameters) const;
+	NIAGARA_API FRibbonComputeUniformParameters SetupComputeVertexGenParams(FNiagaraGpuComputeDispatchInterface* ComputeDispatchInterface, const TSharedPtr<FNiagaraRibbonRenderingFrameResources>& RenderingResources, const FNiagaraRibbonGPUInitParameters& GpuInitParameters) const;
 
 	FNiagaraRibbonGenerationConfig GenerationConfig;
 	
diff --git a/Engine/Plugins/FX/Niagara/Source/Niagara/Public/NiagaraRendererSprites.h b/Engine/Plugins/FX/Niagara/Source/Niagara/Public/NiagaraRendererSprites.h
--- a/Engine/Plugins/FX/Niagara/Source/Niagara/Public/NiagaraRendererSprites.h
+++ b/Engine/Plugins/FX/Niagara/Source/Niagara/Public/NiagaraRendererSprites.h
@@ -16,23 +16,23 @@ struct FNiagaraDynamicDataSprites;
 /**
 * FNiagaraRendererSprites renders an FNiagaraEmitterInstance as sprite particles
 */
-class NIAGARA_API FNiagaraRendererSprites : public FNiagaraRenderer
+class FNiagaraRendererSprites : public FNiagaraRenderer
 {
 public:
-	FNiagaraRendererSprites(ERHIFeatureLevel::Type FeatureLevel, const UNiagaraRendererProperties *InProps, const FNiagaraEmitterInstance* Emitter);
-	~FNiagaraRendererSprites();
+	NIAGARA_API FNiagaraRendererSprites(ERHIFeatureLevel::Type FeatureLevel, const UNiagaraRendererProperties *InProps, const FNiagaraEmitterInstance* Emitter);
+	NIAGARA_API ~FNiagaraRendererSprites();
 
 	//FNiagaraRenderer interface
-	virtual void CreateRenderThreadResources() override;
-	virtual void ReleaseRenderThreadResources() override;
+	NIAGARA_API virtual void CreateRenderThreadResources() override;
+	NIAGARA_API virtual void ReleaseRenderThreadResources() override;
 
-	virtual void GetDynamicMeshElements(const TArray<const FSceneView*>& Views, const FSceneViewFamily& ViewFamily, uint32 VisibilityMap, FMeshElementCollector& Collector, const FNiagaraSceneProxy *SceneProxy) const override;
-	virtual FNiagaraDynamicDataBase* GenerateDynamicData(const FNiagaraSceneProxy* Proxy, const UNiagaraRendererProperties* InProperties, const FNiagaraEmitterInstance* Emitter) const override;
-	virtual int GetDynamicDataSize()const override;
-	virtual bool IsMaterialValid(const UMaterialInterface* Mat)const override;
+	NIAGARA_API virtual void GetDynamicMeshElements(const TArray<const FSceneView*>& Views, const FSceneViewFamily& ViewFamily, uint32 VisibilityMap, FMeshElementCollector& Collector, const FNiagaraSceneProxy *SceneProxy) const override;
+	NIAGARA_API virtual FNiagaraDynamicDataBase* GenerateDynamicData(const FNiagaraSceneProxy* Proxy, const UNiagaraRendererProperties* InProperties, const FNiagaraEmitterInstance* Emitter) const override;
+	NIAGARA_API virtual int GetDynamicDataSize()const override;
+	NIAGARA_API virtual bool IsMaterialValid(const UMaterialInterface* Mat)const override;
 
 #if RHI_RAYTRACING
-		virtual void GetDynamicRayTracingInstances(FRayTracingMaterialGatheringContext& Context, TArray<FRayTracingInstance>& OutRayTracingInstances, const FNiagaraSceneProxy* Proxy) final override;
+		NIAGARA_API virtual void GetDynamicRayTracingInstances(FRayTracingMaterialGatheringContext& Context, TArray<FRayTracingInstance>& OutRayTracingInstances, const FNiagaraSceneProxy* Proxy) final override;
 #endif
 	//FNiagaraRenderer interface END
 
diff --git a/Engine/Plugins/FX/Niagara/Source/Niagara/Public/NiagaraRendererVolumes.h b/Engine/Plugins/FX/Niagara/Source/Niagara/Public/NiagaraRendererVolumes.h
--- a/Engine/Plugins/FX/Niagara/Source/Niagara/Public/NiagaraRendererVolumes.h
+++ b/Engine/Plugins/FX/Niagara/Source/Niagara/Public/NiagaraRendererVolumes.h
@@ -9,24 +9,24 @@
 
 class USceneComponent;
 
-class NIAGARA_API FNiagaraRendererVolumes : public FNiagaraRenderer
+class FNiagaraRendererVolumes : public FNiagaraRenderer
 {
 public:
-	explicit FNiagaraRendererVolumes(ERHIFeatureLevel::Type FeatureLevel, const UNiagaraRendererProperties *InProps, const FNiagaraEmitterInstance* Emitter);
-	~FNiagaraRendererVolumes();
+	NIAGARA_API explicit FNiagaraRendererVolumes(ERHIFeatureLevel::Type FeatureLevel, const UNiagaraRendererProperties *InProps, const FNiagaraEmitterInstance* Emitter);
+	NIAGARA_API ~FNiagaraRendererVolumes();
 
 	//FNiagaraRenderer interface
-	virtual void CreateRenderThreadResources() override;
-	virtual void ReleaseRenderThreadResources() override;
+	NIAGARA_API virtual void CreateRenderThreadResources() override;
+	NIAGARA_API virtual void ReleaseRenderThreadResources() override;
 
-	virtual bool IsMaterialValid(const UMaterialInterface* Material) const override;
+	NIAGARA_API virtual bool IsMaterialValid(const UMaterialInterface* Material) const override;
 
-	virtual FNiagaraDynamicDataBase *GenerateDynamicData(const FNiagaraSceneProxy* Proxy, const UNiagaraRendererProperties* InProperties, const FNiagaraEmitterInstance* Emitter) const override;
-	virtual int GetDynamicDataSize() const override;
+	NIAGARA_API virtual FNiagaraDynamicDataBase *GenerateDynamicData(const FNiagaraSceneProxy* Proxy, const UNiagaraRendererProperties* InProperties, const FNiagaraEmitterInstance* Emitter) const override;
+	NIAGARA_API virtual int GetDynamicDataSize() const override;
 
-	void GetDynamicMeshElements(const TArray<const FSceneView*>& Views, const FSceneViewFamily& ViewFamily, uint32 VisibilityMap, FMeshElementCollector& Collector, const FNiagaraSceneProxy* SceneProxy) const;
+	NIAGARA_API void GetDynamicMeshElements(const TArray<const FSceneView*>& Views, const FSceneViewFamily& ViewFamily, uint32 VisibilityMap, FMeshElementCollector& Collector, const FNiagaraSceneProxy* SceneProxy) const;
 #if RHI_RAYTRACING
-	void GetDynamicRayTracingInstances(FRayTracingMaterialGatheringContext& Context, TArray<FRayTracingInstance>& OutRayTracingInstances, const FNiagaraSceneProxy* SceneProxy);
+	NIAGARA_API void GetDynamicRayTracingInstances(FRayTracingMaterialGatheringContext& Context, TArray<FRayTracingInstance>& OutRayTracingInstances, const FNiagaraSceneProxy* SceneProxy);
 #endif //RHI_RAYTRACING
 	//FNiagaraRenderer interface END
 
diff --git a/Engine/Plugins/FX/Niagara/Source/Niagara/Public/NiagaraRibbonRendererProperties.h b/Engine/Plugins/FX/Niagara/Source/Niagara/Public/NiagaraRibbonRendererProperties.h
--- a/Engine/Plugins/FX/Niagara/Source/Niagara/Public/NiagaraRibbonRendererProperties.h
+++ b/Engine/Plugins/FX/Niagara/Source/Niagara/Public/NiagaraRibbonRendererProperties.h
@@ -189,52 +189,52 @@ namespace ENiagaraRibbonVFLayout
 	};
 };
 
-UCLASS(editinlinenew, meta = (DisplayName = "Ribbon Renderer"))
-class NIAGARA_API UNiagaraRibbonRendererProperties : public UNiagaraRendererProperties
+UCLASS(editinlinenew, meta = (DisplayName = "Ribbon Renderer"), MinimalAPI)
+class UNiagaraRibbonRendererProperties : public UNiagaraRendererProperties
 {
 public:
 	GENERATED_BODY()
 
-	UNiagaraRibbonRendererProperties();
+	NIAGARA_API UNiagaraRibbonRendererProperties();
 
 	//UObject Interface
-	virtual void PostLoad() override;
-	virtual void PostInitProperties() override;
-	virtual void Serialize(FStructuredArchive::FRecord Record) override;
-	virtual void GetResourceSizeEx(FResourceSizeEx& CumulativeResourceSize) override;
+	NIAGARA_API virtual void PostLoad() override;
+	NIAGARA_API virtual void PostInitProperties() override;
+	NIAGARA_API virtual void Serialize(FStructuredArchive::FRecord Record) override;
+	NIAGARA_API virtual void GetResourceSizeEx(FResourceSizeEx& CumulativeResourceSize) override;
 #if WITH_EDITOR
-	virtual void PostEditChangeProperty(struct FPropertyChangedEvent& PropertyChangedEvent) override;
-	virtual void RenameVariable(const FNiagaraVariableBase& OldVariable, const FNiagaraVariableBase& NewVariable, const FVersionedNiagaraEmitter& InEmitter) override;
-	virtual void RemoveVariable(const FNiagaraVariableBase& OldVariable, const FVersionedNiagaraEmitter& InEmitter) override;
+	NIAGARA_API virtual void PostEditChangeProperty(struct FPropertyChangedEvent& PropertyChangedEvent) override;
+	NIAGARA_API virtual void RenameVariable(const FNiagaraVariableBase& OldVariable, const FNiagaraVariableBase& NewVariable, const FVersionedNiagaraEmitter& InEmitter) override;
+	NIAGARA_API virtual void RemoveVariable(const FNiagaraVariableBase& OldVariable, const FVersionedNiagaraEmitter& InEmitter) override;
 #endif
 	//UObject Interface END
 
-	static void InitCDOPropertiesAfterModuleStartup();
+	static NIAGARA_API void InitCDOPropertiesAfterModuleStartup();
 
 	//UNiagaraRendererProperties Interface
-	virtual FNiagaraRenderer* CreateEmitterRenderer(ERHIFeatureLevel::Type FeatureLevel, const FNiagaraEmitterInstance* Emitter, const FNiagaraSystemInstanceController& InController) override;
-	virtual class FNiagaraBoundsCalculator* CreateBoundsCalculator() override;
-	virtual void GetUsedMaterials(const FNiagaraEmitterInstance* InEmitter, TArray<UMaterialInterface*>& OutMaterials) const override;
-	virtual const FVertexFactoryType* GetVertexFactoryType() const override;
-	virtual bool IsBackfaceCullingDisabled() const override;
+	NIAGARA_API virtual FNiagaraRenderer* CreateEmitterRenderer(ERHIFeatureLevel::Type FeatureLevel, const FNiagaraEmitterInstance* Emitter, const FNiagaraSystemInstanceController& InController) override;
+	NIAGARA_API virtual class FNiagaraBoundsCalculator* CreateBoundsCalculator() override;
+	NIAGARA_API virtual void GetUsedMaterials(const FNiagaraEmitterInstance* InEmitter, TArray<UMaterialInterface*>& OutMaterials) const override;
+	NIAGARA_API virtual const FVertexFactoryType* GetVertexFactoryType() const override;
+	NIAGARA_API virtual bool IsBackfaceCullingDisabled() const override;
 	virtual bool IsSimTargetSupported(ENiagaraSimTarget InSimTarget) const override { return true; };
-	virtual bool PopulateRequiredBindings(FNiagaraParameterStore& InParameterStore) override;
-	virtual void CollectPSOPrecacheData(FPSOPrecacheParamsList& OutParams) override;
+	NIAGARA_API virtual bool PopulateRequiredBindings(FNiagaraParameterStore& InParameterStore) override;
+	NIAGARA_API virtual void CollectPSOPrecacheData(FPSOPrecacheParamsList& OutParams) override;
 
 #if WITH_EDITOR
-	virtual const TArray<FNiagaraVariable>& GetOptionalAttributes() override;
-	virtual void GetAdditionalVariables(TArray<FNiagaraVariableBase>& OutArray) const override;
-	virtual FNiagaraVariable GetBoundAttribute(const FNiagaraVariableAttributeBinding* Binding) const override;
-	virtual void GetRendererWidgets(const FNiagaraEmitterInstance* InEmitter, TArray<TSharedPtr<SWidget>>& OutWidgets, TSharedPtr<FAssetThumbnailPool> InThumbnailPool) const override;
-	virtual void GetRendererTooltipWidgets(const FNiagaraEmitterInstance* InEmitter, TArray<TSharedPtr<SWidget>>& OutWidgets, TSharedPtr<FAssetThumbnailPool> InThumbnailPool) const override;
-	virtual void GetRendererFeedback(const FVersionedNiagaraEmitter& InEmitter, TArray<FNiagaraRendererFeedback>& OutErrors, TArray<FNiagaraRendererFeedback>& OutWarnings, TArray<FNiagaraRendererFeedback>& OutInfo) const override;
-
-	virtual TArray<FNiagaraVariable> GetBoundAttributes() const override; 
+	NIAGARA_API virtual const TArray<FNiagaraVariable>& GetOptionalAttributes() override;
+	NIAGARA_API virtual void GetAdditionalVariables(TArray<FNiagaraVariableBase>& OutArray) const override;
+	NIAGARA_API virtual FNiagaraVariable GetBoundAttribute(const FNiagaraVariableAttributeBinding* Binding) const override;
+	NIAGARA_API virtual void GetRendererWidgets(const FNiagaraEmitterInstance* InEmitter, TArray<TSharedPtr<SWidget>>& OutWidgets, TSharedPtr<FAssetThumbnailPool> InThumbnailPool) const override;
+	NIAGARA_API virtual void GetRendererTooltipWidgets(const FNiagaraEmitterInstance* InEmitter, TArray<TSharedPtr<SWidget>>& OutWidgets, TSharedPtr<FAssetThumbnailPool> InThumbnailPool) const override;
+	NIAGARA_API virtual void GetRendererFeedback(const FVersionedNiagaraEmitter& InEmitter, TArray<FNiagaraRendererFeedback>& OutErrors, TArray<FNiagaraRendererFeedback>& OutWarnings, TArray<FNiagaraRendererFeedback>& OutInfo) const override;
+
+	NIAGARA_API virtual TArray<FNiagaraVariable> GetBoundAttributes() const override; 
 #endif
 #if WITH_EDITORONLY_DATA
-	virtual bool IsSupportedVariableForBinding(const FNiagaraVariableBase& InSourceForBinding, const FName& InTargetBindingName) const override;
+	NIAGARA_API virtual bool IsSupportedVariableForBinding(const FNiagaraVariableBase& InSourceForBinding, const FName& InTargetBindingName) const override;
 #endif
-	virtual void CacheFromCompiledData(const FNiagaraDataSetCompiledData* CompiledData) override;
+	NIAGARA_API virtual void CacheFromCompiledData(const FNiagaraDataSetCompiledData* CompiledData) override;
 	//UNiagaraRendererProperties Interface END
 
 	UPROPERTY(EditAnywhere, Category = "Ribbon Rendering")
@@ -493,12 +493,12 @@ public:
 	FNiagaraRendererLayout RendererLayout;
 
 protected:
-	void InitBindings();
-	void SetPreviousBindings(const FVersionedNiagaraEmitter& SrcEmitter);
+	NIAGARA_API void InitBindings();
+	NIAGARA_API void SetPreviousBindings(const FVersionedNiagaraEmitter& SrcEmitter);
 
-	void UpdateSourceModeDerivates(ENiagaraRendererSourceDataMode InSourceMode, bool bFromPropertyEdit);
+	NIAGARA_API void UpdateSourceModeDerivates(ENiagaraRendererSourceDataMode InSourceMode, bool bFromPropertyEdit);
 
-	void UpdateMICs();
+	NIAGARA_API void UpdateMICs();
 
 	virtual bool NeedsMIDsForMaterials() const { return MaterialParameters.HasAnyBindings(); }
 };
diff --git a/Engine/Plugins/FX/Niagara/Source/Niagara/Public/NiagaraSceneProxy.h b/Engine/Plugins/FX/Niagara/Source/Niagara/Public/NiagaraSceneProxy.h
--- a/Engine/Plugins/FX/Niagara/Source/Niagara/Public/NiagaraSceneProxy.h
+++ b/Engine/Plugins/FX/Niagara/Source/Niagara/Public/NiagaraSceneProxy.h
@@ -13,42 +13,42 @@ class UNiagaraComponent;
 /**
 * Scene proxy for drawing niagara particle simulations.
 */
-class NIAGARA_API FNiagaraSceneProxy : public FPrimitiveSceneProxy
+class FNiagaraSceneProxy : public FPrimitiveSceneProxy
 {
 public:
-	virtual SIZE_T GetTypeHash() const override;
+	NIAGARA_API virtual SIZE_T GetTypeHash() const override;
 
-	FNiagaraSceneProxy(UNiagaraComponent* InComponent);
-	~FNiagaraSceneProxy();
+	NIAGARA_API FNiagaraSceneProxy(UNiagaraComponent* InComponent);
+	NIAGARA_API ~FNiagaraSceneProxy();
 
 	/** Retrieves the render data for a single system */
 	FNiagaraSystemRenderData* GetSystemRenderData() { return RenderData; }
 
 	/** Called to allow renderers to free render state */
-	void DestroyRenderState_Concurrent();
+	NIAGARA_API void DestroyRenderState_Concurrent();
 
 	/** Gets whether or not this scene proxy should be rendered. */
-	bool GetRenderingEnabled() const;
+	NIAGARA_API bool GetRenderingEnabled() const;
 
 	/** Sets whether or not this scene proxy should be rendered. */
-	void SetRenderingEnabled(bool bInRenderingEnabled);
+	NIAGARA_API void SetRenderingEnabled(bool bInRenderingEnabled);
 
 	FNiagaraGpuComputeDispatchInterface* GetComputeDispatchInterface() const { return ComputeDispatchInterface; }
 
 #if RHI_RAYTRACING
-	virtual void GetDynamicRayTracingInstances(FRayTracingMaterialGatheringContext& Context, TArray<FRayTracingInstance>& OutRayTracingInstances) override;
+	NIAGARA_API virtual void GetDynamicRayTracingInstances(FRayTracingMaterialGatheringContext& Context, TArray<FRayTracingInstance>& OutRayTracingInstances) override;
 	virtual bool IsRayTracingRelevant() const override { return true; }
 	virtual bool HasRayTracingRepresentation() const override { return true; }
 #endif
 
 	FORCEINLINE const FMatrix& GetLocalToWorldInverse() const { return LocalToWorldInverse; }
 
-	const FVector3f& GetLWCRenderTile() const;
+	NIAGARA_API const FVector3f& GetLWCRenderTile() const;
 
-	TUniformBuffer<FPrimitiveUniformShaderParameters>* GetCustomUniformBufferResource(bool bHasVelocity, const FBox& InstanceBounds = FBox(ForceInitToZero)) const;
-	FRHIUniformBuffer* GetCustomUniformBuffer(bool bHasVelocity, const FBox& InstanceBounds = FBox(ForceInitToZero)) const;
+	NIAGARA_API TUniformBuffer<FPrimitiveUniformShaderParameters>* GetCustomUniformBufferResource(bool bHasVelocity, const FBox& InstanceBounds = FBox(ForceInitToZero)) const;
+	NIAGARA_API FRHIUniformBuffer* GetCustomUniformBuffer(bool bHasVelocity, const FBox& InstanceBounds = FBox(ForceInitToZero)) const;
 
-	virtual FPrimitiveViewRelevance GetViewRelevance(const FSceneView* View) const override;
+	NIAGARA_API virtual FPrimitiveViewRelevance GetViewRelevance(const FSceneView* View) const override;
 
 	/** Some proxy wide dynamic settings passed down with the emitter dynamic data. */
 	struct FDynamicData
@@ -63,17 +63,17 @@ public:
 	void SetProxyDynamicData(const FDynamicData& NewData) { DynamicData = NewData; }
 
 private:
-	void ReleaseRenderThreadResources();
+	NIAGARA_API void ReleaseRenderThreadResources();
 
-	void ReleaseUniformBuffers(bool bEmpty);
+	NIAGARA_API void ReleaseUniformBuffers(bool bEmpty);
 
 	//~ Begin FPrimitiveSceneProxy Interface.
-	virtual void CreateRenderThreadResources() override;
+	NIAGARA_API virtual void CreateRenderThreadResources() override;
 
 	//virtual void OnActorPositionChanged() override;
-	virtual void OnTransformChanged() override;
+	NIAGARA_API virtual void OnTransformChanged() override;
 
-	virtual void GetDynamicMeshElements(const TArray<const FSceneView*>& Views, const FSceneViewFamily& ViewFamily, uint32 VisibilityMap, FMeshElementCollector& Collector) const override;
+	NIAGARA_API virtual void GetDynamicMeshElements(const TArray<const FSceneView*>& Views, const FSceneViewFamily& ViewFamily, uint32 VisibilityMap, FMeshElementCollector& Collector) const override;
 
 
 	virtual bool CanBeOccluded() const override
@@ -84,11 +84,11 @@ private:
 
 
 	/** Callback from the renderer to gather simple lights that this proxy wants renderered. */
-	virtual void GatherSimpleLights(const FSceneViewFamily& ViewFamily, FSimpleLightArray& OutParticleLights) const override;
+	NIAGARA_API virtual void GatherSimpleLights(const FSceneViewFamily& ViewFamily, FSimpleLightArray& OutParticleLights) const override;
 
-	virtual uint32 GetMemoryFootprint() const override;
+	NIAGARA_API virtual uint32 GetMemoryFootprint() const override;
 
-	uint32 GetAllocatedSize() const;
+	NIAGARA_API uint32 GetAllocatedSize() const;
 
 private:
 	/** Custom Uniform Buffers, allows us to have renderer specific data packed inside such as pre-skinned bounds. */
diff --git a/Engine/Plugins/FX/Niagara/Source/Niagara/Public/NiagaraSettings.h b/Engine/Plugins/FX/Niagara/Source/Niagara/Public/NiagaraSettings.h
--- a/Engine/Plugins/FX/Niagara/Source/Niagara/Public/NiagaraSettings.h
+++ b/Engine/Plugins/FX/Niagara/Source/Niagara/Public/NiagaraSettings.h
@@ -93,8 +93,8 @@ namespace ENDICollisionQuery_AsyncGpuTraceProvider
 	};
 }
 
-UCLASS(config = Niagara, defaultconfig, meta=(DisplayName="Niagara"))
-class NIAGARA_API UNiagaraSettings : public UDeveloperSettings
+UCLASS(config = Niagara, defaultconfig, meta=(DisplayName="Niagara"), MinimalAPI)
+class UNiagaraSettings : public UDeveloperSettings
 {
 	GENERATED_UCLASS_BODY()
 
@@ -236,26 +236,26 @@ class NIAGARA_API UNiagaraSettings : public UDeveloperSettings
 	TArray<FNiagaraPlatformSetRedirect> PlatformSetRedirects;
 
 	// Begin UDeveloperSettings Interface
-	virtual FName GetCategoryName() const override;
+	NIAGARA_API virtual FName GetCategoryName() const override;
 #if WITH_EDITOR
-	void AddEnumParameterType(UEnum* Enum);
-	virtual FText GetSectionText() const override;
+	NIAGARA_API void AddEnumParameterType(UEnum* Enum);
+	NIAGARA_API virtual FText GetSectionText() const override;
 
 	// END UDeveloperSettings Interface
 
-	UNiagaraEffectType* GetDefaultEffectType() const;
+	NIAGARA_API UNiagaraEffectType* GetDefaultEffectType() const;
 
-	UNiagaraEffectType* GetRequiredEffectType() const;
+	NIAGARA_API UNiagaraEffectType* GetRequiredEffectType() const;
 
-	virtual void PostEditChangeProperty(FPropertyChangedEvent& PropertyChangedEvent) override;
+	NIAGARA_API virtual void PostEditChangeProperty(FPropertyChangedEvent& PropertyChangedEvent) override;
 
 public:
 	DECLARE_MULTICAST_DELEGATE_TwoParams(FOnNiagaraSettingsChanged, const FName&, const UNiagaraSettings*);
 
 	/** Gets a multicast delegate which is called whenever one of the parameters in this settings object changes. */
-	static FOnNiagaraSettingsChanged& OnSettingsChanged();
+	static NIAGARA_API FOnNiagaraSettingsChanged& OnSettingsChanged();
 
 protected:
-	static FOnNiagaraSettingsChanged SettingsChangedDelegate;
+	static NIAGARA_API FOnNiagaraSettingsChanged SettingsChangedDelegate;
 #endif
-};
\ No newline at end of file
+};
diff --git a/Engine/Plugins/FX/Niagara/Source/Niagara/Public/NiagaraSimulationStageBase.h b/Engine/Plugins/FX/Niagara/Source/Niagara/Public/NiagaraSimulationStageBase.h
--- a/Engine/Plugins/FX/Niagara/Source/Niagara/Public/NiagaraSimulationStageBase.h
+++ b/Engine/Plugins/FX/Niagara/Source/Niagara/Public/NiagaraSimulationStageBase.h
@@ -16,13 +16,13 @@ class UNiagaraScript;
 /**
 * A base class for niagara simulation stages.  This class should be derived to add stage specific information.
 */
-UCLASS(abstract)
-class NIAGARA_API UNiagaraSimulationStageBase : public UNiagaraMergeable
+UCLASS(abstract, MinimalAPI)
+class UNiagaraSimulationStageBase : public UNiagaraMergeable
 {
 	GENERATED_BODY()
 
 public:
-	static const FName ParticleSpawnUpdateName;
+	static NIAGARA_API const FName ParticleSpawnUpdateName;
 
 	UNiagaraSimulationStageBase()
 	{
@@ -39,25 +39,25 @@ public:
 	uint32 bEnabled : 1;
 
 #if WITH_EDITORONLY_DATA
-	virtual bool AppendCompileHash(FNiagaraCompileHashVisitor* InVisitor) const;
+	NIAGARA_API virtual bool AppendCompileHash(FNiagaraCompileHashVisitor* InVisitor) const;
 
-	virtual bool FillCompilationData(TArray<FNiagaraSimulationStageCompilationData>& CompilationSimStageData) const PURE_VIRTUAL(UNiagaraSimulationStageBase::FillCompileSimStageData, return false;);
+	NIAGARA_API virtual bool FillCompilationData(TArray<FNiagaraSimulationStageCompilationData>& CompilationSimStageData) const PURE_VIRTUAL(UNiagaraSimulationStageBase::FillCompileSimStageData, return false;);
 	/** Return the FName to use in place of the default for the location in the stack context. If this would be the default, return NAME_None.*/
 	virtual FName GetStackContextReplacementName() const { return NAME_None; }
-	void SetEnabled(bool bEnabled);
-	void RequestRecompile();
-	virtual void PostEditChangeProperty(struct FPropertyChangedEvent& PropertyChangedEvent) override;
+	NIAGARA_API void SetEnabled(bool bEnabled);
+	NIAGARA_API void RequestRecompile();
+	NIAGARA_API virtual void PostEditChangeProperty(struct FPropertyChangedEvent& PropertyChangedEvent) override;
 
 	UPROPERTY()
 	FGuid OuterEmitterVersion;
 	
-	FVersionedNiagaraEmitterData* GetEmitterData() const;
-	FVersionedNiagaraEmitter GetOuterEmitter() const;
+	NIAGARA_API FVersionedNiagaraEmitterData* GetEmitterData() const;
+	NIAGARA_API FVersionedNiagaraEmitter GetOuterEmitter() const;
 #endif
 };
 
-UCLASS(meta = (DisplayName = "Generic Simulation Stage"))
-class NIAGARA_API UNiagaraSimulationStageGeneric : public UNiagaraSimulationStageBase
+UCLASS(meta = (DisplayName = "Generic Simulation Stage"), MinimalAPI)
+class UNiagaraSimulationStageGeneric : public UNiagaraSimulationStageBase
 {
 	GENERATED_BODY()
 
@@ -204,18 +204,18 @@ public:
 	//////////////////////////////////////////////////////////////////////////////////////////////////////////////////////////////////
 
 
-	virtual void PostInitProperties() override;
-	virtual void PostLoad() override;
+	NIAGARA_API virtual void PostInitProperties() override;
+	NIAGARA_API virtual void PostLoad() override;
 #if WITH_EDITORONLY_DATA
-	virtual bool AppendCompileHash(FNiagaraCompileHashVisitor* InVisitor) const override;
-	virtual bool FillCompilationData(TArray<FNiagaraSimulationStageCompilationData>& CompilationSimStageData) const override;
-	virtual void RenameEmitter(const FName& InOldName, const UNiagaraEmitter* InRenamedEmitter);
-	virtual void RenameVariable(const FNiagaraVariableBase& OldVariable, const FNiagaraVariableBase& NewVariable, const FVersionedNiagaraEmitter& InEmitter);
-	virtual void RemoveVariable(const FNiagaraVariableBase& OldVariable, const FVersionedNiagaraEmitter& InEmitter);
+	NIAGARA_API virtual bool AppendCompileHash(FNiagaraCompileHashVisitor* InVisitor) const override;
+	NIAGARA_API virtual bool FillCompilationData(TArray<FNiagaraSimulationStageCompilationData>& CompilationSimStageData) const override;
+	NIAGARA_API virtual void RenameEmitter(const FName& InOldName, const UNiagaraEmitter* InRenamedEmitter);
+	NIAGARA_API virtual void RenameVariable(const FNiagaraVariableBase& OldVariable, const FNiagaraVariableBase& NewVariable, const FVersionedNiagaraEmitter& InEmitter);
+	NIAGARA_API virtual void RemoveVariable(const FNiagaraVariableBase& OldVariable, const FVersionedNiagaraEmitter& InEmitter);
 #endif
 #if WITH_EDITOR
-	virtual FName GetStackContextReplacementName() const override; 
+	NIAGARA_API virtual FName GetStackContextReplacementName() const override; 
 
-	virtual void PostEditChangeProperty(struct FPropertyChangedEvent& PropertyChangedEvent) override;
+	NIAGARA_API virtual void PostEditChangeProperty(struct FPropertyChangedEvent& PropertyChangedEvent) override;
 #endif
 };
