[Niagara]

diff --git a/Engine/Plugins/FX/Niagara/Source/Niagara/Public/NiagaraSpriteRendererProperties.h b/Engine/Plugins/FX/Niagara/Source/Niagara/Public/NiagaraSpriteRendererProperties.h
--- a/Engine/Plugins/FX/Niagara/Source/Niagara/Public/NiagaraSpriteRendererProperties.h
+++ b/Engine/Plugins/FX/Niagara/Source/Niagara/Public/NiagaraSpriteRendererProperties.h
@@ -106,52 +106,52 @@ namespace ENiagaraSpriteVFLayout
 class FAssetThumbnailPool;
 class SWidget;
 
-UCLASS(editinlinenew, meta = (DisplayName = "Sprite Renderer"))
-class NIAGARA_API UNiagaraSpriteRendererProperties : public UNiagaraRendererProperties
+UCLASS(editinlinenew, meta = (DisplayName = "Sprite Renderer"), MinimalAPI)
+class UNiagaraSpriteRendererProperties : public UNiagaraRendererProperties
 {
 public:
 	GENERATED_BODY()
 
-	UNiagaraSpriteRendererProperties();
+	NIAGARA_API UNiagaraSpriteRendererProperties();
 
 	//UObject Interface
-	virtual void PostLoad() override;
-	virtual void PostInitProperties() override;
-	virtual void Serialize(FStructuredArchive::FRecord Record) override;
-	virtual void GetResourceSizeEx(FResourceSizeEx& CumulativeResourceSize) override;
+	NIAGARA_API virtual void PostLoad() override;
+	NIAGARA_API virtual void PostInitProperties() override;
+	NIAGARA_API virtual void Serialize(FStructuredArchive::FRecord Record) override;
+	NIAGARA_API virtual void GetResourceSizeEx(FResourceSizeEx& CumulativeResourceSize) override;
 #if WITH_EDITORONLY_DATA
-	virtual void PostEditChangeProperty(struct FPropertyChangedEvent& PropertyChangedEvent) override;
-	virtual void RenameVariable(const FNiagaraVariableBase& OldVariable, const FNiagaraVariableBase& NewVariable, const FVersionedNiagaraEmitter& InEmitter) override;
-	virtual void RemoveVariable(const FNiagaraVariableBase& OldVariable, const FVersionedNiagaraEmitter& InEmitter) override;
-	virtual TArray<FNiagaraVariable> GetBoundAttributes() const override;
+	NIAGARA_API virtual void PostEditChangeProperty(struct FPropertyChangedEvent& PropertyChangedEvent) override;
+	NIAGARA_API virtual void RenameVariable(const FNiagaraVariableBase& OldVariable, const FNiagaraVariableBase& NewVariable, const FVersionedNiagaraEmitter& InEmitter) override;
+	NIAGARA_API virtual void RemoveVariable(const FNiagaraVariableBase& OldVariable, const FVersionedNiagaraEmitter& InEmitter) override;
+	NIAGARA_API virtual TArray<FNiagaraVariable> GetBoundAttributes() const override;
 #endif // WITH_EDITORONLY_DATA
 	//UObject Interface END
 
-	static void InitCDOPropertiesAfterModuleStartup();
+	static NIAGARA_API void InitCDOPropertiesAfterModuleStartup();
 
 	//UNiagaraRendererProperties interface
-	virtual FNiagaraRenderer* CreateEmitterRenderer(ERHIFeatureLevel::Type FeatureLevel, const FNiagaraEmitterInstance* Emitter, const FNiagaraSystemInstanceController& InController) override;
-	virtual class FNiagaraBoundsCalculator* CreateBoundsCalculator() override;
-	virtual void GetUsedMaterials(const FNiagaraEmitterInstance* InEmitter, TArray<UMaterialInterface*>& OutMaterials) const override;
-	virtual const FVertexFactoryType* GetVertexFactoryType() const override;
+	NIAGARA_API virtual FNiagaraRenderer* CreateEmitterRenderer(ERHIFeatureLevel::Type FeatureLevel, const FNiagaraEmitterInstance* Emitter, const FNiagaraSystemInstanceController& InController) override;
+	NIAGARA_API virtual class FNiagaraBoundsCalculator* CreateBoundsCalculator() override;
+	NIAGARA_API virtual void GetUsedMaterials(const FNiagaraEmitterInstance* InEmitter, TArray<UMaterialInterface*>& OutMaterials) const override;
+	NIAGARA_API virtual const FVertexFactoryType* GetVertexFactoryType() const override;
 	virtual bool IsSimTargetSupported(ENiagaraSimTarget InSimTarget) const override { return true; };
-	virtual bool PopulateRequiredBindings(FNiagaraParameterStore& InParameterStore)  override;
-	virtual void CollectPSOPrecacheData(FPSOPrecacheParamsList& OutParams) override;
+	NIAGARA_API virtual bool PopulateRequiredBindings(FNiagaraParameterStore& InParameterStore)  override;
+	NIAGARA_API virtual void CollectPSOPrecacheData(FPSOPrecacheParamsList& OutParams) override;
 #if WITH_EDITOR
-	virtual const TArray<FNiagaraVariable>& GetOptionalAttributes() override;
-	virtual void GetAdditionalVariables(TArray<FNiagaraVariableBase>& OutArray) const override;
-	virtual void GetRendererWidgets(const FNiagaraEmitterInstance* InEmitter, TArray<TSharedPtr<SWidget>>& OutWidgets, TSharedPtr<FAssetThumbnailPool> InThumbnailPool) const override;
-	virtual void GetRendererTooltipWidgets(const FNiagaraEmitterInstance* InEmitter, TArray<TSharedPtr<SWidget>>& OutWidgets, TSharedPtr<FAssetThumbnailPool> InThumbnailPool) const override;
-	virtual void GetRendererFeedback(const FVersionedNiagaraEmitter& InEmitter, TArray<FText>& OutErrors, TArray<FText>& OutWarnings, TArray<FText>& OutInfo) const override;
-	virtual void GetRendererFeedback(const FVersionedNiagaraEmitter& InEmitter, TArray<FNiagaraRendererFeedback>& OutErrors, TArray<FNiagaraRendererFeedback>& OutWarnings, TArray<FNiagaraRendererFeedback>& OutInfo) const override;
+	NIAGARA_API virtual const TArray<FNiagaraVariable>& GetOptionalAttributes() override;
+	NIAGARA_API virtual void GetAdditionalVariables(TArray<FNiagaraVariableBase>& OutArray) const override;
+	NIAGARA_API virtual void GetRendererWidgets(const FNiagaraEmitterInstance* InEmitter, TArray<TSharedPtr<SWidget>>& OutWidgets, TSharedPtr<FAssetThumbnailPool> InThumbnailPool) const override;
+	NIAGARA_API virtual void GetRendererTooltipWidgets(const FNiagaraEmitterInstance* InEmitter, TArray<TSharedPtr<SWidget>>& OutWidgets, TSharedPtr<FAssetThumbnailPool> InThumbnailPool) const override;
+	NIAGARA_API virtual void GetRendererFeedback(const FVersionedNiagaraEmitter& InEmitter, TArray<FText>& OutErrors, TArray<FText>& OutWarnings, TArray<FText>& OutInfo) const override;
+	NIAGARA_API virtual void GetRendererFeedback(const FVersionedNiagaraEmitter& InEmitter, TArray<FNiagaraRendererFeedback>& OutErrors, TArray<FNiagaraRendererFeedback>& OutWarnings, TArray<FNiagaraRendererFeedback>& OutInfo) const override;
 #endif
 	virtual ENiagaraRendererSourceDataMode GetCurrentSourceMode() const override { return SourceMode; }
 
-	virtual void CacheFromCompiledData(const FNiagaraDataSetCompiledData* CompiledData) override;
+	NIAGARA_API virtual void CacheFromCompiledData(const FNiagaraDataSetCompiledData* CompiledData) override;
 	//UNiagaraMaterialRendererProperties interface END
 
-	int32 GetNumCutoutVertexPerSubimage() const;
-	uint32 GetNumIndicesPerInstance() const;
+	NIAGARA_API int32 GetNumCutoutVertexPerSubimage() const;
+	NIAGARA_API uint32 GetNumIndicesPerInstance() const;
 
 	/** The material used to render the particle. Note that it must have the Use with Niagara Sprites flag checked.*/
 	UPROPERTY(EditAnywhere, Category = "Sprite Rendering")
@@ -395,8 +395,8 @@ public:
 	UPROPERTY(EditAnywhere, Category="Cutout", meta=(UIMin = "0", UIMax = "1"))
 	float AlphaThreshold;
 
-	void UpdateCutoutTexture();
-	void CacheDerivedData();
+	NIAGARA_API void UpdateCutoutTexture();
+	NIAGARA_API void CacheDerivedData();
 #endif
 
 	UPROPERTY()
@@ -408,19 +408,19 @@ public:
 	FNiagaraRendererLayout RendererLayoutWithoutCustomSort;
 
 protected:
-	void InitBindings();
-	void SetPreviousBindings(const FVersionedNiagaraEmitter& SrcEmitter, ENiagaraRendererSourceDataMode InSourceMode);
-	virtual void UpdateSourceModeDerivates(ENiagaraRendererSourceDataMode InSourceMode, bool bFromPropertyEdit = false) override;
+	NIAGARA_API void InitBindings();
+	NIAGARA_API void SetPreviousBindings(const FVersionedNiagaraEmitter& SrcEmitter, ENiagaraRendererSourceDataMode InSourceMode);
+	NIAGARA_API virtual void UpdateSourceModeDerivates(ENiagaraRendererSourceDataMode InSourceMode, bool bFromPropertyEdit = false) override;
 
-	void UpdateMICs();
+	NIAGARA_API void UpdateMICs();
 
 #if WITH_EDITORONLY_DATA
-	virtual FNiagaraVariable GetBoundAttribute(const FNiagaraVariableAttributeBinding* Binding) const override;
+	NIAGARA_API virtual FNiagaraVariable GetBoundAttribute(const FNiagaraVariableAttributeBinding* Binding) const override;
 #endif
 
 private:
 	/** Derived data for this asset, generated off of SubUVTexture. */
 	FSubUVDerivedData DerivedData;
 
-	static TArray<TWeakObjectPtr<UNiagaraSpriteRendererProperties>> SpriteRendererPropertiesToDeferredInit;
+	static NIAGARA_API TArray<TWeakObjectPtr<UNiagaraSpriteRendererProperties>> SpriteRendererPropertiesToDeferredInit;
 };
diff --git a/Engine/Plugins/FX/Niagara/Source/Niagara/Public/NiagaraSystemInstance.h b/Engine/Plugins/FX/Niagara/Source/Niagara/Public/NiagaraSystemInstance.h
--- a/Engine/Plugins/FX/Niagara/Source/Niagara/Public/NiagaraSystemInstance.h
+++ b/Engine/Plugins/FX/Niagara/Source/Niagara/Public/NiagaraSystemInstance.h
@@ -66,7 +66,7 @@ private:
 #endif
 };
 
-class NIAGARA_API FNiagaraSystemInstance 
+class FNiagaraSystemInstance 
 {
 	friend class FNiagaraSystemSimulation;
 	friend class FNiagaraGPUSystemTick;
@@ -104,40 +104,40 @@ public:
 	FORCEINLINE bool GetAreDataInterfacesInitialized() const { return bDataInterfacesInitialized; }
 
 	/** Creates a new Niagara system instance. */
-	FNiagaraSystemInstance(UWorld& InWorld, UNiagaraSystem& InAsset, FNiagaraUserRedirectionParameterStore* InOverrideParameters = nullptr,
+	NIAGARA_API FNiagaraSystemInstance(UWorld& InWorld, UNiagaraSystem& InAsset, FNiagaraUserRedirectionParameterStore* InOverrideParameters = nullptr,
 		USceneComponent* InAttachComponent = nullptr, ENiagaraTickBehavior InTickBehavior = ENiagaraTickBehavior::UsePrereqs, bool bInPooled = false);
 
 	/** Cleanup*/
-	virtual ~FNiagaraSystemInstance();
+	NIAGARA_API virtual ~FNiagaraSystemInstance();
 
-	void Cleanup();
+	NIAGARA_API void Cleanup();
 
 	/** Initializes this System instance to simulate the supplied System. */
-	void Init(bool bInForceSolo=false);
+	NIAGARA_API void Init(bool bInForceSolo=false);
 
-	void Activate(EResetMode InResetMode = EResetMode::ResetAll);
-	void Deactivate(bool bImmediate = false);
-	void Complete(bool bExternalCompletion);
+	NIAGARA_API void Activate(EResetMode InResetMode = EResetMode::ResetAll);
+	NIAGARA_API void Deactivate(bool bImmediate = false);
+	NIAGARA_API void Complete(bool bExternalCompletion);
 
-	void OnPooledReuse(UWorld& NewWorld);
+	NIAGARA_API void OnPooledReuse(UWorld& NewWorld);
 
-	void SetPaused(bool bInPaused);
+	NIAGARA_API void SetPaused(bool bInPaused);
 	FORCEINLINE bool IsPaused() const { return (SystemInstanceState == ENiagaraSystemInstanceState::PendingSpawnPaused) || (SystemInstanceState == ENiagaraSystemInstanceState::Paused); }
 
-	void SetSolo(bool bInSolo);
+	NIAGARA_API void SetSolo(bool bInSolo);
 
-	void SetGpuComputeDebug(bool bEnableDebug);
-	void SetWarmupSettings(int32 WarmupTickCount, float WarmupTickDelta);
+	NIAGARA_API void SetGpuComputeDebug(bool bEnableDebug);
+	NIAGARA_API void SetWarmupSettings(int32 WarmupTickCount, float WarmupTickDelta);
 
-	UActorComponent* GetPrereqComponent() const;
+	NIAGARA_API UActorComponent* GetPrereqComponent() const;
 
 	//void RebindParameterCollection(UNiagaraParameterCollectionInstance* OldInstance, UNiagaraParameterCollectionInstance* NewInstance);
-	void BindParameters();
-	void UnbindParameters(bool bFromComplete = false);
+	NIAGARA_API void BindParameters();
+	NIAGARA_API void UnbindParameters(bool bFromComplete = false);
 
 	FORCEINLINE FNiagaraParameterStore& GetInstanceParameters() { return InstanceParameters; }
-	FNiagaraLWCConverter GetLWCConverter(bool bLocalSpaceEmitter = false) const;
-	FTransform GetLWCSimToWorld(bool bLocalSpaceEmitter = false) const;
+	NIAGARA_API FNiagaraLWCConverter GetLWCConverter(bool bLocalSpaceEmitter = false) const;
+	NIAGARA_API FTransform GetLWCSimToWorld(bool bLocalSpaceEmitter = false) const;
 
 	FORCEINLINE uint32 GetParameterIndex(bool PreviousFrame = false) const
 	{
@@ -161,35 +161,35 @@ public:
 	FORCEINLINE const FNiagaraEmitterParameters& GetEmitterParameters(int32 EmitterIdx, bool PreviousFrame = false) const { return EmitterParameters[EmitterIdx * 2 + GetParameterIndex(PreviousFrame)]; }
 	FORCEINLINE FNiagaraEmitterParameters& EditEmitterParameters(int32 EmitterIdx) { return EmitterParameters[EmitterIdx * 2 + GetParameterIndex()]; }
 	
-	FNiagaraWorldManager* GetWorldManager()const;
-	bool RequiresDistanceFieldData() const;
-	bool RequiresDepthBuffer() const;
-	bool RequiresEarlyViewData() const;
-	bool RequiresViewUniformBuffer() const;
-	bool RequiresRayTracingScene() const;
+	NIAGARA_API FNiagaraWorldManager* GetWorldManager()const;
+	NIAGARA_API bool RequiresDistanceFieldData() const;
+	NIAGARA_API bool RequiresDepthBuffer() const;
+	NIAGARA_API bool RequiresEarlyViewData() const;
+	NIAGARA_API bool RequiresViewUniformBuffer() const;
+	NIAGARA_API bool RequiresRayTracingScene() const;
 
 	/** Requests the the simulation be reset on the next tick. */
-	void Reset(EResetMode Mode);
+	NIAGARA_API void Reset(EResetMode Mode);
 
-	void ManualTick(float DeltaSeconds, const FGraphEventRef& MyCompletionGraphEvent);
+	NIAGARA_API void ManualTick(float DeltaSeconds, const FGraphEventRef& MyCompletionGraphEvent);
 
 	/** Ticks the system using the a SimCache. */
-	void SimCacheTick_GameThread(UNiagaraSimCache* SimCache, float DesiredAge, float DeltaSeconds, const FGraphEventRef& MyCompletionGraphEvent);
+	NIAGARA_API void SimCacheTick_GameThread(UNiagaraSimCache* SimCache, float DesiredAge, float DeltaSeconds, const FGraphEventRef& MyCompletionGraphEvent);
 	/** Concurrent work for SimCache tick */
-	void SimCacheTick_Concurrent(UNiagaraSimCache* SimCache);
+	NIAGARA_API void SimCacheTick_Concurrent(UNiagaraSimCache* SimCache);
 
 	/** Initial phase of system instance tick. Must be executed on the game thread. */
-	void Tick_GameThread(float DeltaSeconds);
+	NIAGARA_API void Tick_GameThread(float DeltaSeconds);
 	/** Secondary phase of the system instance tick that can be executed on any thread. */
-	void Tick_Concurrent(bool bEnqueueGPUTickIfNeeded = true);
+	NIAGARA_API void Tick_Concurrent(bool bEnqueueGPUTickIfNeeded = true);
 	/** 
 		Final phase of system instance tick. Must be executed on the game thread. 
 		Returns whether the Finalize was actually done. It's possible for the finalize in a task to have already been done earlier on the GT by a WaitForAsyncAndFinalize call.
 	*/
-	void FinalizeTick_GameThread(bool bEnqueueGPUTickIfNeeded = true);
+	NIAGARA_API void FinalizeTick_GameThread(bool bEnqueueGPUTickIfNeeded = true);
 
-	void GenerateAndSubmitGPUTick();
-	void InitGPUTick(FNiagaraGPUSystemTick& OutTick);
+	NIAGARA_API void GenerateAndSubmitGPUTick();
+	NIAGARA_API void InitGPUTick(FNiagaraGPUSystemTick& OutTick);
 
 	void SetPendingFinalize(FNiagaraSystemInstanceFinalizeRef InFinalizeRef) { FinalizeRef = InFinalizeRef; }
 	bool HasPendingFinalize() const { return FinalizeRef.IsPending(); }
@@ -198,27 +198,27 @@ public:
 	Wait for any pending concurrent work to complete, must be called on the GameThread.
 	This will NOT call finalize on the instance so can leave a dangling finalize task.
 	*/
-	void WaitForConcurrentTickDoNotFinalize(bool bEnsureComplete = false);
+	NIAGARA_API void WaitForConcurrentTickDoNotFinalize(bool bEnsureComplete = false);
 
 	/**
 	Wait for any pending concurrent work to complete, must be called on the GameThread.
 	The instance will be finalized if pending, this can complete the instance and remove it from the system simulation.
 	*/
-	void WaitForConcurrentTickAndFinalize(bool bEnsureComplete = false);
+	NIAGARA_API void WaitForConcurrentTickAndFinalize(bool bEnsureComplete = false);
 
 	/** Handles completion of the system and returns true if the system is complete. */
-	bool HandleCompletion();
+	NIAGARA_API bool HandleCompletion();
 
-	void SetEmitterEnable(FName EmitterName, bool bNewEnableState);
+	NIAGARA_API void SetEmitterEnable(FName EmitterName, bool bNewEnableState);
 
 	/** Perform per-tick updates on data interfaces that need it. This can cause systems to complete so cannot be parallelized. */
-	void TickDataInterfaces(float DeltaSeconds, bool bPostSimulate);
+	NIAGARA_API void TickDataInterfaces(float DeltaSeconds, bool bPostSimulate);
 
 	ENiagaraExecutionState GetRequestedExecutionState()const { return RequestedExecutionState; }
-	void SetRequestedExecutionState(ENiagaraExecutionState InState);
+	NIAGARA_API void SetRequestedExecutionState(ENiagaraExecutionState InState);
 
 	ENiagaraExecutionState GetActualExecutionState() { return ActualExecutionState; }
-	void SetActualExecutionState(ENiagaraExecutionState InState);
+	NIAGARA_API void SetActualExecutionState(ENiagaraExecutionState InState);
 
 //	float GetSystemTimeSinceRendered() const { return SystemTimeSinceRenderedParam.GetValue(); }
 
@@ -231,7 +231,7 @@ public:
 	FORCEINLINE bool IsDisabled()const { return ActualExecutionState == ENiagaraExecutionState::Disabled; }
 
 	/** Gets the simulation for the supplied emitter handle. */
-	TSharedPtr<FNiagaraEmitterInstance, ESPMode::ThreadSafe> GetSimulationForHandle(const FNiagaraEmitterHandle& EmitterHandle);
+	NIAGARA_API TSharedPtr<FNiagaraEmitterInstance, ESPMode::ThreadSafe> GetSimulationForHandle(const FNiagaraEmitterHandle& EmitterHandle);
 
 	FORCEINLINE UWorld* GetWorld() const { return World; }
 	FORCEINLINE UNiagaraSystem* GetSystem() const { return System; }
@@ -242,16 +242,16 @@ public:
 	FORCEINLINE const TArray<TSharedRef<FNiagaraEmitterInstance, ESPMode::ThreadSafe> >& GetEmitters() const { return Emitters; }
 	FORCEINLINE const FBox& GetLocalBounds() const { return LocalBounds;  }
 	FORCEINLINE const FVector3f& GetLWCTile() const { return LWCTile;  }
-	TConstArrayView<FNiagaraEmitterExecutionIndex> GetEmitterExecutionOrder() const;
+	NIAGARA_API TConstArrayView<FNiagaraEmitterExecutionIndex> GetEmitterExecutionOrder() const;
 
 	FORCEINLINE void SetSystemFixedBounds(const FBox& InLocalBounds) { FixedBounds_GT = InLocalBounds; }
-	FBox GetSystemFixedBounds() const;
-	void SetEmitterFixedBounds(FName EmitterName, const FBox& InLocalBounds);
-	FBox GetEmitterFixedBounds(FName EmitterName) const;
+	NIAGARA_API FBox GetSystemFixedBounds() const;
+	NIAGARA_API void SetEmitterFixedBounds(FName EmitterName, const FBox& InLocalBounds);
+	NIAGARA_API FBox GetEmitterFixedBounds(FName EmitterName) const;
 
-	FNiagaraEmitterInstance* GetEmitterByID(FGuid InID);
+	NIAGARA_API FNiagaraEmitterInstance* GetEmitterByID(FGuid InID);
 
-	void SetForceSolo(bool bForceSolo);
+	NIAGARA_API void SetForceSolo(bool bForceSolo);
 	FORCEINLINE bool IsSolo() const { return bSolo; }
 
 	FORCEINLINE bool NeedsGPUTick() const { return ActiveGPUEmitterCount > 0 /*&& Component->IsRegistered()*/ && !IsComplete();}
@@ -265,16 +265,16 @@ public:
 
 #if WITH_EDITOR
 	/** Gets a multicast delegate which is called whenever this instance is initialized with an System asset. */
-	FOnInitialized& OnInitialized();
+	NIAGARA_API FOnInitialized& OnInitialized();
 
 	/** Gets a multicast delegate which is called whenever this instance is reset due to external changes in the source System asset. */
-	FOnReset& OnReset();
+	NIAGARA_API FOnReset& OnReset();
 
-	FOnDestroyed& OnDestroyed();
+	NIAGARA_API FOnDestroyed& OnDestroyed();
 #endif
 
 #if WITH_EDITORONLY_DATA
-	bool GetIsolateEnabled() const;
+	NIAGARA_API bool GetIsolateEnabled() const;
 #endif
 
 	FNiagaraSystemInstanceID GetId() const { return ID; }
@@ -324,10 +324,10 @@ public:
 
 	FORCEINLINE const FNiagaraPerInstanceDIFuncInfo& GetPerInstanceDIFunction(ENiagaraSystemSimulationScript ScriptType, int32 FuncIndex)const { return PerInstanceDIFunctions[(int32)ScriptType][FuncIndex]; }
 
-	void EvaluateBoundFunction(FName FunctionName, bool& UsedOnCpu, bool& UsedOnGpu) const;
+	NIAGARA_API void EvaluateBoundFunction(FName FunctionName, bool& UsedOnCpu, bool& UsedOnGpu) const;
 
 #if WITH_EDITORONLY_DATA
-	bool UsesCollection(const UNiagaraParameterCollection* Collection) const;
+	NIAGARA_API bool UsesCollection(const UNiagaraParameterCollection* Collection) const;
 #endif
 
 	FORCEINLINE bool IsPendingSpawn() const { return (SystemInstanceState == ENiagaraSystemInstanceState::PendingSpawn) || (SystemInstanceState == ENiagaraSystemInstanceState::PendingSpawnPaused); }
@@ -343,77 +343,77 @@ public:
 		return SystemSimulation; 
 	}
 
-	bool IsReadyToRun() const;
+	NIAGARA_API bool IsReadyToRun() const;
 
-	UNiagaraParameterCollectionInstance* GetParameterCollectionInstance(UNiagaraParameterCollection* Collection);
+	NIAGARA_API UNiagaraParameterCollectionInstance* GetParameterCollectionInstance(UNiagaraParameterCollection* Collection);
 
 	/** 
 	Manually advances this system's simulation by the specified number of ticks and tick delta. 
 	To be advanced in this way a system must be in solo mode or moved into solo mode which will add additional overhead.
 	*/
-	void AdvanceSimulation(int32 TickCountToSimulate, float TickDeltaSeconds);
+	NIAGARA_API void AdvanceSimulation(int32 TickCountToSimulate, float TickDeltaSeconds);
 
 #if NIAGARA_SYSTEM_CAPTURE
 	/** Request that this simulation capture a frame. Cannot capture if disabled or already completed.*/
-	bool RequestCapture(const FGuid& RequestId);
+	NIAGARA_API bool RequestCapture(const FGuid& RequestId);
 
 	/** Poll for previous frame capture requests. Once queried and bool is returned, the results are cleared from this system instance.*/
-	bool QueryCaptureResults(const FGuid& RequestId, TArray<TSharedPtr<struct FNiagaraScriptDebuggerInfo, ESPMode::ThreadSafe>>& OutCaptureResults);
+	NIAGARA_API bool QueryCaptureResults(const FGuid& RequestId, TArray<TSharedPtr<struct FNiagaraScriptDebuggerInfo, ESPMode::ThreadSafe>>& OutCaptureResults);
 
 	/** Only call from within the script execution states. Value is null if not capturing a frame.*/
-	TArray<TSharedPtr<struct FNiagaraScriptDebuggerInfo, ESPMode::ThreadSafe>>* GetActiveCaptureResults();
+	NIAGARA_API TArray<TSharedPtr<struct FNiagaraScriptDebuggerInfo, ESPMode::ThreadSafe>>* GetActiveCaptureResults();
 
 	/** Only call from within the script execution states. Does nothing if not capturing a frame.*/
-	void FinishCapture();
+	NIAGARA_API void FinishCapture();
 
 	/** Only call from within the script execution states. Value is false if not capturing a frame.*/
-	bool ShouldCaptureThisFrame() const;
+	NIAGARA_API bool ShouldCaptureThisFrame() const;
 
 	/** Only call from within the script execution states. Value is nullptr if not capturing a frame.*/
-	TSharedPtr<struct FNiagaraScriptDebuggerInfo, ESPMode::ThreadSafe> GetActiveCaptureWrite(const FName& InHandleName, ENiagaraScriptUsage InUsage, const FGuid& InUsageId);
+	NIAGARA_API TSharedPtr<struct FNiagaraScriptDebuggerInfo, ESPMode::ThreadSafe> GetActiveCaptureWrite(const FName& InHandleName, ENiagaraScriptUsage InUsage, const FGuid& InUsageId);
 
 #endif
 
 	/** Dumps all of this systems info to the log. */
-	void Dump()const;
+	NIAGARA_API void Dump()const;
 
 	/** Dumps information about the instances tick to the log */
-	void DumpTickInfo(FOutputDevice& Ar);
+	NIAGARA_API void DumpTickInfo(FOutputDevice& Ar);
 
 	FNiagaraGpuComputeDispatchInterface* GetComputeDispatchInterface() const { return ComputeDispatchInterface; }
 
-	static bool AllocateSystemInstance(FNiagaraSystemInstancePtr& OutSystemInstanceAllocation, UWorld& InWorld, UNiagaraSystem& InAsset,
+	static NIAGARA_API bool AllocateSystemInstance(FNiagaraSystemInstancePtr& OutSystemInstanceAllocation, UWorld& InWorld, UNiagaraSystem& InAsset,
 		FNiagaraUserRedirectionParameterStore* InOverrideParameters = nullptr, USceneComponent* InAttachComponent = nullptr,
 		ENiagaraTickBehavior InTickBehavior = ENiagaraTickBehavior::UsePrereqs, bool bInPooled = false);
-	static bool DeallocateSystemInstance(FNiagaraSystemInstancePtr& SystemInstanceAllocation);
+	static NIAGARA_API bool DeallocateSystemInstance(FNiagaraSystemInstancePtr& SystemInstanceAllocation);
 	/*void SetHasGPUEmitters(bool bInHasGPUEmitters) { bHasGPUEmitters = bInHasGPUEmitters; }*/
 	bool HasGPUEmitters() { return bHasGPUEmitters;  }
 
-	void TickInstanceParameters_GameThread(float DeltaSeconds);
+	NIAGARA_API void TickInstanceParameters_GameThread(float DeltaSeconds);
 
-	void TickInstanceParameters_Concurrent();
+	NIAGARA_API void TickInstanceParameters_Concurrent();
 
-	FNiagaraDataSet* CreateEventDataSet(FName EmitterName, FName EventName);
-	FNiagaraDataSet* GetEventDataSet(FName EmitterName, FName EventName) const;
-	void ClearEventDataSets();
+	NIAGARA_API FNiagaraDataSet* CreateEventDataSet(FName EmitterName, FName EventName);
+	NIAGARA_API FNiagaraDataSet* GetEventDataSet(FName EmitterName, FName EventName) const;
+	NIAGARA_API void ClearEventDataSets();
 
-	FORCEINLINE void SetLODDistance(float InLODDistance, float InMaxLODDistance, bool bOverride);
-	FORCEINLINE void ClearLODDistance();
+	NIAGARA_API FORCEINLINE void SetLODDistance(float InLODDistance, float InMaxLODDistance, bool bOverride);
+	NIAGARA_API FORCEINLINE void ClearLODDistance();
 
-	const FString& GetCrashReporterTag()const;
+	NIAGARA_API const FString& GetCrashReporterTag()const;
 
 #if WITH_EDITOR
-	void RaiseNeedsUIResync();
-	bool HandleNeedsUIResync();
+	NIAGARA_API void RaiseNeedsUIResync();
+	NIAGARA_API bool HandleNeedsUIResync();
 #endif
 
 	/** Get the current tick behavior */
 	ENiagaraTickBehavior GetTickBehavior() const { return TickBehavior; }
 	/** Set a new tick behavior, this will not move the instance straight away and will wait until the next time it is evaluated */
-	void SetTickBehavior(ENiagaraTickBehavior NewTickBehavior);
+	NIAGARA_API void SetTickBehavior(ENiagaraTickBehavior NewTickBehavior);
 	
 	/** Calculates which tick group the instance should be in. */
-	ETickingGroup CalculateTickGroup() const;
+	NIAGARA_API ETickingGroup CalculateTickGroup() const;
 
 	/** Gets the current world transform of the system */
 	FORCEINLINE const FTransform& GetWorldTransform() const { return WorldTransform; }
@@ -432,35 +432,35 @@ public:
 	FORCEINLINE void SetSystemSignificanceIndex(int32 InIndex) { SignificanceIndex = InIndex; }
 
 	/** Calculates the distance to use for distance based LODing / culling. */
-	float GetLODDistance();
+	NIAGARA_API float GetLODDistance();
 	float GetMaxLODDistance() const { return MaxLODDistance ;}
 
-	void OnSimulationDestroyed();
+	NIAGARA_API void OnSimulationDestroyed();
 
 	void SetRandomSeedOffset(int32 Offset) { RandomSeedOffset = Offset; }
 	int32 GetRandomSeedOffset() const { return RandomSeedOffset; }
 
-	FNDIStageTickHandler* GetSystemDIStageTickHandler(ENiagaraScriptUsage Usage);
+	NIAGARA_API FNDIStageTickHandler* GetSystemDIStageTickHandler(ENiagaraScriptUsage Usage);
 private:
-	void DumpStalledInfo();
+	NIAGARA_API void DumpStalledInfo();
 
-	void DestroyDataInterfaceInstanceData();
+	NIAGARA_API void DestroyDataInterfaceInstanceData();
 
 	/** Builds the emitter simulations. */
-	void InitEmitters();
+	NIAGARA_API void InitEmitters();
 
 	/** Resets the System, emitter simulations, and renderers to initial conditions. */
-	void ReInitInternal();
+	NIAGARA_API void ReInitInternal();
 	/** Resets for restart, assumes no change in emitter setup */
-	void ResetInternal(bool bResetSimulations);
+	NIAGARA_API void ResetInternal(bool bResetSimulations);
 
 	/** Resets the parameter structs */
-	void ResetParameters();
+	NIAGARA_API void ResetParameters();
 
 	/** Call PrepareForSImulation on each data source from the simulations and determine which need per-tick updates.*/
-	void InitDataInterfaces();
+	NIAGARA_API void InitDataInterfaces();
 
-	void ResolveUserDataInterfaceBindings();
+	NIAGARA_API void ResolveUserDataInterfaceBindings();
 
 	/** The LWC tile of this system instance, used to offset all local simulation relative to the origin */
 	FVector3f LWCTile = FVector3f::ZeroVector;
@@ -686,4 +686,4 @@ FORCEINLINE void FNiagaraSystemInstance::ClearLODDistance()
 	bLODDistanceIsValid = false;
 	LODDistance = 0.0f;
 	MaxLODDistance = 1.0f;
-}
\ No newline at end of file
+}
diff --git a/Engine/Plugins/FX/Niagara/Source/Niagara/Public/NiagaraTypes.h b/Engine/Plugins/FX/Niagara/Source/Niagara/Public/NiagaraTypes.h
--- a/Engine/Plugins/FX/Niagara/Source/Niagara/Public/NiagaraTypes.h
+++ b/Engine/Plugins/FX/Niagara/Source/Niagara/Public/NiagaraTypes.h
@@ -192,7 +192,7 @@ struct FNiagaraMatrix
 };
 
 USTRUCT()
-struct NIAGARA_API FNiagaraAssetVersion
+struct FNiagaraAssetVersion
 {
 	GENERATED_USTRUCT_BODY()
 
@@ -217,7 +217,7 @@ struct NIAGARA_API FNiagaraAssetVersion
 	bool operator<(const FNiagaraAssetVersion& Other) const { return MajorVersion < Other.MajorVersion || (MajorVersion == Other.MajorVersion && MinorVersion < Other.MinorVersion); }
 	bool operator<=(const FNiagaraAssetVersion& Other) const { return *this < Other || (MajorVersion == Other.MajorVersion && MinorVersion == Other.MinorVersion); }
 
-	static FGuid CreateStableVersionGuid(UObject* Object);
+	static NIAGARA_API FGuid CreateStableVersionGuid(UObject* Object);
 
 	friend FORCEINLINE uint32 GetTypeHash(const FNiagaraAssetVersion& Version)
 	{
@@ -225,18 +225,18 @@ struct NIAGARA_API FNiagaraAssetVersion
 	}
 };
 
-struct NIAGARA_API FNiagaraLWCConverter
+struct FNiagaraLWCConverter
 {
-	explicit FNiagaraLWCConverter(FVector InSystemWorldPos = FVector::ZeroVector);
+	NIAGARA_API explicit FNiagaraLWCConverter(FVector InSystemWorldPos = FVector::ZeroVector);
 
-	[[nodiscard]] FVector3f ConvertWorldToSimulationVector(FVector WorldPosition) const;
-	[[nodiscard]] FNiagaraPosition ConvertWorldToSimulationPosition(FVector WorldPosition) const;
+	[[nodiscard]] NIAGARA_API FVector3f ConvertWorldToSimulationVector(FVector WorldPosition) const;
+	[[nodiscard]] NIAGARA_API FNiagaraPosition ConvertWorldToSimulationPosition(FVector WorldPosition) const;
 	
-	[[nodiscard]] FVector ConvertSimulationPositionToWorld(FNiagaraPosition SimulationPosition) const;
-	[[nodiscard]] FVector ConvertSimulationVectorToWorld(FVector3f SimulationPosition) const;
+	[[nodiscard]] NIAGARA_API FVector ConvertSimulationPositionToWorld(FNiagaraPosition SimulationPosition) const;
+	[[nodiscard]] NIAGARA_API FVector ConvertSimulationVectorToWorld(FVector3f SimulationPosition) const;
 
-	[[nodiscard]] FMatrix ConvertWorldToSimulationMatrix(const FMatrix& Matrix) const;
-	[[nodiscard]] FMatrix ConvertSimulationToWorldMatrix(const FMatrix& Matrix) const;
+	[[nodiscard]] NIAGARA_API FMatrix ConvertWorldToSimulationMatrix(const FMatrix& Matrix) const;
+	[[nodiscard]] NIAGARA_API FMatrix ConvertSimulationToWorldMatrix(const FMatrix& Matrix) const;
 
 private:
 	FVector SystemWorldPos;
@@ -294,14 +294,14 @@ struct FNiagaraStructConversionStep
 
 /** Can convert struct data from custom structs containing LWC data such as FVector3d into struct data suitable for Niagara simulations and vice versa. */
 USTRUCT()
-struct NIAGARA_API FNiagaraLwcStructConverter
+struct FNiagaraLwcStructConverter
 {
 	GENERATED_USTRUCT_BODY();
 
-	void ConvertDataToSimulation(uint8* DestinationData, const uint8* SourceData, int32 Count = 1) const;
-	void ConvertDataFromSimulation(uint8* DestinationData, const uint8* SourceData, int32 Count = 1) const;
+	NIAGARA_API void ConvertDataToSimulation(uint8* DestinationData, const uint8* SourceData, int32 Count = 1) const;
+	NIAGARA_API void ConvertDataFromSimulation(uint8* DestinationData, const uint8* SourceData, int32 Count = 1) const;
 	
-	void AddConversionStep(int32 InSourceBytes, int32 InSourceOffset, int32 InSimulationBytes, int32 InSimulationOffset, ENiagaraStructConversionType ConversionType);
+	NIAGARA_API void AddConversionStep(int32 InSourceBytes, int32 InSourceOffset, int32 InSimulationBytes, int32 InSimulationOffset, ENiagaraStructConversionType ConversionType);
 	bool IsValid() const { return ConversionSteps.Num() > 0 && LWCSize > 0 && SWCSize > 0; }
 	
 private:
@@ -405,32 +405,32 @@ enum class ENiagaraStructConversion : uint8
 *  Can convert a UStruct with fields of base types only (float, int... - will likely add native vector types here as well)
 *	to an FNiagaraTypeDefinition (internal representation)
 */
-class NIAGARA_API FNiagaraTypeHelper
+class FNiagaraTypeHelper
 {
 public:
-	static FString ToString(const uint8* ValueData, const UObject* StructOrEnum);
-	static bool IsLWCStructure(UStruct* InStruct);
-	static bool IsLWCType(const FNiagaraTypeDefinition& InType);
-	static UScriptStruct* FindNiagaraFriendlyTopLevelStruct(UScriptStruct* InStruct, ENiagaraStructConversion StructConversion);
-	static bool IsNiagaraFriendlyTopLevelStruct(UScriptStruct* InStruct, ENiagaraStructConversion StructConversion);
-	static UScriptStruct* GetSWCStruct(UScriptStruct* LWCStruct);
-	static UScriptStruct* GetLWCStruct(UScriptStruct* LWCStruct);
-	static FNiagaraTypeDefinition GetSWCType(const FNiagaraTypeDefinition& InType);
-	static void TickTypeRemap();
+	static NIAGARA_API FString ToString(const uint8* ValueData, const UObject* StructOrEnum);
+	static NIAGARA_API bool IsLWCStructure(UStruct* InStruct);
+	static NIAGARA_API bool IsLWCType(const FNiagaraTypeDefinition& InType);
+	static NIAGARA_API UScriptStruct* FindNiagaraFriendlyTopLevelStruct(UScriptStruct* InStruct, ENiagaraStructConversion StructConversion);
+	static NIAGARA_API bool IsNiagaraFriendlyTopLevelStruct(UScriptStruct* InStruct, ENiagaraStructConversion StructConversion);
+	static NIAGARA_API UScriptStruct* GetSWCStruct(UScriptStruct* LWCStruct);
+	static NIAGARA_API UScriptStruct* GetLWCStruct(UScriptStruct* LWCStruct);
+	static NIAGARA_API FNiagaraTypeDefinition GetSWCType(const FNiagaraTypeDefinition& InType);
+	static NIAGARA_API void TickTypeRemap();
 
 	/**
 	Attempt to get the actual LWC type for a given type that would be used in the rest of the engine etc.
 	Used for cases where we actually need to store types in LWC format.
 	These cannot be properly serialized currently so must be constructed at runtime from their Niagara Simulation types.
 	*/
-	static FNiagaraTypeDefinition GetLWCType(const FNiagaraTypeDefinition& InType);
+	static NIAGARA_API FNiagaraTypeDefinition GetLWCType(const FNiagaraTypeDefinition& InType);
 	
-	static FNiagaraTypeDefinition GetVector2DDef();
-	static FNiagaraTypeDefinition GetVectorDef();
-	static FNiagaraTypeDefinition GetVector4Def();
-	static FNiagaraTypeDefinition GetQuatDef();
+	static NIAGARA_API FNiagaraTypeDefinition GetVector2DDef();
+	static NIAGARA_API FNiagaraTypeDefinition GetVectorDef();
+	static NIAGARA_API FNiagaraTypeDefinition GetVector4Def();
+	static NIAGARA_API FNiagaraTypeDefinition GetQuatDef();
 
-	static void InitStaticTypes();
+	static NIAGARA_API void InitStaticTypes();
 private:
 	struct FRemapEntry
 	{
@@ -596,7 +596,7 @@ enum class ENiagaraOrientationAxis : uint32
 
 /** Defines options for conditionally editing and showing script inputs in the UI. */
 USTRUCT()
-struct NIAGARA_API FNiagaraInputConditionMetadata
+struct FNiagaraInputConditionMetadata
 {
 	GENERATED_USTRUCT_BODY()
 public:
@@ -611,7 +611,7 @@ public:
 
 /** Defines override data for enum parameters displayed in the UI. */
 USTRUCT()
-struct NIAGARA_API FNiagaraEnumParameterMetaData
+struct FNiagaraEnumParameterMetaData
 {
 	GENERATED_BODY()
 	
@@ -650,7 +650,7 @@ enum class ENiagaraInputWidgetType : uint8
 };
 
 USTRUCT()
-struct NIAGARA_API FWidgetNamedInputValue
+struct FWidgetNamedInputValue
 {
 	GENERATED_BODY()
 
@@ -666,7 +666,7 @@ struct NIAGARA_API FWidgetNamedInputValue
 
 /** A struct that serves as display metadata for integer type static switches. Is used in conjuction with the 'EnumStyle' widget customization. */
 USTRUCT()
-struct NIAGARA_API FNiagaraWidgetNamedIntegerInputValue
+struct FNiagaraWidgetNamedIntegerInputValue
 {
 	GENERATED_BODY()
 	
@@ -678,7 +678,7 @@ struct NIAGARA_API FNiagaraWidgetNamedIntegerInputValue
 };
 
 USTRUCT()
-struct NIAGARA_API FNiagaraInputParameterCustomization
+struct FNiagaraInputParameterCustomization
 {
 	GENERATED_BODY()
 
@@ -727,7 +727,7 @@ enum class ENiagaraBoolDisplayMode : uint8
 };
 
 USTRUCT()
-struct NIAGARA_API FNiagaraBoolParameterMetaData
+struct FNiagaraBoolParameterMetaData
 {
 	GENERATED_BODY()
 
@@ -756,7 +756,7 @@ struct NIAGARA_API FNiagaraBoolParameterMetaData
 };
 
 USTRUCT()
-struct NIAGARA_API FNiagaraVariableMetaData
+struct FNiagaraVariableMetaData
 {
 	GENERATED_USTRUCT_BODY()
 
@@ -843,7 +843,7 @@ struct NIAGARA_API FNiagaraVariableMetaData
 	int32 GetStaticSwitchDefaultValue_DEPRECATED() const { return StaticSwitchDefaultValue_DEPRECATED; };
 
 	/** Copies all the properties that are marked as editable for the user (e.g. EditAnywhere). */
-	void CopyUserEditableMetaData(const FNiagaraVariableMetaData& OtherMetaData);
+	NIAGARA_API void CopyUserEditableMetaData(const FNiagaraVariableMetaData& OtherMetaData);
 
 	FGuid GetVariableGuid() const { return VariableGuid; };
 
@@ -870,7 +870,7 @@ private:
 };
 
 USTRUCT()
-struct NIAGARA_API FNiagaraTypeDefinition
+struct FNiagaraTypeDefinition
 {
 	GENERATED_USTRUCT_BODY()
 
@@ -1083,7 +1083,7 @@ public:
 		return UnderlyingType == UT_Enum ? CastChecked<UEnum>(ClassStructOrEnum) : nullptr;
 	}
 
-	bool IsDataInterface() const;
+	NIAGARA_API bool IsDataInterface() const;
 
 	FORCEINLINE bool IsUObject() const
 	{
@@ -1105,8 +1105,8 @@ public:
 
 	bool IsIndexWildcard() const { return ClassStructOrEnum == FNiagaraTypeDefinition::GetWildcardStruct(); }
 
-	FNiagaraTypeDefinition ToStaticDef() const;
-	FNiagaraTypeDefinition RemoveStaticDef() const;
+	NIAGARA_API FNiagaraTypeDefinition ToStaticDef() const;
+	NIAGARA_API FNiagaraTypeDefinition RemoveStaticDef() const;
 
 	int32 GetSize() const
 	{
@@ -1163,10 +1163,10 @@ public:
 		return ClassStructOrEnum != nullptr;
 	}
 
-	bool AppendCompileHash(FNiagaraCompileHashVisitor* InVisitor) const;
+	NIAGARA_API bool AppendCompileHash(FNiagaraCompileHashVisitor* InVisitor) const;
 
 #if WITH_EDITORONLY_DATA
-	bool IsInternalType() const;
+	NIAGARA_API bool IsInternalType() const;
 #endif
 
 	/*
@@ -1181,7 +1181,7 @@ public:
 	UPROPERTY(EditAnywhere, Category=Type)
 	uint16 UnderlyingType;
 
-	bool Serialize(FArchive& Ar);
+	NIAGARA_API bool Serialize(FArchive& Ar);
 
 private:
 	UPROPERTY(EditAnywhere, Category = Type)
@@ -1199,9 +1199,9 @@ private:
 #endif
 
 public:
-	static void Init();
+	static NIAGARA_API void Init();
 #if WITH_EDITOR
-	static void RecreateUserDefinedTypeRegistry();
+	static NIAGARA_API void RecreateUserDefinedTypeRegistry();
 #endif
 	static const FNiagaraTypeDefinition& GetFloatDef() { return FloatDef; }
 	static const FNiagaraTypeDefinition& GetBoolDef() { return BoolDef; }
@@ -1267,7 +1267,7 @@ public:
 
 	static const FNiagaraTypeDefinition& GetCollisionEventDef() { return CollisionEventDef; }
 
-	static bool IsScalarDefinition(const FNiagaraTypeDefinition& Type);
+	static NIAGARA_API bool IsScalarDefinition(const FNiagaraTypeDefinition& Type);
 
 	FString ToString(const uint8* ValueData)const
 	{
@@ -1279,94 +1279,94 @@ public:
 		return FNiagaraTypeHelper::ToString(ValueData, ClassStructOrEnum);
 	}
 	
-	static bool TypesAreAssignable(const FNiagaraTypeDefinition& TypeA, const FNiagaraTypeDefinition& TypeB, bool bAllowLossyLWCConversions = false);
-	static bool IsLossyConversion(const FNiagaraTypeDefinition& FromType, const FNiagaraTypeDefinition& ToType);
-	static FNiagaraTypeDefinition GetNumericOutputType(TConstArrayView<FNiagaraTypeDefinition> TypeDefinintions, ENiagaraNumericOutputTypeSelectionMode SelectionMode);
+	static NIAGARA_API bool TypesAreAssignable(const FNiagaraTypeDefinition& TypeA, const FNiagaraTypeDefinition& TypeB, bool bAllowLossyLWCConversions = false);
+	static NIAGARA_API bool IsLossyConversion(const FNiagaraTypeDefinition& FromType, const FNiagaraTypeDefinition& ToType);
+	static NIAGARA_API FNiagaraTypeDefinition GetNumericOutputType(TConstArrayView<FNiagaraTypeDefinition> TypeDefinintions, ENiagaraNumericOutputTypeSelectionMode SelectionMode);
 
 	static const TArray<FNiagaraTypeDefinition>& GetNumericTypes() { return OrderedNumericTypes; }
-	static bool IsValidNumericInput(const FNiagaraTypeDefinition& TypeDef);
+	static NIAGARA_API bool IsValidNumericInput(const FNiagaraTypeDefinition& TypeDef);
 
 	
 
 private:
 
-	static FNiagaraTypeDefinition FloatDef;
-	static FNiagaraTypeDefinition BoolDef;
-	static FNiagaraTypeDefinition IntDef;
-	static FNiagaraTypeDefinition Vec2Def;
-	static FNiagaraTypeDefinition Vec3Def;
-	static FNiagaraTypeDefinition Vec4Def;
-	static FNiagaraTypeDefinition ColorDef;
-	static FNiagaraTypeDefinition PositionDef;
-	static FNiagaraTypeDefinition QuatDef;
-	static FNiagaraTypeDefinition Matrix4Def;
-	static FNiagaraTypeDefinition NumericDef;
-	static FNiagaraTypeDefinition ParameterMapDef;
-	static FNiagaraTypeDefinition IDDef;
-	static FNiagaraTypeDefinition RandInfoDef;
-	static FNiagaraTypeDefinition UObjectDef;
-	static FNiagaraTypeDefinition UMaterialDef;
-	static FNiagaraTypeDefinition UTextureDef;
-	static FNiagaraTypeDefinition UTextureRenderTargetDef;
-	static FNiagaraTypeDefinition UStaticMeshDef;
-	static FNiagaraTypeDefinition USimCacheClassDef;
-	static FNiagaraTypeDefinition WildcardDef;
-
-	static FNiagaraTypeDefinition HalfDef;
-	static FNiagaraTypeDefinition HalfVec2Def;
-	static FNiagaraTypeDefinition HalfVec3Def;
-	static FNiagaraTypeDefinition HalfVec4Def;
-
-	static UScriptStruct* FloatStruct;
-	static UScriptStruct* BoolStruct;
-	static UScriptStruct* IntStruct;
-	static UScriptStruct* Vec2Struct;
-	static UScriptStruct* Vec3Struct;
-	static UScriptStruct* Vec4Struct;
-	static UScriptStruct* QuatStruct;
-	static UScriptStruct* ColorStruct;
-	static UScriptStruct* Matrix4Struct;
-	static UScriptStruct* NumericStruct;
-	static UScriptStruct* WildcardStruct;
-	static UScriptStruct* PositionStruct;
-
-	static UScriptStruct* HalfStruct;
-	static UScriptStruct* HalfVec2Struct;
-	static UScriptStruct* HalfVec3Struct;
-	static UScriptStruct* HalfVec4Struct;
-
-	static UClass* UObjectClass;
-	static UClass* UMaterialClass;
-	static UClass* UTextureClass;
-	static UClass* UTextureRenderTargetClass;
-
-	static UEnum* SimulationTargetEnum;
-	static UEnum* ScriptUsageEnum;
-	static UEnum* ScriptContextEnum;
-	static UEnum* ExecutionStateEnum;
-	static UEnum* CoordinateSpaceEnum;
-	static UEnum* OrientationAxisEnum;
-	static UEnum* ExecutionStateSourceEnum;
-
-	static UEnum* ParameterScopeEnum;
-	static UEnum* ParameterPanelCategoryEnum;
-
-	static UEnum* FunctionDebugStateEnum;
-
-	static UScriptStruct* ParameterMapStruct;
-	static UScriptStruct* IDStruct;
-	static UScriptStruct* RandInfoStruct;
-
-	static TSet<UScriptStruct*> NumericStructs;
-	static TArray<FNiagaraTypeDefinition> OrderedNumericTypes;
-
-	static TSet<UScriptStruct*> ScalarStructs;
-
-	static TSet<UStruct*> FloatStructs;
-	static TSet<UStruct*> IntStructs;
-	static TSet<UStruct*> BoolStructs;
-
-	static FNiagaraTypeDefinition CollisionEventDef;
+	static NIAGARA_API FNiagaraTypeDefinition FloatDef;
+	static NIAGARA_API FNiagaraTypeDefinition BoolDef;
+	static NIAGARA_API FNiagaraTypeDefinition IntDef;
+	static NIAGARA_API FNiagaraTypeDefinition Vec2Def;
+	static NIAGARA_API FNiagaraTypeDefinition Vec3Def;
+	static NIAGARA_API FNiagaraTypeDefinition Vec4Def;
+	static NIAGARA_API FNiagaraTypeDefinition ColorDef;
+	static NIAGARA_API FNiagaraTypeDefinition PositionDef;
+	static NIAGARA_API FNiagaraTypeDefinition QuatDef;
+	static NIAGARA_API FNiagaraTypeDefinition Matrix4Def;
+	static NIAGARA_API FNiagaraTypeDefinition NumericDef;
+	static NIAGARA_API FNiagaraTypeDefinition ParameterMapDef;
+	static NIAGARA_API FNiagaraTypeDefinition IDDef;
+	static NIAGARA_API FNiagaraTypeDefinition RandInfoDef;
+	static NIAGARA_API FNiagaraTypeDefinition UObjectDef;
+	static NIAGARA_API FNiagaraTypeDefinition UMaterialDef;
+	static NIAGARA_API FNiagaraTypeDefinition UTextureDef;
+	static NIAGARA_API FNiagaraTypeDefinition UTextureRenderTargetDef;
+	static NIAGARA_API FNiagaraTypeDefinition UStaticMeshDef;
+	static NIAGARA_API FNiagaraTypeDefinition USimCacheClassDef;
+	static NIAGARA_API FNiagaraTypeDefinition WildcardDef;
+
+	static NIAGARA_API FNiagaraTypeDefinition HalfDef;
+	static NIAGARA_API FNiagaraTypeDefinition HalfVec2Def;
+	static NIAGARA_API FNiagaraTypeDefinition HalfVec3Def;
+	static NIAGARA_API FNiagaraTypeDefinition HalfVec4Def;
+
+	static NIAGARA_API UScriptStruct* FloatStruct;
+	static NIAGARA_API UScriptStruct* BoolStruct;
+	static NIAGARA_API UScriptStruct* IntStruct;
+	static NIAGARA_API UScriptStruct* Vec2Struct;
+	static NIAGARA_API UScriptStruct* Vec3Struct;
+	static NIAGARA_API UScriptStruct* Vec4Struct;
+	static NIAGARA_API UScriptStruct* QuatStruct;
+	static NIAGARA_API UScriptStruct* ColorStruct;
+	static NIAGARA_API UScriptStruct* Matrix4Struct;
+	static NIAGARA_API UScriptStruct* NumericStruct;
+	static NIAGARA_API UScriptStruct* WildcardStruct;
+	static NIAGARA_API UScriptStruct* PositionStruct;
+
+	static NIAGARA_API UScriptStruct* HalfStruct;
+	static NIAGARA_API UScriptStruct* HalfVec2Struct;
+	static NIAGARA_API UScriptStruct* HalfVec3Struct;
+	static NIAGARA_API UScriptStruct* HalfVec4Struct;
+
+	static NIAGARA_API UClass* UObjectClass;
+	static NIAGARA_API UClass* UMaterialClass;
+	static NIAGARA_API UClass* UTextureClass;
+	static NIAGARA_API UClass* UTextureRenderTargetClass;
+
+	static NIAGARA_API UEnum* SimulationTargetEnum;
+	static NIAGARA_API UEnum* ScriptUsageEnum;
+	static NIAGARA_API UEnum* ScriptContextEnum;
+	static NIAGARA_API UEnum* ExecutionStateEnum;
+	static NIAGARA_API UEnum* CoordinateSpaceEnum;
+	static NIAGARA_API UEnum* OrientationAxisEnum;
+	static NIAGARA_API UEnum* ExecutionStateSourceEnum;
+
+	static NIAGARA_API UEnum* ParameterScopeEnum;
+	static NIAGARA_API UEnum* ParameterPanelCategoryEnum;
+
+	static NIAGARA_API UEnum* FunctionDebugStateEnum;
+
+	static NIAGARA_API UScriptStruct* ParameterMapStruct;
+	static NIAGARA_API UScriptStruct* IDStruct;
+	static NIAGARA_API UScriptStruct* RandInfoStruct;
+
+	static NIAGARA_API TSet<UScriptStruct*> NumericStructs;
+	static NIAGARA_API TArray<FNiagaraTypeDefinition> OrderedNumericTypes;
+
+	static NIAGARA_API TSet<UScriptStruct*> ScalarStructs;
+
+	static NIAGARA_API TSet<UStruct*> FloatStructs;
+	static NIAGARA_API TSet<UStruct*> IntStructs;
+	static NIAGARA_API TSet<UStruct*> BoolStructs;
+
+	static NIAGARA_API FNiagaraTypeDefinition CollisionEventDef;
 
 	friend FORCEINLINE uint32 GetTypeHash(const FNiagaraTypeDefinition& Type)
 	{
@@ -1428,7 +1428,7 @@ ENUM_CLASS_FLAGS(ENiagaraTypeRegistryFlags)
 * Used by UI to provide selection; new uniforms and variables
 * may be instanced using the types provided here
 */
-class NIAGARA_API FNiagaraTypeRegistry : public FGCObject
+class FNiagaraTypeRegistry : public FGCObject
 {
 public:
 	enum
@@ -1488,7 +1488,7 @@ public:
 		return Get().RegisteredIndexTypes;
 	}
 
-	static UNiagaraDataInterfaceBase* GetDefaultDataInterfaceByName(const FString& DIClassName);
+	static NIAGARA_API UNiagaraDataInterfaceBase* GetDefaultDataInterfaceByName(const FString& DIClassName);
 
 	static void ClearUserDefinedRegistry()
 	{
@@ -1668,18 +1668,18 @@ public:
 	}
 
 	/** LazySingleton interface */
-	static FNiagaraTypeRegistry& Get();
-	static void TearDown();
+	static NIAGARA_API FNiagaraTypeRegistry& Get();
+	static NIAGARA_API void TearDown();
 
 	/** FGCObject interface */
-	virtual void AddReferencedObjects(FReferenceCollector& Collector) override;
-	virtual FString GetReferencerName() const;
+	NIAGARA_API virtual void AddReferencedObjects(FReferenceCollector& Collector) override;
+	NIAGARA_API virtual FString GetReferencerName() const;
 
 private:
 	friend class FLazySingleton;
 
-	FNiagaraTypeRegistry();
-	virtual ~FNiagaraTypeRegistry();
+	NIAGARA_API FNiagaraTypeRegistry();
+	NIAGARA_API virtual ~FNiagaraTypeRegistry();
 
 	RegisteredTypesArray RegisteredTypes;
 
@@ -2258,7 +2258,7 @@ class UNiagaraEmitter;
 
 /** Struct combining an emitter with a specific version.*/
 USTRUCT()
-struct NIAGARA_API FVersionedNiagaraEmitter
+struct FVersionedNiagaraEmitter
 {
 	GENERATED_USTRUCT_BODY()
 	
@@ -2267,8 +2267,8 @@ struct NIAGARA_API FVersionedNiagaraEmitter
 		: Emitter(InEmitter), Version(InVersion)
 	{};
 
-	FVersionedNiagaraEmitterData* GetEmitterData() const;
-	FVersionedNiagaraEmitterWeakPtr ToWeakPtr() const;
+	NIAGARA_API FVersionedNiagaraEmitterData* GetEmitterData() const;
+	NIAGARA_API FVersionedNiagaraEmitterWeakPtr ToWeakPtr() const;
 
 	UPROPERTY()
 	TObjectPtr<UNiagaraEmitter> Emitter = nullptr;
@@ -2292,13 +2292,13 @@ struct NIAGARA_API FVersionedNiagaraEmitter
 	}
 };
 
-struct NIAGARA_API FVersionedNiagaraEmitterWeakPtr
+struct FVersionedNiagaraEmitterWeakPtr
 {
 	FVersionedNiagaraEmitterWeakPtr() : Emitter(nullptr) {}
-	FVersionedNiagaraEmitterWeakPtr(UNiagaraEmitter* InEmitter, const FGuid& InVersion);
-	FVersionedNiagaraEmitter ResolveWeakPtr() const;
-	FVersionedNiagaraEmitterData* GetEmitterData() const;
-	bool IsValid() const;
+	NIAGARA_API FVersionedNiagaraEmitterWeakPtr(UNiagaraEmitter* InEmitter, const FGuid& InVersion);
+	NIAGARA_API FVersionedNiagaraEmitter ResolveWeakPtr() const;
+	NIAGARA_API FVersionedNiagaraEmitterData* GetEmitterData() const;
+	NIAGARA_API bool IsValid() const;
 	bool operator==(const FVersionedNiagaraEmitterWeakPtr& Other) const
 	{
 		return Emitter == Other.Emitter && Version == Other.Version;
diff --git a/Engine/Plugins/FX/Niagara/Source/Niagara/Public/NiagaraVariant.h b/Engine/Plugins/FX/Niagara/Source/Niagara/Public/NiagaraVariant.h
--- a/Engine/Plugins/FX/Niagara/Source/Niagara/Public/NiagaraVariant.h
+++ b/Engine/Plugins/FX/Niagara/Source/Niagara/Public/NiagaraVariant.h
@@ -19,37 +19,37 @@ enum class ENiagaraVariantMode
 };
 
 USTRUCT()
-struct NIAGARA_API FNiagaraVariant
+struct FNiagaraVariant
 {
 	GENERATED_BODY()
 
-	FNiagaraVariant();
+	NIAGARA_API FNiagaraVariant();
 	// TODO: Add move support
 	FNiagaraVariant(const FNiagaraVariant&) = default;
 	FNiagaraVariant& operator=(const FNiagaraVariant&) = default;
-	explicit FNiagaraVariant(UNiagaraDataInterface* InDataInterface);
-	explicit FNiagaraVariant(UObject* InObject);
-	explicit FNiagaraVariant(const TArray<uint8>& InBytes);
-	FNiagaraVariant(const void* InBytes, int32 Size);
+	NIAGARA_API explicit FNiagaraVariant(UNiagaraDataInterface* InDataInterface);
+	NIAGARA_API explicit FNiagaraVariant(UObject* InObject);
+	NIAGARA_API explicit FNiagaraVariant(const TArray<uint8>& InBytes);
+	NIAGARA_API FNiagaraVariant(const void* InBytes, int32 Size);
 
-	UObject* GetUObject() const;
-	void SetUObject(UObject* InObject);
+	NIAGARA_API UObject* GetUObject() const;
+	NIAGARA_API void SetUObject(UObject* InObject);
 
-	UNiagaraDataInterface* GetDataInterface() const;
-	void SetDataInterface(UNiagaraDataInterface* InDataInterface);
+	NIAGARA_API UNiagaraDataInterface* GetDataInterface() const;
+	NIAGARA_API void SetDataInterface(UNiagaraDataInterface* InDataInterface);
 
-	void AllocateBytes(int32 InCount);
-	void SetBytes(const uint8* InBytes, int32 InCount);
-	uint8* GetBytes() const;
-	int32 GetNumBytes() const;
+	NIAGARA_API void AllocateBytes(int32 InCount);
+	NIAGARA_API void SetBytes(const uint8* InBytes, int32 InCount);
+	NIAGARA_API uint8* GetBytes() const;
+	NIAGARA_API int32 GetNumBytes() const;
 
 	bool IsValid() const { return CurrentMode != ENiagaraVariantMode::None; }
 	ENiagaraVariantMode GetMode() const { return CurrentMode; }
 
-	bool operator==(const FNiagaraVariant& Other) const;
-	bool operator!=(const FNiagaraVariant& Other) const;
+	NIAGARA_API bool operator==(const FNiagaraVariant& Other) const;
+	NIAGARA_API bool operator!=(const FNiagaraVariant& Other) const;
 
-	bool AppendCompileHash(FNiagaraCompileHashVisitor* InVisitor) const;
+	NIAGARA_API bool AppendCompileHash(FNiagaraCompileHashVisitor* InVisitor) const;
 
 private:
 	UPROPERTY(EditAnywhere, Category=Variant, Instanced)
@@ -63,4 +63,4 @@ private:
 
 	UPROPERTY(EditAnywhere, Category=Variant)
 	ENiagaraVariantMode CurrentMode;
-};
\ No newline at end of file
+};
diff --git a/Engine/Plugins/FX/Niagara/Source/Niagara/Public/NiagaraWorldManager.h b/Engine/Plugins/FX/Niagara/Source/Niagara/Public/NiagaraWorldManager.h
--- a/Engine/Plugins/FX/Niagara/Source/Niagara/Public/NiagaraWorldManager.h
+++ b/Engine/Plugins/FX/Niagara/Source/Niagara/Public/NiagaraWorldManager.h
@@ -159,7 +159,7 @@ public:
 		return static_cast<T&>(**ExistingValue);
 	}
 
-	NIAGARA_API TArrayView<const FNiagaraCachedViewInfo> GetCachedViewInfo() const { return MakeArrayView(CachedViewInfo); }
+	TArrayView<const FNiagaraCachedViewInfo> GetCachedViewInfo() const { return MakeArrayView(CachedViewInfo); }
 
 	UNiagaraComponentPool* GetComponentPool() { return ComponentPool; }
 
@@ -226,7 +226,7 @@ public:
 	static void OnRefreshOwnerAllowsScalability();
 
 	NIAGARA_API static void SetScalabilityCullingMode(ENiagaraScalabilityCullingMode NewMode);
-	NIAGARA_API static ENiagaraScalabilityCullingMode GetScalabilityCullingMode() { return ScalabilityCullingMode; }
+	static ENiagaraScalabilityCullingMode GetScalabilityCullingMode() { return ScalabilityCullingMode; }
 
 	FNiagaraDataChannelManager& GetDataChannelManager(){ return *DataChannelManager; }
 	NIAGARA_API void InitDataChannel(const UNiagaraDataChannel* Channel, bool bForce);
diff --git a/Engine/Plugins/FX/Niagara/Source/NiagaraAnimNotifies/Public/AnimNotifyState_TimedNiagaraEffect.h b/Engine/Plugins/FX/Niagara/Source/NiagaraAnimNotifies/Public/AnimNotifyState_TimedNiagaraEffect.h
--- a/Engine/Plugins/FX/Niagara/Source/NiagaraAnimNotifies/Public/AnimNotifyState_TimedNiagaraEffect.h
+++ b/Engine/Plugins/FX/Niagara/Source/NiagaraAnimNotifies/Public/AnimNotifyState_TimedNiagaraEffect.h
@@ -16,8 +16,8 @@ class UAnimInstance;
 // Timed Niagara Effect Notify
 // Allows a looping Niagara effect to be played in an animation that will activate
 // at the beginning of the notify and deactivate at the end.
-UCLASS(Blueprintable, meta = (DisplayName = "Timed Niagara Effect"))
-class NIAGARAANIMNOTIFIES_API UAnimNotifyState_TimedNiagaraEffect : public UAnimNotifyState
+UCLASS(Blueprintable, meta = (DisplayName = "Timed Niagara Effect"), MinimalAPI)
+class UAnimNotifyState_TimedNiagaraEffect : public UAnimNotifyState
 {
 	GENERATED_UCLASS_BODY()
 
@@ -43,26 +43,26 @@ class NIAGARAANIMNOTIFIES_API UAnimNotifyState_TimedNiagaraEffect : public UAnim
 	bool bDestroyAtEnd;
 
 	UE_DEPRECATED(5.0, "Please use the other NotifyBegin function instead")
-	virtual void NotifyBegin(class USkeletalMeshComponent * MeshComp, class UAnimSequenceBase * Animation, float TotalDuration) override;
+	NIAGARAANIMNOTIFIES_API virtual void NotifyBegin(class USkeletalMeshComponent * MeshComp, class UAnimSequenceBase * Animation, float TotalDuration) override;
 
 	UE_DEPRECATED(5.0, "Please use the other NotifyEnd function instead")
-	virtual void NotifyEnd(class USkeletalMeshComponent * MeshComp, class UAnimSequenceBase * Animation) override;
+	NIAGARAANIMNOTIFIES_API virtual void NotifyEnd(class USkeletalMeshComponent * MeshComp, class UAnimSequenceBase * Animation) override;
 	
-	virtual void NotifyBegin(class USkeletalMeshComponent * MeshComp, class UAnimSequenceBase * Animation, float TotalDuration, const FAnimNotifyEventReference& EventReference) override;
-	virtual void NotifyEnd(class USkeletalMeshComponent * MeshComp, class UAnimSequenceBase * Animation, const FAnimNotifyEventReference& EventReference) override;
+	NIAGARAANIMNOTIFIES_API virtual void NotifyBegin(class USkeletalMeshComponent * MeshComp, class UAnimSequenceBase * Animation, float TotalDuration, const FAnimNotifyEventReference& EventReference) override;
+	NIAGARAANIMNOTIFIES_API virtual void NotifyEnd(class USkeletalMeshComponent * MeshComp, class UAnimSequenceBase * Animation, const FAnimNotifyEventReference& EventReference) override;
 
 	// Overridden from UAnimNotifyState to provide custom notify name.
-	FString GetNotifyName_Implementation() const override;
+	NIAGARAANIMNOTIFIES_API FString GetNotifyName_Implementation() const override;
 
 	// Return FXSystemComponent created from SpawnEffect
 	UFUNCTION(BlueprintCallable, Category = "AnimNotify")
-	UFXSystemComponent* GetSpawnedEffect(UMeshComponent* MeshComp);
+	NIAGARAANIMNOTIFIES_API UFXSystemComponent* GetSpawnedEffect(UMeshComponent* MeshComp);
 
 protected:
 	// Spawns the NiagaraSystemComponent. Called from Notify.
-	virtual UFXSystemComponent* SpawnEffect(USkeletalMeshComponent* MeshComp, UAnimSequenceBase* Animation) const;
+	NIAGARAANIMNOTIFIES_API virtual UFXSystemComponent* SpawnEffect(USkeletalMeshComponent* MeshComp, UAnimSequenceBase* Animation) const;
 
-	bool ValidateParameters(USkeletalMeshComponent* MeshComp) const;
+	NIAGARAANIMNOTIFIES_API bool ValidateParameters(USkeletalMeshComponent* MeshComp) const;
 
 	FORCEINLINE FName GetSpawnedComponentTag()const { return GetFName(); }
 };
@@ -81,28 +81,28 @@ public:
 /**
 Same as Timed Niagara Effect but also provides some more advanced abilities at an additional cost. 
 */
-UCLASS(Blueprintable, meta = (DisplayName = "Advanced Timed Niagara Effect"))
-class NIAGARAANIMNOTIFIES_API UAnimNotifyState_TimedNiagaraEffectAdvanced : public UAnimNotifyState_TimedNiagaraEffect
+UCLASS(Blueprintable, meta = (DisplayName = "Advanced Timed Niagara Effect"), MinimalAPI)
+class UAnimNotifyState_TimedNiagaraEffectAdvanced : public UAnimNotifyState_TimedNiagaraEffect
 {
 	GENERATED_UCLASS_BODY()
 
 public:
 	UE_DEPRECATED(5.0, "Please use the other NotifyBegin function instead")
-	virtual void NotifyBegin(class USkeletalMeshComponent* MeshComp, class UAnimSequenceBase* Animation, float TotalDuration) override;
+	NIAGARAANIMNOTIFIES_API virtual void NotifyBegin(class USkeletalMeshComponent* MeshComp, class UAnimSequenceBase* Animation, float TotalDuration) override;
 
 	UE_DEPRECATED(5.0, "Please use the other NotifyEnd function instead")
-	virtual void NotifyEnd(class USkeletalMeshComponent* MeshComp, class UAnimSequenceBase* Animation) override;
+	NIAGARAANIMNOTIFIES_API virtual void NotifyEnd(class USkeletalMeshComponent* MeshComp, class UAnimSequenceBase* Animation) override;
 
 	UE_DEPRECATED(5.0, "Please use the other NotifyTick function instead")
-	virtual void NotifyTick(USkeletalMeshComponent * MeshComp, UAnimSequenceBase * Animation, float FrameDeltaTime) override;
+	NIAGARAANIMNOTIFIES_API virtual void NotifyTick(USkeletalMeshComponent * MeshComp, UAnimSequenceBase * Animation, float FrameDeltaTime) override;
 	
-	virtual void NotifyBegin(class USkeletalMeshComponent* MeshComp, class UAnimSequenceBase* Animation, float TotalDuration, const FAnimNotifyEventReference& EventReference) override;
-	virtual void NotifyEnd(class USkeletalMeshComponent* MeshComp, class UAnimSequenceBase* Animation, const FAnimNotifyEventReference& EventReference) override;
-	virtual void NotifyTick(USkeletalMeshComponent * MeshComp, UAnimSequenceBase * Animation, float FrameDeltaTime, const FAnimNotifyEventReference& EventReference) override;
+	NIAGARAANIMNOTIFIES_API virtual void NotifyBegin(class USkeletalMeshComponent* MeshComp, class UAnimSequenceBase* Animation, float TotalDuration, const FAnimNotifyEventReference& EventReference) override;
+	NIAGARAANIMNOTIFIES_API virtual void NotifyEnd(class USkeletalMeshComponent* MeshComp, class UAnimSequenceBase* Animation, const FAnimNotifyEventReference& EventReference) override;
+	NIAGARAANIMNOTIFIES_API virtual void NotifyTick(USkeletalMeshComponent * MeshComp, UAnimSequenceBase * Animation, float FrameDeltaTime, const FAnimNotifyEventReference& EventReference) override;
 
 	// Returns a 0 to 1 value for the progress of this component along the notify.
 	UFUNCTION(BlueprintCallable, Category = "AnimNotify")
-	float GetNotifyProgress(UMeshComponent* MeshComp);
+	NIAGARAANIMNOTIFIES_API float GetNotifyProgress(UMeshComponent* MeshComp);
 
 	UPROPERTY(EditAnywhere, Category = NotifyProgress, meta = (DisplayName = "Enable Normalized Notify Progress", ToolTip = "This send a 0-1 value of the normalized progress to the FX Component to the float User Parameter."))
 	bool bEnableNormalizedNotifyProgress = true;
diff --git a/Engine/Plugins/FX/Niagara/Source/NiagaraAnimNotifies/Public/AnimNotify_PlayNiagaraEffect.h b/Engine/Plugins/FX/Niagara/Source/NiagaraAnimNotifies/Public/AnimNotify_PlayNiagaraEffect.h
--- a/Engine/Plugins/FX/Niagara/Source/NiagaraAnimNotifies/Public/AnimNotify_PlayNiagaraEffect.h
+++ b/Engine/Plugins/FX/Niagara/Source/NiagaraAnimNotifies/Public/AnimNotify_PlayNiagaraEffect.h
@@ -15,30 +15,30 @@ class USkeletalMeshComponent;
 class UNiagaraSystem;
 class UFXSystemComponent;
 
-UCLASS(const, hidecategories = Object, collapsecategories, meta = (DisplayName = "Play Niagara Particle Effect"))
-class NIAGARAANIMNOTIFIES_API UAnimNotify_PlayNiagaraEffect : public UAnimNotify
+UCLASS(const, hidecategories = Object, collapsecategories, meta = (DisplayName = "Play Niagara Particle Effect"), MinimalAPI)
+class UAnimNotify_PlayNiagaraEffect : public UAnimNotify
 {
 	GENERATED_BODY()
 
 public:
 
-	UAnimNotify_PlayNiagaraEffect();
+	NIAGARAANIMNOTIFIES_API UAnimNotify_PlayNiagaraEffect();
 
 	// Begin UObject interface
-	virtual void PostLoad() override;
+	NIAGARAANIMNOTIFIES_API virtual void PostLoad() override;
 #if WITH_EDITOR
-	virtual void PostEditChangeProperty(struct FPropertyChangedEvent& PropertyChangedEvent) override;
+	NIAGARAANIMNOTIFIES_API virtual void PostEditChangeProperty(struct FPropertyChangedEvent& PropertyChangedEvent) override;
 #endif
 	// End UObject interface
 
 	// Begin UAnimNotify interface
-	virtual FString GetNotifyName_Implementation() const override;
+	NIAGARAANIMNOTIFIES_API virtual FString GetNotifyName_Implementation() const override;
 	
 	UE_DEPRECATED(5.0, "Please use the other Notify function instead")
-	virtual void Notify(USkeletalMeshComponent* MeshComp, UAnimSequenceBase* Animation) override;
-	virtual void Notify(USkeletalMeshComponent* MeshComp, UAnimSequenceBase* Animation, const FAnimNotifyEventReference& EventReference) override;
+	NIAGARAANIMNOTIFIES_API virtual void Notify(USkeletalMeshComponent* MeshComp, UAnimSequenceBase* Animation) override;
+	NIAGARAANIMNOTIFIES_API virtual void Notify(USkeletalMeshComponent* MeshComp, UAnimSequenceBase* Animation, const FAnimNotifyEventReference& EventReference) override;
 #if WITH_EDITOR
-	virtual void ValidateAssociatedAssets() override;
+	NIAGARAANIMNOTIFIES_API virtual void ValidateAssociatedAssets() override;
 #endif
 	// End UAnimNotify interface
 
@@ -64,7 +64,7 @@ public:
 
 	// Return FXSystemComponent created from SpawnEffect
 	UFUNCTION(BlueprintCallable, Category = "AnimNotify")
-	UFXSystemComponent* GetSpawnedEffect();
+	NIAGARAANIMNOTIFIES_API UFXSystemComponent* GetSpawnedEffect();
 
 protected:
 
@@ -75,7 +75,7 @@ protected:
 	FQuat RotationOffsetQuat;
 
 	// Spawns the NiagaraSystemComponent. Called from Notify.
-	virtual UFXSystemComponent* SpawnEffect(USkeletalMeshComponent* MeshComp, UAnimSequenceBase* Animation);
+	NIAGARAANIMNOTIFIES_API virtual UFXSystemComponent* SpawnEffect(USkeletalMeshComponent* MeshComp, UAnimSequenceBase* Animation);
 
 
 public:
diff --git a/Engine/Plugins/FX/Niagara/Source/NiagaraCore/Public/NiagaraCompileHash.h b/Engine/Plugins/FX/Niagara/Source/NiagaraCore/Public/NiagaraCompileHash.h
--- a/Engine/Plugins/FX/Niagara/Source/NiagaraCore/Public/NiagaraCompileHash.h
+++ b/Engine/Plugins/FX/Niagara/Source/NiagaraCore/Public/NiagaraCompileHash.h
@@ -6,7 +6,7 @@
 #include "NiagaraCompileHash.generated.h"
 
 USTRUCT()
-struct NIAGARACORE_API FNiagaraCompileHash
+struct FNiagaraCompileHash
 {
 	GENERATED_BODY()
 
@@ -27,18 +27,18 @@ struct NIAGARACORE_API FNiagaraCompileHash
 
 	bool operator==(const FNiagaraCompileHash& Other) const { return DataHash == Other.DataHash; }
 	bool operator!=(const FNiagaraCompileHash& Other) const { return DataHash != Other.DataHash; }
-	bool operator==(const FSHAHash& Other) const;
+	NIAGARACORE_API bool operator==(const FSHAHash& Other) const;
 	inline bool operator!=(const FSHAHash& Other) const { return !operator==(Other); }
 
-	bool ToSHAHash(FSHAHash& OutHash) const;
+	NIAGARACORE_API bool ToSHAHash(FSHAHash& OutHash) const;
 
-	bool IsValid() const;
+	NIAGARACORE_API bool IsValid() const;
 
-	uint32 GetTypeHash() const;
+	NIAGARACORE_API uint32 GetTypeHash() const;
 
-	const uint8* GetData() const;
+	NIAGARACORE_API const uint8* GetData() const;
 
-	void AppendString(FString& Out) const;
+	NIAGARACORE_API void AppendString(FString& Out) const;
 	
 	FString ToString() const
 	{
diff --git a/Engine/Plugins/FX/Niagara/Source/NiagaraCore/Public/NiagaraCore.h b/Engine/Plugins/FX/Niagara/Source/NiagaraCore/Public/NiagaraCore.h
--- a/Engine/Plugins/FX/Niagara/Source/NiagaraCore/Public/NiagaraCore.h
+++ b/Engine/Plugins/FX/Niagara/Source/NiagaraCore/Public/NiagaraCore.h
@@ -21,7 +21,7 @@ enum class ENiagaraIterationSource : uint8
 
 /** A utility class allowing for references to FNiagaraVariableBase outside of the Niagara module. */
 USTRUCT()
-struct NIAGARACORE_API FNiagaraVariableCommonReference
+struct FNiagaraVariableCommonReference
 {
 	GENERATED_USTRUCT_BODY()
 
@@ -31,7 +31,7 @@ struct NIAGARACORE_API FNiagaraVariableCommonReference
 	UPROPERTY()
 	TObjectPtr<UObject> UnderlyingType;
 
-	bool Serialize(FArchive& Ar);
+	NIAGARACORE_API bool Serialize(FArchive& Ar);
 	friend bool operator<<(FArchive& Ar, FNiagaraVariableCommonReference& VariableReference);
 
 	bool operator==(const FNiagaraVariableCommonReference& Other)const
diff --git a/Engine/Plugins/FX/Niagara/Source/NiagaraCore/Public/NiagaraDataInterfaceBase.h b/Engine/Plugins/FX/Niagara/Source/NiagaraCore/Public/NiagaraDataInterfaceBase.h
--- a/Engine/Plugins/FX/Niagara/Source/NiagaraCore/Public/NiagaraDataInterfaceBase.h
+++ b/Engine/Plugins/FX/Niagara/Source/NiagaraCore/Public/NiagaraDataInterfaceBase.h
@@ -40,8 +40,8 @@ struct FNiagaraDataInterfaceParametersCS
 //////////////////////////////////////////////////////////////////////////
 
 /** Base class for all Niagara data interfaces. */
-UCLASS(abstract, EditInlineNew)
-class NIAGARACORE_API UNiagaraDataInterfaceBase : public UNiagaraMergeable
+UCLASS(abstract, EditInlineNew, MinimalAPI)
+class UNiagaraDataInterfaceBase : public UNiagaraMergeable
 {
 	GENERATED_UCLASS_BODY()
 
diff --git a/Engine/Plugins/FX/Niagara/Source/NiagaraCore/Public/NiagaraMergeable.h b/Engine/Plugins/FX/Niagara/Source/NiagaraCore/Public/NiagaraMergeable.h
--- a/Engine/Plugins/FX/Niagara/Source/NiagaraCore/Public/NiagaraMergeable.h
+++ b/Engine/Plugins/FX/Niagara/Source/NiagaraCore/Public/NiagaraMergeable.h
@@ -4,8 +4,8 @@
 
 #include "NiagaraMergeable.generated.h"
 
-UCLASS()
-class NIAGARACORE_API UNiagaraMergeable : public UObject
+UCLASS(MinimalAPI)
+class UNiagaraMergeable : public UObject
 {
 	GENERATED_BODY()
 
@@ -15,19 +15,19 @@ public:
 #endif
 
 public:
-	UNiagaraMergeable();
+	NIAGARACORE_API UNiagaraMergeable();
 
 #if WITH_EDITOR
-	virtual void PostEditChangeProperty(struct FPropertyChangedEvent& PropertyChangedEvent);
+	NIAGARACORE_API virtual void PostEditChangeProperty(struct FPropertyChangedEvent& PropertyChangedEvent);
 
-	bool Equals(const UNiagaraMergeable* Other);
+	NIAGARACORE_API bool Equals(const UNiagaraMergeable* Other);
 
-	FOnChanged& OnChanged();
+	NIAGARACORE_API FOnChanged& OnChanged();
 
-	FGuid GetMergeId() const;
+	NIAGARACORE_API FGuid GetMergeId() const;
 
 protected:
-	UNiagaraMergeable* StaticDuplicateWithNewMergeIdInternal(UObject* InOuter) const;
+	NIAGARACORE_API UNiagaraMergeable* StaticDuplicateWithNewMergeIdInternal(UObject* InOuter) const;
 #endif
 
 private:
@@ -39,4 +39,4 @@ private:
 	UPROPERTY()
 	FGuid MergeId;
 #endif
-};
\ No newline at end of file
+};
diff --git a/Engine/Plugins/FX/Niagara/Source/NiagaraEditor/Internal/SNiagaraAssetPickerList.h b/Engine/Plugins/FX/Niagara/Source/NiagaraEditor/Internal/SNiagaraAssetPickerList.h
--- a/Engine/Plugins/FX/Niagara/Source/NiagaraEditor/Internal/SNiagaraAssetPickerList.h
+++ b/Engine/Plugins/FX/Niagara/Source/NiagaraEditor/Internal/SNiagaraAssetPickerList.h
@@ -56,7 +56,7 @@ private:
 	bool bAddLibraryOnlyCheckbox;
 };
 
-class NIAGARAEDITOR_API SNiagaraAssetPickerList : public SCompoundWidget
+class SNiagaraAssetPickerList : public SCompoundWidget
 {
 public:
 	DECLARE_DELEGATE_OneParam(FOnTemplateAssetActivated, const FAssetData&);
@@ -94,17 +94,17 @@ public:
 		SLATE_ARGUMENT(EItemSelectorClickActivateMode, ClickActivateMode);
 	SLATE_END_ARGS();
 
-	void Construct(const FArguments& InArgs, UClass* AssetClass);
+	NIAGARAEDITOR_API void Construct(const FArguments& InArgs, UClass* AssetClass);
 
-	virtual void Tick(const FGeometry& AllottedGeometry, const double InCurrentTime, const float InDeltaTime) override;
+	NIAGARAEDITOR_API virtual void Tick(const FGeometry& AllottedGeometry, const double InCurrentTime, const float InDeltaTime) override;
 
-	TArray<FAssetData> GetSelectedAssets() const;
+	NIAGARAEDITOR_API TArray<FAssetData> GetSelectedAssets() const;
 
-	void RefreshAll() const;
+	NIAGARAEDITOR_API void RefreshAll() const;
 
-	void ExpandTree();
+	NIAGARAEDITOR_API void ExpandTree();
 
-	TSharedRef<SWidget> GetSearchBox() const;
+	NIAGARAEDITOR_API TSharedRef<SWidget> GetSearchBox() const;
 
 private:
 	TArray<FAssetData> GetAssetDataForSelector(UClass* AssetClass);
diff --git a/Engine/Plugins/FX/Niagara/Source/NiagaraEditor/Private/GlobalDistanceFieldCapture.h b/Engine/Plugins/FX/Niagara/Source/NiagaraEditor/Private/GlobalDistanceFieldCapture.h
--- a/Engine/Plugins/FX/Niagara/Source/NiagaraEditor/Private/GlobalDistanceFieldCapture.h
+++ b/Engine/Plugins/FX/Niagara/Source/NiagaraEditor/Private/GlobalDistanceFieldCapture.h
@@ -19,16 +19,16 @@ DECLARE_DELEGATE_ThreeParams(FOnDistanceFieldCaptureComplete, FVector /*Extents*
  * Utility interface used by an editor console command to capture the global distance field and store it into a volume
  * texture selected in the Content Browser
  */
-class NIAGARAEDITOR_API FGlobalDistanceFieldCapture
+class FGlobalDistanceFieldCapture
 {
 public:
-	virtual ~FGlobalDistanceFieldCapture();
+	NIAGARAEDITOR_API virtual ~FGlobalDistanceFieldCapture();
 	
 	/**	Requests a capture. Will overwrite the volume texture selected in the Content Browser or create a new one. */
-	static void Request(UVolumeTexture* VolumeTexture, bool bRangeCompress, FOnDistanceFieldCaptureComplete& InCompletionDelegate);
+	static NIAGARAEDITOR_API void Request(UVolumeTexture* VolumeTexture, bool bRangeCompress, FOnDistanceFieldCaptureComplete& InCompletionDelegate);
 
 	/**	Requests a capture at the specified camera location. Will overwrite the volume texture selected in the Content Browser or create a new one. */
-	static void Request(UVolumeTexture* VolumeTexture, bool bRangeCompress, const FVector& CameraPos, FOnDistanceFieldCaptureComplete& InCompletionDelegate);
+	static NIAGARAEDITOR_API void Request(UVolumeTexture* VolumeTexture, bool bRangeCompress, const FVector& CameraPos, FOnDistanceFieldCaptureComplete& InCompletionDelegate);
 
 private:
 	TWeakObjectPtr<UVolumeTexture> VolumeTex;
diff --git a/Engine/Plugins/FX/Niagara/Source/NiagaraEditor/Private/NiagaraCompiler.h b/Engine/Plugins/FX/Niagara/Source/NiagaraEditor/Private/NiagaraCompiler.h
--- a/Engine/Plugins/FX/Niagara/Source/NiagaraEditor/Private/NiagaraCompiler.h
+++ b/Engine/Plugins/FX/Niagara/Source/NiagaraEditor/Private/NiagaraCompiler.h
@@ -29,7 +29,7 @@ struct FNiagaraCompilerJob
 	}
 };
 
-class NIAGARAEDITOR_API FHlslNiagaraCompiler : public INiagaraCompiler
+class FHlslNiagaraCompiler : public INiagaraCompiler
 {
 protected:	
 	/** Captures information about a script compile. */
@@ -37,24 +37,24 @@ protected:
 
 public:
 
-	FHlslNiagaraCompiler();
+	NIAGARAEDITOR_API FHlslNiagaraCompiler();
 	virtual ~FHlslNiagaraCompiler() {}
 
 	//Begin INiagaraCompiler Interface
 	UE_DEPRECATED(5.4, "Please update to supply the GroupName directly as this will be removed in a future version.")
-	virtual int32 CompileScript(const class FNiagaraCompileRequestData* InCompileRequest, const FNiagaraCompileOptions& InOptions, const FNiagaraTranslateResults& InTranslateResults, FNiagaraTranslatorOutput* TranslatorOutput, FString& TranslatedHLSL) override;
+	NIAGARAEDITOR_API virtual int32 CompileScript(const class FNiagaraCompileRequestData* InCompileRequest, const FNiagaraCompileOptions& InOptions, const FNiagaraTranslateResults& InTranslateResults, FNiagaraTranslatorOutput* TranslatorOutput, FString& TranslatedHLSL) override;
 
-	virtual int32 CompileScript(const FStringView GroupName, const FNiagaraCompileOptions& InOptions, const FNiagaraTranslateResults& InTranslateResults, const FNiagaraTranslatorOutput& TranslatorOutput, const FString& TranslatedHLSL) override;
-	virtual TOptional<FNiagaraCompileResults> GetCompileResult(int32 JobID, bool bWait = false) override;
+	NIAGARAEDITOR_API virtual int32 CompileScript(const FStringView GroupName, const FNiagaraCompileOptions& InOptions, const FNiagaraTranslateResults& InTranslateResults, const FNiagaraTranslatorOutput& TranslatorOutput, const FString& TranslatedHLSL) override;
+	NIAGARAEDITOR_API virtual TOptional<FNiagaraCompileResults> GetCompileResult(int32 JobID, bool bWait = false) override;
 
-	virtual void Error(FText ErrorText) override;
-	virtual void Warning(FText WarningText) override;
+	NIAGARAEDITOR_API virtual void Error(FText ErrorText) override;
+	NIAGARAEDITOR_API virtual void Warning(FText WarningText) override;
 
 private:
 	TUniquePtr<FNiagaraCompilerJob> CompilationJob;
 
-	void DumpDebugInfo(const FNiagaraCompileResults& CompileResult, const FShaderCompilerInput& Input, bool bGPUScript);
+	NIAGARAEDITOR_API void DumpDebugInfo(const FNiagaraCompileResults& CompileResult, const FShaderCompilerInput& Input, bool bGPUScript);
 
 	/** SCW doesn't have access to the VM op code names so we do a fixup pass to make these human readable after we get the data back from SCW. */
-	void FixupVMAssembly(FString& Asm);
+	NIAGARAEDITOR_API void FixupVMAssembly(FString& Asm);
 };
diff --git a/Engine/Plugins/FX/Niagara/Source/NiagaraEditor/Private/NiagaraConvertInPlaceEmitterAndSystemState.h b/Engine/Plugins/FX/Niagara/Source/NiagaraEditor/Private/NiagaraConvertInPlaceEmitterAndSystemState.h
--- a/Engine/Plugins/FX/Niagara/Source/NiagaraEditor/Private/NiagaraConvertInPlaceEmitterAndSystemState.h
+++ b/Engine/Plugins/FX/Niagara/Source/NiagaraEditor/Private/NiagaraConvertInPlaceEmitterAndSystemState.h
@@ -7,10 +7,10 @@
 #include "NiagaraNodeFunctionCall.h"
 #include "NiagaraConvertInPlaceEmitterAndSystemState.generated.h"
 
-UCLASS()
-class NIAGARAEDITOR_API  UNiagaraConvertInPlaceEmitterAndSystemState : public UNiagaraConvertInPlaceUtilityBase
+UCLASS(MinimalAPI)
+class  UNiagaraConvertInPlaceEmitterAndSystemState : public UNiagaraConvertInPlaceUtilityBase
 {
 	GENERATED_BODY()
 public:
-	virtual bool Convert(UNiagaraScript* InOldScript, UNiagaraClipboardContent* InOldClipboardContent, UNiagaraScript* InNewScript, UNiagaraStackFunctionInputCollection* InInputCollection, UNiagaraClipboardContent* InNewClipboardContent, UNiagaraNodeFunctionCall* InCallingNode, FText& OutMessage) override;
-};
\ No newline at end of file
+	NIAGARAEDITOR_API virtual bool Convert(UNiagaraScript* InOldScript, UNiagaraClipboardContent* InOldClipboardContent, UNiagaraScript* InNewScript, UNiagaraStackFunctionInputCollection* InInputCollection, UNiagaraClipboardContent* InNewClipboardContent, UNiagaraNodeFunctionCall* InCallingNode, FText& OutMessage) override;
+};
diff --git a/Engine/Plugins/FX/Niagara/Source/NiagaraEditor/Private/NiagaraHlslTranslator.h b/Engine/Plugins/FX/Niagara/Source/NiagaraEditor/Private/NiagaraHlslTranslator.h
--- a/Engine/Plugins/FX/Niagara/Source/NiagaraEditor/Private/NiagaraHlslTranslator.h
+++ b/Engine/Plugins/FX/Niagara/Source/NiagaraEditor/Private/NiagaraHlslTranslator.h
@@ -221,7 +221,7 @@ struct FunctionNodeStackEntry
 	TSet<FString> CulledFunctionNames;
 };
 
-class NIAGARAEDITOR_API FNiagaraHlslTranslator : public INiagaraHlslTranslator
+class FNiagaraHlslTranslator : public INiagaraHlslTranslator
 {
 public:
 
@@ -268,7 +268,7 @@ protected:
 	TMap<FNiagaraDataSetID, TMap<int32, FDataSetAccessInfo>> DataSetWriteInfo[(int32)ENiagaraDataSetAccessMode::Num];
 	TMap<FNiagaraDataSetID, int32> DataSetWriteConditionalInfo[(int32)ENiagaraDataSetAccessMode::Num];
 
-	FString GetDataSetAccessSymbol(FNiagaraDataSetID DataSet, int32 IndexChunk, bool bRead);
+	NIAGARAEDITOR_API FString GetDataSetAccessSymbol(FNiagaraDataSetID DataSet, int32 IndexChunk, bool bRead);
 	FORCEINLINE FNiagaraDataSetID GetInstanceDataSetID()const { return FNiagaraDataSetID(TEXT("DataInstance"), ENiagaraDataSetType::ParticleData); }
 	FORCEINLINE FNiagaraDataSetID GetSystemEngineDataSetID()const { return FNiagaraDataSetID(TEXT("Engine"), ENiagaraDataSetType::ParticleData); }
 	FORCEINLINE FNiagaraDataSetID GetSystemUserDataSetID()const { return FNiagaraDataSetID(TEXT("User"), ENiagaraDataSetType::ParticleData); }
@@ -289,15 +289,15 @@ protected:
 	TMap<FNiagaraFunctionSignature, TArray<FName>> FunctionStageReadTargets;
 	TArray<TArray<FName>> ActiveStageReadTargets;
 
-	void GenerateFunctionCall(ENiagaraScriptUsage ScriptUsage, FNiagaraFunctionSignature& FunctionSignature, TArrayView<const int32> Inputs, TArray<int32>& Outputs);
-	FString GetFunctionIncludeStatement(const FNiagaraCustomHlslInclude& Include) const;
-	FString GetFunctionSignature(const FNiagaraFunctionSignature& Sig);
+	NIAGARAEDITOR_API void GenerateFunctionCall(ENiagaraScriptUsage ScriptUsage, FNiagaraFunctionSignature& FunctionSignature, TArrayView<const int32> Inputs, TArray<int32>& Outputs);
+	NIAGARAEDITOR_API FString GetFunctionIncludeStatement(const FNiagaraCustomHlslInclude& Include) const;
+	NIAGARAEDITOR_API FString GetFunctionSignature(const FNiagaraFunctionSignature& Sig);
 
-	void WriteDataSetContextVars(TMap<FNiagaraDataSetID, TMap<int32, FDataSetAccessInfo>>& DataSetAccessInfo, bool bRead, FString& OutHLSLOutput);
-	void WriteDataSetStructDeclarations(TMap<FNiagaraDataSetID, TMap<int32, FDataSetAccessInfo>>& DataSetAccessInfo, bool bRead, FString& OutHLSLOutput);
-	void DecomposeVariableAccess(UStruct* Struct, bool bRead, FString IndexSymbol, FString HLSLString);
+	NIAGARAEDITOR_API void WriteDataSetContextVars(TMap<FNiagaraDataSetID, TMap<int32, FDataSetAccessInfo>>& DataSetAccessInfo, bool bRead, FString& OutHLSLOutput);
+	NIAGARAEDITOR_API void WriteDataSetStructDeclarations(TMap<FNiagaraDataSetID, TMap<int32, FDataSetAccessInfo>>& DataSetAccessInfo, bool bRead, FString& OutHLSLOutput);
+	NIAGARAEDITOR_API void DecomposeVariableAccess(UStruct* Struct, bool bRead, FString IndexSymbol, FString HLSLString);
 
-	FString GetUniqueSymbolName(FName BaseName);
+	NIAGARAEDITOR_API FString GetUniqueSymbolName(FName BaseName);
 
 	/** Stack of all function params. */
 	struct FFunctionContext
@@ -315,37 +315,37 @@ protected:
 	};
 	TArray<FFunctionContext> FunctionContextStack;
 	const FFunctionContext* FunctionCtx()const { return FunctionContextStack.Num() > 0 ? &FunctionContextStack.Last() : nullptr; }
-	FString GetCallstack();
-	TArray<FGuid> GetCallstackGuids();
+	NIAGARAEDITOR_API FString GetCallstack();
+	NIAGARAEDITOR_API TArray<FGuid> GetCallstackGuids();
 
-	void EnterStatsScope(FNiagaraStatScope StatScope);
-	void ExitStatsScope();
-	void EnterStatsScope(FNiagaraStatScope StatScope, FString& OutHlsl);
-	void ExitStatsScope(FString& OutHlsl);
+	NIAGARAEDITOR_API void EnterStatsScope(FNiagaraStatScope StatScope);
+	NIAGARAEDITOR_API void ExitStatsScope();
+	NIAGARAEDITOR_API void EnterStatsScope(FNiagaraStatScope StatScope, FString& OutHlsl);
+	NIAGARAEDITOR_API void ExitStatsScope(FString& OutHlsl);
 
-	FString GeneratedConstantString(float Constant);
-	FString GeneratedConstantString(FVector4 Constant);
+	NIAGARAEDITOR_API FString GeneratedConstantString(float Constant);
+	NIAGARAEDITOR_API FString GeneratedConstantString(FVector4 Constant);
 
 	/* Add a chunk that is not written to the source, only used as a source chunk for others. */
-	int32 AddSourceChunk(FString SymbolName, const FNiagaraTypeDefinition& Type, bool bSanitize = true);
+	NIAGARAEDITOR_API int32 AddSourceChunk(FString SymbolName, const FNiagaraTypeDefinition& Type, bool bSanitize = true);
 
 	/** Add a chunk defining a uniform value. */
-	int32 AddUniformChunk(FString SymbolName, const FNiagaraVariable& InVariable, ENiagaraCodeChunkMode ChunkMode, bool AddPadding);
+	NIAGARAEDITOR_API int32 AddUniformChunk(FString SymbolName, const FNiagaraVariable& InVariable, ENiagaraCodeChunkMode ChunkMode, bool AddPadding);
 
 	/* Add a chunk that is written to the body of the shader code. */
-	int32 AddBodyChunk(FString SymbolName, FString Definition, const FNiagaraTypeDefinition& Type, TArray<int32>& SourceChunks, bool bDecl = true, bool bIsTerminated = true);
-	int32 AddBodyChunk(FString SymbolName, FString Definition, const FNiagaraTypeDefinition& Type, int32 SourceChunk, bool bDecl = true, bool bIsTerminated = true);
-	int32 AddBodyChunk(FString SymbolName, FString Definition, const FNiagaraTypeDefinition& Type, bool bDecl = true, bool bIsTerminated = true);
-	int32 AddBodyComment(const FString& Comment);
-	int32 AddBodyChunk(const FString& Definition);
+	NIAGARAEDITOR_API int32 AddBodyChunk(FString SymbolName, FString Definition, const FNiagaraTypeDefinition& Type, TArray<int32>& SourceChunks, bool bDecl = true, bool bIsTerminated = true);
+	NIAGARAEDITOR_API int32 AddBodyChunk(FString SymbolName, FString Definition, const FNiagaraTypeDefinition& Type, int32 SourceChunk, bool bDecl = true, bool bIsTerminated = true);
+	NIAGARAEDITOR_API int32 AddBodyChunk(FString SymbolName, FString Definition, const FNiagaraTypeDefinition& Type, bool bDecl = true, bool bIsTerminated = true);
+	NIAGARAEDITOR_API int32 AddBodyComment(const FString& Comment);
+	NIAGARAEDITOR_API int32 AddBodyChunk(const FString& Definition);
 
-	FString GetUniqueEmitterName() const;
-	void ConvertCompileInfoToParamInfo(const FNiagaraScriptDataInterfaceCompileInfo& InCompileInfo, FNiagaraDataInterfaceGPUParamInfo& OutGPUParamInfo, TArray<FNiagaraFunctionSignature>& GeneratedFunctionSignatures);
+	NIAGARAEDITOR_API FString GetUniqueEmitterName() const;
+	NIAGARAEDITOR_API void ConvertCompileInfoToParamInfo(const FNiagaraScriptDataInterfaceCompileInfo& InCompileInfo, FNiagaraDataInterfaceGPUParamInfo& OutGPUParamInfo, TArray<FNiagaraFunctionSignature>& GeneratedFunctionSignatures);
 
-	FString GetFunctionDefinitions();
+	NIAGARAEDITOR_API FString GetFunctionDefinitions();
 
 public:
-	FNiagaraHlslTranslator();
+	NIAGARAEDITOR_API FNiagaraHlslTranslator();
 	virtual ~FNiagaraHlslTranslator() = default;
 
 	virtual const FNiagaraTranslatorOutput& GetTranslateOutput() const override { return CompilationOutput; }
@@ -353,65 +353,65 @@ public:
 
 	bool DisableDebugSwitches() const { return TranslationOptions.bDisableDebugSwitches; }
 
-	bool IsCompileOptionDefined(const TCHAR* InDefineStr);
+	NIAGARAEDITOR_API bool IsCompileOptionDefined(const TCHAR* InDefineStr);
 
-	int32 GetAttribute(const FNiagaraVariable& Attribute);
+	NIAGARAEDITOR_API int32 GetAttribute(const FNiagaraVariable& Attribute);
 
-	int32 GetConstant(const FNiagaraVariable& Constant, FString* DebugOutputValue = nullptr);
+	NIAGARAEDITOR_API int32 GetConstant(const FNiagaraVariable& Constant, FString* DebugOutputValue = nullptr);
 
-	int32 ParameterMapForInnerIndex() const;
+	NIAGARAEDITOR_API int32 ParameterMapForInnerIndex() const;
 
-	void DefinePreviousParametersFunction(FString& HlslOutput, TArray<TArray<FNiagaraVariable>>& DataSetVariables, TMap<FNiagaraDataSetID, int32>& DataSetReads, TMap<FNiagaraDataSetID, int32>& DataSetWrites);
+	NIAGARAEDITOR_API void DefinePreviousParametersFunction(FString& HlslOutput, TArray<TArray<FNiagaraVariable>>& DataSetVariables, TMap<FNiagaraDataSetID, int32>& DataSetReads, TMap<FNiagaraDataSetID, int32>& DataSetWrites);
 
-	void DefineDataSetReadFunction(FString& HlslOutput, TArray<FNiagaraDataSetID>& ReadDataSets);
-	void DefineDataSetWriteFunction(FString& HlslOutput, TArray<FNiagaraDataSetProperties>& WriteDataSets, TArray<int32>& WriteConditionVarIndices);
-	void DefineMain(FString& HLSLOutput, const TArray<TArray<FNiagaraVariable>>& DataSetVariables, const TMap<FNiagaraDataSetID, int32>& DataSetReads, const TMap<FNiagaraDataSetID, int32>& DataSetWrites);
+	NIAGARAEDITOR_API void DefineDataSetReadFunction(FString& HlslOutput, TArray<FNiagaraDataSetID>& ReadDataSets);
+	NIAGARAEDITOR_API void DefineDataSetWriteFunction(FString& HlslOutput, TArray<FNiagaraDataSetProperties>& WriteDataSets, TArray<int32>& WriteConditionVarIndices);
+	NIAGARAEDITOR_API void DefineMain(FString& HLSLOutput, const TArray<TArray<FNiagaraVariable>>& DataSetVariables, const TMap<FNiagaraDataSetID, int32>& DataSetReads, const TMap<FNiagaraDataSetID, int32>& DataSetWrites);
 
-	void DefineDataSetVariableReads(FString& HLSLOutput, const FNiagaraDataSetID& Id, int32 DataSetIndex, const TArray<FNiagaraVariable>& ReadVars);
-	void DefineDataInterfaceHLSL(FString& HlslOutput);
-	void DefineExternalFunctionsHLSL(FString& HlslOutput);
+	NIAGARAEDITOR_API void DefineDataSetVariableReads(FString& HLSLOutput, const FNiagaraDataSetID& Id, int32 DataSetIndex, const TArray<FNiagaraVariable>& ReadVars);
+	NIAGARAEDITOR_API void DefineDataInterfaceHLSL(FString& HlslOutput);
+	NIAGARAEDITOR_API void DefineExternalFunctionsHLSL(FString& HlslOutput);
 
 	// Format string should have up to 5 entries, {{0} = Computed Variable Suffix, {1} = Float or Int, {2} = Data Set Index, {3} = Register Index, {4} Default value for that type.
-	void GatherVariableForDataSetAccess(const FNiagaraVariable& Variable, FString Format, int32& RegisterIdxInt, int32& RegisterIdxFloat, int32& RegisterIdxHalf, int32 DataSetIndex, FString InstanceIdxSymbol, FString& HlslOutput, bool bWriteHLSL = true);
-	void GatherComponentsForDataSetAccess(UScriptStruct* Struct, FString VariableSymbol, bool bMatrixRoot, TArray<FString>& Components, TArray<ENiagaraBaseTypes>& Types);
+	NIAGARAEDITOR_API void GatherVariableForDataSetAccess(const FNiagaraVariable& Variable, FString Format, int32& RegisterIdxInt, int32& RegisterIdxFloat, int32& RegisterIdxHalf, int32 DataSetIndex, FString InstanceIdxSymbol, FString& HlslOutput, bool bWriteHLSL = true);
+	NIAGARAEDITOR_API void GatherComponentsForDataSetAccess(UScriptStruct* Struct, FString VariableSymbol, bool bMatrixRoot, TArray<FString>& Components, TArray<ENiagaraBaseTypes>& Types);
 
-	void EnterFunctionCallNode(const TSet<FName>& UnusedInputs);
-	void ExitFunctionCallNode();
-	bool IsFunctionVariableCulledFromCompilation(const FName& InputName) const;
+	NIAGARAEDITOR_API void EnterFunctionCallNode(const TSet<FName>& UnusedInputs);
+	NIAGARAEDITOR_API void ExitFunctionCallNode();
+	NIAGARAEDITOR_API bool IsFunctionVariableCulledFromCompilation(const FName& InputName) const;
 
-	bool GetFunctionParameter(const FNiagaraVariable& Parameter, int32& OutParam)const;
-	int32 GetUniqueCallerID();
+	NIAGARAEDITOR_API bool GetFunctionParameter(const FNiagaraVariable& Parameter, int32& OutParam)const;
+	NIAGARAEDITOR_API int32 GetUniqueCallerID();
 
-	bool CanReadAttributes()const;
-	ENiagaraScriptUsage GetTargetUsage() const;
-	FGuid GetTargetUsageId() const;
+	NIAGARAEDITOR_API bool CanReadAttributes()const;
+	NIAGARAEDITOR_API ENiagaraScriptUsage GetTargetUsage() const;
+	NIAGARAEDITOR_API FGuid GetTargetUsageId() const;
 	ENiagaraSimTarget GetSimulationTarget() const
 	{
 		return CompilationTarget;
 	}
 
-	static TArray<FName> ConditionPropertyPath(const FNiagaraTypeDefinition& Type, const TArray<FName>& InPath);
-	static FString GetSanitizedSymbolName(FStringView SymbolName, bool bCollapseNamespaces = false);
-	static FString GetSanitizedDIFunctionName(const FString& FunctionName);
-	static FString GetSanitizedFunctionNameSuffix(FString Name);
+	static NIAGARAEDITOR_API TArray<FName> ConditionPropertyPath(const FNiagaraTypeDefinition& Type, const TArray<FName>& InPath);
+	static NIAGARAEDITOR_API FString GetSanitizedSymbolName(FStringView SymbolName, bool bCollapseNamespaces = false);
+	static NIAGARAEDITOR_API FString GetSanitizedDIFunctionName(const FString& FunctionName);
+	static NIAGARAEDITOR_API FString GetSanitizedFunctionNameSuffix(FString Name);
 
 	/** Replaces all non-ascii characters with a "ASCXXX" string, where XXX is their int value */
-	static FString ConvertToAsciiString(FString Name);
+	static NIAGARAEDITOR_API FString ConvertToAsciiString(FString Name);
 
-	bool AddStructToDefinitionSet(const FNiagaraTypeDefinition& TypeDef);
+	NIAGARAEDITOR_API bool AddStructToDefinitionSet(const FNiagaraTypeDefinition& TypeDef);
 
-	static FString GetFunctionSignatureSymbol(const FNiagaraFunctionSignature& Sig);
+	static NIAGARAEDITOR_API FString GetFunctionSignatureSymbol(const FNiagaraFunctionSignature& Sig);
 
-	static FName GetDataInterfaceName(FName BaseName, const FString& UniqueEmitterName, bool bIsParameterMapDataInterface);
+	static NIAGARAEDITOR_API FName GetDataInterfaceName(FName BaseName, const FString& UniqueEmitterName, bool bIsParameterMapDataInterface);
 
 
 
-	static FString GenerateFunctionHlslPrototype(FStringView InVariableName, const FNiagaraFunctionSignature& FunctionSignature);
+	static NIAGARAEDITOR_API FString GenerateFunctionHlslPrototype(FStringView InVariableName, const FNiagaraFunctionSignature& FunctionSignature);
 
-	void SetConstantByStaticVariable(int32& OutValue, const FNiagaraVariable& Var, FString* DebugString = nullptr);
-	void SetConstantByStaticVariable(FNiagaraVariable& OutValue, const FNiagaraVariable& Var, FString* DebugString = nullptr);
+	NIAGARAEDITOR_API void SetConstantByStaticVariable(int32& OutValue, const FNiagaraVariable& Var, FString* DebugString = nullptr);
+	NIAGARAEDITOR_API void SetConstantByStaticVariable(FNiagaraVariable& OutValue, const FNiagaraVariable& Var, FString* DebugString = nullptr);
 
-	bool IsEventSpawnScript()const;
+	NIAGARAEDITOR_API bool IsEventSpawnScript()const;
 
 	// collection of functions that rely on the graph implementation and so must be relegated to the TNiagaraHlslTranslator
 	virtual void DefineDataSetVariableWrites(FString& HlslOutput, const FNiagaraDataSetID& Id, int32 DataSetIndex, const TArray<FNiagaraVariable>& WriteVars) = 0;
@@ -424,48 +424,48 @@ public:
 
 
 protected:
-	bool ShouldConsiderTargetParameterMap(ENiagaraScriptUsage InUsage) const;
+	NIAGARAEDITOR_API bool ShouldConsiderTargetParameterMap(ENiagaraScriptUsage InUsage) const;
 
-	void HandleNamespacedExternalVariablesToDataSetRead(TArray<FNiagaraVariable>& InDataSetVars, FString InNamespaceStr);
+	NIAGARAEDITOR_API void HandleNamespacedExternalVariablesToDataSetRead(TArray<FNiagaraVariable>& InDataSetVars, FString InNamespaceStr);
 
 	// For GPU simulations we have to special case some variables and pass them view shader parameters rather than the uniform buffer as they vary from CPU simulations
-	bool IsVariableInUniformBuffer(const FNiagaraVariable& Variable) const;
+	NIAGARAEDITOR_API bool IsVariableInUniformBuffer(const FNiagaraVariable& Variable) const;
 
-	FString ComputeMatrixColumnAccess(const FString& Name);
-	FString ComputeMatrixRowAccess(const FString& Name);
+	NIAGARAEDITOR_API FString ComputeMatrixColumnAccess(const FString& Name);
+	NIAGARAEDITOR_API FString ComputeMatrixRowAccess(const FString& Name);
 
 	// Add a raw float constant chunk
-	int32 GetConstantDirect(float InValue);
-	int32 GetConstantDirect(bool InValue);
-	int32 GetConstantDirect(int InValue);
+	NIAGARAEDITOR_API int32 GetConstantDirect(float InValue);
+	NIAGARAEDITOR_API int32 GetConstantDirect(bool InValue);
+	NIAGARAEDITOR_API int32 GetConstantDirect(int InValue);
 
-	FNiagaraTypeDefinition GetChildType(const FNiagaraTypeDefinition& BaseType, const FName& PropertyName);
-	FString NamePathToString(const FString& Prefix, const FNiagaraTypeDefinition& RootType, const TArray<FName>& NamePath);
-	FString GenerateAssignment(const FNiagaraTypeDefinition& SrcType, const TArray<FName>& SrcPath, const FNiagaraTypeDefinition& DestType, const TArray<FName>& DestPath);
+	NIAGARAEDITOR_API FNiagaraTypeDefinition GetChildType(const FNiagaraTypeDefinition& BaseType, const FName& PropertyName);
+	NIAGARAEDITOR_API FString NamePathToString(const FString& Prefix, const FNiagaraTypeDefinition& RootType, const TArray<FName>& NamePath);
+	NIAGARAEDITOR_API FString GenerateAssignment(const FNiagaraTypeDefinition& SrcType, const TArray<FName>& SrcPath, const FNiagaraTypeDefinition& DestType, const TArray<FName>& DestPath);
 
 	//Generates the code for the passed chunk.
-	FString GetCode(FNiagaraCodeChunk& Chunk);
-	FString GetCode(int32 ChunkIdx);
+	NIAGARAEDITOR_API FString GetCode(FNiagaraCodeChunk& Chunk);
+	NIAGARAEDITOR_API FString GetCode(int32 ChunkIdx);
 	//Retreives the code for this chunk being used as a source for another chunk
-	FString GetCodeAsSource(int32 ChunkIdx);
+	NIAGARAEDITOR_API FString GetCodeAsSource(int32 ChunkIdx);
 
 	// Generate a structure initializer string
 	// Returns true if we generated the structure successfully or false if we encounter something we could not handle
-	bool GenerateStructInitializer(TStringBuilder<128>& InitializerString, UStruct* UserDefinedStruct, const void* StructData, int32 ByteOffset = 0);
+	NIAGARAEDITOR_API bool GenerateStructInitializer(TStringBuilder<128>& InitializerString, UStruct* UserDefinedStruct, const void* StructData, int32 ByteOffset = 0);
 	// Convert a variable with actual data into a constant string
-	FString GenerateConstantString(const FNiagaraVariable& Constant);
+	NIAGARAEDITOR_API FString GenerateConstantString(const FNiagaraVariable& Constant);
 
 	// Takes the current script state (interpolated or not) and determines the correct context variable.
-	FString GetParameterMapInstanceName(int32 ParamMapHistoryIdx);
+	NIAGARAEDITOR_API FString GetParameterMapInstanceName(int32 ParamMapHistoryIdx);
 
 	// Register an attribute in its non namespaced form
 	virtual bool ParameterMapRegisterNamespaceAttributeVariable(const FNiagaraVariable& InVariable, int32 InParamMapHistoryIdx, int32& Output) = 0;
 
-	bool ShouldInterpolateParameter(const FNiagaraVariable& Parameter);
+	NIAGARAEDITOR_API bool ShouldInterpolateParameter(const FNiagaraVariable& Parameter);
 
-	bool IsBulkSystemScript() const;
-	bool IsSpawnScript() const;
-	bool RequiresInterpolation() const;
+	NIAGARAEDITOR_API bool IsBulkSystemScript() const;
+	NIAGARAEDITOR_API bool IsSpawnScript() const;
+	NIAGARAEDITOR_API bool RequiresInterpolation() const;
 
 	template<typename T>
 	void BuildConstantBuffer(ENiagaraCodeChunkMode ChunkMode);
diff --git a/Engine/Plugins/FX/Niagara/Source/NiagaraEditor/Private/NiagaraVersionMetaData.h b/Engine/Plugins/FX/Niagara/Source/NiagaraEditor/Private/NiagaraVersionMetaData.h
--- a/Engine/Plugins/FX/Niagara/Source/NiagaraEditor/Private/NiagaraVersionMetaData.h
+++ b/Engine/Plugins/FX/Niagara/Source/NiagaraEditor/Private/NiagaraVersionMetaData.h
@@ -9,12 +9,12 @@
 #include "UObject/Object.h"
 #include "NiagaraVersionMetaData.generated.h"
 
-UCLASS()
-class NIAGARAEDITOR_API UNiagaraVersionMetaData : public UObject
+UCLASS(MinimalAPI)
+class UNiagaraVersionMetaData : public UObject
 {
 	GENERATED_BODY()
 public:
-	UNiagaraVersionMetaData();
+	NIAGARAEDITOR_API UNiagaraVersionMetaData();
 
 	/** If true then this version is exposed to the user and is used as the default version for new assets. */
 	UPROPERTY(EditAnywhere, Category="Version Details", meta=(EditCondition="!bIsExposedVersion"))
