[Niagara]

diff --git a/Engine/Plugins/FX/Niagara/Source/NiagaraEditor/Private/ViewModels/NiagaraSimCacheViewModel.h b/Engine/Plugins/FX/Niagara/Source/NiagaraEditor/Private/ViewModels/NiagaraSimCacheViewModel.h
--- a/Engine/Plugins/FX/Niagara/Source/NiagaraEditor/Private/ViewModels/NiagaraSimCacheViewModel.h
+++ b/Engine/Plugins/FX/Niagara/Source/NiagaraEditor/Private/ViewModels/NiagaraSimCacheViewModel.h
@@ -14,7 +14,7 @@ struct FNiagaraSystemSimCacheCaptureReply;
 struct FNiagaraSimCacheDataBuffersLayout;
 class UNiagaraComponent;
 
-class NIAGARAEDITOR_API FNiagaraSimCacheViewModel : public TSharedFromThis<FNiagaraSimCacheViewModel>, public FGCObject
+class FNiagaraSimCacheViewModel : public TSharedFromThis<FNiagaraSimCacheViewModel>, public FGCObject
 {
 public:
 	struct FComponentInfo
@@ -32,18 +32,18 @@ public:
 	DECLARE_MULTICAST_DELEGATE(FOnSimCacheChanged)
 	DECLARE_MULTICAST_DELEGATE(FOnBufferChanged)
 	
-	FNiagaraSimCacheViewModel();
-	~FNiagaraSimCacheViewModel();
+	NIAGARAEDITOR_API FNiagaraSimCacheViewModel();
+	NIAGARAEDITOR_API ~FNiagaraSimCacheViewModel();
 
-	void Initialize(TWeakObjectPtr<UNiagaraSimCache> SimCache);
+	NIAGARAEDITOR_API void Initialize(TWeakObjectPtr<UNiagaraSimCache> SimCache);
 
-	void UpdateSimCache(const FNiagaraSystemSimCacheCaptureReply& Reply);
+	NIAGARAEDITOR_API void UpdateSimCache(const FNiagaraSystemSimCacheCaptureReply& Reply);
 
-	void SetupPreviewComponentAndInstance();
+	NIAGARAEDITOR_API void SetupPreviewComponentAndInstance();
 
 	TConstArrayView<FComponentInfo> GetCurrentComponentInfos() const { return GetComponentInfos(EmitterIndex); }
 
-	TConstArrayView<FComponentInfo> GetComponentInfos(int32 InEmitterIndex) const;
+	NIAGARAEDITOR_API TConstArrayView<FComponentInfo> GetComponentInfos(int32 InEmitterIndex) const;
 
 	TConstArrayView<FString> GetComponentFilters() const { return ComponentFilterArray; }
 
@@ -59,63 +59,63 @@ public:
 	
 	int32 GetNumInstances() const { return NumInstances; }
 
-	int32 GetNumFrames() const;
+	NIAGARAEDITOR_API int32 GetNumFrames() const;
 
 	int32 GetFrameIndex() const { return FrameIndex; }
 
-	void SetFrameIndex(const int32 InFrameIndex);
+	NIAGARAEDITOR_API void SetFrameIndex(const int32 InFrameIndex);
 
 	int32 GetEmitterIndex() const { return EmitterIndex; };
 
-	void SetEmitterIndex(int32 InEmitterIndex);
+	NIAGARAEDITOR_API void SetEmitterIndex(int32 InEmitterIndex);
 
-	FText GetComponentText(FName ComponentName, int32 InstanceIndex) const;
+	NIAGARAEDITOR_API FText GetComponentText(FName ComponentName, int32 InstanceIndex) const;
 
-	bool IsCacheValid();
+	NIAGARAEDITOR_API bool IsCacheValid();
 
-	int32 GetNumEmitterLayouts();
+	NIAGARAEDITOR_API int32 GetNumEmitterLayouts();
 
-	FName GetEmitterLayoutName(int32 Index);
+	NIAGARAEDITOR_API FName GetEmitterLayoutName(int32 Index);
 
-	FOnViewDataChanged& OnViewDataChanged();
+	NIAGARAEDITOR_API FOnViewDataChanged& OnViewDataChanged();
 
-	FOnSimCacheChanged& OnSimCacheChanged();
+	NIAGARAEDITOR_API FOnSimCacheChanged& OnSimCacheChanged();
 
-	FOnBufferChanged& OnBufferChanged();
+	NIAGARAEDITOR_API FOnBufferChanged& OnBufferChanged();
 
-	void OnCacheModified(UNiagaraSimCache* SimCache);
+	NIAGARAEDITOR_API void OnCacheModified(UNiagaraSimCache* SimCache);
 
-	void UpdateCachedFrame();
+	NIAGARAEDITOR_API void UpdateCachedFrame();
 
 	// Build out component infos for each buffer in this cache
-	void UpdateComponentInfos();
+	NIAGARAEDITOR_API void UpdateComponentInfos();
 
-	void BuildTreeItemChildren(TSharedPtr<FNiagaraSimCacheTreeItem> InTreeItem, TWeakPtr<SNiagaraSimCacheTreeView> OwningTreeView);
+	NIAGARAEDITOR_API void BuildTreeItemChildren(TSharedPtr<FNiagaraSimCacheTreeItem> InTreeItem, TWeakPtr<SNiagaraSimCacheTreeView> OwningTreeView);
 
-	void RecursiveBuildTreeItemChildren(FNiagaraSimCacheTreeItem* Root, TSharedRef<FNiagaraSimCacheComponentTreeItem> Parent,
+	NIAGARAEDITOR_API void RecursiveBuildTreeItemChildren(FNiagaraSimCacheTreeItem* Root, TSharedRef<FNiagaraSimCacheComponentTreeItem> Parent,
 	FNiagaraTypeDefinition TypeDefinition, TWeakPtr<SNiagaraSimCacheTreeView> OwningTreeView);
 	
 	// Construct entries for the tree view.
-	void BuildEntries(TWeakPtr<SNiagaraSimCacheTreeView> OwningTreeView);
+	NIAGARAEDITOR_API void BuildEntries(TWeakPtr<SNiagaraSimCacheTreeView> OwningTreeView);
 
-	void UpdateCurrentEntries();
+	NIAGARAEDITOR_API void UpdateCurrentEntries();
 
-	TArray<TSharedRef<FNiagaraSimCacheTreeItem>>* GetCurrentRootEntries();
+	NIAGARAEDITOR_API TArray<TSharedRef<FNiagaraSimCacheTreeItem>>* GetCurrentRootEntries();
 
 
-	TArray<TSharedRef<FNiagaraSimCacheOverviewItem>>* GetBufferEntries();
+	NIAGARAEDITOR_API TArray<TSharedRef<FNiagaraSimCacheOverviewItem>>* GetBufferEntries();
 
 	UNiagaraComponent* GetPreviewComponent() { return PreviewComponent; }
 
 	//~ FGCObject interface
-	virtual void AddReferencedObjects(FReferenceCollector& Collector) override;
+	NIAGARAEDITOR_API virtual void AddReferencedObjects(FReferenceCollector& Collector) override;
 	virtual FString GetReferencerName() const override
 	{
 		return TEXT("FNiagaraSimCacheViewModel");
 	}
 
 private:
-	void BuildComponentInfos(const FName Name, const UScriptStruct* Struct, TArray<FComponentInfo>& ComponentInfos);
+	NIAGARAEDITOR_API void BuildComponentInfos(const FName Name, const UScriptStruct* Struct, TArray<FComponentInfo>& ComponentInfos);
 	
 	// The sim cache being viewed
 	TObjectPtr<UNiagaraSimCache> SimCache = nullptr;
diff --git a/Engine/Plugins/FX/Niagara/Source/NiagaraEditor/Private/ViewModels/NiagaraStandaloneScriptViewModel.h b/Engine/Plugins/FX/Niagara/Source/NiagaraEditor/Private/ViewModels/NiagaraStandaloneScriptViewModel.h
--- a/Engine/Plugins/FX/Niagara/Source/NiagaraEditor/Private/ViewModels/NiagaraStandaloneScriptViewModel.h
+++ b/Engine/Plugins/FX/Niagara/Source/NiagaraEditor/Private/ViewModels/NiagaraStandaloneScriptViewModel.h
@@ -11,10 +11,10 @@ class UNiagaraScript;
 class UNiagaraScriptVariable;
 
 
-class NIAGARAEDITOR_API FNiagaraStandaloneScriptViewModel : public FNiagaraScriptViewModel
+class FNiagaraStandaloneScriptViewModel : public FNiagaraScriptViewModel
 {
 public:
-	FNiagaraStandaloneScriptViewModel(
+	NIAGARAEDITOR_API FNiagaraStandaloneScriptViewModel(
 		FText DisplayName,
 		ENiagaraParameterEditMode InParameterEditMode,
 		TSharedPtr<FNiagaraMessageLogViewModel> InNiagaraMessageLogViewModel,
@@ -22,19 +22,19 @@ public:
 		bool bInIsForDataProcessingOnly
 	);
 
-	void Initialize(FVersionedNiagaraScript& InScript, const FVersionedNiagaraScript& InSourceScript);
+	NIAGARAEDITOR_API void Initialize(FVersionedNiagaraScript& InScript, const FVersionedNiagaraScript& InSourceScript);
 
 	//~ Begin INiagaraParameterDefinitionsSubscriberViewModel Interface
 protected:
-	virtual INiagaraParameterDefinitionsSubscriber* GetParameterDefinitionsSubscriber() override;
+	NIAGARAEDITOR_API virtual INiagaraParameterDefinitionsSubscriber* GetParameterDefinitionsSubscriber() override;
 	//~ End NiagaraParameterDefinitionsSubscriberViewModel Interface
 
 public:
-	virtual FVersionedNiagaraScript GetStandaloneScript() override;
-	const FVersionedNiagaraScript GetStandaloneScript() const;
+	NIAGARAEDITOR_API virtual FVersionedNiagaraScript GetStandaloneScript() override;
+	NIAGARAEDITOR_API const FVersionedNiagaraScript GetStandaloneScript() const;
 
 private:
-	virtual void OnVMScriptCompiled(UNiagaraScript* InScript, const FGuid& ScriptVersion) override;
+	NIAGARAEDITOR_API virtual void OnVMScriptCompiled(UNiagaraScript* InScript, const FGuid& ScriptVersion) override;
 
 	/** Sends messages to FNiagaraMessageManager for all compile events from the last compile. */
 	void SendLastCompileMessages(const FVersionedNiagaraScript& InScript);
diff --git a/Engine/Plugins/FX/Niagara/Source/NiagaraEditor/Private/Widgets/SNiagaraEmitterVersionWidget.h b/Engine/Plugins/FX/Niagara/Source/NiagaraEditor/Private/Widgets/SNiagaraEmitterVersionWidget.h
--- a/Engine/Plugins/FX/Niagara/Source/NiagaraEditor/Private/Widgets/SNiagaraEmitterVersionWidget.h
+++ b/Engine/Plugins/FX/Niagara/Source/NiagaraEditor/Private/Widgets/SNiagaraEmitterVersionWidget.h
@@ -9,14 +9,14 @@
 
 
 
-class NIAGARAEDITOR_API SNiagaraEmitterVersionWidget : public SNiagaraVersionWidget
+class SNiagaraEmitterVersionWidget : public SNiagaraVersionWidget
 {
 public:
-	void Construct(const FArguments& InArgs, UNiagaraEmitter* InEmitter, UNiagaraVersionMetaData* InMetadata, const FString& BasePackageName);
+	NIAGARAEDITOR_API void Construct(const FArguments& InArgs, UNiagaraEmitter* InEmitter, UNiagaraVersionMetaData* InMetadata, const FString& BasePackageName);
 
 protected:
-	virtual FText GetInfoHeaderText() const override;
-	virtual void ExecuteSaveAsAssetAction(FNiagaraAssetVersion AssetVersion) override;
+	NIAGARAEDITOR_API virtual FText GetInfoHeaderText() const override;
+	NIAGARAEDITOR_API virtual void ExecuteSaveAsAssetAction(FNiagaraAssetVersion AssetVersion) override;
 
 private:
 	UNiagaraEmitter* Emitter = nullptr;
diff --git a/Engine/Plugins/FX/Niagara/Source/NiagaraEditor/Private/Widgets/SNiagaraHierarchyAssignment.h b/Engine/Plugins/FX/Niagara/Source/NiagaraEditor/Private/Widgets/SNiagaraHierarchyAssignment.h
--- a/Engine/Plugins/FX/Niagara/Source/NiagaraEditor/Private/Widgets/SNiagaraHierarchyAssignment.h
+++ b/Engine/Plugins/FX/Niagara/Source/NiagaraEditor/Private/Widgets/SNiagaraHierarchyAssignment.h
@@ -9,7 +9,7 @@
 /**
  * 
  */
-class NIAGARAEDITOR_API SNiagaraHierarchyAssignment : public SCompoundWidget
+class SNiagaraHierarchyAssignment : public SCompoundWidget
 {
 public:
 	SLATE_BEGIN_ARGS(SNiagaraHierarchyAssignment)
@@ -19,5 +19,5 @@ public:
 	SLATE_END_ARGS()
 
 	/** Constructs this widget with InArgs */
-	void Construct(const FArguments& InArgs, class UNiagaraNodeAssignment& InAssignmentNode);
+	NIAGARAEDITOR_API void Construct(const FArguments& InArgs, class UNiagaraNodeAssignment& InAssignmentNode);
 };
diff --git a/Engine/Plugins/FX/Niagara/Source/NiagaraEditor/Private/Widgets/SNiagaraScriptVersionWidget.h b/Engine/Plugins/FX/Niagara/Source/NiagaraEditor/Private/Widgets/SNiagaraScriptVersionWidget.h
--- a/Engine/Plugins/FX/Niagara/Source/NiagaraEditor/Private/Widgets/SNiagaraScriptVersionWidget.h
+++ b/Engine/Plugins/FX/Niagara/Source/NiagaraEditor/Private/Widgets/SNiagaraScriptVersionWidget.h
@@ -7,14 +7,14 @@
 #include "NiagaraActions.h"
 #include "SNiagaraVersionWidget.h"
 
-class NIAGARAEDITOR_API SNiagaraScriptVersionWidget : public SNiagaraVersionWidget
+class SNiagaraScriptVersionWidget : public SNiagaraVersionWidget
 {
 public:
-	void Construct(const FArguments& InArgs, UNiagaraScript* InScript, UNiagaraVersionMetaData* InMetadata, const FString& BasePackageName);
+	NIAGARAEDITOR_API void Construct(const FArguments& InArgs, UNiagaraScript* InScript, UNiagaraVersionMetaData* InMetadata, const FString& BasePackageName);
 
 protected:
-	virtual FText GetInfoHeaderText() const override;
-	virtual void ExecuteSaveAsAssetAction(FNiagaraAssetVersion AssetVersion) override;
+	NIAGARAEDITOR_API virtual FText GetInfoHeaderText() const override;
+	NIAGARAEDITOR_API virtual void ExecuteSaveAsAssetAction(FNiagaraAssetVersion AssetVersion) override;
 
 private:
 	UNiagaraScript* Script = nullptr;
diff --git a/Engine/Plugins/FX/Niagara/Source/NiagaraEditor/Private/Widgets/SNiagaraSimCacheViewTransportControls.h b/Engine/Plugins/FX/Niagara/Source/NiagaraEditor/Private/Widgets/SNiagaraSimCacheViewTransportControls.h
--- a/Engine/Plugins/FX/Niagara/Source/NiagaraEditor/Private/Widgets/SNiagaraSimCacheViewTransportControls.h
+++ b/Engine/Plugins/FX/Niagara/Source/NiagaraEditor/Private/Widgets/SNiagaraSimCacheViewTransportControls.h
@@ -5,7 +5,7 @@
 #include "CoreMinimal.h"
 #include "Widgets/SCompoundWidget.h"
 
-class NIAGARAEDITOR_API SNiagaraSimCacheViewTransportControls : public SCompoundWidget
+class SNiagaraSimCacheViewTransportControls : public SCompoundWidget
 {
 public:
 	SLATE_BEGIN_ARGS(SNiagaraSimCacheViewTransportControls) {}
@@ -13,21 +13,21 @@ public:
 	SLATE_END_ARGS()
 
 	/** Constructs this widget with InArgs */
-	void Construct(const FArguments& InArgs);
+	NIAGARAEDITOR_API void Construct(const FArguments& InArgs);
 
 	
 	// Transport Controls
-	TSharedRef<SWidget> MakeTransportControlsWidget(class FEditorWidgetsModule& EditorWidgetsModule);
-	TOptional<int32> GetCurrentFrame() const;
-	void OnFrameIndexChanged(const int32 InFrameIndex, ETextCommit::Type Arg);
-	TOptional<int32> GetMaxFrameIndex() const;
+	NIAGARAEDITOR_API TSharedRef<SWidget> MakeTransportControlsWidget(class FEditorWidgetsModule& EditorWidgetsModule);
+	NIAGARAEDITOR_API TOptional<int32> GetCurrentFrame() const;
+	NIAGARAEDITOR_API void OnFrameIndexChanged(const int32 InFrameIndex, ETextCommit::Type Arg);
+	NIAGARAEDITOR_API TOptional<int32> GetMaxFrameIndex() const;
 	// Custom transport widget for frame selection
-	TSharedRef<SWidget> CreateFrameWidget();
+	NIAGARAEDITOR_API TSharedRef<SWidget> CreateFrameWidget();
 	
-	FReply OnTransportBackwardEnd();
-	FReply OnTransportBackwardStep();
-	FReply OnTransportForwardStep();
-	FReply OnTransportForwardEnd();
+	NIAGARAEDITOR_API FReply OnTransportBackwardEnd();
+	NIAGARAEDITOR_API FReply OnTransportBackwardStep();
+	NIAGARAEDITOR_API FReply OnTransportForwardStep();
+	NIAGARAEDITOR_API FReply OnTransportForwardEnd();
 
 private:
 	TWeakPtr<FNiagaraSimCacheViewModel> WeakViewModel;
diff --git a/Engine/Plugins/FX/Niagara/Source/NiagaraEditor/Private/Widgets/SNiagaraVersionWidget.h b/Engine/Plugins/FX/Niagara/Source/NiagaraEditor/Private/Widgets/SNiagaraVersionWidget.h
--- a/Engine/Plugins/FX/Niagara/Source/NiagaraEditor/Private/Widgets/SNiagaraVersionWidget.h
+++ b/Engine/Plugins/FX/Niagara/Source/NiagaraEditor/Private/Widgets/SNiagaraVersionWidget.h
@@ -33,7 +33,7 @@ struct FNiagaraVersionMenuAction : FNiagaraMenuAction
 	FNiagaraAssetVersion AssetVersion;
 };
 
-class NIAGARAEDITOR_API SNiagaraVersionWidget : public SCompoundWidget, FNotifyHook, public FEditorUndoClient
+class SNiagaraVersionWidget : public SCompoundWidget, FNotifyHook, public FEditorUndoClient
 {
 public:
 	SLATE_BEGIN_ARGS(SNiagaraVersionWidget)
@@ -47,27 +47,27 @@ public:
 	
 	SLATE_END_ARGS()
 
-	virtual ~SNiagaraVersionWidget() override;
+	NIAGARAEDITOR_API virtual ~SNiagaraVersionWidget() override;
 
-	void Construct(const FArguments& InArgs, FNiagaraVersionedObject* InVersionedObject, UNiagaraVersionMetaData* InMetadata);
+	NIAGARAEDITOR_API void Construct(const FArguments& InArgs, FNiagaraVersionedObject* InVersionedObject, UNiagaraVersionMetaData* InMetadata);
 
-	void UpdateEditedObject(FNiagaraVersionedObject* InVersionedObject, UNiagaraVersionMetaData* InMetadata);
+	NIAGARAEDITOR_API void UpdateEditedObject(FNiagaraVersionedObject* InVersionedObject, UNiagaraVersionMetaData* InMetadata);
 
 	//~ Begin FNotifyHook Interface
-	virtual void NotifyPostChange(const FPropertyChangedEvent& PropertyChangedEvent, FProperty* PropertyThatChanged) override;
+	NIAGARAEDITOR_API virtual void NotifyPostChange(const FPropertyChangedEvent& PropertyChangedEvent, FProperty* PropertyThatChanged) override;
 	// End of FNotifyHook
 
 	//~ Begin FEditorUndoClient interface
-	virtual void PostUndo(bool bSuccess) override;
-	virtual void PostRedo(bool bSuccess) override;
+	NIAGARAEDITOR_API virtual void PostUndo(bool bSuccess) override;
+	NIAGARAEDITOR_API virtual void PostRedo(bool bSuccess) override;
 	// End of FEditorUndoClient
 
 	/** See OnVersionDataChanged event */
-	void SetOnVersionDataChanged(FOnVersionPropertyChangedDelegate InOnVersionDataChanged);
+	NIAGARAEDITOR_API void SetOnVersionDataChanged(FOnVersionPropertyChangedDelegate InOnVersionDataChanged);
 
 protected:
-	virtual FText GetInfoHeaderText() const;
-	virtual void ExecuteSaveAsAssetAction(FNiagaraAssetVersion AssetVersion);
+	NIAGARAEDITOR_API virtual FText GetInfoHeaderText() const;
+	NIAGARAEDITOR_API virtual void ExecuteSaveAsAssetAction(FNiagaraAssetVersion AssetVersion);
 	
 private:
 	FNiagaraVersionedObject* VersionedObject = nullptr;
diff --git a/Engine/Plugins/FX/Niagara/Source/NiagaraEditor/Public/EdGraphSchema_Niagara.h b/Engine/Plugins/FX/Niagara/Source/NiagaraEditor/Public/EdGraphSchema_Niagara.h
--- a/Engine/Plugins/FX/Niagara/Source/NiagaraEditor/Public/EdGraphSchema_Niagara.h
+++ b/Engine/Plugins/FX/Niagara/Source/NiagaraEditor/Public/EdGraphSchema_Niagara.h
@@ -18,7 +18,7 @@ enum class ENiagaraDataType : uint8;
 
 /** Action to add a node to the graph */
 USTRUCT()
-struct NIAGARAEDITOR_API FNiagaraSchemaAction_NewNode : public FEdGraphSchemaAction
+struct FNiagaraSchemaAction_NewNode : public FEdGraphSchemaAction
 {
 	GENERATED_USTRUCT_BODY();
 
@@ -40,9 +40,9 @@ struct NIAGARAEDITOR_API FNiagaraSchemaAction_NewNode : public FEdGraphSchemaAct
 	{}
 
 	//~ Begin FEdGraphSchemaAction Interface
-	virtual UEdGraphNode* PerformAction(class UEdGraph* ParentGraph, UEdGraphPin* FromPin, const FVector2D Location, bool bSelectNewNode = true) override;
-	virtual UEdGraphNode* PerformAction(class UEdGraph* ParentGraph, TArray<UEdGraphPin*>& FromPins, const FVector2D Location, bool bSelectNewNode = true) override;
-	virtual void AddReferencedObjects( FReferenceCollector& Collector ) override;
+	NIAGARAEDITOR_API virtual UEdGraphNode* PerformAction(class UEdGraph* ParentGraph, UEdGraphPin* FromPin, const FVector2D Location, bool bSelectNewNode = true) override;
+	NIAGARAEDITOR_API virtual UEdGraphNode* PerformAction(class UEdGraph* ParentGraph, TArray<UEdGraphPin*>& FromPins, const FVector2D Location, bool bSelectNewNode = true) override;
+	NIAGARAEDITOR_API virtual void AddReferencedObjects( FReferenceCollector& Collector ) override;
 	//~ End FEdGraphSchemaAction Interface
 
 	template <typename NodeType>
@@ -56,7 +56,7 @@ struct NIAGARAEDITOR_API FNiagaraSchemaAction_NewNode : public FEdGraphSchemaAct
 };
 
 USTRUCT()
-struct NIAGARAEDITOR_API FNiagaraSchemaAction_NewComment : public FEdGraphSchemaAction
+struct FNiagaraSchemaAction_NewComment : public FEdGraphSchemaAction
 {
 public:
 	GENERATED_USTRUCT_BODY();
@@ -72,7 +72,7 @@ public:
 	{}
 
 	//~ Begin FEdGraphSchemaAction Interface
-	virtual UEdGraphNode* PerformAction(class UEdGraph* ParentGraph, UEdGraphPin* FromPin, const FVector2D Location, bool bSelectNewNode = true) override;
+	NIAGARAEDITOR_API virtual UEdGraphNode* PerformAction(class UEdGraph* ParentGraph, UEdGraphPin* FromPin, const FVector2D Location, bool bSelectNewNode = true) override;
 	//~ End FEdGraphSchemaAction Interface
 
 private:
@@ -80,36 +80,36 @@ private:
 	TSharedPtr<SGraphEditor> GraphEditor;
 };
 
-UCLASS()
-class NIAGARAEDITOR_API UEdGraphSchema_Niagara : public UEdGraphSchema
+UCLASS(MinimalAPI)
+class UEdGraphSchema_Niagara : public UEdGraphSchema
 {
 	GENERATED_UCLASS_BODY()
 
 	// Allowable PinType.PinCategory values
-	static const FName PinCategoryType;
-	static const FName PinCategoryMisc;
-	static const FName PinCategoryClass;
-	static const FName PinCategoryEnum;
-	static const FName PinCategoryStaticType;
-	static const FName PinCategoryStaticClass;
-	static const FName PinCategoryStaticEnum;
+	static NIAGARAEDITOR_API const FName PinCategoryType;
+	static NIAGARAEDITOR_API const FName PinCategoryMisc;
+	static NIAGARAEDITOR_API const FName PinCategoryClass;
+	static NIAGARAEDITOR_API const FName PinCategoryEnum;
+	static NIAGARAEDITOR_API const FName PinCategoryStaticType;
+	static NIAGARAEDITOR_API const FName PinCategoryStaticClass;
+	static NIAGARAEDITOR_API const FName PinCategoryStaticEnum;
 
-	static bool IsStaticPin(const UEdGraphPin* Pin);
+	static NIAGARAEDITOR_API bool IsStaticPin(const UEdGraphPin* Pin);
 
 	//~ Begin EdGraphSchema Interface
-	virtual void GetContextMenuActions(class UToolMenu* Menu, class UGraphNodeContextMenuContext* Context) const override;
-	virtual const FPinConnectionResponse CanCreateConnection(const UEdGraphPin* A, const UEdGraphPin* B) const override;
-	virtual FLinearColor GetPinTypeColor(const FEdGraphPinType& PinType) const override;
-	virtual bool ShouldHidePinDefaultValue(UEdGraphPin* Pin) const override;
-	virtual bool TryCreateConnection(UEdGraphPin* A, UEdGraphPin* B) const override;
-	virtual void BreakSinglePinLink(UEdGraphPin* SourcePin, UEdGraphPin* TargetPin) const override;
-	virtual void BreakPinLinks(UEdGraphPin& TargetPin, bool bSendsNodeNotification) const override;
-	virtual FConnectionDrawingPolicy* CreateConnectionDrawingPolicy(int32 InBackLayerID, int32 InFrontLayerID, float InZoomFactor, const FSlateRect& InClippingRect, class FSlateWindowElementList& InDrawElements, class UEdGraph* InGraphObj) const override;
-	virtual void ResetPinToAutogeneratedDefaultValue(UEdGraphPin* Pin, bool bCallModifyCallbacks = true) const override;
-	virtual void OnPinConnectionDoubleCicked(UEdGraphPin* PinA, UEdGraphPin* PinB, const FVector2D& GraphPosition) const override; 
+	NIAGARAEDITOR_API virtual void GetContextMenuActions(class UToolMenu* Menu, class UGraphNodeContextMenuContext* Context) const override;
+	NIAGARAEDITOR_API virtual const FPinConnectionResponse CanCreateConnection(const UEdGraphPin* A, const UEdGraphPin* B) const override;
+	NIAGARAEDITOR_API virtual FLinearColor GetPinTypeColor(const FEdGraphPinType& PinType) const override;
+	NIAGARAEDITOR_API virtual bool ShouldHidePinDefaultValue(UEdGraphPin* Pin) const override;
+	NIAGARAEDITOR_API virtual bool TryCreateConnection(UEdGraphPin* A, UEdGraphPin* B) const override;
+	NIAGARAEDITOR_API virtual void BreakSinglePinLink(UEdGraphPin* SourcePin, UEdGraphPin* TargetPin) const override;
+	NIAGARAEDITOR_API virtual void BreakPinLinks(UEdGraphPin& TargetPin, bool bSendsNodeNotification) const override;
+	NIAGARAEDITOR_API virtual FConnectionDrawingPolicy* CreateConnectionDrawingPolicy(int32 InBackLayerID, int32 InFrontLayerID, float InZoomFactor, const FSlateRect& InClippingRect, class FSlateWindowElementList& InDrawElements, class UEdGraph* InGraphObj) const override;
+	NIAGARAEDITOR_API virtual void ResetPinToAutogeneratedDefaultValue(UEdGraphPin* Pin, bool bCallModifyCallbacks = true) const override;
+	NIAGARAEDITOR_API virtual void OnPinConnectionDoubleCicked(UEdGraphPin* PinA, UEdGraphPin* PinB, const FVector2D& GraphPosition) const override; 
 	virtual bool ShouldAlwaysPurgeOnModification() const override { return false; }
-	virtual void DroppedAssetsOnGraph(const TArray<FAssetData>& Assets, const FVector2D& GraphPosition, UEdGraph* Graph) const override;
-	virtual void GetAssetsGraphHoverMessage(const TArray<FAssetData>& Assets, const UEdGraph* HoverGraph, FString& OutTooltipText, bool& OutOkIcon) const override;
+	NIAGARAEDITOR_API virtual void DroppedAssetsOnGraph(const TArray<FAssetData>& Assets, const FVector2D& GraphPosition, UEdGraph* Graph) const override;
+	NIAGARAEDITOR_API virtual void GetAssetsGraphHoverMessage(const TArray<FAssetData>& Assets, const UEdGraph* HoverGraph, FString& OutTooltipText, bool& OutOkIcon) const override;
 
 	/**
 	* Sets the string to the specified pin; even if it is invalid it is still set.
@@ -118,17 +118,17 @@ class NIAGARAEDITOR_API UEdGraphSchema_Niagara : public UEdGraphSchema
 	* @param	NewDefaultValue	The new default value.
 	* @param   bMarkAsModified Marks the container of the value as modified. For Niagara Schema, whether or not to notify the Node owning the Pin that the Pin has changed.
 	*/
-	virtual void TrySetDefaultValue(UEdGraphPin& Pin, const FString& NewDefaultValue, bool bMarkAsModified = true) const override;
+	NIAGARAEDITOR_API virtual void TrySetDefaultValue(UEdGraphPin& Pin, const FString& NewDefaultValue, bool bMarkAsModified = true) const override;
 
 	//~ End EdGraphSchema Interface
 
-	static FLinearColor GetTypeColor(const FNiagaraTypeDefinition& Type);
+	static NIAGARAEDITOR_API FLinearColor GetTypeColor(const FNiagaraTypeDefinition& Type);
 
-	TArray<TSharedPtr<FNiagaraAction_NewNode>> GetGraphActions(const UEdGraph* CurrentGraph, const UEdGraphPin* FromPin, UEdGraph* OwnerOfTemporaries) const;
-	void PromoteSinglePinToParameter(UEdGraphPin* SourcePin);
-	static bool CanPromoteSinglePinToParameter(const UEdGraphPin* SourcePin);
-	void ToggleNodeEnabledState(class UNiagaraNode* InNode) const;
-	void RefreshNode(UNiagaraNode* InNode) const;
+	NIAGARAEDITOR_API TArray<TSharedPtr<FNiagaraAction_NewNode>> GetGraphActions(const UEdGraph* CurrentGraph, const UEdGraphPin* FromPin, UEdGraph* OwnerOfTemporaries) const;
+	NIAGARAEDITOR_API void PromoteSinglePinToParameter(UEdGraphPin* SourcePin);
+	static NIAGARAEDITOR_API bool CanPromoteSinglePinToParameter(const UEdGraphPin* SourcePin);
+	NIAGARAEDITOR_API void ToggleNodeEnabledState(class UNiagaraNode* InNode) const;
+	NIAGARAEDITOR_API void RefreshNode(UNiagaraNode* InNode) const;
 
 	/** 
 	  * Creates a niagara variable using the name, type, and default value stored on an ed graph pin.
@@ -136,7 +136,7 @@ class NIAGARAEDITOR_API UEdGraphSchema_Niagara : public UEdGraphSchema
 	  * bNeedsValue Whether or not the returned variable must be allocated to a valid value. When true if the pin doesn't have a valid default value itself the variable will be reset to default before return.
 	  * returns The newly created variable.
 	  */
-	static FNiagaraVariable PinToNiagaraVariable(const UEdGraphPin* Pin, bool bNeedsValue=false, ENiagaraStructConversion StructConversion = ENiagaraStructConversion::UserFacing);
+	static NIAGARAEDITOR_API FNiagaraVariable PinToNiagaraVariable(const UEdGraphPin* Pin, bool bNeedsValue=false, ENiagaraStructConversion StructConversion = ENiagaraStructConversion::UserFacing);
 
 	/** 
 	  * Tries to get a default value string for a graph pin from a niagara variable.
@@ -144,61 +144,61 @@ class NIAGARAEDITOR_API UEdGraphSchema_Niagara : public UEdGraphSchema
 	  * @param OutPinDefaultValue The pin default value string if this call was successful, otherwise empty.
 	  * @returns Whether or not we could get a default value form the supplied variable.
 	  */
-	bool TryGetPinDefaultValueFromNiagaraVariable(const FNiagaraVariable& Variable, FString& OutPinDefaultValue) const;
+	NIAGARAEDITOR_API bool TryGetPinDefaultValueFromNiagaraVariable(const FNiagaraVariable& Variable, FString& OutPinDefaultValue) const;
 
-	static void ConvertIllegalPinsInPlace(UEdGraphPin* Pin);
-	static FNiagaraTypeDefinition PinToTypeDefinition(const UEdGraphPin* Pin, ENiagaraStructConversion StructConversion = ENiagaraStructConversion::UserFacing);
-	static FNiagaraTypeDefinition PinTypeToTypeDefinition(const FEdGraphPinType& PinType);
-	static FEdGraphPinType TypeDefinitionToPinType(FNiagaraTypeDefinition TypeDef);
+	static NIAGARAEDITOR_API void ConvertIllegalPinsInPlace(UEdGraphPin* Pin);
+	static NIAGARAEDITOR_API FNiagaraTypeDefinition PinToTypeDefinition(const UEdGraphPin* Pin, ENiagaraStructConversion StructConversion = ENiagaraStructConversion::UserFacing);
+	static NIAGARAEDITOR_API FNiagaraTypeDefinition PinTypeToTypeDefinition(const FEdGraphPinType& PinType);
+	static NIAGARAEDITOR_API FEdGraphPinType TypeDefinitionToPinType(FNiagaraTypeDefinition TypeDef);
 	
-	static bool IsPinWildcard(const UEdGraphPin* Pin);
+	static NIAGARAEDITOR_API bool IsPinWildcard(const UEdGraphPin* Pin);
 	
 	// Fast path helper to skip doing full type conversion
-	static bool IsPinStatic(const UEdGraphPin* Pin);
-	static FPinConnectionResponse GetWildcardConnectionResponse(const UEdGraphPin* PinA, const UEdGraphPin* PinB);
+	static NIAGARAEDITOR_API bool IsPinStatic(const UEdGraphPin* Pin);
+	static NIAGARAEDITOR_API FPinConnectionResponse GetWildcardConnectionResponse(const UEdGraphPin* PinA, const UEdGraphPin* PinB);
 	
-	bool IsSystemConstant(const FNiagaraVariable& Variable)const;
-
-	class UNiagaraParameterCollection* VariableIsFromParameterCollection(const FNiagaraVariable& Var)const;
-	class UNiagaraParameterCollection* VariableIsFromParameterCollection(const FString& VarName, bool bAllowPartialMatch, FNiagaraVariable& OutVar)const;
-
-	bool IsValidNiagaraPropertyType(const FProperty* Property) const;
-	FNiagaraTypeDefinition GetTypeDefForProperty(const FProperty* Property) const;
-
-	static const FLinearColor NodeTitleColor_Attribute;
-	static const FLinearColor NodeTitleColor_Constant;
-	static const FLinearColor NodeTitleColor_SystemConstant;
-	static const FLinearColor NodeTitleColor_FunctionCall;
-	static const FLinearColor NodeTitleColor_CustomHlsl;
-	static const FLinearColor NodeTitleColor_Event; 
-	static const FLinearColor NodeTitleColor_TranslatorConstant;
-	static const FLinearColor NodeTitleColor_RapidIteration;
-
-	static const FText ReplaceExistingInputConnectionsText;
-	static const FText TypesAreNotCompatibleText;
-	static const FText ConvertText;
-	static const FText ConvertLossyText;
-	static const FText PinNotConnectableText;
-	static const FText SameNodeConnectionForbiddenText;
-	static const FText DirectionsNotCompatibleText;
-	static const FText AddPinIncompatibleTypeText;
-	static const FText CircularConnectionFoundText;
+	NIAGARAEDITOR_API bool IsSystemConstant(const FNiagaraVariable& Variable)const;
+
+	NIAGARAEDITOR_API class UNiagaraParameterCollection* VariableIsFromParameterCollection(const FNiagaraVariable& Var)const;
+	NIAGARAEDITOR_API class UNiagaraParameterCollection* VariableIsFromParameterCollection(const FString& VarName, bool bAllowPartialMatch, FNiagaraVariable& OutVar)const;
+
+	NIAGARAEDITOR_API bool IsValidNiagaraPropertyType(const FProperty* Property) const;
+	NIAGARAEDITOR_API FNiagaraTypeDefinition GetTypeDefForProperty(const FProperty* Property) const;
+
+	static NIAGARAEDITOR_API const FLinearColor NodeTitleColor_Attribute;
+	static NIAGARAEDITOR_API const FLinearColor NodeTitleColor_Constant;
+	static NIAGARAEDITOR_API const FLinearColor NodeTitleColor_SystemConstant;
+	static NIAGARAEDITOR_API const FLinearColor NodeTitleColor_FunctionCall;
+	static NIAGARAEDITOR_API const FLinearColor NodeTitleColor_CustomHlsl;
+	static NIAGARAEDITOR_API const FLinearColor NodeTitleColor_Event; 
+	static NIAGARAEDITOR_API const FLinearColor NodeTitleColor_TranslatorConstant;
+	static NIAGARAEDITOR_API const FLinearColor NodeTitleColor_RapidIteration;
+
+	static NIAGARAEDITOR_API const FText ReplaceExistingInputConnectionsText;
+	static NIAGARAEDITOR_API const FText TypesAreNotCompatibleText;
+	static NIAGARAEDITOR_API const FText ConvertText;
+	static NIAGARAEDITOR_API const FText ConvertLossyText;
+	static NIAGARAEDITOR_API const FText PinNotConnectableText;
+	static NIAGARAEDITOR_API const FText SameNodeConnectionForbiddenText;
+	static NIAGARAEDITOR_API const FText DirectionsNotCompatibleText;
+	static NIAGARAEDITOR_API const FText AddPinIncompatibleTypeText;
+	static NIAGARAEDITOR_API const FText CircularConnectionFoundText;
 	
-	bool PinTypesValidForNumericConversion(FEdGraphPinType AType, FEdGraphPinType BType) const;	
+	NIAGARAEDITOR_API bool PinTypesValidForNumericConversion(FEdGraphPinType AType, FEdGraphPinType BType) const;	
 
 	// Return type is Name, Prototype, Description
-	static TArray<TTuple<FString, FString, FText>> GetDataInterfaceFunctionPrototypes(const UEdGraphPin* GraphPin);
+	static NIAGARAEDITOR_API TArray<TTuple<FString, FString, FText>> GetDataInterfaceFunctionPrototypes(const UEdGraphPin* GraphPin);
 
 private:
-	void GetNumericConversionToSubMenuActions(class UToolMenu* Menu, const FName SectionName, UEdGraphPin* InGraphPin);
-	void GetNumericConversionToSubMenuActionsAll(class UToolMenu* Menu, const FName SectionName, UNiagaraNode* InGraphPin);
-	void ConvertNumericPinToType(UEdGraphPin* InPin, FNiagaraTypeDefinition TypeDef);
-	void ConvertNumericPinToTypeAll(UNiagaraNode* InPin, FNiagaraTypeDefinition TypeDef);
-	void ConvertPinToType(UEdGraphPin* InPin, FNiagaraTypeDefinition TypeDef) const;
+	NIAGARAEDITOR_API void GetNumericConversionToSubMenuActions(class UToolMenu* Menu, const FName SectionName, UEdGraphPin* InGraphPin);
+	NIAGARAEDITOR_API void GetNumericConversionToSubMenuActionsAll(class UToolMenu* Menu, const FName SectionName, UNiagaraNode* InGraphPin);
+	NIAGARAEDITOR_API void ConvertNumericPinToType(UEdGraphPin* InPin, FNiagaraTypeDefinition TypeDef);
+	NIAGARAEDITOR_API void ConvertNumericPinToTypeAll(UNiagaraNode* InPin, FNiagaraTypeDefinition TypeDef);
+	NIAGARAEDITOR_API void ConvertPinToType(UEdGraphPin* InPin, FNiagaraTypeDefinition TypeDef) const;
 
-	void GenerateDataInterfacePinMenu(UToolMenu* ToolMenu, const FName SectionName, const UEdGraphPin* GraphPin, FNiagaraTypeDefinition TypeDef) const;
+	NIAGARAEDITOR_API void GenerateDataInterfacePinMenu(UToolMenu* ToolMenu, const FName SectionName, const UEdGraphPin* GraphPin, FNiagaraTypeDefinition TypeDef) const;
 
-	static bool CheckCircularConnection(TSet<const UEdGraphNode*>& VisitedNodes, const UEdGraphNode* InNode, const UEdGraphNode* InTestNode);
+	static NIAGARAEDITOR_API bool CheckCircularConnection(TSet<const UEdGraphNode*>& VisitedNodes, const UEdGraphNode* InNode, const UEdGraphNode* InTestNode);
 };
 
 class FNiagaraConnectionDrawingPolicy : public FConnectionDrawingPolicy
diff --git a/Engine/Plugins/FX/Niagara/Source/NiagaraEditor/Public/EdGraphSchema_NiagaraSystemOverview.h b/Engine/Plugins/FX/Niagara/Source/NiagaraEditor/Public/EdGraphSchema_NiagaraSystemOverview.h
--- a/Engine/Plugins/FX/Niagara/Source/NiagaraEditor/Public/EdGraphSchema_NiagaraSystemOverview.h
+++ b/Engine/Plugins/FX/Niagara/Source/NiagaraEditor/Public/EdGraphSchema_NiagaraSystemOverview.h
@@ -9,23 +9,23 @@
 
 
 
-UCLASS()
-class NIAGARAEDITOR_API UEdGraphSchema_NiagaraSystemOverview : public UEdGraphSchema
+UCLASS(MinimalAPI)
+class UEdGraphSchema_NiagaraSystemOverview : public UEdGraphSchema
 {
 	GENERATED_UCLASS_BODY()
 
 	//~ Begin EdGraphSchema Interface
-	virtual void GetContextMenuActions(class UToolMenu* Menu, class UGraphNodeContextMenuContext* Context) const override;
+	NIAGARAEDITOR_API virtual void GetContextMenuActions(class UToolMenu* Menu, class UGraphNodeContextMenuContext* Context) const override;
 	virtual const FPinConnectionResponse CanCreateConnection(const UEdGraphPin* A, const UEdGraphPin* B) const override { return FPinConnectionResponse(CONNECT_RESPONSE_DISALLOW, FText()); }; //@TODO System Overview: write text response
 	virtual bool TryCreateConnection(UEdGraphPin* A, UEdGraphPin* B) const override { return false; };
 	virtual bool ShouldAlwaysPurgeOnModification() const override { return false; }
-	virtual bool IsCacheVisualizationOutOfDate(int32 InVisualizationCacheID) const override;
-	virtual int32 GetCurrentVisualizationCacheID() const override;
-	virtual void ForceVisualizationCacheClear() const override;
-	virtual void DroppedAssetsOnGraph(const TArray<FAssetData>& Assets, const FVector2D& GraphPosition, UEdGraph* Graph) const override;
-	virtual void GetAssetsGraphHoverMessage(const TArray<FAssetData>& Assets, const UEdGraph* HoverGraph, FString& OutTooltipText, bool& OutOkIcon) const override;
+	NIAGARAEDITOR_API virtual bool IsCacheVisualizationOutOfDate(int32 InVisualizationCacheID) const override;
+	NIAGARAEDITOR_API virtual int32 GetCurrentVisualizationCacheID() const override;
+	NIAGARAEDITOR_API virtual void ForceVisualizationCacheClear() const override;
+	NIAGARAEDITOR_API virtual void DroppedAssetsOnGraph(const TArray<FAssetData>& Assets, const FVector2D& GraphPosition, UEdGraph* Graph) const override;
+	NIAGARAEDITOR_API virtual void GetAssetsGraphHoverMessage(const TArray<FAssetData>& Assets, const UEdGraph* HoverGraph, FString& OutTooltipText, bool& OutOkIcon) const override;
 
 private:
-	static int32 CurrentCacheRefreshID;
+	static NIAGARAEDITOR_API int32 CurrentCacheRefreshID;
 };
 
diff --git a/Engine/Plugins/FX/Niagara/Source/NiagaraEditor/Public/INiagaraCompiler.h b/Engine/Plugins/FX/Niagara/Source/NiagaraEditor/Public/INiagaraCompiler.h
--- a/Engine/Plugins/FX/Niagara/Source/NiagaraEditor/Public/INiagaraCompiler.h
+++ b/Engine/Plugins/FX/Niagara/Source/NiagaraEditor/Public/INiagaraCompiler.h
@@ -103,7 +103,7 @@ struct FNiagaraTranslateResults
 	static ENiagaraScriptCompileStatus TranslateResultsToSummary(const FNiagaraTranslateResults* CompileResults);
 };
 
-class NIAGARAEDITOR_API FHlslNiagaraTranslatorOptions
+class FHlslNiagaraTranslatorOptions
 {
 public:
 	ENiagaraSimTarget SimTarget = ENiagaraSimTarget::CPUSim;
@@ -121,7 +121,7 @@ public:
 	TArray<FNiagaraVariable> OverrideModuleConstants;
 };
 
-class NIAGARAEDITOR_API INiagaraHlslTranslator
+class INiagaraHlslTranslator
 {
 public:
 	virtual ~INiagaraHlslTranslator() = default;
@@ -130,14 +130,14 @@ public:
 	virtual const FNiagaraTranslatorOutput& GetTranslateOutput() const = 0;
 	virtual const FString& GetTranslatedHLSL() const = 0;
 
-	static TUniquePtr<INiagaraHlslTranslator> CreateTranslator(const FNiagaraCompileRequestDataBase* InCompileData, const FNiagaraCompileRequestDuplicateDataBase* InDuplicateData);
-	static TUniquePtr<INiagaraHlslTranslator> CreateTranslator(const FNiagaraPrecompileData* InCompileData, const FNiagaraCompilationCopyData* InDuplicateData);
-
-	static FString BuildHLSLStructDecl(const FNiagaraTypeDefinition& Type, FText& OutErrorMessage, bool bGpuScript);
-	static bool IsBuiltInHlslType(const FNiagaraTypeDefinition& Type);
-	static FString GetStructHlslTypeName(const FNiagaraTypeDefinition& Type);
-	static FString GetPropertyHlslTypeName(const FProperty* Property);
-	static bool IsHlslBuiltinVector(const FNiagaraTypeDefinition& Type);
-	static TArray<FName> ConditionPropertyPath(const FNiagaraTypeDefinition& Type, const TArray<FName>& InPath);
-	static FString GetHlslDefaultForType(const FNiagaraTypeDefinition& Type);
-};
\ No newline at end of file
+	static NIAGARAEDITOR_API TUniquePtr<INiagaraHlslTranslator> CreateTranslator(const FNiagaraCompileRequestDataBase* InCompileData, const FNiagaraCompileRequestDuplicateDataBase* InDuplicateData);
+	static NIAGARAEDITOR_API TUniquePtr<INiagaraHlslTranslator> CreateTranslator(const FNiagaraPrecompileData* InCompileData, const FNiagaraCompilationCopyData* InDuplicateData);
+
+	static NIAGARAEDITOR_API FString BuildHLSLStructDecl(const FNiagaraTypeDefinition& Type, FText& OutErrorMessage, bool bGpuScript);
+	static NIAGARAEDITOR_API bool IsBuiltInHlslType(const FNiagaraTypeDefinition& Type);
+	static NIAGARAEDITOR_API FString GetStructHlslTypeName(const FNiagaraTypeDefinition& Type);
+	static NIAGARAEDITOR_API FString GetPropertyHlslTypeName(const FProperty* Property);
+	static NIAGARAEDITOR_API bool IsHlslBuiltinVector(const FNiagaraTypeDefinition& Type);
+	static NIAGARAEDITOR_API TArray<FName> ConditionPropertyPath(const FNiagaraTypeDefinition& Type, const TArray<FName>& InPath);
+	static NIAGARAEDITOR_API FString GetHlslDefaultForType(const FNiagaraTypeDefinition& Type);
+};
diff --git a/Engine/Plugins/FX/Niagara/Source/NiagaraEditor/Public/NiagaraActions.h b/Engine/Plugins/FX/Niagara/Source/NiagaraEditor/Public/NiagaraActions.h
--- a/Engine/Plugins/FX/Niagara/Source/NiagaraEditor/Public/NiagaraActions.h
+++ b/Engine/Plugins/FX/Niagara/Source/NiagaraEditor/Public/NiagaraActions.h
@@ -34,7 +34,7 @@ enum class EScriptSource : uint8
 };
 
 USTRUCT()
-struct NIAGARAEDITOR_API FNiagaraActionSourceData
+struct FNiagaraActionSourceData
 {
 	GENERATED_BODY()
 	
@@ -53,7 +53,7 @@ struct NIAGARAEDITOR_API FNiagaraActionSourceData
 };
 
 USTRUCT()
-struct NIAGARAEDITOR_API FNiagaraMenuAction : public FEdGraphSchemaAction
+struct FNiagaraMenuAction : public FEdGraphSchemaAction
 {
 	GENERATED_USTRUCT_BODY();
 
@@ -61,8 +61,8 @@ struct NIAGARAEDITOR_API FNiagaraMenuAction : public FEdGraphSchemaAction
 	DECLARE_DELEGATE_RetVal(bool, FCanExecuteStackAction);
 
 	FNiagaraMenuAction() {}
-	FNiagaraMenuAction(FText InNodeCategory, FText InMenuDesc, FText InToolTip, const int32 InGrouping, FText InKeywords, FOnExecuteStackAction InAction, int32 InSectionID = 0);
-	FNiagaraMenuAction(FText InNodeCategory, FText InMenuDesc, FText InToolTip, const int32 InGrouping, FText InKeywords, FOnExecuteStackAction InAction, FCanExecuteStackAction InCanPerformAction, int32 InSectionID = 0);
+	NIAGARAEDITOR_API FNiagaraMenuAction(FText InNodeCategory, FText InMenuDesc, FText InToolTip, const int32 InGrouping, FText InKeywords, FOnExecuteStackAction InAction, int32 InSectionID = 0);
+	NIAGARAEDITOR_API FNiagaraMenuAction(FText InNodeCategory, FText InMenuDesc, FText InToolTip, const int32 InGrouping, FText InKeywords, FOnExecuteStackAction InAction, FCanExecuteStackAction InCanPerformAction, int32 InSectionID = 0);
 
 	void ExecuteAction()
 	{
@@ -80,8 +80,8 @@ struct NIAGARAEDITOR_API FNiagaraMenuAction : public FEdGraphSchemaAction
 
 	bool IsExperimental = false;
 
-	TOptional<FNiagaraVariable> GetParameterVariable() const;
-	void SetParameterVariable(const FNiagaraVariable& InParameterVariable);
+	NIAGARAEDITOR_API TOptional<FNiagaraVariable> GetParameterVariable() const;
+	NIAGARAEDITOR_API void SetParameterVariable(const FNiagaraVariable& InParameterVariable);
 	void SetSectionId(const int32 NewSectionId) { SectionID = NewSectionId; };
 
 private:
@@ -110,7 +110,7 @@ private:
 // new action hierarchy for the new menus. Prefer inheriting from this rather than the above
 // this action does not have any use; inherit from it and provide your own functionality
 USTRUCT()
-struct NIAGARAEDITOR_API FNiagaraMenuAction_Base
+struct FNiagaraMenuAction_Base
 {
 	GENERATED_USTRUCT_BODY();
 
@@ -118,9 +118,9 @@ struct NIAGARAEDITOR_API FNiagaraMenuAction_Base
 	DECLARE_DELEGATE_RetVal(bool, FCanExecuteAction);
 
 	FNiagaraMenuAction_Base() {}
-	FNiagaraMenuAction_Base(FText DisplayName, ENiagaraMenuSections Section, TArray<FString> InNodeCategories, FText InToolTip, FText InKeywords, float InIntrinsicWeightMultiplier = 1.f);
+	NIAGARAEDITOR_API FNiagaraMenuAction_Base(FText DisplayName, ENiagaraMenuSections Section, TArray<FString> InNodeCategories, FText InToolTip, FText InKeywords, float InIntrinsicWeightMultiplier = 1.f);
 
-	void UpdateFullSearchText();
+	NIAGARAEDITOR_API void UpdateFullSearchText();
 
 	bool bIsExperimental = false;
 
@@ -152,7 +152,7 @@ struct NIAGARAEDITOR_API FNiagaraMenuAction_Base
 };
 
 USTRUCT()
-struct NIAGARAEDITOR_API FNiagaraMenuAction_Generic : public FNiagaraMenuAction_Base
+struct FNiagaraMenuAction_Generic : public FNiagaraMenuAction_Base
 {
 	GENERATED_BODY()
 
@@ -188,8 +188,8 @@ struct NIAGARAEDITOR_API FNiagaraMenuAction_Generic : public FNiagaraMenuAction_
 		}
 	}
 
-	TOptional<FNiagaraVariable> GetParameterVariable() const;
-	void SetParameterVariable(const FNiagaraVariable& InParameterVariable);
+	NIAGARAEDITOR_API TOptional<FNiagaraVariable> GetParameterVariable() const;
+	NIAGARAEDITOR_API void SetParameterVariable(const FNiagaraVariable& InParameterVariable);
 protected:
 	FOnExecuteAction Action;
 	FCanExecuteAction CanExecuteAction;
@@ -198,7 +198,7 @@ protected:
 };
 
 USTRUCT()
-struct NIAGARAEDITOR_API FNiagaraAction_NewNode : public FNiagaraMenuAction_Generic
+struct FNiagaraAction_NewNode : public FNiagaraMenuAction_Generic
 {
 	GENERATED_BODY()
 
@@ -209,14 +209,14 @@ struct NIAGARAEDITOR_API FNiagaraAction_NewNode : public FNiagaraMenuAction_Gene
 	{		
 	}
 
-	class UEdGraphNode* CreateNode(UEdGraph* Graph, UEdGraphPin* FromPin, FVector2D NodePosition, bool bSelectNewNode = true) const;
-	class UEdGraphNode* CreateNode(UEdGraph* Graph, TArray<UEdGraphPin*>& FromPins, FVector2D NodePosition, bool bSelectNewNode = true) const;
+	NIAGARAEDITOR_API class UEdGraphNode* CreateNode(UEdGraph* Graph, UEdGraphPin* FromPin, FVector2D NodePosition, bool bSelectNewNode = true) const;
+	NIAGARAEDITOR_API class UEdGraphNode* CreateNode(UEdGraph* Graph, TArray<UEdGraphPin*>& FromPins, FVector2D NodePosition, bool bSelectNewNode = true) const;
 
 	UPROPERTY()
 	TObjectPtr<class UEdGraphNode> NodeTemplate = nullptr;
 };
 
-struct NIAGARAEDITOR_API FNiagaraParameterAction : public FEdGraphSchemaAction
+struct FNiagaraParameterAction : public FEdGraphSchemaAction
 {
 	FNiagaraParameterAction()
 		: ScriptVar(nullptr)
@@ -228,18 +228,18 @@ struct NIAGARAEDITOR_API FNiagaraParameterAction : public FEdGraphSchemaAction
 	{
 	}
 
-	FNiagaraParameterAction(const FNiagaraVariable& InParameter,
+	NIAGARAEDITOR_API FNiagaraParameterAction(const FNiagaraVariable& InParameter,
 		const TArray<FNiagaraGraphParameterReferenceCollection>& InReferenceCollection,
 		FText InNodeCategory, FText InMenuDesc, FText InToolTip, const int32 InGrouping, FText InKeywords,
 		TSharedPtr<TArray<FName>> ParametersWithNamespaceModifierRenamePending,
 		int32 InSectionID = 0);
 
-	FNiagaraParameterAction(const FNiagaraVariable& InParameter, 
+	NIAGARAEDITOR_API FNiagaraParameterAction(const FNiagaraVariable& InParameter, 
 		FText InNodeCategory, FText InMenuDesc, FText InToolTip, const int32 InGrouping, FText InKeywords,
 		TSharedPtr<TArray<FName>> ParametersWithNamespaceModifierRenamePending,
 		int32 InSectionID = 0);
 
-	FNiagaraParameterAction(const UNiagaraScriptVariable* InScriptVar,
+	NIAGARAEDITOR_API FNiagaraParameterAction(const UNiagaraScriptVariable* InScriptVar,
 		FText InNodeCategory, FText InMenuDesc, FText InToolTip, const int32 InGrouping, FText InKeywords,
 		int32 InSectionID = 0);
 
@@ -247,23 +247,23 @@ struct NIAGARAEDITOR_API FNiagaraParameterAction : public FEdGraphSchemaAction
 	static FName StaticGetTypeId() { return FNiagaraEditorStrings::FNiagaraParameterActionId; };
 	virtual FName GetTypeId() const { return StaticGetTypeId(); };
 
-	const UNiagaraScriptVariable* GetScriptVar() const;
+	NIAGARAEDITOR_API const UNiagaraScriptVariable* GetScriptVar() const;
 
-	const FNiagaraVariable& GetParameter() const;
+	NIAGARAEDITOR_API const FNiagaraVariable& GetParameter() const;
 
-	TArray<FNiagaraGraphParameterReferenceCollection>& GetReferenceCollection();
+	NIAGARAEDITOR_API TArray<FNiagaraGraphParameterReferenceCollection>& GetReferenceCollection();
 
-	bool GetIsNamespaceModifierRenamePending() const;
+	NIAGARAEDITOR_API bool GetIsNamespaceModifierRenamePending() const;
 
-	void SetIsNamespaceModifierRenamePending(bool bIsNamespaceModifierRenamePending);
+	NIAGARAEDITOR_API void SetIsNamespaceModifierRenamePending(bool bIsNamespaceModifierRenamePending);
 
-	bool GetIsExternallyReferenced() const;
+	NIAGARAEDITOR_API bool GetIsExternallyReferenced() const;
 
-	void SetIsExternallyReferenced(bool bInIsExternallyReferenced);
+	NIAGARAEDITOR_API void SetIsExternallyReferenced(bool bInIsExternallyReferenced);
 
-	bool GetIsSourcedFromCustomStackContext() const;
+	NIAGARAEDITOR_API bool GetIsSourcedFromCustomStackContext() const;
 
-	void SetIsSourcedFromCustomStackContext(bool bInIsSourcedFromCustomStackContext);
+	NIAGARAEDITOR_API void SetIsSourcedFromCustomStackContext(bool bInIsSourcedFromCustomStackContext);
 
 private:
 	const UNiagaraScriptVariable* ScriptVar;
@@ -279,23 +279,23 @@ private:
 	TWeakPtr<TArray<FName>> ParametersWithNamespaceModifierRenamePendingWeak;
 };
 
-struct NIAGARAEDITOR_API FNiagaraScriptParameterAction : public FEdGraphSchemaAction
+struct FNiagaraScriptParameterAction : public FEdGraphSchemaAction
 {
 	FNiagaraScriptParameterAction() {}
-	FNiagaraScriptParameterAction(const FNiagaraVariable& InVariable, const FNiagaraVariableMetaData& InVariableMetaData);
+	NIAGARAEDITOR_API FNiagaraScriptParameterAction(const FNiagaraVariable& InVariable, const FNiagaraVariableMetaData& InVariableMetaData);
 };
 
-class NIAGARAEDITOR_API FNiagaraParameterGraphDragOperation : public FGraphSchemaActionDragDropAction
+class FNiagaraParameterGraphDragOperation : public FGraphSchemaActionDragDropAction
 {
 public:
 	DRAG_DROP_OPERATOR_TYPE(FNiagaraParameterGraphDragOperation, FGraphSchemaActionDragDropAction)
 
-	static TSharedRef<FNiagaraParameterGraphDragOperation> New(const TSharedPtr<FEdGraphSchemaAction>& InActionNode);
+	static NIAGARAEDITOR_API TSharedRef<FNiagaraParameterGraphDragOperation> New(const TSharedPtr<FEdGraphSchemaAction>& InActionNode);
 
 	// FGraphEditorDragDropAction interface
-	virtual void HoverTargetChanged() override;
-	virtual FReply DroppedOnNode(FVector2D ScreenPosition, FVector2D GraphPosition) override;
-	virtual FReply DroppedOnPanel(const TSharedRef< class SWidget >& Panel, FVector2D ScreenPosition, FVector2D GraphPosition, UEdGraph& Graph) override;
+	NIAGARAEDITOR_API virtual void HoverTargetChanged() override;
+	NIAGARAEDITOR_API virtual FReply DroppedOnNode(FVector2D ScreenPosition, FVector2D GraphPosition) override;
+	NIAGARAEDITOR_API virtual FReply DroppedOnPanel(const TSharedRef< class SWidget >& Panel, FVector2D ScreenPosition, FVector2D GraphPosition, UEdGraph& Graph) override;
 	// End of FGraphEditorDragDropAction
 
 	/** Set if operation is modified by alt */
@@ -305,13 +305,13 @@ public:
 	void SetCtrlDrag(bool InIsCtrlDrag) { bControlDrag = InIsCtrlDrag; }
 
 	/** Returns true if the drag operation is currently hovering over the supplied node */
-	bool IsCurrentlyHoveringNode(const UEdGraphNode* TestNode) const;
+	NIAGARAEDITOR_API bool IsCurrentlyHoveringNode(const UEdGraphNode* TestNode) const;
 
 	const TSharedPtr<FEdGraphSchemaAction>& GetSourceAction() const { return SourceAction; }
 
 protected:
 	/** Constructor */
-	FNiagaraParameterGraphDragOperation();
+	NIAGARAEDITOR_API FNiagaraParameterGraphDragOperation();
 
 	/** Structure for required node construction parameters */
 	struct FNiagaraParameterNodeConstructionParams
@@ -335,12 +335,12 @@ protected:
 		const UNiagaraScriptVariable* ScriptVar;
 	};
 
-	static void MakeGetMap(FNiagaraParameterNodeConstructionParams InParams);
-	static void MakeSetMap(FNiagaraParameterNodeConstructionParams InParams);
-	static void MakeStaticSwitch(FNiagaraParameterNodeConstructionParams InParams, const UNiagaraScriptVariable* ScriptVariable);
+	static NIAGARAEDITOR_API void MakeGetMap(FNiagaraParameterNodeConstructionParams InParams);
+	static NIAGARAEDITOR_API void MakeSetMap(FNiagaraParameterNodeConstructionParams InParams);
+	static NIAGARAEDITOR_API void MakeStaticSwitch(FNiagaraParameterNodeConstructionParams InParams, const UNiagaraScriptVariable* ScriptVariable);
 
-	virtual EVisibility GetIconVisible() const override;
-	virtual EVisibility GetErrorIconVisible() const override;
+	NIAGARAEDITOR_API virtual EVisibility GetIconVisible() const override;
+	NIAGARAEDITOR_API virtual EVisibility GetErrorIconVisible() const override;
 
 	/** Was ctrl held down at start of drag */
 	bool bControlDrag;
@@ -348,7 +348,7 @@ protected:
 	bool bAltDrag;
 };
 
-class NIAGARAEDITOR_API FNiagaraParameterDragOperation : public FDecoratedDragDropOp
+class FNiagaraParameterDragOperation : public FDecoratedDragDropOp
 {
 public:
 	DRAG_DROP_OPERATOR_TYPE(FNiagaraParameterDragOperation, FDecoratedDragDropOp)
@@ -360,16 +360,16 @@ public:
 
 	const TSharedPtr<FEdGraphSchemaAction>& GetSourceAction() const { return SourceAction; }
 
-	virtual TSharedPtr<SWidget> GetDefaultDecorator() const override;
+	NIAGARAEDITOR_API virtual TSharedPtr<SWidget> GetDefaultDecorator() const override;
 
 	void SetAdditionalText(FText InAdditionalText) { CurrentHoverText = InAdditionalText; }
-	EVisibility IsTextVisible() const;
+	NIAGARAEDITOR_API EVisibility IsTextVisible() const;
 private:
 	TSharedPtr<FEdGraphSchemaAction> SourceAction;
 	TOptional<FNiagaraVariable> TargetParameter;
 };
 
-class NIAGARAEDITOR_API FNiagaraScriptDragOperation : public FDecoratedDragDropOp
+class FNiagaraScriptDragOperation : public FDecoratedDragDropOp
 {
 public:
 	DRAG_DROP_OPERATOR_TYPE(FNiagaraScriptDragOperation, FDecoratedDragDropOp)
@@ -385,7 +385,7 @@ public:
 };
 
 
-class NIAGARAEDITOR_API INiagaraDataInterfaceNodeActionProvider
+class INiagaraDataInterfaceNodeActionProvider
 {
 public:
 
@@ -411,15 +411,15 @@ public:
 	template<typename DIClass>
 	static void GetNodeContextMenuActions(UToolMenu* Menu, UGraphNodeContextMenuContext* Context, FNiagaraFunctionSignature Signature);
 
-	static void GetNodeContextMenuActions(UClass* DIClass, UToolMenu* Menu, UGraphNodeContextMenuContext* Context, FNiagaraFunctionSignature Signature);
+	static NIAGARAEDITOR_API void GetNodeContextMenuActions(UClass* DIClass, UToolMenu* Menu, UGraphNodeContextMenuContext* Context, FNiagaraFunctionSignature Signature);
 
 	template<typename DIClass>
 	static void CollectAddPinActions(FNiagaraMenuActionCollector& Collector, UEdGraphPin* AddPin);
 
-	static void CollectAddPinActions(UClass* DIClass, FNiagaraMenuActionCollector& Collector, UEdGraphPin* AddPin);
+	static NIAGARAEDITOR_API void CollectAddPinActions(UClass* DIClass, FNiagaraMenuActionCollector& Collector, UEdGraphPin* AddPin);
 
 	/** All currently registered action providers. */
-	static TMap<FName, TUniquePtr<INiagaraDataInterfaceNodeActionProvider>> RegisteredActionProviders;
+	static NIAGARAEDITOR_API TMap<FName, TUniquePtr<INiagaraDataInterfaceNodeActionProvider>> RegisteredActionProviders;
 };
 
 template<typename DIClass, typename ActionProviderClass>
@@ -470,4 +470,4 @@ public:
 
 	virtual void GetNodeContextMenuActionsImpl(UToolMenu* Menu, UGraphNodeContextMenuContext* Context, FNiagaraFunctionSignature Signature) const override;
 	virtual void CollectAddPinActionsImpl(FNiagaraMenuActionCollector& Collector, UEdGraphPin* AddPin)const override;
-};
\ No newline at end of file
+};
diff --git a/Engine/Plugins/FX/Niagara/Source/NiagaraEditor/Public/NiagaraClipboard.h b/Engine/Plugins/FX/Niagara/Source/NiagaraEditor/Public/NiagaraClipboard.h
--- a/Engine/Plugins/FX/Niagara/Source/NiagaraEditor/Public/NiagaraClipboard.h
+++ b/Engine/Plugins/FX/Niagara/Source/NiagaraEditor/Public/NiagaraClipboard.h
@@ -30,25 +30,25 @@ enum class ENiagaraClipboardFunctionInputValueMode
 
 class UNiagaraClipboardFunction;
 
-UCLASS()
-class NIAGARAEDITOR_API UNiagaraClipboardFunctionInput : public UObject
+UCLASS(MinimalAPI)
+class UNiagaraClipboardFunctionInput : public UObject
 {
 	GENERATED_BODY()
 
 public:
-	static const UNiagaraClipboardFunctionInput* CreateLocalValue(UObject* InOuter, FName InInputName, FNiagaraTypeDefinition InInputType, TOptional<bool> bInEditConditionValue, TArray<uint8>& InLocalValueData);
+	static NIAGARAEDITOR_API const UNiagaraClipboardFunctionInput* CreateLocalValue(UObject* InOuter, FName InInputName, FNiagaraTypeDefinition InInputType, TOptional<bool> bInEditConditionValue, TArray<uint8>& InLocalValueData);
 
-	static const UNiagaraClipboardFunctionInput* CreateLinkedValue(UObject* InOuter, FName InInputName, FNiagaraTypeDefinition InInputType, TOptional<bool> bInEditConditionValue, FName InLinkedValue);
+	static NIAGARAEDITOR_API const UNiagaraClipboardFunctionInput* CreateLinkedValue(UObject* InOuter, FName InInputName, FNiagaraTypeDefinition InInputType, TOptional<bool> bInEditConditionValue, FName InLinkedValue);
 
-	static const UNiagaraClipboardFunctionInput* CreateDataValue(UObject* InOuter, FName InInputName, FNiagaraTypeDefinition InInputType, TOptional<bool> bInEditConditionValue, UNiagaraDataInterface* InDataValue);
+	static NIAGARAEDITOR_API const UNiagaraClipboardFunctionInput* CreateDataValue(UObject* InOuter, FName InInputName, FNiagaraTypeDefinition InInputType, TOptional<bool> bInEditConditionValue, UNiagaraDataInterface* InDataValue);
 
-	static const UNiagaraClipboardFunctionInput* CreateObjectAssetValue(UObject* InOuter, FName InInputName, FNiagaraTypeDefinition InInputType, TOptional<bool> bInEditConditionValue, UObject* InObject);
+	static NIAGARAEDITOR_API const UNiagaraClipboardFunctionInput* CreateObjectAssetValue(UObject* InOuter, FName InInputName, FNiagaraTypeDefinition InInputType, TOptional<bool> bInEditConditionValue, UObject* InObject);
 
-	static const UNiagaraClipboardFunctionInput* CreateExpressionValue(UObject* InOuter, FName InInputName, FNiagaraTypeDefinition InInputType, TOptional<bool> bInEditConditionValue, const FString& InExpressionValue);
+	static NIAGARAEDITOR_API const UNiagaraClipboardFunctionInput* CreateExpressionValue(UObject* InOuter, FName InInputName, FNiagaraTypeDefinition InInputType, TOptional<bool> bInEditConditionValue, const FString& InExpressionValue);
 
-	static const UNiagaraClipboardFunctionInput* CreateDynamicValue(UObject* InOuter, FName InInputName, FNiagaraTypeDefinition InInputType, TOptional<bool> bInEditConditionValue, FString InDynamicValueName, UNiagaraScript* InDynamicValue, const FGuid& InScriptVersion);
+	static NIAGARAEDITOR_API const UNiagaraClipboardFunctionInput* CreateDynamicValue(UObject* InOuter, FName InInputName, FNiagaraTypeDefinition InInputType, TOptional<bool> bInEditConditionValue, FString InDynamicValueName, UNiagaraScript* InDynamicValue, const FGuid& InScriptVersion);
 
-	static const UNiagaraClipboardFunctionInput* CreateDefaultInputValue(UObject* InOuter, FName InInputName, FNiagaraTypeDefinition InInputType);
+	static NIAGARAEDITOR_API const UNiagaraClipboardFunctionInput* CreateDefaultInputValue(UObject* InOuter, FName InInputName, FNiagaraTypeDefinition InInputType);
 	
 	UPROPERTY()
 	FName InputName;
@@ -83,7 +83,7 @@ public:
 	UPROPERTY()
 	TObjectPtr<UNiagaraClipboardFunction> Dynamic;
 
-	bool CopyValuesFrom(const UNiagaraClipboardFunctionInput* InOther);
+	NIAGARAEDITOR_API bool CopyValuesFrom(const UNiagaraClipboardFunctionInput* InOther);
 
 	const FNiagaraTypeDefinition& GetTypeDef() const { return InputType; };
 
@@ -96,17 +96,17 @@ enum class ENiagaraClipboardFunctionScriptMode
 	Assignment
 };
 
-UCLASS()
-class NIAGARAEDITOR_API UNiagaraClipboardFunction : public UObject
+UCLASS(MinimalAPI)
+class UNiagaraClipboardFunction : public UObject
 {
 	GENERATED_BODY()
 
 public:
 	DECLARE_DYNAMIC_DELEGATE_OneParam(FOnPastedFunctionCallNode, UNiagaraNodeFunctionCall*, PastedFunctionCall);
 
-	static UNiagaraClipboardFunction* CreateScriptFunction(UObject* InOuter, FString InFunctionName, UNiagaraScript* InScript, const FGuid& InScriptVersion = FGuid(), const TArray<FNiagaraStackMessage> InStackMessages = TArray<FNiagaraStackMessage>());
+	static NIAGARAEDITOR_API UNiagaraClipboardFunction* CreateScriptFunction(UObject* InOuter, FString InFunctionName, UNiagaraScript* InScript, const FGuid& InScriptVersion = FGuid(), const TArray<FNiagaraStackMessage> InStackMessages = TArray<FNiagaraStackMessage>());
 
-	static UNiagaraClipboardFunction* CreateAssignmentFunction(UObject* InOuter, FString InFunctionName, const TArray<FNiagaraVariable>& InAssignmentTargets, const TArray<FString>& InAssignmentDefaults);
+	static NIAGARAEDITOR_API UNiagaraClipboardFunction* CreateAssignmentFunction(UObject* InOuter, FString InFunctionName, const TArray<FNiagaraVariable>& InAssignmentTargets, const TArray<FString>& InAssignmentDefaults);
 
 	UPROPERTY()
 	FString FunctionName;
@@ -140,7 +140,7 @@ public:
 };
 
 USTRUCT()
-struct NIAGARAEDITOR_API FNiagaraClipboardScriptVariable
+struct FNiagaraClipboardScriptVariable
 {
 	GENERATED_BODY()
 
@@ -164,13 +164,13 @@ struct NIAGARAEDITOR_API FNiagaraClipboardScriptVariable
 		return OriginalChangeId == OtherVariable.OriginalChangeId;
 	}
 };
-UCLASS()
-class NIAGARAEDITOR_API UNiagaraClipboardContent : public UObject
+UCLASS(MinimalAPI)
+class UNiagaraClipboardContent : public UObject
 {
 	GENERATED_BODY()
 
 public:
-	static UNiagaraClipboardContent* Create();
+	static NIAGARAEDITOR_API UNiagaraClipboardContent* Create();
 
 	UPROPERTY()
 	TArray<TObjectPtr<const UNiagaraClipboardFunction>> Functions;
@@ -196,69 +196,69 @@ public:
 	mutable bool bFixupPasteIndexForScriptDependenciesInStack = false;
 };
 
-class NIAGARAEDITOR_API FNiagaraClipboard
+class FNiagaraClipboard
 {
 public:
-	FNiagaraClipboard();
+	NIAGARAEDITOR_API FNiagaraClipboard();
 
-	void SetClipboardContent(UNiagaraClipboardContent* ClipboardContent);
+	NIAGARAEDITOR_API void SetClipboardContent(UNiagaraClipboardContent* ClipboardContent);
 
-	const UNiagaraClipboardContent* GetClipboardContent() const;
+	NIAGARAEDITOR_API const UNiagaraClipboardContent* GetClipboardContent() const;
 };
 
-UCLASS()
-class NIAGARAEDITOR_API UNiagaraClipboardEditorScriptingUtilities : public UObject
+UCLASS(MinimalAPI)
+class UNiagaraClipboardEditorScriptingUtilities : public UObject
 {
 	GENERATED_BODY()
 
 public:
 	UFUNCTION(BlueprintPure, Category = "Input")
-	static void TryGetInputByName(const TArray<UNiagaraClipboardFunctionInput*>& InInputs, FName InInputName, bool& bOutSucceeded, UNiagaraClipboardFunctionInput*& OutInput);
+	static NIAGARAEDITOR_API void TryGetInputByName(const TArray<UNiagaraClipboardFunctionInput*>& InInputs, FName InInputName, bool& bOutSucceeded, UNiagaraClipboardFunctionInput*& OutInput);
 
 	UFUNCTION(BlueprintPure, Category = "Input")
-	static void TryGetLocalValueAsFloat(const UNiagaraClipboardFunctionInput* InInput, bool& bOutSucceeded, float& OutValue);
+	static NIAGARAEDITOR_API void TryGetLocalValueAsFloat(const UNiagaraClipboardFunctionInput* InInput, bool& bOutSucceeded, float& OutValue);
 
 	UFUNCTION(BlueprintPure, Category = "Input")
-	static void TryGetLocalValueAsInt(const UNiagaraClipboardFunctionInput* InInput, bool& bOutSucceeded, int32& OutValue);
+	static NIAGARAEDITOR_API void TryGetLocalValueAsInt(const UNiagaraClipboardFunctionInput* InInput, bool& bOutSucceeded, int32& OutValue);
 
 	UFUNCTION(BlueprintPure, Category = "Input")
-	static void TrySetLocalValueAsInt(UNiagaraClipboardFunctionInput* InInput, bool& bOutSucceeded, int32 InValue, bool bLooseTyping = true);
+	static NIAGARAEDITOR_API void TrySetLocalValueAsInt(UNiagaraClipboardFunctionInput* InInput, bool& bOutSucceeded, int32 InValue, bool bLooseTyping = true);
 
 	UFUNCTION(BlueprintPure, Category = "Input")
-	static FName GetTypeName(const UNiagaraClipboardFunctionInput* InInput);
+	static NIAGARAEDITOR_API FName GetTypeName(const UNiagaraClipboardFunctionInput* InInput);
 
 	UFUNCTION(BlueprintPure, Category = "Input")
-	static UNiagaraClipboardFunctionInput* CreateFloatLocalValueInput(UObject* InOuter, FName InInputName, bool bInHasEditCondition, bool bInEditConditionValue, float InLocalValue);
+	static NIAGARAEDITOR_API UNiagaraClipboardFunctionInput* CreateFloatLocalValueInput(UObject* InOuter, FName InInputName, bool bInHasEditCondition, bool bInEditConditionValue, float InLocalValue);
 
 	UFUNCTION(BlueprintPure, Category = "Input")
-	static UNiagaraClipboardFunctionInput* CreateVec2LocalValueInput(UObject* InOuter, FName InInputName, bool bInHasEditCondition, bool bInEditConditionValue, FVector2D InVec2Value);
+	static NIAGARAEDITOR_API UNiagaraClipboardFunctionInput* CreateVec2LocalValueInput(UObject* InOuter, FName InInputName, bool bInHasEditCondition, bool bInEditConditionValue, FVector2D InVec2Value);
 
 	UFUNCTION(BlueprintPure, Category = "Input")
-	static UNiagaraClipboardFunctionInput* CreateVec3LocalValueInput(UObject* InOuter, FName InInputName, bool bInHasEditCondition, bool bInEditConditionValue, FVector InVec3Value);
+	static NIAGARAEDITOR_API UNiagaraClipboardFunctionInput* CreateVec3LocalValueInput(UObject* InOuter, FName InInputName, bool bInHasEditCondition, bool bInEditConditionValue, FVector InVec3Value);
 
 	UFUNCTION(BlueprintPure, Category = "Input")
-	static UNiagaraClipboardFunctionInput* CreateIntLocalValueInput(UObject* InOuter, FName InInputName, bool bInHasEditCondition, bool bInEditConditionValue, int32 InLocalValue);
+	static NIAGARAEDITOR_API UNiagaraClipboardFunctionInput* CreateIntLocalValueInput(UObject* InOuter, FName InInputName, bool bInHasEditCondition, bool bInEditConditionValue, int32 InLocalValue);
 	
 	UFUNCTION(BlueprintPure, Category = "Input")
-	static UNiagaraClipboardFunctionInput* CreateBoolLocalValueInput(UObject* InOuter, FName InInputName, bool bInHasEditCondition, bool bInEditConditionValue, bool InBoolValue);
+	static NIAGARAEDITOR_API UNiagaraClipboardFunctionInput* CreateBoolLocalValueInput(UObject* InOuter, FName InInputName, bool bInHasEditCondition, bool bInEditConditionValue, bool InBoolValue);
 	
 	UFUNCTION(BlueprintPure, Category = "Input")
-	static UNiagaraClipboardFunctionInput* CreateStructLocalValueInput(UObject* InOuter, FName InInputName, bool bInHasEditCondition, bool bInEditConditionValue, UUserDefinedStruct* InStructValue);
+	static NIAGARAEDITOR_API UNiagaraClipboardFunctionInput* CreateStructLocalValueInput(UObject* InOuter, FName InInputName, bool bInHasEditCondition, bool bInEditConditionValue, UUserDefinedStruct* InStructValue);
 
 	UFUNCTION(BlueprintPure, Category = "Input")
-	static UNiagaraClipboardFunctionInput* CreateEnumLocalValueInput(UObject* InOuter, FName InInputName, bool bInHasEditCondition, bool bInEditCoditionValue, UUserDefinedEnum* InEnumType, int32 InEnumValue);
+	static NIAGARAEDITOR_API UNiagaraClipboardFunctionInput* CreateEnumLocalValueInput(UObject* InOuter, FName InInputName, bool bInHasEditCondition, bool bInEditCoditionValue, UUserDefinedEnum* InEnumType, int32 InEnumValue);
 
 	UFUNCTION(BlueprintPure, Category = "Input")
-	static UNiagaraClipboardFunctionInput* CreateLinkedValueInput(UObject* InOuter, FName InInputName, FName InInputTypeName, bool bInHasEditCondition, bool bInEditConditionValue, FName InLinkedValue);
+	static NIAGARAEDITOR_API UNiagaraClipboardFunctionInput* CreateLinkedValueInput(UObject* InOuter, FName InInputName, FName InInputTypeName, bool bInHasEditCondition, bool bInEditConditionValue, FName InLinkedValue);
 
 	UFUNCTION(BlueprintPure, Category = "Input")
-	static UNiagaraClipboardFunctionInput* CreateDataValueInput(UObject* InOuter, FName InInputName, bool bInHasEditCondition, bool bInEditConditionValue, UNiagaraDataInterface* InDataValue);
+	static NIAGARAEDITOR_API UNiagaraClipboardFunctionInput* CreateDataValueInput(UObject* InOuter, FName InInputName, bool bInHasEditCondition, bool bInEditConditionValue, UNiagaraDataInterface* InDataValue);
 
 	UFUNCTION(BlueprintPure, Category = "Input")
-	static UNiagaraClipboardFunctionInput* CreateExpressionValueInput(UObject* InOuter, FName InInputName, FName InInputTypeName, bool bInHasEditCondition, bool bInEditConditionValue, const FString& InExpressionValue);
+	static NIAGARAEDITOR_API UNiagaraClipboardFunctionInput* CreateExpressionValueInput(UObject* InOuter, FName InInputName, FName InInputTypeName, bool bInHasEditCondition, bool bInEditConditionValue, const FString& InExpressionValue);
 
 	UFUNCTION(BlueprintPure, Category = "Input")
-	static UNiagaraClipboardFunctionInput* CreateDynamicValueInput(UObject* InOuter, FName InInputName, FName InInputTypeName, bool bInHasEditCondition, bool bInEditConditionValue, FString InDynamicValueName, UNiagaraScript* InDynamicValue);
+	static NIAGARAEDITOR_API UNiagaraClipboardFunctionInput* CreateDynamicValueInput(UObject* InOuter, FName InInputName, FName InInputTypeName, bool bInHasEditCondition, bool bInEditConditionValue, FString InDynamicValueName, UNiagaraScript* InDynamicValue);
 
-	static FNiagaraTypeDefinition GetRegisteredTypeDefinitionByName(FName InTypeName);
+	static NIAGARAEDITOR_API FNiagaraTypeDefinition GetRegisteredTypeDefinitionByName(FName InTypeName);
 };
diff --git a/Engine/Plugins/FX/Niagara/Source/NiagaraEditor/Public/NiagaraEditorCommon.h b/Engine/Plugins/FX/Niagara/Source/NiagaraEditor/Public/NiagaraEditorCommon.h
--- a/Engine/Plugins/FX/Niagara/Source/NiagaraEditor/Public/NiagaraEditorCommon.h
+++ b/Engine/Plugins/FX/Niagara/Source/NiagaraEditor/Public/NiagaraEditorCommon.h
@@ -190,7 +190,7 @@ enum class EParameterDefinitionMatchState : uint8
 
 
 USTRUCT()
-struct NIAGARAEDITOR_API FFunctionInputSummaryViewKey
+struct FFunctionInputSummaryViewKey
 {
 	GENERATED_BODY()
 
@@ -239,7 +239,7 @@ public:
 };
 
 USTRUCT()
-struct NIAGARAEDITOR_API FFunctionInputSummaryViewMetadata
+struct FFunctionInputSummaryViewMetadata
 {
 	GENERATED_BODY()
 
@@ -267,4 +267,4 @@ public:
 			Category == Other.Category &&
 			SortIndex == Other.SortIndex;
 	}
-};
\ No newline at end of file
+};
diff --git a/Engine/Plugins/FX/Niagara/Source/NiagaraEditor/Public/NiagaraEditorModule.h b/Engine/Plugins/FX/Niagara/Source/NiagaraEditor/Public/NiagaraEditorModule.h
--- a/Engine/Plugins/FX/Niagara/Source/NiagaraEditor/Public/NiagaraEditorModule.h
+++ b/Engine/Plugins/FX/Niagara/Source/NiagaraEditor/Public/NiagaraEditorModule.h
@@ -58,7 +58,7 @@ extern NIAGARAEDITOR_API int32 GbEnableExperimentalInlineDynamicInputs;
 extern NIAGARAEDITOR_API int32 GbEnableCustomInlineDynamicInputFormats;
 
 /* Defines methods for allowing external modules to supply widgets to the core editor module. */
-class NIAGARAEDITOR_API INiagaraEditorWidgetProvider
+class INiagaraEditorWidgetProvider
 {
 public:
 	virtual TSharedRef<SWidget> CreateStackView(UNiagaraStackViewModel& StackViewModel) const = 0;
@@ -101,7 +101,7 @@ private:
 };
 
 USTRUCT()
-struct NIAGARAEDITOR_API FNiagaraRendererCreationInfo
+struct FNiagaraRendererCreationInfo
 {
 	DECLARE_DELEGATE_RetVal_OneParam(UNiagaraRendererProperties*, FRendererFactory, UObject* OuterEmitter);
 
@@ -192,7 +192,7 @@ public:
 
 	/** Registers a new renderer creation delegate with the display name it's going to use for the UI. */
 	NIAGARAEDITOR_API void RegisterRendererCreationInfo(FNiagaraRendererCreationInfo RendererCreationInfo);
-	NIAGARAEDITOR_API const TArray<FNiagaraRendererCreationInfo>& GetRendererCreationInfos() const { return RendererCreationInfo; }
+	const TArray<FNiagaraRendererCreationInfo>& GetRendererCreationInfos() const { return RendererCreationInfo; }
 	
 	void RegisterParameterTrackCreatorForType(const UScriptStruct& StructType, FOnCreateMovieSceneTrackForParameter CreateTrack);
 	void UnregisterParameterTrackCreatorForType(const UScriptStruct& StructType);
diff --git a/Engine/Plugins/FX/Niagara/Source/NiagaraEditor/Public/NiagaraEditorSettings.h b/Engine/Plugins/FX/Niagara/Source/NiagaraEditor/Public/NiagaraEditorSettings.h
--- a/Engine/Plugins/FX/Niagara/Source/NiagaraEditor/Public/NiagaraEditorSettings.h
+++ b/Engine/Plugins/FX/Niagara/Source/NiagaraEditor/Public/NiagaraEditorSettings.h
@@ -90,13 +90,13 @@ enum class ENiagaraAddDefaultsTrackMode : uint8
 };
 
 USTRUCT()
-struct NIAGARAEDITOR_API FNiagaraNamespaceMetadata
+struct FNiagaraNamespaceMetadata
 {
 	GENERATED_BODY()
 
-	FNiagaraNamespaceMetadata();
+	NIAGARAEDITOR_API FNiagaraNamespaceMetadata();
 
-	FNiagaraNamespaceMetadata(TArray<FName> InNamespaces, FName InRequiredNamespaceModifier = NAME_None);
+	NIAGARAEDITOR_API FNiagaraNamespaceMetadata(TArray<FName> InNamespaces, FName InRequiredNamespaceModifier = NAME_None);
 
 	//FNiagaraNamespaceMetadata(const FNiagaraNamespaceMetadata& Other);
 
@@ -209,7 +209,7 @@ private:
 };
 
 USTRUCT()
-struct NIAGARAEDITOR_API FNiagaraCurveTemplate
+struct FNiagaraCurveTemplate
 {
 	GENERATED_BODY();
 
@@ -221,11 +221,11 @@ struct NIAGARAEDITOR_API FNiagaraCurveTemplate
 };
 
 USTRUCT()
-struct NIAGARAEDITOR_API FNiagaraActionColors
+struct FNiagaraActionColors
 {
 	GENERATED_BODY()
 
-	FNiagaraActionColors();
+	NIAGARAEDITOR_API FNiagaraActionColors();
 	
 	UPROPERTY(EditAnywhere, Category = Color)
 	FLinearColor NiagaraColor;
@@ -242,7 +242,7 @@ struct NIAGARAEDITOR_API FNiagaraActionColors
 
 
 USTRUCT()
-struct NIAGARAEDITOR_API FNiagaraParameterPanelSectionStorage
+struct FNiagaraParameterPanelSectionStorage
 {
 	GENERATED_BODY()
 		
@@ -265,8 +265,8 @@ FORCEINLINE uint32 GetTypeHash(const FNiagaraNamespaceMetadata& NamespaceMetaDat
 	return GetTypeHash(NamespaceMetaData.GetGuid());
 }
 
-UCLASS(config = Niagara, defaultconfig, meta=(DisplayName="Niagara"))
-class NIAGARAEDITOR_API UNiagaraEditorSettings : public UDeveloperSettings
+UCLASS(config = Niagara, defaultconfig, meta=(DisplayName="Niagara"), MinimalAPI)
+class UNiagaraEditorSettings : public UDeveloperSettings
 {
 public:
 	GENERATED_UCLASS_BODY()
@@ -317,122 +317,122 @@ public:
 	UPROPERTY(config, EditAnywhere, Category = Niagara)
 	ENiagaraAddDefaultsTrackMode DefaultsSequencerSubtracks = ENiagaraAddDefaultsTrackMode::NoSubtracks;
 
-	TArray<float> GetPlaybackSpeeds() const;
+	NIAGARAEDITOR_API TArray<float> GetPlaybackSpeeds() const;
 
 	/** Gets whether or not auto-compile is enabled in the editors. */
-	bool GetAutoCompile() const;
+	NIAGARAEDITOR_API bool GetAutoCompile() const;
 
 	/** Sets whether or not auto-compile is enabled in the editors. */
-	void SetAutoCompile(bool bInAutoCompile);
+	NIAGARAEDITOR_API void SetAutoCompile(bool bInAutoCompile);
 
 	/** Gets whether or not simulations should start playing automatically when the emitter or system editor is opened, or when the data is changed in the editor. */
-	bool GetAutoPlay() const;
+	NIAGARAEDITOR_API bool GetAutoPlay() const;
 
 	/** Sets whether or not simulations should start playing automatically when the emitter or system editor is opened, or when the data is changed in the editor. */
-	void SetAutoPlay(bool bInAutoPlay);
+	NIAGARAEDITOR_API void SetAutoPlay(bool bInAutoPlay);
 
 	/** Gets whether or not the simulation should reset when a value on the emitter or system is changed. */
-	bool GetResetSimulationOnChange() const;
+	NIAGARAEDITOR_API bool GetResetSimulationOnChange() const;
 
 	/** Sets whether or not the simulation should reset when a value on the emitter or system is changed. */
-	void SetResetSimulationOnChange(bool bInResetSimulationOnChange);
+	NIAGARAEDITOR_API void SetResetSimulationOnChange(bool bInResetSimulationOnChange);
 
 	/** Gets whether or not to rerun the simulation to the current time when making modifications while paused. */
-	bool GetResimulateOnChangeWhilePaused() const;
+	NIAGARAEDITOR_API bool GetResimulateOnChangeWhilePaused() const;
 
 	/** Sets whether or not to rerun the simulation to the current time when making modifications while paused. */
-	void SetResimulateOnChangeWhilePaused(bool bInResimulateOnChangeWhilePaused);
+	NIAGARAEDITOR_API void SetResimulateOnChangeWhilePaused(bool bInResimulateOnChangeWhilePaused);
 
 	/** Gets whether or not to reset all components that include the system that is currently being reset */
-	bool GetResetDependentSystemsWhenEditingEmitters() const;
+	NIAGARAEDITOR_API bool GetResetDependentSystemsWhenEditingEmitters() const;
 
 	/** Sets whether or not to reset all components that include the system that is currently being reset */
-	void SetResetDependentSystemsWhenEditingEmitters(bool bInResetDependentSystemsWhenEditingEmitters);
+	NIAGARAEDITOR_API void SetResetDependentSystemsWhenEditingEmitters(bool bInResetDependentSystemsWhenEditingEmitters);
 
 	/** Gets whether or not to display advanced categories for the parameter panel. */
-	bool GetDisplayAdvancedParameterPanelCategories() const;
+	NIAGARAEDITOR_API bool GetDisplayAdvancedParameterPanelCategories() const;
 
 	/** Sets whether or not to display advanced categories for the parameter panel. */
-	void SetDisplayAdvancedParameterPanelCategories(bool bInDisplayAdvancedParameterPanelCategories);
+	NIAGARAEDITOR_API void SetDisplayAdvancedParameterPanelCategories(bool bInDisplayAdvancedParameterPanelCategories);
 
-	bool GetDisplayAffectedAssetStats() const;
-	int32 GetAssetStatsSearchLimit() const;
+	NIAGARAEDITOR_API bool GetDisplayAffectedAssetStats() const;
+	NIAGARAEDITOR_API int32 GetAssetStatsSearchLimit() const;
 
-	bool IsShowGridInViewport() const;
-	void SetShowGridInViewport(bool bShowGridInViewport);
+	NIAGARAEDITOR_API bool IsShowGridInViewport() const;
+	NIAGARAEDITOR_API void SetShowGridInViewport(bool bShowGridInViewport);
 	
-	bool IsShowParticleCountsInViewport() const;
-	void SetShowParticleCountsInViewport(bool bShowParticleCountsInViewport);
+	NIAGARAEDITOR_API bool IsShowParticleCountsInViewport() const;
+	NIAGARAEDITOR_API void SetShowParticleCountsInViewport(bool bShowParticleCountsInViewport);
 	
-	FNiagaraNewAssetDialogConfig GetNewAssetDailogConfig(FName InDialogConfigKey) const;
+	NIAGARAEDITOR_API FNiagaraNewAssetDialogConfig GetNewAssetDailogConfig(FName InDialogConfigKey) const;
 
-	void SetNewAssetDialogConfig(FName InDialogConfigKey, const FNiagaraNewAssetDialogConfig& InNewAssetDialogConfig);
+	NIAGARAEDITOR_API void SetNewAssetDialogConfig(FName InDialogConfigKey, const FNiagaraNewAssetDialogConfig& InNewAssetDialogConfig);
 
-	FNiagaraNamespaceMetadata GetDefaultNamespaceMetadata() const;
-	FNiagaraNamespaceMetadata GetMetaDataForNamespaces(TArray<FName> Namespaces) const;
-	FNiagaraNamespaceMetadata GetMetaDataForId(const FGuid& NamespaceId) const;
-	const FGuid& GetIdForUsage(ENiagaraScriptUsage Usage) const;
-	const TArray<FNiagaraNamespaceMetadata>& GetAllNamespaceMetadata() const;
+	NIAGARAEDITOR_API FNiagaraNamespaceMetadata GetDefaultNamespaceMetadata() const;
+	NIAGARAEDITOR_API FNiagaraNamespaceMetadata GetMetaDataForNamespaces(TArray<FName> Namespaces) const;
+	NIAGARAEDITOR_API FNiagaraNamespaceMetadata GetMetaDataForId(const FGuid& NamespaceId) const;
+	NIAGARAEDITOR_API const FGuid& GetIdForUsage(ENiagaraScriptUsage Usage) const;
+	NIAGARAEDITOR_API const TArray<FNiagaraNamespaceMetadata>& GetAllNamespaceMetadata() const;
 
-	FNiagaraNamespaceMetadata GetDefaultNamespaceModifierMetadata() const;
-	FNiagaraNamespaceMetadata GetMetaDataForNamespaceModifier(FName NamespaceModifier) const;
-	const TArray<FNiagaraNamespaceMetadata>& GetAllNamespaceModifierMetadata() const;
+	NIAGARAEDITOR_API FNiagaraNamespaceMetadata GetDefaultNamespaceModifierMetadata() const;
+	NIAGARAEDITOR_API FNiagaraNamespaceMetadata GetMetaDataForNamespaceModifier(FName NamespaceModifier) const;
+	NIAGARAEDITOR_API const TArray<FNiagaraNamespaceMetadata>& GetAllNamespaceModifierMetadata() const;
 
-	const TArray<FNiagaraCurveTemplate>& GetCurveTemplates() const;
+	NIAGARAEDITOR_API const TArray<FNiagaraCurveTemplate>& GetCurveTemplates() const;
 	
 	// Begin UDeveloperSettings Interface
-	virtual FName GetCategoryName() const override;
-	virtual FText GetSectionText() const override;
+	NIAGARAEDITOR_API virtual FName GetCategoryName() const override;
+	NIAGARAEDITOR_API virtual FText GetSectionText() const override;
 	// END UDeveloperSettings Interface
 
-	virtual void PostEditChangeProperty(FPropertyChangedEvent& PropertyChangedEvent) override;
+	NIAGARAEDITOR_API virtual void PostEditChangeProperty(FPropertyChangedEvent& PropertyChangedEvent) override;
 
 	DECLARE_MULTICAST_DELEGATE_TwoParams(FOnNiagaraEditorSettingsChanged, const FString&, const UNiagaraEditorSettings*);
 
 	/** Gets a multicast delegate which is called whenever one of the parameters in this settings object changes. */
-	static FOnNiagaraEditorSettingsChanged& OnSettingsChanged();
+	static NIAGARAEDITOR_API FOnNiagaraEditorSettingsChanged& OnSettingsChanged();
 
 	const TMap<FString, FString>& GetHLSLKeywordReplacementsMap()const { return HLSLKeywordReplacements; }
 
-	FLinearColor GetSourceColor(EScriptSource Source) const;
+	NIAGARAEDITOR_API FLinearColor GetSourceColor(EScriptSource Source) const;
 
-	FNiagaraParameterPanelSectionStorage& FindOrAddParameterPanelSectionStorage(FGuid PanelSectionId, bool& bOutAdded);
+	NIAGARAEDITOR_API FNiagaraParameterPanelSectionStorage& FindOrAddParameterPanelSectionStorage(FGuid PanelSectionId, bool& bOutAdded);
 
 	DECLARE_DELEGATE_RetVal_OneParam(bool, FOnIsClassAllowed, const UClass* /*InClass*/);
 	DECLARE_DELEGATE_RetVal_OneParam(bool, FOnIsClassPathAllowed, const FTopLevelAssetPath& /*InClassPath*/);
 	DECLARE_DELEGATE_RetVal_OneParam(bool, FOnShouldFilterAssetByClassUsage, const FTopLevelAssetPath& /*InAssetPath*/)
 	
 	/** Sets a delegate that allows external code to restrict which features can be used within the niagara editor by filtering which classes are allowed. */
-	void SetOnIsClassAllowed(const FOnIsClassAllowed& InOnIsClassAllowed);
+	NIAGARAEDITOR_API void SetOnIsClassAllowed(const FOnIsClassAllowed& InOnIsClassAllowed);
 
 	/** Sets a delegate that allows external code to restrict which features can be used within the niagara editor by filtering which classes are allowed by class path. */
-	void SetOnIsClassPathAllowed(const FOnIsClassPathAllowed& InOnIsClassPathAllowed);
+	NIAGARAEDITOR_API void SetOnIsClassPathAllowed(const FOnIsClassPathAllowed& InOnIsClassPathAllowed);
 
 	/** Sets a delegate that allows external code to restrict which features can be used within the niagara editor by filtering which assets are allowed. */
-	void SetOnShouldFilterAssetByClassUsage(const FOnShouldFilterAssetByClassUsage& InOnShouldFilterAssetByClassUsage);
+	NIAGARAEDITOR_API void SetOnShouldFilterAssetByClassUsage(const FOnShouldFilterAssetByClassUsage& InOnShouldFilterAssetByClassUsage);
 
 	/** Returns whether or not the supplied class can be used in the current editor context. */
-	bool IsAllowedClass(const UClass* InClass) const;
+	NIAGARAEDITOR_API bool IsAllowedClass(const UClass* InClass) const;
 
 	/** Returns whether or not the class referenced by the supplied class path can be used in the current editor context. */
-	bool IsAllowedClassPath(const FTopLevelAssetPath& InClassPath) const;
+	NIAGARAEDITOR_API bool IsAllowedClassPath(const FTopLevelAssetPath& InClassPath) const;
 
 	/** Returns whether or not the supplied niagara type definition can be used in the current editor context. */
-	bool IsAllowedTypeDefinition(const FNiagaraTypeDefinition& InTypeDefinition) const;
+	NIAGARAEDITOR_API bool IsAllowedTypeDefinition(const FNiagaraTypeDefinition& InTypeDefinition) const;
 
-	FAssetRegistryTag CreateClassUsageAssetRegistryTag(const UObject* SourceObject) const;
+	NIAGARAEDITOR_API FAssetRegistryTag CreateClassUsageAssetRegistryTag(const UObject* SourceObject) const;
 
-	bool IsAllowedAssetByClassUsage(const FAssetData& InAssetData) const;
-	bool IsAllowedAssetObjectByClassUsage(const UObject& InAssetObject) const;
+	NIAGARAEDITOR_API bool IsAllowedAssetByClassUsage(const FAssetData& InAssetData) const;
+	NIAGARAEDITOR_API bool IsAllowedAssetObjectByClassUsage(const UObject& InAssetObject) const;
 
-	bool GetUpdateStackValuesOnCommitOnly() const;
+	NIAGARAEDITOR_API bool GetUpdateStackValuesOnCommitOnly() const;
 
 private:
-	bool IsAllowedAssetObjectByClassUsageInternal(const UObject& InAssetObject, TSet<const UObject*>& CheckedAssetObjects) const;
+	NIAGARAEDITOR_API bool IsAllowedAssetObjectByClassUsageInternal(const UObject& InAssetObject, TSet<const UObject*>& CheckedAssetObjects) const;
 
-	void SetupNamespaceMetadata();
-	void BuildCachedPlaybackSpeeds() const;
-	bool ShouldTrackClassUsage(const UClass* InClass) const;
+	NIAGARAEDITOR_API void SetupNamespaceMetadata();
+	NIAGARAEDITOR_API void BuildCachedPlaybackSpeeds() const;
+	NIAGARAEDITOR_API bool ShouldTrackClassUsage(const UClass* InClass) const;
 protected:
 	FOnNiagaraEditorSettingsChanged SettingsChangedDelegate;
 private:
@@ -531,12 +531,12 @@ private:
 	TArray<UClass*> TrackedUsageBaseClasses;
 
 public:
-	bool IsShowInstructionsCount() const;
-	void SetShowInstructionsCount(bool bShowInstructionsCount);
-	bool IsShowEmitterExecutionOrder() const;
-	void SetShowEmitterExecutionOrder(bool bShowEmitterExecutionOrder);
-	bool IsShowGpuTickInformation() const;
-	void SetShowGpuTickInformation(bool bShowGpuTickInformation);
-	bool IsShowMemoryInfo() const;
-	void SetShowMemoryInfo(bool bInShowInfo);
+	NIAGARAEDITOR_API bool IsShowInstructionsCount() const;
+	NIAGARAEDITOR_API void SetShowInstructionsCount(bool bShowInstructionsCount);
+	NIAGARAEDITOR_API bool IsShowEmitterExecutionOrder() const;
+	NIAGARAEDITOR_API void SetShowEmitterExecutionOrder(bool bShowEmitterExecutionOrder);
+	NIAGARAEDITOR_API bool IsShowGpuTickInformation() const;
+	NIAGARAEDITOR_API void SetShowGpuTickInformation(bool bShowGpuTickInformation);
+	NIAGARAEDITOR_API bool IsShowMemoryInfo() const;
+	NIAGARAEDITOR_API void SetShowMemoryInfo(bool bInShowInfo);
 };
diff --git a/Engine/Plugins/FX/Niagara/Source/NiagaraEditor/Public/NiagaraEditorStyle.h b/Engine/Plugins/FX/Niagara/Source/NiagaraEditor/Public/NiagaraEditorStyle.h
--- a/Engine/Plugins/FX/Niagara/Source/NiagaraEditor/Public/NiagaraEditorStyle.h
+++ b/Engine/Plugins/FX/Niagara/Source/NiagaraEditor/Public/NiagaraEditorStyle.h
@@ -5,20 +5,20 @@
 #include "Styling/SlateStyle.h"
 
 /** Manages the style which provides resources for niagara editor widgets. */
-class NIAGARAEDITOR_API FNiagaraEditorStyle : public FSlateStyleSet
+class FNiagaraEditorStyle : public FSlateStyleSet
 {
 public:
-	static void Register();
-	static void Unregister();
-	static void Shutdown();
+	static NIAGARAEDITOR_API void Register();
+	static NIAGARAEDITOR_API void Unregister();
+	static NIAGARAEDITOR_API void Shutdown();
 
 	/** reloads textures used by slate renderer */
-	static void ReloadTextures();
+	static NIAGARAEDITOR_API void ReloadTextures();
 
 	/** @return The Slate style set for niagara editor widgets */
-	static const FNiagaraEditorStyle& Get();
+	static NIAGARAEDITOR_API const FNiagaraEditorStyle& Get();
 
-	static void ReinitializeStyle();
+	static NIAGARAEDITOR_API void ReinitializeStyle();
 
 private:	
 	FNiagaraEditorStyle();
diff --git a/Engine/Plugins/FX/Niagara/Source/NiagaraEditor/Public/NiagaraEditorUtilities.h b/Engine/Plugins/FX/Niagara/Source/NiagaraEditor/Public/NiagaraEditorUtilities.h
--- a/Engine/Plugins/FX/Niagara/Source/NiagaraEditor/Public/NiagaraEditorUtilities.h
+++ b/Engine/Plugins/FX/Niagara/Source/NiagaraEditor/Public/NiagaraEditorUtilities.h
@@ -486,7 +486,7 @@ namespace FNiagaraParameterUtilities
 	/** A simple utility function to generate a parameter widget. */
 	NIAGARAEDITOR_API TSharedRef<SWidget> GetParameterWidget(FNiagaraVariable Variable, bool bAddTypeIcon, bool bShowValue);
 
-	struct NIAGARAEDITOR_API FNiagaraParameterWidgetOptions
+	struct FNiagaraParameterWidgetOptions
 	{
 		bool bAddTypeIcon = false;
 		bool bShowValue = false;
diff --git a/Engine/Plugins/FX/Niagara/Source/NiagaraEditor/Public/NiagaraEmitterEditorData.h b/Engine/Plugins/FX/Niagara/Source/NiagaraEditor/Public/NiagaraEmitterEditorData.h
--- a/Engine/Plugins/FX/Niagara/Source/NiagaraEditor/Public/NiagaraEmitterEditorData.h
+++ b/Engine/Plugins/FX/Niagara/Source/NiagaraEditor/Public/NiagaraEmitterEditorData.h
@@ -40,7 +40,7 @@ public:
 
 	NIAGARAEDITOR_API FSimpleMulticastDelegate& OnSummaryViewStateChanged();
 	
-	NIAGARAEDITOR_API bool ShouldShowSummaryView() const { return bShowSummaryView; }
+	bool ShouldShowSummaryView() const { return bShowSummaryView; }
 	NIAGARAEDITOR_API void SetShowSummaryView(bool bInShouldShowSummaryView);
 	NIAGARAEDITOR_API void ToggleShowSummaryView();
 	
diff --git a/Engine/Plugins/FX/Niagara/Source/NiagaraEditor/Public/NiagaraMessageManager.h b/Engine/Plugins/FX/Niagara/Source/NiagaraEditor/Public/NiagaraMessageManager.h
--- a/Engine/Plugins/FX/Niagara/Source/NiagaraEditor/Public/NiagaraMessageManager.h
+++ b/Engine/Plugins/FX/Niagara/Source/NiagaraEditor/Public/NiagaraMessageManager.h
@@ -61,7 +61,7 @@ private:
 	const TOptional<const FGuid> PinGUID;
 };
 
-class NIAGARAEDITOR_API INiagaraMessageRegistrationHandle : public TSharedFromThis<INiagaraMessageRegistrationHandle>
+class INiagaraMessageRegistrationHandle : public TSharedFromThis<INiagaraMessageRegistrationHandle>
 {
 public:
 	DECLARE_DELEGATE_OneParam(FOnRequestRefresh, const TArray<TSharedRef<const INiagaraMessage>>& /** NewMessages*/)
@@ -154,7 +154,7 @@ public:
 
 private:
 
-	struct NIAGARAEDITOR_API FAssetMessageInfo
+	struct FAssetMessageInfo
 	{
 		FAssetMessageInfo()
 			: bDirty(false)
diff --git a/Engine/Plugins/FX/Niagara/Source/NiagaraEditor/Public/NiagaraMessages.h b/Engine/Plugins/FX/Niagara/Source/NiagaraEditor/Public/NiagaraMessages.h
--- a/Engine/Plugins/FX/Niagara/Source/NiagaraEditor/Public/NiagaraMessages.h
+++ b/Engine/Plugins/FX/Niagara/Source/NiagaraEditor/Public/NiagaraMessages.h
@@ -12,10 +12,10 @@
 #include "NiagaraMessageDataBase.h"
 #include "NiagaraMessages.generated.h"
 
-struct NIAGARAEDITOR_API FNiagaraMessageTopics
+struct FNiagaraMessageTopics
 {
-	static const FName CompilerTopicName;
-	static const FName ObjectTopicName;
+	static NIAGARAEDITOR_API const FName CompilerTopicName;
+	static NIAGARAEDITOR_API const FName ObjectTopicName;
 };
 
 UENUM()
@@ -43,7 +43,7 @@ public:
 };
 
 // Struct for passing around named simple delegates.
-struct NIAGARAEDITOR_API FLinkNameAndDelegate
+struct FLinkNameAndDelegate
 {
 public:
 	FLinkNameAndDelegate() = default;
@@ -238,13 +238,13 @@ private:
 };
 
 USTRUCT()
-struct NIAGARAEDITOR_API FNiagaraStackMessage
+struct FNiagaraStackMessage
 {
 	GENERATED_BODY()
 
-	FNiagaraStackMessage();
+	NIAGARAEDITOR_API FNiagaraStackMessage();
 
-	FNiagaraStackMessage(const FText& MessageText, const FText& ShortDescription, ENiagaraMessageSeverity Severity, bool bAllowDismissal, FGuid Guid = FGuid::NewGuid());
+	NIAGARAEDITOR_API FNiagaraStackMessage(const FText& MessageText, const FText& ShortDescription, ENiagaraMessageSeverity Severity, bool bAllowDismissal, FGuid Guid = FGuid::NewGuid());
 	
 	UPROPERTY()
 	FText MessageText;
@@ -262,7 +262,7 @@ struct NIAGARAEDITOR_API FNiagaraStackMessage
 	FGuid Guid;
 };
 
-struct NIAGARAEDITOR_API FGenerateNiagaraMessageInfo
+struct FGenerateNiagaraMessageInfo
 {
 public:
 	FGenerateNiagaraMessageInfo() = default;
@@ -280,26 +280,26 @@ private:
 	FSimpleDelegate DismissHandler;
 };
 
-UCLASS(abstract)
-class NIAGARAEDITOR_API UNiagaraMessageData : public UNiagaraMessageDataBase
+UCLASS(abstract, MinimalAPI)
+class UNiagaraMessageData : public UNiagaraMessageDataBase
 {
 	GENERATED_BODY()
 
 public:
-	virtual TSharedRef<const INiagaraMessage> GenerateNiagaraMessage(const FGenerateNiagaraMessageInfo& InGenerateInfo = FGenerateNiagaraMessageInfo()) const PURE_VIRTUAL(UNiagaraMessageData::GenerateNiagaraMessage, return TSharedRef<const INiagaraMessage>(static_cast<INiagaraMessage*>(nullptr)););
+	NIAGARAEDITOR_API virtual TSharedRef<const INiagaraMessage> GenerateNiagaraMessage(const FGenerateNiagaraMessageInfo& InGenerateInfo = FGenerateNiagaraMessageInfo()) const PURE_VIRTUAL(UNiagaraMessageData::GenerateNiagaraMessage, return TSharedRef<const INiagaraMessage>(static_cast<INiagaraMessage*>(nullptr)););
 };
 
-UCLASS()
-class NIAGARAEDITOR_API UNiagaraMessageDataText : public UNiagaraMessageData
+UCLASS(MinimalAPI)
+class UNiagaraMessageDataText : public UNiagaraMessageData
 {
 	GENERATED_BODY()
 
 public:
-	void Init(const FText& InMessageText, const ENiagaraMessageSeverity InMessageSeverity, const FName& InTopicName);
+	NIAGARAEDITOR_API void Init(const FText& InMessageText, const ENiagaraMessageSeverity InMessageSeverity, const FName& InTopicName);
 	
-	void Init(const FText& InMessageText, const FText& InShortDescription, const ENiagaraMessageSeverity InMessageSeverity, const FName& InTopicName);
+	NIAGARAEDITOR_API void Init(const FText& InMessageText, const FText& InShortDescription, const ENiagaraMessageSeverity InMessageSeverity, const FName& InTopicName);
 
-	virtual TSharedRef<const INiagaraMessage> GenerateNiagaraMessage(const FGenerateNiagaraMessageInfo& InGenerateInfo = FGenerateNiagaraMessageInfo()) const override;
+	NIAGARAEDITOR_API virtual TSharedRef<const INiagaraMessage> GenerateNiagaraMessage(const FGenerateNiagaraMessageInfo& InGenerateInfo = FGenerateNiagaraMessageInfo()) const override;
 
 	virtual bool GetAllowDismissal() const override { return bAllowDismissal; }
 	void SetAllowDismissal(bool bInAllowDismissal) { bAllowDismissal = bInAllowDismissal; }
@@ -322,4 +322,4 @@ private:
 	
 	UPROPERTY()
 	FName TopicName;
-};
\ No newline at end of file
+};
