[Niagara]

diff --git a/Engine/Plugins/FX/Niagara/Source/NiagaraEditor/Public/NiagaraNode.h b/Engine/Plugins/FX/Niagara/Source/NiagaraEditor/Public/NiagaraNode.h
--- a/Engine/Plugins/FX/Niagara/Source/NiagaraEditor/Public/NiagaraNode.h
+++ b/Engine/Plugins/FX/Niagara/Source/NiagaraEditor/Public/NiagaraNode.h
@@ -24,17 +24,17 @@ template<typename T> class TNiagaraHlslTranslator;
 
 typedef TArray<UEdGraphPin*, TInlineAllocator<16>> FPinCollectorArray;
 
-UCLASS()
-class NIAGARAEDITOR_API UNiagaraNode : public UEdGraphNode
+UCLASS(MinimalAPI)
+class UNiagaraNode : public UEdGraphNode
 {
 	GENERATED_UCLASS_BODY()
 protected:
 
 	using FTranslator = TNiagaraHlslTranslator<FNiagaraCompilationGraphBridge>;
 
-	bool ReallocatePins(bool bMarkNeedsResynchronizeOnChange = true);
+	NIAGARAEDITOR_API bool ReallocatePins(bool bMarkNeedsResynchronizeOnChange = true);
 
-	bool CompileInputPins(FTranslator* Translator, TArray<int32>& OutCompiledInputs) const;
+	NIAGARAEDITOR_API bool CompileInputPins(FTranslator* Translator, TArray<int32>& OutCompiledInputs) const;
 
 	virtual void OnPostSynchronizationInReallocatePins() {}
 	
@@ -43,34 +43,34 @@ public:
 	DECLARE_MULTICAST_DELEGATE_OneParam(FOnNodeVisualsChanged, UNiagaraNode*);
 
 	//~ Begin UObject interface
-	virtual void PostLoad() override;
+	NIAGARAEDITOR_API virtual void PostLoad() override;
 	//~ End UObject interface
 
 	//~ Begin EdGraphNode Interface
-	virtual void PostPlacedNewNode() override;
-	virtual void AutowireNewNode(UEdGraphPin* FromPin) override;	
-	virtual void PinDefaultValueChanged(UEdGraphPin* Pin) override;
-	virtual void PinConnectionListChanged(UEdGraphPin* Pin) override;
-	virtual void PinTypeChanged(UEdGraphPin* Pin) override;
-	virtual void OnRenameNode(const FString& NewName) override;
-	virtual void OnPinRemoved(UEdGraphPin* InRemovedPin) override;
-	virtual void NodeConnectionListChanged() override;	
-	virtual TSharedPtr<SGraphNode> CreateVisualWidget() override; 
-	virtual void GetPinHoverText(const UEdGraphPin& Pin, FString& HoverTextOut) const override;
-	virtual void GetNodeContextMenuActions(class UToolMenu* Menu, class UGraphNodeContextMenuContext* Context) const override;
-	virtual bool CanCreateUnderSpecifiedSchema(const UEdGraphSchema* Schema) const override;
+	NIAGARAEDITOR_API virtual void PostPlacedNewNode() override;
+	NIAGARAEDITOR_API virtual void AutowireNewNode(UEdGraphPin* FromPin) override;	
+	NIAGARAEDITOR_API virtual void PinDefaultValueChanged(UEdGraphPin* Pin) override;
+	NIAGARAEDITOR_API virtual void PinConnectionListChanged(UEdGraphPin* Pin) override;
+	NIAGARAEDITOR_API virtual void PinTypeChanged(UEdGraphPin* Pin) override;
+	NIAGARAEDITOR_API virtual void OnRenameNode(const FString& NewName) override;
+	NIAGARAEDITOR_API virtual void OnPinRemoved(UEdGraphPin* InRemovedPin) override;
+	NIAGARAEDITOR_API virtual void NodeConnectionListChanged() override;	
+	NIAGARAEDITOR_API virtual TSharedPtr<SGraphNode> CreateVisualWidget() override; 
+	NIAGARAEDITOR_API virtual void GetPinHoverText(const UEdGraphPin& Pin, FString& HoverTextOut) const override;
+	NIAGARAEDITOR_API virtual void GetNodeContextMenuActions(class UToolMenu* Menu, class UGraphNodeContextMenuContext* Context) const override;
+	NIAGARAEDITOR_API virtual bool CanCreateUnderSpecifiedSchema(const UEdGraphSchema* Schema) const override;
 	//~ End EdGraphNode Interface
 
 	/** Virtual function to allow for custom widgets in the input or output box */
-	virtual void AddWidgetsToInputBox(TSharedPtr<SVerticalBox> InputBox);
-	virtual void AddWidgetsToOutputBox(TSharedPtr<SVerticalBox> OutputBox);
+	NIAGARAEDITOR_API virtual void AddWidgetsToInputBox(TSharedPtr<SVerticalBox> InputBox);
+	NIAGARAEDITOR_API virtual void AddWidgetsToOutputBox(TSharedPtr<SVerticalBox> OutputBox);
 	
 	/** Get the Niagara graph that owns this node */
-	const class UNiagaraGraph* GetNiagaraGraph()const;
-	class UNiagaraGraph* GetNiagaraGraph();
+	NIAGARAEDITOR_API const class UNiagaraGraph* GetNiagaraGraph()const;
+	NIAGARAEDITOR_API class UNiagaraGraph* GetNiagaraGraph();
 
 	/** Get the source object */
-	class UNiagaraScriptSource* GetSource()const;
+	NIAGARAEDITOR_API class UNiagaraScriptSource* GetSource()const;
 
 	/** Gets the asset referenced by this node, or nullptr if there isn't one. */
 	virtual UObject* GetReferencedAsset() const { return nullptr; }
@@ -79,10 +79,10 @@ public:
 	/** Refreshes the node due to external changes, e.g. the underlying function changed for a function call node. Return true if the graph changed.*/
 	virtual bool RefreshFromExternalChanges() { return false; }
 
-	virtual void Compile(FTranslator* Translator, TArray<int32>& Outputs) const;
+	NIAGARAEDITOR_API virtual void Compile(FTranslator* Translator, TArray<int32>& Outputs) const;
 
-	UEdGraphPin* GetInputPin(int32 InputIndex) const;
-	UEdGraphPin* GetOutputPin(int32 OutputIndex) const;
+	NIAGARAEDITOR_API UEdGraphPin* GetInputPin(int32 InputIndex) const;
+	NIAGARAEDITOR_API UEdGraphPin* GetOutputPin(int32 OutputIndex) const;
 
 	template<typename ContainerType>
 	void GetInputPins(ContainerType& OutInputPins) const
@@ -95,27 +95,27 @@ public:
 		GetOutputPinsInternal<ContainerType>(OutOutputPins);
 	}
 
-	UEdGraphPin* GetPinByPersistentGuid(const FGuid& InGuid) const;
-	virtual void ResolveNumerics(const UEdGraphSchema_Niagara* Schema, bool bSetInline, TMap<TPair<FGuid, class UEdGraphNode*>, FNiagaraTypeDefinition>* PinCache);
+	NIAGARAEDITOR_API UEdGraphPin* GetPinByPersistentGuid(const FGuid& InGuid) const;
+	NIAGARAEDITOR_API virtual void ResolveNumerics(const UEdGraphSchema_Niagara* Schema, bool bSetInline, TMap<TPair<FGuid, class UEdGraphNode*>, FNiagaraTypeDefinition>* PinCache);
 
 	/** Apply any node-specific logic to determine if it is safe to add this node to the graph. This is meant to be called only in the Editor before placing the node.*/
-	virtual bool CanAddToGraph(UNiagaraGraph* TargetGraph, FString& OutErrorMsg) const;
+	NIAGARAEDITOR_API virtual bool CanAddToGraph(UNiagaraGraph* TargetGraph, FString& OutErrorMsg) const;
 
 	/** Request a pin type change to a specific type. */
-	void RequestNewPinType(UEdGraphPin* PinToChange, FNiagaraTypeDefinition NewType);
+	NIAGARAEDITOR_API void RequestNewPinType(UEdGraphPin* PinToChange, FNiagaraTypeDefinition NewType);
 
 	/** Determine whether we are allowed to change a pin's type from UI. Enables the type conversion widget. */
 	virtual bool AllowExternalPinTypeChanges(const UEdGraphPin* InGraphPin) const { return false; }
 	/** Determine whether or not a pin can be changed to a certain type.
 	 *  Used to populate the type conversion menu if external pin type changes are allowed or for wildcard responses */
 	virtual bool AllowNiagaraTypeForPinTypeChange(const FNiagaraTypeDefinition& InType, UEdGraphPin* Pin) const { return true; }
-	virtual void GetWildcardPinHoverConnectionTextAddition(const UEdGraphPin* WildcardPin, const UEdGraphPin* OtherPin, ECanCreateConnectionResponse ConnectionResponse, FString& OutString) const;
+	NIAGARAEDITOR_API virtual void GetWildcardPinHoverConnectionTextAddition(const UEdGraphPin* WildcardPin, const UEdGraphPin* OtherPin, ECanCreateConnectionResponse ConnectionResponse, FString& OutString) const;
 	
 	/** Gets which mode to use when deducing the type of numeric output pins from the types of the input pins. */
-	virtual ENiagaraNumericOutputTypeSelectionMode GetNumericOutputTypeSelectionMode() const;
+	NIAGARAEDITOR_API virtual ENiagaraNumericOutputTypeSelectionMode GetNumericOutputTypeSelectionMode() const;
 
 	/** Convert the type of an existing numeric pin to a more known type.*/
-	virtual bool ConvertNumericPinToType(UEdGraphPin* InGraphPin, FNiagaraTypeDefinition TypeDef);
+	NIAGARAEDITOR_API virtual bool ConvertNumericPinToType(UEdGraphPin* InGraphPin, FNiagaraTypeDefinition TypeDef);
 
 	/** Determine whether or not a pin should be renamable. */
 	virtual bool IsPinNameEditable(const UEdGraphPin* GraphPinObj) const { return false; }
@@ -132,12 +132,12 @@ public:
 	virtual bool CancelEditablePinName(const FText& InName, UEdGraphPin* InGraphPinObj) { return false; }
 
 	/** Returns whether or not the supplied pin has a rename pending. */
-	bool GetIsPinRenamePending(const UEdGraphPin* Pin);
+	NIAGARAEDITOR_API bool GetIsPinRenamePending(const UEdGraphPin* Pin);
 
 	/** Sets whether or not the supplied pin has a rename pending. */
-	void SetIsPinRenamePending(const UEdGraphPin* Pin, bool bInIsRenamePending);
+	NIAGARAEDITOR_API void SetIsPinRenamePending(const UEdGraphPin* Pin, bool bInIsRenamePending);
 
-	bool IsParameterMapPin(const UEdGraphPin* Pin) const;
+	NIAGARAEDITOR_API bool IsParameterMapPin(const UEdGraphPin* Pin) const;
 	
 	/** Adds the current node information to the parameter map history
 	 *
@@ -145,7 +145,7 @@ public:
 	 *  @Param	bRecursive				If true then the history of all of this node's input pins will also be added
 	 *  @Param	bFilterForCompilation	If true, some nodes like static switches or reroute nodes are jumped over. The ui usually sets this to false to follow all possible paths.  
 	 */
-	virtual void BuildParameterMapHistory(FNiagaraParameterMapHistoryBuilder& OutHistory, bool bRecursive = true, bool bFilterForCompilation = true) const;
+	NIAGARAEDITOR_API virtual void BuildParameterMapHistory(FNiagaraParameterMapHistoryBuilder& OutHistory, bool bRecursive = true, bool bFilterForCompilation = true) const;
 	
 	/** Go through all the external dependencies of this node in isolation and add them to the reference id list.*/
 	virtual void GatherExternalDependencyData(ENiagaraScriptUsage InUsage, const FGuid& InUsageId, TArray<FNiagaraCompileHash>& InReferencedCompileHashes, TArray<FString>& InReferencedObjs) const {};
@@ -155,81 +155,81 @@ public:
 	 *  @Param	LocallyOwnedOutputPin	The pin to trace, must be a child of this node
 	 *  @Param	bFilterForCompilation	If true, some nodes like static switches or reroute nodes are jumped over. The ui usually sets this to false to follow all possible paths.  
 	 */
-	virtual UEdGraphPin* GetTracedOutputPin(UEdGraphPin* LocallyOwnedOutputPin, bool bFilterForCompilation, TArray<const UNiagaraNode*>* OutNodesVisitedDuringTrace = nullptr) const;
+	NIAGARAEDITOR_API virtual UEdGraphPin* GetTracedOutputPin(UEdGraphPin* LocallyOwnedOutputPin, bool bFilterForCompilation, TArray<const UNiagaraNode*>* OutNodesVisitedDuringTrace = nullptr) const;
 
 	/** Traces a node's output pins to its source output pin.
 	*
 	*  @Param	LocallyOwnedOutputPin	The pin to trace
 	*  @Param	bFilterForCompilation	If true, some nodes like static switches or reroute nodes are jumped over. The ui usually sets this to false to follow all possible paths.  
 	*/
-	static UEdGraphPin* TraceOutputPin(UEdGraphPin* LocallyOwnedOutputPin, bool bFilterForCompilation = true, TArray<const UNiagaraNode*>* OutNodesVisitedDuringTrace = nullptr);
+	static NIAGARAEDITOR_API UEdGraphPin* TraceOutputPin(UEdGraphPin* LocallyOwnedOutputPin, bool bFilterForCompilation = true, TArray<const UNiagaraNode*>* OutNodesVisitedDuringTrace = nullptr);
 
 	/** Allows a node to replace a pin that is about to be compiled with another pin. This can be used for either optimizations or features such as the static switch. Returns true if the pin was successfully replaced, false otherwise. */
-	virtual bool SubstituteCompiledPin(FTranslator* Translator, UEdGraphPin** LocallyOwnedPin);
+	NIAGARAEDITOR_API virtual bool SubstituteCompiledPin(FTranslator* Translator, UEdGraphPin** LocallyOwnedPin);
 
 	virtual UEdGraphPin* GetPassThroughPin(const UEdGraphPin* LocallyOwnedOutputPin) const override { return nullptr; }
 	virtual UEdGraphPin* GetPassThroughPin(const UEdGraphPin* LocallyOwnedOutputPin, ENiagaraScriptUsage InUsage) const { return nullptr; }
 
 	/** Identify that this node has undergone changes that will require synchronization with a compiled script.*/
-	void MarkNodeRequiresSynchronization(FString Reason, bool bRaiseGraphNeedsRecompile);
+	NIAGARAEDITOR_API void MarkNodeRequiresSynchronization(FString Reason, bool bRaiseGraphNeedsRecompile);
 
 	/** Get the change id for this node. This change id is updated whenever the node is manipulated in a way that should force a recompile.*/
 	const FGuid& GetChangeId() const { return ChangeId; }
 
 	/** Set the change id for this node to an explicit value. This should only be called by internal code. */
-	void ForceChangeId(const FGuid& InId, bool bRaiseGraphNeedsRecompile);
+	NIAGARAEDITOR_API void ForceChangeId(const FGuid& InId, bool bRaiseGraphNeedsRecompile);
 
-	FOnNodeVisualsChanged& OnVisualsChanged();
+	NIAGARAEDITOR_API FOnNodeVisualsChanged& OnVisualsChanged();
 
 	virtual void AppendFunctionAliasForContext(const FNiagaraGraphFunctionAliasContext& InFunctionAliasContext, FString& InOutFunctionAlias, bool& OutOnlyOncePerNodeType) const { };
 
 	/** Old style compile hash code. To be removed in the future.*/
-	virtual void UpdateCompileHashForNode(FSHA1& HashState) const;
+	NIAGARAEDITOR_API virtual void UpdateCompileHashForNode(FSHA1& HashState) const;
 
 	/** Entry point for generating the compile hash.*/
-	bool AppendCompileHash(FNiagaraCompileHashVisitor* InVisitor) const;
+	NIAGARAEDITOR_API bool AppendCompileHash(FNiagaraCompileHashVisitor* InVisitor) const;
 
-	static bool SetPinDefaultToTypeDefaultIfUnset(UEdGraphPin* InPin);
+	static NIAGARAEDITOR_API bool SetPinDefaultToTypeDefaultIfUnset(UEdGraphPin* InPin);
 
 protected:
 	/** Pin type changes are handled by the individual subclasses to account for pin sets (like the if node or select node) */
 	virtual bool OnNewPinTypeRequested(UEdGraphPin* PinToChange, FNiagaraTypeDefinition NewType) { return false; }
 	
 	/** Go through all class members for a given UClass on this object and hash them into the visitor.*/
-	virtual bool GenerateCompileHashForClassMembers(const UClass* InClass, FNiagaraCompileHashVisitor* InVisitor) const;
+	NIAGARAEDITOR_API virtual bool GenerateCompileHashForClassMembers(const UClass* InClass, FNiagaraCompileHashVisitor* InVisitor) const;
 
 	/** Write out the specific entries for UNiagaraNode into the visitor hash. */
-	bool NiagaraNodeAppendCompileHash(FNiagaraCompileHashVisitor* InVisitor) const;
+	NIAGARAEDITOR_API bool NiagaraNodeAppendCompileHash(FNiagaraCompileHashVisitor* InVisitor) const;
 
 	/** Write out the specific entries of this pin to the visitor hash.*/
-	virtual bool PinAppendCompileHash(const UEdGraphPin* InPin, FNiagaraCompileHashVisitor* InVisitor) const;
+	NIAGARAEDITOR_API virtual bool PinAppendCompileHash(const UEdGraphPin* InPin, FNiagaraCompileHashVisitor* InVisitor) const;
 	
 	/** Helper function to hash arbitrary UProperty entries (Arrays, Maps, Structs, etc).*/
-	virtual bool NestedPropertiesAppendCompileHash(const void* Container, const UStruct* Struct, EFieldIteratorFlags::SuperClassFlags IteratorFlags, const FString& BaseName, FNiagaraCompileHashVisitor* InVisitor) const;
+	NIAGARAEDITOR_API virtual bool NestedPropertiesAppendCompileHash(const void* Container, const UStruct* Struct, EFieldIteratorFlags::SuperClassFlags IteratorFlags, const FString& BaseName, FNiagaraCompileHashVisitor* InVisitor) const;
 	
 	/** For a simple Plain old data type UProperty, hash the data.*/
-	virtual bool PODPropertyAppendCompileHash(const void* Container, FProperty* Property, const FString& PropertyName, FNiagaraCompileHashVisitor* InVisitor) const;
+	NIAGARAEDITOR_API virtual bool PODPropertyAppendCompileHash(const void* Container, FProperty* Property, const FString& PropertyName, FNiagaraCompileHashVisitor* InVisitor) const;
 
-	virtual int32 CompileInputPin(FTranslator* Translator, UEdGraphPin* Pin) const;
-	virtual bool IsValidPinToCompile(UEdGraphPin* Pin) const;
+	NIAGARAEDITOR_API virtual int32 CompileInputPin(FTranslator* Translator, UEdGraphPin* Pin) const;
+	NIAGARAEDITOR_API virtual bool IsValidPinToCompile(UEdGraphPin* Pin) const;
 
-	void NumericResolutionByPins(const UEdGraphSchema_Niagara* Schema, TArrayView<UEdGraphPin* const> InputPins, TArrayView<UEdGraphPin* const> OutputPins,
+	NIAGARAEDITOR_API void NumericResolutionByPins(const UEdGraphSchema_Niagara* Schema, TArrayView<UEdGraphPin* const> InputPins, TArrayView<UEdGraphPin* const> OutputPins,
 		bool bFixInline, TMap<TPair<FGuid, UEdGraphNode*>, FNiagaraTypeDefinition>* PinCache);
 		
-	virtual FNiagaraTypeDefinition ResolveCustomNumericType(const TArray<FNiagaraTypeDefinition>& NonNumericInputs) const;
+	NIAGARAEDITOR_API virtual FNiagaraTypeDefinition ResolveCustomNumericType(const TArray<FNiagaraTypeDefinition>& NonNumericInputs) const;
 
 	/** Route input parameter map to output parameter map if it exists. Note that before calling this function,
 		the input pins should have been visited already.*/
-	virtual void RouteParameterMapAroundMe(FNiagaraParameterMapHistoryBuilder& OutHistory, bool bRecursive) const;
+	NIAGARAEDITOR_API virtual void RouteParameterMapAroundMe(FNiagaraParameterMapHistoryBuilder& OutHistory, bool bRecursive) const;
 
 	/** Basically routes the pin through the parameter map builder so that it looks like a regular pin. Handles visiting the node connected as well to the input pin.*/
-	virtual void RegisterPassthroughPin(FNiagaraParameterMapHistoryBuilder& OutHistory, UEdGraphPin* InputPin, UEdGraphPin* OutputPin, bool bFilterForCompilation, bool bVisitInputPin) const;
+	NIAGARAEDITOR_API virtual void RegisterPassthroughPin(FNiagaraParameterMapHistoryBuilder& OutHistory, UEdGraphPin* InputPin, UEdGraphPin* OutputPin, bool bFilterForCompilation, bool bVisitInputPin) const;
 
 	/** If the pin is a known name (like Engine.DeltaTime) this tries to return a default tooltip for it. */
-	bool GetTooltipTextForKnownPin(const UEdGraphPin& Pin, FText& OutTooltip) const;
+	NIAGARAEDITOR_API bool GetTooltipTextForKnownPin(const UEdGraphPin& Pin, FText& OutTooltip) const;
 	
 #if WITH_EDITORONLY_DATA
-	virtual void GatherForLocalization(FPropertyLocalizationDataGatherer& PropertyLocalizationDataGatherer, const EPropertyLocalizationGathererTextFlags GatherTextFlags) const override;
+	NIAGARAEDITOR_API virtual void GatherForLocalization(FPropertyLocalizationDataGatherer& PropertyLocalizationDataGatherer, const EPropertyLocalizationGathererTextFlags GatherTextFlags) const override;
 #endif
 
 	template<typename ContainerType>
@@ -260,8 +260,8 @@ protected:
 		}
 	}
 
-	virtual void GetCompilationInputPins(FPinCollectorArray& InputPins) const;
-	virtual void GetCompilationOutputPins(FPinCollectorArray& OutputPins) const;
+	NIAGARAEDITOR_API virtual void GetCompilationInputPins(FPinCollectorArray& InputPins) const;
+	NIAGARAEDITOR_API virtual void GetCompilationOutputPins(FPinCollectorArray& OutputPins) const;
 
 	/** The current change identifier for this node. Used to sync status with UNiagaraScripts.*/
 	UPROPERTY()
diff --git a/Engine/Plugins/FX/Niagara/Source/NiagaraEditor/Public/NiagaraNodeAssignment.h b/Engine/Plugins/FX/Niagara/Source/NiagaraEditor/Public/NiagaraNodeAssignment.h
--- a/Engine/Plugins/FX/Niagara/Source/NiagaraEditor/Public/NiagaraNodeAssignment.h
+++ b/Engine/Plugins/FX/Niagara/Source/NiagaraEditor/Public/NiagaraNodeAssignment.h
@@ -13,23 +13,23 @@ class SComboButton;
 class SWidget;
 class UNiagaraScript;
 
-UCLASS()
-class NIAGARAEDITOR_API UNiagaraNodeAssignment : public UNiagaraNodeFunctionCall
+UCLASS(MinimalAPI)
+class UNiagaraNodeAssignment : public UNiagaraNodeFunctionCall
 {
 	GENERATED_BODY()
 
 public:
 	int32 NumTargets() const { return AssignmentTargets.Num(); }
 
-	int32 FindAssignmentTarget(const FName& InName);
-	int32 FindAssignmentTarget(const FName& InName, const FNiagaraTypeDefinition& TypeDef);
+	NIAGARAEDITOR_API int32 FindAssignmentTarget(const FName& InName);
+	NIAGARAEDITOR_API int32 FindAssignmentTarget(const FName& InName, const FNiagaraTypeDefinition& TypeDef);
 
 	/** Set the assignment target and default value. This does not call RefreshFromExternalChanges, which you will need to do to update the internal graph. Returns true if values were changed.*/
-	bool SetAssignmentTarget(int32 Idx, const FNiagaraVariable& InVar, const FString* InDefaultValue = nullptr);
-	int32 AddAssignmentTarget(const FNiagaraVariable& InVar, const FString* InDefaultValue = nullptr);
+	NIAGARAEDITOR_API bool SetAssignmentTarget(int32 Idx, const FNiagaraVariable& InVar, const FString* InDefaultValue = nullptr);
+	NIAGARAEDITOR_API int32 AddAssignmentTarget(const FNiagaraVariable& InVar, const FString* InDefaultValue = nullptr);
 
 	/** This will rename an existing input but will not refresh the underlying script.  RefreshFromExternalChanges must be called if this function returns true. */
-	bool RenameAssignmentTarget(FName OldName, FName NewName);
+	NIAGARAEDITOR_API bool RenameAssignmentTarget(FName OldName, FName NewName);
 
 	const FNiagaraVariable& GetAssignmentTarget(int32 Idx) const { return AssignmentTargets[Idx]; }
 
@@ -37,29 +37,29 @@ public:
 	const TArray<FString>& GetAssignmentDefaults() const { return AssignmentDefaultValues; }
 
 	//~ Begin EdGraphNode Interface
-	virtual void PostLoad() override;
+	NIAGARAEDITOR_API virtual void PostLoad() override;
 #if WITH_EDITORONLY_DATA
-	static void DeclareConstructClasses(TArray<FTopLevelAssetPath>& OutConstructClasses, const UClass* SpecificSubclass);
+	static NIAGARAEDITOR_API void DeclareConstructClasses(TArray<FTopLevelAssetPath>& OutConstructClasses, const UClass* SpecificSubclass);
 #endif
-	virtual FText GetNodeTitle(ENodeTitleType::Type TitleType) const override;
-	virtual void AllocateDefaultPins() override;
+	NIAGARAEDITOR_API virtual FText GetNodeTitle(ENodeTitleType::Type TitleType) const override;
+	NIAGARAEDITOR_API virtual void AllocateDefaultPins() override;
 
 	//~ UNiagaraNodeFunctionCall interface
 	virtual bool AllowDynamicPins() const override { return true; }
-	virtual bool RefreshFromExternalChanges() override;
-	virtual void BuildParameterMapHistory(FNiagaraParameterMapHistoryBuilder& OutHistory, bool bRecursive = true, bool bFilterForCompilation = true) const override;
-	virtual void GatherExternalDependencyData(ENiagaraScriptUsage InUsage, const FGuid& InUsageId, TArray<FNiagaraCompileHash>& InReferencedCompileHashes, TArray<FString>& InReferencedObjs) const override;
+	NIAGARAEDITOR_API virtual bool RefreshFromExternalChanges() override;
+	NIAGARAEDITOR_API virtual void BuildParameterMapHistory(FNiagaraParameterMapHistoryBuilder& OutHistory, bool bRecursive = true, bool bFilterForCompilation = true) const override;
+	NIAGARAEDITOR_API virtual void GatherExternalDependencyData(ENiagaraScriptUsage InUsage, const FGuid& InUsageId, TArray<FNiagaraCompileHash>& InReferencedCompileHashes, TArray<FString>& InReferencedObjs) const override;
 
-	void AddParameter(FNiagaraVariable InVar, FString InDefaultValue);
-	void RemoveParameter(const FNiagaraVariable& InVar);
+	NIAGARAEDITOR_API void AddParameter(FNiagaraVariable InVar, FString InDefaultValue);
+	NIAGARAEDITOR_API void RemoveParameter(const FNiagaraVariable& InVar);
 
-	void UpdateUsageBitmaskFromOwningScript();
+	NIAGARAEDITOR_API void UpdateUsageBitmaskFromOwningScript();
 
-	virtual FText GetTooltipText() const override;
+	NIAGARAEDITOR_API virtual FText GetTooltipText() const override;
 
 	FSimpleMulticastDelegate& OnAssignmentTargetsChanged() { return AssignmentTargetsChangedDelegate; }
 
-	TSharedRef<SWidget> CreateAddParameterMenu(const TSharedPtr<SComboButton>& AddButton);
+	NIAGARAEDITOR_API TSharedRef<SWidget> CreateAddParameterMenu(const TSharedPtr<SComboButton>& AddButton);
 
 protected:
 
@@ -83,14 +83,14 @@ protected:
 	FSimpleMulticastDelegate AssignmentTargetsChangedDelegate;
 
 private:
-	void GenerateScript();
+	NIAGARAEDITOR_API void GenerateScript();
 
-	void InitializeScript(UNiagaraScript* NewScript);
+	NIAGARAEDITOR_API void InitializeScript(UNiagaraScript* NewScript);
 
-	int32 CalculateScriptUsageBitmask();
+	NIAGARAEDITOR_API int32 CalculateScriptUsageBitmask();
 
-	void RefreshTitle();
+	NIAGARAEDITOR_API void RefreshTitle();
 
-	bool PostLoad_LWCFixup(int32 NiagaraVersion);
+	NIAGARAEDITOR_API bool PostLoad_LWCFixup(int32 NiagaraVersion);
 };
 
diff --git a/Engine/Plugins/FX/Niagara/Source/NiagaraEditor/Public/NiagaraNodeFunctionCall.h b/Engine/Plugins/FX/Niagara/Source/NiagaraEditor/Public/NiagaraNodeFunctionCall.h
--- a/Engine/Plugins/FX/Niagara/Source/NiagaraEditor/Public/NiagaraNodeFunctionCall.h
+++ b/Engine/Plugins/FX/Niagara/Source/NiagaraEditor/Public/NiagaraNodeFunctionCall.h
@@ -180,10 +180,10 @@ public:
 	FNiagaraMessageStore& GetMessageStore() { return MessageStore; }
 
 	// Custom Notes API
-	NIAGARAEDITOR_API const TArray<FNiagaraStackMessage>& GetCustomNotes() const { return StackMessages; };
+	const TArray<FNiagaraStackMessage>& GetCustomNotes() const { return StackMessages; };
 	NIAGARAEDITOR_API void AddCustomNote(const FNiagaraStackMessage& StackMessage);
 	NIAGARAEDITOR_API void RemoveCustomNote(const FGuid& MessageKey);
-	NIAGARAEDITOR_API FSimpleDelegate& OnCustomNotesChanged() { return OnCustomNotesChangedDelegate; }
+	FSimpleDelegate& OnCustomNotesChanged() { return OnCustomNotesChangedDelegate; }
 	void RemoveCustomNoteViaDelegate(const FGuid MessageKey);
 	TArray<FGuid> GetBoundPinGuidsByName(FName InputName) const;
 
diff --git a/Engine/Plugins/FX/Niagara/Source/NiagaraEditor/Public/NiagaraOverviewNode.h b/Engine/Plugins/FX/Niagara/Source/NiagaraEditor/Public/NiagaraOverviewNode.h
--- a/Engine/Plugins/FX/Niagara/Source/NiagaraEditor/Public/NiagaraOverviewNode.h
+++ b/Engine/Plugins/FX/Niagara/Source/NiagaraEditor/Public/NiagaraOverviewNode.h
@@ -7,39 +7,39 @@
 class UNiagaraSystem;
 class FNiagaraEmitterHandleViewModel;
 
-UCLASS()
-class NIAGARAEDITOR_API UNiagaraOverviewNode : public UEdGraphNode
+UCLASS(MinimalAPI)
+class UNiagaraOverviewNode : public UEdGraphNode
 {
 	GENERATED_BODY()
 public:
-	UNiagaraOverviewNode();
-	void Initialize(UNiagaraSystem* InOwningSystem);
-	void Initialize(UNiagaraSystem* InOwningSystem, FGuid InEmitterHandleGuid);
-	const FGuid GetEmitterHandleGuid() const;
-	struct FNiagaraEmitterHandle* TryGetEmitterHandle();
+	NIAGARAEDITOR_API UNiagaraOverviewNode();
+	NIAGARAEDITOR_API void Initialize(UNiagaraSystem* InOwningSystem);
+	NIAGARAEDITOR_API void Initialize(UNiagaraSystem* InOwningSystem, FGuid InEmitterHandleGuid);
+	NIAGARAEDITOR_API const FGuid GetEmitterHandleGuid() const;
+	NIAGARAEDITOR_API struct FNiagaraEmitterHandle* TryGetEmitterHandle();
 	
 	//~ Begin UEdGraphNode Interface
 	/** Gets the name of this node, shown in title bar */
-	virtual FText GetNodeTitle(ENodeTitleType::Type TitleType) const override;
+	NIAGARAEDITOR_API virtual FText GetNodeTitle(ENodeTitleType::Type TitleType) const override;
 
-	virtual FLinearColor GetNodeTitleColor() const override;
+	NIAGARAEDITOR_API virtual FLinearColor GetNodeTitleColor() const override;
 
 	/** Whether or not this node can be deleted by user action */
-	virtual bool CanUserDeleteNode() const override;
+	NIAGARAEDITOR_API virtual bool CanUserDeleteNode() const override;
 
 	/** Whether or not this node can be safely duplicated (via copy/paste, etc...) in the graph */
-	virtual bool CanDuplicateNode() const override;
+	NIAGARAEDITOR_API virtual bool CanDuplicateNode() const override;
 
-	virtual void OnRenameNode(const FString& NewName) override;
+	NIAGARAEDITOR_API virtual void OnRenameNode(const FString& NewName) override;
 
-	virtual bool GetCanRenameNode() const override;
+	NIAGARAEDITOR_API virtual bool GetCanRenameNode() const override;
 
-	virtual void GetNodeContextMenuActions(class UToolMenu* Menu, class UGraphNodeContextMenuContext* Context) const override;
+	NIAGARAEDITOR_API virtual void GetNodeContextMenuActions(class UToolMenu* Menu, class UGraphNodeContextMenuContext* Context) const override;
 
-	virtual bool CanCreateUnderSpecifiedSchema(const UEdGraphSchema* Schema) const override;
+	NIAGARAEDITOR_API virtual bool CanCreateUnderSpecifiedSchema(const UEdGraphSchema* Schema) const override;
 	//~ End UEdGraphNode Interface
 
-	UNiagaraSystem* GetOwningSystem();
+	NIAGARAEDITOR_API UNiagaraSystem* GetOwningSystem();
 
 	void RequestRename() { bRenamePending = true; }
 	void RenameStarted() { bRenamePending = false; }
@@ -54,9 +54,9 @@ private:
 
 	bool bRenamePending;
 
-	static bool bColorsAreInitialized;
-	static FLinearColor SystemColor;
-	static FLinearColor EmitterColor;
-	static FLinearColor IsolatedColor;
-	static FLinearColor NotIsolatedColor;
+	static NIAGARAEDITOR_API bool bColorsAreInitialized;
+	static NIAGARAEDITOR_API FLinearColor SystemColor;
+	static NIAGARAEDITOR_API FLinearColor EmitterColor;
+	static NIAGARAEDITOR_API FLinearColor IsolatedColor;
+	static NIAGARAEDITOR_API FLinearColor NotIsolatedColor;
 };
diff --git a/Engine/Plugins/FX/Niagara/Source/NiagaraEditor/Public/NiagaraScriptVariable.h b/Engine/Plugins/FX/Niagara/Source/NiagaraEditor/Public/NiagaraScriptVariable.h
--- a/Engine/Plugins/FX/Niagara/Source/NiagaraEditor/Public/NiagaraScriptVariable.h
+++ b/Engine/Plugins/FX/Niagara/Source/NiagaraEditor/Public/NiagaraScriptVariable.h
@@ -10,15 +10,15 @@
 class UNiagaraScriptVariable;
 
 USTRUCT()
-struct NIAGARAEDITOR_API FNiagaraScriptVariableData
+struct FNiagaraScriptVariableData
 {
 	GENERATED_USTRUCT_BODY()
 
 	FNiagaraScriptVariableData() = default;
 
-	void Init(const FNiagaraVariable& InVar, const FNiagaraVariableMetaData& InVarMetaData);
-	void InitFrom(const UNiagaraScriptVariable& Source, bool bCreateNewGuid = true);
-	void InitFrom(const FNiagaraScriptVariableData& Source, bool bCreateNewGuid = true);
+	NIAGARAEDITOR_API void Init(const FNiagaraVariable& InVar, const FNiagaraVariableMetaData& InVarMetaData);
+	NIAGARAEDITOR_API void InitFrom(const UNiagaraScriptVariable& Source, bool bCreateNewGuid = true);
+	NIAGARAEDITOR_API void InitFrom(const FNiagaraScriptVariableData& Source, bool bCreateNewGuid = true);
 
 	/** The default mode. Can be Value, Binding or Custom. */
 	UPROPERTY(EditAnywhere, Category = "Default Value")
@@ -37,16 +37,16 @@ struct NIAGARAEDITOR_API FNiagaraScriptVariableData
 	FNiagaraVariableMetaData Metadata;
 
 	/** Entry point for generating the compile hash.*/
-	bool AppendCompileHash(FNiagaraCompileHashVisitor* InVisitor) const;
+	NIAGARAEDITOR_API bool AppendCompileHash(FNiagaraCompileHashVisitor* InVisitor) const;
 
 	bool GetIsStaticSwitch() const { return bIsStaticSwitch; };
 	void SetIsStaticSwitch(bool bInIsStaticSwitch) { bIsStaticSwitch = bInIsStaticSwitch; };
 
 	bool GetIsSubscribedToParameterDefinitions() const { return bSubscribedToParameterDefinitions; };
-	void SetIsSubscribedToParameterDefinitions(bool bInSubscribedToParameterDefinitions);
+	NIAGARAEDITOR_API void SetIsSubscribedToParameterDefinitions(bool bInSubscribedToParameterDefinitions);
 
 	bool GetIsOverridingParameterDefinitionsDefaultValue() const { return bOverrideParameterDefinitionsDefaultValue; };
-	void SetIsOverridingParameterDefinitionsDefaultValue(bool bInOverridingParameterDefinitionsDefaultValue);
+	NIAGARAEDITOR_API void SetIsOverridingParameterDefinitionsDefaultValue(bool bInOverridingParameterDefinitionsDefaultValue);
 
 	const FGuid& GetChangeId() const { return ChangeId; };
 	void SetChangeId(const FGuid& NewId) { ChangeId = NewId; };
@@ -133,20 +133,20 @@ private:
 /*
 * Used to store variable data and metadata per graph. 
 */
-UCLASS()
-class NIAGARAEDITOR_API UNiagaraScriptVariable : public UObject
+UCLASS(MinimalAPI)
+class UNiagaraScriptVariable : public UObject
 {
 public:
 	GENERATED_UCLASS_BODY()
 
 	DECLARE_DELEGATE_OneParam(FOnChanged, const UNiagaraScriptVariable* /*ThisScriptVariable*/);
 
-	virtual void PostLoad() override;
+	NIAGARAEDITOR_API virtual void PostLoad() override;
 
-	virtual void PostEditChangeProperty(struct FPropertyChangedEvent& PropertyChangedEvent) override;
+	NIAGARAEDITOR_API virtual void PostEditChangeProperty(struct FPropertyChangedEvent& PropertyChangedEvent) override;
 
-	void Init(const FNiagaraVariable& InVar, const FNiagaraVariableMetaData& InVarMetaData);
-	void InitFrom(const UNiagaraScriptVariable* Value, bool bCreateNewGuid = true);
+	NIAGARAEDITOR_API void Init(const FNiagaraVariable& InVar, const FNiagaraVariableMetaData& InVarMetaData);
+	NIAGARAEDITOR_API void InitFrom(const UNiagaraScriptVariable* Value, bool bCreateNewGuid = true);
 
 	/** The default mode. Can be Value, Binding or Custom. */
 	UPROPERTY(EditAnywhere, Category = "Default Value")
@@ -165,16 +165,16 @@ public:
 	FNiagaraVariableMetaData Metadata;
 
 	/** Entry point for generating the compile hash.*/
-	bool AppendCompileHash(FNiagaraCompileHashVisitor* InVisitor) const;
+	NIAGARAEDITOR_API bool AppendCompileHash(FNiagaraCompileHashVisitor* InVisitor) const;
 
 	bool GetIsStaticSwitch() const { return bIsStaticSwitch; };
 	void SetIsStaticSwitch(bool bInIsStaticSwitch) { bIsStaticSwitch = bInIsStaticSwitch; };
 
 	bool GetIsSubscribedToParameterDefinitions() const { return bSubscribedToParameterDefinitions; };
-	void SetIsSubscribedToParameterDefinitions(bool bInSubscribedToParameterDefinitions);
+	NIAGARAEDITOR_API void SetIsSubscribedToParameterDefinitions(bool bInSubscribedToParameterDefinitions);
 
 	bool GetIsOverridingParameterDefinitionsDefaultValue() const { return bOverrideParameterDefinitionsDefaultValue; };
-	void SetIsOverridingParameterDefinitionsDefaultValue(bool bInOverridingParameterDefinitionsDefaultValue);
+	NIAGARAEDITOR_API void SetIsOverridingParameterDefinitionsDefaultValue(bool bInOverridingParameterDefinitionsDefaultValue);
 
 	const FGuid& GetChangeId() const { return ChangeId; };
 	void SetChangeId(const FGuid& NewId) { ChangeId = NewId; };
@@ -216,9 +216,9 @@ public:
 		return StaticSwitchDefaultValue;
 	}
 
-	static bool DefaultsAreEquivalent(const UNiagaraScriptVariable* ScriptVarA, const UNiagaraScriptVariable* ScriptVarB);
+	static NIAGARAEDITOR_API bool DefaultsAreEquivalent(const UNiagaraScriptVariable* ScriptVarA, const UNiagaraScriptVariable* ScriptVarB);
 
-	static FGuid GenerateStableGuid(const UNiagaraScriptVariable* ScriptVariable);
+	static NIAGARAEDITOR_API FGuid GenerateStableGuid(const UNiagaraScriptVariable* ScriptVariable);
 
 private:
 	void AllocateData()
diff --git a/Engine/Plugins/FX/Niagara/Source/NiagaraEditor/Public/NiagaraSystemEditorData.h b/Engine/Plugins/FX/Niagara/Source/NiagaraEditor/Public/NiagaraSystemEditorData.h
--- a/Engine/Plugins/FX/Niagara/Source/NiagaraEditor/Public/NiagaraSystemEditorData.h
+++ b/Engine/Plugins/FX/Niagara/Source/NiagaraEditor/Public/NiagaraSystemEditorData.h
@@ -46,27 +46,27 @@ private:
 };
 
 /** Editor only data for systems. */
-UCLASS()
-class NIAGARAEDITOR_API UNiagaraSystemEditorData : public UNiagaraEditorDataBase
+UCLASS(MinimalAPI)
+class UNiagaraSystemEditorData : public UNiagaraEditorDataBase
 {
 	GENERATED_BODY()
 
 	DECLARE_MULTICAST_DELEGATE(FOnUserParameterScriptVariablesSynced)
 public:
-	UNiagaraSystemEditorData(const FObjectInitializer& ObjectInitializer);
+	NIAGARAEDITOR_API UNiagaraSystemEditorData(const FObjectInitializer& ObjectInitializer);
 
-	void PostInitProperties();
+	NIAGARAEDITOR_API void PostInitProperties();
 
-	virtual void PostLoadFromOwner(UObject* InOwner) override;
+	NIAGARAEDITOR_API virtual void PostLoadFromOwner(UObject* InOwner) override;
 #if WITH_EDITORONLY_DATA
-	static void DeclareConstructClasses(TArray<FTopLevelAssetPath>& OutConstructClasses, const UClass* SpecificSubclass);
+	static NIAGARAEDITOR_API void DeclareConstructClasses(TArray<FTopLevelAssetPath>& OutConstructClasses, const UClass* SpecificSubclass);
 #endif
 
 	/** Gets the root folder for UI folders for emitters. */
-	UNiagaraSystemEditorFolder& GetRootFolder() const;
+	NIAGARAEDITOR_API UNiagaraSystemEditorFolder& GetRootFolder() const;
 
 	/** Gets the stack editor data for the system. */
-	UNiagaraStackEditorData& GetStackEditorData() const;
+	NIAGARAEDITOR_API UNiagaraStackEditorData& GetStackEditorData() const;
 
 	const FTransform& GetOwnerTransform() const {
 		return OwnerTransform;
@@ -76,31 +76,31 @@ public:
 		OwnerTransform = InTransform;
 	}
 
-	TRange<float> GetPlaybackRange() const;
+	NIAGARAEDITOR_API TRange<float> GetPlaybackRange() const;
 
-	void SetPlaybackRange(TRange<float> InPlaybackRange);
+	NIAGARAEDITOR_API void SetPlaybackRange(TRange<float> InPlaybackRange);
 
-	UEdGraph* GetSystemOverviewGraph() const;
+	NIAGARAEDITOR_API UEdGraph* GetSystemOverviewGraph() const;
 
-	const FNiagaraGraphViewSettings& GetSystemOverviewGraphViewSettings() const;
+	NIAGARAEDITOR_API const FNiagaraGraphViewSettings& GetSystemOverviewGraphViewSettings() const;
 
-	void SetSystemOverviewGraphViewSettings(const FNiagaraGraphViewSettings& InOverviewGraphViewSettings);
+	NIAGARAEDITOR_API void SetSystemOverviewGraphViewSettings(const FNiagaraGraphViewSettings& InOverviewGraphViewSettings);
 
-	bool GetOwningSystemIsPlaceholder() const;
+	NIAGARAEDITOR_API bool GetOwningSystemIsPlaceholder() const;
 
-	void SetOwningSystemIsPlaceholder(bool bInSystemIsPlaceholder, UNiagaraSystem& OwnerSystem);
+	NIAGARAEDITOR_API void SetOwningSystemIsPlaceholder(bool bInSystemIsPlaceholder, UNiagaraSystem& OwnerSystem);
 
-	void SynchronizeOverviewGraphWithSystem(UNiagaraSystem& OwnerSystem);
+	NIAGARAEDITOR_API void SynchronizeOverviewGraphWithSystem(UNiagaraSystem& OwnerSystem);
 
-	void InitOnSyncScriptVariables(UNiagaraSystem& System);
+	NIAGARAEDITOR_API void InitOnSyncScriptVariables(UNiagaraSystem& System);
 	
-	void SyncUserScriptVariables(UNiagaraSystem* System);
+	NIAGARAEDITOR_API void SyncUserScriptVariables(UNiagaraSystem* System);
 	FOnUserParameterScriptVariablesSynced& OnUserParameterScriptVariablesSynced() { return OnUserParameterScriptVariablesSyncedDelegate; }
 	
-	UNiagaraScriptVariable* FindOrAddUserScriptVariable(FNiagaraVariable UserParameter, UNiagaraSystem& System);
-	const UNiagaraScriptVariable* FindUserScriptVariable(FGuid UserParameterGuid) const;
-	bool RenameUserScriptVariable(FNiagaraVariable OldVariable, FName NewName);
-	bool RemoveUserScriptVariable(FNiagaraVariable Variable);
+	NIAGARAEDITOR_API UNiagaraScriptVariable* FindOrAddUserScriptVariable(FNiagaraVariable UserParameter, UNiagaraSystem& System);
+	NIAGARAEDITOR_API const UNiagaraScriptVariable* FindUserScriptVariable(FGuid UserParameterGuid) const;
+	NIAGARAEDITOR_API bool RenameUserScriptVariable(FNiagaraVariable OldVariable, FName NewName);
+	NIAGARAEDITOR_API bool RemoveUserScriptVariable(FNiagaraVariable Variable);
 
 	// If true then the preview viewport's orbit setting is saved in the asset data
 	UPROPERTY()
@@ -114,7 +114,7 @@ public:
 	TObjectPtr<UNiagaraHierarchyRoot> UserParameterHierarchy;
 	
 private:
-	void UpdatePlaybackRangeFromEmitters(UNiagaraSystem& OwnerSystem);
+	NIAGARAEDITOR_API void UpdatePlaybackRangeFromEmitters(UNiagaraSystem& OwnerSystem);
 
 private:
 	UPROPERTY(Instanced)
diff --git a/Engine/Plugins/FX/Niagara/Source/NiagaraEditor/Public/SNiagaraGraphActionWidget.h b/Engine/Plugins/FX/Niagara/Source/NiagaraEditor/Public/SNiagaraGraphActionWidget.h
--- a/Engine/Plugins/FX/Niagara/Source/NiagaraEditor/Public/SNiagaraGraphActionWidget.h
+++ b/Engine/Plugins/FX/Niagara/Source/NiagaraEditor/Public/SNiagaraGraphActionWidget.h
@@ -32,7 +32,7 @@ struct FCreateNiagaraWidgetForActionData
 };
 
 /** Custom widget for GraphActionMenu */
-class NIAGARAEDITOR_API SNiagaraGraphActionWidget : public SCompoundWidget
+class SNiagaraGraphActionWidget : public SCompoundWidget
 {
 	public:
 
@@ -40,8 +40,8 @@ class NIAGARAEDITOR_API SNiagaraGraphActionWidget : public SCompoundWidget
 	SLATE_ATTRIBUTE(FText, HighlightText)
 SLATE_END_ARGS()
 
-void Construct(const FArguments& InArgs, const FCreateWidgetForActionData* InCreateData);
-	virtual FReply OnMouseButtonDown( const FGeometry& MyGeometry, const FPointerEvent& MouseEvent ) override;
+NIAGARAEDITOR_API void Construct(const FArguments& InArgs, const FCreateWidgetForActionData* InCreateData);
+	NIAGARAEDITOR_API virtual FReply OnMouseButtonDown( const FGeometry& MyGeometry, const FPointerEvent& MouseEvent ) override;
 
 	private:
 	/** The item that we want to display with this widget */
@@ -51,7 +51,7 @@ void Construct(const FArguments& InArgs, const FCreateWidgetForActionData* InCre
 };
 
 /** Custom widget for GraphActionMenu */
-class NIAGARAEDITOR_API SNiagaraActionWidget : public SCompoundWidget
+class SNiagaraActionWidget : public SCompoundWidget
 {
 public:
 
@@ -61,7 +61,7 @@ public:
 		SLATE_ARGUMENT(bool, bShowTypeIfParameter)
 	SLATE_END_ARGS()
 
-	void Construct(const FArguments& InArgs, const FCreateNiagaraWidgetForActionData& InCreateData);
+	NIAGARAEDITOR_API void Construct(const FArguments& InArgs, const FCreateNiagaraWidgetForActionData& InCreateData);
 	//virtual FReply OnMouseButtonDown( const FGeometry& MyGeometry, const FPointerEvent& MouseEvent ) override;
 
 private:
diff --git a/Engine/Plugins/FX/Niagara/Source/NiagaraEditor/Public/SNiagaraParameterEditor.h b/Engine/Plugins/FX/Niagara/Source/NiagaraEditor/Public/SNiagaraParameterEditor.h
--- a/Engine/Plugins/FX/Niagara/Source/NiagaraEditor/Public/SNiagaraParameterEditor.h
+++ b/Engine/Plugins/FX/Niagara/Source/NiagaraEditor/Public/SNiagaraParameterEditor.h
@@ -9,7 +9,7 @@
 
 /** Base class for inline parameter editors. These editors are expected to maintain an internal value which
 	is populated from a parameter struct. */
-class NIAGARAEDITOR_API SNiagaraParameterEditor : public SCompoundWidget
+class SNiagaraParameterEditor : public SCompoundWidget
 {
 public:
 	DECLARE_DELEGATE(FOnValueChange);
@@ -25,7 +25,7 @@ public:
 		SLATE_ARGUMENT(TOptional<float>, MaximumDesiredWidth)
 	SLATE_END_ARGS();
 
-	void Construct(const FArguments& InArgs);
+	NIAGARAEDITOR_API void Construct(const FArguments& InArgs);
 
 	/** Updates the internal value of the widget from a struct. */
 	virtual void UpdateInternalValueFromStruct(TSharedRef<FStructOnScope> Struct) = 0;
@@ -36,37 +36,37 @@ public:
 	/** Gets whether this is currently the exclusive editor of this parameter, meaning that the corresponding details view
 		should not be updated.  This hack is necessary because the details view closes all color pickers when
 		it's changed! */
-	bool GetIsEditingExclusively();
+	NIAGARAEDITOR_API bool GetIsEditingExclusively();
 
 	/** Sets the OnBeginValueChange delegate which is run when a continuous internal value change begins. */
-	void SetOnBeginValueChange(FOnValueChange InOnBeginValueChange);
+	NIAGARAEDITOR_API void SetOnBeginValueChange(FOnValueChange InOnBeginValueChange);
 
 	/** Sets the OnBeginValueChange delegate which is run when a continuous internal value change ends. */
-	void SetOnEndValueChange(FOnValueChange InOnEndValueChange);
+	NIAGARAEDITOR_API void SetOnEndValueChange(FOnValueChange InOnEndValueChange);
 
 	/** Sets the OnValueChanged delegate which is run when the internal value changes. */
-	void SetOnValueChanged(FOnValueChange InOnValueChanged);
+	NIAGARAEDITOR_API void SetOnValueChanged(FOnValueChange InOnValueChanged);
 
 	/** Gets an optional minimum desired width for this parameter editor. */
-	const TOptional<float>& GetMinimumDesiredWidth() const;
+	NIAGARAEDITOR_API const TOptional<float>& GetMinimumDesiredWidth() const;
 
 	/** Sets the minimum desired width for this parameter editor.  If the option parameter is not set it clears the current value. */
-	void SetMinimumDesiredWidth(TOptional<float> InMinimumDesiredWidth);
+	NIAGARAEDITOR_API void SetMinimumDesiredWidth(TOptional<float> InMinimumDesiredWidth);
 
 	/** Gets an optional maximum desired width for this parameter editor. */
-	const TOptional<float>& GetMaximumDesiredWidth() const;
+	NIAGARAEDITOR_API const TOptional<float>& GetMaximumDesiredWidth() const;
 
 	/** Sets the maximum desired width for this parameter editor.  If the option parameter is not set it clears the current value. */
-	void SetMaximumDesiredWidth(TOptional<float> InMaximumDesiredWidth);
+	NIAGARAEDITOR_API void SetMaximumDesiredWidth(TOptional<float> InMaximumDesiredWidth);
 
 	/** Gets the desired horizontal alignment of this parameter editor for it's parent container. */
-	EHorizontalAlignment GetHorizontalAlignment() const;
+	NIAGARAEDITOR_API EHorizontalAlignment GetHorizontalAlignment() const;
 
 	/** Gets the desired horizontal alignment of this parameter editor for it's parent container. */
-	EVerticalAlignment GetVerticalAlignment() const;
+	NIAGARAEDITOR_API EVerticalAlignment GetVerticalAlignment() const;
 
 	//~ SWidget interface
-	virtual FVector2D ComputeDesiredSize(float LayoutScaleMultiplier) const override;
+	NIAGARAEDITOR_API virtual FVector2D ComputeDesiredSize(float LayoutScaleMultiplier) const override;
 
 	/** */
 	virtual bool CanChangeContinuously() const { return false; }
@@ -74,16 +74,16 @@ protected:
 	/** Sets whether this is currently the exclusive editor of this parameter, meaning that the corresponding details view
 		should not be updated.  This hack is necessary because the details view closes all color pickers when
 		it's changed! */
-	void SetIsEditingExclusively(bool bInIsEditingExclusively);
+	NIAGARAEDITOR_API void SetIsEditingExclusively(bool bInIsEditingExclusively);
 
 	/** Executes the OnBeginValueChange delegate */
-	void ExecuteOnBeginValueChange();
+	NIAGARAEDITOR_API void ExecuteOnBeginValueChange();
 
 	/** Executes the OnEndValueChange delegate. */
-	void ExecuteOnEndValueChange();
+	NIAGARAEDITOR_API void ExecuteOnEndValueChange();
 
 	/** Executes the OnValueChanged delegate. */
-	void ExecuteOnValueChanged();
+	NIAGARAEDITOR_API void ExecuteOnValueChanged();
 
 	template<typename NumericType>
 	static TSharedPtr<TNumericUnitTypeInterface<NumericType>> GetTypeInterface(EUnit TypeUnit)
@@ -96,7 +96,7 @@ protected:
 	}
 
 protected:
-	static const float DefaultInputSize;
+	static NIAGARAEDITOR_API const float DefaultInputSize;
 
 private:
 	//** Flag set when a continuous change is in progress to prevent invoking ExecuteOnBeginValueChange() or ExecuteOnEndValueChange() more than once for a change. */
diff --git a/Engine/Plugins/FX/Niagara/Source/NiagaraEditor/Public/UpgradeNiagaraScriptResults.h b/Engine/Plugins/FX/Niagara/Source/NiagaraEditor/Public/UpgradeNiagaraScriptResults.h
--- a/Engine/Plugins/FX/Niagara/Source/NiagaraEditor/Public/UpgradeNiagaraScriptResults.h
+++ b/Engine/Plugins/FX/Niagara/Source/NiagaraEditor/Public/UpgradeNiagaraScriptResults.h
@@ -89,16 +89,16 @@ public:
 /**
  * Wrapper class for passing results back from the version upgrade python script.
  */
-UCLASS(BlueprintType)
-class NIAGARAEDITOR_API UUpgradeNiagaraScriptResults : public UObject
+UCLASS(BlueprintType, MinimalAPI)
+class UUpgradeNiagaraScriptResults : public UObject
 {
 	GENERATED_BODY()
 
 public:
 
-	UUpgradeNiagaraScriptResults();
+	NIAGARAEDITOR_API UUpgradeNiagaraScriptResults();
 
-	void Init();
+	NIAGARAEDITOR_API void Init();
 	
 	// Whether the converter process was cancelled due to an unrecoverable error in the python script process.
 	UPROPERTY(BlueprintReadWrite, Category = "Scripting")
@@ -111,43 +111,43 @@ public:
 	TArray<TObjectPtr<UNiagaraPythonScriptModuleInput>> NewInputs;
 
 	UFUNCTION(BlueprintCallable, Category = "Scripting")
-    UNiagaraPythonScriptModuleInput* GetOldInput(const FString& InputName);
+    NIAGARAEDITOR_API UNiagaraPythonScriptModuleInput* GetOldInput(const FString& InputName);
 
 	UFUNCTION(BlueprintCallable, Category = "Scripting")
-    void SetFloatInput(const FString& InputName, float Value);
+    NIAGARAEDITOR_API void SetFloatInput(const FString& InputName, float Value);
 
 	UFUNCTION(BlueprintCallable, Category = "Scripting")
-    void SetIntInput(const FString& InputName, int32 Value);
+    NIAGARAEDITOR_API void SetIntInput(const FString& InputName, int32 Value);
 
 	UFUNCTION(BlueprintCallable, Category = "Scripting")
-    void SetBoolInput(const FString& InputName, bool Value);
+    NIAGARAEDITOR_API void SetBoolInput(const FString& InputName, bool Value);
 
 	UFUNCTION(BlueprintCallable, Category = "Scripting")
-    void SetVec2Input(const FString& InputName, FVector2D Value);
+    NIAGARAEDITOR_API void SetVec2Input(const FString& InputName, FVector2D Value);
 
 	UFUNCTION(BlueprintCallable, Category = "Scripting")
-    void SetVec3Input(const FString& InputName, FVector Value);
+    NIAGARAEDITOR_API void SetVec3Input(const FString& InputName, FVector Value);
 
 	UFUNCTION(BlueprintCallable, Category = "Scripting")
-    void SetVec4Input(const FString& InputName, FVector4 Value);
+    NIAGARAEDITOR_API void SetVec4Input(const FString& InputName, FVector4 Value);
 
 	UFUNCTION(BlueprintCallable, Category = "Scripting")
-    void SetColorInput(const FString& InputName, FLinearColor Value);
+    NIAGARAEDITOR_API void SetColorInput(const FString& InputName, FLinearColor Value);
 
 	UFUNCTION(BlueprintCallable, Category = "Scripting")
-    void SetQuatInput(const FString& InputName, FQuat Value);
+    NIAGARAEDITOR_API void SetQuatInput(const FString& InputName, FQuat Value);
 
 	UFUNCTION(BlueprintCallable, Category = "Scripting")
-    void SetEnumInput(const FString& InputName, FString Value);
+    NIAGARAEDITOR_API void SetEnumInput(const FString& InputName, FString Value);
     
     UFUNCTION(BlueprintCallable, Category = "Scripting")
-    void SetLinkedInput(const FString& InputName, FString Value);
+    NIAGARAEDITOR_API void SetLinkedInput(const FString& InputName, FString Value);
 
 	UFUNCTION(BlueprintCallable, Category = "Scripting")
-	void SetNewInput(const FString& InputName, UNiagaraPythonScriptModuleInput* Value);
+	NIAGARAEDITOR_API void SetNewInput(const FString& InputName, UNiagaraPythonScriptModuleInput* Value);
 
 	UFUNCTION(BlueprintCallable, Category = "Scripting")
-	void ResetToDefault(const FString& InputName);
+	NIAGARAEDITOR_API void ResetToDefault(const FString& InputName);
 
 private:
 	UNiagaraPythonScriptModuleInput* GetNewInput(const FName& InputName) const;
@@ -157,7 +157,7 @@ private:
 	TObjectPtr<UNiagaraPythonScriptModuleInput> DummyInput;
 };
 
-struct NIAGARAEDITOR_API FNiagaraScriptVersionUpgradeContext {
+struct FNiagaraScriptVersionUpgradeContext {
 	TFunction<void(UNiagaraClipboardContent*)> CreateClipboardCallback;
 	TFunction<void(UNiagaraClipboardContent*, FText&)> ApplyClipboardCallback;
 	FCompileConstantResolver ConstantResolver;
@@ -235,8 +235,8 @@ public:
 /**
  * Wrapper class for passing results back from the version upgrade python script.
  */
-UCLASS(BlueprintType)
-class NIAGARAEDITOR_API UUpgradeNiagaraEmitterContext : public UObject
+UCLASS(BlueprintType, MinimalAPI)
+class UUpgradeNiagaraEmitterContext : public UObject
 {
 	GENERATED_BODY()
 
@@ -244,9 +244,9 @@ public:
 
 	UUpgradeNiagaraEmitterContext() {}
 
-	void Init(UNiagaraPythonEmitter* InOldEmitter, UNiagaraPythonEmitter* InNewEmitter);
-	bool IsValid() const;
-	const TArray<FVersionedNiagaraEmitterData*>& GetUpgradeData() const;
+	NIAGARAEDITOR_API void Init(UNiagaraPythonEmitter* InOldEmitter, UNiagaraPythonEmitter* InNewEmitter);
+	NIAGARAEDITOR_API bool IsValid() const;
+	NIAGARAEDITOR_API const TArray<FVersionedNiagaraEmitterData*>& GetUpgradeData() const;
 	
 	// Whether the converter process was cancelled due to an unrecoverable error in the python script process.
 	UPROPERTY(BlueprintReadWrite, Category = "Scripting")
@@ -260,4 +260,4 @@ public:
 
 private:
 	TArray<FVersionedNiagaraEmitterData*> UpgradeVersionData;
-};
\ No newline at end of file
+};
diff --git a/Engine/Plugins/FX/Niagara/Source/NiagaraEditor/Public/ViewModels/HierarchyEditor/NiagaraHierarchyViewModelBase.h b/Engine/Plugins/FX/Niagara/Source/NiagaraEditor/Public/ViewModels/HierarchyEditor/NiagaraHierarchyViewModelBase.h
--- a/Engine/Plugins/FX/Niagara/Source/NiagaraEditor/Public/ViewModels/HierarchyEditor/NiagaraHierarchyViewModelBase.h
+++ b/Engine/Plugins/FX/Niagara/Source/NiagaraEditor/Public/ViewModels/HierarchyEditor/NiagaraHierarchyViewModelBase.h
@@ -92,8 +92,8 @@ class UNiagaraHierarchyDataRefreshContext : public UObject
 	TMap<FNiagaraHierarchyIdentity, TObjectPtr<UObject>> IdentityToObjectMap;
 };
 
-UCLASS()
-class NIAGARAEDITOR_API UNiagaraHierarchyItemBase : public UObject
+UCLASS(MinimalAPI)
+class UNiagaraHierarchyItemBase : public UObject
 {
 	GENERATED_BODY()
 
@@ -102,18 +102,18 @@ public:
 	virtual ~UNiagaraHierarchyItemBase() override {}
 
 	TArray<TObjectPtr<UNiagaraHierarchyItemBase>>& GetChildrenMutable() { return Children; }
-	const TArray<UNiagaraHierarchyItemBase*>& GetChildren() const;
+	NIAGARAEDITOR_API const TArray<UNiagaraHierarchyItemBase*>& GetChildren() const;
 
 	template<class ChildClass>
 	ChildClass* AddChild();
 
-	UNiagaraHierarchyItemBase* FindChildWithIdentity(FNiagaraHierarchyIdentity ChildIdentity, bool bSearchRecursively = false);
+	NIAGARAEDITOR_API UNiagaraHierarchyItemBase* FindChildWithIdentity(FNiagaraHierarchyIdentity ChildIdentity, bool bSearchRecursively = false);
 
-	UNiagaraHierarchyItemBase* CopyAndAddItemAsChild(const UNiagaraHierarchyItemBase& ItemToCopy);
-	UNiagaraHierarchyItemBase* CopyAndAddItemUnderParentIdentity(const UNiagaraHierarchyItemBase& ItemToCopy, FNiagaraHierarchyIdentity ParentIdentity);
+	NIAGARAEDITOR_API UNiagaraHierarchyItemBase* CopyAndAddItemAsChild(const UNiagaraHierarchyItemBase& ItemToCopy);
+	NIAGARAEDITOR_API UNiagaraHierarchyItemBase* CopyAndAddItemUnderParentIdentity(const UNiagaraHierarchyItemBase& ItemToCopy, FNiagaraHierarchyIdentity ParentIdentity);
 	
 	/** Remove a child with a given identity. Can be searched recursively. This function operates under the assumption there will be only one item with a given identity. */
-	bool RemoveChildWithIdentity(FNiagaraHierarchyIdentity ChildIdentity, bool bSearchRecursively = false);
+	NIAGARAEDITOR_API bool RemoveChildWithIdentity(FNiagaraHierarchyIdentity ChildIdentity, bool bSearchRecursively = false);
 	
 	template<class ChildClass>
 	bool DoesOneChildExist(bool bRecursive = false) const;
@@ -129,13 +129,13 @@ public:
 	/** An identity can be optionally set to create a mapping from previously existing guids or names to hierarchy items that represent them. */
 	void SetIdentity(FNiagaraHierarchyIdentity InIdentity) { Identity = InIdentity; }
 	FNiagaraHierarchyIdentity GetPersistentIdentity() const { return Identity; }
-	TArray<FNiagaraHierarchyIdentity> GetParentIdentities() const;
+	NIAGARAEDITOR_API TArray<FNiagaraHierarchyIdentity> GetParentIdentities() const;
 	
 	/** Overridden modify method to also mark all children as modified */
-	virtual bool Modify(bool bAlwaysMarkDirty = true) override;
+	NIAGARAEDITOR_API virtual bool Modify(bool bAlwaysMarkDirty = true) override;
 
 protected:
-	virtual void PostLoad() override;
+	NIAGARAEDITOR_API virtual void PostLoad() override;
 
 	UPROPERTY()
 	TArray<TObjectPtr<UNiagaraHierarchyItemBase>> Children;
@@ -220,8 +220,8 @@ void UNiagaraHierarchyItemBase::SortChildren(const PREDICATE_CLASS& Predicate, b
 
 class UNiagaraHierarchySection;
 
-UCLASS()
-class NIAGARAEDITOR_API UNiagaraHierarchyRoot : public UNiagaraHierarchyItemBase
+UCLASS(MinimalAPI)
+class UNiagaraHierarchyRoot : public UNiagaraHierarchyItemBase
 {
 	GENERATED_BODY()
 public:
@@ -232,25 +232,25 @@ public:
 	TArray<TObjectPtr<UNiagaraHierarchySection>>& GetSectionDataMutable() { return Sections; }
 
 
-	TSet<FName> GetSections() const;
-	int32 GetSectionIndex(UNiagaraHierarchySection* Section) const;
+	NIAGARAEDITOR_API TSet<FName> GetSections() const;
+	NIAGARAEDITOR_API int32 GetSectionIndex(UNiagaraHierarchySection* Section) const;
 
-	UNiagaraHierarchySection* AddSection(FText InNewSectionName, int32 InsertIndex = INDEX_NONE);
-	UNiagaraHierarchySection* FindSectionByIdentity(FNiagaraHierarchyIdentity SectionIdentity);
+	NIAGARAEDITOR_API UNiagaraHierarchySection* AddSection(FText InNewSectionName, int32 InsertIndex = INDEX_NONE);
+	NIAGARAEDITOR_API UNiagaraHierarchySection* FindSectionByIdentity(FNiagaraHierarchyIdentity SectionIdentity);
 	/** This will copy the section element itself */
-	void DuplicateSectionFromOtherRoot(const UNiagaraHierarchySection& SectionToCopy);
-	void RemoveSection(FText SectionName);
-	void RemoveSectionByIdentity(FNiagaraHierarchyIdentity SectionIdentity);
+	NIAGARAEDITOR_API void DuplicateSectionFromOtherRoot(const UNiagaraHierarchySection& SectionToCopy);
+	NIAGARAEDITOR_API void RemoveSection(FText SectionName);
+	NIAGARAEDITOR_API void RemoveSectionByIdentity(FNiagaraHierarchyIdentity SectionIdentity);
 	
-	virtual bool Modify(bool bAlwaysMarkDirty = true) override;
+	NIAGARAEDITOR_API virtual bool Modify(bool bAlwaysMarkDirty = true) override;
 protected:
 
 	UPROPERTY()
 	TArray<TObjectPtr<UNiagaraHierarchySection>> Sections;
 };
 
-UCLASS()
-class NIAGARAEDITOR_API UNiagaraHierarchyItem : public UNiagaraHierarchyItemBase
+UCLASS(MinimalAPI)
+class UNiagaraHierarchyItem : public UNiagaraHierarchyItemBase
 {
 	GENERATED_BODY()
 public:
@@ -258,8 +258,8 @@ public:
 	virtual ~UNiagaraHierarchyItem() override {}
 };
 
-UCLASS()
-class NIAGARAEDITOR_API UNiagaraHierarchyCategory : public UNiagaraHierarchyItemBase
+UCLASS(MinimalAPI)
+class UNiagaraHierarchyCategory : public UNiagaraHierarchyItemBase
 {
 	GENERATED_BODY()
 public:
@@ -281,7 +281,7 @@ public:
 	 *  We fix this up by looking through our available sections and match up via persistent identity.
 	 *  This function expects the correct section with the same identity to exist already at the root level
 	 */
-	void FixupSectionLinkage();
+	NIAGARAEDITOR_API void FixupSectionLinkage();
 private:
 	UPROPERTY()
 	FName Category;
@@ -294,8 +294,8 @@ private:
 	TObjectPtr<UNiagaraHierarchySection> Section = nullptr;
 };
 
-UCLASS()
-class NIAGARAEDITOR_API UNiagaraHierarchySection : public UNiagaraHierarchyItemBase
+UCLASS(MinimalAPI)
+class UNiagaraHierarchySection : public UNiagaraHierarchyItemBase
 {
 	GENERATED_BODY()
 
@@ -305,7 +305,7 @@ public:
 	void SetSectionName(FName InSectionName) { Section = InSectionName; }
 	FName GetSectionName() const { return Section; }
 	
-	void SetSectionNameAsText(const FText& Text);
+	NIAGARAEDITOR_API void SetSectionNameAsText(const FText& Text);
 	FText GetSectionNameAsText() const { return FText::FromName(Section); }
 
 	FText GetTooltip() const { return Tooltip; }
@@ -320,13 +320,13 @@ private:
 	FText Tooltip;
 };
 
-UCLASS()
-class NIAGARAEDITOR_API UNiagaraHierarchyObjectProperty : public UNiagaraHierarchyItem
+UCLASS(MinimalAPI)
+class UNiagaraHierarchyObjectProperty : public UNiagaraHierarchyItem
 {
 	GENERATED_BODY()
 public:
 	/** To know what object this ObjectProperty is referring to, a persistent guid that can be mapped back to an object is required. */
-	void Initialize(FGuid ObjectGuid, FString PropertyName);
+	NIAGARAEDITOR_API void Initialize(FGuid ObjectGuid, FString PropertyName);
 };
 
 struct FNiagaraHierarchyItemViewModelBase;
@@ -512,7 +512,7 @@ protected:
 	FTimerHandle FullRefreshNextFrameHandle;
 };
 
-struct NIAGARAEDITOR_API FNiagaraHierarchyItemViewModelBase : TSharedFromThis<FNiagaraHierarchyItemViewModelBase>, public FTickableEditorObject
+struct FNiagaraHierarchyItemViewModelBase : TSharedFromThis<FNiagaraHierarchyItemViewModelBase>, public FTickableEditorObject
 {
 	DECLARE_DELEGATE(FOnSynced)
 	DECLARE_DELEGATE_RetVal_OneParam(bool, FOnFilterChild, const TSharedPtr<const FNiagaraHierarchyItemViewModelBase> Child);
@@ -530,7 +530,7 @@ struct NIAGARAEDITOR_API FNiagaraHierarchyItemViewModelBase : TSharedFromThis<FN
 	/** Can be implemented for additional logic that the constructor isn't valid for. */
 	virtual void Initialize() {}
 	
-	virtual ~FNiagaraHierarchyItemViewModelBase() override;
+	NIAGARAEDITOR_API virtual ~FNiagaraHierarchyItemViewModelBase() override;
 	
 	UNiagaraHierarchyItemBase* GetDataMutable() const { return ItemBase; }
 	const UNiagaraHierarchyItemBase* GetData() const { return ItemBase; }
@@ -541,49 +541,49 @@ struct NIAGARAEDITOR_API FNiagaraHierarchyItemViewModelBase : TSharedFromThis<FN
 	template<class T>
 	const T* GetData() const { return Cast<T>(ItemBase); }
 	
-	virtual void Tick(float DeltaTime) override;
-	virtual TStatId GetStatId() const override;
+	NIAGARAEDITOR_API virtual void Tick(float DeltaTime) override;
+	NIAGARAEDITOR_API virtual TStatId GetStatId() const override;
 	
 	virtual FString ToString() const { return ItemBase->ToString(); }
 	FText ToStringAsText() const { return FText::FromString(ToString()); }
 	virtual TArray<FString> GetSearchTerms() const { return {ToString()} ;}
 
-	void RefreshChildrenData();
-	void SyncViewModelsToData();
+	NIAGARAEDITOR_API void RefreshChildrenData();
+	NIAGARAEDITOR_API void SyncViewModelsToData();
 	const TArray<TSharedPtr<FNiagaraHierarchyItemViewModelBase>>& GetChildren() const { return Children; }
 	TArray<TSharedPtr<FNiagaraHierarchyItemViewModelBase>>& GetChildrenMutable() { return Children; }
-	const TArray<TSharedPtr<FNiagaraHierarchyItemViewModelBase>>& GetFilteredChildren() const;
+	NIAGARAEDITOR_API const TArray<TSharedPtr<FNiagaraHierarchyItemViewModelBase>>& GetFilteredChildren() const;
 
-	void AddChildFilter(FOnFilterChild InFilterChild);
+	NIAGARAEDITOR_API void AddChildFilter(FOnFilterChild InFilterChild);
 	
 	template<class DataClass, class ViewModelChildClass>
 	void GetChildrenViewModelsForType(TArray<TSharedPtr<ViewModelChildClass>>& OutChildren, bool bRecursive = false);
 	
-	bool HasParent(TSharedPtr<FNiagaraHierarchyItemViewModelBase> ParentCandidate, bool bRecursive = false);
+	NIAGARAEDITOR_API bool HasParent(TSharedPtr<FNiagaraHierarchyItemViewModelBase> ParentCandidate, bool bRecursive = false);
 
-	void AddChild(TSharedPtr<FNiagaraHierarchyItemViewModelBase> Item);
-	TSharedPtr<FNiagaraHierarchyItemViewModelBase> AddNewItem(TSubclassOf<UNiagaraHierarchyItemBase> NewItemClass);
-	TSharedRef<FNiagaraHierarchyItemViewModelBase> DuplicateToThis(TSharedPtr<FNiagaraHierarchyItemViewModelBase> ItemToDuplicate, int32 InsertIndex = INDEX_NONE);
-	TSharedRef<FNiagaraHierarchyItemViewModelBase> ReparentToThis(TSharedPtr<FNiagaraHierarchyItemViewModelBase> ItemToMove, int32 InsertIndex = INDEX_NONE);
+	NIAGARAEDITOR_API void AddChild(TSharedPtr<FNiagaraHierarchyItemViewModelBase> Item);
+	NIAGARAEDITOR_API TSharedPtr<FNiagaraHierarchyItemViewModelBase> AddNewItem(TSubclassOf<UNiagaraHierarchyItemBase> NewItemClass);
+	NIAGARAEDITOR_API TSharedRef<FNiagaraHierarchyItemViewModelBase> DuplicateToThis(TSharedPtr<FNiagaraHierarchyItemViewModelBase> ItemToDuplicate, int32 InsertIndex = INDEX_NONE);
+	NIAGARAEDITOR_API TSharedRef<FNiagaraHierarchyItemViewModelBase> ReparentToThis(TSharedPtr<FNiagaraHierarchyItemViewModelBase> ItemToMove, int32 InsertIndex = INDEX_NONE);
 
-	TSharedPtr<FNiagaraHierarchyItemViewModelBase> FindViewModelForChild(UNiagaraHierarchyItemBase* Child, bool bSearchRecursively = false) const;
-	TSharedPtr<FNiagaraHierarchyItemViewModelBase> FindViewModelForChild(FNiagaraHierarchyIdentity ChildIdentity, bool bSearchRecursively = false) const;
-	int32 FindIndexOfChild(UNiagaraHierarchyItemBase* Child) const;
-	int32 FindIndexOfDataChild(TSharedPtr<FNiagaraHierarchyItemViewModelBase> Child) const;
-	int32 FindIndexOfDataChild(UNiagaraHierarchyItemBase* Child) const;
+	NIAGARAEDITOR_API TSharedPtr<FNiagaraHierarchyItemViewModelBase> FindViewModelForChild(UNiagaraHierarchyItemBase* Child, bool bSearchRecursively = false) const;
+	NIAGARAEDITOR_API TSharedPtr<FNiagaraHierarchyItemViewModelBase> FindViewModelForChild(FNiagaraHierarchyIdentity ChildIdentity, bool bSearchRecursively = false) const;
+	NIAGARAEDITOR_API int32 FindIndexOfChild(UNiagaraHierarchyItemBase* Child) const;
+	NIAGARAEDITOR_API int32 FindIndexOfDataChild(TSharedPtr<FNiagaraHierarchyItemViewModelBase> Child) const;
+	NIAGARAEDITOR_API int32 FindIndexOfDataChild(UNiagaraHierarchyItemBase* Child) const;
 
-	UNiagaraHierarchyItemBase* AddChild(TSubclassOf<UNiagaraHierarchyItemBase> NewChildClass, FNiagaraHierarchyIdentity ChildIdentity);
+	NIAGARAEDITOR_API UNiagaraHierarchyItemBase* AddChild(TSubclassOf<UNiagaraHierarchyItemBase> NewChildClass, FNiagaraHierarchyIdentity ChildIdentity);
 	
 	/** Deleting will ask the parent to delete its child */
-	void Delete();
-	void DeleteChild(TSharedPtr<FNiagaraHierarchyItemViewModelBase> Child);
+	NIAGARAEDITOR_API void Delete();
+	NIAGARAEDITOR_API void DeleteChild(TSharedPtr<FNiagaraHierarchyItemViewModelBase> Child);
 
 	TWeakObjectPtr<UNiagaraHierarchyViewModelBase> GetHierarchyViewModel() const { return HierarchyViewModel; }
 
 	/** Returns a set result if the item can accept a drop either above/onto/below the item.  */
-	TOptional<EItemDropZone> OnCanRowAcceptDrop(const FDragDropEvent& DragDropEvent, EItemDropZone ItemDropZone, TSharedPtr<FNiagaraHierarchyItemViewModelBase> Item);
-	virtual FReply OnDroppedOnRow(const FDragDropEvent& DragDropEvent, EItemDropZone ItemDropZone, TSharedPtr<FNiagaraHierarchyItemViewModelBase> Item);
-	void OnRowDragLeave(const FDragDropEvent& DragDropEvent);
+	NIAGARAEDITOR_API TOptional<EItemDropZone> OnCanRowAcceptDrop(const FDragDropEvent& DragDropEvent, EItemDropZone ItemDropZone, TSharedPtr<FNiagaraHierarchyItemViewModelBase> Item);
+	NIAGARAEDITOR_API virtual FReply OnDroppedOnRow(const FDragDropEvent& DragDropEvent, EItemDropZone ItemDropZone, TSharedPtr<FNiagaraHierarchyItemViewModelBase> Item);
+	NIAGARAEDITOR_API void OnRowDragLeave(const FDragDropEvent& DragDropEvent);
 
 	struct FCanPerformActionResults
 	{
@@ -611,7 +611,7 @@ struct NIAGARAEDITOR_API FNiagaraHierarchyItemViewModelBase : TSharedFromThis<FN
 	virtual bool CanHaveChildren() const { return false; }
 	
 	/** Should return true if an item should be draggable. An uneditable item can not be dragged even if CanDragInternal returns true. */
-	FCanPerformActionResults CanDrag();
+	NIAGARAEDITOR_API FCanPerformActionResults CanDrag();
 	
 	/** Returns true if renamable */
 	bool CanRename() { return IsEditableByUser().bCanPerform && CanRenameInternal(); }
@@ -663,9 +663,9 @@ struct NIAGARAEDITOR_API FNiagaraHierarchyItemViewModelBase : TSharedFromThis<FN
 	virtual UObject* GetDataForEditing() { return ItemBase; }
 	
 	/** Used to create customized drag drop ops. */
-	TSharedRef<class FNiagaraHierarchyDragDropOp> CreateDragDropOp();
+	NIAGARAEDITOR_API TSharedRef<class FNiagaraHierarchyDragDropOp> CreateDragDropOp();
 	
-	FReply OnDragDetected(const FGeometry& Geometry, const FPointerEvent& PointerEvent, bool bIsSource);
+	NIAGARAEDITOR_API FReply OnDragDetected(const FGeometry& Geometry, const FPointerEvent& PointerEvent, bool bIsSource);
 
 	FSimpleDelegate& GetOnRequestRename() { return OnRequestRenameDelegate; }
 	FOnSynced& GetOnSynced() { return OnSyncedDelegate; }
@@ -686,7 +686,7 @@ protected:
 	virtual bool CanDeleteInternal() { return IsForHierarchy(); }
 
 	/** Should return true if the given drag drop operation is allowed to succeed. */
-	virtual FCanPerformActionResults CanDropOnInternal(TSharedPtr<FNiagaraHierarchyItemViewModelBase>, EItemDropZone ItemDropZone);
+	NIAGARAEDITOR_API virtual FCanPerformActionResults CanDropOnInternal(TSharedPtr<FNiagaraHierarchyItemViewModelBase>, EItemDropZone ItemDropZone);
 	
 	/** Override this to handle drop-on logic. This is called when an item has been dropped on the item that has implemented this function. */
 	virtual void OnDroppedOnInternal(TSharedPtr<FNiagaraHierarchyItemViewModelBase>, EItemDropZone ItemDropZone) { }
@@ -703,7 +703,7 @@ private:
 	virtual void SyncViewModelsToDataInternal() {}
 	virtual void FinalizeInternal() {}
 
-	void PropagateOnChildSynced();
+	NIAGARAEDITOR_API void PropagateOnChildSynced();
 protected:
 	UNiagaraHierarchyItemBase* const ItemBase;
 	/** Parent should be valid for all instances of this struct except for root objects */
@@ -814,7 +814,7 @@ private:
 	const FSlateBrush* SectionImage = nullptr;
 };
 
-struct NIAGARAEDITOR_API FNiagaraHierarchyItemViewModel : FNiagaraHierarchyItemViewModelBase
+struct FNiagaraHierarchyItemViewModel : FNiagaraHierarchyItemViewModelBase
 {
 	FNiagaraHierarchyItemViewModel(UNiagaraHierarchyItem* InItem, TSharedPtr<FNiagaraHierarchyItemViewModelBase> InParent, TWeakObjectPtr<UNiagaraHierarchyViewModelBase> InHierarchyViewModel, bool bInIsForHierarchy) : FNiagaraHierarchyItemViewModelBase(InItem, InParent, InHierarchyViewModel, bInIsForHierarchy) {}
 	
@@ -824,8 +824,8 @@ struct NIAGARAEDITOR_API FNiagaraHierarchyItemViewModel : FNiagaraHierarchyItemV
 	virtual bool CanHaveChildren() const override { return false; }
 	virtual FCanPerformActionResults CanDragInternal() override { return true; }
 
-	virtual FCanPerformActionResults CanDropOnInternal(TSharedPtr<FNiagaraHierarchyItemViewModelBase>, EItemDropZone ItemDropZone) override;
-	virtual void OnDroppedOnInternal(TSharedPtr<FNiagaraHierarchyItemViewModelBase> DroppedItem, EItemDropZone ItemDropZone) override;
+	NIAGARAEDITOR_API virtual FCanPerformActionResults CanDropOnInternal(TSharedPtr<FNiagaraHierarchyItemViewModelBase>, EItemDropZone ItemDropZone) override;
+	NIAGARAEDITOR_API virtual void OnDroppedOnInternal(TSharedPtr<FNiagaraHierarchyItemViewModelBase> DroppedItem, EItemDropZone ItemDropZone) override;
 };
 
 struct FNiagaraHierarchyCategoryViewModel : FNiagaraHierarchyItemViewModelBase
diff --git a/Engine/Plugins/FX/Niagara/Source/NiagaraEditor/Public/ViewModels/HierarchyEditor/NiagaraSummaryViewViewModel.h b/Engine/Plugins/FX/Niagara/Source/NiagaraEditor/Public/ViewModels/HierarchyEditor/NiagaraSummaryViewViewModel.h
--- a/Engine/Plugins/FX/Niagara/Source/NiagaraEditor/Public/ViewModels/HierarchyEditor/NiagaraSummaryViewViewModel.h
+++ b/Engine/Plugins/FX/Niagara/Source/NiagaraEditor/Public/ViewModels/HierarchyEditor/NiagaraSummaryViewViewModel.h
@@ -24,16 +24,16 @@ public:
 	TSharedPtr<FNiagaraEmitterViewModel> EmitterViewModel;
 };
 
-UCLASS()
-class NIAGARAEDITOR_API UNiagaraHierarchyModule : public UNiagaraHierarchyItem
+UCLASS(MinimalAPI)
+class UNiagaraHierarchyModule : public UNiagaraHierarchyItem
 {
 	GENERATED_BODY()
 public:
-	void Initialize(const UNiagaraNodeFunctionCall& ModuleNode);
+	NIAGARAEDITOR_API void Initialize(const UNiagaraNodeFunctionCall& ModuleNode);
 };
 
-UCLASS()
-class NIAGARAEDITOR_API UNiagaraHierarchyModuleInput : public UNiagaraHierarchyItem
+UCLASS(MinimalAPI)
+class UNiagaraHierarchyModuleInput : public UNiagaraHierarchyItem
 {
 	GENERATED_BODY()
 
@@ -41,7 +41,7 @@ public:
 	UNiagaraHierarchyModuleInput() {}
 	virtual ~UNiagaraHierarchyModuleInput() override {}
 
-	void Initialize(const UNiagaraNodeFunctionCall& FunctionCall, FGuid InputGuid);
+	NIAGARAEDITOR_API void Initialize(const UNiagaraNodeFunctionCall& FunctionCall, FGuid InputGuid);
 
 	void SetDisplayNameOverride(const FText& InText) { DisplayNameOverride = InText; }
 	FText GetDisplayNameOverride() const { return DisplayNameOverride; }
@@ -57,8 +57,8 @@ private:
 	FText TooltipOverride;
 };
 
-UCLASS()
-class NIAGARAEDITOR_API UNiagaraHierarchyAssignmentInput : public UNiagaraHierarchyItem
+UCLASS(MinimalAPI)
+class UNiagaraHierarchyAssignmentInput : public UNiagaraHierarchyItem
 {
 	GENERATED_BODY()
 
@@ -66,7 +66,7 @@ public:
 	UNiagaraHierarchyAssignmentInput() {}
 	virtual ~UNiagaraHierarchyAssignmentInput() override {}
 
-	void Initialize(const UNiagaraNodeAssignment& AssignmentNode, FName AssignmentTarget);
+	NIAGARAEDITOR_API void Initialize(const UNiagaraNodeAssignment& AssignmentNode, FName AssignmentTarget);
 
 	FText GetTooltipOverride() const { return TooltipOverride; }
 private:
@@ -75,52 +75,52 @@ private:
 	FText TooltipOverride;
 };
 
-UCLASS()
-class NIAGARAEDITOR_API UNiagaraHierarchyRenderer : public UNiagaraHierarchyItem
+UCLASS(MinimalAPI)
+class UNiagaraHierarchyRenderer : public UNiagaraHierarchyItem
 {
 	GENERATED_BODY()
 public:
-	void Initialize(const UNiagaraRendererProperties& Renderer);
+	NIAGARAEDITOR_API void Initialize(const UNiagaraRendererProperties& Renderer);
 };
 
-UCLASS()
-class NIAGARAEDITOR_API UNiagaraHierarchyEventHandler : public UNiagaraHierarchyItem
+UCLASS(MinimalAPI)
+class UNiagaraHierarchyEventHandler : public UNiagaraHierarchyItem
 {
 	GENERATED_BODY()
 public:
-	void Initialize(const FNiagaraEventScriptProperties& EventHandler);
+	NIAGARAEDITOR_API void Initialize(const FNiagaraEventScriptProperties& EventHandler);
 };
 
-UCLASS()
-class NIAGARAEDITOR_API UNiagaraHierarchyEventHandlerProperties : public UNiagaraHierarchyItem
+UCLASS(MinimalAPI)
+class UNiagaraHierarchyEventHandlerProperties : public UNiagaraHierarchyItem
 {
 	GENERATED_BODY()
 public:
-	void Initialize(const FNiagaraEventScriptProperties& EventHandler);
+	NIAGARAEDITOR_API void Initialize(const FNiagaraEventScriptProperties& EventHandler);
 	
-	static FNiagaraHierarchyIdentity MakeIdentity(const FNiagaraEventScriptProperties& EventHandler);
+	static NIAGARAEDITOR_API FNiagaraHierarchyIdentity MakeIdentity(const FNiagaraEventScriptProperties& EventHandler);
 };
 
-UCLASS()
-class NIAGARAEDITOR_API UNiagaraHierarchySimStage : public UNiagaraHierarchyItem
+UCLASS(MinimalAPI)
+class UNiagaraHierarchySimStage : public UNiagaraHierarchyItem
 {
 	GENERATED_BODY()
 public:
-	void Initialize(const UNiagaraSimulationStageBase& SimStage);
+	NIAGARAEDITOR_API void Initialize(const UNiagaraSimulationStageBase& SimStage);
 };
 
-UCLASS()
-class NIAGARAEDITOR_API UNiagaraHierarchySimStageProperties : public UNiagaraHierarchyItem
+UCLASS(MinimalAPI)
+class UNiagaraHierarchySimStageProperties : public UNiagaraHierarchyItem
 {
 	GENERATED_BODY()
 public:
-	void Initialize(const UNiagaraSimulationStageBase& SimStage);
+	NIAGARAEDITOR_API void Initialize(const UNiagaraSimulationStageBase& SimStage);
 	
-	static FNiagaraHierarchyIdentity MakeIdentity(const UNiagaraSimulationStageBase& SimStage);
+	static NIAGARAEDITOR_API FNiagaraHierarchyIdentity MakeIdentity(const UNiagaraSimulationStageBase& SimStage);
 };
 
-UCLASS()
-class NIAGARAEDITOR_API UNiagaraSummaryViewViewModel : public UNiagaraHierarchyViewModelBase
+UCLASS(MinimalAPI)
+class UNiagaraSummaryViewViewModel : public UNiagaraHierarchyViewModelBase
 {
 	GENERATED_BODY()
 public:
@@ -129,33 +129,33 @@ public:
 	{
 	}
 
-	void Initialize(TSharedRef<FNiagaraEmitterViewModel> EmitterViewModel);
-	virtual void FinalizeInternal() override;
+	NIAGARAEDITOR_API void Initialize(TSharedRef<FNiagaraEmitterViewModel> EmitterViewModel);
+	NIAGARAEDITOR_API virtual void FinalizeInternal() override;
 	
-	TSharedRef<FNiagaraEmitterViewModel> GetEmitterViewModel() const;
+	NIAGARAEDITOR_API TSharedRef<FNiagaraEmitterViewModel> GetEmitterViewModel() const;
 	
-	virtual UNiagaraHierarchyRoot* GetHierarchyRoot() const override;
-	virtual TSharedPtr<FNiagaraHierarchyItemViewModelBase> CreateViewModelForData(UNiagaraHierarchyItemBase* ItemBase, TSharedPtr<FNiagaraHierarchyItemViewModelBase> Parent) override;
+	NIAGARAEDITOR_API virtual UNiagaraHierarchyRoot* GetHierarchyRoot() const override;
+	NIAGARAEDITOR_API virtual TSharedPtr<FNiagaraHierarchyItemViewModelBase> CreateViewModelForData(UNiagaraHierarchyItemBase* ItemBase, TSharedPtr<FNiagaraHierarchyItemViewModelBase> Parent) override;
 	
-	virtual void PrepareSourceItems(UNiagaraHierarchyRoot* SourceRoot, TSharedPtr<FNiagaraHierarchyRootViewModel>) override;
-	virtual void SetupCommands() override;
+	NIAGARAEDITOR_API virtual void PrepareSourceItems(UNiagaraHierarchyRoot* SourceRoot, TSharedPtr<FNiagaraHierarchyRootViewModel>) override;
+	NIAGARAEDITOR_API virtual void SetupCommands() override;
 	
-	virtual TSharedRef<FNiagaraHierarchyDragDropOp> CreateDragDropOp(TSharedRef<FNiagaraHierarchyItemViewModelBase> Item) override;
+	NIAGARAEDITOR_API virtual TSharedRef<FNiagaraHierarchyDragDropOp> CreateDragDropOp(TSharedRef<FNiagaraHierarchyItemViewModelBase> Item) override;
 	
 	virtual bool SupportsDetailsPanel() override { return true; }
-	virtual TArray<TTuple<UClass*, FOnGetDetailCustomizationInstance>> GetInstanceCustomizations() override;
+	NIAGARAEDITOR_API virtual TArray<TTuple<UClass*, FOnGetDetailCustomizationInstance>> GetInstanceCustomizations() override;
 
-	TMap<FGuid, UObject*> GetObjectsForProperties();
+	NIAGARAEDITOR_API TMap<FGuid, UObject*> GetObjectsForProperties();
 
-	UNiagaraNodeFunctionCall* GetFunctionCallNode(const FGuid& NodeIdentity);
-	void ClearFunctionCallNodeCache(const FGuid& NodeIdentity);
-	TOptional<struct FInputData> GetInputData(const UNiagaraHierarchyModuleInput& Input);
+	NIAGARAEDITOR_API UNiagaraNodeFunctionCall* GetFunctionCallNode(const FGuid& NodeIdentity);
+	NIAGARAEDITOR_API void ClearFunctionCallNodeCache(const FGuid& NodeIdentity);
+	NIAGARAEDITOR_API TOptional<struct FInputData> GetInputData(const UNiagaraHierarchyModuleInput& Input);
 
 private:
-	void OnScriptGraphChanged(const FEdGraphEditAction& Action, const UNiagaraScript& Script);
-	void OnRenderersChanged();
-	void OnSimStagesChanged();
-	void OnEventHandlersChanged();
+	NIAGARAEDITOR_API void OnScriptGraphChanged(const FEdGraphEditAction& Action, const UNiagaraScript& Script);
+	NIAGARAEDITOR_API void OnRenderersChanged();
+	NIAGARAEDITOR_API void OnSimStagesChanged();
+	NIAGARAEDITOR_API void OnEventHandlersChanged();
 protected:
 	// The cache is used to speed up access across different inputs, as the view models for both regular inputs & modules, dynamic inputs & assignment nodes need to 'find' these nodes which is expensive
 	TMap<FGuid, TWeakObjectPtr<UNiagaraNodeFunctionCall>> FunctionCallCache;
@@ -469,4 +469,4 @@ public:
 private:
 	TWeakPtr<struct FNiagaraFunctionViewModel> ModuleViewModel;
 	TSharedPtr<SInlineEditableTextBlock> InlineEditableTextBlock;
-};
\ No newline at end of file
+};
diff --git a/Engine/Plugins/FX/Niagara/Source/NiagaraEditor/Public/ViewModels/NiagaraCurveSelectionViewModel.h b/Engine/Plugins/FX/Niagara/Source/NiagaraEditor/Public/ViewModels/NiagaraCurveSelectionViewModel.h
--- a/Engine/Plugins/FX/Niagara/Source/NiagaraEditor/Public/ViewModels/NiagaraCurveSelectionViewModel.h
+++ b/Engine/Plugins/FX/Niagara/Source/NiagaraEditor/Public/ViewModels/NiagaraCurveSelectionViewModel.h
@@ -17,22 +17,22 @@ class UNiagaraSystem;
 class FNiagaraSystemViewModel;
 struct FRichCurve;
 
-struct NIAGARAEDITOR_API FNiagaraCurveSelectionTreeNodeDataId
+struct FNiagaraCurveSelectionTreeNodeDataId
 {
 	FNiagaraCurveSelectionTreeNodeDataId()
 		: Object(nullptr)
 	{
 	}
 
-	bool operator==(const FNiagaraCurveSelectionTreeNodeDataId& Other) const;
+	NIAGARAEDITOR_API bool operator==(const FNiagaraCurveSelectionTreeNodeDataId& Other) const;
 
 	FName UniqueName;
 	FGuid Guid;
 	UObject* Object;
 
-	static FNiagaraCurveSelectionTreeNodeDataId FromUniqueName(FName UniqueName);
-	static FNiagaraCurveSelectionTreeNodeDataId FromGuid(FGuid Guid);
-	static FNiagaraCurveSelectionTreeNodeDataId FromObject(UObject* Object);
+	static NIAGARAEDITOR_API FNiagaraCurveSelectionTreeNodeDataId FromUniqueName(FName UniqueName);
+	static NIAGARAEDITOR_API FNiagaraCurveSelectionTreeNodeDataId FromGuid(FGuid Guid);
+	static NIAGARAEDITOR_API FNiagaraCurveSelectionTreeNodeDataId FromObject(UObject* Object);
 };
 
 enum class ENiagaraCurveSelectionNodeStyleMode
@@ -45,90 +45,90 @@ enum class ENiagaraCurveSelectionNodeStyleMode
 	CurveComponent
 };
 
-struct NIAGARAEDITOR_API FNiagaraCurveSelectionTreeNode : TSharedFromThis<FNiagaraCurveSelectionTreeNode>
+struct FNiagaraCurveSelectionTreeNode : TSharedFromThis<FNiagaraCurveSelectionTreeNode>
 {
 public:
-	FNiagaraCurveSelectionTreeNode();
+	NIAGARAEDITOR_API FNiagaraCurveSelectionTreeNode();
 
-	const FNiagaraCurveSelectionTreeNodeDataId& GetDataId() const;
+	NIAGARAEDITOR_API const FNiagaraCurveSelectionTreeNodeDataId& GetDataId() const;
 
-	void SetDataId(const FNiagaraCurveSelectionTreeNodeDataId& InDataId);
+	NIAGARAEDITOR_API void SetDataId(const FNiagaraCurveSelectionTreeNodeDataId& InDataId);
 
-	FGuid GetNodeUniqueId() const;
+	NIAGARAEDITOR_API FGuid GetNodeUniqueId() const;
 
-	FText GetDisplayName() const;
+	NIAGARAEDITOR_API FText GetDisplayName() const;
 
-	void SetDisplayName(FText InDisplayName);
+	NIAGARAEDITOR_API void SetDisplayName(FText InDisplayName);
 
-	FText GetSecondDisplayName() const;
+	NIAGARAEDITOR_API FText GetSecondDisplayName() const;
 
-	void SetSecondDisplayName(FText InSecondDisplayName);
+	NIAGARAEDITOR_API void SetSecondDisplayName(FText InSecondDisplayName);
 
-	ENiagaraCurveSelectionNodeStyleMode GetStyleMode() const;
+	NIAGARAEDITOR_API ENiagaraCurveSelectionNodeStyleMode GetStyleMode() const;
 
-	FName GetExecutionCategory() const;
+	NIAGARAEDITOR_API FName GetExecutionCategory() const;
 
-	FName GetExecutionSubcategory() const;
+	NIAGARAEDITOR_API FName GetExecutionSubcategory() const;
 
-	bool GetIsParameter() const;
+	NIAGARAEDITOR_API bool GetIsParameter() const;
 
-	void SetStyle(ENiagaraCurveSelectionNodeStyleMode InStyleMode, FName InExecutionCategory, FName InExecutionSubcategory, bool bInIsParameter);
+	NIAGARAEDITOR_API void SetStyle(ENiagaraCurveSelectionNodeStyleMode InStyleMode, FName InExecutionCategory, FName InExecutionSubcategory, bool bInIsParameter);
 
-	TSharedPtr<FNiagaraCurveSelectionTreeNode> GetParent() const;
+	NIAGARAEDITOR_API TSharedPtr<FNiagaraCurveSelectionTreeNode> GetParent() const;
 
 protected:
-	void SetParent(TSharedPtr<FNiagaraCurveSelectionTreeNode> InParent);
+	NIAGARAEDITOR_API void SetParent(TSharedPtr<FNiagaraCurveSelectionTreeNode> InParent);
 
 public:
-	const TArray<TSharedRef<FNiagaraCurveSelectionTreeNode>>& GetChildNodes() const;
+	NIAGARAEDITOR_API const TArray<TSharedRef<FNiagaraCurveSelectionTreeNode>>& GetChildNodes() const;
 
-	void SetChildNodes(const TArray<TSharedRef<FNiagaraCurveSelectionTreeNode>> InChildNodes);
+	NIAGARAEDITOR_API void SetChildNodes(const TArray<TSharedRef<FNiagaraCurveSelectionTreeNode>> InChildNodes);
 
-	static TSharedPtr<FNiagaraCurveSelectionTreeNode> FindNodeWithDataId(const TArray<TSharedRef<FNiagaraCurveSelectionTreeNode>>& Nodes, FNiagaraCurveSelectionTreeNodeDataId DataId);
+	static NIAGARAEDITOR_API TSharedPtr<FNiagaraCurveSelectionTreeNode> FindNodeWithDataId(const TArray<TSharedRef<FNiagaraCurveSelectionTreeNode>>& Nodes, FNiagaraCurveSelectionTreeNodeDataId DataId);
 
-	TWeakObjectPtr<UNiagaraDataInterfaceCurveBase> GetCurveDataInterface() const;
+	NIAGARAEDITOR_API TWeakObjectPtr<UNiagaraDataInterfaceCurveBase> GetCurveDataInterface() const;
 
-	FRichCurve* GetCurve() const;
+	NIAGARAEDITOR_API FRichCurve* GetCurve() const;
 
-	FName GetCurveName() const;
+	NIAGARAEDITOR_API FName GetCurveName() const;
 
-	FLinearColor GetCurveColor() const;
+	NIAGARAEDITOR_API FLinearColor GetCurveColor() const;
 
-	bool GetCurveIsReadOnly() const;
+	NIAGARAEDITOR_API bool GetCurveIsReadOnly() const;
 
-	void SetCurveDataInterface(UNiagaraDataInterfaceCurveBase* InCurveDataInterface);
+	NIAGARAEDITOR_API void SetCurveDataInterface(UNiagaraDataInterfaceCurveBase* InCurveDataInterface);
 
-	void SetPlaceholderDataInterfaceHandle(TSharedPtr<FNiagaraPlaceholderDataInterfaceHandle> InPlaceholderDataInterfaceHandle);
+	NIAGARAEDITOR_API void SetPlaceholderDataInterfaceHandle(TSharedPtr<FNiagaraPlaceholderDataInterfaceHandle> InPlaceholderDataInterfaceHandle);
 
-	void SetCurveData(UNiagaraDataInterfaceCurveBase* InCurveDataInterface, FRichCurve* InCurve, FName InCurveName, FLinearColor InCurveColor);
+	NIAGARAEDITOR_API void SetCurveData(UNiagaraDataInterfaceCurveBase* InCurveDataInterface, FRichCurve* InCurve, FName InCurveName, FLinearColor InCurveColor);
 
-	const TOptional<FObjectKey>& GetDisplayedObjectKey() const;
+	NIAGARAEDITOR_API const TOptional<FObjectKey>& GetDisplayedObjectKey() const;
 
-	void SetDisplayedObjectKey(FObjectKey InDisplayedObjectKey);
+	NIAGARAEDITOR_API void SetDisplayedObjectKey(FObjectKey InDisplayedObjectKey);
 
-	bool GetShowInTree() const;
+	NIAGARAEDITOR_API bool GetShowInTree() const;
 
-	void SetShowInTree(bool bInShouldShowInTree);
+	NIAGARAEDITOR_API void SetShowInTree(bool bInShouldShowInTree);
 
-	bool GetIsExpanded() const;
+	NIAGARAEDITOR_API bool GetIsExpanded() const;
 
-	void SetIsExpanded(bool bInIsExpanded);
+	NIAGARAEDITOR_API void SetIsExpanded(bool bInIsExpanded);
 
-	bool GetIsEnabled() const;
+	NIAGARAEDITOR_API bool GetIsEnabled() const;
 
-	void SetIsEnabled(bool bInIsEnabled);
+	NIAGARAEDITOR_API void SetIsEnabled(bool bInIsEnabled);
 
-	bool GetIsEnabledAndParentIsEnabled() const;
+	NIAGARAEDITOR_API bool GetIsEnabledAndParentIsEnabled() const;
 
-	const TArray<int32>& GetSortIndices() const;
+	NIAGARAEDITOR_API const TArray<int32>& GetSortIndices() const;
 
-	void UpdateSortIndices(int32 Index);
+	NIAGARAEDITOR_API void UpdateSortIndices(int32 Index);
 
-	void ResetCachedEnabledState();
+	NIAGARAEDITOR_API void ResetCachedEnabledState();
 
-	FSimpleMulticastDelegate& GetOnCurveChanged();
+	NIAGARAEDITOR_API FSimpleMulticastDelegate& GetOnCurveChanged();
 
-	void NotifyCurveChanged();
+	NIAGARAEDITOR_API void NotifyCurveChanged();
 
 private:
 	FNiagaraCurveSelectionTreeNodeDataId DataId;
@@ -158,8 +158,8 @@ private:
 	FSimpleMulticastDelegate OnCurveChangedDelegate;
 };
 
-UCLASS()
-class NIAGARAEDITOR_API UNiagaraCurveSelectionViewModel : public UObject
+UCLASS(MinimalAPI)
+class UNiagaraCurveSelectionViewModel : public UObject
 {
 public:
 	GENERATED_BODY()
@@ -168,23 +168,23 @@ public:
 	DECLARE_MULTICAST_DELEGATE_OneParam(FOnRequestSelectNode, FGuid /* NodeIdToSelect */);
 
 public:
-	void Initialize(TSharedRef<FNiagaraSystemViewModel> InSystemViewModel);
+	NIAGARAEDITOR_API void Initialize(TSharedRef<FNiagaraSystemViewModel> InSystemViewModel);
 
-	void Finalize();
+	NIAGARAEDITOR_API void Finalize();
 
-	const TArray<TSharedRef<FNiagaraCurveSelectionTreeNode>>& GetRootNodes();
+	NIAGARAEDITOR_API const TArray<TSharedRef<FNiagaraCurveSelectionTreeNode>>& GetRootNodes();
 
-	void FocusAndSelectCurveDataInterface(UNiagaraDataInterfaceCurveBase& CurveDataInterface);
+	NIAGARAEDITOR_API void FocusAndSelectCurveDataInterface(UNiagaraDataInterfaceCurveBase& CurveDataInterface);
 
-	void Refresh();
+	NIAGARAEDITOR_API void Refresh();
 
-	void RefreshDeferred();
+	NIAGARAEDITOR_API void RefreshDeferred();
 
-	void Tick();
+	NIAGARAEDITOR_API void Tick();
 
-	FSimpleMulticastDelegate& OnRefreshed();
+	NIAGARAEDITOR_API FSimpleMulticastDelegate& OnRefreshed();
 
-	FOnRequestSelectNode& OnRequestSelectNode();
+	NIAGARAEDITOR_API FOnRequestSelectNode& OnRequestSelectNode();
 
 private:
 	TSharedRef<FNiagaraCurveSelectionTreeNode> CreateNodeForCurveDataInterface(const FNiagaraCurveSelectionTreeNodeDataId& DataId, UNiagaraDataInterfaceCurveBase& CurveDataInterface, bool bIsParameter) const;
@@ -226,4 +226,4 @@ private:
 	mutable bool bHandlingInternalCurveChanged;
 
 	bool bRefreshPending;
-};
\ No newline at end of file
+};
diff --git a/Engine/Plugins/FX/Niagara/Source/NiagaraEditor/Public/ViewModels/NiagaraParameterDefinitionsPanelViewModel.h b/Engine/Plugins/FX/Niagara/Source/NiagaraEditor/Public/ViewModels/NiagaraParameterDefinitionsPanelViewModel.h
--- a/Engine/Plugins/FX/Niagara/Source/NiagaraEditor/Public/ViewModels/NiagaraParameterDefinitionsPanelViewModel.h
+++ b/Engine/Plugins/FX/Niagara/Source/NiagaraEditor/Public/ViewModels/NiagaraParameterDefinitionsPanelViewModel.h
@@ -12,13 +12,13 @@ class UNiagaraParameterDefinitions;
 class FNiagaraObjectSelection;
 
 /** Interface for view models to the parameter panel. */
-class NIAGARAEDITOR_API INiagaraParameterDefinitionsPanelViewModel : public INiagaraImmutableParameterPanelViewModel
+class INiagaraParameterDefinitionsPanelViewModel : public INiagaraImmutableParameterPanelViewModel
 {
 public:
 	//~ Begin INiagaraImmutableParameterPanelViewModel interface
-	virtual const TArray<UNiagaraScriptVariable*> GetEditableScriptVariablesWithName(const FName ParameterName) const override;
+	NIAGARAEDITOR_API virtual const TArray<UNiagaraScriptVariable*> GetEditableScriptVariablesWithName(const FName ParameterName) const override;
 
-	virtual const TArray<FNiagaraGraphParameterReference> GetGraphParameterReferencesForItem(const FNiagaraParameterPanelItemBase& Item) const override;
+	NIAGARAEDITOR_API virtual const TArray<FNiagaraGraphParameterReference> GetGraphParameterReferencesForItem(const FNiagaraParameterPanelItemBase& Item) const override;
 
 	//~ Begin Pure Virtual Methods
 	virtual const TArray<UNiagaraGraph*> GetEditableGraphsConst() const override = 0;
diff --git a/Engine/Plugins/FX/Niagara/Source/NiagaraEditor/Public/ViewModels/NiagaraParameterPanelViewModel.h b/Engine/Plugins/FX/Niagara/Source/NiagaraEditor/Public/ViewModels/NiagaraParameterPanelViewModel.h
--- a/Engine/Plugins/FX/Niagara/Source/NiagaraEditor/Public/ViewModels/NiagaraParameterPanelViewModel.h
+++ b/Engine/Plugins/FX/Niagara/Source/NiagaraEditor/Public/ViewModels/NiagaraParameterPanelViewModel.h
@@ -77,7 +77,7 @@ struct FMenuAndSearchBoxWidgets
 };
 
 /** Base Interface for view models to SiagaraParameterPanel and SNiagaraParameterDefinitionsPanel. */
-class NIAGARAEDITOR_API INiagaraImmutableParameterPanelViewModel : public TSharedFromThis<INiagaraImmutableParameterPanelViewModel>, public FSelfRegisteringEditorUndoClient
+class INiagaraImmutableParameterPanelViewModel : public TSharedFromThis<INiagaraImmutableParameterPanelViewModel>, public FSelfRegisteringEditorUndoClient
 {
 public:
 	/** Delegate to signal the view model's state has changed. */
@@ -98,17 +98,17 @@ public:
 	//~ End Pure Virtual Methods
 
 	//~ Begin FEditorUndoClient Interface 
-	virtual void PostUndo(bool bSuccess) override;
+	NIAGARAEDITOR_API virtual void PostUndo(bool bSuccess) override;
 	virtual void PostRedo(bool bSuccess) override { PostUndo(bSuccess); };
 	//~ End FEditorUndoClient Interface 
 
-	virtual void CopyParameterReference(const FNiagaraParameterPanelItemBase& ItemToCopy) const;
+	NIAGARAEDITOR_API virtual void CopyParameterReference(const FNiagaraParameterPanelItemBase& ItemToCopy) const;
 
-	virtual bool GetCanCopyParameterReferenceAndToolTip(const FNiagaraParameterPanelItemBase& ItemToCopy, FText& OutCanCopyParameterToolTip) const;
+	NIAGARAEDITOR_API virtual bool GetCanCopyParameterReferenceAndToolTip(const FNiagaraParameterPanelItemBase& ItemToCopy, FText& OutCanCopyParameterToolTip) const;
 
-	virtual void CopyParameterMetaData(const FNiagaraParameterPanelItemBase ItemToCopy) const;
+	NIAGARAEDITOR_API virtual void CopyParameterMetaData(const FNiagaraParameterPanelItemBase ItemToCopy) const;
 
-	virtual bool GetCanCopyParameterMetaDataAndToolTip(const FNiagaraParameterPanelItemBase& ItemToCopy, FText& OutCanCopyToolTip) const;
+	NIAGARAEDITOR_API virtual bool GetCanCopyParameterMetaDataAndToolTip(const FNiagaraParameterPanelItemBase& ItemToCopy, FText& OutCanCopyToolTip) const;
 
 	// Refresh often encapsulates GUI rebuilding, so we have multiple API's for very specific use cases within the view model.
 
@@ -116,13 +116,13 @@ public:
 	void Refresh() const { RefreshFull(false); };
 	
 	// RefreshFull allows for optional category tree expansion/collapse, which can be more expensive.
-	virtual void RefreshFull(bool bDoCategoryExpansion) const;
+	NIAGARAEDITOR_API virtual void RefreshFull(bool bDoCategoryExpansion) const;
 
 	// Enqueue light refresh next tick
 	void RefreshNextTick() const { RefreshFullNextTick(false); }
 	
 	// Equeue a full refresh, parameterized by inputs next tick
-	virtual void RefreshFullNextTick(bool bDoCategoryExpansion) const;
+	NIAGARAEDITOR_API virtual void RefreshFullNextTick(bool bDoCategoryExpansion) const;
 
 	// Do an exhaustive refresh as context may be swapped out entirely
 	virtual void RefreshDueToActiveDocumentChanged() {};
diff --git a/Engine/Plugins/FX/Niagara/Source/NiagaraEditor/Public/ViewModels/NiagaraScratchPadScriptViewModel.h b/Engine/Plugins/FX/Niagara/Source/NiagaraEditor/Public/ViewModels/NiagaraScratchPadScriptViewModel.h
--- a/Engine/Plugins/FX/Niagara/Source/NiagaraEditor/Public/ViewModels/NiagaraScratchPadScriptViewModel.h
+++ b/Engine/Plugins/FX/Niagara/Source/NiagaraEditor/Public/ViewModels/NiagaraScratchPadScriptViewModel.h
@@ -10,7 +10,7 @@ class FNiagaraScriptToolkitParameterPanelViewModel;
 class FUICommandList;
 
 
-class NIAGARAEDITOR_API FNiagaraScratchPadScriptViewModel : public FNiagaraScriptViewModel, public FGCObject
+class FNiagaraScratchPadScriptViewModel : public FNiagaraScriptViewModel, public FGCObject
 {
 public:
 	DECLARE_MULTICAST_DELEGATE(FOnRenamed);
@@ -21,28 +21,28 @@ public:
 	DECLARE_MULTICAST_DELEGATE_OneParam(FOnPinIDFocusRequested, FNiagaraScriptIDAndGraphFocusInfo*  /* FocusInfo */);
 	DECLARE_MULTICAST_DELEGATE_OneParam(FOnGraphSubObjectSelectionChanged, const UObject*);
 
-	FNiagaraScratchPadScriptViewModel(bool bInIsForDataProcessingOnly);
+	NIAGARAEDITOR_API FNiagaraScratchPadScriptViewModel(bool bInIsForDataProcessingOnly);
 
-	~FNiagaraScratchPadScriptViewModel();
+	NIAGARAEDITOR_API ~FNiagaraScratchPadScriptViewModel();
 
-	void Initialize(UNiagaraScript* Script, UNiagaraScript* InEditScript, TWeakPtr<class FNiagaraSystemViewModel> InSystemViewModel);
+	NIAGARAEDITOR_API void Initialize(UNiagaraScript* Script, UNiagaraScript* InEditScript, TWeakPtr<class FNiagaraSystemViewModel> InSystemViewModel);
 
-	void TransferFromOldWhenDoingApply(TSharedPtr< FNiagaraScratchPadScriptViewModel> InOldScriptVM);
+	NIAGARAEDITOR_API void TransferFromOldWhenDoingApply(TSharedPtr< FNiagaraScratchPadScriptViewModel> InOldScriptVM);
 
-	void Finalize();
+	NIAGARAEDITOR_API void Finalize();
 
 	//~ Begin FGCObject
-	virtual void AddReferencedObjects(FReferenceCollector& Collector) override;
+	NIAGARAEDITOR_API virtual void AddReferencedObjects(FReferenceCollector& Collector) override;
 	//~ End FGCObject
 
-	TArray<UNiagaraGraph*> GetEditableGraphs();
+	NIAGARAEDITOR_API TArray<UNiagaraGraph*> GetEditableGraphs();
 
 	//~ Begin NiagaraParameterDefinitionsSubscriberViewModel Interface
 protected:
 	virtual INiagaraParameterDefinitionsSubscriber* GetParameterDefinitionsSubscriber() override { return &EditScript; };
 	//~ End NiagaraParameterDefinitionsSubscriberViewModel Interface
 
-	void OnGraphSubObjectSelectionChanged(const UObject* Obj);
+	NIAGARAEDITOR_API void OnGraphSubObjectSelectionChanged(const UObject* Obj);
 
 	FDelegateHandle ExternalSelectionChangedDelegate;
 
@@ -52,60 +52,60 @@ public:
 		return TEXT("FNiagaraScratchPadScriptViewModel");
 	}
 
-	UNiagaraScript* GetOriginalScript() const;
+	NIAGARAEDITOR_API UNiagaraScript* GetOriginalScript() const;
 
-	const FVersionedNiagaraScript& GetEditScript() const;
+	NIAGARAEDITOR_API const FVersionedNiagaraScript& GetEditScript() const;
 
-	TSharedPtr<INiagaraParameterPanelViewModel> GetParameterPanelViewModel() const;
+	NIAGARAEDITOR_API TSharedPtr<INiagaraParameterPanelViewModel> GetParameterPanelViewModel() const;
 
-	TSharedPtr<FUICommandList> GetParameterPanelCommands() const;
+	NIAGARAEDITOR_API TSharedPtr<FUICommandList> GetParameterPanelCommands() const;
 
-	FText GetToolTip() const;
+	NIAGARAEDITOR_API FText GetToolTip() const;
 
-	bool GetIsPendingRename() const;
+	NIAGARAEDITOR_API bool GetIsPendingRename() const;
 
-	void SetIsPendingRename(bool bInIsPendingRename);
+	NIAGARAEDITOR_API void SetIsPendingRename(bool bInIsPendingRename);
 
-	void SetScriptName(FText InScriptName);
+	NIAGARAEDITOR_API void SetScriptName(FText InScriptName);
 
-	bool GetIsPinned() const;
+	NIAGARAEDITOR_API bool GetIsPinned() const;
 
-	void SetIsPinned(bool bInIsPinned);
+	NIAGARAEDITOR_API void SetIsPinned(bool bInIsPinned);
 
-	float GetEditorHeight() const;
+	NIAGARAEDITOR_API float GetEditorHeight() const;
 
-	void SetEditorHeight(float InEditorHeight);
+	NIAGARAEDITOR_API void SetEditorHeight(float InEditorHeight);
 
-	bool HasUnappliedChanges() const;
+	NIAGARAEDITOR_API bool HasUnappliedChanges() const;
 
-	void ApplyChanges();
+	NIAGARAEDITOR_API void ApplyChanges();
 
-	void DiscardChanges();
+	NIAGARAEDITOR_API void DiscardChanges();
 
-	FOnRenamed& OnRenamed();
+	NIAGARAEDITOR_API FOnRenamed& OnRenamed();
 
-	FOnPinnedChanged& OnPinnedChanged();
+	NIAGARAEDITOR_API FOnPinnedChanged& OnPinnedChanged();
 
-	FOnHasUnappliedChangesChanged& OnHasUnappliedChangesChanged();
+	NIAGARAEDITOR_API FOnHasUnappliedChangesChanged& OnHasUnappliedChangesChanged();
 
-	FOnChangesApplied& OnChangesApplied();
+	NIAGARAEDITOR_API FOnChangesApplied& OnChangesApplied();
 
-	FSimpleDelegate& OnRequestDiscardChanges();
+	NIAGARAEDITOR_API FSimpleDelegate& OnRequestDiscardChanges();
 
-	FOnNodeIDFocusRequested& OnNodeIDFocusRequested();
-	FOnPinIDFocusRequested& OnPinIDFocusRequested();
+	NIAGARAEDITOR_API FOnNodeIDFocusRequested& OnNodeIDFocusRequested();
+	NIAGARAEDITOR_API FOnPinIDFocusRequested& OnPinIDFocusRequested();
 
 	FOnGraphSubObjectSelectionChanged& OnGraphSelectionChanged() { return OnGraphSubObjectSelectionChangedDelegate;	};
 
-	void RaisePinFocusRequested(FNiagaraScriptIDAndGraphFocusInfo*  InFocusInfo);
-	void RaiseNodeFocusRequested(FNiagaraScriptIDAndGraphFocusInfo* InFocusInfo);
+	NIAGARAEDITOR_API void RaisePinFocusRequested(FNiagaraScriptIDAndGraphFocusInfo*  InFocusInfo);
+	NIAGARAEDITOR_API void RaiseNodeFocusRequested(FNiagaraScriptIDAndGraphFocusInfo* InFocusInfo);
 
 private:
-	FText GetDisplayNameInternal() const;
+	NIAGARAEDITOR_API FText GetDisplayNameInternal() const;
 
-	void OnScriptGraphChanged(const FEdGraphEditAction &Action);
+	NIAGARAEDITOR_API void OnScriptGraphChanged(const FEdGraphEditAction &Action);
 
-	void OnScriptPropertyChanged(FPropertyChangedEvent& PropertyChangedEvent);
+	NIAGARAEDITOR_API void OnScriptPropertyChanged(FPropertyChangedEvent& PropertyChangedEvent);
 
 	bool bIsPendingRename;
 	bool bIsPinned;
@@ -132,4 +132,4 @@ private:
 
 	FOnNodeIDFocusRequested OnNodeIDFocusRequestedDelegate;
 	FOnPinIDFocusRequested OnPinIDFocusRequestedDelegate;
-};
\ No newline at end of file
+};
