[Niagara]

diff --git a/Engine/Plugins/FX/Niagara/Source/NiagaraEditor/Public/ViewModels/NiagaraScratchPadViewModel.h b/Engine/Plugins/FX/Niagara/Source/NiagaraEditor/Public/ViewModels/NiagaraScratchPadViewModel.h
--- a/Engine/Plugins/FX/Niagara/Source/NiagaraEditor/Public/ViewModels/NiagaraScratchPadViewModel.h
+++ b/Engine/Plugins/FX/Niagara/Source/NiagaraEditor/Public/ViewModels/NiagaraScratchPadViewModel.h
@@ -11,8 +11,8 @@ class FNiagaraScratchPadScriptViewModel;
 class FNiagaraObjectSelection;
 class UNiagaraScript;
 
-UCLASS()
-class NIAGARAEDITOR_API UNiagaraScratchPadViewModel : public UObject
+UCLASS(MinimalAPI)
+class UNiagaraScratchPadViewModel : public UObject
 {
 	GENERATED_BODY()
 
@@ -23,68 +23,68 @@ public:
 	DECLARE_MULTICAST_DELEGATE(FOnScriptDeleted);
 
 public:
-	void Initialize(TSharedRef<FNiagaraSystemViewModel> InSystemViewModel);
+	NIAGARAEDITOR_API void Initialize(TSharedRef<FNiagaraSystemViewModel> InSystemViewModel);
 
-	void Finalize();
+	NIAGARAEDITOR_API void Finalize();
 
-	void RefreshScriptViewModels();
+	NIAGARAEDITOR_API void RefreshScriptViewModels();
 
-	void ApplyScratchPadChanges();
+	NIAGARAEDITOR_API void ApplyScratchPadChanges();
 	
-	const TArray<TSharedRef<FNiagaraScratchPadScriptViewModel>>& GetScriptViewModels() const;
+	NIAGARAEDITOR_API const TArray<TSharedRef<FNiagaraScratchPadScriptViewModel>>& GetScriptViewModels() const;
 
-	const TArray<TSharedRef<FNiagaraScratchPadScriptViewModel>>& GetEditScriptViewModels() const;
+	NIAGARAEDITOR_API const TArray<TSharedRef<FNiagaraScratchPadScriptViewModel>>& GetEditScriptViewModels() const;
 
-	TSharedPtr<FNiagaraScratchPadScriptViewModel> GetViewModelForScript(UNiagaraScript* InScript);
-	TSharedPtr<FNiagaraScratchPadScriptViewModel> GetViewModelForScript(FName InScriptName);
+	NIAGARAEDITOR_API TSharedPtr<FNiagaraScratchPadScriptViewModel> GetViewModelForScript(UNiagaraScript* InScript);
+	NIAGARAEDITOR_API TSharedPtr<FNiagaraScratchPadScriptViewModel> GetViewModelForScript(FName InScriptName);
 
-	TSharedPtr<FNiagaraScratchPadScriptViewModel> GetViewModelForEditScript(UNiagaraScript* InEditScript);
+	NIAGARAEDITOR_API TSharedPtr<FNiagaraScratchPadScriptViewModel> GetViewModelForEditScript(UNiagaraScript* InEditScript);
 
-	const TArray<ENiagaraScriptUsage>& GetAvailableUsages() const;
+	NIAGARAEDITOR_API const TArray<ENiagaraScriptUsage>& GetAvailableUsages() const;
 
-	FText GetDisplayNameForUsage(ENiagaraScriptUsage InUsage) const;
+	NIAGARAEDITOR_API FText GetDisplayNameForUsage(ENiagaraScriptUsage InUsage) const;
 
-	TSharedRef<FNiagaraObjectSelection> GetObjectSelection();
+	NIAGARAEDITOR_API TSharedRef<FNiagaraObjectSelection> GetObjectSelection();
 
-	TSharedPtr<FNiagaraScratchPadScriptViewModel> GetActiveScriptViewModel();
+	NIAGARAEDITOR_API TSharedPtr<FNiagaraScratchPadScriptViewModel> GetActiveScriptViewModel();
 
-	void SetActiveScriptViewModel(TSharedRef<FNiagaraScratchPadScriptViewModel> InActiveScriptViewModel);
+	NIAGARAEDITOR_API void SetActiveScriptViewModel(TSharedRef<FNiagaraScratchPadScriptViewModel> InActiveScriptViewModel);
 
-	void FocusScratchPadScriptViewModel(TSharedRef<FNiagaraScratchPadScriptViewModel> InScriptViewModel);
+	NIAGARAEDITOR_API void FocusScratchPadScriptViewModel(TSharedRef<FNiagaraScratchPadScriptViewModel> InScriptViewModel);
 
-	void ResetActiveScriptViewModel();
+	NIAGARAEDITOR_API void ResetActiveScriptViewModel();
 
-	void CopyActiveScript();
+	NIAGARAEDITOR_API void CopyActiveScript();
 
-	bool CanPasteScript() const;
+	NIAGARAEDITOR_API bool CanPasteScript() const;
 
-	void PasteScript();
+	NIAGARAEDITOR_API void PasteScript();
 
-	void DeleteActiveScript();
+	NIAGARAEDITOR_API void DeleteActiveScript();
 
-	TSharedPtr<FNiagaraScratchPadScriptViewModel> CreateNewScript(ENiagaraScriptUsage InScriptUsage, ENiagaraScriptUsage InTargetSupportedUsage, FNiagaraTypeDefinition InOutputType);
+	NIAGARAEDITOR_API TSharedPtr<FNiagaraScratchPadScriptViewModel> CreateNewScript(ENiagaraScriptUsage InScriptUsage, ENiagaraScriptUsage InTargetSupportedUsage, FNiagaraTypeDefinition InOutputType);
 
-	TSharedPtr<FNiagaraScratchPadScriptViewModel> CreateNewScriptAsDuplicate(const UNiagaraScript* ScriptToDuplicate);
+	NIAGARAEDITOR_API TSharedPtr<FNiagaraScratchPadScriptViewModel> CreateNewScriptAsDuplicate(const UNiagaraScript* ScriptToDuplicate);
 
-	void CreateAssetFromActiveScript();
+	NIAGARAEDITOR_API void CreateAssetFromActiveScript();
 
-	bool CanSelectNextUsageForActiveScript();
+	NIAGARAEDITOR_API bool CanSelectNextUsageForActiveScript();
 
-	void SelectNextUsageForActiveScript();
+	NIAGARAEDITOR_API void SelectNextUsageForActiveScript();
 
-	bool HasUnappliedChanges() const;
+	NIAGARAEDITOR_API bool HasUnappliedChanges() const;
 
-	FOnScriptViewModelsChanged& OnScriptViewModelsChanged();
+	NIAGARAEDITOR_API FOnScriptViewModelsChanged& OnScriptViewModelsChanged();
 
-	FOnScriptViewModelsChanged& OnEditScriptViewModelsChanged();
+	NIAGARAEDITOR_API FOnScriptViewModelsChanged& OnEditScriptViewModelsChanged();
 
-	FOnActiveScriptChanged& OnActiveScriptChanged();
+	NIAGARAEDITOR_API FOnActiveScriptChanged& OnActiveScriptChanged();
 
-	FOnScriptRenamed& OnScriptRenamed();
+	NIAGARAEDITOR_API FOnScriptRenamed& OnScriptRenamed();
 
-	FOnScriptDeleted& OnScriptDeleted();
+	NIAGARAEDITOR_API FOnScriptDeleted& OnScriptDeleted();
 
-	void OpenEditorForActive();
+	NIAGARAEDITOR_API void OpenEditorForActive();
 
 private:
 	TSharedRef<FNiagaraSystemViewModel> GetSystemViewModel();
@@ -140,4 +140,4 @@ private:
 	FOnScriptRenamed OnScriptRenamedDelegate;
 
 	FOnScriptDeleted OnScriptDeletedDelegate;
-};
\ No newline at end of file
+};
diff --git a/Engine/Plugins/FX/Niagara/Source/NiagaraEditor/Public/ViewModels/NiagaraScriptViewModel.h b/Engine/Plugins/FX/Niagara/Source/NiagaraEditor/Public/ViewModels/NiagaraScriptViewModel.h
--- a/Engine/Plugins/FX/Niagara/Source/NiagaraEditor/Public/ViewModels/NiagaraScriptViewModel.h
+++ b/Engine/Plugins/FX/Niagara/Source/NiagaraEditor/Public/ViewModels/NiagaraScriptViewModel.h
@@ -77,7 +77,7 @@ public:
 
 	//~ FEditorUndoClient Interface
 	NIAGARAEDITOR_API virtual void PostUndo(bool bSuccess) override;
-	NIAGARAEDITOR_API virtual void PostRedo(bool bSuccess) override { PostUndo(bSuccess); }
+	virtual void PostRedo(bool bSuccess) override { PostUndo(bSuccess); }
 
 	UNiagaraScript* GetContainerScript(ENiagaraScriptUsage InUsage, FGuid InUsageId = FGuid());
 	UNiagaraScript* GetScript(ENiagaraScriptUsage InUsage, FGuid InUsageId = FGuid());
diff --git a/Engine/Plugins/FX/Niagara/Source/NiagaraEditor/Public/ViewModels/NiagaraSystemEditorDocumentsViewModel.h b/Engine/Plugins/FX/Niagara/Source/NiagaraEditor/Public/ViewModels/NiagaraSystemEditorDocumentsViewModel.h
--- a/Engine/Plugins/FX/Niagara/Source/NiagaraEditor/Public/ViewModels/NiagaraSystemEditorDocumentsViewModel.h
+++ b/Engine/Plugins/FX/Niagara/Source/NiagaraEditor/Public/ViewModels/NiagaraSystemEditorDocumentsViewModel.h
@@ -17,8 +17,8 @@ class FNiagaraObjectSelection;
 class UNiagaraScript;
 class UEdGraph;
 
-UCLASS()
-class NIAGARAEDITOR_API UNiagaraSystemEditorDocumentsViewModel : public UObject
+UCLASS(MinimalAPI)
+class UNiagaraSystemEditorDocumentsViewModel : public UObject
 {
 	GENERATED_BODY()
 
@@ -30,55 +30,55 @@ public:
 
 
 public:
-	void Initialize(TSharedRef<FNiagaraSystemViewModel> InSystemViewModel);
-	void Finalize();
+	NIAGARAEDITOR_API void Initialize(TSharedRef<FNiagaraSystemViewModel> InSystemViewModel);
+	NIAGARAEDITOR_API void Finalize();
 
-	void OpenChildScript(UEdGraph* Graph);
-	void CloseChildScript(UEdGraph* Graph);
+	NIAGARAEDITOR_API void OpenChildScript(UEdGraph* Graph);
+	NIAGARAEDITOR_API void CloseChildScript(UEdGraph* Graph);
 
-	TArray<class UNiagaraGraph*> GetEditableGraphsForActiveScriptDocument();
-	TArray<class UNiagaraGraph*> GetAllGraphsForActiveScriptDocument();
-	TArray<class UNiagaraGraph*> GetEditableGraphsForPrimaryDocument();
-	TArray<class UNiagaraGraph*> GetAllGraphsForPrimaryDocument();
-	void InitializePreTabManager(TSharedPtr<class FNiagaraSystemToolkit> InToolkit);
-	void InitializePostTabManager(TSharedPtr<class FNiagaraSystemToolkit> InToolkit);
+	NIAGARAEDITOR_API TArray<class UNiagaraGraph*> GetEditableGraphsForActiveScriptDocument();
+	NIAGARAEDITOR_API TArray<class UNiagaraGraph*> GetAllGraphsForActiveScriptDocument();
+	NIAGARAEDITOR_API TArray<class UNiagaraGraph*> GetEditableGraphsForPrimaryDocument();
+	NIAGARAEDITOR_API TArray<class UNiagaraGraph*> GetAllGraphsForPrimaryDocument();
+	NIAGARAEDITOR_API void InitializePreTabManager(TSharedPtr<class FNiagaraSystemToolkit> InToolkit);
+	NIAGARAEDITOR_API void InitializePostTabManager(TSharedPtr<class FNiagaraSystemToolkit> InToolkit);
 
-	void SetActiveDocumentTab(TSharedPtr<SDockTab> Tab);
+	NIAGARAEDITOR_API void SetActiveDocumentTab(TSharedPtr<SDockTab> Tab);
 	TWeakPtr<SDockTab> GetActiveDocumentTab() const { return ActiveDocumentTab; }
 	FScriptToolkitsActiveDocumentChanged& OnActiveDocumentChanged() {
 		return ActiveDocChangedDelegate;
 	}
 
-	bool IsPrimaryDocumentActive() const;
-	TSharedPtr<class FNiagaraScratchPadScriptViewModel> GetActiveScratchPadViewModelIfSet();
-	static TSharedPtr < class FNiagaraScratchPadScriptViewModel> GetScratchPadViewModelFromGraph(FNiagaraSystemViewModel* InSysViewModel, UEdGraph* InTargetGraph);
+	NIAGARAEDITOR_API bool IsPrimaryDocumentActive() const;
+	NIAGARAEDITOR_API TSharedPtr<class FNiagaraScratchPadScriptViewModel> GetActiveScratchPadViewModelIfSet();
+	static NIAGARAEDITOR_API TSharedPtr < class FNiagaraScratchPadScriptViewModel> GetScratchPadViewModelFromGraph(FNiagaraSystemViewModel* InSysViewModel, UEdGraph* InTargetGraph);
 
-	void DrawAttentionToPrimaryDocument();
+	NIAGARAEDITOR_API void DrawAttentionToPrimaryDocument();
 	void SetPrimaryDocumentID(const FName& TabId) { PrimaryDocumentTabId = TabId; }
-	void SwapEditableScripts(TSharedPtr < class FNiagaraScratchPadScriptViewModel> OldScriptViewModel, TSharedPtr < class FNiagaraScratchPadScriptViewModel> NewScriptViewModel);
+	NIAGARAEDITOR_API void SwapEditableScripts(TSharedPtr < class FNiagaraScratchPadScriptViewModel> OldScriptViewModel, TSharedPtr < class FNiagaraScratchPadScriptViewModel> NewScriptViewModel);
 
 protected:
-	TSharedPtr<SDockTab> OpenDocument(const UObject* DocumentID, FDocumentTracker::EOpenDocumentCause Cause);
+	NIAGARAEDITOR_API TSharedPtr<SDockTab> OpenDocument(const UObject* DocumentID, FDocumentTracker::EOpenDocumentCause Cause);
 
-	void NavigateTab(FDocumentTracker::EOpenDocumentCause InCause);
+	NIAGARAEDITOR_API void NavigateTab(FDocumentTracker::EOpenDocumentCause InCause);
 
-	void CloseDocumentTab(const UObject* DocumentID);
+	NIAGARAEDITOR_API void CloseDocumentTab(const UObject* DocumentID);
 
 	// Finds any open tabs containing the specified document and adds them to the specified array; returns true if at least one is found
-	bool FindOpenTabsContainingDocument(const UObject* DocumentID, /*inout*/ TArray< TSharedPtr<SDockTab> >& Results);
+	NIAGARAEDITOR_API bool FindOpenTabsContainingDocument(const UObject* DocumentID, /*inout*/ TArray< TSharedPtr<SDockTab> >& Results);
 
 	FName PrimaryDocumentTabId;
 
 
 private:
-	TSharedRef<FNiagaraSystemViewModel> GetSystemViewModel();
+	NIAGARAEDITOR_API TSharedRef<FNiagaraSystemViewModel> GetSystemViewModel();
 
 	FScriptToolkitsActiveDocumentChanged ActiveDocChangedDelegate;
 
 	TWeakPtr<FNiagaraSystemViewModel> SystemViewModelWeak;
 	TWeakPtr< FNiagaraSystemGraphSelectionViewModel> SystemGraphSelectionVMWeak;
 
-	TSharedRef<class SNiagaraScratchPadScriptEditor> CreateGraphEditorWidget(TSharedRef<FTabInfo> InTabInfo, UEdGraph* InGraph);
+	NIAGARAEDITOR_API TSharedRef<class SNiagaraScratchPadScriptEditor> CreateGraphEditorWidget(TSharedRef<FTabInfo> InTabInfo, UEdGraph* InGraph);
 
 private:
 	TSharedPtr<FDocumentTracker> DocumentManager;
@@ -92,4 +92,4 @@ private:
 	TSharedPtr<FNiagaraScratchPadScriptViewModel> ActiveDocumentTabScriptViewModel;
 
 
-};
\ No newline at end of file
+};
diff --git a/Engine/Plugins/FX/Niagara/Source/NiagaraEditor/Public/ViewModels/NiagaraSystemScalabilityViewModel.h b/Engine/Plugins/FX/Niagara/Source/NiagaraEditor/Public/ViewModels/NiagaraSystemScalabilityViewModel.h
--- a/Engine/Plugins/FX/Niagara/Source/NiagaraEditor/Public/ViewModels/NiagaraSystemScalabilityViewModel.h
+++ b/Engine/Plugins/FX/Niagara/Source/NiagaraEditor/Public/ViewModels/NiagaraSystemScalabilityViewModel.h
@@ -17,8 +17,8 @@ class FNiagaraSystemViewModel;
 
 /** A view model for viewing and editing a UNiagaraSystem. */
 
-UCLASS()
-class NIAGARAEDITOR_API UNiagaraSystemScalabilityViewModel : public UObject
+UCLASS(MinimalAPI)
+class UNiagaraSystemScalabilityViewModel : public UObject
 {
 	GENERATED_BODY()
 public:
@@ -28,15 +28,15 @@ public:
 public:
 
 	/** Creates a new view model with the supplied System and System instance. */
-	UNiagaraSystemScalabilityViewModel();
+	NIAGARAEDITOR_API UNiagaraSystemScalabilityViewModel();
 	
 	/** Initializes this scalability view model with the supplied system view model. */
-	void Initialize(TSharedRef<FNiagaraSystemViewModel> InSystemViewModel);
+	NIAGARAEDITOR_API void Initialize(TSharedRef<FNiagaraSystemViewModel> InSystemViewModel);
 
 	/** Returns whether or not this view model is initialized and safe to use. */
-	bool IsValid() const;
+	NIAGARAEDITOR_API bool IsValid() const;
 	
-	bool IsActive() const;
+	NIAGARAEDITOR_API bool IsActive() const;
 
 	TWeakPtr<FNiagaraSystemViewModel> GetSystemViewModel() const { return SystemViewModel; }
 	FNiagaraPlatformSet* GetPreviewPlatformSet() { return PreviewPlatforms.Get(); }
@@ -47,10 +47,10 @@ public:
 	TOptional<TObjectPtr<UDeviceProfile>> GetPreviewDeviceProfile() const { return PreviewDeviceProfile; }
 	bool IsViewModeQualityEnabled(int32 QualityLevel) const { return PreviewPlatforms->IsEffectQualityEnabled(QualityLevel); }
 	
-	void UpdatePreviewDeviceProfile(UDeviceProfile* DeviceProfile);
-	void UpdatePreviewQualityLevel(int32 QualityLevel);
+	NIAGARAEDITOR_API void UpdatePreviewDeviceProfile(UDeviceProfile* DeviceProfile);
+	NIAGARAEDITOR_API void UpdatePreviewQualityLevel(int32 QualityLevel);
 
-	bool IsPlatformActive(const FNiagaraPlatformSet& PlatformSet);
+	NIAGARAEDITOR_API bool IsPlatformActive(const FNiagaraPlatformSet& PlatformSet);
 
 	//void NavigateToScalabilityProperty(UObject* Object, FName PropertyName);
 	
diff --git a/Engine/Plugins/FX/Niagara/Source/NiagaraEditor/Public/ViewModels/NiagaraSystemSelectionViewModel.h b/Engine/Plugins/FX/Niagara/Source/NiagaraEditor/Public/ViewModels/NiagaraSystemSelectionViewModel.h
--- a/Engine/Plugins/FX/Niagara/Source/NiagaraEditor/Public/ViewModels/NiagaraSystemSelectionViewModel.h
+++ b/Engine/Plugins/FX/Niagara/Source/NiagaraEditor/Public/ViewModels/NiagaraSystemSelectionViewModel.h
@@ -11,8 +11,8 @@ class UNiagaraStackEntry;
 class UNiagaraStackSelection;
 class UNiagaraStackViewModel;
 
-UCLASS()
-class NIAGARAEDITOR_API UNiagaraSystemSelectionViewModel : public UObject
+UCLASS(MinimalAPI)
+class UNiagaraSystemSelectionViewModel : public UObject
 {
 public:
 	DECLARE_MULTICAST_DELEGATE(FOnSelectionChanged);
@@ -30,49 +30,49 @@ private:
 public:
 	GENERATED_BODY()
 
-	void Initialize(TSharedRef<FNiagaraSystemViewModel> InSystemViewModel);
+	NIAGARAEDITOR_API void Initialize(TSharedRef<FNiagaraSystemViewModel> InSystemViewModel);
 
-	void Finalize();
+	NIAGARAEDITOR_API void Finalize();
 
-	bool ContainsEntry(UNiagaraStackEntry* StackEntry) const;
+	NIAGARAEDITOR_API bool ContainsEntry(UNiagaraStackEntry* StackEntry) const;
 
-	void GetSelectedEntries(TArray<UNiagaraStackEntry*>& OutSelectedEntries) const;
+	NIAGARAEDITOR_API void GetSelectedEntries(TArray<UNiagaraStackEntry*>& OutSelectedEntries) const;
 
-	bool GetSystemIsSelected() const;
+	NIAGARAEDITOR_API bool GetSystemIsSelected() const;
 
-	const TArray<FGuid>& GetSelectedEmitterHandleIds() const;
+	NIAGARAEDITOR_API const TArray<FGuid>& GetSelectedEmitterHandleIds() const;
 
-	void UpdateSelectedEntries(const TArray<UNiagaraStackEntry*> InSelectedEntries, const TArray<UNiagaraStackEntry*> InDeselectedEntries, bool bClearCurrentSelection);
+	NIAGARAEDITOR_API void UpdateSelectedEntries(const TArray<UNiagaraStackEntry*> InSelectedEntries, const TArray<UNiagaraStackEntry*> InDeselectedEntries, bool bClearCurrentSelection);
 
-	UNiagaraStackViewModel* GetSelectionStackViewModel();
+	NIAGARAEDITOR_API UNiagaraStackViewModel* GetSelectionStackViewModel();
 
-	void EmptySelection();
+	NIAGARAEDITOR_API void EmptySelection();
 
-	void RemoveEntriesFromSelection(const TArray<UNiagaraStackEntry*>& InEntriesToRemove);
+	NIAGARAEDITOR_API void RemoveEntriesFromSelection(const TArray<UNiagaraStackEntry*>& InEntriesToRemove);
 
-	void RemoveEntryFromSelectionByDisplayedObject(const UObject* InObject);
+	NIAGARAEDITOR_API void RemoveEntryFromSelectionByDisplayedObject(const UObject* InObject);
 
-	void AddEntriesToSelectionByDisplayedObjectsDeferred(const TArray<const UObject*>& InObjects);
+	NIAGARAEDITOR_API void AddEntriesToSelectionByDisplayedObjectsDeferred(const TArray<const UObject*>& InObjects);
 
-	void AddEntryToSelectionByDisplayedObjectDeferred(const UObject* InObjects);
+	NIAGARAEDITOR_API void AddEntryToSelectionByDisplayedObjectDeferred(const UObject* InObjects);
 
-	void AddEntriesToSelectionByDisplayedObjectKeysDeferred(const TArray<FObjectKey>& InObjectKeys);
+	NIAGARAEDITOR_API void AddEntriesToSelectionByDisplayedObjectKeysDeferred(const TArray<FObjectKey>& InObjectKeys);
 
-	void AddEntryToSelectionByDisplayedObjectKeyDeferred(const FObjectKey& InObjectKey);
+	NIAGARAEDITOR_API void AddEntryToSelectionByDisplayedObjectKeyDeferred(const FObjectKey& InObjectKey);
 
-	void AddEntryToSelectionBySelectionIdDeferred(const FGuid& InSelectionId);
+	NIAGARAEDITOR_API void AddEntryToSelectionBySelectionIdDeferred(const FGuid& InSelectionId);
 
-	bool Refresh();
+	NIAGARAEDITOR_API bool Refresh();
 
-	void RefreshDeferred();
+	NIAGARAEDITOR_API void RefreshDeferred();
 
-	FOnSelectionChanged& OnEntrySelectionChanged();
+	NIAGARAEDITOR_API FOnSelectionChanged& OnEntrySelectionChanged();
 
-	FOnSelectionChanged& OnEmitterHandleIdSelectionChanged();
+	NIAGARAEDITOR_API FOnSelectionChanged& OnEmitterHandleIdSelectionChanged();
 
-	FOnSelectionChanged& OnSystemIsSelectedChanged();
+	NIAGARAEDITOR_API FOnSelectionChanged& OnSystemIsSelectedChanged();
 
-	void Tick();
+	NIAGARAEDITOR_API void Tick();
 
 private:
 	TSharedRef<FNiagaraSystemViewModel> GetSystemViewModel();
@@ -107,4 +107,4 @@ private:
 	TArray<FGuid> DeferredSelectionIdsToAddToSelection;
 
 	bool bRefreshIsPending;
-};
\ No newline at end of file
+};
diff --git a/Engine/Plugins/FX/Niagara/Source/NiagaraEditor/Public/ViewModels/NiagaraSystemViewModel.h b/Engine/Plugins/FX/Niagara/Source/NiagaraEditor/Public/ViewModels/NiagaraSystemViewModel.h
--- a/Engine/Plugins/FX/Niagara/Source/NiagaraEditor/Public/ViewModels/NiagaraSystemViewModel.h
+++ b/Engine/Plugins/FX/Niagara/Source/NiagaraEditor/Public/ViewModels/NiagaraSystemViewModel.h
@@ -74,9 +74,9 @@ enum class ENiagaraGetGraphParameterReferencesMode
 DECLARE_DELEGATE_RetVal(FName, FOnGetWorkflowMode);
 
 /** Defines options for the niagara System view model */
-struct NIAGARAEDITOR_API FNiagaraSystemViewModelOptions
+struct FNiagaraSystemViewModelOptions
 {
-	FNiagaraSystemViewModelOptions();
+	NIAGARAEDITOR_API FNiagaraSystemViewModelOptions();
 
 	/** Whether or not the user can edit emitters from the timeline. */
 	bool bCanModifyEmittersFromTimeline = true;
diff --git a/Engine/Plugins/FX/Niagara/Source/NiagaraEditor/Public/ViewModels/Stack/NiagaraParameterHandle.h b/Engine/Plugins/FX/Niagara/Source/NiagaraEditor/Public/ViewModels/Stack/NiagaraParameterHandle.h
--- a/Engine/Plugins/FX/Niagara/Source/NiagaraEditor/Public/ViewModels/Stack/NiagaraParameterHandle.h
+++ b/Engine/Plugins/FX/Niagara/Source/NiagaraEditor/Public/ViewModels/Stack/NiagaraParameterHandle.h
@@ -8,67 +8,67 @@
 class UNiagaraNodeFunctionCall;
 struct FNiagaraVariable;
 
-class NIAGARAEDITOR_API FNiagaraParameterHandle
+class FNiagaraParameterHandle
 {
 public:
-	FNiagaraParameterHandle();
+	NIAGARAEDITOR_API FNiagaraParameterHandle();
 
-	FNiagaraParameterHandle(FName InParameterHandleString);
+	NIAGARAEDITOR_API FNiagaraParameterHandle(FName InParameterHandleString);
 
-	FNiagaraParameterHandle(FName InNamespace, FName InName);
+	NIAGARAEDITOR_API FNiagaraParameterHandle(FName InNamespace, FName InName);
 
-	bool operator==(const FNiagaraParameterHandle& Other) const;
+	NIAGARAEDITOR_API bool operator==(const FNiagaraParameterHandle& Other) const;
 
 	bool operator!=(const FNiagaraParameterHandle& Other) const { return (*this == Other) == false; }
 
-	static FNiagaraParameterHandle CreateAliasedModuleParameterHandle(const FNiagaraParameterHandle& ModuleParameterHandle, const UNiagaraNodeFunctionCall* ModuleNode);
-	static FNiagaraParameterHandle CreateAliasedModuleParameterHandle(const FName FullInputName, const FName FunctionName);
+	static NIAGARAEDITOR_API FNiagaraParameterHandle CreateAliasedModuleParameterHandle(const FNiagaraParameterHandle& ModuleParameterHandle, const UNiagaraNodeFunctionCall* ModuleNode);
+	static NIAGARAEDITOR_API FNiagaraParameterHandle CreateAliasedModuleParameterHandle(const FName FullInputName, const FName FunctionName);
 
-	static FNiagaraParameterHandle CreateEngineParameterHandle(const FNiagaraVariable& SystemVariable);
+	static NIAGARAEDITOR_API FNiagaraParameterHandle CreateEngineParameterHandle(const FNiagaraVariable& SystemVariable);
 
-	static FNiagaraParameterHandle CreateEmitterParameterHandle(const FNiagaraVariable& EmitterVariable);
+	static NIAGARAEDITOR_API FNiagaraParameterHandle CreateEmitterParameterHandle(const FNiagaraVariable& EmitterVariable);
 
-	static FNiagaraParameterHandle CreateParticleAttributeParameterHandle(const FName InName);
+	static NIAGARAEDITOR_API FNiagaraParameterHandle CreateParticleAttributeParameterHandle(const FName InName);
 
-	static FNiagaraParameterHandle CreateModuleParameterHandle(const FName InName);
+	static NIAGARAEDITOR_API FNiagaraParameterHandle CreateModuleParameterHandle(const FName InName);
 
-	static FNiagaraParameterHandle CreateInitialParameterHandle(const FNiagaraParameterHandle& Handle);
+	static NIAGARAEDITOR_API FNiagaraParameterHandle CreateInitialParameterHandle(const FNiagaraParameterHandle& Handle);
 
-	bool IsValid() const;
+	NIAGARAEDITOR_API bool IsValid() const;
 
-	const FName GetParameterHandleString() const;
+	NIAGARAEDITOR_API const FName GetParameterHandleString() const;
 
-	const FName GetName() const;
+	NIAGARAEDITOR_API const FName GetName() const;
 
-	const FName GetNamespace() const;
+	NIAGARAEDITOR_API const FName GetNamespace() const;
 
-	const TArray<FName> GetHandleParts() const;
+	NIAGARAEDITOR_API const TArray<FName> GetHandleParts() const;
 
-	bool IsUserHandle() const;
+	NIAGARAEDITOR_API bool IsUserHandle() const;
 
-	bool IsEngineHandle() const;
+	NIAGARAEDITOR_API bool IsEngineHandle() const;
 
-	bool IsSystemHandle() const;
+	NIAGARAEDITOR_API bool IsSystemHandle() const;
 
-	bool IsEmitterHandle() const;
+	NIAGARAEDITOR_API bool IsEmitterHandle() const;
 
-	bool IsParticleAttributeHandle() const;
+	NIAGARAEDITOR_API bool IsParticleAttributeHandle() const;
 
-	bool IsModuleHandle() const;
+	NIAGARAEDITOR_API bool IsModuleHandle() const;
 
-	bool IsOutputHandle() const;
+	NIAGARAEDITOR_API bool IsOutputHandle() const;
 
-	bool IsLocalHandle() const;
+	NIAGARAEDITOR_API bool IsLocalHandle() const;
 
-	bool IsParameterCollectionHandle() const;
+	NIAGARAEDITOR_API bool IsParameterCollectionHandle() const;
 
-	bool IsReadOnlyHandle() const;
+	NIAGARAEDITOR_API bool IsReadOnlyHandle() const;
 
-	bool IsTransientHandle() const;
+	NIAGARAEDITOR_API bool IsTransientHandle() const;
 
-	bool IsDataInstanceHandle() const;
+	NIAGARAEDITOR_API bool IsDataInstanceHandle() const;
 
-	bool IsStackContextHandle() const;
+	NIAGARAEDITOR_API bool IsStackContextHandle() const;
 
 private:
 	FName ParameterHandleName;
diff --git a/Engine/Plugins/FX/Niagara/Source/NiagaraEditor/Public/ViewModels/Stack/NiagaraStackCommentCollection.h b/Engine/Plugins/FX/Niagara/Source/NiagaraEditor/Public/ViewModels/Stack/NiagaraStackCommentCollection.h
--- a/Engine/Plugins/FX/Niagara/Source/NiagaraEditor/Public/ViewModels/Stack/NiagaraStackCommentCollection.h
+++ b/Engine/Plugins/FX/Niagara/Source/NiagaraEditor/Public/ViewModels/Stack/NiagaraStackCommentCollection.h
@@ -7,21 +7,21 @@
 #include "ViewModels/Stack/NiagaraStackEntry.h"
 #include "NiagaraStackCommentCollection.generated.h"
 
-UCLASS()
-class NIAGARAEDITOR_API UNiagaraStackCommentCollection : public UNiagaraStackEntry
+UCLASS(MinimalAPI)
+class UNiagaraStackCommentCollection : public UNiagaraStackEntry
 {
 	GENERATED_BODY()
 
 public:
 	UNiagaraStackCommentCollection() {}
 
-	void Initialize(FRequiredEntryData InRequiredEntryData);
+	NIAGARAEDITOR_API void Initialize(FRequiredEntryData InRequiredEntryData);
 
 	virtual bool GetShouldShowInOverview() const override { return false; }
 	virtual bool GetShouldShowInStack() const override { return false; }
 
-	UNiagaraStackObject* FindStackObjectForCommentNode(UEdGraphNode_Comment* CommentNode) const;
+	NIAGARAEDITOR_API UNiagaraStackObject* FindStackObjectForCommentNode(UEdGraphNode_Comment* CommentNode) const;
 	
 protected:
-	virtual void RefreshChildrenInternal(const TArray<UNiagaraStackEntry*>& CurrentChildren, TArray<UNiagaraStackEntry*>& NewChildren, TArray<FStackIssue>& NewIssues) override;
-};
\ No newline at end of file
+	NIAGARAEDITOR_API virtual void RefreshChildrenInternal(const TArray<UNiagaraStackEntry*>& CurrentChildren, TArray<UNiagaraStackEntry*>& NewChildren, TArray<FStackIssue>& NewIssues) override;
+};
diff --git a/Engine/Plugins/FX/Niagara/Source/NiagaraEditor/Public/ViewModels/Stack/NiagaraStackEmitterPropertiesGroup.h b/Engine/Plugins/FX/Niagara/Source/NiagaraEditor/Public/ViewModels/Stack/NiagaraStackEmitterPropertiesGroup.h
--- a/Engine/Plugins/FX/Niagara/Source/NiagaraEditor/Public/ViewModels/Stack/NiagaraStackEmitterPropertiesGroup.h
+++ b/Engine/Plugins/FX/Niagara/Source/NiagaraEditor/Public/ViewModels/Stack/NiagaraStackEmitterPropertiesGroup.h
@@ -12,32 +12,32 @@ class UNiagaraSimulationStageBase;
 class UNiagaraStackEmitterPropertiesItem;
 class UNiagaraStackObject;
 
-UCLASS()
-class NIAGARAEDITOR_API UNiagaraStackEmitterPropertiesGroup : public UNiagaraStackItemGroup
+UCLASS(MinimalAPI)
+class UNiagaraStackEmitterPropertiesGroup : public UNiagaraStackItemGroup
 {
 	GENERATED_BODY()
 
 public:
-	UNiagaraStackEmitterPropertiesGroup();
+	NIAGARAEDITOR_API UNiagaraStackEmitterPropertiesGroup();
 
-	void Initialize(FRequiredEntryData InRequiredEntryData);
+	NIAGARAEDITOR_API void Initialize(FRequiredEntryData InRequiredEntryData);
 
 	virtual EIconMode GetSupportedIconMode() const override { return EIconMode::Brush; }
-	virtual const FSlateBrush* GetIconBrush() const override;
+	NIAGARAEDITOR_API virtual const FSlateBrush* GetIconBrush() const override;
 	virtual bool SupportsSecondaryIcon() const override { return true; }
-	virtual const FSlateBrush* GetSecondaryIconBrush() const override;
+	NIAGARAEDITOR_API virtual const FSlateBrush* GetSecondaryIconBrush() const override;
 	virtual bool GetCanExpandInOverview() const override { return false; }
 	virtual bool GetShouldShowInStack() const override { return false; }
 
 protected:
-	void RefreshChildrenInternal(const TArray<UNiagaraStackEntry*>& CurrentChildren, TArray<UNiagaraStackEntry*>& NewChildren, TArray<FStackIssue>& NewIssues) override;
+	NIAGARAEDITOR_API void RefreshChildrenInternal(const TArray<UNiagaraStackEntry*>& CurrentChildren, TArray<UNiagaraStackEntry*>& NewChildren, TArray<FStackIssue>& NewIssues) override;
 
 private:
-	void ItemAddedFromUtilities(FGuid AddedEventHandlerId, UNiagaraSimulationStageBase* AddedSimulationStage);
+	NIAGARAEDITOR_API void ItemAddedFromUtilities(FGuid AddedEventHandlerId, UNiagaraSimulationStageBase* AddedSimulationStage);
 
 private:
 	UPROPERTY()
 	TObjectPtr<UNiagaraStackEmitterPropertiesItem> PropertiesItem;
 
 	TSharedPtr<INiagaraStackItemGroupAddUtilities> AddUtilities;
-};
\ No newline at end of file
+};
diff --git a/Engine/Plugins/FX/Niagara/Source/NiagaraEditor/Public/ViewModels/Stack/NiagaraStackEmitterSettingsGroup.h b/Engine/Plugins/FX/Niagara/Source/NiagaraEditor/Public/ViewModels/Stack/NiagaraStackEmitterSettingsGroup.h
--- a/Engine/Plugins/FX/Niagara/Source/NiagaraEditor/Public/ViewModels/Stack/NiagaraStackEmitterSettingsGroup.h
+++ b/Engine/Plugins/FX/Niagara/Source/NiagaraEditor/Public/ViewModels/Stack/NiagaraStackEmitterSettingsGroup.h
@@ -13,37 +13,37 @@ class FNiagaraScriptViewModel;
 class UNiagaraStackObject;
 class IDetailTreeNode;
 
-UCLASS()
-class NIAGARAEDITOR_API UNiagaraStackEmitterPropertiesItem : public UNiagaraStackItem
+UCLASS(MinimalAPI)
+class UNiagaraStackEmitterPropertiesItem : public UNiagaraStackItem
 {
 	GENERATED_BODY()
 
 public:
-	void Initialize(FRequiredEntryData InRequiredEntryData);
+	NIAGARAEDITOR_API void Initialize(FRequiredEntryData InRequiredEntryData);
 
-	virtual FText GetDisplayName() const override;
+	NIAGARAEDITOR_API virtual FText GetDisplayName() const override;
 
-	virtual FText GetTooltipText() const override;
+	NIAGARAEDITOR_API virtual FText GetTooltipText() const override;
 
 	virtual bool SupportsResetToBase() const override { return true; }
-	virtual bool TestCanResetToBaseWithMessage(FText& OutCanResetToBaseMessage) const override;
-	virtual void ResetToBase() override;
+	NIAGARAEDITOR_API virtual bool TestCanResetToBaseWithMessage(FText& OutCanResetToBaseMessage) const override;
+	NIAGARAEDITOR_API virtual void ResetToBase() override;
 
 	virtual EIconMode GetSupportedIconMode() const { return EIconMode::Brush; }
-	virtual const FSlateBrush* GetIconBrush() const override;
+	NIAGARAEDITOR_API virtual const FSlateBrush* GetIconBrush() const override;
 
 	virtual bool GetShouldShowInOverview() const override { return false; }
 
 protected:
-	virtual void FinalizeInternal() override;
+	NIAGARAEDITOR_API virtual void FinalizeInternal() override;
 
-	virtual void RefreshChildrenInternal(const TArray<UNiagaraStackEntry*>& CurrentChildren, TArray<UNiagaraStackEntry*>& NewChildren, TArray<FStackIssue>& NewIssues) override;
+	NIAGARAEDITOR_API virtual void RefreshChildrenInternal(const TArray<UNiagaraStackEntry*>& CurrentChildren, TArray<UNiagaraStackEntry*>& NewChildren, TArray<FStackIssue>& NewIssues) override;
 
-	virtual void RefreshIssues(TArray<FStackIssue>& NewIssues);
+	NIAGARAEDITOR_API virtual void RefreshIssues(TArray<FStackIssue>& NewIssues);
 
 private:
-	void EmitterPropertiesChanged();
-	FStackIssueFixDelegate GetUpgradeVersionFix();
+	NIAGARAEDITOR_API void EmitterPropertiesChanged();
+	NIAGARAEDITOR_API FStackIssueFixDelegate GetUpgradeVersionFix();
 
 	mutable TOptional<bool> bCanResetToBaseCache;
 
@@ -53,34 +53,34 @@ private:
 	TObjectPtr<UNiagaraStackObject> EmitterObject;
 };
 
-UCLASS()
-class NIAGARAEDITOR_API UNiagaraStackEmitterSummaryItem : public UNiagaraStackItem
+UCLASS(MinimalAPI)
+class UNiagaraStackEmitterSummaryItem : public UNiagaraStackItem
 {
 	GENERATED_BODY()
 
 public:
-	void Initialize(FRequiredEntryData InRequiredEntryData);
+	NIAGARAEDITOR_API void Initialize(FRequiredEntryData InRequiredEntryData);
 
-	virtual FText GetDisplayName() const override;
-	virtual FText GetTooltipText() const override;
+	NIAGARAEDITOR_API virtual FText GetDisplayName() const override;
+	NIAGARAEDITOR_API virtual FText GetTooltipText() const override;
 
 	virtual bool SupportsResetToBase() const override { return false; }
 	virtual EIconMode GetSupportedIconMode() const override { return EIconMode::Text; }
-	virtual FText GetIconText() const override;
+	NIAGARAEDITOR_API virtual FText GetIconText() const override;
 	virtual bool GetShouldShowInOverview() const override { return false; }
 
 	virtual bool SupportsEditMode() const override { return true; }
-	virtual void OnEditButtonClicked() override;
-	virtual TOptional<FText> GetEditModeButtonText() const override;
-	virtual TOptional<FText> GetEditModeButtonTooltip() const override;
-	virtual EVisibility IsEditButtonVisible() const override;
+	NIAGARAEDITOR_API virtual void OnEditButtonClicked() override;
+	NIAGARAEDITOR_API virtual TOptional<FText> GetEditModeButtonText() const override;
+	NIAGARAEDITOR_API virtual TOptional<FText> GetEditModeButtonTooltip() const override;
+	NIAGARAEDITOR_API virtual EVisibility IsEditButtonVisible() const override;
 
 protected:
 
-	virtual void RefreshChildrenInternal(const TArray<UNiagaraStackEntry*>& CurrentChildren, TArray<UNiagaraStackEntry*>& NewChildren, TArray<FStackIssue>& NewIssues) override;
+	NIAGARAEDITOR_API virtual void RefreshChildrenInternal(const TArray<UNiagaraStackEntry*>& CurrentChildren, TArray<UNiagaraStackEntry*>& NewChildren, TArray<FStackIssue>& NewIssues) override;
 
 private:
-	void SelectSummaryNodesFromEmitterEditorDataRootNodes(TArray<TSharedRef<IDetailTreeNode>> Source, TArray<TSharedRef<IDetailTreeNode>>* Selected);
+	NIAGARAEDITOR_API void SelectSummaryNodesFromEmitterEditorDataRootNodes(TArray<TSharedRef<IDetailTreeNode>> Source, TArray<TSharedRef<IDetailTreeNode>>* Selected);
 
 private:
 	FVersionedNiagaraEmitterWeakPtr Emitter;
@@ -89,21 +89,21 @@ private:
 	TObjectPtr<UNiagaraStackSummaryViewCollection> FilteredObject;
 };
 
-UCLASS()
-class NIAGARAEDITOR_API UNiagaraStackEmitterSummaryGroup : public UNiagaraStackItemGroup
+UCLASS(MinimalAPI)
+class UNiagaraStackEmitterSummaryGroup : public UNiagaraStackItemGroup
 {
 	GENERATED_BODY()
 
 public:
-	UNiagaraStackEmitterSummaryGroup();
+	NIAGARAEDITOR_API UNiagaraStackEmitterSummaryGroup();
 
 	virtual EIconMode GetSupportedIconMode() const override { return EIconMode::Text; }
-	virtual  FText GetIconText() const override;
+	NIAGARAEDITOR_API virtual  FText GetIconText() const override;
 	virtual bool GetCanExpandInOverview() const override { return false; }
 	virtual bool GetShouldShowInStack() const override { return false; }
 
 protected:
-	void RefreshChildrenInternal(const TArray<UNiagaraStackEntry*>& CurrentChildren, TArray<UNiagaraStackEntry*>& NewChildren, TArray<FStackIssue>& NewIssues) override;
+	NIAGARAEDITOR_API void RefreshChildrenInternal(const TArray<UNiagaraStackEntry*>& CurrentChildren, TArray<UNiagaraStackEntry*>& NewChildren, TArray<FStackIssue>& NewIssues) override;
 
 
 private:
@@ -112,25 +112,25 @@ private:
 };
 
 
-UCLASS()
-class NIAGARAEDITOR_API UNiagaraStackSummaryViewCollapseButton : public UNiagaraStackEntry
+UCLASS(MinimalAPI)
+class UNiagaraStackSummaryViewCollapseButton : public UNiagaraStackEntry
 {
 	GENERATED_BODY()
 
 public:
-	void Initialize(FRequiredEntryData InRequiredEntryData);
+	NIAGARAEDITOR_API void Initialize(FRequiredEntryData InRequiredEntryData);
 
 	//~ UNiagaraStackEntry interface
-	virtual FText GetDisplayName() const override;
-	virtual EStackRowStyle GetStackRowStyle() const override;
-	virtual FText GetTooltipText() const override;
+	NIAGARAEDITOR_API virtual FText GetDisplayName() const override;
+	NIAGARAEDITOR_API virtual EStackRowStyle GetStackRowStyle() const override;
+	NIAGARAEDITOR_API virtual FText GetTooltipText() const override;
 
 	virtual bool GetCanExpandInOverview() const override { return true; }
-	virtual bool GetIsEnabled() const override;
+	NIAGARAEDITOR_API virtual bool GetIsEnabled() const override;
 	virtual void SetIsEnabled(bool bEnabled) {}
 	virtual bool SupportsChangeEnabled() const { return false; }
 	
 protected:
-	virtual void RefreshChildrenInternal(const TArray<UNiagaraStackEntry*>& CurrentChildren, TArray<UNiagaraStackEntry*>& NewChildren, TArray<FStackIssue>& NewIssues) override;
+	NIAGARAEDITOR_API virtual void RefreshChildrenInternal(const TArray<UNiagaraStackEntry*>& CurrentChildren, TArray<UNiagaraStackEntry*>& NewChildren, TArray<FStackIssue>& NewIssues) override;
 
-};
\ No newline at end of file
+};
diff --git a/Engine/Plugins/FX/Niagara/Source/NiagaraEditor/Public/ViewModels/Stack/NiagaraStackEntry.h b/Engine/Plugins/FX/Niagara/Source/NiagaraEditor/Public/ViewModels/Stack/NiagaraStackEntry.h
--- a/Engine/Plugins/FX/Niagara/Source/NiagaraEditor/Public/ViewModels/Stack/NiagaraStackEntry.h
+++ b/Engine/Plugins/FX/Niagara/Source/NiagaraEditor/Public/ViewModels/Stack/NiagaraStackEntry.h
@@ -46,8 +46,8 @@ private:
 	TArray<UNiagaraStackEntry*> DraggedEntries;
 };
 
-UCLASS()
-class NIAGARAEDITOR_API UNiagaraStackEntry : public UObject
+UCLASS(MinimalAPI)
+class UNiagaraStackEntry : public UObject
 {
 	GENERATED_BODY()
 
@@ -110,22 +110,22 @@ public:
 	DECLARE_MULTICAST_DELEGATE(FOnAlternateDisplayNameChanged);
 
 public:
-	struct NIAGARAEDITOR_API FExecutionCategoryNames
+	struct FExecutionCategoryNames
 	{
-		static const FName System;
-		static const FName Emitter;
-		static const FName Particle;
-		static const FName Render;
+		static NIAGARAEDITOR_API const FName System;
+		static NIAGARAEDITOR_API const FName Emitter;
+		static NIAGARAEDITOR_API const FName Particle;
+		static NIAGARAEDITOR_API const FName Render;
 	};
 
-	struct NIAGARAEDITOR_API FExecutionSubcategoryNames
+	struct FExecutionSubcategoryNames
 	{
-		static const FName Settings;
-		static const FName Spawn;
-		static const FName Update;
-		static const FName Event;
-		static const FName SimulationStage;
-		static const FName Render;
+		static NIAGARAEDITOR_API const FName Settings;
+		static NIAGARAEDITOR_API const FName Spawn;
+		static NIAGARAEDITOR_API const FName Update;
+		static NIAGARAEDITOR_API const FName Event;
+		static NIAGARAEDITOR_API const FName SimulationStage;
+		static NIAGARAEDITOR_API const FName Render;
 	};
 
 	enum class EStackRowStyle
@@ -181,23 +181,23 @@ public:
 		Link
 	};
 
-	struct NIAGARAEDITOR_API FStackIssueFix
+	struct FStackIssueFix
 	{
-		FStackIssueFix();
+		NIAGARAEDITOR_API FStackIssueFix();
 
-		FStackIssueFix(FText InDescription, FStackIssueFixDelegate InFixDelegate, EStackIssueFixStyle FixStyle = EStackIssueFixStyle::Fix);
+		NIAGARAEDITOR_API FStackIssueFix(FText InDescription, FStackIssueFixDelegate InFixDelegate, EStackIssueFixStyle FixStyle = EStackIssueFixStyle::Fix);
 
-		bool IsValid() const;
+		NIAGARAEDITOR_API bool IsValid() const;
 
-		const FText& GetDescription() const;
+		NIAGARAEDITOR_API const FText& GetDescription() const;
 
-		void SetFixDelegate(const FStackIssueFixDelegate& InFixDelegate);
+		NIAGARAEDITOR_API void SetFixDelegate(const FStackIssueFixDelegate& InFixDelegate);
 
-		const FStackIssueFixDelegate& GetFixDelegate() const;
+		NIAGARAEDITOR_API const FStackIssueFixDelegate& GetFixDelegate() const;
 
-		EStackIssueFixStyle GetStyle() const;
+		NIAGARAEDITOR_API EStackIssueFixStyle GetStyle() const;
 
-		const FString& GetUniqueIdentifier() const;
+		NIAGARAEDITOR_API const FString& GetUniqueIdentifier() const;
 
 	private:
 		FText Description;
@@ -206,37 +206,37 @@ public:
 		FString UniqueIdentifier;
 	};
 
-	struct NIAGARAEDITOR_API FStackIssue
+	struct FStackIssue
 	{
-		FStackIssue();
+		NIAGARAEDITOR_API FStackIssue();
 
-		FStackIssue(EStackIssueSeverity InSeverity, FText InShortDescription, FText InLongDescription, FString InStackEditorDataKey, bool bInCanBeDismissed, const TArray<FStackIssueFix>& InFixes, const FSimpleDelegate& InDismissHandler = FSimpleDelegate());
+		NIAGARAEDITOR_API FStackIssue(EStackIssueSeverity InSeverity, FText InShortDescription, FText InLongDescription, FString InStackEditorDataKey, bool bInCanBeDismissed, const TArray<FStackIssueFix>& InFixes, const FSimpleDelegate& InDismissHandler = FSimpleDelegate());
 
-		FStackIssue(EStackIssueSeverity InSeverity, FText InShortDescription, FText InLongDescription, FString InStackEditorDataKey, bool bInCanBeDismissed, FStackIssueFix InFix);
+		NIAGARAEDITOR_API FStackIssue(EStackIssueSeverity InSeverity, FText InShortDescription, FText InLongDescription, FString InStackEditorDataKey, bool bInCanBeDismissed, FStackIssueFix InFix);
 
-		FStackIssue(EStackIssueSeverity InSeverity, FText InShortDescription, FText InLongDescription, FString InStackEditorDataKey, bool bInCanBeDismissed);
+		NIAGARAEDITOR_API FStackIssue(EStackIssueSeverity InSeverity, FText InShortDescription, FText InLongDescription, FString InStackEditorDataKey, bool bInCanBeDismissed);
 
-		bool IsValid();
+		NIAGARAEDITOR_API bool IsValid();
 
-		EStackIssueSeverity GetSeverity() const;
+		NIAGARAEDITOR_API EStackIssueSeverity GetSeverity() const;
 
-		const FText& GetShortDescription() const;
+		NIAGARAEDITOR_API const FText& GetShortDescription() const;
 
-		const FText& GetLongDescription() const;
+		NIAGARAEDITOR_API const FText& GetLongDescription() const;
 
-		const FString& GetUniqueIdentifier() const;
+		NIAGARAEDITOR_API const FString& GetUniqueIdentifier() const;
 
-		bool GetCanBeDismissed() const;
+		NIAGARAEDITOR_API bool GetCanBeDismissed() const;
 
-		const FSimpleDelegate& GetDismissHandler() const;
+		NIAGARAEDITOR_API const FSimpleDelegate& GetDismissHandler() const;
 
-		const TArray<FStackIssueFix>& GetFixes() const;
+		NIAGARAEDITOR_API const TArray<FStackIssueFix>& GetFixes() const;
 
-		bool GetIsExpandedByDefault() const;
+		NIAGARAEDITOR_API bool GetIsExpandedByDefault() const;
 
-		void SetIsExpandedByDefault(bool InExpanded);
+		NIAGARAEDITOR_API void SetIsExpandedByDefault(bool InExpanded);
 		
-		void InsertFix(int32 InsertionIdx, const FStackIssueFix& Fix);
+		NIAGARAEDITOR_API void InsertFix(int32 InsertionIdx, const FStackIssueFix& Fix);
 
 	private:
 		EStackIssueSeverity Severity;
@@ -250,70 +250,70 @@ public:
 	};
 
 public:
-	UNiagaraStackEntry();
+	NIAGARAEDITOR_API UNiagaraStackEntry();
 
-	void Initialize(FRequiredEntryData InRequiredEntryData, FString InStackEditorDataKey);
+	NIAGARAEDITOR_API void Initialize(FRequiredEntryData InRequiredEntryData, FString InStackEditorDataKey);
 
-	void Finalize();
+	NIAGARAEDITOR_API void Finalize();
 
-	bool IsFinalized() const;
+	NIAGARAEDITOR_API bool IsFinalized() const;
 
-	virtual FText GetDisplayName() const;
+	NIAGARAEDITOR_API virtual FText GetDisplayName() const;
 
-	TOptional<FText> GetAlternateDisplayName() const;
+	NIAGARAEDITOR_API TOptional<FText> GetAlternateDisplayName() const;
 
-	virtual UObject* GetDisplayedObject() const;
+	NIAGARAEDITOR_API virtual UObject* GetDisplayedObject() const;
 
-	virtual FGuid GetSelectionId() const;
+	NIAGARAEDITOR_API virtual FGuid GetSelectionId() const;
 
-	UNiagaraStackEditorData& GetStackEditorData() const;
+	NIAGARAEDITOR_API UNiagaraStackEditorData& GetStackEditorData() const;
 
-	FString GetStackEditorDataKey() const;
+	NIAGARAEDITOR_API FString GetStackEditorDataKey() const;
 
-	virtual FText GetTooltipText() const;
+	NIAGARAEDITOR_API virtual FText GetTooltipText() const;
 
-	virtual bool GetCanExpand() const;
+	NIAGARAEDITOR_API virtual bool GetCanExpand() const;
 
-	virtual bool GetCanExpandInOverview() const;
+	NIAGARAEDITOR_API virtual bool GetCanExpandInOverview() const;
 
-	virtual bool IsExpandedByDefault() const;
+	NIAGARAEDITOR_API virtual bool IsExpandedByDefault() const;
 
-	bool GetIsExpanded() const;
+	NIAGARAEDITOR_API bool GetIsExpanded() const;
 
-	void SetIsExpanded(bool bInExpanded);
+	NIAGARAEDITOR_API void SetIsExpanded(bool bInExpanded);
 
-	void SetIsExpanded_Recursive(bool bInExpanded);
+	NIAGARAEDITOR_API void SetIsExpanded_Recursive(bool bInExpanded);
 
-	bool GetIsExpandedInOverview() const;
+	NIAGARAEDITOR_API bool GetIsExpandedInOverview() const;
 
-	void SetIsExpandedInOverview(bool bInExpanded);
+	NIAGARAEDITOR_API void SetIsExpandedInOverview(bool bInExpanded);
 
-	virtual bool GetIsEnabled() const;
+	NIAGARAEDITOR_API virtual bool GetIsEnabled() const;
 
-	bool GetOwnerIsEnabled() const;
+	NIAGARAEDITOR_API bool GetOwnerIsEnabled() const;
 
 	bool GetIsEnabledAndOwnerIsEnabled() const { return GetIsEnabled() && GetOwnerIsEnabled(); }
 
-	FName GetExecutionCategoryName() const;
+	NIAGARAEDITOR_API FName GetExecutionCategoryName() const;
 
-	FName GetExecutionSubcategoryName() const;
+	NIAGARAEDITOR_API FName GetExecutionSubcategoryName() const;
 
-	virtual EStackRowStyle GetStackRowStyle() const;
+	NIAGARAEDITOR_API virtual EStackRowStyle GetStackRowStyle() const;
 
-	int32 GetIndentLevel() const;
+	NIAGARAEDITOR_API int32 GetIndentLevel() const;
 
 	/** Returns whether or not this entry should be treated as a child of a previous sibling for layout purposes. */
 	virtual bool IsSemanticChild() const { return false; }
 
-	virtual bool GetShouldShowInStack() const;
+	NIAGARAEDITOR_API virtual bool GetShouldShowInStack() const;
 
-	virtual bool GetShouldShowInOverview() const;
+	NIAGARAEDITOR_API virtual bool GetShouldShowInOverview() const;
 
-	void GetFilteredChildren(TArray<UNiagaraStackEntry*>& OutFilteredChildren) const;
+	NIAGARAEDITOR_API void GetFilteredChildren(TArray<UNiagaraStackEntry*>& OutFilteredChildren) const;
 	
-	void GetUnfilteredChildren(TArray<UNiagaraStackEntry*>& OutUnfilteredChildren) const;
+	NIAGARAEDITOR_API void GetUnfilteredChildren(TArray<UNiagaraStackEntry*>& OutUnfilteredChildren) const;
 
-	void GetCustomFilteredChildren(TArray<UNiagaraStackEntry*>& OutFilteredChildren, const TArray<FOnFilterChild>& ChildFilters) const;
+	NIAGARAEDITOR_API void GetCustomFilteredChildren(TArray<UNiagaraStackEntry*>& OutFilteredChildren, const TArray<FOnFilterChild>& ChildFilters) const;
 
 	void GetFilteredChildrenOfTypes(TArray<UNiagaraStackEntry*>& OutFilteredChildren, const TSet<UClass*>& AllowedClasses) const
 	{
@@ -393,35 +393,35 @@ public:
 		}
 	}
 
-	FOnExpansionChanged& OnExpansionChanged();
+	NIAGARAEDITOR_API FOnExpansionChanged& OnExpansionChanged();
 
-	FOnExpansionChanged& OnExpansionInOverviewChanged();
+	NIAGARAEDITOR_API FOnExpansionChanged& OnExpansionInOverviewChanged();
 
-	FOnStructureChanged& OnStructureChanged();
+	NIAGARAEDITOR_API FOnStructureChanged& OnStructureChanged();
 
-	FOnDataObjectModified& OnDataObjectModified();
+	NIAGARAEDITOR_API FOnDataObjectModified& OnDataObjectModified();
 
-	FOnRequestFullRefresh& OnRequestFullRefresh();
+	NIAGARAEDITOR_API FOnRequestFullRefresh& OnRequestFullRefresh();
 
-	const FOnRequestFullRefresh& OnRequestFullRefreshDeferred() const;
+	NIAGARAEDITOR_API const FOnRequestFullRefresh& OnRequestFullRefreshDeferred() const;
 
-	FOnRequestFullRefresh& OnRequestFullRefreshDeferred();
+	NIAGARAEDITOR_API FOnRequestFullRefresh& OnRequestFullRefreshDeferred();
 
-	FOnAlternateDisplayNameChanged& OnAlternateDisplayNameChanged();
+	NIAGARAEDITOR_API FOnAlternateDisplayNameChanged& OnAlternateDisplayNameChanged();
 
 	/** Recursively refreshes the children for the current stack entry.  This may cause children to be added or removed and will automatically cause the filtered 
 	children to be refreshed. This will also cause the structure changed delegate to be broadcast. */
-	void RefreshChildren();
+	NIAGARAEDITOR_API void RefreshChildren();
 
 	/** Invalidates the cached filtered children so that the filters will be run the next time that GetFilteredChildren is called.  This should be called any time
 	a change to the data is made which will affect how children are filtered.  This will also cause the structure changed delegate to be broadcast. */
-	void RefreshFilteredChildren();
+	NIAGARAEDITOR_API void RefreshFilteredChildren();
 
-	FDelegateHandle AddChildFilter(FOnFilterChild ChildFilter);
-	void RemoveChildFilter(FDelegateHandle FilterHandle);
+	NIAGARAEDITOR_API FDelegateHandle AddChildFilter(FOnFilterChild ChildFilter);
+	NIAGARAEDITOR_API void RemoveChildFilter(FDelegateHandle FilterHandle);
 
-	TSharedRef<FNiagaraSystemViewModel> GetSystemViewModel() const;
-	TSharedPtr<FNiagaraEmitterViewModel> GetEmitterViewModel() const;
+	NIAGARAEDITOR_API TSharedRef<FNiagaraSystemViewModel> GetSystemViewModel() const;
+	NIAGARAEDITOR_API TSharedPtr<FNiagaraEmitterViewModel> GetEmitterViewModel() const;
 
 	template<typename ChildType, typename PredicateType>
 	static ChildType* FindCurrentChildOfTypeByPredicate(const TArray<UNiagaraStackEntry*>& CurrentChildren, PredicateType Predicate)
@@ -437,47 +437,47 @@ public:
 		return nullptr;
 	}
 
-	virtual void GetSearchItems(TArray<FStackSearchItem>& SearchItems) const;
+	NIAGARAEDITOR_API virtual void GetSearchItems(TArray<FStackSearchItem>& SearchItems) const;
 
-	virtual UObject* GetExternalAsset() const;
+	NIAGARAEDITOR_API virtual UObject* GetExternalAsset() const;
 
-	virtual bool CanDrag() const;
+	NIAGARAEDITOR_API virtual bool CanDrag() const;
 
-	TOptional<FDropRequestResponse> CanDrop(const FDropRequest& DropRequest);
+	NIAGARAEDITOR_API TOptional<FDropRequestResponse> CanDrop(const FDropRequest& DropRequest);
 
-	TOptional<FDropRequestResponse> Drop(const FDropRequest& DropRequest);
+	NIAGARAEDITOR_API TOptional<FDropRequestResponse> Drop(const FDropRequest& DropRequest);
 
-	void SetOnRequestCanDrop(FOnRequestDrop InOnRequestCanDrop);
+	NIAGARAEDITOR_API void SetOnRequestCanDrop(FOnRequestDrop InOnRequestCanDrop);
 
-	void SetOnRequestDrop(FOnRequestDrop InOnRequestCanDrop);
+	NIAGARAEDITOR_API void SetOnRequestDrop(FOnRequestDrop InOnRequestCanDrop);
 
-	const bool GetIsSearchResult() const;
+	NIAGARAEDITOR_API const bool GetIsSearchResult() const;
 
-	void SetIsSearchResult(bool bInIsSearchResult);
+	NIAGARAEDITOR_API void SetIsSearchResult(bool bInIsSearchResult);
 
-	bool HasBaseEmitter() const;
+	NIAGARAEDITOR_API bool HasBaseEmitter() const;
 
-	bool HasIssuesOrAnyChildHasIssues() const;
-	bool HasUsagesOrAnyChildHasUsages() const;
-	void GetRecursiveUsages(bool& bRead, bool& bWrite) const;
+	NIAGARAEDITOR_API bool HasIssuesOrAnyChildHasIssues() const;
+	NIAGARAEDITOR_API bool HasUsagesOrAnyChildHasUsages() const;
+	NIAGARAEDITOR_API void GetRecursiveUsages(bool& bRead, bool& bWrite) const;
 
-	int32 GetTotalNumberOfCustomNotes() const;
+	NIAGARAEDITOR_API int32 GetTotalNumberOfCustomNotes() const;
 	
-	int32 GetTotalNumberOfInfoIssues() const;
+	NIAGARAEDITOR_API int32 GetTotalNumberOfInfoIssues() const;
 
-	int32 GetTotalNumberOfWarningIssues() const;
+	NIAGARAEDITOR_API int32 GetTotalNumberOfWarningIssues() const;
 
-	int32 GetTotalNumberOfErrorIssues() const;
+	NIAGARAEDITOR_API int32 GetTotalNumberOfErrorIssues() const;
 
 	virtual EStackIssueSeverity GetIssueSeverity() const { return EStackIssueSeverity::None; }
 
-	const TArray<FStackIssue>& GetIssues() const;
+	NIAGARAEDITOR_API const TArray<FStackIssue>& GetIssues() const;
 
-	const TArray<UNiagaraStackEntry*>& GetAllChildrenWithIssues() const;
+	NIAGARAEDITOR_API const TArray<UNiagaraStackEntry*>& GetAllChildrenWithIssues() const;
 	
-	void AddValidationIssue(EStackIssueSeverity Severity, const FText& SummaryText, const FText& Description, bool bCanBeDismissed, const TArray<FNiagaraValidationFix>& Fixes, const TArray<FNiagaraValidationFix>& Links);
-	void AddExternalIssue(EStackIssueSeverity Severity, const FText& SummaryText, const FText& Description, bool bCanBeDismissed);
-	void ClearExternalIssues();
+	NIAGARAEDITOR_API void AddValidationIssue(EStackIssueSeverity Severity, const FText& SummaryText, const FText& Description, bool bCanBeDismissed, const TArray<FNiagaraValidationFix>& Fixes, const TArray<FNiagaraValidationFix>& Links);
+	NIAGARAEDITOR_API void AddExternalIssue(EStackIssueSeverity Severity, const FText& SummaryText, const FText& Description, bool bCanBeDismissed);
+	NIAGARAEDITOR_API void ClearExternalIssues();
 
 	virtual bool SupportsCut() const { return false; }
 
@@ -514,20 +514,20 @@ public:
 	virtual bool SupportsRename() const { return false; }
 
 	/** Gets whether this entry has a rename pending. */
-	virtual bool GetIsRenamePending() const;
+	NIAGARAEDITOR_API virtual bool GetIsRenamePending() const;
 
 	/** Sets whether this entry has a rename pending. */
-	virtual void SetIsRenamePending(bool bIsRenamePending);
+	NIAGARAEDITOR_API virtual void SetIsRenamePending(bool bIsRenamePending);
 
 	/** Handler for when a rename is committed for this stack entry. */
-	virtual void OnRenamed(FText NewName);
+	NIAGARAEDITOR_API virtual void OnRenamed(FText NewName);
 
 	virtual bool SupportsSummaryView() const { return false; }
-	virtual struct FNiagaraHierarchyIdentity DetermineSummaryIdentity() const;
-	bool IsInSummaryView() const;
-	bool IsAnyParentInSummaryView() const;
-	bool IsAnyChildInSummaryView(bool bRecursive = false) const;
-	bool ExistsInParentEmitterSummary() const;
+	NIAGARAEDITOR_API virtual struct FNiagaraHierarchyIdentity DetermineSummaryIdentity() const;
+	NIAGARAEDITOR_API bool IsInSummaryView() const;
+	NIAGARAEDITOR_API bool IsAnyParentInSummaryView() const;
+	NIAGARAEDITOR_API bool IsAnyChildInSummaryView(bool bRecursive = false) const;
+	NIAGARAEDITOR_API bool ExistsInParentEmitterSummary() const;
 
 	virtual EIconMode GetSupportedIconMode() const { return EIconMode::None; }
 
@@ -541,7 +541,7 @@ public:
 
 	virtual FText GetInheritanceMessage() const { return FText(); }
 
-	virtual void InvalidateCollectedUsage();
+	NIAGARAEDITOR_API virtual void InvalidateCollectedUsage();
 
 
 	struct FCollectedUsageData
@@ -550,30 +550,30 @@ public:
 		bool bHasReferencedParameterRead = false;
 		bool bHasReferencedParameterWrite = false;
 	};
-	virtual const FCollectedUsageData& GetCollectedUsageData() const;
+	NIAGARAEDITOR_API virtual const FCollectedUsageData& GetCollectedUsageData() const;
 
 protected:
-	virtual void BeginDestroy() override;
+	NIAGARAEDITOR_API virtual void BeginDestroy() override;
 
-	virtual void RefreshChildrenInternal(const TArray<UNiagaraStackEntry*>& CurrentChildren, TArray<UNiagaraStackEntry*>& NewChildren, TArray<FStackIssue>& NewIssues);
+	NIAGARAEDITOR_API virtual void RefreshChildrenInternal(const TArray<UNiagaraStackEntry*>& CurrentChildren, TArray<UNiagaraStackEntry*>& NewChildren, TArray<FStackIssue>& NewIssues);
 
-	virtual void PostRefreshChildrenInternal();
+	NIAGARAEDITOR_API virtual void PostRefreshChildrenInternal();
 
-	FRequiredEntryData CreateDefaultChildRequiredData() const;
+	NIAGARAEDITOR_API FRequiredEntryData CreateDefaultChildRequiredData() const;
 
-	virtual int32 GetChildIndentLevel() const;
+	NIAGARAEDITOR_API virtual int32 GetChildIndentLevel() const;
 
-	virtual TOptional<FDropRequestResponse> CanDropInternal(const FDropRequest& DropRequest);
+	NIAGARAEDITOR_API virtual TOptional<FDropRequestResponse> CanDropInternal(const FDropRequest& DropRequest);
 
-	virtual TOptional<FDropRequestResponse> DropInternal(const FDropRequest& DropRequest);
+	NIAGARAEDITOR_API virtual TOptional<FDropRequestResponse> DropInternal(const FDropRequest& DropRequest);
 
-	virtual TOptional<FDropRequestResponse> ChildRequestCanDropInternal(const UNiagaraStackEntry& TargetChild, const FDropRequest& DropRequest);
+	NIAGARAEDITOR_API virtual TOptional<FDropRequestResponse> ChildRequestCanDropInternal(const UNiagaraStackEntry& TargetChild, const FDropRequest& DropRequest);
 
-	virtual TOptional<FDropRequestResponse> ChildRequestDropInternal(const UNiagaraStackEntry& TargetChild, const FDropRequest& DropRequest);
+	NIAGARAEDITOR_API virtual TOptional<FDropRequestResponse> ChildRequestDropInternal(const UNiagaraStackEntry& TargetChild, const FDropRequest& DropRequest);
 
-	virtual void ChildStructureChangedInternal();
+	NIAGARAEDITOR_API virtual void ChildStructureChangedInternal();
 
-	virtual void FinalizeInternal();
+	NIAGARAEDITOR_API virtual void FinalizeInternal();
 
 
 	mutable TOptional<FCollectedUsageData> CachedCollectedUsageData;
@@ -582,27 +582,27 @@ protected:
 
 
 private:
-	void ChildStructureChanged(ENiagaraStructureChangedFlags Info);
+	NIAGARAEDITOR_API void ChildStructureChanged(ENiagaraStructureChangedFlags Info);
 
-	void ChildExpansionChanged();
+	NIAGARAEDITOR_API void ChildExpansionChanged();
 
-	void ChildExpansionInOverviewChanged();
+	NIAGARAEDITOR_API void ChildExpansionInOverviewChanged();
 	
-	void ChildDataObjectModified(TArray<UObject*> ChangedObjects, ENiagaraDataObjectChange ChangeType);
+	NIAGARAEDITOR_API void ChildDataObjectModified(TArray<UObject*> ChangedObjects, ENiagaraDataObjectChange ChangeType);
 
-	void ChildRequestFullRefresh();
+	NIAGARAEDITOR_API void ChildRequestFullRefresh();
 
-	void ChildRequestFullRefreshDeferred();
+	NIAGARAEDITOR_API void ChildRequestFullRefreshDeferred();
 
-	TOptional<FDropRequestResponse> ChildRequestCanDrop(const UNiagaraStackEntry& TargetChild, const FDropRequest& DropRequest);
+	NIAGARAEDITOR_API TOptional<FDropRequestResponse> ChildRequestCanDrop(const UNiagaraStackEntry& TargetChild, const FDropRequest& DropRequest);
 
-	TOptional<FDropRequestResponse> ChildRequestDrop(const UNiagaraStackEntry& TargetChild, const FDropRequest& DropRequest);
+	NIAGARAEDITOR_API TOptional<FDropRequestResponse> ChildRequestDrop(const UNiagaraStackEntry& TargetChild, const FDropRequest& DropRequest);
 
-	void RefreshStackErrorChildren();
+	NIAGARAEDITOR_API void RefreshStackErrorChildren();
 
-	void IssueModified();
+	NIAGARAEDITOR_API void IssueModified();
 
-	void InvalidateFilteredChildren();
+	NIAGARAEDITOR_API void InvalidateFilteredChildren();
 
 	struct FCollectedIssueData
 	{
@@ -624,7 +624,7 @@ private:
 	};
 
 
-	const FCollectedIssueData& GetCollectedIssueData() const;
+	NIAGARAEDITOR_API const FCollectedIssueData& GetCollectedIssueData() const;
 
 	TWeakPtr<FNiagaraSystemViewModel> SystemViewModel;
 	TWeakPtr<FNiagaraEmitterViewModel> EmitterViewModel;
@@ -688,12 +688,12 @@ private:
 };
 
 
-UCLASS()
-class NIAGARAEDITOR_API UNiagaraStackSpacer : public UNiagaraStackEntry
+UCLASS(MinimalAPI)
+class UNiagaraStackSpacer : public UNiagaraStackEntry
 {
 	GENERATED_BODY()
 public:
-	void Initialize(FRequiredEntryData InRequiredEntryData, float InSpacerHeight, TAttribute<bool> InShouldShowInStack, FString InOwningStackItemEditorDataKey);
+	NIAGARAEDITOR_API void Initialize(FRequiredEntryData InRequiredEntryData, float InSpacerHeight, TAttribute<bool> InShouldShowInStack, FString InOwningStackItemEditorDataKey);
 	virtual EStackRowStyle GetStackRowStyle() const override { return UNiagaraStackEntry::EStackRowStyle::Spacer; }
 	virtual bool GetCanExpand() const override { return false; }
 	virtual bool GetShouldShowInStack() const override { return ShouldShowInStack.Get(); }
@@ -703,4 +703,4 @@ public:
 private:
 	float SpacerHeight;
 	TAttribute<bool> ShouldShowInStack;
-};
\ No newline at end of file
+};
diff --git a/Engine/Plugins/FX/Niagara/Source/NiagaraEditor/Public/ViewModels/Stack/NiagaraStackErrorItem.h b/Engine/Plugins/FX/Niagara/Source/NiagaraEditor/Public/ViewModels/Stack/NiagaraStackErrorItem.h
--- a/Engine/Plugins/FX/Niagara/Source/NiagaraEditor/Public/ViewModels/Stack/NiagaraStackErrorItem.h
+++ b/Engine/Plugins/FX/Niagara/Source/NiagaraEditor/Public/ViewModels/Stack/NiagaraStackErrorItem.h
@@ -6,26 +6,26 @@
 #include "Input/Reply.h"
 #include "NiagaraStackErrorItem.generated.h"
 
-UCLASS()
-class NIAGARAEDITOR_API UNiagaraStackErrorItem : public UNiagaraStackEntry
+UCLASS(MinimalAPI)
+class UNiagaraStackErrorItem : public UNiagaraStackEntry
 {
 	GENERATED_BODY()
 public:
 	DECLARE_MULTICAST_DELEGATE(FOnIssueNotify);
 public:
-	void Initialize(FRequiredEntryData InRequiredEntryData, FStackIssue InStackIssue, FString InStackEditorDataKey);
+	NIAGARAEDITOR_API void Initialize(FRequiredEntryData InRequiredEntryData, FStackIssue InStackIssue, FString InStackEditorDataKey);
 	FStackIssue GetStackIssue() const { return StackIssue; }
-	void SetStackIssue(const FStackIssue& InStackIssue);
-	virtual FText GetDisplayName() const override;
-	virtual EStackRowStyle GetStackRowStyle() const override;
-	FOnIssueNotify& OnIssueModified();
-	virtual void GetSearchItems(TArray<FStackSearchItem>& SearchItems) const override;
-	virtual EStackIssueSeverity GetIssueSeverity() const override;
-	virtual bool IsExpandedByDefault() const override;
-	virtual void SetIsExpandedByDefault(bool bIsExpanded);
+	NIAGARAEDITOR_API void SetStackIssue(const FStackIssue& InStackIssue);
+	NIAGARAEDITOR_API virtual FText GetDisplayName() const override;
+	NIAGARAEDITOR_API virtual EStackRowStyle GetStackRowStyle() const override;
+	NIAGARAEDITOR_API FOnIssueNotify& OnIssueModified();
+	NIAGARAEDITOR_API virtual void GetSearchItems(TArray<FStackSearchItem>& SearchItems) const override;
+	NIAGARAEDITOR_API virtual EStackIssueSeverity GetIssueSeverity() const override;
+	NIAGARAEDITOR_API virtual bool IsExpandedByDefault() const override;
+	NIAGARAEDITOR_API virtual void SetIsExpandedByDefault(bool bIsExpanded);
 
 protected:
-	virtual void RefreshChildrenInternal(const TArray<UNiagaraStackEntry*>& CurrentChildren, TArray<UNiagaraStackEntry*>& NewChildren, TArray<FStackIssue>& NewIssues) override;
+	NIAGARAEDITOR_API virtual void RefreshChildrenInternal(const TArray<UNiagaraStackEntry*>& CurrentChildren, TArray<UNiagaraStackEntry*>& NewChildren, TArray<FStackIssue>& NewIssues) override;
 
 protected:
 	FStackIssue StackIssue;
@@ -34,40 +34,40 @@ protected:
 	bool bIsExpandedByDefault = true;
 
 private:
-	void IssueFixed();
+	NIAGARAEDITOR_API void IssueFixed();
 }; 
 
-UCLASS()
-class NIAGARAEDITOR_API UNiagaraStackErrorItemLongDescription : public UNiagaraStackEntry
+UCLASS(MinimalAPI)
+class UNiagaraStackErrorItemLongDescription : public UNiagaraStackEntry
 {
 	GENERATED_BODY()
 
 public:
-	void Initialize(FRequiredEntryData InRequiredEntryData, FStackIssue InStackIssue, FString InStackEditorDataKey);
-	virtual FText GetDisplayName() const override;
-	virtual EStackRowStyle GetStackRowStyle() const override;
-	virtual EStackIssueSeverity GetIssueSeverity() const override;
+	NIAGARAEDITOR_API void Initialize(FRequiredEntryData InRequiredEntryData, FStackIssue InStackIssue, FString InStackEditorDataKey);
+	NIAGARAEDITOR_API virtual FText GetDisplayName() const override;
+	NIAGARAEDITOR_API virtual EStackRowStyle GetStackRowStyle() const override;
+	NIAGARAEDITOR_API virtual EStackIssueSeverity GetIssueSeverity() const override;
 
 protected:
 	FStackIssue StackIssue;
 };
 
-UCLASS()
-class NIAGARAEDITOR_API UNiagaraStackErrorItemFix : public UNiagaraStackEntry
+UCLASS(MinimalAPI)
+class UNiagaraStackErrorItemFix : public UNiagaraStackEntry
 {
 	GENERATED_BODY()
 
 public:
-	void Initialize(FRequiredEntryData InRequiredEntryData, FStackIssue InStackIssue, FStackIssueFix InIssueFix, FString InStackEditorDataKey);
+	NIAGARAEDITOR_API void Initialize(FRequiredEntryData InRequiredEntryData, FStackIssue InStackIssue, FStackIssueFix InIssueFix, FString InStackEditorDataKey);
 
 	FStackIssueFix GetStackIssueFix() const { return IssueFix; };
-	FReply OnTryFixError();
-	virtual FText GetDisplayName() const override;
-	virtual EStackRowStyle GetStackRowStyle() const override;
-	virtual FText GetFixButtonText() const;
-	UNiagaraStackErrorItem::FOnIssueNotify& OnIssueFixed();
-	void SetFixDelegate(const FStackIssueFixDelegate& InFixDelegate);
-	virtual EStackIssueSeverity GetIssueSeverity() const override;
+	NIAGARAEDITOR_API FReply OnTryFixError();
+	NIAGARAEDITOR_API virtual FText GetDisplayName() const override;
+	NIAGARAEDITOR_API virtual EStackRowStyle GetStackRowStyle() const override;
+	NIAGARAEDITOR_API virtual FText GetFixButtonText() const;
+	NIAGARAEDITOR_API UNiagaraStackErrorItem::FOnIssueNotify& OnIssueFixed();
+	NIAGARAEDITOR_API void SetFixDelegate(const FStackIssueFixDelegate& InFixDelegate);
+	NIAGARAEDITOR_API virtual EStackIssueSeverity GetIssueSeverity() const override;
 
 protected:
 	FStackIssue StackIssue;
@@ -75,16 +75,16 @@ protected:
 	UNiagaraStackErrorItem::FOnIssueNotify IssueFixedDelegate;
 };
 
-UCLASS()
-class NIAGARAEDITOR_API UNiagaraStackErrorItemDismiss : public UNiagaraStackErrorItemFix
+UCLASS(MinimalAPI)
+class UNiagaraStackErrorItemDismiss : public UNiagaraStackErrorItemFix
 {
 	GENERATED_BODY()
 
 public:
-	void Initialize(FRequiredEntryData InRequiredEntryData, UNiagaraStackEntry::FStackIssue InStackIssue, FString InStackEditorDataKey);
+	NIAGARAEDITOR_API void Initialize(FRequiredEntryData InRequiredEntryData, UNiagaraStackEntry::FStackIssue InStackIssue, FString InStackEditorDataKey);
 
-	virtual EStackRowStyle GetStackRowStyle() const override;
-	virtual FText GetFixButtonText() const override;
+	NIAGARAEDITOR_API virtual EStackRowStyle GetStackRowStyle() const override;
+	NIAGARAEDITOR_API virtual FText GetFixButtonText() const override;
 
 private:
 	void DismissIssue();
diff --git a/Engine/Plugins/FX/Niagara/Source/NiagaraEditor/Public/ViewModels/Stack/NiagaraStackEventScriptItemGroup.h b/Engine/Plugins/FX/Niagara/Source/NiagaraEditor/Public/ViewModels/Stack/NiagaraStackEventScriptItemGroup.h
--- a/Engine/Plugins/FX/Niagara/Source/NiagaraEditor/Public/ViewModels/Stack/NiagaraStackEventScriptItemGroup.h
+++ b/Engine/Plugins/FX/Niagara/Source/NiagaraEditor/Public/ViewModels/Stack/NiagaraStackEventScriptItemGroup.h
@@ -14,8 +14,8 @@ struct FNiagaraEventScriptProperties;
 class IDetailTreeNode;
 
 // This is a wrapper class used for the details customization in the stack, since the event script properties were moved from the emitter object into the version data struct
-UCLASS()
-class NIAGARAEDITOR_API UNiagaraStackEventWrapper : public UObject
+UCLASS(MinimalAPI)
+class UNiagaraStackEventWrapper : public UObject
 {
 	GENERATED_BODY()
 
@@ -26,40 +26,40 @@ public:
 
 	FVersionedNiagaraEmitterWeakPtr EmitterWeakPtr;
 
-	virtual void PostEditChangeProperty(FPropertyChangedEvent& PropertyChangedEvent) override;
+	NIAGARAEDITOR_API virtual void PostEditChangeProperty(FPropertyChangedEvent& PropertyChangedEvent) override;
 };
 
-UCLASS()
-class NIAGARAEDITOR_API UNiagaraStackEventHandlerPropertiesItem : public UNiagaraStackItem
+UCLASS(MinimalAPI)
+class UNiagaraStackEventHandlerPropertiesItem : public UNiagaraStackItem
 {
 	GENERATED_BODY()
 		 
 public:
-	void Initialize(FRequiredEntryData InRequiredEntryData, FGuid InEventScriptUsageId);
+	NIAGARAEDITOR_API void Initialize(FRequiredEntryData InRequiredEntryData, FGuid InEventScriptUsageId);
 
-	virtual FText GetDisplayName() const override;
+	NIAGARAEDITOR_API virtual FText GetDisplayName() const override;
 
 	virtual bool SupportsResetToBase() const override { return true; }
-	virtual bool TestCanResetToBaseWithMessage(FText& OutCanResetToBaseMessage) const override;
-	virtual void ResetToBase() override;
+	NIAGARAEDITOR_API virtual bool TestCanResetToBaseWithMessage(FText& OutCanResetToBaseMessage) const override;
+	NIAGARAEDITOR_API virtual void ResetToBase() override;
 
 	virtual FGuid GetSelectionId() const override { return EventScriptUsageId; }
 
 	FGuid GetEventScriptUsageId() const { return EventScriptUsageId; };
 
 protected:
-	virtual void FinalizeInternal() override;
+	NIAGARAEDITOR_API virtual void FinalizeInternal() override;
 
-	virtual void RefreshChildrenInternal(const TArray<UNiagaraStackEntry*>& CurrentChildren, TArray<UNiagaraStackEntry*>& NewChildren, TArray<FStackIssue>& NewIssues) override;
+	NIAGARAEDITOR_API virtual void RefreshChildrenInternal(const TArray<UNiagaraStackEntry*>& CurrentChildren, TArray<UNiagaraStackEntry*>& NewChildren, TArray<FStackIssue>& NewIssues) override;
 
 	virtual bool SupportsSummaryView() const override { return true; }
-	virtual FNiagaraHierarchyIdentity DetermineSummaryIdentity() const override;
+	NIAGARAEDITOR_API virtual FNiagaraHierarchyIdentity DetermineSummaryIdentity() const override;
 private:
-	void EventHandlerPropertiesChanged();
+	NIAGARAEDITOR_API void EventHandlerPropertiesChanged();
 
-	void SelectEmitterStackObjectRootTreeNodes(TArray<TSharedRef<IDetailTreeNode>> Source, TArray<TSharedRef<IDetailTreeNode>>* Selected);
+	NIAGARAEDITOR_API void SelectEmitterStackObjectRootTreeNodes(TArray<TSharedRef<IDetailTreeNode>> Source, TArray<TSharedRef<IDetailTreeNode>>* Selected);
 
-	bool HasBaseEventHandler() const;
+	NIAGARAEDITOR_API bool HasBaseEventHandler() const;
 
 private:
 	FGuid EventScriptUsageId;
@@ -77,9 +77,9 @@ private:
 	TObjectPtr<UNiagaraStackEventWrapper> EventWrapper;
 };
 
-UCLASS()
+UCLASS(MinimalAPI)
 /** Meant to contain a single binding of a Emitter::EventScriptProperties to the stack.*/
-class NIAGARAEDITOR_API UNiagaraStackEventScriptItemGroup : public UNiagaraStackScriptItemGroup
+class UNiagaraStackEventScriptItemGroup : public UNiagaraStackScriptItemGroup
 {
 	GENERATED_BODY()
 
@@ -87,22 +87,22 @@ public:
 	DECLARE_DELEGATE(FOnModifiedEventHandlers);
 
 public:
-	void Initialize(
+	NIAGARAEDITOR_API void Initialize(
 		FRequiredEntryData InRequiredEntryData,
 		TSharedRef<FNiagaraScriptViewModel> InScriptViewModel,
 		ENiagaraScriptUsage InScriptUsage,
 		FGuid InScriptUsageId,
 		FGuid InEventSourceEmitterId);
 
-	void SetOnModifiedEventHandlers(FOnModifiedEventHandlers OnModifiedEventHandlers);
+	NIAGARAEDITOR_API void SetOnModifiedEventHandlers(FOnModifiedEventHandlers OnModifiedEventHandlers);
 
 	virtual bool SupportsDelete() const override { return true; }
-	virtual bool TestCanDeleteWithMessage(FText& OutCanDeleteMessage) const override;
-	virtual void Delete() override;
+	NIAGARAEDITOR_API virtual bool TestCanDeleteWithMessage(FText& OutCanDeleteMessage) const override;
+	NIAGARAEDITOR_API virtual void Delete() override;
 
 	virtual bool SupportsInheritance() const override { return true; }
-	virtual bool GetIsInherited() const override;
-	virtual FText GetInheritanceMessage() const override;
+	NIAGARAEDITOR_API virtual bool GetIsInherited() const override;
+	NIAGARAEDITOR_API virtual FText GetInheritanceMessage() const override;
 
 	const TObjectPtr<UNiagaraStackEventHandlerPropertiesItem>& GetEventHandlerPropertiesItem() const { return EventHandlerProperties; }
 	FGuid GetEventSourceEmitterId() const {return EventSourceEmitterId; };
@@ -110,12 +110,12 @@ public:
 protected:
 	FOnModifiedEventHandlers OnModifiedEventHandlersDelegate;
 
-	virtual void RefreshChildrenInternal(const TArray<UNiagaraStackEntry*>& CurrentChildren, TArray<UNiagaraStackEntry*>& NewChildren, TArray<FStackIssue>& NewIssues) override;
+	NIAGARAEDITOR_API virtual void RefreshChildrenInternal(const TArray<UNiagaraStackEntry*>& CurrentChildren, TArray<UNiagaraStackEntry*>& NewChildren, TArray<FStackIssue>& NewIssues) override;
 
 	virtual bool SupportsSummaryView() const override { return true; }
-	virtual FNiagaraHierarchyIdentity DetermineSummaryIdentity() const override;
+	NIAGARAEDITOR_API virtual FNiagaraHierarchyIdentity DetermineSummaryIdentity() const override;
 private:
-	bool HasBaseEventHandler() const;
+	NIAGARAEDITOR_API bool HasBaseEventHandler() const;
 	
 private:
 	mutable TOptional<bool> bHasBaseEventHandlerCache;
diff --git a/Engine/Plugins/FX/Niagara/Source/NiagaraEditor/Public/ViewModels/Stack/NiagaraStackFunctionInput.h b/Engine/Plugins/FX/Niagara/Source/NiagaraEditor/Public/ViewModels/Stack/NiagaraStackFunctionInput.h
--- a/Engine/Plugins/FX/Niagara/Source/NiagaraEditor/Public/ViewModels/Stack/NiagaraStackFunctionInput.h
+++ b/Engine/Plugins/FX/Niagara/Source/NiagaraEditor/Public/ViewModels/Stack/NiagaraStackFunctionInput.h
@@ -27,8 +27,8 @@ class UNiagaraScriptVariable;
 class FNiagaraPlaceholderDataInterfaceHandle;
 
 /** Represents a single module input in the module stack view model. */
-UCLASS()
-class NIAGARAEDITOR_API UNiagaraStackFunctionInput : public UNiagaraStackItemContent
+UCLASS(MinimalAPI)
+class UNiagaraStackFunctionInput : public UNiagaraStackItemContent
 {
 	GENERATED_BODY()
 
@@ -61,7 +61,7 @@ public:
 	DECLARE_MULTICAST_DELEGATE(FOnValueChanged);
 
 public:
-	UNiagaraStackFunctionInput();
+	NIAGARAEDITOR_API UNiagaraStackFunctionInput();
 
 	/** 
 	 * Sets the input data for this entry.
@@ -74,7 +74,7 @@ public:
 	 * @param InParameterBehavior Determines how the parameter should behave in the stack
 	 * @param InOwnerStackItemEditorDataKey The editor data key of the item that owns this input.
 	 */
-	void Initialize(
+	NIAGARAEDITOR_API void Initialize(
 		FRequiredEntryData InRequiredEntryData,
 		UNiagaraNodeFunctionCall& InModuleNode,
 		UNiagaraNodeFunctionCall& InInputFunctionCallNode,
@@ -84,240 +84,240 @@ public:
 		FString InOwnerStackItemEditorDataKey);
 
 	/** Gets the function call node which owns this input. */
-	const UNiagaraNodeFunctionCall& GetInputFunctionCallNode() const;
+	NIAGARAEDITOR_API const UNiagaraNodeFunctionCall& GetInputFunctionCallNode() const;
 
 	/** Gets the script that the function call node was referencing when this input was initialized. */
-	UNiagaraScript* GetInputFunctionCallInitialScript() const;
+	NIAGARAEDITOR_API UNiagaraScript* GetInputFunctionCallInitialScript() const;
 
 	/** Gets the current value mode */
-	EValueMode GetValueMode() const;
+	NIAGARAEDITOR_API EValueMode GetValueMode() const;
 
 	/** Gets the type of this input. */
-	const FNiagaraTypeDefinition& GetInputType() const;
+	NIAGARAEDITOR_API const FNiagaraTypeDefinition& GetInputType() const;
 
 	/** Gets the unit of this input. */
-	EUnit GetInputDisplayUnit() const;
+	NIAGARAEDITOR_API EUnit GetInputDisplayUnit() const;
 
-	FNiagaraInputParameterCustomization GetInputWidgetCustomization() const;
+	NIAGARAEDITOR_API FNiagaraInputParameterCustomization GetInputWidgetCustomization() const;
 
 	//~ UNiagaraStackEntry interface
-	virtual FText GetDisplayName() const override;
-	virtual FText GetTooltipText() const override;
-	virtual bool GetIsEnabled() const override;
-	virtual UObject* GetExternalAsset() const override;
+	NIAGARAEDITOR_API virtual FText GetDisplayName() const override;
+	NIAGARAEDITOR_API virtual FText GetTooltipText() const override;
+	NIAGARAEDITOR_API virtual bool GetIsEnabled() const override;
+	NIAGARAEDITOR_API virtual UObject* GetExternalAsset() const override;
 	virtual bool SupportsCut() const override { return true; }
-	virtual bool TestCanCutWithMessage(FText& OutMessage) const override;
-	virtual FText GetCutTransactionText() const override;
-	virtual void CopyForCut(UNiagaraClipboardContent* ClipboardContent) const override;
-	virtual void RemoveForCut() override;
+	NIAGARAEDITOR_API virtual bool TestCanCutWithMessage(FText& OutMessage) const override;
+	NIAGARAEDITOR_API virtual FText GetCutTransactionText() const override;
+	NIAGARAEDITOR_API virtual void CopyForCut(UNiagaraClipboardContent* ClipboardContent) const override;
+	NIAGARAEDITOR_API virtual void RemoveForCut() override;
 	virtual bool SupportsCopy() const override { return true; }
-	virtual bool TestCanCopyWithMessage(FText& OutMessage) const override;
-	virtual void Copy(UNiagaraClipboardContent* ClipboardContent) const override;
+	NIAGARAEDITOR_API virtual bool TestCanCopyWithMessage(FText& OutMessage) const override;
+	NIAGARAEDITOR_API virtual void Copy(UNiagaraClipboardContent* ClipboardContent) const override;
 	virtual bool SupportsPaste() const override { return true; }
-	virtual bool TestCanPasteWithMessage(const UNiagaraClipboardContent* ClipboardContent, FText& OutMessage) const override;
-	virtual FText GetPasteTransactionText(const UNiagaraClipboardContent* ClipboardContent) const override;
-	virtual void Paste(const UNiagaraClipboardContent* ClipboardContent, FText& OutPasteWarning) override;
-	virtual bool HasOverridenContent() const override;
-	virtual bool SupportsSummaryView() const override;
-	virtual FNiagaraHierarchyIdentity DetermineSummaryIdentity() const override;
+	NIAGARAEDITOR_API virtual bool TestCanPasteWithMessage(const UNiagaraClipboardContent* ClipboardContent, FText& OutMessage) const override;
+	NIAGARAEDITOR_API virtual FText GetPasteTransactionText(const UNiagaraClipboardContent* ClipboardContent) const override;
+	NIAGARAEDITOR_API virtual void Paste(const UNiagaraClipboardContent* ClipboardContent, FText& OutPasteWarning) override;
+	NIAGARAEDITOR_API virtual bool HasOverridenContent() const override;
+	NIAGARAEDITOR_API virtual bool SupportsSummaryView() const override;
+	NIAGARAEDITOR_API virtual FNiagaraHierarchyIdentity DetermineSummaryIdentity() const override;
 	
 	/** Gets the tooltip that should be shown for the value of this input. */
-	FText GetValueToolTip() const;
+	NIAGARAEDITOR_API FText GetValueToolTip() const;
 
 	/** Gets the tooltip that should be shown for the value of this input. */
-	FText GetCollapsedStateText() const;
+	NIAGARAEDITOR_API FText GetCollapsedStateText() const;
 
-	void SetSummaryViewDisplayName(TAttribute<FText> InDisplayName);
-	void SetSummaryViewTooltip(TAttribute<FText> InTooltipOverride);
+	NIAGARAEDITOR_API void SetSummaryViewDisplayName(TAttribute<FText> InDisplayName);
+	NIAGARAEDITOR_API void SetSummaryViewTooltip(TAttribute<FText> InTooltipOverride);
 
 	/** Gets the path of parameter handles from the owning module to the function call which owns this input. */
-	const TArray<FNiagaraParameterHandle>& GetInputParameterHandlePath() const;
+	NIAGARAEDITOR_API const TArray<FNiagaraParameterHandle>& GetInputParameterHandlePath() const;
 
 	/** Gets the parameter handle which defined this input in the module. */
-	const FNiagaraParameterHandle& GetInputParameterHandle() const;
+	NIAGARAEDITOR_API const FNiagaraParameterHandle& GetInputParameterHandle() const;
 
 	/** Gets the handle to the linked value for this input if there is one. */
-	const FNiagaraParameterHandle& GetLinkedValueHandle() const;
+	NIAGARAEDITOR_API const FNiagaraParameterHandle& GetLinkedValueHandle() const;
 
 	/** Sets the value of this input to a linked parameter handle. */
-	void SetLinkedValueHandle(const FNiagaraParameterHandle& InParameterHandle);
+	NIAGARAEDITOR_API void SetLinkedValueHandle(const FNiagaraParameterHandle& InParameterHandle);
 
 	/** Gets the current set of available parameter handles which can be assigned to this input. Optionally returns possible conversion scripts. */
-	void GetAvailableParameterHandles(TArray<FNiagaraParameterHandle>& AvailableParameterHandles, TMap<FNiagaraVariable, UNiagaraScript*>& AvailableConversionHandles, bool bIncludeConversionScripts = true) const;
+	NIAGARAEDITOR_API void GetAvailableParameterHandles(TArray<FNiagaraParameterHandle>& AvailableParameterHandles, TMap<FNiagaraVariable, UNiagaraScript*>& AvailableConversionHandles, bool bIncludeConversionScripts = true) const;
 
 	/** Gets the function node form the script graph if the current value mode is DefaultFunction. */
-	UNiagaraNodeFunctionCall* GetDefaultFunctionNode() const;
+	NIAGARAEDITOR_API UNiagaraNodeFunctionCall* GetDefaultFunctionNode() const;
 
 	/** Gets the dynamic input node providing the value for this input, if one is available. */
-	UNiagaraNodeFunctionCall* GetDynamicInputNode() const;
+	NIAGARAEDITOR_API UNiagaraNodeFunctionCall* GetDynamicInputNode() const;
 
 	/** Gets the dynamic inputs available for this input. */
-	void GetAvailableDynamicInputs(TArray<UNiagaraScript*>& AvailableDynamicInputs, bool bIncludeNonLibraryInputs = false);
+	NIAGARAEDITOR_API void GetAvailableDynamicInputs(TArray<UNiagaraScript*>& AvailableDynamicInputs, bool bIncludeNonLibraryInputs = false);
 
 	/** Sets the dynamic input script for this input. */
-	void SetDynamicInput(UNiagaraScript* DynamicInput, FString SuggestedName = FString(), const FGuid& InScriptVersion = FGuid());
+	NIAGARAEDITOR_API void SetDynamicInput(UNiagaraScript* DynamicInput, FString SuggestedName = FString(), const FGuid& InScriptVersion = FGuid());
 
 	/** Gets the expression providing the value for this input, if one is available. */
-	FText GetCustomExpressionText() const;
+	NIAGARAEDITOR_API FText GetCustomExpressionText() const;
 
 	/** Sets the dynamic custom expression script for this input. */
-	void SetCustomExpression(const FString& InCustomExpression);
+	NIAGARAEDITOR_API void SetCustomExpression(const FString& InCustomExpression);
 
 	/** Create a new scratch pad dynamic inputs and set this input to use it. */
-	void SetScratch();
+	NIAGARAEDITOR_API void SetScratch();
 
 	/** Gets the current struct value of this input is there is one. */
-	TSharedPtr<const FStructOnScope> GetLocalValueStruct();
+	NIAGARAEDITOR_API TSharedPtr<const FStructOnScope> GetLocalValueStruct();
 
 	/** Gets the current data object value of this input is there is one. */
-	UNiagaraDataInterface* GetDataValueObject();
+	NIAGARAEDITOR_API UNiagaraDataInterface* GetDataValueObject();
 
 	/** Gets the current object Asset value of this input is there is one. */
-	UObject* GetObjectAssetValue();
+	NIAGARAEDITOR_API UObject* GetObjectAssetValue();
 
 	/** Sets the current object Asset value of this input is there is one. */
-	void SetObjectAssetValue(UObject* NewValue);
+	NIAGARAEDITOR_API void SetObjectAssetValue(UObject* NewValue);
 
 	/** Called to notify the input that an ongoing change to it's value has begun. */
-	void NotifyBeginLocalValueChange();
+	NIAGARAEDITOR_API void NotifyBeginLocalValueChange();
 
 	/** Called to notify the input that an ongoing change to it's value has ended. */
-	void NotifyEndLocalValueChange();
+	NIAGARAEDITOR_API void NotifyEndLocalValueChange();
 
 	/** Is this pin editable or should it show as disabled?*/
-	bool IsEnabled() const;
+	NIAGARAEDITOR_API bool IsEnabled() const;
 
 	/** Sets this input's local value. */
-	void SetLocalValue(TSharedRef<FStructOnScope> InLocalValue);
+	NIAGARAEDITOR_API void SetLocalValue(TSharedRef<FStructOnScope> InLocalValue);
 	
 	/** Returns whether or not the value or handle of this input has been overridden and can be reset. */
-	bool CanReset() const;
+	NIAGARAEDITOR_API bool CanReset() const;
 
 	/** Resets the value and handle of this input to the value and handle defined in the module. */
-	void Reset();
+	NIAGARAEDITOR_API void Reset();
 
 	/** Determine if this field is editable */
-	bool IsEditable() const;
+	NIAGARAEDITOR_API bool IsEditable() const;
 
 	/** If true the parameter can only be set to local constant values */
-	bool IsStaticParameter() const;
+	NIAGARAEDITOR_API bool IsStaticParameter() const;
 
 	/** Whether or not this input has a base value.  This is true for emitter instances in systems. */
-	bool EmitterHasBase() const;
+	NIAGARAEDITOR_API bool EmitterHasBase() const;
 
 	/** Whether or not this input can be reset to a base value. */
-	bool CanResetToBase() const;
+	NIAGARAEDITOR_API bool CanResetToBase() const;
 
 	/** Resets this input to its base value. */
-	void ResetToBase();
+	NIAGARAEDITOR_API void ResetToBase();
 
 	/** Returns whether or not this input can be renamed. */
-	virtual bool SupportsRename() const override;
+	NIAGARAEDITOR_API virtual bool SupportsRename() const override;
 
 	/** Renames this input to the name specified. */
-	virtual void OnRenamed(FText NewName) override;
+	NIAGARAEDITOR_API virtual void OnRenamed(FText NewName) override;
 
 	/** Returns whether or not this input can be deleted. */
-	bool CanDeleteInput() const;
+	NIAGARAEDITOR_API bool CanDeleteInput() const;
 
 	/** Deletes this input */
-	void DeleteInput();
+	NIAGARAEDITOR_API void DeleteInput();
 
 	/** Gets the namespaces which new parameters for this input can be read from. */
-	void GetNamespacesForNewReadParameters(TArray<FName>& OutNamespacesForNewParameters) const;
+	NIAGARAEDITOR_API void GetNamespacesForNewReadParameters(TArray<FName>& OutNamespacesForNewParameters) const;
 
 	/** Gets the namespaces which new parameters for this input can write to. */
-	void GetNamespacesForNewWriteParameters(TArray<FName>& OutNamespacesForNewParameters) const;
+	NIAGARAEDITOR_API void GetNamespacesForNewWriteParameters(TArray<FName>& OutNamespacesForNewParameters) const;
 
 	/** Gets a multicast delegate which is called whenever the value on this input changes. */
-	FOnValueChanged& OnValueChanged();
+	NIAGARAEDITOR_API FOnValueChanged& OnValueChanged();
 
 	/** Gets the variable that serves as an edit condition for this input. */
-	TOptional<FNiagaraVariable> GetEditConditionVariable() const;
+	NIAGARAEDITOR_API TOptional<FNiagaraVariable> GetEditConditionVariable() const;
 	
 	/** Gets whether or not this input has an associated edit condition input. */
-	bool GetHasEditCondition() const;
+	NIAGARAEDITOR_API bool GetHasEditCondition() const;
 
 	/** Gets whether or not to show a control inline for the edit condition input associated with this input. */
-	bool GetShowEditConditionInline() const;
+	NIAGARAEDITOR_API bool GetShowEditConditionInline() const;
 
 	/** Gets the enabled value of the edit condition input associated with this input. */
-	bool GetEditConditionEnabled() const;
+	NIAGARAEDITOR_API bool GetEditConditionEnabled() const;
 
 	/** Sets the enabled value of the edit condition input associated with this input. */
-	void SetEditConditionEnabled(bool bIsEnabled);
+	NIAGARAEDITOR_API void SetEditConditionEnabled(bool bIsEnabled);
 
 	/** Gets whether or not this input has an associated visible condition input. */
-	bool GetHasVisibleCondition() const;
+	NIAGARAEDITOR_API bool GetHasVisibleCondition() const;
 
 	/** Gets the enabled value of the visible condition input associated with this input. */
-	bool GetVisibleConditionEnabled() const;
+	NIAGARAEDITOR_API bool GetVisibleConditionEnabled() const;
 
 	/** Gets whether or not this input is used as an edit condition for another input and should be hidden. */
-	bool GetIsInlineEditConditionToggle() const;
+	NIAGARAEDITOR_API bool GetIsInlineEditConditionToggle() const;
 
 	/** Gets whether or not a dynamic input script reassignment is pending.  This can happen when trying to fix dynamic inputs which are missing their scripts. */
-	bool GetIsDynamicInputScriptReassignmentPending() const;
+	NIAGARAEDITOR_API bool GetIsDynamicInputScriptReassignmentPending() const;
 
 	/** Gets whether or not a dynamic input script reassignment should be be pending. */
-	void SetIsDynamicInputScriptReassignmentPending(bool bIsPending);
+	NIAGARAEDITOR_API void SetIsDynamicInputScriptReassignmentPending(bool bIsPending);
 
 	/** Reassigns the function script for the current dynamic input without resetting the sub-inputs. */
-	void ReassignDynamicInputScript(UNiagaraScript* DynamicInputScript);
+	NIAGARAEDITOR_API void ReassignDynamicInputScript(UNiagaraScript* DynamicInputScript);
 
 	/** Gets whether or not this input is filtered from search results and appearing in stack due to visibility metadata*/
-	bool GetShouldPassFilterForVisibleCondition() const;
+	NIAGARAEDITOR_API bool GetShouldPassFilterForVisibleCondition() const;
 	
-	TArray<UNiagaraScript*> GetPossibleConversionScripts(const FNiagaraTypeDefinition& FromType) const;
-	static TArray<UNiagaraScript*> GetPossibleConversionScripts(const FNiagaraTypeDefinition& FromType, const FNiagaraTypeDefinition& ToType);
+	NIAGARAEDITOR_API TArray<UNiagaraScript*> GetPossibleConversionScripts(const FNiagaraTypeDefinition& FromType) const;
+	static NIAGARAEDITOR_API TArray<UNiagaraScript*> GetPossibleConversionScripts(const FNiagaraTypeDefinition& FromType, const FNiagaraTypeDefinition& ToType);
 
-	void SetLinkedInputViaConversionScript(const FName& LinkedInputName, const FNiagaraTypeDefinition& FromType);
-	void SetLinkedInputViaConversionScript(const FNiagaraVariable& LinkedInput, UNiagaraScript* ConversionScript);
-	void SetClipboardContentViaConversionScript(const UNiagaraClipboardFunctionInput& ClipboardFunctionInput);
+	NIAGARAEDITOR_API void SetLinkedInputViaConversionScript(const FName& LinkedInputName, const FNiagaraTypeDefinition& FromType);
+	NIAGARAEDITOR_API void SetLinkedInputViaConversionScript(const FNiagaraVariable& LinkedInput, UNiagaraScript* ConversionScript);
+	NIAGARAEDITOR_API void SetClipboardContentViaConversionScript(const UNiagaraClipboardFunctionInput& ClipboardFunctionInput);
 
-	void ChangeScriptVersion(FGuid NewScriptVersion);
+	NIAGARAEDITOR_API void ChangeScriptVersion(FGuid NewScriptVersion);
 
-	const UNiagaraClipboardFunctionInput* ToClipboardFunctionInput(UObject* InOuter) const;
+	NIAGARAEDITOR_API const UNiagaraClipboardFunctionInput* ToClipboardFunctionInput(UObject* InOuter) const;
 
-	void SetValueFromClipboardFunctionInput(const UNiagaraClipboardFunctionInput& ClipboardFunctionInput);
+	NIAGARAEDITOR_API void SetValueFromClipboardFunctionInput(const UNiagaraClipboardFunctionInput& ClipboardFunctionInput);
 
-	bool IsScratchDynamicInput() const;
+	NIAGARAEDITOR_API bool IsScratchDynamicInput() const;
 
-	bool ShouldDisplayInline() const;
+	NIAGARAEDITOR_API bool ShouldDisplayInline() const;
 	
-	TArray<UNiagaraStackFunctionInput*> GetChildInputs() const;
+	NIAGARAEDITOR_API TArray<UNiagaraStackFunctionInput*> GetChildInputs() const;
 
-	TOptional<FNiagaraVariableMetaData> GetInputMetaData() const;
+	NIAGARAEDITOR_API TOptional<FNiagaraVariableMetaData> GetInputMetaData() const;
 	
-	void GetFilteredChildInputs(TArray<UNiagaraStackFunctionInput*>& OutFilteredChildInputs) const;
+	NIAGARAEDITOR_API void GetFilteredChildInputs(TArray<UNiagaraStackFunctionInput*>& OutFilteredChildInputs) const;
 
-	UNiagaraStackObject* GetChildDataObject() const;
+	NIAGARAEDITOR_API UNiagaraStackObject* GetChildDataObject() const;
 
-	virtual bool IsSemanticChild() const override;
-	void SetSemanticChild(bool IsSemanticChild);
+	NIAGARAEDITOR_API virtual bool IsSemanticChild() const override;
+	NIAGARAEDITOR_API void SetSemanticChild(bool IsSemanticChild);
 
 	//~ UNiagaraStackEntry interface
-	virtual void GetSearchItems(TArray<FStackSearchItem>& SearchItems) const override;
+	NIAGARAEDITOR_API virtual void GetSearchItems(TArray<FStackSearchItem>& SearchItems) const override;
 
-	virtual const FCollectedUsageData& GetCollectedUsageData() const override;
+	NIAGARAEDITOR_API virtual const FCollectedUsageData& GetCollectedUsageData() const override;
 
-	ENiagaraStackEntryInlineDisplayMode GetInlineDisplayMode() const;
+	NIAGARAEDITOR_API ENiagaraStackEntryInlineDisplayMode GetInlineDisplayMode() const;
 
-	void SetInlineDisplayMode(ENiagaraStackEntryInlineDisplayMode InlineDisplayMode);
+	NIAGARAEDITOR_API void SetInlineDisplayMode(ENiagaraStackEntryInlineDisplayMode InlineDisplayMode);
 
-	bool OpenSourceAsset() const;
+	NIAGARAEDITOR_API bool OpenSourceAsset() const;
 
-	bool SupportsCustomExpressions() const;
+	NIAGARAEDITOR_API bool SupportsCustomExpressions() const;
 
 protected:
 	//~ UNiagaraStackEntry interface
-	virtual void FinalizeInternal() override;
-	virtual void RefreshChildrenInternal(const TArray<UNiagaraStackEntry*>& CurrentChildren, TArray<UNiagaraStackEntry*>& NewChildren, TArray<FStackIssue>& NewIssues) override;
+	NIAGARAEDITOR_API virtual void FinalizeInternal() override;
+	NIAGARAEDITOR_API virtual void RefreshChildrenInternal(const TArray<UNiagaraStackEntry*>& CurrentChildren, TArray<UNiagaraStackEntry*>& NewChildren, TArray<FStackIssue>& NewIssues) override;
 
-	bool UpdateRapidIterationParametersForAffectedScripts(const uint8* Data);
-	bool RemoveRapidIterationParametersForAffectedScripts(bool bUpdateGraphGuidsForAffected = false);
-	FString ResolveDisplayNameArgument(const FString& InArg) const;
-	FStackIssueFixDelegate GetUpgradeDynamicInputVersionFix();
+	NIAGARAEDITOR_API bool UpdateRapidIterationParametersForAffectedScripts(const uint8* Data);
+	NIAGARAEDITOR_API bool RemoveRapidIterationParametersForAffectedScripts(bool bUpdateGraphGuidsForAffected = false);
+	NIAGARAEDITOR_API FString ResolveDisplayNameArgument(const FString& InArg) const;
+	NIAGARAEDITOR_API FStackIssueFixDelegate GetUpgradeDynamicInputVersionFix();
 
 private:
 	struct FInputValues
@@ -344,65 +344,65 @@ private:
 
 private:
 	/** Refreshes the current values for this input from the state of the graph. */
-	void RefreshValues();
+	NIAGARAEDITOR_API void RefreshValues();
 
 	/** Refreshes additional state for this input which comes from input metadata. */
-	void RefreshFromMetaData(TArray<FStackIssue>& NewIssues);
+	NIAGARAEDITOR_API void RefreshFromMetaData(TArray<FStackIssue>& NewIssues);
 
 	/** Called whenever the graph which generated this input changes. */
-	void OnGraphChanged(const struct FEdGraphEditAction& InAction);
+	NIAGARAEDITOR_API void OnGraphChanged(const struct FEdGraphEditAction& InAction);
 
 	/** Called whenever rapid iteration parameters are changed for the script that owns the function that owns this input. */
-	void OnRapidIterationParametersChanged();
+	NIAGARAEDITOR_API void OnRapidIterationParametersChanged();
 
 	/** Called whenever the script source that owns the function that owns this input changes. */
-	void OnScriptSourceChanged();
+	NIAGARAEDITOR_API void OnScriptSourceChanged();
 
 	/** Gets the graph node which owns the local overrides for the module that owns this input if it exists. */
-	UNiagaraNodeParameterMapSet* GetOverrideNode() const;
+	NIAGARAEDITOR_API UNiagaraNodeParameterMapSet* GetOverrideNode() const;
 
 	/** Gets the graph node which owns the local overrides for the module that owns this input this input.  
 	  * This will create the node and add it to the graph if it doesn't exist. */
-	UNiagaraNodeParameterMapSet& GetOrCreateOverrideNode();
+	NIAGARAEDITOR_API UNiagaraNodeParameterMapSet& GetOrCreateOverrideNode();
 
 	/** Gets the pin on the override node which is associated with this input if it exists. */
-	UEdGraphPin* GetOverridePin() const;
+	NIAGARAEDITOR_API UEdGraphPin* GetOverridePin() const;
 
 	/** Gets the pin on the override node which is associated with this input.  If either the override node or
 	  * pin don't exist, they will be created. */
-	UEdGraphPin& GetOrCreateOverridePin();
+	NIAGARAEDITOR_API UEdGraphPin& GetOrCreateOverridePin();
 
-	void GetDefaultDataInterfaceValueFromDefaultPin(UEdGraphPin* DefaultPin, UNiagaraStackFunctionInput::FInputValues& InInputValues) const;
+	NIAGARAEDITOR_API void GetDefaultDataInterfaceValueFromDefaultPin(UEdGraphPin* DefaultPin, UNiagaraStackFunctionInput::FInputValues& InInputValues) const;
 
-	void GetDefaultObjectAssetValueFromDefaultPin(UEdGraphPin* DefaultPin, UNiagaraStackFunctionInput::FInputValues& InInputValues) const;
+	NIAGARAEDITOR_API void GetDefaultObjectAssetValueFromDefaultPin(UEdGraphPin* DefaultPin, UNiagaraStackFunctionInput::FInputValues& InInputValues) const;
 
-	void GetDefaultLocalValueFromDefaultPin(UEdGraphPin* DefaultPin, UNiagaraStackFunctionInput::FInputValues& InInputValues) const;
+	NIAGARAEDITOR_API void GetDefaultLocalValueFromDefaultPin(UEdGraphPin* DefaultPin, UNiagaraStackFunctionInput::FInputValues& InInputValues) const;
 
-	void GetDefaultLinkedHandleOrLinkedFunctionFromDefaultPin(UEdGraphPin* DefaultPin, UNiagaraStackFunctionInput::FInputValues& InInputValues) const;
+	NIAGARAEDITOR_API void GetDefaultLinkedHandleOrLinkedFunctionFromDefaultPin(UEdGraphPin* DefaultPin, UNiagaraStackFunctionInput::FInputValues& InInputValues) const;
 
-	void UpdateValuesFromScriptDefaults(FInputValues& InInputValues) const;
+	NIAGARAEDITOR_API void UpdateValuesFromScriptDefaults(FInputValues& InInputValues) const;
 
-	void UpdateValuesFromOverridePin(const FInputValues& OldInputValues, FInputValues& NewInputValues, UEdGraphPin& InOverridePin) const;
+	NIAGARAEDITOR_API void UpdateValuesFromOverridePin(const FInputValues& OldInputValues, FInputValues& NewInputValues, UEdGraphPin& InOverridePin) const;
 
 	/** Removes all nodes connected to the override pin which provide it's value. */
-	void RemoveNodesForOverridePin(UEdGraphPin& OverridePin);
+	NIAGARAEDITOR_API void RemoveNodesForOverridePin(UEdGraphPin& OverridePin);
 
 	/** Remove the override pin and all nodes connected to it. */
-	void RemoveOverridePin();
+	NIAGARAEDITOR_API void RemoveOverridePin();
 
 	/** Determine if the values in this input are possibly under the control of the rapid iteration array on the script.*/
-	bool IsRapidIterationCandidate() const;
+	NIAGARAEDITOR_API bool IsRapidIterationCandidate() const;
 
-	FNiagaraVariable CreateRapidIterationVariable(const FName& InName);
+	NIAGARAEDITOR_API FNiagaraVariable CreateRapidIterationVariable(const FName& InName);
 
 	/** Handles the message manager refreshing messages. */
-	void OnMessageManagerRefresh(const TArray<TSharedRef<const INiagaraMessage>>& NewMessages);
+	NIAGARAEDITOR_API void OnMessageManagerRefresh(const TArray<TSharedRef<const INiagaraMessage>>& NewMessages);
 
-	void GetCurrentChangeIds(FGuid& OutOwningGraphChangeId, FGuid& OutFunctionGraphChangeId) const;
+	NIAGARAEDITOR_API void GetCurrentChangeIds(FGuid& OutOwningGraphChangeId, FGuid& OutFunctionGraphChangeId) const;
 
-	UNiagaraScript* FindConversionScript(const FNiagaraTypeDefinition& FromType, TMap<FNiagaraTypeDefinition, UNiagaraScript*>& ConversionScriptCache, bool bIncludeConversionScripts) const;
+	NIAGARAEDITOR_API UNiagaraScript* FindConversionScript(const FNiagaraTypeDefinition& FromType, TMap<FNiagaraTypeDefinition, UNiagaraScript*>& ConversionScriptCache, bool bIncludeConversionScripts) const;
 
-	bool FilterInlineChildren(const UNiagaraStackEntry& Child) const;
+	NIAGARAEDITOR_API bool FilterInlineChildren(const UNiagaraStackEntry& Child) const;
 
 private:
 	/** The module function call which owns this input entry. NOTE: This input might not be an input to the module function
