[Niagara]

diff --git a/Engine/Plugins/FX/Niagara/Source/NiagaraEditor/Public/ViewModels/Stack/NiagaraStackFunctionInputCollection.h b/Engine/Plugins/FX/Niagara/Source/NiagaraEditor/Public/ViewModels/Stack/NiagaraStackFunctionInputCollection.h
--- a/Engine/Plugins/FX/Niagara/Source/NiagaraEditor/Public/ViewModels/Stack/NiagaraStackFunctionInputCollection.h
+++ b/Engine/Plugins/FX/Niagara/Source/NiagaraEditor/Public/ViewModels/Stack/NiagaraStackFunctionInputCollection.h
@@ -15,36 +15,36 @@ class UEdGraphPin;
 /** A base class for value collections. Values can be all kinds of input data, such as module inputs, object properties etc.
  * This has base functionality for sections.
  */
-UCLASS()
-class NIAGARAEDITOR_API UNiagaraStackValueCollection : public UNiagaraStackItemContent
+UCLASS(MinimalAPI)
+class UNiagaraStackValueCollection : public UNiagaraStackItemContent
 {
 	GENERATED_BODY()
 public:
-	void Initialize(FRequiredEntryData InRequiredEntryData, FString InOwningStackItemEditorDataKey, FString InStackEditorDataKey);
+	NIAGARAEDITOR_API void Initialize(FRequiredEntryData InRequiredEntryData, FString InOwningStackItemEditorDataKey, FString InStackEditorDataKey);
 	
 	bool GetShouldDisplayLabel() const { return bShouldDisplayLabel; }
-	void SetShouldDisplayLabel(bool bInShouldShowLabel);
+	NIAGARAEDITOR_API void SetShouldDisplayLabel(bool bInShouldShowLabel);
 
-	const TArray<FText>& GetSections() const;
+	NIAGARAEDITOR_API const TArray<FText>& GetSections() const;
 
-	FText GetActiveSection() const;
+	NIAGARAEDITOR_API FText GetActiveSection() const;
 
-	void SetActiveSection(FText InActiveSection);
+	NIAGARAEDITOR_API void SetActiveSection(FText InActiveSection);
 
-	FText GetTooltipForSection(FString Section) const;
+	NIAGARAEDITOR_API FText GetTooltipForSection(FString Section) const;
 public:
-	static FText UncategorizedName;
+	static NIAGARAEDITOR_API FText UncategorizedName;
 
-	static FText AllSectionName;
+	static NIAGARAEDITOR_API FText AllSectionName;
 protected:
-	virtual void RefreshChildrenInternal(const TArray<UNiagaraStackEntry*>& CurrentChildren, TArray<UNiagaraStackEntry*>& NewChildren, TArray<FStackIssue>& NewIssues) override;
+	NIAGARAEDITOR_API virtual void RefreshChildrenInternal(const TArray<UNiagaraStackEntry*>& CurrentChildren, TArray<UNiagaraStackEntry*>& NewChildren, TArray<FStackIssue>& NewIssues) override;
 	virtual void GetSectionsInternal(TArray<FNiagaraStackSection>& OutStackSections) const { }
-	void UpdateCachedSectionData() const;
-	virtual bool FilterByActiveSection(const UNiagaraStackEntry& Child) const;
+	NIAGARAEDITOR_API void UpdateCachedSectionData() const;
+	NIAGARAEDITOR_API virtual bool FilterByActiveSection(const UNiagaraStackEntry& Child) const;
 private:
-	virtual bool GetCanExpand() const override;
-	virtual bool GetShouldShowInStack() const override;
-	virtual int32 GetChildIndentLevel() const override;
+	NIAGARAEDITOR_API virtual bool GetCanExpand() const override;
+	NIAGARAEDITOR_API virtual bool GetShouldShowInStack() const override;
+	NIAGARAEDITOR_API virtual int32 GetChildIndentLevel() const override;
 private:
 	mutable TOptional<TArray<FText>> SectionsCache;
 	mutable TOptional<TMap<FString, TArray<FText>>> SectionToCategoryMapCache;
@@ -55,39 +55,39 @@ private:
 	bool bShouldDisplayLabel;
 };
 
-UCLASS()
-class NIAGARAEDITOR_API UNiagaraStackFunctionInputCollection : public UNiagaraStackValueCollection
+UCLASS(MinimalAPI)
+class UNiagaraStackFunctionInputCollection : public UNiagaraStackValueCollection
 {
 	GENERATED_BODY()
 
 public:
-	UNiagaraStackFunctionInputCollection();
+	NIAGARAEDITOR_API UNiagaraStackFunctionInputCollection();
 
-	UNiagaraNodeFunctionCall* GetModuleNode() const;
+	NIAGARAEDITOR_API UNiagaraNodeFunctionCall* GetModuleNode() const;
 
-	UNiagaraNodeFunctionCall* GetInputFunctionCallNode() const;
+	NIAGARAEDITOR_API UNiagaraNodeFunctionCall* GetInputFunctionCallNode() const;
 
-	void Initialize(
+	NIAGARAEDITOR_API void Initialize(
 		FRequiredEntryData InRequiredEntryData,
 		UNiagaraNodeFunctionCall& InModuleNode,
 		UNiagaraNodeFunctionCall& InInputFunctionCallNode,
 		FString InOwnerStackItemEditorDataKey);
 
 	//~ UNiagaraStackEntry interface
-	virtual FText GetDisplayName() const override;
-	virtual bool GetIsEnabled() const;
+	NIAGARAEDITOR_API virtual FText GetDisplayName() const override;
+	NIAGARAEDITOR_API virtual bool GetIsEnabled() const;
 
-	void ToClipboardFunctionInputs(UObject* InOuter, TArray<const UNiagaraClipboardFunctionInput*>& OutClipboardFunctionInputs) const;
+	NIAGARAEDITOR_API void ToClipboardFunctionInputs(UObject* InOuter, TArray<const UNiagaraClipboardFunctionInput*>& OutClipboardFunctionInputs) const;
 
-	void SetValuesFromClipboardFunctionInputs(const TArray<const UNiagaraClipboardFunctionInput*>& ClipboardFunctionInputs);
+	NIAGARAEDITOR_API void SetValuesFromClipboardFunctionInputs(const TArray<const UNiagaraClipboardFunctionInput*>& ClipboardFunctionInputs);
 
-	void GetChildInputs(TArray<UNiagaraStackFunctionInput*>& OutResult) const;
+	NIAGARAEDITOR_API void GetChildInputs(TArray<UNiagaraStackFunctionInput*>& OutResult) const;
 
-	void GetFilteredChildInputs(TArray<UNiagaraStackFunctionInput*>& OutFilteredChildInputs) const;
+	NIAGARAEDITOR_API void GetFilteredChildInputs(TArray<UNiagaraStackFunctionInput*>& OutFilteredChildInputs) const;
 
-	void GetCustomFilteredChildInputs(TArray<UNiagaraStackFunctionInput*>& OutResult, const TArray<FOnFilterChild>& CustomFilters) const;
+	NIAGARAEDITOR_API void GetCustomFilteredChildInputs(TArray<UNiagaraStackFunctionInput*>& OutResult, const TArray<FOnFilterChild>& CustomFilters) const;
 
-	TArray<UNiagaraStackFunctionInput*> GetInlineParameterInputs() const;
+	NIAGARAEDITOR_API TArray<UNiagaraStackFunctionInput*> GetInlineParameterInputs() const;
 
 private:
 
@@ -121,11 +121,11 @@ private:
 	
 	void OnScriptApplied(UNiagaraScript* NiagaraScript, FGuid Guid);
 	
-	virtual void FinalizeInternal() override;
+	NIAGARAEDITOR_API virtual void FinalizeInternal() override;
 
-	virtual void RefreshChildrenInternal(const TArray<UNiagaraStackEntry*>& CurrentChildren, TArray<UNiagaraStackEntry*>& NewChildren, TArray<FStackIssue>& NewIssues) override;
+	NIAGARAEDITOR_API virtual void RefreshChildrenInternal(const TArray<UNiagaraStackEntry*>& CurrentChildren, TArray<UNiagaraStackEntry*>& NewChildren, TArray<FStackIssue>& NewIssues) override;
 
-	virtual void GetSectionsInternal(TArray<FNiagaraStackSection>& OutStackSections) const override;
+	NIAGARAEDITOR_API virtual void GetSectionsInternal(TArray<FNiagaraStackSection>& OutStackSections) const override;
 	
 	void RefreshChildrenForFunctionCall(const TArray<UNiagaraStackEntry*>& CurrentChildren, TArray<UNiagaraStackEntry*>& NewChildren, TArray<FStackIssue>& NewIssues);
 	
diff --git a/Engine/Plugins/FX/Niagara/Source/NiagaraEditor/Public/ViewModels/Stack/NiagaraStackInputCategory.h b/Engine/Plugins/FX/Niagara/Source/NiagaraEditor/Public/ViewModels/Stack/NiagaraStackInputCategory.h
--- a/Engine/Plugins/FX/Niagara/Source/NiagaraEditor/Public/ViewModels/Stack/NiagaraStackInputCategory.h
+++ b/Engine/Plugins/FX/Niagara/Source/NiagaraEditor/Public/ViewModels/Stack/NiagaraStackInputCategory.h
@@ -18,29 +18,29 @@ enum class EStackParameterBehavior
 	Dynamic, Static
 };
 
-UCLASS()
-class NIAGARAEDITOR_API UNiagaraStackCategory : public UNiagaraStackItemContent
+UCLASS(MinimalAPI)
+class UNiagaraStackCategory : public UNiagaraStackItemContent
 {
 	GENERATED_BODY()
 
 public:
-	void Initialize(FRequiredEntryData InRequiredEntryData, FString InOwningStackItemEditorDataKey, FString InStackEditorDataKey);
+	NIAGARAEDITOR_API void Initialize(FRequiredEntryData InRequiredEntryData, FString InOwningStackItemEditorDataKey, FString InStackEditorDataKey);
 
 	//~ UNiagaraStackEntry interface
-	virtual FText GetDisplayName() const override;
-	virtual bool GetShouldShowInStack() const override;
-	virtual EStackRowStyle GetStackRowStyle() const override;
-	virtual void GetSearchItems(TArray<FStackSearchItem>& SearchItems) const override;
+	NIAGARAEDITOR_API virtual FText GetDisplayName() const override;
+	NIAGARAEDITOR_API virtual bool GetShouldShowInStack() const override;
+	NIAGARAEDITOR_API virtual EStackRowStyle GetStackRowStyle() const override;
+	NIAGARAEDITOR_API virtual void GetSearchItems(TArray<FStackSearchItem>& SearchItems) const override;
 
 	virtual bool IsTopLevelCategory() const { return false; }
-	virtual int32 GetChildIndentLevel() const override;
+	NIAGARAEDITOR_API virtual int32 GetChildIndentLevel() const override;
 
 protected:
-	virtual void RefreshChildrenInternal(const TArray<UNiagaraStackEntry*>& CurrentChildren, TArray<UNiagaraStackEntry*>& NewChildren, TArray<FStackIssue>& NewIssues) override;
+	NIAGARAEDITOR_API virtual void RefreshChildrenInternal(const TArray<UNiagaraStackEntry*>& CurrentChildren, TArray<UNiagaraStackEntry*>& NewChildren, TArray<FStackIssue>& NewIssues) override;
 private:
 private:
-	bool FilterForVisibleCondition(const UNiagaraStackEntry& Child) const;
-	bool FilterForIsInlineEditConditionToggle(const UNiagaraStackEntry& Child) const;
+	NIAGARAEDITOR_API bool FilterForVisibleCondition(const UNiagaraStackEntry& Child) const;
+	NIAGARAEDITOR_API bool FilterForIsInlineEditConditionToggle(const UNiagaraStackEntry& Child) const;
 protected:
 	bool bShouldShowInStack;
 
@@ -48,42 +48,42 @@ protected:
 	TObjectPtr<UNiagaraStackSpacer> CategorySpacer;
 };
 
-UCLASS()
-class NIAGARAEDITOR_API UNiagaraStackInputCategory : public UNiagaraStackCategory
+UCLASS(MinimalAPI)
+class UNiagaraStackInputCategory : public UNiagaraStackCategory
 {
 	GENERATED_BODY() 
 
 public:
-	void Initialize(
+	NIAGARAEDITOR_API void Initialize(
 		FRequiredEntryData InRequiredEntryData,
 		FString InputCategoryStackEditorDataKey,
 		FText InCategoryName,
 		bool bInIsTopLevelCategory,
 		FString InOwnerStackItemEditorDataKey);
 	
-	bool GetIsEnabled() const;
+	NIAGARAEDITOR_API bool GetIsEnabled() const;
 
-	virtual FText GetDisplayName() const override;
+	NIAGARAEDITOR_API virtual FText GetDisplayName() const override;
 
-	void ResetInputs();
+	NIAGARAEDITOR_API void ResetInputs();
 
-	void AddInput(UNiagaraNodeFunctionCall* InModuleNode, UNiagaraNodeFunctionCall* InInputFunctionCallNode, FName InInputParameterHandle, FNiagaraTypeDefinition InInputType, EStackParameterBehavior InParameterBehavior, TOptional<FText> InOptionalDisplayName, bool bIsHidden, bool bIsChildInput);
+	NIAGARAEDITOR_API void AddInput(UNiagaraNodeFunctionCall* InModuleNode, UNiagaraNodeFunctionCall* InInputFunctionCallNode, FName InInputParameterHandle, FNiagaraTypeDefinition InInputType, EStackParameterBehavior InParameterBehavior, TOptional<FText> InOptionalDisplayName, bool bIsHidden, bool bIsChildInput);
 
-	void SetShouldShowInStack(bool bInShouldShowInStack);
+	NIAGARAEDITOR_API void SetShouldShowInStack(bool bInShouldShowInStack);
 
-	void ToClipboardFunctionInputs(UObject* InOuter, TArray<const UNiagaraClipboardFunctionInput*>& OutClipboardFunctionInputs) const;
+	NIAGARAEDITOR_API void ToClipboardFunctionInputs(UObject* InOuter, TArray<const UNiagaraClipboardFunctionInput*>& OutClipboardFunctionInputs) const;
 
 	// We pass in the owning input collection to refresh the inputs of a category after setting a static switch value.
 	// This is done to support copy paste for chained static switches, as otherwise the stack entry for the 2nd static switch won't exist for pasting data
-	void SetStaticSwitchValuesFromClipboardFunctionInputs(const TArray<const UNiagaraClipboardFunctionInput*>& ClipboardFunctionInputs, class UNiagaraStackFunctionInputCollection& OwningFunctionCollection);
+	NIAGARAEDITOR_API void SetStaticSwitchValuesFromClipboardFunctionInputs(const TArray<const UNiagaraClipboardFunctionInput*>& ClipboardFunctionInputs, class UNiagaraStackFunctionInputCollection& OwningFunctionCollection);
 
-	void SetStandardValuesFromClipboardFunctionInputs(const TArray<const UNiagaraClipboardFunctionInput*>& ClipboardFunctionInputs);
+	NIAGARAEDITOR_API void SetStandardValuesFromClipboardFunctionInputs(const TArray<const UNiagaraClipboardFunctionInput*>& ClipboardFunctionInputs);
 
-	void GetFilteredChildInputs(TArray<UNiagaraStackFunctionInput*>& OutFilteredChildInputs) const;
+	NIAGARAEDITOR_API void GetFilteredChildInputs(TArray<UNiagaraStackFunctionInput*>& OutFilteredChildInputs) const;
 
 protected:
 	//~ UNiagaraStackEntry interface
-	virtual void RefreshChildrenInternal(const TArray<UNiagaraStackEntry*>& CurrentChildren, TArray<UNiagaraStackEntry*>& NewChildren, TArray<FStackIssue>& NewIssues) override;
+	NIAGARAEDITOR_API virtual void RefreshChildrenInternal(const TArray<UNiagaraStackEntry*>& CurrentChildren, TArray<UNiagaraStackEntry*>& NewChildren, TArray<FStackIssue>& NewIssues) override;
 	virtual bool IsTopLevelCategory() const override { return bIsTopLevelCategory; }
 private:
 	struct FInputParameterHandleAndType
@@ -107,30 +107,30 @@ private:
 
 void AddSummaryItem(UNiagaraHierarchyItemBase* HierarchyItem, UNiagaraStackEntry* Parent);
 
-UCLASS()
-class NIAGARAEDITOR_API UNiagaraStackSummaryCategory : public UNiagaraStackCategory
+UCLASS(MinimalAPI)
+class UNiagaraStackSummaryCategory : public UNiagaraStackCategory
 {
 	GENERATED_BODY()
 
 public:
 	UNiagaraStackSummaryCategory() {}
 
-	void Initialize(
+	NIAGARAEDITOR_API void Initialize(
 		FRequiredEntryData InRequiredEntryData,
 		TSharedPtr<FNiagaraHierarchyCategoryViewModel> InCategory,
 		FString InOwnerStackItemEditorDataKey);
 
 	//~ UNiagaraStackEntry interface
-	virtual FText GetDisplayName() const override;
+	NIAGARAEDITOR_API virtual FText GetDisplayName() const override;
 	virtual bool GetIsEnabled() const override { return true; }
 	
 	TWeakPtr<FNiagaraHierarchyCategoryViewModel> GetHierarchyCategory() const { return CategoryViewModelWeakPtr; }
 
 protected:
-	virtual void RefreshChildrenInternal(const TArray<UNiagaraStackEntry*>& CurrentChildren, TArray<UNiagaraStackEntry*>& NewChildren, TArray<FStackIssue>& NewIssues) override;
-	virtual bool IsTopLevelCategory() const override;
-	virtual FText GetTooltipText() const override;
-	virtual int32 GetChildIndentLevel() const override;
+	NIAGARAEDITOR_API virtual void RefreshChildrenInternal(const TArray<UNiagaraStackEntry*>& CurrentChildren, TArray<UNiagaraStackEntry*>& NewChildren, TArray<FStackIssue>& NewIssues) override;
+	NIAGARAEDITOR_API virtual bool IsTopLevelCategory() const override;
+	NIAGARAEDITOR_API virtual FText GetTooltipText() const override;
+	NIAGARAEDITOR_API virtual int32 GetChildIndentLevel() const override;
 
 private:
 	TWeakPtr<FNiagaraHierarchyCategoryViewModel> CategoryViewModelWeakPtr;
diff --git a/Engine/Plugins/FX/Niagara/Source/NiagaraEditor/Public/ViewModels/Stack/NiagaraStackItem.h b/Engine/Plugins/FX/Niagara/Source/NiagaraEditor/Public/ViewModels/Stack/NiagaraStackItem.h
--- a/Engine/Plugins/FX/Niagara/Source/NiagaraEditor/Public/ViewModels/Stack/NiagaraStackItem.h
+++ b/Engine/Plugins/FX/Niagara/Source/NiagaraEditor/Public/ViewModels/Stack/NiagaraStackItem.h
@@ -11,8 +11,8 @@ class UNiagaraStackItemFooter;
 class UNiagaraNode;
 class UNiagaraClipboardContent;
 
-UCLASS()
-class NIAGARAEDITOR_API UNiagaraStackItem : public UNiagaraStackEntry
+UCLASS(MinimalAPI)
+class UNiagaraStackItem : public UNiagaraStackEntry
 {
 	GENERATED_BODY()
 
@@ -22,17 +22,17 @@ public:
 	DECLARE_DELEGATE_ThreeParams(FOnRequestPaste, const UNiagaraClipboardContent* /* ClipboardContent */, int32 /* PasteIndex */, FText& /* OutPasteWarning */);
 
 public:
-	void Initialize(FRequiredEntryData InRequiredEntryData, FString InStackEditorDataKey);
+	NIAGARAEDITOR_API void Initialize(FRequiredEntryData InRequiredEntryData, FString InStackEditorDataKey);
 
-	virtual EStackRowStyle GetStackRowStyle() const override;
+	NIAGARAEDITOR_API virtual EStackRowStyle GetStackRowStyle() const override;
 
-	FOnModifiedGroupItems& OnModifiedGroupItems();
+	NIAGARAEDITOR_API FOnModifiedGroupItems& OnModifiedGroupItems();
 
-	void SetOnRequestCanPaste(FOnRequestCanPaste InOnRequestCanPaste);
-	void SetOnRequestPaste(FOnRequestPaste InOnRequestCanPaste);
+	NIAGARAEDITOR_API void SetOnRequestCanPaste(FOnRequestCanPaste InOnRequestCanPaste);
+	NIAGARAEDITOR_API void SetOnRequestPaste(FOnRequestPaste InOnRequestCanPaste);
 
 	virtual bool SupportsChangeEnabled() const { return false; }
-	void SetIsEnabled(bool bInIsEnabled);
+	NIAGARAEDITOR_API void SetIsEnabled(bool bInIsEnabled);
 
 	virtual bool SupportsResetToBase() const { return false; }
 	virtual bool TestCanResetToBaseWithMessage(FText& OutCanResetToBaseMessage) const { return false; }
@@ -47,20 +47,20 @@ public:
 	virtual bool GetIsInherited() const { return false; }
 
 protected:
-	virtual void RefreshChildrenInternal(const TArray<UNiagaraStackEntry*>& CurrentChildren, TArray<UNiagaraStackEntry*>& NewChildren, TArray<FStackIssue>& NewIssues) override;
+	NIAGARAEDITOR_API virtual void RefreshChildrenInternal(const TArray<UNiagaraStackEntry*>& CurrentChildren, TArray<UNiagaraStackEntry*>& NewChildren, TArray<FStackIssue>& NewIssues) override;
 
-	virtual void PostRefreshChildrenInternal() override;
+	NIAGARAEDITOR_API virtual void PostRefreshChildrenInternal() override;
 
-	virtual int32 GetChildIndentLevel() const override;
+	NIAGARAEDITOR_API virtual int32 GetChildIndentLevel() const override;
 
 	virtual void SetIsEnabledInternal(bool bInIsEnabled) { }
 
 private:
-	bool FilterAdvancedChildren(const UNiagaraStackEntry& Child) const;
+	NIAGARAEDITOR_API bool FilterAdvancedChildren(const UNiagaraStackEntry& Child) const;
 
-	bool FilterHiddenChildren(const UNiagaraStackEntry& Child) const;
+	NIAGARAEDITOR_API bool FilterHiddenChildren(const UNiagaraStackEntry& Child) const;
 
-	void ToggleShowAdvanced();
+	NIAGARAEDITOR_API void ToggleShowAdvanced();
 
 protected:
 	FOnModifiedGroupItems ModifiedGroupItemsDelegate;
@@ -72,30 +72,30 @@ private:
 	TObjectPtr<UNiagaraStackItemFooter> ItemFooter;
 };
 
-UCLASS()
-class NIAGARAEDITOR_API UNiagaraStackItemContent : public UNiagaraStackEntry
+UCLASS(MinimalAPI)
+class UNiagaraStackItemContent : public UNiagaraStackEntry
 {
 	GENERATED_BODY()
 
 public:
-	void Initialize(FRequiredEntryData InRequiredEntryData, FString InOwningStackItemEditorDataKey, FString InStackEditorDataKey);
+	NIAGARAEDITOR_API void Initialize(FRequiredEntryData InRequiredEntryData, FString InOwningStackItemEditorDataKey, FString InStackEditorDataKey);
 
-	virtual EStackRowStyle GetStackRowStyle() const override;
+	NIAGARAEDITOR_API virtual EStackRowStyle GetStackRowStyle() const override;
 
-	bool GetIsAdvanced() const;
+	NIAGARAEDITOR_API bool GetIsAdvanced() const;
 
-	bool GetIsHidden() const;
+	NIAGARAEDITOR_API bool GetIsHidden() const;
 
-	void SetIsHidden(bool bInIsHidden);
+	NIAGARAEDITOR_API void SetIsHidden(bool bInIsHidden);
 
 	// Returns true if this stack entry was changed by a user and differs from the default value
-	virtual bool HasOverridenContent() const;
+	NIAGARAEDITOR_API virtual bool HasOverridenContent() const;
 
-	bool FilterHiddenChildren(const UNiagaraStackEntry& Child) const;
+	NIAGARAEDITOR_API bool FilterHiddenChildren(const UNiagaraStackEntry& Child) const;
 protected:
-	FString GetOwnerStackItemEditorDataKey() const;
+	NIAGARAEDITOR_API FString GetOwnerStackItemEditorDataKey() const;
 
-	void SetIsAdvanced(bool bInIsAdvanced);
+	NIAGARAEDITOR_API void SetIsAdvanced(bool bInIsAdvanced);
 
 private:
 	bool FilterAdvancedChildren(const UNiagaraStackEntry& Child) const;
@@ -106,15 +106,15 @@ private:
 	bool bIsHidden;
 };
 
-UCLASS()
-class NIAGARAEDITOR_API UNiagaraStackItemTextContent : public UNiagaraStackItemContent
+UCLASS(MinimalAPI)
+class UNiagaraStackItemTextContent : public UNiagaraStackItemContent
 {
 	GENERATED_BODY()
 
 public:
-	void Initialize(FRequiredEntryData InRequiredEntryData, FText InDisplayText, FString InOwningStackItemEditorDataKey);
+	NIAGARAEDITOR_API void Initialize(FRequiredEntryData InRequiredEntryData, FText InDisplayText, FString InOwningStackItemEditorDataKey);
 
-	virtual FText GetDisplayName() const override;
+	NIAGARAEDITOR_API virtual FText GetDisplayName() const override;
 
 private:
 	FText DisplayText;
diff --git a/Engine/Plugins/FX/Niagara/Source/NiagaraEditor/Public/ViewModels/Stack/NiagaraStackItemFooter.h b/Engine/Plugins/FX/Niagara/Source/NiagaraEditor/Public/ViewModels/Stack/NiagaraStackItemFooter.h
--- a/Engine/Plugins/FX/Niagara/Source/NiagaraEditor/Public/ViewModels/Stack/NiagaraStackItemFooter.h
+++ b/Engine/Plugins/FX/Niagara/Source/NiagaraEditor/Public/ViewModels/Stack/NiagaraStackItemFooter.h
@@ -10,8 +10,8 @@ class UNiagaraStackFunctionInputCollection;
 class UNiagaraStackModuleItemOutputCollection;
 class UNiagaraNode;
 
-UCLASS()
-class NIAGARAEDITOR_API UNiagaraStackItemFooter : public UNiagaraStackEntry
+UCLASS(MinimalAPI)
+class UNiagaraStackItemFooter : public UNiagaraStackEntry
 {
 	GENERATED_BODY()
 
@@ -19,24 +19,24 @@ public:
 	DECLARE_DELEGATE(FOnToggleShowAdvanced);
 
 public:
-	void Initialize(
+	NIAGARAEDITOR_API void Initialize(
 		FRequiredEntryData InRequiredEntryData,
 		FString InOwnerStackItemEditorDataKey);
 
-	virtual bool GetCanExpand() const override;
-	virtual EStackRowStyle GetStackRowStyle() const override;
+	NIAGARAEDITOR_API virtual bool GetCanExpand() const override;
+	NIAGARAEDITOR_API virtual EStackRowStyle GetStackRowStyle() const override;
 
-	virtual bool GetIsEnabled() const override;
-	void SetIsEnabled(bool bInIsEnabled);
+	NIAGARAEDITOR_API virtual bool GetIsEnabled() const override;
+	NIAGARAEDITOR_API void SetIsEnabled(bool bInIsEnabled);
 
-	bool GetHasAdvancedContent() const;
-	void SetHasAdvancedContent(bool bHInHasAdvancedRows, bool bHasChangedAdvancedContent);
+	NIAGARAEDITOR_API bool GetHasAdvancedContent() const;
+	NIAGARAEDITOR_API void SetHasAdvancedContent(bool bHInHasAdvancedRows, bool bHasChangedAdvancedContent);
 
-	void SetOnToggleShowAdvanced(FOnToggleShowAdvanced OnToggleShowAdvanced);
+	NIAGARAEDITOR_API void SetOnToggleShowAdvanced(FOnToggleShowAdvanced OnToggleShowAdvanced);
 
-	bool GetShowAdvanced() const;
+	NIAGARAEDITOR_API bool GetShowAdvanced() const;
 
-	void ToggleShowAdvanced();
+	NIAGARAEDITOR_API void ToggleShowAdvanced();
 
 	bool HasChangedAdvancedContent() const { return bHasChangedAdvancedContent; }
 
diff --git a/Engine/Plugins/FX/Niagara/Source/NiagaraEditor/Public/ViewModels/Stack/NiagaraStackItemGroup.h b/Engine/Plugins/FX/Niagara/Source/NiagaraEditor/Public/ViewModels/Stack/NiagaraStackItemGroup.h
--- a/Engine/Plugins/FX/Niagara/Source/NiagaraEditor/Public/ViewModels/Stack/NiagaraStackItemGroup.h
+++ b/Engine/Plugins/FX/Niagara/Source/NiagaraEditor/Public/ViewModels/Stack/NiagaraStackItemGroup.h
@@ -11,43 +11,43 @@ class FNiagaraEmitterHandleViewModel;
 class UNiagaraStackItemGroupFooter;
 struct FSlateBrush;
 
-UCLASS()
-class NIAGARAEDITOR_API UNiagaraStackItemGroup : public UNiagaraStackEntry
+UCLASS(MinimalAPI)
+class UNiagaraStackItemGroup : public UNiagaraStackEntry
 {
 	GENERATED_BODY()
 
 public:
-	void Initialize(FRequiredEntryData InRequiredEntryData, FText InDisplayName, FText InToolTip, INiagaraStackItemGroupAddUtilities* InAddUtilities);
+	NIAGARAEDITOR_API void Initialize(FRequiredEntryData InRequiredEntryData, FText InDisplayName, FText InToolTip, INiagaraStackItemGroupAddUtilities* InAddUtilities);
 
 	//~ UNiagaraStackEntry interface
-	virtual FText GetDisplayName() const override;
-	virtual EStackRowStyle GetStackRowStyle() const override;
-	virtual FText GetTooltipText() const override;
+	NIAGARAEDITOR_API virtual FText GetDisplayName() const override;
+	NIAGARAEDITOR_API virtual EStackRowStyle GetStackRowStyle() const override;
+	NIAGARAEDITOR_API virtual FText GetTooltipText() const override;
 
 	virtual bool GetCanExpandInOverview() const override { return true; }
-	virtual bool GetIsEnabled() const override;
+	NIAGARAEDITOR_API virtual bool GetIsEnabled() const override;
 	virtual void SetIsEnabled(bool bEnabled) {}
 	virtual bool SupportsChangeEnabled() const { return false; }
 
 	virtual bool SupportsSecondaryIcon() const { return false; }
 	virtual const FSlateBrush* GetSecondaryIconBrush() const { return nullptr; };
 
-	INiagaraStackItemGroupAddUtilities* GetAddUtilities() const;
+	NIAGARAEDITOR_API INiagaraStackItemGroupAddUtilities* GetAddUtilities() const;
 
-	uint32 GetRecursiveStackIssuesCount() const;
-	EStackIssueSeverity GetHighestStackIssueSeverity() const;
+	NIAGARAEDITOR_API uint32 GetRecursiveStackIssuesCount() const;
+	NIAGARAEDITOR_API EStackIssueSeverity GetHighestStackIssueSeverity() const;
 	
 protected:
-	void SetDisplayName(FText InDisplayName);
+	NIAGARAEDITOR_API void SetDisplayName(FText InDisplayName);
 	
-	virtual void RefreshChildrenInternal(const TArray<UNiagaraStackEntry*>& CurrentChildren, TArray<UNiagaraStackEntry*>& NewChildren, TArray<FStackIssue>& NewIssues) override;
+	NIAGARAEDITOR_API virtual void RefreshChildrenInternal(const TArray<UNiagaraStackEntry*>& CurrentChildren, TArray<UNiagaraStackEntry*>& NewChildren, TArray<FStackIssue>& NewIssues) override;
 
-	virtual int32 GetChildIndentLevel() const override;
+	NIAGARAEDITOR_API virtual int32 GetChildIndentLevel() const override;
 
-	virtual void ChildStructureChangedInternal() override;
+	NIAGARAEDITOR_API virtual void ChildStructureChangedInternal() override;
 
 private:
-	bool FilterChildrenWithIssues(const UNiagaraStackEntry& Child) const;
+	NIAGARAEDITOR_API bool FilterChildrenWithIssues(const UNiagaraStackEntry& Child) const;
 
 private:
 	UPROPERTY()
@@ -66,15 +66,15 @@ private:
 	TWeakPtr<FNiagaraEmitterHandleViewModel> OwningEmitterHandleViewModelWeak;
 };
 
-UCLASS()
-class NIAGARAEDITOR_API UNiagaraStackItemGroupFooter : public UNiagaraStackEntry
+UCLASS(MinimalAPI)
+class UNiagaraStackItemGroupFooter : public UNiagaraStackEntry
 {
 	GENERATED_BODY()
 
 public:
-	void Initialize(FRequiredEntryData InRequiredEntryData);
+	NIAGARAEDITOR_API void Initialize(FRequiredEntryData InRequiredEntryData);
 
-	virtual EStackRowStyle GetStackRowStyle() const override;
+	NIAGARAEDITOR_API virtual EStackRowStyle GetStackRowStyle() const override;
 
-	virtual bool GetCanExpand() const;
-};
\ No newline at end of file
+	NIAGARAEDITOR_API virtual bool GetCanExpand() const;
+};
diff --git a/Engine/Plugins/FX/Niagara/Source/NiagaraEditor/Public/ViewModels/Stack/NiagaraStackModuleItem.h b/Engine/Plugins/FX/Niagara/Source/NiagaraEditor/Public/ViewModels/Stack/NiagaraStackModuleItem.h
--- a/Engine/Plugins/FX/Niagara/Source/NiagaraEditor/Public/ViewModels/Stack/NiagaraStackModuleItem.h
+++ b/Engine/Plugins/FX/Niagara/Source/NiagaraEditor/Public/ViewModels/Stack/NiagaraStackModuleItem.h
@@ -16,8 +16,8 @@ struct FAssetData;
 class UNiagaraClipboardFunctionInput;
 class INiagaraMessage;
 
-UCLASS()
-class NIAGARAEDITOR_API UNiagaraStackModuleItem : public UNiagaraStackItem
+UCLASS(MinimalAPI)
+class UNiagaraStackModuleItem : public UNiagaraStackItem
 {
 	GENERATED_BODY()
 
@@ -25,84 +25,84 @@ public:
 	DECLARE_DELEGATE_OneParam(FOnRequestDeprecationRecommended, UNiagaraStackModuleItem*);
 	DECLARE_DELEGATE_OneParam(FOnNoteModeSet, bool);
 
-	UNiagaraStackModuleItem();
+	NIAGARAEDITOR_API UNiagaraStackModuleItem();
 
-	UNiagaraNodeFunctionCall& GetModuleNode() const;
+	NIAGARAEDITOR_API UNiagaraNodeFunctionCall& GetModuleNode() const;
 
-	void Initialize(FRequiredEntryData InRequiredEntryData, INiagaraStackItemGroupAddUtilities* GroupAddUtilities, UNiagaraNodeFunctionCall& InFunctionCallNode);
+	NIAGARAEDITOR_API void Initialize(FRequiredEntryData InRequiredEntryData, INiagaraStackItemGroupAddUtilities* GroupAddUtilities, UNiagaraNodeFunctionCall& InFunctionCallNode);
 
-	virtual FText GetDisplayName() const override;
-	virtual UObject* GetDisplayedObject() const override;
-	virtual FText GetTooltipText() const override;
+	NIAGARAEDITOR_API virtual FText GetDisplayName() const override;
+	NIAGARAEDITOR_API virtual UObject* GetDisplayedObject() const override;
+	NIAGARAEDITOR_API virtual FText GetTooltipText() const override;
 
-	INiagaraStackItemGroupAddUtilities* GetGroupAddUtilities();
+	NIAGARAEDITOR_API INiagaraStackItemGroupAddUtilities* GetGroupAddUtilities();
 
-	bool CanMoveAndDelete() const;
-	bool CanRefresh() const;
-	void Refresh();
+	NIAGARAEDITOR_API bool CanMoveAndDelete() const;
+	NIAGARAEDITOR_API bool CanRefresh() const;
+	NIAGARAEDITOR_API void Refresh();
 
 	virtual bool SupportsRename() const override { return true; }
 
 	virtual bool SupportsChangeEnabled() const override { return true; }
-	virtual bool GetIsEnabled() const override;
+	NIAGARAEDITOR_API virtual bool GetIsEnabled() const override;
 
-	int32 GetModuleIndex() const;
+	NIAGARAEDITOR_API int32 GetModuleIndex() const;
 	
-	UObject* GetExternalAsset() const override;
+	NIAGARAEDITOR_API UObject* GetExternalAsset() const override;
 
-	virtual bool CanDrag() const override;
+	NIAGARAEDITOR_API virtual bool CanDrag() const override;
 
 	/** Gets the output node of this module. */
-	class UNiagaraNodeOutput* GetOutputNode() const;
+	NIAGARAEDITOR_API class UNiagaraNodeOutput* GetOutputNode() const;
 
-	bool CanAddInput(FNiagaraVariable InputParameter) const;
+	NIAGARAEDITOR_API bool CanAddInput(FNiagaraVariable InputParameter) const;
 
-	void AddInput(FNiagaraVariable InputParameter);
+	NIAGARAEDITOR_API void AddInput(FNiagaraVariable InputParameter);
 
 	/** Gets whether or not a module script reassignment is pending.  This can happen when trying to fix modules which are missing their scripts. */
-	bool GetIsModuleScriptReassignmentPending() const;
+	NIAGARAEDITOR_API bool GetIsModuleScriptReassignmentPending() const;
 
 	/** Gets whether or not a module script reassignment should be be pending. */
-	void SetIsModuleScriptReassignmentPending(bool bIsPending);
+	NIAGARAEDITOR_API void SetIsModuleScriptReassignmentPending(bool bIsPending);
 
 	/** Reassigns the function script for the module without resetting the inputs. */
-	void ReassignModuleScript(UNiagaraScript* ModuleScript);
+	NIAGARAEDITOR_API void ReassignModuleScript(UNiagaraScript* ModuleScript);
 	
-	void ChangeScriptVersion(FGuid NewScriptVersion);
+	NIAGARAEDITOR_API void ChangeScriptVersion(FGuid NewScriptVersion);
 
-	void SetInputValuesFromClipboardFunctionInputs(const TArray<const UNiagaraClipboardFunctionInput*>& ClipboardFunctionInputs);
+	NIAGARAEDITOR_API void SetInputValuesFromClipboardFunctionInputs(const TArray<const UNiagaraClipboardFunctionInput*>& ClipboardFunctionInputs);
 
-	void GetParameterInputs(TArray<class UNiagaraStackFunctionInput*>& OutResult) const;
-	TArray<class UNiagaraStackFunctionInput*> GetInlineParameterInputs() const;
+	NIAGARAEDITOR_API void GetParameterInputs(TArray<class UNiagaraStackFunctionInput*>& OutResult) const;
+	NIAGARAEDITOR_API TArray<class UNiagaraStackFunctionInput*> GetInlineParameterInputs() const;
 
 	virtual bool SupportsCut() const override { return true; }
-	virtual bool TestCanCutWithMessage(FText& OutMessage) const override;
-	virtual FText GetCutTransactionText() const override;
-	virtual void CopyForCut(UNiagaraClipboardContent* ClipboardContent) const override;
-	virtual void RemoveForCut() override;
+	NIAGARAEDITOR_API virtual bool TestCanCutWithMessage(FText& OutMessage) const override;
+	NIAGARAEDITOR_API virtual FText GetCutTransactionText() const override;
+	NIAGARAEDITOR_API virtual void CopyForCut(UNiagaraClipboardContent* ClipboardContent) const override;
+	NIAGARAEDITOR_API virtual void RemoveForCut() override;
 
 	virtual bool SupportsCopy() const override { return true; }
-	virtual bool TestCanCopyWithMessage(FText& OutMessage) const override;
-	virtual void Copy(UNiagaraClipboardContent* ClipboardContent) const override;
+	NIAGARAEDITOR_API virtual bool TestCanCopyWithMessage(FText& OutMessage) const override;
+	NIAGARAEDITOR_API virtual void Copy(UNiagaraClipboardContent* ClipboardContent) const override;
 
 	virtual bool SupportsPaste() const override { return true; }
-	virtual bool TestCanPasteWithMessage(const UNiagaraClipboardContent* ClipboardContent, FText& OutMessage) const override;
-	virtual FText GetPasteTransactionText(const UNiagaraClipboardContent* ClipboardContent) const override;
-	virtual void Paste(const UNiagaraClipboardContent* ClipboardContent, FText& OutPasteWarning) override;
+	NIAGARAEDITOR_API virtual bool TestCanPasteWithMessage(const UNiagaraClipboardContent* ClipboardContent, FText& OutMessage) const override;
+	NIAGARAEDITOR_API virtual FText GetPasteTransactionText(const UNiagaraClipboardContent* ClipboardContent) const override;
+	NIAGARAEDITOR_API virtual void Paste(const UNiagaraClipboardContent* ClipboardContent, FText& OutPasteWarning) override;
 
 	virtual bool SupportsDelete() const override { return true; }
-	virtual bool TestCanDeleteWithMessage(FText& OutCanDeleteMessage) const override;
-	virtual FText GetDeleteTransactionText() const override;
-	virtual void Delete() override;
+	NIAGARAEDITOR_API virtual bool TestCanDeleteWithMessage(FText& OutCanDeleteMessage) const override;
+	NIAGARAEDITOR_API virtual FText GetDeleteTransactionText() const override;
+	NIAGARAEDITOR_API virtual void Delete() override;
 
 	virtual bool SupportsInheritance() const override { return true; }
-	virtual bool GetIsInherited() const override;
-	virtual FText GetInheritanceMessage() const override;
+	NIAGARAEDITOR_API virtual bool GetIsInherited() const override;
+	NIAGARAEDITOR_API virtual FText GetInheritanceMessage() const override;
 
 	virtual bool SupportsSummaryView() const override { return true; }
-	virtual FNiagaraHierarchyIdentity DetermineSummaryIdentity() const override;
+	NIAGARAEDITOR_API virtual FNiagaraHierarchyIdentity DetermineSummaryIdentity() const override;
 	
-	bool IsScratchModule() const;
+	NIAGARAEDITOR_API bool IsScratchModule() const;
 
 	void SetOnRequestDeprecationRecommended(FOnRequestDeprecationRecommended InOnRequest)
 	{
@@ -125,34 +125,34 @@ public:
 		return bIsNoteModeActive;
 	}
 
-	bool IsDebugDrawEnabled() const;
-	void SetDebugDrawEnabled(bool bInEnabled);
+	NIAGARAEDITOR_API bool IsDebugDrawEnabled() const;
+	NIAGARAEDITOR_API void SetDebugDrawEnabled(bool bInEnabled);
 
 	FOnNoteModeSet& OnNoteModeSet() { return OnNoteModeSetDelegate;}
 
-	bool OpenSourceAsset() const;
+	NIAGARAEDITOR_API bool OpenSourceAsset() const;
 
 protected:
 	FOnRequestDeprecationRecommended DeprecationDelegate;
 
-	virtual void FinalizeInternal() override;
-	virtual void RefreshChildrenInternal(const TArray<UNiagaraStackEntry*>& CurrentChildren, TArray<UNiagaraStackEntry*>& NewChildren, TArray<FStackIssue>& NewIssues) override;
-	virtual void SetIsEnabledInternal(bool bInIsEnabled) override;
+	NIAGARAEDITOR_API virtual void FinalizeInternal() override;
+	NIAGARAEDITOR_API virtual void RefreshChildrenInternal(const TArray<UNiagaraStackEntry*>& CurrentChildren, TArray<UNiagaraStackEntry*>& NewChildren, TArray<FStackIssue>& NewIssues) override;
+	NIAGARAEDITOR_API virtual void SetIsEnabledInternal(bool bInIsEnabled) override;
 
-	virtual TOptional<FDropRequestResponse> CanDropInternal(const FDropRequest& DropRequest) override;
-	virtual TOptional<FDropRequestResponse> DropInternal(const FDropRequest& DropRequest) override;
+	NIAGARAEDITOR_API virtual TOptional<FDropRequestResponse> CanDropInternal(const FDropRequest& DropRequest) override;
+	NIAGARAEDITOR_API virtual TOptional<FDropRequestResponse> DropInternal(const FDropRequest& DropRequest) override;
 
-	virtual const FCollectedUsageData& GetCollectedUsageData() const override;
+	NIAGARAEDITOR_API virtual const FCollectedUsageData& GetCollectedUsageData() const override;
 
 private:
-	bool FilterOutputCollection(const UNiagaraStackEntry& Child) const;
-	bool FilterOutputCollectionChild(const UNiagaraStackEntry& Child) const;
-	bool FilterLinkedInputCollection(const UNiagaraStackEntry& Child) const;
-	bool FilterLinkedInputCollectionChild(const UNiagaraStackEntry& Child) const;
-	void RefreshIssues(TArray<FStackIssue>& NewIssues);
-	void RefreshIsEnabled();
-	void OnMessageManagerRefresh(const TArray<TSharedRef<const INiagaraMessage>>& NewMessages);
-	FStackIssueFixDelegate GetUpgradeVersionFix();
+	NIAGARAEDITOR_API bool FilterOutputCollection(const UNiagaraStackEntry& Child) const;
+	NIAGARAEDITOR_API bool FilterOutputCollectionChild(const UNiagaraStackEntry& Child) const;
+	NIAGARAEDITOR_API bool FilterLinkedInputCollection(const UNiagaraStackEntry& Child) const;
+	NIAGARAEDITOR_API bool FilterLinkedInputCollectionChild(const UNiagaraStackEntry& Child) const;
+	NIAGARAEDITOR_API void RefreshIssues(TArray<FStackIssue>& NewIssues);
+	NIAGARAEDITOR_API void RefreshIsEnabled();
+	NIAGARAEDITOR_API void OnMessageManagerRefresh(const TArray<TSharedRef<const INiagaraMessage>>& NewMessages);
+	NIAGARAEDITOR_API FStackIssueFixDelegate GetUpgradeVersionFix();
 
 private:
 	UNiagaraNodeOutput* OutputNode;
diff --git a/Engine/Plugins/FX/Niagara/Source/NiagaraEditor/Public/ViewModels/Stack/NiagaraStackModuleItemLinkedInputCollection.h b/Engine/Plugins/FX/Niagara/Source/NiagaraEditor/Public/ViewModels/Stack/NiagaraStackModuleItemLinkedInputCollection.h
--- a/Engine/Plugins/FX/Niagara/Source/NiagaraEditor/Public/ViewModels/Stack/NiagaraStackModuleItemLinkedInputCollection.h
+++ b/Engine/Plugins/FX/Niagara/Source/NiagaraEditor/Public/ViewModels/Stack/NiagaraStackModuleItemLinkedInputCollection.h
@@ -7,25 +7,25 @@
 
 class UNiagaraNodeFunctionCall;
 
-UCLASS()
-class NIAGARAEDITOR_API UNiagaraStackModuleItemLinkedInputCollection : public UNiagaraStackEntry
+UCLASS(MinimalAPI)
+class UNiagaraStackModuleItemLinkedInputCollection : public UNiagaraStackEntry
 {
 	GENERATED_BODY()
 
 public:
-	UNiagaraStackModuleItemLinkedInputCollection();
+	NIAGARAEDITOR_API UNiagaraStackModuleItemLinkedInputCollection();
 
-	void Initialize(FRequiredEntryData InRequiredEntryData, UNiagaraNodeFunctionCall& InFunctionCallNode);
+	NIAGARAEDITOR_API void Initialize(FRequiredEntryData InRequiredEntryData, UNiagaraNodeFunctionCall& InFunctionCallNode);
 
 	//~ UNiagaraStackEntry interface
-	virtual FText GetDisplayName() const override;
-	virtual bool IsExpandedByDefault() const override;
-	virtual bool GetIsEnabled() const override;
-	virtual EStackRowStyle GetStackRowStyle() const override;
-	virtual bool GetShouldShowInStack() const override;
+	NIAGARAEDITOR_API virtual FText GetDisplayName() const override;
+	NIAGARAEDITOR_API virtual bool IsExpandedByDefault() const override;
+	NIAGARAEDITOR_API virtual bool GetIsEnabled() const override;
+	NIAGARAEDITOR_API virtual EStackRowStyle GetStackRowStyle() const override;
+	NIAGARAEDITOR_API virtual bool GetShouldShowInStack() const override;
 
 protected:
-	virtual void RefreshChildrenInternal(const TArray<UNiagaraStackEntry*>& CurrentChildren, TArray<UNiagaraStackEntry*>& NewChildren, TArray<FStackIssue>& NewIssues) override;
+	NIAGARAEDITOR_API virtual void RefreshChildrenInternal(const TArray<UNiagaraStackEntry*>& CurrentChildren, TArray<UNiagaraStackEntry*>& NewChildren, TArray<FStackIssue>& NewIssues) override;
 
 private:
 	void GetCurrentChangeIds(FGuid& OutOwningGraphChangeId, FGuid& OutFunctionGraphChangeId) const;
diff --git a/Engine/Plugins/FX/Niagara/Source/NiagaraEditor/Public/ViewModels/Stack/NiagaraStackModuleItemOutput.h b/Engine/Plugins/FX/Niagara/Source/NiagaraEditor/Public/ViewModels/Stack/NiagaraStackModuleItemOutput.h
--- a/Engine/Plugins/FX/Niagara/Source/NiagaraEditor/Public/ViewModels/Stack/NiagaraStackModuleItemOutput.h
+++ b/Engine/Plugins/FX/Niagara/Source/NiagaraEditor/Public/ViewModels/Stack/NiagaraStackModuleItemOutput.h
@@ -10,35 +10,35 @@
 class UNiagaraNodeFunctionCall;
 
 /** Represents a single module Output in the module stack view model. */
-UCLASS()
-class NIAGARAEDITOR_API UNiagaraStackModuleItemOutput : public UNiagaraStackEntry
+UCLASS(MinimalAPI)
+class UNiagaraStackModuleItemOutput : public UNiagaraStackEntry
 {
 	GENERATED_BODY()
 
 public:
-	UNiagaraStackModuleItemOutput();
+	NIAGARAEDITOR_API UNiagaraStackModuleItemOutput();
 
 	/** 
 	 * Sets the Output data for this entry.
 	 * @param InFunctionCallNode The function call node representing the module in the stack graph which owns this Output.
 	 * @param InOutputParameterHandle The Namespace.Name handle for the Output to the owning module.
 	 */
-	void Initialize(FRequiredEntryData InRequiredEntryData, UNiagaraNodeFunctionCall& InFunctionCallNode, FName InOutputParameterHandle, FNiagaraTypeDefinition InOutputType);
+	NIAGARAEDITOR_API void Initialize(FRequiredEntryData InRequiredEntryData, UNiagaraNodeFunctionCall& InFunctionCallNode, FName InOutputParameterHandle, FNiagaraTypeDefinition InOutputType);
 
 	//~ UNiagaraStackEntry interface
-	virtual FText GetDisplayName() const override;
-	virtual FText GetTooltipText() const override;
-	virtual bool GetIsEnabled() const override;
-	virtual EStackRowStyle GetStackRowStyle() const override;
-	virtual void GetSearchItems(TArray<FStackSearchItem>& SearchItems) const override;
+	NIAGARAEDITOR_API virtual FText GetDisplayName() const override;
+	NIAGARAEDITOR_API virtual FText GetTooltipText() const override;
+	NIAGARAEDITOR_API virtual bool GetIsEnabled() const override;
+	NIAGARAEDITOR_API virtual EStackRowStyle GetStackRowStyle() const override;
+	NIAGARAEDITOR_API virtual void GetSearchItems(TArray<FStackSearchItem>& SearchItems) const override;
 
 	/** Gets the parameter handle which defined this Output in the module. */
-	const FNiagaraParameterHandle& GetOutputParameterHandle() const;
+	NIAGARAEDITOR_API const FNiagaraParameterHandle& GetOutputParameterHandle() const;
 
 	/** Gets the assigned parameter handle as displayable text. */
-	FText GetOutputParameterHandleText() const;
+	NIAGARAEDITOR_API FText GetOutputParameterHandleText() const;
 
-	virtual const FCollectedUsageData& GetCollectedUsageData() const override;
+	NIAGARAEDITOR_API virtual const FCollectedUsageData& GetCollectedUsageData() const override;
 
 private:
 	/** The function call node which represents the emitter which owns this Output */
diff --git a/Engine/Plugins/FX/Niagara/Source/NiagaraEditor/Public/ViewModels/Stack/NiagaraStackModuleItemOutputCollection.h b/Engine/Plugins/FX/Niagara/Source/NiagaraEditor/Public/ViewModels/Stack/NiagaraStackModuleItemOutputCollection.h
--- a/Engine/Plugins/FX/Niagara/Source/NiagaraEditor/Public/ViewModels/Stack/NiagaraStackModuleItemOutputCollection.h
+++ b/Engine/Plugins/FX/Niagara/Source/NiagaraEditor/Public/ViewModels/Stack/NiagaraStackModuleItemOutputCollection.h
@@ -7,24 +7,24 @@
 
 class UNiagaraNodeFunctionCall;
 
-UCLASS()
-class NIAGARAEDITOR_API UNiagaraStackModuleItemOutputCollection : public UNiagaraStackEntry
+UCLASS(MinimalAPI)
+class UNiagaraStackModuleItemOutputCollection : public UNiagaraStackEntry
 {
 	GENERATED_BODY()
 
 public:
-	UNiagaraStackModuleItemOutputCollection();
+	NIAGARAEDITOR_API UNiagaraStackModuleItemOutputCollection();
 
-	void Initialize(FRequiredEntryData InRequiredEntryData, UNiagaraNodeFunctionCall& InFunctionCallNode);
+	NIAGARAEDITOR_API void Initialize(FRequiredEntryData InRequiredEntryData, UNiagaraNodeFunctionCall& InFunctionCallNode);
 	
 	//~ UNiagaraStackEntry interface
-	virtual FText GetDisplayName() const override;
-	virtual bool IsExpandedByDefault() const override;
-	virtual bool GetIsEnabled() const override;
-	virtual EStackRowStyle GetStackRowStyle() const override;
+	NIAGARAEDITOR_API virtual FText GetDisplayName() const override;
+	NIAGARAEDITOR_API virtual bool IsExpandedByDefault() const override;
+	NIAGARAEDITOR_API virtual bool GetIsEnabled() const override;
+	NIAGARAEDITOR_API virtual EStackRowStyle GetStackRowStyle() const override;
 
 protected:
-	virtual void RefreshChildrenInternal(const TArray<UNiagaraStackEntry*>& CurrentChildren, TArray<UNiagaraStackEntry*>& NewChildren, TArray<FStackIssue>& NewIssues) override;
+	NIAGARAEDITOR_API virtual void RefreshChildrenInternal(const TArray<UNiagaraStackEntry*>& CurrentChildren, TArray<UNiagaraStackEntry*>& NewChildren, TArray<FStackIssue>& NewIssues) override;
 
 private:
 	void GetCurrentChangeIds(FGuid& OutOwningGraphChangeId, FGuid& OutFunctionGraphChangeId) const;
diff --git a/Engine/Plugins/FX/Niagara/Source/NiagaraEditor/Public/ViewModels/Stack/NiagaraStackObject.h b/Engine/Plugins/FX/Niagara/Source/NiagaraEditor/Public/ViewModels/Stack/NiagaraStackObject.h
--- a/Engine/Plugins/FX/Niagara/Source/NiagaraEditor/Public/ViewModels/Stack/NiagaraStackObject.h
+++ b/Engine/Plugins/FX/Niagara/Source/NiagaraEditor/Public/ViewModels/Stack/NiagaraStackObject.h
@@ -12,8 +12,8 @@ class UNiagaraNode;
 class IDetailTreeNode;
 class INiagaraMessage;
 
-UCLASS()
-class NIAGARAEDITOR_API UNiagaraStackObject : public UNiagaraStackItemContent, public FNotifyHook
+UCLASS(MinimalAPI)
+class UNiagaraStackObject : public UNiagaraStackItemContent, public FNotifyHook
 {
 	GENERATED_BODY()
 		
@@ -21,7 +21,7 @@ public:
 	DECLARE_DELEGATE_TwoParams(FOnSelectRootNodes, TArray<TSharedRef<IDetailTreeNode>>, TArray<TSharedRef<IDetailTreeNode>>*);
 
 public:
-	UNiagaraStackObject();
+	NIAGARAEDITOR_API UNiagaraStackObject();
 
 	/** Initialized an entry for displaying a UObject in the niagara stack. 
 		@param InRequiredEntryData Struct with required data for all stack entries.
@@ -29,35 +29,35 @@ public:
 		@param bInIsTopLevelObject Whether or not the object being displayed should be treated as a top level object shown in the root of the stack.  This can be used to make layout decisions such as category styling.
 		@param InOwnerStackItemEditorDataKey The stack editor data key of the owning stack entry. 
 		@param InOwningNiagaraNode An optional niagara node which owns this object. */
-	void Initialize(FRequiredEntryData InRequiredEntryData, UObject* InObject, bool bInIsTopLevelObject, FString InOwnerStackItemEditorDataKey, UNiagaraNode* InOwningNiagaraNode = nullptr);
+	NIAGARAEDITOR_API void Initialize(FRequiredEntryData InRequiredEntryData, UObject* InObject, bool bInIsTopLevelObject, FString InOwnerStackItemEditorDataKey, UNiagaraNode* InOwningNiagaraNode = nullptr);
 
-	void SetOnSelectRootNodes(FOnSelectRootNodes OnSelectRootNodes);
+	NIAGARAEDITOR_API void SetOnSelectRootNodes(FOnSelectRootNodes OnSelectRootNodes);
 
-	void RegisterInstancedCustomPropertyLayout(UStruct* Class, FOnGetDetailCustomizationInstance DetailLayoutDelegate);
-	void RegisterInstancedCustomPropertyTypeLayout(FName PropertyTypeName, FOnGetPropertyTypeCustomizationInstance PropertyTypeLayoutDelegate, TSharedPtr<IPropertyTypeIdentifier> Identifier = nullptr);
+	NIAGARAEDITOR_API void RegisterInstancedCustomPropertyLayout(UStruct* Class, FOnGetDetailCustomizationInstance DetailLayoutDelegate);
+	NIAGARAEDITOR_API void RegisterInstancedCustomPropertyTypeLayout(FName PropertyTypeName, FOnGetPropertyTypeCustomizationInstance PropertyTypeLayoutDelegate, TSharedPtr<IPropertyTypeIdentifier> Identifier = nullptr);
 
-	UObject* GetObject();
+	NIAGARAEDITOR_API UObject* GetObject();
 
 	//~ FNotifyHook interface
-	virtual void NotifyPostChange(const FPropertyChangedEvent& PropertyChangedEvent, FProperty* PropertyThatChanged) override;
+	NIAGARAEDITOR_API virtual void NotifyPostChange(const FPropertyChangedEvent& PropertyChangedEvent, FProperty* PropertyThatChanged) override;
 
 	//~ UNiagaraStackEntry interface
-	virtual bool GetIsEnabled() const override;
-	virtual bool GetShouldShowInStack() const override;
-	virtual UObject* GetDisplayedObject() const override;
+	NIAGARAEDITOR_API virtual bool GetIsEnabled() const override;
+	NIAGARAEDITOR_API virtual bool GetShouldShowInStack() const override;
+	NIAGARAEDITOR_API virtual UObject* GetDisplayedObject() const override;
 
 	void SetObjectGuid(FGuid InGuid) { ObjectGuid = InGuid; }
 	virtual bool SupportsSummaryView() const override { return ObjectGuid.IsSet() && ObjectGuid->IsValid(); }
-	virtual FNiagaraHierarchyIdentity DetermineSummaryIdentity() const override;
+	NIAGARAEDITOR_API virtual FNiagaraHierarchyIdentity DetermineSummaryIdentity() const override;
 protected:
-	virtual void FinalizeInternal() override;
+	NIAGARAEDITOR_API virtual void FinalizeInternal() override;
 
-	virtual void RefreshChildrenInternal(const TArray<UNiagaraStackEntry*>& CurrentChildren, TArray<UNiagaraStackEntry*>& NewChildren, TArray<FStackIssue>& NewIssues) override;
-	virtual void PostRefreshChildrenInternal() override;
+	NIAGARAEDITOR_API virtual void RefreshChildrenInternal(const TArray<UNiagaraStackEntry*>& CurrentChildren, TArray<UNiagaraStackEntry*>& NewChildren, TArray<FStackIssue>& NewIssues) override;
+	NIAGARAEDITOR_API virtual void PostRefreshChildrenInternal() override;
 
 private:
-	void PropertyRowsRefreshed();
-	void OnMessageManagerRefresh(const TArray<TSharedRef<const INiagaraMessage>>& NewMessages);
+	NIAGARAEDITOR_API void PropertyRowsRefreshed();
+	NIAGARAEDITOR_API void OnMessageManagerRefresh(const TArray<TSharedRef<const INiagaraMessage>>& NewMessages);
 
 private:
 	struct FRegisteredClassCustomization
diff --git a/Engine/Plugins/FX/Niagara/Source/NiagaraEditor/Public/ViewModels/Stack/NiagaraStackPropertyRow.h b/Engine/Plugins/FX/Niagara/Source/NiagaraEditor/Public/ViewModels/Stack/NiagaraStackPropertyRow.h
--- a/Engine/Plugins/FX/Niagara/Source/NiagaraEditor/Public/ViewModels/Stack/NiagaraStackPropertyRow.h
+++ b/Engine/Plugins/FX/Niagara/Source/NiagaraEditor/Public/ViewModels/Stack/NiagaraStackPropertyRow.h
@@ -8,42 +8,42 @@
 class IDetailTreeNode;
 class UNiagaraNode;
 
-UCLASS()
-class NIAGARAEDITOR_API UNiagaraStackPropertyRow : public UNiagaraStackItemContent
+UCLASS(MinimalAPI)
+class UNiagaraStackPropertyRow : public UNiagaraStackItemContent
 {
 	GENERATED_BODY()
 		
 public:
-	void Initialize(FRequiredEntryData InRequiredEntryData, TSharedRef<IDetailTreeNode> InDetailTreeNode, bool bInIsTopLevelProperty, FString InOwnerStackItemEditorDataKey, FString InOwnerStackEditorDataKey, UNiagaraNode* InOwningNiagaraNode);
+	NIAGARAEDITOR_API void Initialize(FRequiredEntryData InRequiredEntryData, TSharedRef<IDetailTreeNode> InDetailTreeNode, bool bInIsTopLevelProperty, FString InOwnerStackItemEditorDataKey, FString InOwnerStackEditorDataKey, UNiagaraNode* InOwningNiagaraNode);
 	
-	virtual EStackRowStyle GetStackRowStyle() const override;
+	NIAGARAEDITOR_API virtual EStackRowStyle GetStackRowStyle() const override;
 
-	TSharedRef<IDetailTreeNode> GetDetailTreeNode() const;
+	NIAGARAEDITOR_API TSharedRef<IDetailTreeNode> GetDetailTreeNode() const;
 
-	virtual bool GetIsEnabled() const override;
+	NIAGARAEDITOR_API virtual bool GetIsEnabled() const override;
 
-	virtual bool HasOverridenContent() const override;
+	NIAGARAEDITOR_API virtual bool HasOverridenContent() const override;
 
-	virtual bool IsExpandedByDefault() const override;
+	NIAGARAEDITOR_API virtual bool IsExpandedByDefault() const override;
 
-	virtual bool CanDrag() const override;
+	NIAGARAEDITOR_API virtual bool CanDrag() const override;
 
 	void SetOwnerGuid(TOptional<FGuid> InGuid) { OwnerGuid = InGuid; }
 protected:
-	virtual void FinalizeInternal() override;
+	NIAGARAEDITOR_API virtual void FinalizeInternal() override;
 
-	virtual void RefreshChildrenInternal(const TArray<UNiagaraStackEntry*>& CurrentChildren, TArray<UNiagaraStackEntry*>& NewChildren, TArray<FStackIssue>& NewIssues) override;
+	NIAGARAEDITOR_API virtual void RefreshChildrenInternal(const TArray<UNiagaraStackEntry*>& CurrentChildren, TArray<UNiagaraStackEntry*>& NewChildren, TArray<FStackIssue>& NewIssues) override;
 
-	virtual int32 GetChildIndentLevel() const override;
+	NIAGARAEDITOR_API virtual int32 GetChildIndentLevel() const override;
 
-	virtual void GetSearchItems(TArray<FStackSearchItem>& SearchItems) const override;
+	NIAGARAEDITOR_API virtual void GetSearchItems(TArray<FStackSearchItem>& SearchItems) const override;
 
-	virtual TOptional<FDropRequestResponse> CanDropInternal(const FDropRequest& DropRequest) override;
+	NIAGARAEDITOR_API virtual TOptional<FDropRequestResponse> CanDropInternal(const FDropRequest& DropRequest) override;
 
-	virtual TOptional<FDropRequestResponse> DropInternal(const FDropRequest& DropRequest) override;
+	NIAGARAEDITOR_API virtual TOptional<FDropRequestResponse> DropInternal(const FDropRequest& DropRequest) override;
 
-	virtual bool SupportsSummaryView() const override;
-	virtual FNiagaraHierarchyIdentity DetermineSummaryIdentity() const override;
+	NIAGARAEDITOR_API virtual bool SupportsSummaryView() const override;
+	NIAGARAEDITOR_API virtual FNiagaraHierarchyIdentity DetermineSummaryIdentity() const override;
 private:
 	TSharedPtr<IDetailTreeNode> DetailTreeNode;
 	UNiagaraNode* OwningNiagaraNode;
@@ -57,4 +57,4 @@ private:
 
 	bool bCannotEditInThisContext;
 	bool bIsTopLevelProperty;
-};
\ No newline at end of file
+};
diff --git a/Engine/Plugins/FX/Niagara/Source/NiagaraEditor/Public/ViewModels/Stack/NiagaraStackRenderItemGroup.h b/Engine/Plugins/FX/Niagara/Source/NiagaraEditor/Public/ViewModels/Stack/NiagaraStackRenderItemGroup.h
--- a/Engine/Plugins/FX/Niagara/Source/NiagaraEditor/Public/ViewModels/Stack/NiagaraStackRenderItemGroup.h
+++ b/Engine/Plugins/FX/Niagara/Source/NiagaraEditor/Public/ViewModels/Stack/NiagaraStackRenderItemGroup.h
@@ -10,30 +10,30 @@ class UNiagaraRendererProperties;
 class UNiagaraEmitter;
 class UNiagaraClipboardContent;
 
-UCLASS()
-class NIAGARAEDITOR_API UNiagaraStackRenderItemGroup : public UNiagaraStackItemGroup
+UCLASS(MinimalAPI)
+class UNiagaraStackRenderItemGroup : public UNiagaraStackItemGroup
 {
 	GENERATED_BODY()
 
 public:
-	void Initialize(FRequiredEntryData InRequiredEntryData);
+	NIAGARAEDITOR_API void Initialize(FRequiredEntryData InRequiredEntryData);
 
 	virtual bool SupportsPaste() const override { return true; }
-	virtual bool TestCanPasteWithMessage(const UNiagaraClipboardContent* ClipboardContent, FText& OutMessage) const override;
-	virtual FText GetPasteTransactionText(const UNiagaraClipboardContent* ClipboardContent) const override;
-	virtual void Paste(const UNiagaraClipboardContent* ClipboardContent, FText& OutPasteWarning) override;
+	NIAGARAEDITOR_API virtual bool TestCanPasteWithMessage(const UNiagaraClipboardContent* ClipboardContent, FText& OutMessage) const override;
+	NIAGARAEDITOR_API virtual FText GetPasteTransactionText(const UNiagaraClipboardContent* ClipboardContent) const override;
+	NIAGARAEDITOR_API virtual void Paste(const UNiagaraClipboardContent* ClipboardContent, FText& OutPasteWarning) override;
 
 protected:
-	virtual void RefreshChildrenInternal(const TArray<UNiagaraStackEntry*>& CurrentChildren, TArray<UNiagaraStackEntry*>& NewChildren, TArray<FStackIssue>& NewIssues) override;
-	virtual void FinalizeInternal() override;
+	NIAGARAEDITOR_API virtual void RefreshChildrenInternal(const TArray<UNiagaraStackEntry*>& CurrentChildren, TArray<UNiagaraStackEntry*>& NewChildren, TArray<FStackIssue>& NewIssues) override;
+	NIAGARAEDITOR_API virtual void FinalizeInternal() override;
 private:
-	void EmitterRenderersChanged();
+	NIAGARAEDITOR_API void EmitterRenderersChanged();
 
-	bool ChildRequestCanPaste(const UNiagaraClipboardContent* ClipboardContent,FText& OutCanPasteMessage);
-	void ChildRequestPaste(const UNiagaraClipboardContent* ClipboardContent, int32 PasteIndex, FText& OutPasteWarning);
+	NIAGARAEDITOR_API bool ChildRequestCanPaste(const UNiagaraClipboardContent* ClipboardContent,FText& OutCanPasteMessage);
+	NIAGARAEDITOR_API void ChildRequestPaste(const UNiagaraClipboardContent* ClipboardContent, int32 PasteIndex, FText& OutPasteWarning);
 
 private:
-	void OnRendererAdded(UNiagaraRendererProperties* RendererProperties) const;
+	NIAGARAEDITOR_API void OnRendererAdded(UNiagaraRendererProperties* RendererProperties) const;
 	TSharedPtr<INiagaraStackItemGroupAddUtilities> AddUtilities;
 
 	FVersionedNiagaraEmitterWeakPtr EmitterWeak;
diff --git a/Engine/Plugins/FX/Niagara/Source/NiagaraEditor/Public/ViewModels/Stack/NiagaraStackRendererItem.h b/Engine/Plugins/FX/Niagara/Source/NiagaraEditor/Public/ViewModels/Stack/NiagaraStackRendererItem.h
--- a/Engine/Plugins/FX/Niagara/Source/NiagaraEditor/Public/ViewModels/Stack/NiagaraStackRendererItem.h
+++ b/Engine/Plugins/FX/Niagara/Source/NiagaraEditor/Public/ViewModels/Stack/NiagaraStackRendererItem.h
@@ -12,85 +12,85 @@ class UNiagaraEmitter;
 class UNiagaraRendererProperties;
 class UNiagaraStackObject;
 
-UCLASS()
-class NIAGARAEDITOR_API UNiagaraStackRendererItem : public UNiagaraStackItem
+UCLASS(MinimalAPI)
+class UNiagaraStackRendererItem : public UNiagaraStackItem
 {
 	GENERATED_BODY()
 		
 public:
-	UNiagaraStackRendererItem();
+	NIAGARAEDITOR_API UNiagaraStackRendererItem();
 
-	void Initialize(FRequiredEntryData InRequiredEntryData, UNiagaraRendererProperties* InRendererProperties);
+	NIAGARAEDITOR_API void Initialize(FRequiredEntryData InRequiredEntryData, UNiagaraRendererProperties* InRendererProperties);
 
-	UNiagaraRendererProperties* GetRendererProperties();
-	const UNiagaraRendererProperties* GetRendererProperties() const;
+	NIAGARAEDITOR_API UNiagaraRendererProperties* GetRendererProperties();
+	NIAGARAEDITOR_API const UNiagaraRendererProperties* GetRendererProperties() const;
 	
-	virtual FText GetDisplayName() const override;
+	NIAGARAEDITOR_API virtual FText GetDisplayName() const override;
 
 	virtual bool SupportsCut() const override { return true; }
-	virtual bool TestCanCutWithMessage(FText& OutMessage) const override;
-	virtual FText GetCutTransactionText() const override;
-	virtual void CopyForCut(UNiagaraClipboardContent* ClipboardContent) const override;
-	virtual void RemoveForCut() override;
+	NIAGARAEDITOR_API virtual bool TestCanCutWithMessage(FText& OutMessage) const override;
+	NIAGARAEDITOR_API virtual FText GetCutTransactionText() const override;
+	NIAGARAEDITOR_API virtual void CopyForCut(UNiagaraClipboardContent* ClipboardContent) const override;
+	NIAGARAEDITOR_API virtual void RemoveForCut() override;
 
 	virtual bool SupportsCopy() const override { return true; }
-	virtual bool TestCanCopyWithMessage(FText& OutMessage) const override;
-	virtual void Copy(UNiagaraClipboardContent* ClipboardContent) const override;
+	NIAGARAEDITOR_API virtual bool TestCanCopyWithMessage(FText& OutMessage) const override;
+	NIAGARAEDITOR_API virtual void Copy(UNiagaraClipboardContent* ClipboardContent) const override;
 
 	virtual bool SupportsPaste() const override { return true; }
-	virtual bool TestCanPasteWithMessage(const UNiagaraClipboardContent* ClipboardContent, FText& OutMessage) const override;
-	virtual FText GetPasteTransactionText(const UNiagaraClipboardContent* ClipboardContent) const override;
-	virtual void Paste(const UNiagaraClipboardContent* ClipboardContent, FText& OutPasteWarning) override;
+	NIAGARAEDITOR_API virtual bool TestCanPasteWithMessage(const UNiagaraClipboardContent* ClipboardContent, FText& OutMessage) const override;
+	NIAGARAEDITOR_API virtual FText GetPasteTransactionText(const UNiagaraClipboardContent* ClipboardContent) const override;
+	NIAGARAEDITOR_API virtual void Paste(const UNiagaraClipboardContent* ClipboardContent, FText& OutPasteWarning) override;
 
 	virtual bool SupportsDelete() const override { return true; }
-	virtual bool TestCanDeleteWithMessage(FText& OutCanDeleteMessage) const override;
-	virtual FText GetDeleteTransactionText() const override;
-	virtual void Delete() override;
+	NIAGARAEDITOR_API virtual bool TestCanDeleteWithMessage(FText& OutCanDeleteMessage) const override;
+	NIAGARAEDITOR_API virtual FText GetDeleteTransactionText() const override;
+	NIAGARAEDITOR_API virtual void Delete() override;
 
 	virtual bool SupportsInheritance() const override { return true; }
-	virtual bool GetIsInherited() const override;
-	virtual FText GetInheritanceMessage() const override;
+	NIAGARAEDITOR_API virtual bool GetIsInherited() const override;
+	NIAGARAEDITOR_API virtual FText GetInheritanceMessage() const override;
 
 	virtual bool SupportsSummaryView() const override { return true; }
-	virtual FNiagaraHierarchyIdentity DetermineSummaryIdentity() const override;
+	NIAGARAEDITOR_API virtual FNiagaraHierarchyIdentity DetermineSummaryIdentity() const override;
 	
 	virtual bool SupportsRename() const override { return true; }
 
-	bool HasBaseRenderer() const;
+	NIAGARAEDITOR_API bool HasBaseRenderer() const;
 
 	virtual bool SupportsChangeEnabled() const override { return true; }
-	virtual bool GetIsEnabled() const override;
+	NIAGARAEDITOR_API virtual bool GetIsEnabled() const override;
 
 	virtual EIconMode GetSupportedIconMode() const override { return EIconMode::Brush; }
-	virtual const FSlateBrush* GetIconBrush() const override;
+	NIAGARAEDITOR_API virtual const FSlateBrush* GetIconBrush() const override;
 
 	virtual bool SupportsResetToBase() const override { return true; }
-	virtual bool TestCanResetToBaseWithMessage(FText& OutCanResetToBaseMessage) const override;
-	virtual void ResetToBase() override;
+	NIAGARAEDITOR_API virtual bool TestCanResetToBaseWithMessage(FText& OutCanResetToBaseMessage) const override;
+	NIAGARAEDITOR_API virtual void ResetToBase() override;
 
-	static TArray<FNiagaraVariable> GetMissingVariables(UNiagaraRendererProperties* RendererProperties, const FVersionedNiagaraEmitterData* EmitterData);
-	static bool AddMissingVariable(const FVersionedNiagaraEmitterData* EmitterData, const FNiagaraVariable& Variable);
+	static NIAGARAEDITOR_API TArray<FNiagaraVariable> GetMissingVariables(UNiagaraRendererProperties* RendererProperties, const FVersionedNiagaraEmitterData* EmitterData);
+	static NIAGARAEDITOR_API bool AddMissingVariable(const FVersionedNiagaraEmitterData* EmitterData, const FNiagaraVariable& Variable);
 
-	bool IsExcludedFromScalability() const;
-	bool IsOwningEmitterExcludedFromScalability() const;
+	NIAGARAEDITOR_API bool IsExcludedFromScalability() const;
+	NIAGARAEDITOR_API bool IsOwningEmitterExcludedFromScalability() const;
 
-	virtual const FCollectedUsageData& GetCollectedUsageData() const override;
+	NIAGARAEDITOR_API virtual const FCollectedUsageData& GetCollectedUsageData() const override;
 
-	bool CanMoveRendererUp() const;
-	void MoveRendererUp() const;
-	bool CanMoveRendererDown() const;
-	void MoveRendererDown() const;
+	NIAGARAEDITOR_API bool CanMoveRendererUp() const;
+	NIAGARAEDITOR_API void MoveRendererUp() const;
+	NIAGARAEDITOR_API bool CanMoveRendererDown() const;
+	NIAGARAEDITOR_API void MoveRendererDown() const;
 protected:
-	virtual void FinalizeInternal() override;
+	NIAGARAEDITOR_API virtual void FinalizeInternal() override;
 
-	virtual void RefreshChildrenInternal(const TArray<UNiagaraStackEntry*>& CurrentChildren, TArray<UNiagaraStackEntry*>& NewChildren, TArray<FStackIssue>& NewIssues) override;
+	NIAGARAEDITOR_API virtual void RefreshChildrenInternal(const TArray<UNiagaraStackEntry*>& CurrentChildren, TArray<UNiagaraStackEntry*>& NewChildren, TArray<FStackIssue>& NewIssues) override;
 
-	virtual void SetIsEnabledInternal(bool bInIsEnabled) override;
+	NIAGARAEDITOR_API virtual void SetIsEnabledInternal(bool bInIsEnabled) override;
 
 private:
-	void RendererChanged();
-	void RefreshIssues(TArray<FStackIssue>& NewIssues);
-	void ProcessRendererIssues(const TArray<FNiagaraRendererFeedback>& InIssues, EStackIssueSeverity Severity, TArray<FStackIssue>& OutIssues);
+	NIAGARAEDITOR_API void RendererChanged();
+	NIAGARAEDITOR_API void RefreshIssues(TArray<FStackIssue>& NewIssues);
+	NIAGARAEDITOR_API void ProcessRendererIssues(const TArray<FNiagaraRendererFeedback>& InIssues, EStackIssueSeverity Severity, TArray<FStackIssue>& OutIssues);
 
 private:
 	TWeakObjectPtr<UNiagaraRendererProperties> RendererProperties;
diff --git a/Engine/Plugins/FX/Niagara/Source/NiagaraEditor/Public/ViewModels/Stack/NiagaraStackRoot.h b/Engine/Plugins/FX/Niagara/Source/NiagaraEditor/Public/ViewModels/Stack/NiagaraStackRoot.h
--- a/Engine/Plugins/FX/Niagara/Source/NiagaraEditor/Public/ViewModels/Stack/NiagaraStackRoot.h
+++ b/Engine/Plugins/FX/Niagara/Source/NiagaraEditor/Public/ViewModels/Stack/NiagaraStackRoot.h
@@ -15,19 +15,19 @@ class UNiagaraStackRenderItemGroup;
 class UNiagaraStackEmitterSummaryGroup;
 class UNiagaraStackSummaryViewCollapseButton;
 
-UCLASS()
-class NIAGARAEDITOR_API UNiagaraStackRoot : public UNiagaraStackEntry
+UCLASS(MinimalAPI)
+class UNiagaraStackRoot : public UNiagaraStackEntry
 {
 	GENERATED_BODY()
 
 public:
-	UNiagaraStackRoot();
+	NIAGARAEDITOR_API UNiagaraStackRoot();
 	
-	void Initialize(FRequiredEntryData InRequiredEntryData, bool bInIncludeSystemInformation, bool bInIncludeEmitterInformation);
-	virtual void FinalizeInternal() override;
+	NIAGARAEDITOR_API void Initialize(FRequiredEntryData InRequiredEntryData, bool bInIncludeSystemInformation, bool bInIncludeEmitterInformation);
+	NIAGARAEDITOR_API virtual void FinalizeInternal() override;
 
-	virtual bool GetCanExpand() const override;
-	virtual bool GetShouldShowInStack() const override;
+	NIAGARAEDITOR_API virtual bool GetCanExpand() const override;
+	NIAGARAEDITOR_API virtual bool GetShouldShowInStack() const override;
 	UNiagaraStackRenderItemGroup* GetRenderGroup() const
 	{
 		return RenderGroup;
@@ -39,11 +39,11 @@ public:
 	}
 	
 protected:
-	virtual void RefreshChildrenInternal(const TArray<UNiagaraStackEntry*>& CurrentChildren, TArray<UNiagaraStackEntry*>& NewChildren, TArray<FStackIssue>& NewIssues) override;
+	NIAGARAEDITOR_API virtual void RefreshChildrenInternal(const TArray<UNiagaraStackEntry*>& CurrentChildren, TArray<UNiagaraStackEntry*>& NewChildren, TArray<FStackIssue>& NewIssues) override;
 
 private:
-	void EmitterArraysChanged();
-	void OnSummaryViewStateChanged();
+	NIAGARAEDITOR_API void EmitterArraysChanged();
+	NIAGARAEDITOR_API void OnSummaryViewStateChanged();
 
 private:
 	UPROPERTY()
diff --git a/Engine/Plugins/FX/Niagara/Source/NiagaraEditor/Public/ViewModels/Stack/NiagaraStackScriptItemGroup.h b/Engine/Plugins/FX/Niagara/Source/NiagaraEditor/Public/ViewModels/Stack/NiagaraStackScriptItemGroup.h
--- a/Engine/Plugins/FX/Niagara/Source/NiagaraEditor/Public/ViewModels/Stack/NiagaraStackScriptItemGroup.h
+++ b/Engine/Plugins/FX/Niagara/Source/NiagaraEditor/Public/ViewModels/Stack/NiagaraStackScriptItemGroup.h
@@ -13,13 +13,13 @@ class UNiagaraNodeFunctionCall;
 class FScriptItemGroupAddUtilities;
 class UEdGraph;
 
-UCLASS()
-class NIAGARAEDITOR_API UNiagaraStackScriptItemGroup : public UNiagaraStackItemGroup
+UCLASS(MinimalAPI)
+class UNiagaraStackScriptItemGroup : public UNiagaraStackItemGroup
 {
 	GENERATED_BODY()
 
 public:
-	void Initialize(
+	NIAGARAEDITOR_API void Initialize(
 		FRequiredEntryData InRequiredEntryData,
 		FText InDisplayName,
 		FText InToolTip,
@@ -29,71 +29,71 @@ public:
 
 	ENiagaraScriptUsage GetScriptUsage() const { return ScriptUsage; }
 	FGuid GetScriptUsageId() const { return ScriptUsageId; }
-	UNiagaraNodeOutput* GetScriptOutputNode() const;
+	NIAGARAEDITOR_API UNiagaraNodeOutput* GetScriptOutputNode() const;
 
-	virtual UObject* GetDisplayedObject() const override;
+	NIAGARAEDITOR_API virtual UObject* GetDisplayedObject() const override;
 
 	virtual bool SupportsPaste() const override { return true; }
-	virtual bool TestCanPasteWithMessage(const UNiagaraClipboardContent* ClipboardContent, FText& OutMessage) const override;
-	virtual FText GetPasteTransactionText(const UNiagaraClipboardContent* ClipboardContent) const override;
-	virtual void Paste(const UNiagaraClipboardContent* ClipboardContent, FText& OutPasteWarning) override;
+	NIAGARAEDITOR_API virtual bool TestCanPasteWithMessage(const UNiagaraClipboardContent* ClipboardContent, FText& OutMessage) const override;
+	NIAGARAEDITOR_API virtual FText GetPasteTransactionText(const UNiagaraClipboardContent* ClipboardContent) const override;
+	NIAGARAEDITOR_API virtual void Paste(const UNiagaraClipboardContent* ClipboardContent, FText& OutPasteWarning) override;
 
 protected:
-	virtual void RefreshChildrenInternal(const TArray<UNiagaraStackEntry*>& CurrentChildren, TArray<UNiagaraStackEntry*>& NewChildren, TArray<FStackIssue>& NewIssues) override;
+	NIAGARAEDITOR_API virtual void RefreshChildrenInternal(const TArray<UNiagaraStackEntry*>& CurrentChildren, TArray<UNiagaraStackEntry*>& NewChildren, TArray<FStackIssue>& NewIssues) override;
 
-	virtual void FinalizeInternal() override;
+	NIAGARAEDITOR_API virtual void FinalizeInternal() override;
 
-	virtual TOptional<FDropRequestResponse> CanDropInternal(const FDropRequest& DropRequest) override;
+	NIAGARAEDITOR_API virtual TOptional<FDropRequestResponse> CanDropInternal(const FDropRequest& DropRequest) override;
 
-	virtual TOptional<FDropRequestResponse> DropInternal(const FDropRequest& DropRequest) override;
+	NIAGARAEDITOR_API virtual TOptional<FDropRequestResponse> DropInternal(const FDropRequest& DropRequest) override;
 
-	virtual TOptional<FDropRequestResponse> ChildRequestCanDropInternal(const UNiagaraStackEntry& TargetChild, const FDropRequest& DropRequest) override;
+	NIAGARAEDITOR_API virtual TOptional<FDropRequestResponse> ChildRequestCanDropInternal(const UNiagaraStackEntry& TargetChild, const FDropRequest& DropRequest) override;
 
-	virtual TOptional<FDropRequestResponse> ChildRequestDropInternal(const UNiagaraStackEntry& TargetChild, const FDropRequest& DropRequest) override;
+	NIAGARAEDITOR_API virtual TOptional<FDropRequestResponse> ChildRequestDropInternal(const UNiagaraStackEntry& TargetChild, const FDropRequest& DropRequest) override;
 
 private:
-	void ItemAdded(UNiagaraNodeFunctionCall* AddedModule);
+	NIAGARAEDITOR_API void ItemAdded(UNiagaraNodeFunctionCall* AddedModule);
 
-	void ChildModifiedGroupItems();
+	NIAGARAEDITOR_API void ChildModifiedGroupItems();
 
-	bool ChildRequestCanPaste(const UNiagaraClipboardContent* ClipboardContent, FText& OutCanPasteMessage);
+	NIAGARAEDITOR_API bool ChildRequestCanPaste(const UNiagaraClipboardContent* ClipboardContent, FText& OutCanPasteMessage);
 
-	void ChildRequestPaste(const UNiagaraClipboardContent* ClipboardContent, int32 PasteIndex, FText& OutPasteWarning);
+	NIAGARAEDITOR_API void ChildRequestPaste(const UNiagaraClipboardContent* ClipboardContent, int32 PasteIndex, FText& OutPasteWarning);
 
-	void ChildRequestDeprecatedRecommendation(class UNiagaraStackModuleItem* TargetChild);
+	NIAGARAEDITOR_API void ChildRequestDeprecatedRecommendation(class UNiagaraStackModuleItem* TargetChild);
 
-	void OnScriptGraphChanged(const struct FEdGraphEditAction& InAction);
+	NIAGARAEDITOR_API void OnScriptGraphChanged(const struct FEdGraphEditAction& InAction);
 
-	void OnSystemScriptCompiled(UNiagaraScript* InScript, const FGuid& ScriptVersion);
+	NIAGARAEDITOR_API void OnSystemScriptCompiled(UNiagaraScript* InScript, const FGuid& ScriptVersion);
 
-	void OnParticleScriptCompiled(UNiagaraScript* InScript, const FGuid& ScriptVersion);
+	NIAGARAEDITOR_API void OnParticleScriptCompiled(UNiagaraScript* InScript, const FGuid& ScriptVersion);
 
-	TOptional<FDropRequestResponse> CanDropOnTarget(const UNiagaraStackEntry& TargetEntry, const FDropRequest& DropRequest);
+	NIAGARAEDITOR_API TOptional<FDropRequestResponse> CanDropOnTarget(const UNiagaraStackEntry& TargetEntry, const FDropRequest& DropRequest);
 
-	TOptional<FDropRequestResponse> CanDropEntriesOnTarget(const UNiagaraStackEntry& TargetEntry, const FDropRequest& DropRequest);
+	NIAGARAEDITOR_API TOptional<FDropRequestResponse> CanDropEntriesOnTarget(const UNiagaraStackEntry& TargetEntry, const FDropRequest& DropRequest);
 
-	TOptional<FDropRequestResponse> CanDropAssetsOnTarget(const UNiagaraStackEntry& TargetEntry, const FDropRequest& DropRequest);
+	NIAGARAEDITOR_API TOptional<FDropRequestResponse> CanDropAssetsOnTarget(const UNiagaraStackEntry& TargetEntry, const FDropRequest& DropRequest);
 
-	TOptional<FDropRequestResponse> CanDropScriptsOnTarget(const UNiagaraStackEntry& TargetEntry, const FDropRequest& DropRequest);
+	NIAGARAEDITOR_API TOptional<FDropRequestResponse> CanDropScriptsOnTarget(const UNiagaraStackEntry& TargetEntry, const FDropRequest& DropRequest);
 
-	TOptional<FDropRequestResponse> CanDropParameterOnTarget(const UNiagaraStackEntry& TargetEntry, const FDropRequest& DropRequest);
+	NIAGARAEDITOR_API TOptional<FDropRequestResponse> CanDropParameterOnTarget(const UNiagaraStackEntry& TargetEntry, const FDropRequest& DropRequest);
 
-	TOptional<FDropRequestResponse> DropOnTarget(const UNiagaraStackEntry& TargetEntry, const FDropRequest& DropRequest);
+	NIAGARAEDITOR_API TOptional<FDropRequestResponse> DropOnTarget(const UNiagaraStackEntry& TargetEntry, const FDropRequest& DropRequest);
 
-	TOptional<FDropRequestResponse> DropEntriesOnTarget(const UNiagaraStackEntry& TargetEntry, const FDropRequest& DropRequest);
+	NIAGARAEDITOR_API TOptional<FDropRequestResponse> DropEntriesOnTarget(const UNiagaraStackEntry& TargetEntry, const FDropRequest& DropRequest);
 
-	TOptional<FDropRequestResponse> DropAssetsOnTarget(const UNiagaraStackEntry& TargetEntry, const FDropRequest& DropRequest);
+	NIAGARAEDITOR_API TOptional<FDropRequestResponse> DropAssetsOnTarget(const UNiagaraStackEntry& TargetEntry, const FDropRequest& DropRequest);
 
-	TOptional<FDropRequestResponse> DropScriptsOnTarget(const UNiagaraStackEntry& TargetEntry, const FDropRequest& DropRequest);
+	NIAGARAEDITOR_API TOptional<FDropRequestResponse> DropScriptsOnTarget(const UNiagaraStackEntry& TargetEntry, const FDropRequest& DropRequest);
 
-	TOptional<FDropRequestResponse> DropParameterOnTarget(const UNiagaraStackEntry& TargetEntry, const FDropRequest& DropRequest);
+	NIAGARAEDITOR_API TOptional<FDropRequestResponse> DropParameterOnTarget(const UNiagaraStackEntry& TargetEntry, const FDropRequest& DropRequest);
 
 protected:
 	TWeakPtr<FNiagaraScriptViewModel> ScriptViewModel;
-	void RefreshIssues(TArray<FStackIssue>& NewIssues);
+	NIAGARAEDITOR_API void RefreshIssues(TArray<FStackIssue>& NewIssues);
 
 private:
-	void PasteModules(const UNiagaraClipboardContent* ClipboardContent, int32 PasteIndex, FText& OutPasteWarning);
+	NIAGARAEDITOR_API void PasteModules(const UNiagaraClipboardContent* ClipboardContent, int32 PasteIndex, FText& OutPasteWarning);
 
 private:
 	TSharedPtr<FScriptItemGroupAddUtilities> AddUtilities;
diff --git a/Engine/Plugins/FX/Niagara/Source/NiagaraEditor/Public/ViewModels/Stack/NiagaraStackSelection.h b/Engine/Plugins/FX/Niagara/Source/NiagaraEditor/Public/ViewModels/Stack/NiagaraStackSelection.h
--- a/Engine/Plugins/FX/Niagara/Source/NiagaraEditor/Public/ViewModels/Stack/NiagaraStackSelection.h
+++ b/Engine/Plugins/FX/Niagara/Source/NiagaraEditor/Public/ViewModels/Stack/NiagaraStackSelection.h
@@ -5,23 +5,23 @@
 #include "NiagaraStackEntry.h"
 #include "NiagaraStackSelection.generated.h"
 
-UCLASS()
-class NIAGARAEDITOR_API UNiagaraStackSelection : public UNiagaraStackEntry
+UCLASS(MinimalAPI)
+class UNiagaraStackSelection : public UNiagaraStackEntry
 {
 	GENERATED_BODY()
 
 public:
-	void Initialize(FRequiredEntryData InRequiredEntryData);
+	NIAGARAEDITOR_API void Initialize(FRequiredEntryData InRequiredEntryData);
 
-	virtual bool GetCanExpand() const override;
+	NIAGARAEDITOR_API virtual bool GetCanExpand() const override;
 
-	virtual bool GetShouldShowInStack() const override;
+	NIAGARAEDITOR_API virtual bool GetShouldShowInStack() const override;
 
-	void SetSelectedEntries(const TArray<UNiagaraStackEntry*>& InSelectedEntries);
+	NIAGARAEDITOR_API void SetSelectedEntries(const TArray<UNiagaraStackEntry*>& InSelectedEntries);
 
 protected:
-	virtual void RefreshChildrenInternal(const TArray<UNiagaraStackEntry*>& CurrentChildren, TArray<UNiagaraStackEntry*>& NewChildren, TArray<FStackIssue>& NewIssues) override;
+	NIAGARAEDITOR_API virtual void RefreshChildrenInternal(const TArray<UNiagaraStackEntry*>& CurrentChildren, TArray<UNiagaraStackEntry*>& NewChildren, TArray<FStackIssue>& NewIssues) override;
 
 private:
 	TArray<TWeakObjectPtr<UNiagaraStackEntry>> SelectedEntries;
-};
\ No newline at end of file
+};
diff --git a/Engine/Plugins/FX/Niagara/Source/NiagaraEditor/Public/ViewModels/Stack/NiagaraStackSimulationStageGroup.h b/Engine/Plugins/FX/Niagara/Source/NiagaraEditor/Public/ViewModels/Stack/NiagaraStackSimulationStageGroup.h
--- a/Engine/Plugins/FX/Niagara/Source/NiagaraEditor/Public/ViewModels/Stack/NiagaraStackSimulationStageGroup.h
+++ b/Engine/Plugins/FX/Niagara/Source/NiagaraEditor/Public/ViewModels/Stack/NiagaraStackSimulationStageGroup.h
@@ -13,34 +13,34 @@ class UNiagaraStackObject;
 struct FNiagaraEventScriptProperties;
 class IDetailTreeNode;
 
-UCLASS()
-class NIAGARAEDITOR_API UNiagaraStackSimulationStagePropertiesItem : public UNiagaraStackItem
+UCLASS(MinimalAPI)
+class UNiagaraStackSimulationStagePropertiesItem : public UNiagaraStackItem
 {
 	GENERATED_BODY()
 		 
 public:
-	void Initialize(FRequiredEntryData InRequiredEntryData, UNiagaraSimulationStageBase* InSimulationStage);
+	NIAGARAEDITOR_API void Initialize(FRequiredEntryData InRequiredEntryData, UNiagaraSimulationStageBase* InSimulationStage);
 
-	virtual FText GetDisplayName() const override;
+	NIAGARAEDITOR_API virtual FText GetDisplayName() const override;
 
 	virtual bool SupportsResetToBase() const override { return true; }
-	virtual bool TestCanResetToBaseWithMessage(FText& OutCanResetToBaseMessage) const override;
-	virtual void ResetToBase() override;
+	NIAGARAEDITOR_API virtual bool TestCanResetToBaseWithMessage(FText& OutCanResetToBaseMessage) const override;
+	NIAGARAEDITOR_API virtual void ResetToBase() override;
 
-	void SetSimulationStageEnabled(bool bIsEnabled);
+	NIAGARAEDITOR_API void SetSimulationStageEnabled(bool bIsEnabled);
 
 	TWeakObjectPtr<UNiagaraSimulationStageBase> GetSimulationStage() const { return SimulationStage; }
 protected:
-	virtual void FinalizeInternal() override;
+	NIAGARAEDITOR_API virtual void FinalizeInternal() override;
 
-	virtual void RefreshChildrenInternal(const TArray<UNiagaraStackEntry*>& CurrentChildren, TArray<UNiagaraStackEntry*>& NewChildren, TArray<FStackIssue>& NewIssues) override;
+	NIAGARAEDITOR_API virtual void RefreshChildrenInternal(const TArray<UNiagaraStackEntry*>& CurrentChildren, TArray<UNiagaraStackEntry*>& NewChildren, TArray<FStackIssue>& NewIssues) override;
 
 	virtual bool SupportsSummaryView() const override { return true; }
-	virtual FNiagaraHierarchyIdentity DetermineSummaryIdentity() const override;
+	NIAGARAEDITOR_API virtual FNiagaraHierarchyIdentity DetermineSummaryIdentity() const override;
 private:
-	void SimulationStagePropertiesChanged();
+	NIAGARAEDITOR_API void SimulationStagePropertiesChanged();
 
-	bool HasBaseSimulationStage() const;
+	NIAGARAEDITOR_API bool HasBaseSimulationStage() const;
 
 private:
 	TWeakObjectPtr<UNiagaraSimulationStageBase> SimulationStage;
@@ -53,9 +53,9 @@ private:
 	TObjectPtr<UNiagaraStackObject> SimulationStageObject;
 };
 
-UCLASS()
+UCLASS(MinimalAPI)
 /** Meant to contain a single binding of a Emitter::EventScriptProperties to the stack.*/
-class NIAGARAEDITOR_API UNiagaraStackSimulationStageGroup : public UNiagaraStackScriptItemGroup
+class UNiagaraStackSimulationStageGroup : public UNiagaraStackScriptItemGroup
 {
 	GENERATED_BODY()
 
@@ -63,44 +63,44 @@ public:
 	DECLARE_DELEGATE(FOnModifiedSimulationStages);
 
 public:
-	void Initialize(
+	NIAGARAEDITOR_API void Initialize(
 		FRequiredEntryData InRequiredEntryData,
 		TSharedRef<FNiagaraScriptViewModel> InScriptViewModel,
 		UNiagaraSimulationStageBase* InSimulationStage);
 
-	UNiagaraSimulationStageBase* GetSimulationStage() const;
+	NIAGARAEDITOR_API UNiagaraSimulationStageBase* GetSimulationStage() const;
 	const TObjectPtr<UNiagaraStackSimulationStagePropertiesItem>& GetSimulationStagePropertiesItem() const { return SimulationStageProperties; }
-	void SetOnModifiedSimulationStages(FOnModifiedSimulationStages OnModifiedSimulationStages);
+	NIAGARAEDITOR_API void SetOnModifiedSimulationStages(FOnModifiedSimulationStages OnModifiedSimulationStages);
 
 	virtual bool SupportsDelete() const override { return true; }
-	virtual bool TestCanDeleteWithMessage(FText& OutCanDeleteMessage) const override;
-	virtual void Delete() override;
+	NIAGARAEDITOR_API virtual bool TestCanDeleteWithMessage(FText& OutCanDeleteMessage) const override;
+	NIAGARAEDITOR_API virtual void Delete() override;
 
 	virtual bool SupportsInheritance() const override { return true; }
-	virtual bool GetIsInherited() const override;
-	virtual FText GetInheritanceMessage() const override;
+	NIAGARAEDITOR_API virtual bool GetIsInherited() const override;
+	NIAGARAEDITOR_API virtual FText GetInheritanceMessage() const override;
 
 	virtual bool CanDrag() const override { return true; }
 
-	virtual bool GetIsEnabled() const override;
-	virtual void SetIsEnabled(bool bEnabled) override;
+	NIAGARAEDITOR_API virtual bool GetIsEnabled() const override;
+	NIAGARAEDITOR_API virtual void SetIsEnabled(bool bEnabled) override;
 	virtual bool SupportsChangeEnabled() const override { return true; }
 
 protected:
-	virtual void FinalizeInternal() override;
+	NIAGARAEDITOR_API virtual void FinalizeInternal() override;
 
-	virtual void RefreshChildrenInternal(const TArray<UNiagaraStackEntry*>& CurrentChildren, TArray<UNiagaraStackEntry*>& NewChildren, TArray<FStackIssue>& NewIssues) override;
+	NIAGARAEDITOR_API virtual void RefreshChildrenInternal(const TArray<UNiagaraStackEntry*>& CurrentChildren, TArray<UNiagaraStackEntry*>& NewChildren, TArray<FStackIssue>& NewIssues) override;
 
-	virtual TOptional<FDropRequestResponse> CanDropInternal(const FDropRequest& DropRequest) override;
+	NIAGARAEDITOR_API virtual TOptional<FDropRequestResponse> CanDropInternal(const FDropRequest& DropRequest) override;
 
-	virtual TOptional<FDropRequestResponse> DropInternal(const FDropRequest& DropRequest) override;
+	NIAGARAEDITOR_API virtual TOptional<FDropRequestResponse> DropInternal(const FDropRequest& DropRequest) override;
 
 	virtual bool SupportsSummaryView() const override { return true; }
-	virtual FNiagaraHierarchyIdentity DetermineSummaryIdentity() const override;
+	NIAGARAEDITOR_API virtual FNiagaraHierarchyIdentity DetermineSummaryIdentity() const override;
 private:
-	void SimulationStagePropertiesChanged();
+	NIAGARAEDITOR_API void SimulationStagePropertiesChanged();
 
-	bool HasBaseSimulationStage() const;
+	NIAGARAEDITOR_API bool HasBaseSimulationStage() const;
 	
 private:
 	TWeakObjectPtr<UNiagaraSimulationStageBase> SimulationStage;
diff --git a/Engine/Plugins/FX/Niagara/Source/NiagaraEditor/Public/ViewModels/Stack/NiagaraStackSummaryViewInputCollection.h b/Engine/Plugins/FX/Niagara/Source/NiagaraEditor/Public/ViewModels/Stack/NiagaraStackSummaryViewInputCollection.h
--- a/Engine/Plugins/FX/Niagara/Source/NiagaraEditor/Public/ViewModels/Stack/NiagaraStackSummaryViewInputCollection.h
+++ b/Engine/Plugins/FX/Niagara/Source/NiagaraEditor/Public/ViewModels/Stack/NiagaraStackSummaryViewInputCollection.h
@@ -7,8 +7,8 @@
 #include "ViewModels/HierarchyEditor/NiagaraHierarchyViewModelBase.h"
 #include "NiagaraStackSummaryViewInputCollection.generated.h"
 
-UCLASS()
-class NIAGARAEDITOR_API UNiagaraStackSummaryViewCollection : public UNiagaraStackValueCollection
+UCLASS(MinimalAPI)
+class UNiagaraStackSummaryViewCollection : public UNiagaraStackValueCollection
 {
 	GENERATED_BODY()
 		
@@ -18,21 +18,21 @@ public:
 public:
 	UNiagaraStackSummaryViewCollection() {}
 
-	void Initialize(FRequiredEntryData InRequiredEntryData, FVersionedNiagaraEmitterWeakPtr InEmitter, FString InOwningStackItemEditorDataKey);
-	virtual void FinalizeInternal() override;
+	NIAGARAEDITOR_API void Initialize(FRequiredEntryData InRequiredEntryData, FVersionedNiagaraEmitterWeakPtr InEmitter, FString InOwningStackItemEditorDataKey);
+	NIAGARAEDITOR_API virtual void FinalizeInternal() override;
 
-	virtual FText GetDisplayName() const override;
-	virtual bool GetIsEnabled() const override;
+	NIAGARAEDITOR_API virtual FText GetDisplayName() const override;
+	NIAGARAEDITOR_API virtual bool GetIsEnabled() const override;
 	
 protected:
-	virtual void RefreshChildrenInternal(const TArray<UNiagaraStackEntry*>& CurrentChildren, TArray<UNiagaraStackEntry*>& NewChildren, TArray<FStackIssue>& NewIssues) override;
+	NIAGARAEDITOR_API virtual void RefreshChildrenInternal(const TArray<UNiagaraStackEntry*>& CurrentChildren, TArray<UNiagaraStackEntry*>& NewChildren, TArray<FStackIssue>& NewIssues) override;
 
-	virtual void GetSectionsInternal(TArray<FNiagaraStackSection>& OutStackSections) const override;
-	const TArray<UNiagaraHierarchySection*>& GetHierarchySections() const;
+	NIAGARAEDITOR_API virtual void GetSectionsInternal(TArray<FNiagaraStackSection>& OutStackSections) const override;
+	NIAGARAEDITOR_API const TArray<UNiagaraHierarchySection*>& GetHierarchySections() const;
 
-	virtual bool FilterByActiveSection(const UNiagaraStackEntry& Child) const override;
+	NIAGARAEDITOR_API virtual bool FilterByActiveSection(const UNiagaraStackEntry& Child) const override;
 private:
-	void OnViewStateChanged();
+	NIAGARAEDITOR_API void OnViewStateChanged();
 private:
 
 	FVersionedNiagaraEmitterWeakPtr Emitter;
diff --git a/Engine/Plugins/FX/Niagara/Source/NiagaraEditor/Public/ViewModels/Stack/NiagaraStackSystemPropertiesItem.h b/Engine/Plugins/FX/Niagara/Source/NiagaraEditor/Public/ViewModels/Stack/NiagaraStackSystemPropertiesItem.h
--- a/Engine/Plugins/FX/Niagara/Source/NiagaraEditor/Public/ViewModels/Stack/NiagaraStackSystemPropertiesItem.h
+++ b/Engine/Plugins/FX/Niagara/Source/NiagaraEditor/Public/ViewModels/Stack/NiagaraStackSystemPropertiesItem.h
@@ -8,25 +8,25 @@
 
 class UNiagaraStackObject;
 
-UCLASS()
-class NIAGARAEDITOR_API UNiagaraStackSystemPropertiesItem : public UNiagaraStackItem
+UCLASS(MinimalAPI)
+class UNiagaraStackSystemPropertiesItem : public UNiagaraStackItem
 {
 	GENERATED_BODY()
 
 public:
-	void Initialize(FRequiredEntryData InRequiredEntryData);
+	NIAGARAEDITOR_API void Initialize(FRequiredEntryData InRequiredEntryData);
 
-	virtual FText GetDisplayName() const override;
+	NIAGARAEDITOR_API virtual FText GetDisplayName() const override;
 
-	virtual FText GetTooltipText() const override;
+	NIAGARAEDITOR_API virtual FText GetTooltipText() const override;
 
 	virtual bool GetShouldShowInOverview() const override { return false; }
 
 protected:
-	virtual void RefreshChildrenInternal(const TArray<UNiagaraStackEntry*>& CurrentChildren, TArray<UNiagaraStackEntry*>& NewChildren, TArray<FStackIssue>& NewIssues) override;
+	NIAGARAEDITOR_API virtual void RefreshChildrenInternal(const TArray<UNiagaraStackEntry*>& CurrentChildren, TArray<UNiagaraStackEntry*>& NewChildren, TArray<FStackIssue>& NewIssues) override;
 
 private:
-	void SystemPropertiesChanged();
+	NIAGARAEDITOR_API void SystemPropertiesChanged();
 
 private:
 	mutable TOptional<bool> bCanResetToBase;
diff --git a/Engine/Plugins/FX/Niagara/Source/NiagaraEditor/Public/ViewModels/Stack/NiagaraStackSystemSettingsGroup.h b/Engine/Plugins/FX/Niagara/Source/NiagaraEditor/Public/ViewModels/Stack/NiagaraStackSystemSettingsGroup.h
--- a/Engine/Plugins/FX/Niagara/Source/NiagaraEditor/Public/ViewModels/Stack/NiagaraStackSystemSettingsGroup.h
+++ b/Engine/Plugins/FX/Niagara/Source/NiagaraEditor/Public/ViewModels/Stack/NiagaraStackSystemSettingsGroup.h
@@ -11,33 +11,33 @@
 
 class FNiagaraScriptViewModel;
 
-UCLASS()
-class NIAGARAEDITOR_API UNiagaraStackSystemPropertiesGroup : public UNiagaraStackItemGroup
+UCLASS(MinimalAPI)
+class UNiagaraStackSystemPropertiesGroup : public UNiagaraStackItemGroup
 {
 	GENERATED_BODY()
 		
 public:
-	void Initialize(FRequiredEntryData InRequiredEntryData);
+	NIAGARAEDITOR_API void Initialize(FRequiredEntryData InRequiredEntryData);
 
 	virtual EIconMode GetSupportedIconMode() const override { return UNiagaraStackEntry::EIconMode::Brush; }
-	virtual const FSlateBrush* GetIconBrush() const override;
+	NIAGARAEDITOR_API virtual const FSlateBrush* GetIconBrush() const override;
 	virtual bool GetCanExpandInOverview() const override { return false; }
 	virtual bool GetShouldShowInStack() const override { return false; }
 
 protected:
-	virtual void RefreshChildrenInternal(const TArray<UNiagaraStackEntry*>& CurrentChildren, TArray<UNiagaraStackEntry*>& NewChildren, TArray<FStackIssue>& NewIssues) override;
+	NIAGARAEDITOR_API virtual void RefreshChildrenInternal(const TArray<UNiagaraStackEntry*>& CurrentChildren, TArray<UNiagaraStackEntry*>& NewChildren, TArray<FStackIssue>& NewIssues) override;
 };
 
-UCLASS()
-class NIAGARAEDITOR_API UNiagaraStackSystemUserParametersGroup : public UNiagaraStackItemGroup
+UCLASS(MinimalAPI)
+class UNiagaraStackSystemUserParametersGroup : public UNiagaraStackItemGroup
 {
 	GENERATED_BODY()
 
 public:
-	void Initialize(FRequiredEntryData InRequiredEntryData, UObject* InOwner, FNiagaraParameterStore* InParameterStore);
+	NIAGARAEDITOR_API void Initialize(FRequiredEntryData InRequiredEntryData, UObject* InOwner, FNiagaraParameterStore* InParameterStore);
 
 	virtual EIconMode GetSupportedIconMode() const override { return UNiagaraStackEntry::EIconMode::Text; }
-	virtual FText GetIconText() const override;
+	NIAGARAEDITOR_API virtual FText GetIconText() const override;
 	virtual bool GetCanExpandInOverview() const override { return false; }
 	virtual bool GetShouldShowInStack() const override { return false; }
 private:
diff --git a/Engine/Plugins/FX/Niagara/Source/NiagaraEditor/Public/ViewModels/Stack/NiagaraStackViewModel.h b/Engine/Plugins/FX/Niagara/Source/NiagaraEditor/Public/ViewModels/Stack/NiagaraStackViewModel.h
--- a/Engine/Plugins/FX/Niagara/Source/NiagaraEditor/Public/ViewModels/Stack/NiagaraStackViewModel.h
+++ b/Engine/Plugins/FX/Niagara/Source/NiagaraEditor/Public/ViewModels/Stack/NiagaraStackViewModel.h
@@ -35,8 +35,8 @@ private:
 	bool bIncludeEmitterInformation;
 };
 
-UCLASS()
-class NIAGARAEDITOR_API UNiagaraStackViewModel : public UObject
+UCLASS(MinimalAPI)
+class UNiagaraStackViewModel : public UObject
 {
 	GENERATED_BODY()
 
@@ -59,21 +59,21 @@ public:
 		}
 	};
 
-	struct NIAGARAEDITOR_API FTopLevelViewModel
+	struct FTopLevelViewModel
 	{
-		FTopLevelViewModel(TSharedPtr<FNiagaraSystemViewModel> InSystemViewModel);
+		NIAGARAEDITOR_API FTopLevelViewModel(TSharedPtr<FNiagaraSystemViewModel> InSystemViewModel);
 
-		FTopLevelViewModel(TSharedPtr<FNiagaraEmitterHandleViewModel> InEmitterHandleViewModel);
+		NIAGARAEDITOR_API FTopLevelViewModel(TSharedPtr<FNiagaraEmitterHandleViewModel> InEmitterHandleViewModel);
 
-		bool IsValid() const;
+		NIAGARAEDITOR_API bool IsValid() const;
 
-		UNiagaraStackEditorData* GetStackEditorData() const;
+		NIAGARAEDITOR_API UNiagaraStackEditorData* GetStackEditorData() const;
 
-		void GetMessageStores(TArray<FNiagaraMessageSourceAndStore>& OutMessageStores);
+		NIAGARAEDITOR_API void GetMessageStores(TArray<FNiagaraMessageSourceAndStore>& OutMessageStores);
 
-		FText GetDisplayName() const;
+		NIAGARAEDITOR_API FText GetDisplayName() const;
 
-		bool operator==(const FTopLevelViewModel& Other)const;
+		NIAGARAEDITOR_API bool operator==(const FTopLevelViewModel& Other)const;
 
 		const TSharedPtr<FNiagaraSystemViewModel> SystemViewModel;
 		const TSharedPtr<FNiagaraEmitterHandleViewModel> EmitterHandleViewModel;
@@ -81,79 +81,79 @@ public:
 	};
 
 public:
-	void InitializeWithViewModels(TSharedPtr<FNiagaraSystemViewModel> InSystemViewModel, TSharedPtr<FNiagaraEmitterHandleViewModel> InEmitterHandleViewModel, FNiagaraStackViewModelOptions InOptions);
-	void InitializeWithRootEntry(UNiagaraStackEntry* Root);
+	NIAGARAEDITOR_API void InitializeWithViewModels(TSharedPtr<FNiagaraSystemViewModel> InSystemViewModel, TSharedPtr<FNiagaraEmitterHandleViewModel> InEmitterHandleViewModel, FNiagaraStackViewModelOptions InOptions);
+	NIAGARAEDITOR_API void InitializeWithRootEntry(UNiagaraStackEntry* Root);
 
-	void Finalize();
+	NIAGARAEDITOR_API void Finalize();
 
-	virtual void BeginDestroy() override;
+	NIAGARAEDITOR_API virtual void BeginDestroy() override;
 
-	UNiagaraStackEntry* GetRootEntry();
+	NIAGARAEDITOR_API UNiagaraStackEntry* GetRootEntry();
 
-	TArray<UNiagaraStackEntry*>& GetRootEntryAsArray();
+	NIAGARAEDITOR_API TArray<UNiagaraStackEntry*>& GetRootEntryAsArray();
 
-	FOnChangeSearchTextExternal& OnChangeSearchTextExternal();
+	NIAGARAEDITOR_API FOnChangeSearchTextExternal& OnChangeSearchTextExternal();
 
-	FOnExpansionChanged& OnExpansionChanged();
-	FOnExpansionChanged& OnExpansionInOverviewChanged();
-	FOnStructureChanged& OnStructureChanged();
-	FOnSearchCompleted& OnSearchCompleted();
-	FOnDataObjectChanged& OnDataObjectChanged();
+	NIAGARAEDITOR_API FOnExpansionChanged& OnExpansionChanged();
+	NIAGARAEDITOR_API FOnExpansionChanged& OnExpansionInOverviewChanged();
+	NIAGARAEDITOR_API FOnStructureChanged& OnStructureChanged();
+	NIAGARAEDITOR_API FOnSearchCompleted& OnSearchCompleted();
+	NIAGARAEDITOR_API FOnDataObjectChanged& OnDataObjectChanged();
 
-	bool GetShowAllAdvanced() const;
-	void SetShowAllAdvanced(bool bInShowAllAdvanced);
+	NIAGARAEDITOR_API bool GetShowAllAdvanced() const;
+	NIAGARAEDITOR_API void SetShowAllAdvanced(bool bInShowAllAdvanced);
 
-	bool GetShowOutputs() const;
-	void SetShowOutputs(bool bInShowOutputs);
+	NIAGARAEDITOR_API bool GetShowOutputs() const;
+	NIAGARAEDITOR_API void SetShowOutputs(bool bInShowOutputs);
 
-	bool GetShowLinkedInputs() const;
-	void SetShowLinkedInputs(bool bInShowLinkedInputs);
+	NIAGARAEDITOR_API bool GetShowLinkedInputs() const;
+	NIAGARAEDITOR_API void SetShowLinkedInputs(bool bInShowLinkedInputs);
 
-	bool GetShowOnlyIssues() const;
-	void SetShowOnlyIssues(bool bInShowOnlyIssues);
+	NIAGARAEDITOR_API bool GetShowOnlyIssues() const;
+	NIAGARAEDITOR_API void SetShowOnlyIssues(bool bInShowOnlyIssues);
 
-	double GetLastScrollPosition() const;
-	void SetLastScrollPosition(double InLastScrollPosition);
+	NIAGARAEDITOR_API double GetLastScrollPosition() const;
+	NIAGARAEDITOR_API void SetLastScrollPosition(double InLastScrollPosition);
 
-	virtual void Tick();
+	NIAGARAEDITOR_API virtual void Tick();
 	//~ stack search stuff
-	void ResetSearchText();
+	NIAGARAEDITOR_API void ResetSearchText();
 	FText GetCurrentSearchText() const { return CurrentSearchText; };
-	void OnSearchTextChanged(const FText& SearchText);
-	void SetSearchTextExternal(const FText& NewSearchText);
-	bool IsSearching();
-	const TArray<FSearchResult>& GetCurrentSearchResults();
-	UNiagaraStackEntry* GetCurrentFocusedEntry() const;
-	int32 GetCurrentFocusedEntryIndex() const;
-	void AddSearchScrollOffset(int NumberOfSteps);
+	NIAGARAEDITOR_API void OnSearchTextChanged(const FText& SearchText);
+	NIAGARAEDITOR_API void SetSearchTextExternal(const FText& NewSearchText);
+	NIAGARAEDITOR_API bool IsSearching();
+	NIAGARAEDITOR_API const TArray<FSearchResult>& GetCurrentSearchResults();
+	NIAGARAEDITOR_API UNiagaraStackEntry* GetCurrentFocusedEntry() const;
+	NIAGARAEDITOR_API int32 GetCurrentFocusedEntryIndex() const;
+	NIAGARAEDITOR_API void AddSearchScrollOffset(int NumberOfSteps);
 
-	void OnCycleThroughIssues(TSharedPtr<FTopLevelViewModel> TopLevelToCycle);
-	UNiagaraStackEntry* GetCurrentFocusedIssue() const;
+	NIAGARAEDITOR_API void OnCycleThroughIssues(TSharedPtr<FTopLevelViewModel> TopLevelToCycle);
+	NIAGARAEDITOR_API UNiagaraStackEntry* GetCurrentFocusedIssue() const;
 
-	void GetPathForEntry(UNiagaraStackEntry* Entry, TArray<UNiagaraStackEntry*>& EntryPath) const;
+	NIAGARAEDITOR_API void GetPathForEntry(UNiagaraStackEntry* Entry, TArray<UNiagaraStackEntry*>& EntryPath) const;
 
 	/** Starts recursing through all entries to expand all groups and collapse all items. */
-	void CollapseToHeaders();
+	NIAGARAEDITOR_API void CollapseToHeaders();
 	
-	void UndismissAllIssues();
-	bool HasDismissedStackIssues();
+	NIAGARAEDITOR_API void UndismissAllIssues();
+	NIAGARAEDITOR_API bool HasDismissedStackIssues();
 
-	const TArray<TSharedRef<FTopLevelViewModel>>& GetTopLevelViewModels() const;
+	NIAGARAEDITOR_API const TArray<TSharedRef<FTopLevelViewModel>>& GetTopLevelViewModels() const;
 
-	TSharedPtr<FTopLevelViewModel> GetTopLevelViewModelForEntry(UNiagaraStackEntry& InEntry) const;
+	NIAGARAEDITOR_API TSharedPtr<FTopLevelViewModel> GetTopLevelViewModelForEntry(UNiagaraStackEntry& InEntry) const;
 
-	void Reset();
-	bool HasIssues() const;
-	void RequestRefreshDeferred();
+	NIAGARAEDITOR_API void Reset();
+	NIAGARAEDITOR_API bool HasIssues() const;
+	NIAGARAEDITOR_API void RequestRefreshDeferred();
 	void RequestValidationUpdate() { bValidatorUpdatePending = true; }
 
-	bool ShouldHideDisabledModules() const;
+	NIAGARAEDITOR_API bool ShouldHideDisabledModules() const;
 
-	void InvalidateCachedParameterUsage();
+	NIAGARAEDITOR_API void InvalidateCachedParameterUsage();
 
 private:
 	/** Recursively Expands all groups and collapses all items in the stack. */
-	void CollapseToHeadersRecursive(TArray<UNiagaraStackEntry*> Entries);
+	NIAGARAEDITOR_API void CollapseToHeadersRecursive(TArray<UNiagaraStackEntry*> Entries);
 
 	struct FSearchWorkItem
 	{
@@ -166,23 +166,23 @@ private:
 		}
 	};
 
-	void EntryExpansionChanged();
-	void EntryExpansionInOverviewChanged();
-	void EntryStructureChanged(ENiagaraStructureChangedFlags Flags);
-	void EntryDataObjectModified(TArray<UObject*> ChangedObjects, ENiagaraDataObjectChange ChangeType);
-	void EntryRequestFullRefresh();
-	void RefreshTopLevelViewModels();
-	void RefreshHasIssues();
-	void EmitterParentRemoved();
+	NIAGARAEDITOR_API void EntryExpansionChanged();
+	NIAGARAEDITOR_API void EntryExpansionInOverviewChanged();
+	NIAGARAEDITOR_API void EntryStructureChanged(ENiagaraStructureChangedFlags Flags);
+	NIAGARAEDITOR_API void EntryDataObjectModified(TArray<UObject*> ChangedObjects, ENiagaraDataObjectChange ChangeType);
+	NIAGARAEDITOR_API void EntryRequestFullRefresh();
+	NIAGARAEDITOR_API void RefreshTopLevelViewModels();
+	NIAGARAEDITOR_API void RefreshHasIssues();
+	NIAGARAEDITOR_API void EmitterParentRemoved();
 	/** Called by the tick function to perform partial search */
-	void SearchTick();
-	void GenerateTraversalEntries(UNiagaraStackEntry* Root, TArray<UNiagaraStackEntry*> ParentChain, 
+	NIAGARAEDITOR_API void SearchTick();
+	NIAGARAEDITOR_API void GenerateTraversalEntries(UNiagaraStackEntry* Root, TArray<UNiagaraStackEntry*> ParentChain, 
 		TArray<FSearchWorkItem>& TraversedArray);
-	bool ItemMatchesSearchCriteria(UNiagaraStackEntry::FStackSearchItem SearchItem);
-	void GeneratePathForEntry(UNiagaraStackEntry* Root, UNiagaraStackEntry* Entry, TArray<UNiagaraStackEntry*> CurrentPath, TArray<UNiagaraStackEntry*>& EntryPath) const;
+	NIAGARAEDITOR_API bool ItemMatchesSearchCriteria(UNiagaraStackEntry::FStackSearchItem SearchItem);
+	NIAGARAEDITOR_API void GeneratePathForEntry(UNiagaraStackEntry* Root, UNiagaraStackEntry* Entry, TArray<UNiagaraStackEntry*> CurrentPath, TArray<UNiagaraStackEntry*>& EntryPath) const;
 
-	void RestartSearch();
-	void UpdateStackWithValidationResults();
+	NIAGARAEDITOR_API void RestartSearch();
+	NIAGARAEDITOR_API void UpdateStackWithValidationResults();
 
 private:
 	TWeakPtr<FNiagaraEmitterHandleViewModel> EmitterHandleViewModel;
@@ -208,7 +208,7 @@ private:
 	FOnSearchCompleted SearchCompletedDelegate;
 	TArray<FSearchWorkItem> ItemsToSearch;
 	TArray<FSearchResult> CurrentSearchResults;
-	static const double MaxSearchTime;
+	static NIAGARAEDITOR_API const double MaxSearchTime;
 	TOptional<FSearchResult> FocusedSearchResultCache;
 	bool bRestartSearch;
 	bool bRefreshPending;
@@ -221,4 +221,4 @@ private:
 	TArray<TSharedRef<FTopLevelViewModel>> TopLevelViewModels;
 
 	FNiagaraStackViewModelOptions Options;
-};
\ No newline at end of file
+};
