[Niagara]

diff --git a/Engine/Plugins/FX/Niagara/Source/NiagaraEditor/Public/Widgets/NiagaraHLSLSyntaxHighlighter.h b/Engine/Plugins/FX/Niagara/Source/NiagaraEditor/Public/Widgets/NiagaraHLSLSyntaxHighlighter.h
--- a/Engine/Plugins/FX/Niagara/Source/NiagaraEditor/Public/Widgets/NiagaraHLSLSyntaxHighlighter.h
+++ b/Engine/Plugins/FX/Niagara/Source/NiagaraEditor/Public/Widgets/NiagaraHLSLSyntaxHighlighter.h
@@ -7,13 +7,13 @@
 
 class FTextLayout;
 
-class NIAGARAEDITOR_API FNiagaraHLSLSyntaxHighlighter : public FHLSLSyntaxHighlighterMarshaller
+class FNiagaraHLSLSyntaxHighlighter : public FHLSLSyntaxHighlighterMarshaller
 {
 public:
 
-	static TSharedRef<FNiagaraHLSLSyntaxHighlighter> Create(const FSyntaxTextStyle& InSyntaxTextStyle = GetSyntaxTextStyle());
-	static FSyntaxTextStyle GetSyntaxTextStyle();
+	static NIAGARAEDITOR_API TSharedRef<FNiagaraHLSLSyntaxHighlighter> Create(const FSyntaxTextStyle& InSyntaxTextStyle = GetSyntaxTextStyle());
+	static NIAGARAEDITOR_API FSyntaxTextStyle GetSyntaxTextStyle();
 
 protected:
-	FNiagaraHLSLSyntaxHighlighter(TSharedPtr< ISyntaxTokenizer > InTokenizer, const FSyntaxTextStyle& InSyntaxTextStyle);
+	NIAGARAEDITOR_API FNiagaraHLSLSyntaxHighlighter(TSharedPtr< ISyntaxTokenizer > InTokenizer, const FSyntaxTextStyle& InSyntaxTextStyle);
 };
diff --git a/Engine/Plugins/FX/Niagara/Source/NiagaraEditor/Public/Widgets/SDynamicLayoutBox.h b/Engine/Plugins/FX/Niagara/Source/NiagaraEditor/Public/Widgets/SDynamicLayoutBox.h
--- a/Engine/Plugins/FX/Niagara/Source/NiagaraEditor/Public/Widgets/SDynamicLayoutBox.h
+++ b/Engine/Plugins/FX/Niagara/Source/NiagaraEditor/Public/Widgets/SDynamicLayoutBox.h
@@ -5,18 +5,18 @@
 #include "Widgets/SCompoundWidget.h"
 #include "Widgets/DeclarativeSyntaxSupport.h"
 
-class NIAGARAEDITOR_API SDynamicLayoutBox : public SCompoundWidget
+class SDynamicLayoutBox : public SCompoundWidget
 {
 public:
 	DECLARE_DELEGATE_RetVal_OneParam(TSharedRef<SWidget>, FOnGenerateNamedWidget, FName /* InWidgetName */);
 
-	class NIAGARAEDITOR_API FNamedWidgetProvider
+	class FNamedWidgetProvider
 	{
 	public:
-		FNamedWidgetProvider();
-		FNamedWidgetProvider(FOnGenerateNamedWidget InGeneratedNamedWidget);
+		NIAGARAEDITOR_API FNamedWidgetProvider();
+		NIAGARAEDITOR_API FNamedWidgetProvider(FOnGenerateNamedWidget InGeneratedNamedWidget);
 
-		TSharedRef<SWidget> GetNamedWidget(FName InWidgetName) const;
+		NIAGARAEDITOR_API TSharedRef<SWidget> GetNamedWidget(FName InWidgetName) const;
 
 	private:
 		mutable TMap<FName, TSharedRef<SWidget>> NameToGeneratedWidgetCache;
@@ -35,13 +35,13 @@ public:
 		SLATE_EVENT(FOnChooseLayout, ChooseLayout);
 	SLATE_END_ARGS();
 
-	void Construct(const FArguments& InArgs);
+	NIAGARAEDITOR_API void Construct(const FArguments& InArgs);
 
-	virtual void Tick(const FGeometry& AllottedGeometry, const double InCurrentTime, const float InDeltaTime) override;
+	NIAGARAEDITOR_API virtual void Tick(const FGeometry& AllottedGeometry, const double InCurrentTime, const float InDeltaTime) override;
 
 private:
 	FNamedWidgetProvider NamedWidgetProvider;
 	FOnGenerateNamedLayout GenerateNamedLayout;
 	FOnChooseLayout ChooseLayout;
 	TOptional<FName> CurrentLayout;
-};
\ No newline at end of file
+};
diff --git a/Engine/Plugins/FX/Niagara/Source/NiagaraEditor/Public/Widgets/SNiagaraFilterBox.h b/Engine/Plugins/FX/Niagara/Source/NiagaraEditor/Public/Widgets/SNiagaraFilterBox.h
--- a/Engine/Plugins/FX/Niagara/Source/NiagaraEditor/Public/Widgets/SNiagaraFilterBox.h
+++ b/Engine/Plugins/FX/Niagara/Source/NiagaraEditor/Public/Widgets/SNiagaraFilterBox.h
@@ -14,7 +14,7 @@
 /**
  * A single source filter checkbox, inherited from SCheckBox to add shift-click functionality
  */
-class NIAGARAEDITOR_API SNiagaraSourceFilterCheckBox : public SCheckBox
+class SNiagaraSourceFilterCheckBox : public SCheckBox
 {
 	DECLARE_DELEGATE_TwoParams(FOnSourceStateChanged, EScriptSource, bool);
 	DECLARE_DELEGATE_TwoParams(FOnShiftClicked, EScriptSource, bool);
@@ -26,9 +26,9 @@ class NIAGARAEDITOR_API SNiagaraSourceFilterCheckBox : public SCheckBox
 		SLATE_ATTRIBUTE(ECheckBoxState, IsChecked)
 	SLATE_END_ARGS()
 
-	void Construct(const FArguments& Args, EScriptSource Source);
+	NIAGARAEDITOR_API void Construct(const FArguments& Args, EScriptSource Source);
 
-	virtual FReply OnMouseButtonUp(const FGeometry& MyGeometry, const FPointerEvent& MouseEvent) override;
+	NIAGARAEDITOR_API virtual FReply OnMouseButtonUp(const FGeometry& MyGeometry, const FPointerEvent& MouseEvent) override;
 private:	
 	/** The script source entry associated with this checkbox. */
 	EScriptSource Source;
@@ -44,7 +44,7 @@ private:
 /**
  * The source filter box will create a source filter checkbox per enum entry and can be used to filter for different sources of assets, actions etc.
  */
-class NIAGARAEDITOR_API SNiagaraSourceFilterBox : public SCompoundWidget
+class SNiagaraSourceFilterBox : public SCompoundWidget
 {
 public:
 	typedef TMap<EScriptSource, bool> SourceMap;
@@ -55,9 +55,9 @@ public:
 		SLATE_EVENT(FOnFiltersChanged, OnFiltersChanged)
 	SLATE_END_ARGS()
 
-	void Construct(const FArguments& Args);
+	NIAGARAEDITOR_API void Construct(const FArguments& Args);
 
-	bool IsFilterActive(EScriptSource Source) const;
+	NIAGARAEDITOR_API bool IsFilterActive(EScriptSource Source) const;
 
 private:
 	TMap<EScriptSource, TSharedRef<SNiagaraSourceFilterCheckBox>> SourceButtons;
@@ -78,7 +78,7 @@ private:
  * The template tab box offers buttons (working like tabs or radio buttons) for filtering purposes.
  * It is used for emitters and systems to filter between standard assets, templates and behavior examples.
  */
-class NIAGARAEDITOR_API SNiagaraTemplateTabBox : public SCompoundWidget
+class SNiagaraTemplateTabBox : public SCompoundWidget
 {
 public:
 	DECLARE_DELEGATE_OneParam(FOnTabActivated, ENiagaraScriptTemplateSpecification Tab);
@@ -86,19 +86,19 @@ public:
 	/**
 	 * The template tab options will tell the tab box which tabs to offer in the form of buttons.
 	 */
-	struct NIAGARAEDITOR_API FNiagaraTemplateTabOptions
+	struct FNiagaraTemplateTabOptions
 	{
 		FNiagaraTemplateTabOptions()
 		{}
 
 		void ChangeTabState(ENiagaraScriptTemplateSpecification AssetTab, bool bAvailable = true) { TabData.Add(AssetTab, bAvailable); }
-		bool IsTabAvailable(ENiagaraScriptTemplateSpecification AssetTab) const;
+		NIAGARAEDITOR_API bool IsTabAvailable(ENiagaraScriptTemplateSpecification AssetTab) const;
 
-		int32 GetNumAvailableTabs() const;
-		bool GetOnlyAvailableTab(ENiagaraScriptTemplateSpecification& OutTab) const;
-		bool GetOnlyShowTemplates() const;
+		NIAGARAEDITOR_API int32 GetNumAvailableTabs() const;
+		NIAGARAEDITOR_API bool GetOnlyAvailableTab(ENiagaraScriptTemplateSpecification& OutTab) const;
+		NIAGARAEDITOR_API bool GetOnlyShowTemplates() const;
 
-		const TMap<ENiagaraScriptTemplateSpecification, bool>& GetTabData() const;
+		NIAGARAEDITOR_API const TMap<ENiagaraScriptTemplateSpecification, bool>& GetTabData() const;
 	
 	private:
 		TMap<ENiagaraScriptTemplateSpecification, bool> TabData;
@@ -111,11 +111,11 @@ public:
 		SLATE_ARGUMENT(UClass*, Class)
 	SLATE_END_ARGS()
 
-	void Construct(const FArguments& InArgs, FNiagaraTemplateTabOptions TabOptions);
+	NIAGARAEDITOR_API void Construct(const FArguments& InArgs, FNiagaraTemplateTabOptions TabOptions);
 
-	ENiagaraScriptTemplateSpecification GetActiveTab() const;
+	NIAGARAEDITOR_API ENiagaraScriptTemplateSpecification GetActiveTab() const;
 	/** If true, will set the parameter to the currently active tab. If false, the template filter wasn't initialized with any tabs. */
-	bool GetActiveTab(ENiagaraScriptTemplateSpecification& OutTemplateSpecification) const;
+	NIAGARAEDITOR_API bool GetActiveTab(ENiagaraScriptTemplateSpecification& OutTemplateSpecification) const;
 private:
 	FOnTabActivated OnTabActivatedDelegate;
 	void OnTabActivated(ENiagaraScriptTemplateSpecification AssetTab);
@@ -140,29 +140,29 @@ private:
  * The primary access point to use different applicable filters.
  * Can be configured to add a source, library and template filter and offers callbacks via delegates to execute logic when the filter states change.
  */
-class NIAGARAEDITOR_API SNiagaraFilterBox : public SCompoundWidget
+class SNiagaraFilterBox : public SCompoundWidget
 {
 public:
 	/**
 	 * The filter options used to initialized the filter box.
 	 * There needs to be at least two tab options set to true for the tab filter to show up, if added.
 	 */
-	struct NIAGARAEDITOR_API FFilterOptions
+	struct FFilterOptions
 	{
 		FFilterOptions() : bAddSourceFilter(true), bAddLibraryFilter(true), bAddTemplateFilter(false)
 		{}
 
-		bool IsAnyFilterActive();
+		NIAGARAEDITOR_API bool IsAnyFilterActive();
 
-		void SetAddSourceFilter(bool bAddSourceFilter);
-		void SetAddLibraryFilter(bool bAddLibraryFilter);
-		void SetAddTemplateFilter(bool bAddTemplateFilter);
-		void SetTabOptions(const SNiagaraTemplateTabBox::FNiagaraTemplateTabOptions& TabOptions);
+		NIAGARAEDITOR_API void SetAddSourceFilter(bool bAddSourceFilter);
+		NIAGARAEDITOR_API void SetAddLibraryFilter(bool bAddLibraryFilter);
+		NIAGARAEDITOR_API void SetAddTemplateFilter(bool bAddTemplateFilter);
+		NIAGARAEDITOR_API void SetTabOptions(const SNiagaraTemplateTabBox::FNiagaraTemplateTabOptions& TabOptions);
 		
-		bool GetAddSourceFilter() const;
-		bool GetAddLibraryFilter() const;
-		bool GetAddTemplateFilter() const;
-		SNiagaraTemplateTabBox::FNiagaraTemplateTabOptions GetTabOptions() const;
+		NIAGARAEDITOR_API bool GetAddSourceFilter() const;
+		NIAGARAEDITOR_API bool GetAddLibraryFilter() const;
+		NIAGARAEDITOR_API bool GetAddTemplateFilter() const;
+		NIAGARAEDITOR_API SNiagaraTemplateTabBox::FNiagaraTemplateTabOptions GetTabOptions() const;
 	private:
 		bool bAddSourceFilter;
 		bool bAddLibraryFilter;
@@ -182,13 +182,13 @@ public:
 		SLATE_ARGUMENT(UClass*, Class)
 	SLATE_END_ARGS()
 
-	void Construct(const FArguments& InArgs, FFilterOptions FilterOptions);
+	NIAGARAEDITOR_API void Construct(const FArguments& InArgs, FFilterOptions FilterOptions);
 	
-	bool IsSourceFilterActive(EScriptSource Source) const;
-	bool GetActiveTemplateTab(ENiagaraScriptTemplateSpecification& OutTemplateSpecification) const;
+	NIAGARAEDITOR_API bool IsSourceFilterActive(EScriptSource Source) const;
+	NIAGARAEDITOR_API bool GetActiveTemplateTab(ENiagaraScriptTemplateSpecification& OutTemplateSpecification) const;
 	
 private:
 	TSharedPtr<SNiagaraSourceFilterBox> SourceFilterBox;
 	TSharedPtr<SNiagaraLibraryOnlyToggleHeader> LibraryOnlyToggleHeader;
 	TSharedPtr<SNiagaraTemplateTabBox> TemplateTabBox;
-};
\ No newline at end of file
+};
diff --git a/Engine/Plugins/FX/Niagara/Source/NiagaraEditor/Public/Widgets/SNiagaraGraphActionMenu.h b/Engine/Plugins/FX/Niagara/Source/NiagaraEditor/Public/Widgets/SNiagaraGraphActionMenu.h
--- a/Engine/Plugins/FX/Niagara/Source/NiagaraEditor/Public/Widgets/SNiagaraGraphActionMenu.h
+++ b/Engine/Plugins/FX/Niagara/Source/NiagaraEditor/Public/Widgets/SNiagaraGraphActionMenu.h
@@ -32,7 +32,7 @@ struct FNiagaraGraphActionMenuFilter
 	TDelegate<void()> OnFilterChanged;
 };
 
-class NIAGARAEDITOR_API SNiagaraGraphActionMenu : public SBorder
+class SNiagaraGraphActionMenu : public SBorder
 {
 public:
 	SLATE_BEGIN_ARGS( SNiagaraGraphActionMenu )
@@ -48,11 +48,11 @@ public:
 		SLATE_ARGUMENT( bool, AutoExpandActionMenu )
 	SLATE_END_ARGS()
 
-	void Construct( const FArguments& InArgs );
+	NIAGARAEDITOR_API void Construct( const FArguments& InArgs );
 
-	~SNiagaraGraphActionMenu();
+	NIAGARAEDITOR_API ~SNiagaraGraphActionMenu();
 
-	TSharedRef<SWidget> GetFilterTextBox();
+	NIAGARAEDITOR_API TSharedRef<SWidget> GetFilterTextBox();
 
 protected:
 	UEdGraph* GraphObj;
@@ -66,11 +66,11 @@ protected:
 	TSharedPtr<SNiagaraFilterBox> FilterBox;
 	TSharedPtr<SNiagaraMenuActionSelector> ActionSelector;
 
-	TArray<TSharedPtr<FNiagaraAction_NewNode>> CollectAllActions();
+	NIAGARAEDITOR_API TArray<TSharedPtr<FNiagaraAction_NewNode>> CollectAllActions();
 	
 	UEdGraph* OwnerOfTemporaries = nullptr;
 
-	static bool bLibraryOnly;
+	static NIAGARAEDITOR_API bool bLibraryOnly;
 private:
 	TArray<FString> OnGetCategoriesForItem(const TSharedPtr<FNiagaraAction_NewNode>& Item);
 	TArray<ENiagaraMenuSections> OnGetSectionsForItem(const TSharedPtr<FNiagaraAction_NewNode>& Item);
diff --git a/Engine/Plugins/FX/Niagara/Source/NiagaraEditor/Public/Widgets/SNiagaraParameterName.h b/Engine/Plugins/FX/Niagara/Source/NiagaraEditor/Public/Widgets/SNiagaraParameterName.h
--- a/Engine/Plugins/FX/Niagara/Source/NiagaraEditor/Public/Widgets/SNiagaraParameterName.h
+++ b/Engine/Plugins/FX/Niagara/Source/NiagaraEditor/Public/Widgets/SNiagaraParameterName.h
@@ -14,7 +14,7 @@ class SInlineEditableTextBlock;
 class SBorder;
 class UEdGraphPin;
 
-class NIAGARAEDITOR_API SNiagaraParameterName : public SCompoundWidget
+class SNiagaraParameterName : public SCompoundWidget
 {
 public:
 	DECLARE_DELEGATE_OneParam(FOnNameChanged, FName /* InNewName */);
@@ -52,17 +52,17 @@ public:
 		SLATE_NAMED_SLOT(FArguments, Decorator)
 	SLATE_END_ARGS()
 
-	void Construct(const FArguments& InArgs);
+	NIAGARAEDITOR_API void Construct(const FArguments& InArgs);
 
-	virtual void Tick(const FGeometry& AllottedGeometry, const double InCurrentTime, const float InDeltaTime) override;
+	NIAGARAEDITOR_API virtual void Tick(const FGeometry& AllottedGeometry, const double InCurrentTime, const float InDeltaTime) override;
 
-	virtual FReply OnMouseButtonDoubleClick(const FGeometry& InMyGeometry, const FPointerEvent& InMouseEvent) override;
+	NIAGARAEDITOR_API virtual FReply OnMouseButtonDoubleClick(const FGeometry& InMyGeometry, const FPointerEvent& InMouseEvent) override;
 
-	void EnterEditingMode();
+	NIAGARAEDITOR_API void EnterEditingMode();
 
-	void EnterNamespaceModifierEditingMode();
+	NIAGARAEDITOR_API void EnterNamespaceModifierEditingMode();
 
-	void UpdateDecorator(TSharedRef<SWidget> InDecorator);
+	NIAGARAEDITOR_API void UpdateDecorator(TSharedRef<SWidget> InDecorator);
 	
 private:
 	TSharedRef<SBorder> CreateNamespaceWidget(FText NamespaceDisplayName, FText NamespaceDescription, FLinearColor NamespaceBorderColor, FName NamespaceForegroundStyle);
@@ -103,7 +103,7 @@ private:
 	TSharedPtr<SWidget> Decorator;
 };
 
-class NIAGARAEDITOR_API SNiagaraParameterNameTextBlock : public SCompoundWidget
+class SNiagaraParameterNameTextBlock : public SCompoundWidget
 {
 public:
 	DECLARE_DELEGATE_RetVal_TwoParams(FReply, FOnDragDetectedHandler, const FGeometry&, const FPointerEvent&);
@@ -131,11 +131,11 @@ public:
 		SLATE_NAMED_SLOT(FArguments, Decorator)
 	SLATE_END_ARGS()
 
-	void Construct(const FArguments& InArgs);
+	NIAGARAEDITOR_API void Construct(const FArguments& InArgs);
 
-	void EnterEditingMode();
+	NIAGARAEDITOR_API void EnterEditingMode();
 
-	void EnterNamespaceModifierEditingMode();
+	NIAGARAEDITOR_API void EnterNamespaceModifierEditingMode();
 
 private:
 	FName GetParameterName() const;
@@ -144,8 +144,8 @@ private:
 
 	void NameChanged(FName InNewName);
 
-	virtual FReply OnMouseButtonDown(const FGeometry& MyGeometry, const FPointerEvent& MouseEvent) override;
-	virtual FReply OnDragDetected(const FGeometry& MyGeometry, const FPointerEvent& MouseEvent) override;
+	NIAGARAEDITOR_API virtual FReply OnMouseButtonDown(const FGeometry& MyGeometry, const FPointerEvent& MouseEvent) override;
+	NIAGARAEDITOR_API virtual FReply OnDragDetected(const FGeometry& MyGeometry, const FPointerEvent& MouseEvent) override;
 
 private:
 	TAttribute<FText> ParameterText;
@@ -158,7 +158,7 @@ private:
 	TSharedPtr<SNiagaraParameterName> ParameterName;
 };
 
-class NIAGARAEDITOR_API SNiagaraParameterNamePinLabel : public SCompoundWidget
+class SNiagaraParameterNamePinLabel : public SCompoundWidget
 {
 public:
 	SLATE_BEGIN_ARGS(SNiagaraParameterNamePinLabel)
@@ -181,13 +181,13 @@ public:
 		SLATE_NAMED_SLOT(FArguments, Decorator)
 	SLATE_END_ARGS()
 
-	void Construct(const FArguments& InArgs, UEdGraphPin* InTargetPin);
+	NIAGARAEDITOR_API void Construct(const FArguments& InArgs, UEdGraphPin* InTargetPin);
 
-	virtual void Tick(const FGeometry& AllottedGeometry, const double InCurrentTime, const float InDeltaTime) override;
+	NIAGARAEDITOR_API virtual void Tick(const FGeometry& AllottedGeometry, const double InCurrentTime, const float InDeltaTime) override;
 
-	virtual FSlateColor GetForegroundColor() const override;
+	NIAGARAEDITOR_API virtual FSlateColor GetForegroundColor() const override;
 
-	void EnterEditingMode();
+	NIAGARAEDITOR_API void EnterEditingMode();
 private:
 	UEdGraphPin* TargetPin;
 	TSharedPtr<SNiagaraParameterNameTextBlock> ParameterNameTextBlock;
diff --git a/Engine/Plugins/FX/Niagara/Source/NiagaraEditor/Public/Widgets/SNiagaraPinTypeSelector.h b/Engine/Plugins/FX/Niagara/Source/NiagaraEditor/Public/Widgets/SNiagaraPinTypeSelector.h
--- a/Engine/Plugins/FX/Niagara/Source/NiagaraEditor/Public/Widgets/SNiagaraPinTypeSelector.h
+++ b/Engine/Plugins/FX/Niagara/Source/NiagaraEditor/Public/Widgets/SNiagaraPinTypeSelector.h
@@ -15,25 +15,25 @@
 #include "Widgets/SCompoundWidget.h"
 #include "Widgets/SWidget.h"
 
-class NIAGARAEDITOR_API SNiagaraPinTypeSelector : public SCompoundWidget
+class SNiagaraPinTypeSelector : public SCompoundWidget
 {
 public:
 	SLATE_BEGIN_ARGS(SNiagaraPinTypeSelector)
 	{}
 	SLATE_END_ARGS()
 public:
-	void Construct(const FArguments& InArgs, UEdGraphPin * InGraphPin);
+	NIAGARAEDITOR_API void Construct(const FArguments& InArgs, UEdGraphPin * InGraphPin);
 
 protected:
-	virtual TSharedRef<SWidget>	GetMenuContent();
-	FText GetTooltipText() const;
+	NIAGARAEDITOR_API virtual TSharedRef<SWidget>	GetMenuContent();
+	NIAGARAEDITOR_API FText GetTooltipText() const;
 
 private:
 	UEdGraphPin* Pin = nullptr;
 	TSharedPtr<SComboButton> SelectorButton;
 };
 
-class NIAGARAEDITOR_API SNiagaraIconWidget : public SGraphPaletteItem
+class SNiagaraIconWidget : public SGraphPaletteItem
 {
 public:
 
@@ -63,4 +63,4 @@ public:
 			)
 		];
 	}
-};
\ No newline at end of file
+};
diff --git a/Engine/Plugins/FX/Niagara/Source/NiagaraEditor/Public/Widgets/SNiagaraSelectedObjectsDetails.h b/Engine/Plugins/FX/Niagara/Source/NiagaraEditor/Public/Widgets/SNiagaraSelectedObjectsDetails.h
--- a/Engine/Plugins/FX/Niagara/Source/NiagaraEditor/Public/Widgets/SNiagaraSelectedObjectsDetails.h
+++ b/Engine/Plugins/FX/Niagara/Source/NiagaraEditor/Public/Widgets/SNiagaraSelectedObjectsDetails.h
@@ -27,7 +27,7 @@ public:
 
 	//~ Begin FEditorUndoClient Interface 
 	NIAGARAEDITOR_API virtual void PostUndo(bool bSuccess) override;
-	NIAGARAEDITOR_API virtual void PostRedo(bool bSuccess) override { PostUndo(bSuccess); };
+	virtual void PostRedo(bool bSuccess) override { PostUndo(bSuccess); };
 	//~ End FEditorUndoClient Interface 
 
 	/** Delegate to know when one of the properties has been changed.*/
diff --git a/Engine/Plugins/FX/Niagara/Source/NiagaraEditor/Public/Widgets/SVerticalResizeBox.h b/Engine/Plugins/FX/Niagara/Source/NiagaraEditor/Public/Widgets/SVerticalResizeBox.h
--- a/Engine/Plugins/FX/Niagara/Source/NiagaraEditor/Public/Widgets/SVerticalResizeBox.h
+++ b/Engine/Plugins/FX/Niagara/Source/NiagaraEditor/Public/Widgets/SVerticalResizeBox.h
@@ -7,7 +7,7 @@
 #include "Styling/SlateBrush.h"
 #include "Types/SlateStructs.h"
 
-class NIAGARAEDITOR_API SVerticalResizeBox : public SCompoundWidget
+class SVerticalResizeBox : public SCompoundWidget
 {
 public:
 	DECLARE_DELEGATE_OneParam(FOnContentHeightChanged, float);
@@ -27,13 +27,13 @@ public:
 		SLATE_DEFAULT_SLOT(FArguments, Content);
 	SLATE_END_ARGS();
 
-	void Construct(const FArguments& InArgs);
+	NIAGARAEDITOR_API void Construct(const FArguments& InArgs);
 
-	virtual FReply OnMouseButtonDown(const FGeometry& MyGeometry, const FPointerEvent& MouseEvent) override;
-	virtual FReply OnMouseButtonUp(const FGeometry& MyGeometry, const FPointerEvent& MouseEvent) override;
-	virtual FReply OnMouseMove(const FGeometry& MyGeometry, const FPointerEvent& MouseEvent) override;
-	virtual FCursorReply OnCursorQuery(const FGeometry& MyGeometry, const FPointerEvent& CursorEvent) const override;
-	virtual int32 OnPaint(const FPaintArgs& Args, const FGeometry& AllottedGeometry, const FSlateRect& MyClippingRect, FSlateWindowElementList& OutDrawElements, int32 LayerId, const FWidgetStyle& InWidgetStyle, bool bParentEnabled) const override;
+	NIAGARAEDITOR_API virtual FReply OnMouseButtonDown(const FGeometry& MyGeometry, const FPointerEvent& MouseEvent) override;
+	NIAGARAEDITOR_API virtual FReply OnMouseButtonUp(const FGeometry& MyGeometry, const FPointerEvent& MouseEvent) override;
+	NIAGARAEDITOR_API virtual FReply OnMouseMove(const FGeometry& MyGeometry, const FPointerEvent& MouseEvent) override;
+	NIAGARAEDITOR_API virtual FCursorReply OnCursorQuery(const FGeometry& MyGeometry, const FPointerEvent& CursorEvent) const override;
+	NIAGARAEDITOR_API virtual int32 OnPaint(const FPaintArgs& Args, const FGeometry& AllottedGeometry, const FSlateRect& MyClippingRect, FSlateWindowElementList& OutDrawElements, int32 LayerId, const FWidgetStyle& InWidgetStyle, bool bParentEnabled) const override;
 
 private:
 	FOptionalSize GetHeightOverride() const;
@@ -52,4 +52,4 @@ private:
 	FSlateBrush HandleBrush;
 
 	FOnContentHeightChanged ContentHeightChanged;
-};
\ No newline at end of file
+};
diff --git a/Engine/Plugins/FX/Niagara/Source/NiagaraShader/Public/NiagaraAsyncGpuTraceProvider.h b/Engine/Plugins/FX/Niagara/Source/NiagaraShader/Public/NiagaraAsyncGpuTraceProvider.h
--- a/Engine/Plugins/FX/Niagara/Source/NiagaraShader/Public/NiagaraAsyncGpuTraceProvider.h
+++ b/Engine/Plugins/FX/Niagara/Source/NiagaraShader/Public/NiagaraAsyncGpuTraceProvider.h
@@ -31,7 +31,7 @@ struct FNiagaraAsyncGpuTraceResult
 	float _Pad0; // padding to force 16 byte alignment to meet requirements for VK
 };
 
-class NIAGARASHADER_API FNiagaraAsyncGpuTraceProvider
+class FNiagaraAsyncGpuTraceProvider
 {
 public:
 	using EProviderType = ENDICollisionQuery_AsyncGpuTraceProvider::Type;
@@ -58,9 +58,9 @@ public:
 		uint32 MaxRetraceCount = 0;
 	};
 
-	static EProviderType ResolveSupportedType(EProviderType InType, const FProviderPriorityArray& Priorities);
-	static bool RequiresDistanceFieldData(EProviderType InType, const FProviderPriorityArray& Priorities);
-	static bool RequiresRayTracingScene(EProviderType InType, const FProviderPriorityArray& Priorities);
+	static NIAGARASHADER_API EProviderType ResolveSupportedType(EProviderType InType, const FProviderPriorityArray& Priorities);
+	static NIAGARASHADER_API bool RequiresDistanceFieldData(EProviderType InType, const FProviderPriorityArray& Priorities);
+	static NIAGARASHADER_API bool RequiresRayTracingScene(EProviderType InType, const FProviderPriorityArray& Priorities);
 
 	/** Hash table.
 		PrimIdHashTable is the main hash table that maps GPUSceneInstanceIndex to and Index we can use to store Collision Groups inside HashToCollisionGroups.
@@ -72,15 +72,15 @@ public:
 		uint32 HashTableSize = 0;
 	};
 
-	static void BuildCollisionGroupHashMap(FRHICommandList& RHICmdList, ERHIFeatureLevel::Type FeatureLevel, FSceneInterface* Scene, const TMap<FPrimitiveComponentId, uint32>& CollisionGroupMap, FCollisionGroupHashMap& Result);
+	static NIAGARASHADER_API void BuildCollisionGroupHashMap(FRHICommandList& RHICmdList, ERHIFeatureLevel::Type FeatureLevel, FSceneInterface* Scene, const TMap<FPrimitiveComponentId, uint32>& CollisionGroupMap, FCollisionGroupHashMap& Result);
 
-	static TArray<TUniquePtr<FNiagaraAsyncGpuTraceProvider>> CreateSupportedProviders(EShaderPlatform ShaderPlatform, FNiagaraGpuComputeDispatchInterface* Dispatcher, const FProviderPriorityArray& Priorities);
-	static void ClearResults(FRHICommandList& RHICmdList, EShaderPlatform ShaderPlatform, const FDispatchRequest& Request);
+	static NIAGARASHADER_API TArray<TUniquePtr<FNiagaraAsyncGpuTraceProvider>> CreateSupportedProviders(EShaderPlatform ShaderPlatform, FNiagaraGpuComputeDispatchInterface* Dispatcher, const FProviderPriorityArray& Priorities);
+	static NIAGARASHADER_API void ClearResults(FRHICommandList& RHICmdList, EShaderPlatform ShaderPlatform, const FDispatchRequest& Request);
 
 	virtual bool IsAvailable() const = 0;
 	virtual EProviderType GetType() const = 0;
 
-	virtual void PostRenderOpaque(FRHICommandList& RHICmdList, TConstStridedView<FSceneView> Views, TUniformBufferRef<FSceneUniformParameters> SceneUniformBufferRHI, FCollisionGroupHashMap* CollisionGroupHash);
+	NIAGARASHADER_API virtual void PostRenderOpaque(FRHICommandList& RHICmdList, TConstStridedView<FSceneView> Views, TUniformBufferRef<FSceneUniformParameters> SceneUniformBufferRHI, FCollisionGroupHashMap* CollisionGroupHash);
 
 	virtual void IssueTraces(FRHICommandList& RHICmdList, const FDispatchRequest& Request, TUniformBufferRef<FSceneUniformParameters> SceneUniformBufferRHI, FCollisionGroupHashMap* CollisionGroupHash)
 	{
diff --git a/Engine/Plugins/FX/Niagara/Source/NiagaraShader/Public/NiagaraBatchedElements.h b/Engine/Plugins/FX/Niagara/Source/NiagaraShader/Public/NiagaraBatchedElements.h
--- a/Engine/Plugins/FX/Niagara/Source/NiagaraShader/Public/NiagaraBatchedElements.h
+++ b/Engine/Plugins/FX/Niagara/Source/NiagaraShader/Public/NiagaraBatchedElements.h
@@ -9,7 +9,7 @@ class FTexture;
 /**
  * Batched element parameters for gathering attributes from different slices into a single color
  */
-class NIAGARASHADER_API FBatchedElementNiagara2DArrayAttribute : public FBatchedElementParameters
+class FBatchedElementNiagara2DArrayAttribute : public FBatchedElementParameters
 {
 public:
 	typedef TFunction<void(FRHITexture*&, FRHISamplerState*&)> FGetTextureAndSamplerDelegate;
@@ -21,7 +21,7 @@ public:
 	}
 
 	/** Binds vertex and pixel shaders for this element */
-	virtual void BindShaders(FRHICommandList& RHICmdList, FGraphicsPipelineStateInitializer& GraphicsPSOInit, ERHIFeatureLevel::Type InFeatureLevel, const FMatrix& InTransform, const float InGamma, const FMatrix& ColorWeights, const FTexture* Texture) override;
+	NIAGARASHADER_API virtual void BindShaders(FRHICommandList& RHICmdList, FGraphicsPipelineStateInitializer& GraphicsPSOInit, ERHIFeatureLevel::Type InFeatureLevel, const FMatrix& InTransform, const float InGamma, const FMatrix& ColorWeights, const FTexture* Texture) override;
 
 private:
 	FIntVector4 AttributeSlices;
@@ -31,7 +31,7 @@ private:
 /**
  * Batched element parameters for gathering attributes from a volume texture into a single color
  */
-class NIAGARASHADER_API FBatchedElementNiagaraVolumeAttribute : public FBatchedElementParameters
+class FBatchedElementNiagaraVolumeAttribute : public FBatchedElementParameters
 {
 public:
 	typedef TFunction<void(FRHITexture*&, FRHISamplerState*&)> FGetTextureAndSamplerDelegate;
@@ -45,7 +45,7 @@ public:
 	}
 
 	/** Binds vertex and pixel shaders for this element */
-	virtual void BindShaders(FRHICommandList& RHICmdList, FGraphicsPipelineStateInitializer& GraphicsPSOInit, ERHIFeatureLevel::Type InFeatureLevel, const FMatrix& InTransform, const float InGamma, const FMatrix& ColorWeights, const FTexture* Texture) override;
+	NIAGARASHADER_API virtual void BindShaders(FRHICommandList& RHICmdList, FGraphicsPipelineStateInitializer& GraphicsPSOInit, ERHIFeatureLevel::Type InFeatureLevel, const FMatrix& InTransform, const float InGamma, const FMatrix& ColorWeights, const FTexture* Texture) override;
 
 private:
 	FVector2f								TileUVs;
@@ -58,7 +58,7 @@ private:
 /**
  * Batched element parameters for inverting a color channel
  */
-class NIAGARASHADER_API FBatchedElementNiagaraInvertColorChannel : public FBatchedElementParameters
+class FBatchedElementNiagaraInvertColorChannel : public FBatchedElementParameters
 {
 public:
 	FBatchedElementNiagaraInvertColorChannel(uint32 InChannelMask)
@@ -67,7 +67,7 @@ public:
 	}
 
 	/** Binds vertex and pixel shaders for this element */
-	virtual void BindShaders(FRHICommandList& RHICmdList, FGraphicsPipelineStateInitializer& GraphicsPSOInit, ERHIFeatureLevel::Type InFeatureLevel, const FMatrix& InTransform, const float InGamma, const FMatrix& ColorWeights, const FTexture* Texture) override;
+	NIAGARASHADER_API virtual void BindShaders(FRHICommandList& RHICmdList, FGraphicsPipelineStateInitializer& GraphicsPSOInit, ERHIFeatureLevel::Type InFeatureLevel, const FMatrix& InTransform, const float InGamma, const FMatrix& ColorWeights, const FTexture* Texture) override;
 
 private:
 	uint32	ChannelMask;
@@ -76,7 +76,7 @@ private:
 /**
  * Simple batched element using a 2d texture
  */
-class NIAGARASHADER_API FBatchedElementNiagaraSimple : public FBatchedElementParameters
+class FBatchedElementNiagaraSimple : public FBatchedElementParameters
 {
 public:
 	FBatchedElementNiagaraSimple(FMatrix InColorTransform = FMatrix::Identity, bool InAlphaBlend = false)
@@ -86,7 +86,7 @@ public:
 	}
 
 	/** Binds vertex and pixel shaders for this element */
-	virtual void BindShaders(FRHICommandList& RHICmdList, FGraphicsPipelineStateInitializer& GraphicsPSOInit, ERHIFeatureLevel::Type InFeatureLevel, const FMatrix& InTransform, const float InGamma, const FMatrix& ColorWeights, const FTexture* Texture) override;
+	NIAGARASHADER_API virtual void BindShaders(FRHICommandList& RHICmdList, FGraphicsPipelineStateInitializer& GraphicsPSOInit, ERHIFeatureLevel::Type InFeatureLevel, const FMatrix& InTransform, const float InGamma, const FMatrix& ColorWeights, const FTexture* Texture) override;
 
 private:
 	FMatrix	ColorTransform;
diff --git a/Engine/Plugins/FX/Niagara/Source/NiagaraShader/Public/NiagaraGPUSceneUtils.h b/Engine/Plugins/FX/Niagara/Source/NiagaraShader/Public/NiagaraGPUSceneUtils.h
--- a/Engine/Plugins/FX/Niagara/Source/NiagaraShader/Public/NiagaraGPUSceneUtils.h
+++ b/Engine/Plugins/FX/Niagara/Source/NiagaraShader/Public/NiagaraGPUSceneUtils.h
@@ -10,7 +10,7 @@
 
 DECLARE_UNIFORM_BUFFER_STRUCT(FSceneUniformParameters, RENDERER_API)
 
-class NIAGARASHADER_API FNiagaraGPUSceneUtils
+class FNiagaraGPUSceneUtils
 {
 public:	
 	BEGIN_SHADER_PARAMETER_STRUCT(FUpdateMeshParticleInstancesParams, NIAGARASHADER_API)
@@ -35,7 +35,7 @@ public:
 		SHADER_PARAMETER(int, bNeedsPrevTransform)
 	END_SHADER_PARAMETER_STRUCT()
 
-	static void AddUpdateMeshParticleInstancesPass(
+	static NIAGARASHADER_API void AddUpdateMeshParticleInstancesPass(
 		FRDGBuilder& GraphBuilder,
 		FUpdateMeshParticleInstancesParams& Params,
 		ERHIFeatureLevel::Type FeatureLevel,
diff --git a/Engine/Plugins/FX/Niagara/Source/NiagaraShader/Public/NiagaraScriptBase.h b/Engine/Plugins/FX/Niagara/Source/NiagaraShader/Public/NiagaraScriptBase.h
--- a/Engine/Plugins/FX/Niagara/Source/NiagaraShader/Public/NiagaraScriptBase.h
+++ b/Engine/Plugins/FX/Niagara/Source/NiagaraShader/Public/NiagaraScriptBase.h
@@ -56,11 +56,11 @@ enum class ENiagaraSimStageExecuteBehavior : uint8
 };
 
 USTRUCT()
-struct NIAGARASHADER_API FSimulationStageMetaData
+struct FSimulationStageMetaData
 {
 	GENERATED_USTRUCT_BODY()
 public:
-	FSimulationStageMetaData();
+	NIAGARASHADER_API FSimulationStageMetaData();
 
 	/** User simulation stage name. */
 	UPROPERTY()
diff --git a/Engine/Plugins/FX/Niagara/Source/NiagaraShader/Public/NiagaraShader.h b/Engine/Plugins/FX/Niagara/Source/NiagaraShader/Public/NiagaraShader.h
--- a/Engine/Plugins/FX/Niagara/Source/NiagaraShader/Public/NiagaraShader.h
+++ b/Engine/Plugins/FX/Niagara/Source/NiagaraShader/Public/NiagaraShader.h
@@ -28,7 +28,7 @@ class UClass;
 template<typename TBufferStruct> class TUniformBufferRef;
 
 /** Base class of all shaders that need material parameters. */
-class NIAGARASHADER_API FNiagaraShader : public FShader
+class FNiagaraShader : public FShader
 {
 public:
 	DECLARE_SHADER_TYPE(FNiagaraShader, Niagara);
@@ -208,7 +208,7 @@ public:
 
 	using FPermutationParameters = FNiagaraShaderPermutationParameters;
 
-	static FName UniformBufferLayoutName;
+	static NIAGARASHADER_API FName UniformBufferLayoutName;
 
 	static FIntVector GetDefaultThreadGroupSize(ENiagaraGpuDispatchType DispatchType)
 	{
@@ -230,7 +230,7 @@ public:
 	}
 
 	FNiagaraShader() {}
-	FNiagaraShader(const FNiagaraShaderType::CompiledShaderInitializerType& Initializer);
+	NIAGARASHADER_API FNiagaraShader(const FNiagaraShaderType::CompiledShaderInitializerType& Initializer);
 
 	TConstArrayView<FNiagaraDataInterfaceParamRef> GetDIParameters()
 	{
diff --git a/Engine/Plugins/FX/Niagara/Source/NiagaraShader/Public/NiagaraShaderParametersBuilder.h b/Engine/Plugins/FX/Niagara/Source/NiagaraShader/Public/NiagaraShaderParametersBuilder.h
--- a/Engine/Plugins/FX/Niagara/Source/NiagaraShader/Public/NiagaraShaderParametersBuilder.h
+++ b/Engine/Plugins/FX/Niagara/Source/NiagaraShader/Public/NiagaraShaderParametersBuilder.h
@@ -5,7 +5,7 @@
 #include "NiagaraShared.h"
 #include "ShaderParameterMetadataBuilder.h"
 
-class NIAGARASHADER_API FNiagaraShaderParametersBuilder
+class FNiagaraShaderParametersBuilder
 {
 public:
 	explicit FNiagaraShaderParametersBuilder(const FNiagaraDataInterfaceGPUParamInfo& InGPUParamInfo, TArray<FString>& InLooseNames, TArray<FNiagaraDataInterfaceStructIncludeInfo>& InStructIncludeInfos, FShaderParametersMetadataBuilder& InMetadataBuilder)
diff --git a/Engine/Plugins/FX/Niagara/Source/NiagaraShader/Public/NiagaraShared.h b/Engine/Plugins/FX/Niagara/Source/NiagaraShader/Public/NiagaraShared.h
--- a/Engine/Plugins/FX/Niagara/Source/NiagaraShader/Public/NiagaraShared.h
+++ b/Engine/Plugins/FX/Niagara/Source/NiagaraShader/Public/NiagaraShared.h
@@ -133,7 +133,7 @@ using FNiagaraShaderRef = TNiagaraShaderRef<FNiagaraShader>;
 * Information about a data interface function generated by the translator.
 */
 USTRUCT()
-struct NIAGARASHADER_API FNiagaraDataInterfaceGeneratedFunction
+struct FNiagaraDataInterfaceGeneratedFunction
 {
 	GENERATED_USTRUCT_BODY()
 	
@@ -166,7 +166,7 @@ struct NIAGARASHADER_API FNiagaraDataInterfaceGeneratedFunction
 		return nullptr;
 	}
 
-	bool Serialize(FArchive& Ar);
+	NIAGARASHADER_API bool Serialize(FArchive& Ar);
 	friend bool operator<<(FArchive& Ar, FNiagaraDataInterfaceGeneratedFunction& DIFunction);
 
 	bool operator==(const FNiagaraDataInterfaceGeneratedFunction& Other) const
@@ -207,7 +207,7 @@ template<> struct TStructOpsTypeTraits<FNiagaraDataInterfaceGeneratedFunction> :
 * Data coming from that translator that describes parameters needed for each data interface.
 */
 USTRUCT()
-struct NIAGARASHADER_API FNiagaraDataInterfaceGPUParamInfo
+struct FNiagaraDataInterfaceGPUParamInfo
 {
 	GENERATED_USTRUCT_BODY()
 		
@@ -227,9 +227,9 @@ struct NIAGARASHADER_API FNiagaraDataInterfaceGPUParamInfo
 	UPROPERTY()
 	TArray<FNiagaraDataInterfaceGeneratedFunction> GeneratedFunctions;
 
-	bool IsUserParameter() const;
+	NIAGARASHADER_API bool IsUserParameter() const;
 
-	bool Serialize(FArchive& Ar);
+	NIAGARASHADER_API bool Serialize(FArchive& Ar);
 };
 
 template<> struct TStructOpsTypeTraits<FNiagaraDataInterfaceGPUParamInfo> : public TStructOpsTypeTraitsBase2<FNiagaraDataInterfaceGPUParamInfo>
@@ -262,7 +262,7 @@ struct FNiagaraDataInterfaceStructIncludeInfo
 
 /* Data provided from the translator to track information requires to bind data interface or parameters at runtime. */
 USTRUCT()
-struct NIAGARASHADER_API FNiagaraShaderScriptParametersMetadata
+struct FNiagaraShaderScriptParametersMetadata
 {
 	GENERATED_USTRUCT_BODY()
 
@@ -897,14 +897,14 @@ public:
 	TSharedRef<FNiagaraShaderScriptParametersMetadata> GetScriptParametersMetadata() const { check(ScriptParametersMetadata.IsValid()); return ScriptParametersMetadata.ToSharedRef(); }
 	const FNiagaraShaderScriptParametersMetadata& GetScriptParametersMetadata_RT() const { check(ScriptParametersMetadata_RT.IsValid()); return *ScriptParametersMetadata_RT.Get(); }
 
-	NIAGARASHADER_API FOnNiagaraScriptCompilationComplete& OnCompilationComplete()
+	FOnNiagaraScriptCompilationComplete& OnCompilationComplete()
 	{
 		return OnCompilationCompleteDelegate;
 	}
 
 	bool IsSame(const FNiagaraShaderMapId& InId) const;
 
-	NIAGARASHADER_API int32 GetNumPermutations() const { return NumPermutations; }
+	int32 GetNumPermutations() const { return NumPermutations; }
 
 	NIAGARASHADER_API bool IsShaderMapComplete() const;
 
diff --git a/Engine/Plugins/FX/Niagara/Source/NiagaraVertexFactories/Public/NiagaraCutoutVertexBuffer.h b/Engine/Plugins/FX/Niagara/Source/NiagaraVertexFactories/Public/NiagaraCutoutVertexBuffer.h
--- a/Engine/Plugins/FX/Niagara/Source/NiagaraVertexFactories/Public/NiagaraCutoutVertexBuffer.h
+++ b/Engine/Plugins/FX/Niagara/Source/NiagaraVertexFactories/Public/NiagaraCutoutVertexBuffer.h
@@ -12,17 +12,17 @@ NiagaraCutoutVertexBuffer.h: Niagara cutout uv buffer.
 /**
  * Vertex buffer to hold cutout UVs. 
  */
-class NIAGARAVERTEXFACTORIES_API FNiagaraCutoutVertexBuffer : public FVertexBuffer
+class FNiagaraCutoutVertexBuffer : public FVertexBuffer
 {
 public:
 
-	FNiagaraCutoutVertexBuffer(int32 ZeroInitCount = 0);
+	NIAGARAVERTEXFACTORIES_API FNiagaraCutoutVertexBuffer(int32 ZeroInitCount = 0);
 
 	/**
 	 * Initialize the RHI for this rendering resource
 	 */
-	virtual void InitRHI() override;
-	virtual void ReleaseRHI() override;
+	NIAGARAVERTEXFACTORIES_API virtual void InitRHI() override;
+	NIAGARAVERTEXFACTORIES_API virtual void ReleaseRHI() override;
 	FShaderResourceViewRHIRef VertexBufferSRV;
 
 	/** Data to initialize the buffer. */
diff --git a/Engine/Plugins/FX/Niagara/Source/NiagaraVertexFactories/Public/NiagaraDispatchIndirect.h b/Engine/Plugins/FX/Niagara/Source/NiagaraVertexFactories/Public/NiagaraDispatchIndirect.h
--- a/Engine/Plugins/FX/Niagara/Source/NiagaraVertexFactories/Public/NiagaraDispatchIndirect.h
+++ b/Engine/Plugins/FX/Niagara/Source/NiagaraVertexFactories/Public/NiagaraDispatchIndirect.h
@@ -29,9 +29,9 @@ struct FNiagaraDispatchIndirectParametersCS
 	FUintVector4	ThreadCount = FUintVector4(0, 0, 0, 0);
 };
 
-class NIAGARAVERTEXFACTORIES_API FNiagaraDispatchIndirectArgsGenCS : public FGlobalShader
+class FNiagaraDispatchIndirectArgsGenCS : public FGlobalShader
 {
-	DECLARE_GLOBAL_SHADER(FNiagaraDispatchIndirectArgsGenCS);
+	DECLARE_EXPORTED_GLOBAL_SHADER(FNiagaraDispatchIndirectArgsGenCS, NIAGARAVERTEXFACTORIES_API);
 	SHADER_USE_PARAMETER_STRUCT(FNiagaraDispatchIndirectArgsGenCS, FGlobalShader);
 
 public:
@@ -46,5 +46,5 @@ public:
 		SHADER_PARAMETER_RDG_BUFFER_UAV(RWBuffer<uint4>,	RWDispatchIndirectArgs)
 	END_SHADER_PARAMETER_STRUCT()
 
-	static void ModifyCompilationEnvironment(const FGlobalShaderPermutationParameters& Parameters, FShaderCompilerEnvironment& OutEnvironment);
+	static NIAGARAVERTEXFACTORIES_API void ModifyCompilationEnvironment(const FGlobalShaderPermutationParameters& Parameters, FShaderCompilerEnvironment& OutEnvironment);
 };
diff --git a/Engine/Plugins/FX/Niagara/Source/NiagaraVertexFactories/Public/NiagaraMeshVertexFactory.h b/Engine/Plugins/FX/Niagara/Source/NiagaraVertexFactories/Public/NiagaraMeshVertexFactory.h
--- a/Engine/Plugins/FX/Niagara/Source/NiagaraVertexFactories/Public/NiagaraMeshVertexFactory.h
+++ b/Engine/Plugins/FX/Niagara/Source/NiagaraVertexFactories/Public/NiagaraMeshVertexFactory.h
@@ -121,9 +121,9 @@ class FNiagaraMeshInstanceVertices;
 /**
 * Vertex factory for rendering instanced mesh particles with out dynamic parameter support.
 */
-class NIAGARAVERTEXFACTORIES_API FNiagaraMeshVertexFactory : public FNiagaraVertexFactoryBase
+class FNiagaraMeshVertexFactory : public FNiagaraVertexFactoryBase
 {
-	DECLARE_VERTEX_FACTORY_TYPE(FNiagaraMeshVertexFactory);
+	DECLARE_VERTEX_FACTORY_TYPE_API(FNiagaraMeshVertexFactory, NIAGARAVERTEXFACTORIES_API);
 public:
 	
 	/** Default constructor. */
@@ -144,25 +144,25 @@ public:
 	/**
 	* Should we cache the material's shadertype on this platform with this vertex factory?
 	*/
-	static bool ShouldCompilePermutation(const FVertexFactoryShaderPermutationParameters& Parameters);
+	static NIAGARAVERTEXFACTORIES_API bool ShouldCompilePermutation(const FVertexFactoryShaderPermutationParameters& Parameters);
 
 
 	/**
 	* Modify compile environment to enable instancing
 	* @param OutEnvironment - shader compile environment to modify
 	*/
-	static void ModifyCompilationEnvironment(const FVertexFactoryShaderPermutationParameters& Parameters, FShaderCompilerEnvironment& OutEnvironment);
+	static NIAGARAVERTEXFACTORIES_API void ModifyCompilationEnvironment(const FVertexFactoryShaderPermutationParameters& Parameters, FShaderCompilerEnvironment& OutEnvironment);
 
 	/**
 	* Get vertex elements used when during PSO precaching materials using this vertex factory type
 	*/
-	static void GetPSOPrecacheVertexFetchElements(EVertexInputStreamType VertexInputStreamType, FVertexDeclarationElementList& Elements);
-	static void GetVertexElements(ERHIFeatureLevel::Type FeatureLevel, bool bSupportsManualVertexFetch, FStaticMeshDataType& Data, FVertexDeclarationElementList& Elements);
+	static NIAGARAVERTEXFACTORIES_API void GetPSOPrecacheVertexFetchElements(EVertexInputStreamType VertexInputStreamType, FVertexDeclarationElementList& Elements);
+	static NIAGARAVERTEXFACTORIES_API void GetVertexElements(ERHIFeatureLevel::Type FeatureLevel, bool bSupportsManualVertexFetch, FStaticMeshDataType& Data, FVertexDeclarationElementList& Elements);
 
 	/**
 	* An implementation of the interface used by TSynchronizedResource to update the resource with new data from the game thread.
 	*/
-	void SetData(const FStaticMeshDataType& InData);
+	NIAGARAVERTEXFACTORIES_API void SetData(const FStaticMeshDataType& InData);
 
 	/**
 	* Set the uniform buffer for this vertex factory.
@@ -206,7 +206,7 @@ public:
 	}
 
 	// FRenderResource interface.
-	virtual void InitRHI() override;
+	NIAGARAVERTEXFACTORIES_API virtual void InitRHI() override;
 
 	int32 GetMeshIndex() const { return MeshIndex; }
 	void SetMeshIndex(int32 InMeshIndex) { MeshIndex = InMeshIndex; }
@@ -222,7 +222,7 @@ public:
 		FNiagaraVertexFactoryBase::SetFeatureLevel(InFeatureLevel);
 	}
 protected:
-	static void GetVertexElements(ERHIFeatureLevel::Type FeatureLevel, bool bSupportsManualVertexFetch, FStaticMeshDataType& Data, FVertexDeclarationElementList& Elements, FVertexStreamList& InOutStreams);
+	static NIAGARAVERTEXFACTORIES_API void GetVertexElements(ERHIFeatureLevel::Type FeatureLevel, bool bSupportsManualVertexFetch, FStaticMeshDataType& Data, FVertexDeclarationElementList& Elements, FVertexStreamList& InOutStreams);
 
 protected:
 	FStaticMeshDataType Data;
diff --git a/Engine/Plugins/FX/Niagara/Source/NiagaraVertexFactories/Public/NiagaraRibbonVertexFactory.h b/Engine/Plugins/FX/Niagara/Source/NiagaraVertexFactories/Public/NiagaraRibbonVertexFactory.h
--- a/Engine/Plugins/FX/Niagara/Source/NiagaraVertexFactories/Public/NiagaraRibbonVertexFactory.h
+++ b/Engine/Plugins/FX/Niagara/Source/NiagaraVertexFactories/Public/NiagaraRibbonVertexFactory.h
@@ -89,9 +89,9 @@ typedef TUniformBufferRef<FNiagaraRibbonVFLooseParameters> FNiagaraRibbonVFLoose
 /**
 * Beam/Trail particle vertex factory.
 */
-class NIAGARAVERTEXFACTORIES_API FNiagaraRibbonVertexFactory : public FNiagaraVertexFactoryBase
+class FNiagaraRibbonVertexFactory : public FNiagaraVertexFactoryBase
 {
-	DECLARE_VERTEX_FACTORY_TYPE(FNiagaraRibbonVertexFactory);
+	DECLARE_VERTEX_FACTORY_TYPE_API(FNiagaraRibbonVertexFactory, NIAGARAVERTEXFACTORIES_API);
 
 public:
 
@@ -109,20 +109,20 @@ public:
 	/**
 	* Should we cache the material's shadertype on this platform with this vertex factory?
 	*/
-	static bool ShouldCompilePermutation(const FVertexFactoryShaderPermutationParameters& Parameters);
+	static NIAGARAVERTEXFACTORIES_API bool ShouldCompilePermutation(const FVertexFactoryShaderPermutationParameters& Parameters);
 
 	/**
 	* Can be overridden by FVertexFactory subclasses to modify their compile environment just before compilation occurs.
 	*/
-	static void ModifyCompilationEnvironment(const FVertexFactoryShaderPermutationParameters& Parameters, FShaderCompilerEnvironment& OutEnvironment);
+	static NIAGARAVERTEXFACTORIES_API void ModifyCompilationEnvironment(const FVertexFactoryShaderPermutationParameters& Parameters, FShaderCompilerEnvironment& OutEnvironment);
 
 	/**
 	* Get vertex elements used when during PSO precaching materials using this vertex factory type
 	*/
-	static void GetPSOPrecacheVertexFetchElements(EVertexInputStreamType VertexInputStreamType, FVertexDeclarationElementList& Elements);
+	static NIAGARAVERTEXFACTORIES_API void GetPSOPrecacheVertexFetchElements(EVertexInputStreamType VertexInputStreamType, FVertexDeclarationElementList& Elements);
 
 	// FRenderResource interface.
-	virtual void InitRHI() override;
+	NIAGARAVERTEXFACTORIES_API virtual void InitRHI() override;
 
 	/**
 	* Set the uniform buffer for this vertex factory.
@@ -143,12 +143,12 @@ public:
 	/**
 	* Set the source vertex buffer.
 	*/
-	void SetVertexBuffer(const FVertexBuffer* InBuffer, uint32 StreamOffset, uint32 Stride);
+	NIAGARAVERTEXFACTORIES_API void SetVertexBuffer(const FVertexBuffer* InBuffer, uint32 StreamOffset, uint32 Stride);
 
 	/**
 	* Set the source vertex buffer that contains particle dynamic parameter data.
 	*/
-	void SetDynamicParameterBuffer(const FVertexBuffer* InDynamicParameterBuffer, int32 ParameterIndex, uint32 StreamOffset, uint32 Stride);
+	NIAGARAVERTEXFACTORIES_API void SetDynamicParameterBuffer(const FVertexBuffer* InDynamicParameterBuffer, int32 ParameterIndex, uint32 StreamOffset, uint32 Stride);
 
 
 	FUniformBufferRHIRef LooseParameterUniformBuffer;
@@ -160,4 +160,4 @@ private:
 
 
 	const FNiagaraDataSet *DataSet;
-};
\ No newline at end of file
+};
diff --git a/Engine/Plugins/FX/Niagara/Source/NiagaraVertexFactories/Public/NiagaraSpriteVertexFactory.h b/Engine/Plugins/FX/Niagara/Source/NiagaraVertexFactories/Public/NiagaraSpriteVertexFactory.h
--- a/Engine/Plugins/FX/Niagara/Source/NiagaraVertexFactories/Public/NiagaraSpriteVertexFactory.h
+++ b/Engine/Plugins/FX/Niagara/Source/NiagaraVertexFactories/Public/NiagaraSpriteVertexFactory.h
@@ -110,9 +110,9 @@ typedef TUniformBufferRef<FNiagaraSpriteVFLooseParameters> FNiagaraSpriteVFLoose
 /**
  * Vertex factory for rendering particle sprites.
  */
-class NIAGARAVERTEXFACTORIES_API FNiagaraSpriteVertexFactory : public FNiagaraVertexFactoryBase
+class FNiagaraSpriteVertexFactory : public FNiagaraVertexFactoryBase
 {
-	DECLARE_VERTEX_FACTORY_TYPE(FNiagaraSpriteVertexFactory);
+	DECLARE_VERTEX_FACTORY_TYPE_API(FNiagaraSpriteVertexFactory, NIAGARAVERTEXFACTORIES_API);
 
 public:
 
@@ -140,26 +140,26 @@ public:
 	{}
 
 	// FRenderResource interface.
-	virtual void InitRHI() override;
+	NIAGARAVERTEXFACTORIES_API virtual void InitRHI() override;
 
 	virtual bool RendersPrimitivesAsCameraFacingSprites() const override { return true; }
 
 	/**
 	 * Should we cache the material's shadertype on this platform with this vertex factory? 
 	 */
-	static bool ShouldCompilePermutation(const FVertexFactoryShaderPermutationParameters& Parameters);
+	static NIAGARAVERTEXFACTORIES_API bool ShouldCompilePermutation(const FVertexFactoryShaderPermutationParameters& Parameters);
 
 	/**
 	 * Can be overridden by FVertexFactory subclasses to modify their compile environment just before compilation occurs.
 	 */
-	static void ModifyCompilationEnvironment(const FVertexFactoryShaderPermutationParameters& Parameters, FShaderCompilerEnvironment& OutEnvironment);
+	static NIAGARAVERTEXFACTORIES_API void ModifyCompilationEnvironment(const FVertexFactoryShaderPermutationParameters& Parameters, FShaderCompilerEnvironment& OutEnvironment);
 
 	/**
 	* Get vertex elements used when during PSO precaching materials using this vertex factory type
 	*/
-	static void GetPSOPrecacheVertexFetchElements(EVertexInputStreamType VertexInputStreamType, FVertexDeclarationElementList& Elements);
+	static NIAGARAVERTEXFACTORIES_API void GetPSOPrecacheVertexFetchElements(EVertexInputStreamType VertexInputStreamType, FVertexDeclarationElementList& Elements);
 
-	void SetTexCoordBuffer(const FVertexBuffer* InTexCoordBuffer);
+	NIAGARAVERTEXFACTORIES_API void SetTexCoordBuffer(const FVertexBuffer* InTexCoordBuffer);
 	
 	/**
 	 * Set the uniform buffer for this vertex factory.
@@ -230,7 +230,7 @@ public:
 
 protected:
 	/** Initialize streams for this vertex factory. */
-	void InitStreams();
+	NIAGARAVERTEXFACTORIES_API void InitStreams();
 
 private:
 
