better error reporting

diff --git a/Engine/Plugins/FX/Niagara/Source/Niagara/Private/NiagaraDataInterfaceRenderTargetVolume.cpp b/Engine/Plugins/FX/Niagara/Source/Niagara/Private/NiagaraDataInterfaceRenderTargetVolume.cpp
--- a/Engine/Plugins/FX/Niagara/Source/Niagara/Private/NiagaraDataInterfaceRenderTargetVolume.cpp
+++ b/Engine/Plugins/FX/Niagara/Source/Niagara/Private/NiagaraDataInterfaceRenderTargetVolume.cpp
@@ -686,8 +686,15 @@ bool UNiagaraDataInterfaceRenderTargetVolume::SimCacheWriteFrame(UObject* Storag
 				SVTCreateInfo.AttributesFormats[1] = PF_Unknown;
 
 				UE::SVT::FTextureData SparseTextureData{};
-				SparseTextureData.CreateFromDense(SVTCreateInfo, TArrayView<uint8, int64>((uint8*)TextureData.GetData(), (int64)TextureData.Num() * sizeof(TextureData[0])), TArrayView<uint8>());
+				bool Success = SparseTextureData.CreateFromDense(SVTCreateInfo, TArrayView<uint8, int64>((uint8*)TextureData.GetData(), (int64)TextureData.Num() * sizeof(TextureData[0])), TArrayView<uint8>());
 				
+				if (!Success)
+				{
+					UE_LOG(LogNiagara, Error, TEXT("Cannot create SVT for volume render target"));
+
+					return false;
+				}
+
 				CurrCache->AppendFrame(SparseTextureData);
 #endif
 			}
@@ -758,6 +765,7 @@ bool UNiagaraDataInterfaceRenderTargetVolume::SimCacheReadFrame(UObject* Storage
 	if (Cast<UVolumeCache>(StorageObject))
 	{
 		UE_LOG(LogNiagara, Error, TEXT("vdb caches are no longer supported.  Regenerate your Sim Cache, which will use SVT to store volume data."));
+
 	}
 	else if (Cast<UAnimatedSparseVolumeTexture>(StorageObject))
 	{			
