Early out of NDC spwaning functions if the target emitter is not active.

diff --git a/Engine/Plugins/FX/Niagara/Source/Niagara/Classes/NiagaraEmitterInstance.h b/Engine/Plugins/FX/Niagara/Source/Niagara/Classes/NiagaraEmitterInstance.h
--- a/Engine/Plugins/FX/Niagara/Source/Niagara/Classes/NiagaraEmitterInstance.h
+++ b/Engine/Plugins/FX/Niagara/Source/Niagara/Classes/NiagaraEmitterInstance.h
@@ -87,6 +87,7 @@ public:
 
 	FNiagaraDataSet& GetData()const { return *ParticleDataSet; }
 
+	FORCEINLINE bool IsActive()const { return ExecutionState == ENiagaraExecutionState::Active; }
 	FORCEINLINE bool IsDisabled()const { return ExecutionState == ENiagaraExecutionState::Disabled; }
 	FORCEINLINE bool IsInactive()const { return ExecutionState == ENiagaraExecutionState::Inactive; }
 	FORCEINLINE bool IsComplete()const { return ExecutionState == ENiagaraExecutionState::Complete || ExecutionState == ENiagaraExecutionState::Disabled; }
diff --git a/Engine/Plugins/FX/Niagara/Source/Niagara/Private/DataInterface/NiagaraDataInterfaceDataChannelRead.cpp b/Engine/Plugins/FX/Niagara/Source/Niagara/Private/DataInterface/NiagaraDataInterfaceDataChannelRead.cpp
--- a/Engine/Plugins/FX/Niagara/Source/Niagara/Private/DataInterface/NiagaraDataInterfaceDataChannelRead.cpp
+++ b/Engine/Plugins/FX/Niagara/Source/Niagara/Private/DataInterface/NiagaraDataInterfaceDataChannelRead.cpp
@@ -1094,7 +1094,7 @@ void UNiagaraDataInterfaceDataChannelRead::SpawnConditional(FVectorVMExternalFun
 	FNiagaraDataChannelData* DataChannelData = InstData->DataChannelData.Get();
 	FNiagaraDataBuffer* Data = DataChannelData ? DataChannelData->GetCPUData(!bReadCurrentFrame) : nullptr;
 
-	bool bSpawn = INiagaraModule::DataChannelsEnabled() && Data && Data->GetNumInstances() > 0 && EmittterInst && InEnabled.GetAndAdvance();
+	bool bSpawn = INiagaraModule::DataChannelsEnabled() && Data && Data->GetNumInstances() > 0 && EmittterInst && EmittterInst->IsActive() && InEnabled.GetAndAdvance();
 	if(bSpawn)
 	{
 		FNDIRandomHelperFromStream RandHelper(Context);
