Add the perf graph mode and sampling type to the overview title box when in perf graph mode.
Easier to tell at a glance what stats are being shown.

diff --git a/Engine/Plugins/FX/Niagara/Source/Niagara/Private/NiagaraDebugHud.cpp b/Engine/Plugins/FX/Niagara/Source/Niagara/Private/NiagaraDebugHud.cpp
--- a/Engine/Plugins/FX/Niagara/Source/Niagara/Private/NiagaraDebugHud.cpp
+++ b/Engine/Plugins/FX/Niagara/Source/Niagara/Private/NiagaraDebugHud.cpp
@@ -1946,7 +1946,21 @@ void FNiagaraDebugHud::DrawOverview(class FNiagaraWorldManager* WorldManager, FC
 			TextLocation.Y += fAdvanceHeight;
 
 			const UEnum* EnumPtr = FindObjectChecked<UEnum>(nullptr, TEXT("/Script/Niagara.ENiagaraDebugHUDOverviewMode"), true);
-			DrawCanvas->DrawShadowedString(TextLocation.X, TextLocation.Y, *EnumPtr->GetDisplayNameTextByValue((int32)Settings.OverviewMode).ToString(), Font, HeadingColor);
+			TStringBuilder<128> ModelStringBuilder;
+			ModelStringBuilder.Append(EnumPtr->GetDisplayNameTextByValue((int32)Settings.OverviewMode).ToString());
+			if(Settings.OverviewMode == ENiagaraDebugHUDOverviewMode::PerformanceGraph)
+			{
+				const UEnum* PerfGraphModeEnumPtr = FindObjectChecked<UEnum>(nullptr, TEXT("/Script/Niagara.ENiagaraDebugHUDPerfGraphMode"), true);
+				const UEnum* PerfGraphSampleModeEnumPtr = FindObjectChecked<UEnum>(nullptr, TEXT("/Script/Niagara.ENiagaraDebugHUDPerfSampleMode"), true);
+				
+				ModelStringBuilder.Append(TEXT(" - "));
+				ModelStringBuilder.Append(PerfGraphModeEnumPtr->GetDisplayNameTextByValue((int32)Settings.PerfGraphMode).ToString());
+				ModelStringBuilder.Append(TEXT(" - "));
+				ModelStringBuilder.Append(PerfGraphSampleModeEnumPtr->GetDisplayNameTextByValue((int32)Settings.PerfSampleMode).ToString());
+			}			
+			DrawCanvas->DrawShadowedString(TextLocation.X, TextLocation.Y, ModelStringBuilder.ToString(), Font, HeadingColor);
+			ModelStringBuilder.Reset();
+
 			TextLocation.Y += fAdvanceHeight;
 
 			if (OverviewString.Len() > 0)
