Refactor LODDistance calcualtions out of FNiagaraSystemIntance into the world manager allowing reuse for NDC work in a later CL.

diff --git a/Engine/Plugins/FX/Niagara/Source/Niagara/Private/NiagaraSystemInstance.cpp b/Engine/Plugins/FX/Niagara/Source/Niagara/Private/NiagaraSystemInstance.cpp
--- a/Engine/Plugins/FX/Niagara/Source/Niagara/Private/NiagaraSystemInstance.cpp
+++ b/Engine/Plugins/FX/Niagara/Source/Niagara/Private/NiagaraSystemInstance.cpp
@@ -1776,53 +1776,7 @@ float FNiagaraSystemInstance::GetLODDistance()
 	const FVector EffectLocation = WorldTransform.GetLocation() + (FVector(LWCTile) * FLargeWorldRenderScalar::GetTileSize());
 	LODDistance = DefaultLODDistance;
 
-	// If we are inside the WorldManager tick we will use the cache player view locations as we can be ticked on different threads
-	if (WorldManager->GetCachedViewInfo().Num() > 0)
-	{
-		// We are being ticked inside the WorldManager and can safely use the list of cached player view locations
-		FVector::FReal LODDistanceSqr = FMath::Square(WORLD_MAX);
-		for (const FNiagaraCachedViewInfo& ViewInfo : WorldManager->GetCachedViewInfo())
-		{
-			const FVector::FReal DistanceToEffectSqr = FVector(ViewInfo.ViewToWorld.GetOrigin() - EffectLocation).SizeSquared();
-			LODDistanceSqr = FMath::Min(LODDistanceSqr, DistanceToEffectSqr);
-		}
-		LODDistance = float(FMath::Sqrt(LODDistanceSqr));
-	}
-	else
-	{
-		// If we are not inside the WorldManager tick (solo tick) we must look over the player view locations manually
-		ensureMsgf(IsInGameThread(), TEXT("FNiagaraSystemInstance::GetLODDistance called in potentially thread unsafe way"));
-
-		TArray<FVector, TInlineAllocator<8> > PlayerViewLocations;
-		if (World->GetPlayerControllerIterator())
-		{
-			for (FConstPlayerControllerIterator Iterator = World->GetPlayerControllerIterator(); Iterator; ++Iterator)
-			{
-				APlayerController* PlayerController = Iterator->Get();
-				if (PlayerController && PlayerController->IsLocalPlayerController())
-				{
-					FVector* ViewLocation = new(PlayerViewLocations) FVector;
-					FRotator ViewRotation;
-					PlayerController->GetPlayerViewPoint(*ViewLocation, ViewRotation);
-				}
-			}
-		}
-		else
-		{
-			PlayerViewLocations = World->ViewLocationsRenderedLastFrame;
-		}
-
-		if (PlayerViewLocations.Num() > 0)
-		{
-			FVector::FReal LODDistanceSqr = FMath::Square(WORLD_MAX);
-			for (const FVector& ViewLocation : PlayerViewLocations)
-			{
-				const FVector::FReal DistanceToEffectSqr = FVector(ViewLocation - EffectLocation).SizeSquared();
-				LODDistanceSqr = FMath::Min(LODDistanceSqr, DistanceToEffectSqr);
-			}
-			LODDistance = float(FMath::Sqrt(LODDistanceSqr));
-		}
-	}
+	LODDistance = WorldManager->GetLODDistance(EffectLocation);
 
 	bLODDistanceIsValid = true;
 	return LODDistance;
diff --git a/Engine/Plugins/FX/Niagara/Source/Niagara/Private/NiagaraWorldManager.cpp b/Engine/Plugins/FX/Niagara/Source/Niagara/Private/NiagaraWorldManager.cpp
--- a/Engine/Plugins/FX/Niagara/Source/Niagara/Private/NiagaraWorldManager.cpp
+++ b/Engine/Plugins/FX/Niagara/Source/Niagara/Private/NiagaraWorldManager.cpp
@@ -1363,6 +1363,59 @@ UWorld* FNiagaraWorldManager::GetWorld()
 
 //////////////////////////////////////////////////////////////////////////
 
+FVector::FReal FNiagaraWorldManager::GetLODDistance(FVector Location)const
+{
+	FVector::FReal LODDistance = 0.0f;
+	// If we are inside the WorldManager tick we will use the cache player view locations as we can be ticked on different threads
+	if (GetCachedViewInfo().Num() > 0)
+	{
+		// We are being ticked inside the WorldManager and can safely use the list of cached player view locations
+		FVector::FReal LODDistanceSqr = FMath::Square(WORLD_MAX);
+		for (const FNiagaraCachedViewInfo& ViewInfo : GetCachedViewInfo())
+		{
+			const FVector::FReal DistanceToEffectSqr = FVector(ViewInfo.ViewToWorld.GetOrigin() - Location).SizeSquared();
+			LODDistanceSqr = FMath::Min(LODDistanceSqr, DistanceToEffectSqr);
+		}
+		LODDistance = FVector::FReal(FMath::Sqrt(LODDistanceSqr));
+	}
+	else
+	{
+		// If we are not inside the WorldManager tick (solo tick) we must look over the player view locations manually
+		ensureMsgf(IsInGameThread(), TEXT("FNiagaraSystemInstance::GetLODDistance called in potentially thread unsafe way"));
+
+		TArray<FVector, TInlineAllocator<8> > PlayerViewLocations;
+		if (World->GetPlayerControllerIterator())
+		{
+			for (FConstPlayerControllerIterator Iterator = World->GetPlayerControllerIterator(); Iterator; ++Iterator)
+			{
+				APlayerController* PlayerController = Iterator->Get();
+				if (PlayerController && PlayerController->IsLocalPlayerController())
+				{
+					FVector* ViewLocation = new(PlayerViewLocations) FVector;
+					FRotator ViewRotation;
+					PlayerController->GetPlayerViewPoint(*ViewLocation, ViewRotation);
+				}
+			}
+		}
+		else
+		{
+			PlayerViewLocations = World->ViewLocationsRenderedLastFrame;
+		}
+
+		if (PlayerViewLocations.Num() > 0)
+		{
+			FVector::FReal LODDistanceSqr = FMath::Square(WORLD_MAX);
+			for (const FVector& ViewLocation : PlayerViewLocations)
+			{
+				const FVector::FReal DistanceToEffectSqr = FVector(ViewLocation - Location).SizeSquared();
+				LODDistanceSqr = FMath::Min(LODDistanceSqr, DistanceToEffectSqr);
+			}
+			LODDistance = FVector::FReal(FMath::Sqrt(LODDistanceSqr));
+		}
+	}
+	return LODDistance;
+}
+
 void FNiagaraWorldManager::UpdateScalabilityManagers(float DeltaSeconds, bool bNewSpawnsOnly)
 {
 	SCOPE_CYCLE_COUNTER(STAT_UpdateScalabilityManagers);
diff --git a/Engine/Plugins/FX/Niagara/Source/Niagara/Public/NiagaraWorldManager.h b/Engine/Plugins/FX/Niagara/Source/Niagara/Public/NiagaraWorldManager.h
--- a/Engine/Plugins/FX/Niagara/Source/Niagara/Public/NiagaraWorldManager.h
+++ b/Engine/Plugins/FX/Niagara/Source/Niagara/Public/NiagaraWorldManager.h
@@ -161,6 +161,9 @@ public:
 
 	TArrayView<const FNiagaraCachedViewInfo> GetCachedViewInfo() const { return MakeArrayView(CachedViewInfo); }
 
+	//Returns the distance to the nearest viewpoint to the give location. Used for a distance on which to base LODs.
+	NIAGARA_API FVector::FReal GetLODDistance(FVector Location)const;
+
 	UNiagaraComponentPool* GetComponentPool() { return ComponentPool; }
 
 	void UpdateScalabilityManagers(float DeltaSeconds, bool bNewSpawnsOnly);
