Defer addition and deletion of NDCs to the next frame.
Must happen on GT.
Adding GTOnly hook to allow cases where a sync load will generate a new NDC to refresh them immediately.
Changed NDC handler to use a weakptr so NDCs that are not used by other things in the scene can be destroyed.

diff --git a/Engine/Plugins/FX/Niagara/Source/Niagara/Private/NiagaraDataChannel.cpp b/Engine/Plugins/FX/Niagara/Source/Niagara/Private/NiagaraDataChannel.cpp
--- a/Engine/Plugins/FX/Niagara/Source/Niagara/Private/NiagaraDataChannel.cpp
+++ b/Engine/Plugins/FX/Niagara/Source/Niagara/Private/NiagaraDataChannel.cpp
@@ -805,12 +805,7 @@ namespace NiagaraDataChannel
 void UNiagaraDataChannel::PostInitProperties()
 {
 	Super::PostInitProperties();
-
-	FNiagaraWorldManager::ForAllWorldManagers(
-		[DataChannel = this](FNiagaraWorldManager& WorldMan)
-		{
-			WorldMan.InitDataChannel(DataChannel, true);
-		});
+	INiagaraModule::RequestRefreshDataChannels();
 }
 
 void UNiagaraDataChannel::PostLoad()
@@ -851,22 +846,13 @@ void UNiagaraDataChannel::PostLoad()
 	//TODO: Can serialize?
 	GameDataLayout.Init(Variables);
 
-	FNiagaraWorldManager::ForAllWorldManagers(
-		[DataChannel = this](FNiagaraWorldManager& WorldMan)
-		{
-			WorldMan.InitDataChannel(DataChannel, true);
-		});
+	INiagaraModule::RequestRefreshDataChannels();
 }
 
 void UNiagaraDataChannel::BeginDestroy()
 {
 	Super::BeginDestroy();
-
-	FNiagaraWorldManager::ForAllWorldManagers(
-		[DataChannel = this](FNiagaraWorldManager& WorldMan)
-		{
-			WorldMan.RemoveDataChannel(DataChannel);
-		});
+	INiagaraModule::RequestRefreshDataChannels();
 }
 
 #if WITH_EDITOR
@@ -894,11 +880,7 @@ void UNiagaraDataChannel::PostEditChangeProperty(FPropertyChangedEvent& Property
 
 	GameDataLayout.Init(Variables);
 
-	FNiagaraWorldManager::ForAllWorldManagers(
-		[DataChannel = this](FNiagaraWorldManager& WorldMan)
-		{
-			WorldMan.InitDataChannel(DataChannel, true);
-		});
+	INiagaraModule::RequestRefreshDataChannels();
 }
 
 #endif//WITH_EDITOR
diff --git a/Engine/Plugins/FX/Niagara/Source/Niagara/Private/NiagaraDataChannelManager.cpp b/Engine/Plugins/FX/Niagara/Source/Niagara/Private/NiagaraDataChannelManager.cpp
--- a/Engine/Plugins/FX/Niagara/Source/Niagara/Private/NiagaraDataChannelManager.cpp
+++ b/Engine/Plugins/FX/Niagara/Source/Niagara/Private/NiagaraDataChannelManager.cpp
@@ -21,6 +21,14 @@ void FNiagaraDataChannelManager::AddReferencedObjects(FReferenceCollector& Colle
 	Collector.AddReferencedObjects(Channels);
 }
 
+void FNiagaraDataChannelManager::RefreshDataChannels()
+{
+	for (TObjectIterator<UNiagaraDataChannel> DCIt; DCIt; ++DCIt)
+	{
+		InitDataChannel(*DCIt, false);
+	}
+}
+
 void FNiagaraDataChannelManager::Init()
 {
 	//Initialize any existing data channels, more may be initialized later as they are loaded.
@@ -42,11 +50,19 @@ void FNiagaraDataChannelManager::BeginFrame(float DeltaSeconds)
 {
 	if (INiagaraModule::DataChannelsEnabled())
 	{
+		check(IsInGameThread());
 		SCOPE_CYCLE_COUNTER(STAT_DataChannelManager_BeginFrame);
 		//Tick all DataChannel channel handlers.
-		for (auto& ChannelPair : Channels)
+		for (auto It = Channels.CreateIterator(); It; ++It)
 		{
-			ChannelPair.Value->BeginFrame(DeltaSeconds, WorldMan);
+			if(const UNiagaraDataChannel* Channel = It.Key().Get())
+			{
+				It.Value()->BeginFrame(DeltaSeconds, WorldMan);
+			}
+			else
+			{
+				It.RemoveCurrent();
+			}
 		}
 	}
 }
@@ -55,6 +71,7 @@ void FNiagaraDataChannelManager::EndFrame(float DeltaSeconds)
 {
 	if (INiagaraModule::DataChannelsEnabled())
 	{
+		check(IsInGameThread());
 		SCOPE_CYCLE_COUNTER(STAT_DataChannelManager_EndFrame);
 		//Tick all DataChannel channel handlers.
 		for (auto& ChannelPair : Channels)
@@ -68,6 +85,7 @@ void FNiagaraDataChannelManager::Tick(float DeltaSeconds, ETickingGroup TickGrou
 {
 	if(INiagaraModule::DataChannelsEnabled())
 	{
+		check(IsInGameThread());
 		SCOPE_CYCLE_COUNTER(STAT_DataChannelManager_Tick);
 		//Tick all DataChannel channel handlers.
 		for (auto& ChannelPair : Channels)
@@ -83,6 +101,7 @@ void FNiagaraDataChannelManager::Tick(float DeltaSeconds, ETickingGroup TickGrou
 
 UNiagaraDataChannelHandler* FNiagaraDataChannelManager::FindDataChannelHandler(const UNiagaraDataChannel* Channel)
 {
+	check(IsInGameThread());
 	if (TObjectPtr<UNiagaraDataChannelHandler>* Found = Channels.Find(Channel))
 	{
 		return (*Found).Get();
@@ -92,6 +111,7 @@ UNiagaraDataChannelHandler* FNiagaraDataChannelManager::FindDataChannelHandler(c
 
 void FNiagaraDataChannelManager::InitDataChannel(const UNiagaraDataChannel* InChannel, bool bForce)
 {
+	check(IsInGameThread());
 	UWorld* World = GetWorld();
 	if(INiagaraModule::DataChannelsEnabled() && World && !World->IsNetMode(NM_DedicatedServer) && InChannel->IsValid())
 	{
@@ -106,6 +126,7 @@ void FNiagaraDataChannelManager::InitDataChannel(const UNiagaraDataChannel* InCh
 
 void FNiagaraDataChannelManager::RemoveDataChannel(const UNiagaraDataChannel* InChannel)
 {
+	check(IsInGameThread());
 	Channels.Remove(InChannel);
 }
 
diff --git a/Engine/Plugins/FX/Niagara/Source/Niagara/Private/NiagaraDataChannelManager.h b/Engine/Plugins/FX/Niagara/Source/Niagara/Private/NiagaraDataChannelManager.h
--- a/Engine/Plugins/FX/Niagara/Source/Niagara/Private/NiagaraDataChannelManager.h
+++ b/Engine/Plugins/FX/Niagara/Source/Niagara/Private/NiagaraDataChannelManager.h
@@ -31,6 +31,8 @@ public:
 	void EndFrame(float DeltaSeconds);
 
 	void Tick(float DeltaSeconds, ETickingGroup TickGroup);
+
+	void RefreshDataChannels();
 	
 	/**
 	Return the DataChannel handler for the given channel.
@@ -44,5 +46,5 @@ private:
 	FNiagaraWorldManager* WorldMan = nullptr;
 
 	/** Runtime handlers for each DataChannel channel. */
-	TMap<TObjectPtr<const UNiagaraDataChannel>, TObjectPtr<UNiagaraDataChannelHandler>> Channels;
+	TMap<TWeakObjectPtr<const UNiagaraDataChannel>, TObjectPtr<UNiagaraDataChannelHandler>> Channels;
 };
diff --git a/Engine/Plugins/FX/Niagara/Source/Niagara/Private/NiagaraDataChannel_Global.cpp b/Engine/Plugins/FX/Niagara/Source/Niagara/Private/NiagaraDataChannel_Global.cpp
--- a/Engine/Plugins/FX/Niagara/Source/Niagara/Private/NiagaraDataChannel_Global.cpp
+++ b/Engine/Plugins/FX/Niagara/Source/Niagara/Private/NiagaraDataChannel_Global.cpp
@@ -44,8 +44,8 @@ void UNiagaraDataChannelHandler_Global::BeginDestroy()
 
 void UNiagaraDataChannelHandler_Global::Init(const UNiagaraDataChannel* InChannel)
 {
+	check(InChannel);
 	Super::Init(InChannel);
-	check(DataChannel);
 	Data = CreateData();
 }
 
diff --git a/Engine/Plugins/FX/Niagara/Source/Niagara/Private/NiagaraDataChannel_Islands.cpp b/Engine/Plugins/FX/Niagara/Source/Niagara/Private/NiagaraDataChannel_Islands.cpp
--- a/Engine/Plugins/FX/Niagara/Source/Niagara/Private/NiagaraDataChannel_Islands.cpp
+++ b/Engine/Plugins/FX/Niagara/Source/Niagara/Private/NiagaraDataChannel_Islands.cpp
@@ -82,7 +82,7 @@ void UNiagaraDataChannelHandler_Islands::EndFrame(float DeltaTime, FNiagaraWorld
 
 void UNiagaraDataChannelHandler_Islands::Tick(float DeltaTime, ETickingGroup TickGroup, FNiagaraWorldManager* OwningWorld)
 {
-	UNiagaraDataChannel_Islands* IslandChannel = GetChannelTyped<UNiagaraDataChannel_Islands>();
+	const UNiagaraDataChannel_Islands* IslandChannel = GetChannelTyped<const UNiagaraDataChannel_Islands>();
 	for (auto It = ActiveIslands.CreateIterator(); It; ++It)
 	{
 		FNDCIsland& Island = IslandPool[*It];
@@ -115,7 +115,7 @@ FNiagaraDataChannelDataPtr UNiagaraDataChannelHandler_Islands::FindData(FNiagara
 
 FNDCIsland* UNiagaraDataChannelHandler_Islands::FindOrCreateIsland(const FNiagaraDataChannelSearchParameters& SearchParams, ENiagaraResourceAccess AccessType)
 {
-	UNiagaraDataChannel_Islands* IslandChannel = CastChecked<UNiagaraDataChannel_Islands>(DataChannel);
+	const UNiagaraDataChannel_Islands* IslandChannel = CastChecked<const UNiagaraDataChannel_Islands>(DataChannel);
 
 	FVector Location = SearchParams.GetLocation();
 	if (AccessType == ENiagaraResourceAccess::ReadOnly)
@@ -290,7 +290,7 @@ void FNDCIsland::OnAcquired(FVector Location)
 	FBox HandlerSystemBounds(-Bounds.BoxExtent, Bounds.BoxExtent);
 	for (auto& SoftSys : Channel->Systems)
 	{
-		if (UNiagaraSystem* Sys = SoftSys.Get())
+		if (UNiagaraSystem* Sys = SoftSys.LoadSynchronous())
 		{
 			FFXSystemSpawnParameters SpawnParams;
 			SpawnParams.bAutoActivate = false;//We must activate AFTER adding this component to our handle system array.
diff --git a/Engine/Plugins/FX/Niagara/Source/Niagara/Private/NiagaraModule.cpp b/Engine/Plugins/FX/Niagara/Source/Niagara/Private/NiagaraModule.cpp
--- a/Engine/Plugins/FX/Niagara/Source/Niagara/Private/NiagaraModule.cpp
+++ b/Engine/Plugins/FX/Niagara/Source/Niagara/Private/NiagaraModule.cpp
@@ -64,6 +64,8 @@ static FAutoConsoleVariableRef CVarLogCompileIdGeneration(
 float INiagaraModule::EngineGlobalSpawnCountScale = 1.0f;
 float INiagaraModule::EngineGlobalSystemCountScale = 1.0f;
 
+std::atomic<bool> INiagaraModule::bDataChannelRefreshRequested = false;
+
 int32 GEnableVerboseNiagaraChangeIdLogging = 0;
 static FAutoConsoleVariableRef CVarEnableVerboseNiagaraChangeIdLogging(
 	TEXT("fx.EnableVerboseNiagaraChangeIdLogging"),
@@ -520,6 +522,8 @@ void INiagaraModule::OnWorldTickStart(UWorld* World, ELevelTick TickType, float
 #if WITH_EDITOR
 	PollSystemCompilations();
 #endif
+
+	RefreshDataChannels();
 	
 #if NIAGARA_PERF_BASELINES
 	if (BaselineHandler.IsValid())
@@ -529,6 +533,19 @@ void INiagaraModule::OnWorldTickStart(UWorld* World, ELevelTick TickType, float
 #endif
 }
 
+void INiagaraModule::RefreshDataChannels()
+{
+	if (bDataChannelRefreshRequested)
+	{
+		bDataChannelRefreshRequested = false;
+		FNiagaraWorldManager::ForAllWorldManagers(
+			[](FNiagaraWorldManager& WorldMan)
+			{
+				WorldMan.GetDataChannelManager().RefreshDataChannels();
+			});
+	}
+}
+
 void INiagaraModule::OnPostGarbageCollect()
 {
 	FNiagaraTypeHelper::TickTypeRemap();
diff --git a/Engine/Plugins/FX/Niagara/Source/Niagara/Public/NiagaraDataChannelHandler.h b/Engine/Plugins/FX/Niagara/Source/Niagara/Public/NiagaraDataChannelHandler.h
--- a/Engine/Plugins/FX/Niagara/Source/Niagara/Public/NiagaraDataChannelHandler.h
+++ b/Engine/Plugins/FX/Niagara/Source/Niagara/Public/NiagaraDataChannelHandler.h
@@ -37,7 +37,12 @@ public:
 	*/
 	NIAGARA_API virtual FNiagaraDataChannelDataPtr FindData(FNiagaraDataChannelSearchParameters SearchParams, ENiagaraResourceAccess AccessType)  PURE_VIRTUAL(UNiagaraDataChannelHandler::FindData, return nullptr;);
 
-	const UNiagaraDataChannel* GetDataChannel()const { return DataChannel; }
+	const UNiagaraDataChannel* GetDataChannel()const
+	{
+		const UNiagaraDataChannel* Ret = DataChannel.Get();
+		check(Ret);
+		return Ret;
+	}
 
 	UFUNCTION(BlueprintCallable, Category="Data Channel")
 	NIAGARA_API UNiagaraDataChannelWriter* GetDataChannelWriter();
@@ -53,7 +58,7 @@ public:
 protected:
 
 	UPROPERTY()
-	TObjectPtr<const UNiagaraDataChannel> DataChannel;
+	TWeakObjectPtr<const UNiagaraDataChannel> DataChannel;
 
 	/** Helper object allowing BP to write data in this channel. */
 	UPROPERTY()
diff --git a/Engine/Plugins/FX/Niagara/Source/Niagara/Public/NiagaraModule.h b/Engine/Plugins/FX/Niagara/Source/Niagara/Public/NiagaraModule.h
--- a/Engine/Plugins/FX/Niagara/Source/Niagara/Public/NiagaraModule.h
+++ b/Engine/Plugins/FX/Niagara/Source/Niagara/Public/NiagaraModule.h
@@ -124,9 +124,11 @@ public:
 	NIAGARA_API void UnregisterAbortSystemCompile(FDelegateHandle DelegateHandle);
 
 	NIAGARA_API void OnAssetLoaded(UObject* Asset);
-
 #endif
 
+	static void RequestRefreshDataChannels() { bDataChannelRefreshRequested = true; }
+	static void RefreshDataChannels();
+
 #if NIAGARA_PERF_BASELINES
 	NIAGARA_API void GeneratePerfBaselines(TArray<UNiagaraEffectType*>& BaselinesToGenerate);
 
@@ -399,5 +401,7 @@ private:
 #endif
 
 	FDelegateHandle OnCVarUnregisteredHandle;
+
+	static std::atomic<bool> bDataChannelRefreshRequested;
 };
 
