- Add a global deferred execution queue to the world manager
- Move the export & simple counter DIs to use this global queue to ensure callbacks execute at safe points for Niagara

diff --git a/Engine/Plugins/FX/Niagara/Source/Niagara/Private/DataInterface/NiagaraDataInterfaceSimpleCounter.cpp b/Engine/Plugins/FX/Niagara/Source/Niagara/Private/DataInterface/NiagaraDataInterfaceSimpleCounter.cpp
--- a/Engine/Plugins/FX/Niagara/Source/Niagara/Private/DataInterface/NiagaraDataInterfaceSimpleCounter.cpp
+++ b/Engine/Plugins/FX/Niagara/Source/Niagara/Private/DataInterface/NiagaraDataInterfaceSimpleCounter.cpp
@@ -7,6 +7,7 @@
 #include "NiagaraGpuReadbackManager.h"
 #include "NiagaraShaderParametersBuilder.h"
 #include "NiagaraSystemInstance.h"
+#include "NiagaraWorldManager.h"
 
 #include "Internationalization/Internationalization.h"
 #include "Async/Async.h"
@@ -113,10 +114,11 @@ struct FNDISimpleCounterProxy : public FNiagaraDataInterfaceProxyRW
 							[SystemInstanceID, WeakOwner=Proxy->WeakOwner, Proxy](TConstArrayView<TPair<void*, uint32>> ReadbackData)
 							{
 								const int32 CounterValue = *reinterpret_cast<const int32*>(ReadbackData[0].Key);
-								AsyncTask(
-									ENamedThreads::GameThread,
+								FNiagaraWorldManager::EnqueueGlobalDeferredCallback(
 									[SystemInstanceID, CounterValue, WeakOwner, Proxy]()
 									{
+										TRACE_CPUPROFILER_EVENT_SCOPE(NDISimpleCounterGTCallback);
+
 										// FNiagaraDataInterfaceProxy do not outlive UNiagaraDataInterface so if our Object is valid so is the proxy
 										// Equally because we do not share instance IDs (monotonically increasing number) we won't ever stomp something that has 'gone away'
 										if ( WeakOwner.Get() )
diff --git a/Engine/Plugins/FX/Niagara/Source/Niagara/Private/NiagaraDataInterfaceExport.cpp b/Engine/Plugins/FX/Niagara/Source/Niagara/Private/NiagaraDataInterfaceExport.cpp
--- a/Engine/Plugins/FX/Niagara/Source/Niagara/Private/NiagaraDataInterfaceExport.cpp
+++ b/Engine/Plugins/FX/Niagara/Source/Niagara/Private/NiagaraDataInterfaceExport.cpp
@@ -10,6 +10,7 @@
 #include "NiagaraShaderParametersBuilder.h"
 #include "NiagaraSystem.h"
 #include "NiagaraSystemInstance.h"
+#include "NiagaraWorldManager.h"
 
 #include "Internationalization/Internationalization.h"
 #include "RenderGraphBuilder.h"
@@ -40,51 +41,33 @@ namespace NDIExportLocal
 		TEXT("Maximum buffer instance count for the GPU readback when in PerParticleMode, where <= 0 means ignore."),
 		ECVF_Default
 	);
-}
-
 
-/**
-Async task to call the blueprint callback on the game thread and isolate the Niagara tick from the blueprint
-*/
-class FNiagaraExportCallbackAsyncTask
-{
-	TWeakObjectPtr<UObject> WeakCallbackHandler;
-	TArray<FBasicParticleData> Data;
-	TWeakObjectPtr<UNiagaraSystem> WeakSystem;
-	FVector3f SystemTileOffset;
-
-public:
-	FNiagaraExportCallbackAsyncTask(TWeakObjectPtr<UObject> InCallbackHandler, TArray<FBasicParticleData> Data, TWeakObjectPtr<UNiagaraSystem> InSystem, FVector3f InSystemTileOffset)
-		: WeakCallbackHandler(InCallbackHandler)
-		, Data(Data)
-		, WeakSystem(InSystem)
-		, SystemTileOffset(InSystemTileOffset)
+	void EnqueueGameThreadCallback(TWeakObjectPtr<UObject> WeakCallbackHandler, TArray<FBasicParticleData> ParticleData, TWeakObjectPtr<UNiagaraSystem> WeakSystem, FVector3f SystemTileOffset)
 	{
-	}
-
-	FORCEINLINE TStatId GetStatId() const { RETURN_QUICK_DECLARE_CYCLE_STAT(FNiagaraExportCallbackAsyncTask, STATGROUP_TaskGraphTasks); }
-	static FORCEINLINE ENamedThreads::Type GetDesiredThread() { return ENamedThreads::GameThread; }
-	static FORCEINLINE ESubsequentsMode::Type GetSubsequentsMode() { return ESubsequentsMode::FireAndForget; }
+		FNiagaraWorldManager::EnqueueGlobalDeferredCallback(
+			[WeakCallbackHandler, ParticleData=MoveTemp(ParticleData), WeakSystem, SystemTileOffset]()
+			{
+				TRACE_CPUPROFILER_EVENT_SCOPE(NiagaraExportGTCallback);
 
-	void DoTask(ENamedThreads::Type CurrentThread, const FGraphEventRef& MyCompletionGraphEvent)
-	{
-		UNiagaraSystem* System = WeakSystem.Get();
-		if (System == nullptr)
-		{
-			UE_LOG(LogNiagara, Warning, TEXT("Invalid system handle in export data interface callback, skipping"));
-			return;
-		}
+				UNiagaraSystem* System = WeakSystem.Get();
+				if (System == nullptr)
+				{
+					UE_LOG(LogNiagara, Warning, TEXT("Invalid system handle in export data interface callback, skipping"));
+					return;
+				}
 
-		UObject* CallbackHandler = WeakCallbackHandler.Get();
-		if (CallbackHandler == nullptr)
-		{
-			UE_LOG(LogNiagara, Warning, TEXT("Invalid CallbackHandler in export data interface callback, skipping"));
-			return;
-		}
+				UObject* CallbackHandler = WeakCallbackHandler.Get();
+				if (CallbackHandler == nullptr)
+				{
+					UE_LOG(LogNiagara, Warning, TEXT("Invalid CallbackHandler in export data interface callback, skipping"));
+					return;
+				}
 
-		INiagaraParticleCallbackHandler::Execute_ReceiveParticleData(CallbackHandler, Data, System, FVector(SystemTileOffset) * FLargeWorldRenderScalar::GetTileSize());
+				INiagaraParticleCallbackHandler::Execute_ReceiveParticleData(CallbackHandler, ParticleData, System, FVector(SystemTileOffset) * FLargeWorldRenderScalar::GetTileSize());
+			}
+		);
 	}
-};
+}
 
 //////////////////////////////////////////////////////////////////////////
 // Instance and Proxy Data
@@ -162,7 +145,7 @@ struct FNDIExportProxy : public FNiagaraDataInterfaceProxy
 							FloatData += NDIExportLocal::NumFloatsPerInstance;
 						}
 
-						TGraphTask<FNiagaraExportCallbackAsyncTask>::CreateTask().ConstructAndDispatchWhenReady(WeakCallbackHandler, MoveTemp(ExportParticleData), WeakSystem, SystemLWCTile);
+						NDIExportLocal::EnqueueGameThreadCallback(WeakCallbackHandler, MoveTemp(ExportParticleData), WeakSystem, SystemLWCTile);
 					}
 				}
 			);
@@ -294,7 +277,7 @@ bool UNiagaraDataInterfaceExport::PerInstanceTickPostSimulate(void* PerInstanceD
 				Data.Add(Value);
 			}
 
-			TGraphTask<FNiagaraExportCallbackAsyncTask>::CreateTask().ConstructAndDispatchWhenReady(PIData->CallbackHandler, MoveTemp(Data), SystemInstance->GetSystem(), SystemInstance->GetLWCTile());
+			NDIExportLocal::EnqueueGameThreadCallback(CallbackHandler, MoveTemp(Data), SystemInstance->GetSystem(), SystemInstance->GetLWCTile());
 		}
 	}
 	return false;
diff --git a/Engine/Plugins/FX/Niagara/Source/Niagara/Private/NiagaraWorldManager.cpp b/Engine/Plugins/FX/Niagara/Source/Niagara/Private/NiagaraWorldManager.cpp
--- a/Engine/Plugins/FX/Niagara/Source/Niagara/Private/NiagaraWorldManager.cpp
+++ b/Engine/Plugins/FX/Niagara/Source/Niagara/Private/NiagaraWorldManager.cpp
@@ -255,8 +255,19 @@ FDelegateHandle FNiagaraWorldManager::ViewTargetChangedHandle;
 std::atomic<bool> FNiagaraWorldManager::bInvalidateCachedSystemScalabilityData(false);
 TMap<class UWorld*, class FNiagaraWorldManager*> FNiagaraWorldManager::WorldManagers;
 
-namespace FNiagaraUtilities
+namespace NiagaraWorldManagerInternal
 {
+	static TQueue<TFunction<void()>, EQueueMode::Mpsc> GlobalDeferredCallbacks;
+
+	void ExecuteGlobalDeferredCallbacks()
+	{
+		TFunction<void()> Callback;
+		while (GlobalDeferredCallbacks.Dequeue(Callback))
+		{
+			Callback();
+		}
+	}
+
 	int GetNiagaraTickGroup(ETickingGroup TickGroup)
 	{
 		const int ActualTickGroup = FMath::Clamp(TickGroup - NiagaraFirstTickGroup, 0, NiagaraNumTickGroups - 1);
@@ -563,10 +574,10 @@ FNiagaraSystemSimulationRef FNiagaraWorldManager::GetSystemSimulation(ETickingGr
 {
 	LLM_SCOPE(ELLMTag::Niagara);
 
-	int32 ActualTickGroup = FNiagaraUtilities::GetNiagaraTickGroup(TickGroup);
+	int32 ActualTickGroup = NiagaraWorldManagerInternal::GetNiagaraTickGroup(TickGroup);
 	if (ActiveNiagaraTickGroup == ActualTickGroup)
 	{
-		int32 DemotedTickGroup = FNiagaraUtilities::GetNiagaraTickGroup((ETickingGroup)(TickGroup + 1));
+		int32 DemotedTickGroup = NiagaraWorldManagerInternal::GetNiagaraTickGroup((ETickingGroup)(TickGroup + 1));
 		ActualTickGroup = DemotedTickGroup == ActualTickGroup ? 0 : DemotedTickGroup;
 	}
 
@@ -941,6 +952,7 @@ void FNiagaraWorldManager::OnRefreshOwnerAllowsScalability()
 void FNiagaraWorldManager::PreActorTick(ELevelTick InLevelTick, float InDeltaSeconds)
 {	
 	DataChannelManager->BeginFrame(InDeltaSeconds);
+	NiagaraWorldManagerInternal::ExecuteGlobalDeferredCallbacks();
 }
 
 void FNiagaraWorldManager::PostActorTick(float DeltaSeconds)
@@ -951,6 +963,8 @@ void FNiagaraWorldManager::PostActorTick(float DeltaSeconds)
 	SCOPE_CYCLE_COUNTER(STAT_NiagaraOverview_GT);
 	QUICK_SCOPE_CYCLE_COUNTER(STAT_NiagaraPostActorTick_GT);
 
+	NiagaraWorldManagerInternal::ExecuteGlobalDeferredCallbacks();
+
 	DeltaSeconds *= DebugPlaybackRate;
 
 	// Update any systems with post actor work
@@ -1151,13 +1165,16 @@ void FNiagaraWorldManager::Tick(ETickingGroup TickGroup, float DeltaSeconds, ELe
 {
 	check(TickGroup >= NiagaraFirstTickGroup && TickGroup <= NiagaraLastTickGroup);
 
-	DeferredMethods.ExecuteAndClear();
-
 	CSV_SCOPED_TIMING_STAT_EXCLUSIVE(Effects);
 	LLM_SCOPE(ELLMTag::Niagara);
 	SCOPE_CYCLE_COUNTER(STAT_NiagaraWorldManTick);
 	SCOPE_CYCLE_COUNTER(STAT_NiagaraOverview_GT);
 
+	// Execute any global and local deferred functions
+	NiagaraWorldManagerInternal::ExecuteGlobalDeferredCallbacks();
+	DeferredMethods.ExecuteAndClear();
+
+	// Modify playback rate
 	DeltaSeconds *= DebugPlaybackRate;
 
 	//Tick DataChannel Manager.
@@ -1274,7 +1291,7 @@ void FNiagaraWorldManager::Tick(ETickingGroup TickGroup, float DeltaSeconds, ELe
 	}
 
 	// Now tick all system instances. 
-	const int ActualTickGroup = FNiagaraUtilities::GetNiagaraTickGroup(TickGroup);
+	const int ActualTickGroup = NiagaraWorldManagerInternal::GetNiagaraTickGroup(TickGroup);
 
 	ActiveNiagaraTickGroup = ActualTickGroup;
 
@@ -1986,6 +2003,12 @@ void FNiagaraWorldManager::PrimePool(UNiagaraSystem* System)
 	}
 }
 
+void FNiagaraWorldManager::EnqueueGlobalDeferredCallback(TFunction<void()>&& Callback)
+{
+	using namespace NiagaraWorldManagerInternal;
+	GlobalDeferredCallbacks.Enqueue(Callback);
+}
+
 bool FNiagaraWorldManager::IsComponentLocalPlayerLinked(const USceneComponent* InComponent)
 {
 	check(InComponent);
diff --git a/Engine/Plugins/FX/Niagara/Source/Niagara/Public/NiagaraWorldManager.h b/Engine/Plugins/FX/Niagara/Source/Niagara/Public/NiagaraWorldManager.h
--- a/Engine/Plugins/FX/Niagara/Source/Niagara/Public/NiagaraWorldManager.h
+++ b/Engine/Plugins/FX/Niagara/Source/Niagara/Public/NiagaraWorldManager.h
@@ -223,6 +223,12 @@ public:
 
 	class FNiagaraDeferredMethodQueue& GetDeferredMethodQueue() { return DeferredMethods; }
 
+	/**
+	This is a threadsafe queue which will execute on the game thread.
+	The queue is flushed at the start of each tick group & pre / post actor tick.
+	*/
+	NIAGARA_API static void EnqueueGlobalDeferredCallback(TFunction<void()>&& Callback);
+
 	/** Is this component in anyway linked to the local player. */
 	static bool IsComponentLocalPlayerLinked(const USceneComponent* Component);
 
