- Fix SimCache capturing the same frame when the game is paused
- We needed to update CaptureTickCount when writing the frame

#rnx

diff --git a/Engine/Plugins/FX/Niagara/Source/Niagara/Private/NiagaraSimCache.cpp b/Engine/Plugins/FX/Niagara/Source/Niagara/Private/NiagaraSimCache.cpp
--- a/Engine/Plugins/FX/Niagara/Source/Niagara/Private/NiagaraSimCache.cpp
+++ b/Engine/Plugins/FX/Niagara/Source/Niagara/Private/NiagaraSimCache.cpp
@@ -295,27 +295,30 @@ bool UNiagaraSimCache::WriteFrame(UNiagaraComponent* NiagaraComponent, FNiagaraS
 		return false;
 	}
 
-	// First frame we are about to cache?
+	// Is this the first frame?  If so capture the cache start time
 	if (CaptureTickCount == INDEX_NONE)
 	{
 		StartSeconds = Helper.SystemInstance->GetAge();
-		CaptureTickCount = Helper.SystemInstance->GetTickCount();
 	}
-	// If our tick counter hasn't moved then we won't capture a frame as there's nothing new to process
-	else if (CaptureTickCount == Helper.SystemInstance->GetTickCount())
+	else
 	{
-		FeedbackContext.Errors.Emplace(FString::Printf(TEXT("System was not ticked since the last capture. %s"), *Helper.SystemInstance->GetCrashReporterTag()));
-		return false;
-	}
+		// If our tick counter hasn't moved then we won't capture the frame as there's nothing new to process
+		if (CaptureTickCount == Helper.SystemInstance->GetTickCount())
+		{
+			//FeedbackContext.Errors.Emplace(FString::Printf(TEXT("System was not ticked since the last capture. %s"), *Helper.SystemInstance->GetCrashReporterTag()));
+			return false;
+		}
 
-	// If the tick counter is lower than the previous value then the system was reset so we won't capture
-	if ( Helper.SystemInstance->GetTickCount() < CaptureTickCount )
-	{
-		FeedbackContext.Errors.Emplace(FString::Printf(TEXT("System was was reset since the last capture, skipping further cache writes. %s"), *Helper.SystemInstance->GetCrashReporterTag()));
-		SoftNiagaraSystem.Reset();
-		return false;
+		// If the tick counter is lower than the previous value then the system was reset so we won't capture
+		if (Helper.SystemInstance->GetTickCount() < CaptureTickCount)
+		{
+			FeedbackContext.Errors.Emplace(FString::Printf(TEXT("System was was reset since the last capture, skipping further cache writes. %s"), *Helper.SystemInstance->GetCrashReporterTag()));
+			SoftNiagaraSystem.Reset();
+			return false;
+		}
 	}
 
+	CaptureTickCount = Helper.SystemInstance->GetTickCount();
 	DurationSeconds = Helper.SystemInstance->GetAge() - StartSeconds;
 
 	// Cache frame
