Fix breakage with new functions hitting dll changes downstream.

diff --git a/Engine/Plugins/FX/Niagara/Source/Niagara/Public/NiagaraModule.h b/Engine/Plugins/FX/Niagara/Source/Niagara/Public/NiagaraModule.h
--- a/Engine/Plugins/FX/Niagara/Source/Niagara/Public/NiagaraModule.h
+++ b/Engine/Plugins/FX/Niagara/Source/Niagara/Public/NiagaraModule.h
@@ -126,8 +126,8 @@ public:
 	NIAGARA_API void OnAssetLoaded(UObject* Asset);
 #endif
 
-	static void RequestRefreshDataChannels() { bDataChannelRefreshRequested = true; }
-	static void RefreshDataChannels();
+	static NIAGARA_API void RequestRefreshDataChannels() { bDataChannelRefreshRequested = true; }
+	static NIAGARA_API void RefreshDataChannels();
 
 #if NIAGARA_PERF_BASELINES
 	NIAGARA_API void GeneratePerfBaselines(TArray<UNiagaraEffectType*>& BaselinesToGenerate);
