Fixed constructor initialization order

diff --git a/Engine/Plugins/FX/Niagara/Source/Niagara/Private/NiagaraEmitterHandle.cpp b/Engine/Plugins/FX/Niagara/Source/Niagara/Private/NiagaraEmitterHandle.cpp
--- a/Engine/Plugins/FX/Niagara/Source/Niagara/Private/NiagaraEmitterHandle.cpp
+++ b/Engine/Plugins/FX/Niagara/Source/Niagara/Private/NiagaraEmitterHandle.cpp
@@ -24,10 +24,10 @@ FNiagaraEmitterHandle::FNiagaraEmitterHandle()
 
 #if WITH_EDITORONLY_DATA
 FNiagaraEmitterHandle::FNiagaraEmitterHandle(UNiagaraEmitter& InEmitter, const FGuid& Version)
-	: Id(FGuid::NewGuid())
+	: Name(*InEmitter.GetUniqueEmitterName())
+	, Id(FGuid::NewGuid())
 	, IdName(*Id.ToString())
 	, bIsEnabled(true)
-	, Name(*InEmitter.GetUniqueEmitterName())
 	, Source_DEPRECATED(nullptr)
 	, LastMergedSource_DEPRECATED(nullptr)
 	, bIsolated(false)
@@ -37,10 +37,10 @@ FNiagaraEmitterHandle::FNiagaraEmitterHandle(UNiagaraEmitter& InEmitter, const F
 }
 
 FNiagaraEmitterHandle::FNiagaraEmitterHandle(const FVersionedNiagaraEmitter& InEmitter)
-	: Id(FGuid::NewGuid())
+	: Name(*InEmitter.Emitter->GetUniqueEmitterName())
+	, Id(FGuid::NewGuid())
 	, IdName(*Id.ToString())
 	, bIsEnabled(true)
-	, Name(*InEmitter.Emitter->GetUniqueEmitterName())
 	, Source_DEPRECATED(nullptr)
 	, LastMergedSource_DEPRECATED(nullptr)
 	, bIsolated(false)
