Fixed a crash when emitters are loaded that inherit from unloaded parent emitters & modules (disabled fluids plugin when placing a user created fluid system for example). Summary View Post Load Transfer would expect some data to exist that doesn't exist

diff --git a/Engine/Plugins/FX/Niagara/Source/NiagaraEditor/Private/NiagaraEmitterEditorData.cpp b/Engine/Plugins/FX/Niagara/Source/NiagaraEditor/Private/NiagaraEmitterEditorData.cpp
--- a/Engine/Plugins/FX/Niagara/Source/NiagaraEditor/Private/NiagaraEmitterEditorData.cpp
+++ b/Engine/Plugins/FX/Niagara/Source/NiagaraEditor/Private/NiagaraEmitterEditorData.cpp
@@ -275,6 +275,13 @@ void UNiagaraEmitterEditorData::PostLoad_TransferSummaryDataToNewFormat()
 		// advanced items always get put at the bottom
 		auto SortInputsPredicate = ([&](UNiagaraHierarchyItemBase& ItemA, UNiagaraHierarchyItemBase& ItemB)
 		{
+			// while these should generally be valid, it's possible due to inheritance that a user created fluid system is executing this code despite the fluid plugin being disabled.
+			// in that case, asset variable metadata might not contain these entries.
+			if(AssetVariableMetadata.Contains(ItemA.GetPersistentIdentity()) == false || AssetVariableMetadata.Contains(ItemB.GetPersistentIdentity()) == false)
+			{
+				return false;
+			}
+			
 			const FNiagaraVariableMetaData& MetaDataA = AssetVariableMetadata[ItemA.GetPersistentIdentity()];
 			const FNiagaraVariableMetaData& MetaDataB = AssetVariableMetadata[ItemB.GetPersistentIdentity()];
 
@@ -323,7 +330,10 @@ void UNiagaraEmitterEditorData::PostLoad_TransferSummaryDataToNewFormat()
 				int32 MinimumCategorySortOrder = SummarySortOrder[MinimumItem];
 				if(MinimumCategorySortOrder == INDEX_NONE)
 				{
-					MinimumCategorySortOrder = AssetVariableMetadata[MinimumItem->GetPersistentIdentity()].EditorSortPriority;
+					if(AssetVariableMetadata.Contains(MinimumItem->GetPersistentIdentity()))
+					{
+						MinimumCategorySortOrder = AssetVariableMetadata[MinimumItem->GetPersistentIdentity()].EditorSortPriority;
+					}
 				}
 				
 				CategoryMinimumInputSortOrderMap.Add(Category, MinimumCategorySortOrder);
