Allowing stack items to customize the advanced toggle logic.
This is used by the summary view to handle categories with only advanced inputs correctly, as the Toggle would only refresh filtered children before, but not reevaluate entire categories for showing up or not as this is done in section evaluation
Also changed the active section widget to evaluate using 'Equals' instead of 'Identitcal' as there was an issue where equal Texts would still return false

diff --git a/Engine/Plugins/FX/Niagara/Source/NiagaraEditor/Private/ViewModels/Stack/NiagaraStackEmitterSettingsGroup.cpp b/Engine/Plugins/FX/Niagara/Source/NiagaraEditor/Private/ViewModels/Stack/NiagaraStackEmitterSettingsGroup.cpp
--- a/Engine/Plugins/FX/Niagara/Source/NiagaraEditor/Private/ViewModels/Stack/NiagaraStackEmitterSettingsGroup.cpp
+++ b/Engine/Plugins/FX/Niagara/Source/NiagaraEditor/Private/ViewModels/Stack/NiagaraStackEmitterSettingsGroup.cpp
@@ -311,17 +311,23 @@ FText UNiagaraStackEmitterSummaryItem::GetIconText() const
 
 void UNiagaraStackEmitterSummaryItem::RefreshChildrenInternal(const TArray<UNiagaraStackEntry*>& CurrentChildren, TArray<UNiagaraStackEntry*>& NewChildren, TArray<FStackIssue>& NewIssues)
 {
-	if (FilteredObject == nullptr)
+	if (SummaryViewCollection == nullptr)
 	{
-		FilteredObject = NewObject<UNiagaraStackSummaryViewCollection>(this);
+		SummaryViewCollection = NewObject<UNiagaraStackSummaryViewCollection>(this);
 		FRequiredEntryData RequiredEntryData(GetSystemViewModel(), GetEmitterViewModel(), FExecutionCategoryNames::Emitter, NAME_None, GetStackEditorData());
-		FilteredObject->Initialize(RequiredEntryData, Emitter, GetStackEditorDataKey());
+		SummaryViewCollection->Initialize(RequiredEntryData, Emitter, GetStackEditorDataKey());
 	}
 
-	NewChildren.Add(FilteredObject);
+	NewChildren.Add(SummaryViewCollection);
 	Super::RefreshChildrenInternal(CurrentChildren, NewChildren, NewIssues);
 }
 
+void UNiagaraStackEmitterSummaryItem::ToggleShowAdvancedInternal()
+{
+	// this will cause section data & filtered children to be refreshed while caching the last active section
+	SummaryViewCollection->RefreshForAdvancedToggle();
+}
+
 TSharedPtr<IDetailTreeNode> GetSummarySectionsPropertyNode(const TArray<TSharedRef<IDetailTreeNode>>& Nodes)
 {
 	TArray<TSharedRef<IDetailTreeNode>> ChildrenToCheck;
diff --git a/Engine/Plugins/FX/Niagara/Source/NiagaraEditor/Private/ViewModels/Stack/NiagaraStackFunctionInputCollection.cpp b/Engine/Plugins/FX/Niagara/Source/NiagaraEditor/Private/ViewModels/Stack/NiagaraStackFunctionInputCollection.cpp
--- a/Engine/Plugins/FX/Niagara/Source/NiagaraEditor/Private/ViewModels/Stack/NiagaraStackFunctionInputCollection.cpp
+++ b/Engine/Plugins/FX/Niagara/Source/NiagaraEditor/Private/ViewModels/Stack/NiagaraStackFunctionInputCollection.cpp
@@ -90,6 +90,14 @@ FText UNiagaraStackValueCollection::GetTooltipForSection(FString Section) const
 	return FText::GetEmpty();
 }
 
+void UNiagaraStackValueCollection::CacheLastActiveSection()
+{
+	if(ActiveSectionCache.IsSet())
+	{
+		LastActiveSection = ActiveSectionCache.GetValue();
+	}
+}
+
 bool UNiagaraStackValueCollection::GetCanExpand() const
 {
 	return bShouldDisplayLabel;
diff --git a/Engine/Plugins/FX/Niagara/Source/NiagaraEditor/Private/ViewModels/Stack/NiagaraStackItem.cpp b/Engine/Plugins/FX/Niagara/Source/NiagaraEditor/Private/ViewModels/Stack/NiagaraStackItem.cpp
--- a/Engine/Plugins/FX/Niagara/Source/NiagaraEditor/Private/ViewModels/Stack/NiagaraStackItem.cpp
+++ b/Engine/Plugins/FX/Niagara/Source/NiagaraEditor/Private/ViewModels/Stack/NiagaraStackItem.cpp
@@ -130,6 +130,11 @@ void UNiagaraStackItem::ToggleShowAdvanced()
 {
 	bool bCurrentShowAdvanced = GetStackEditorData().GetStackItemShowAdvanced(GetStackEditorDataKey(), false);
 	GetStackEditorData().SetStackItemShowAdvanced(GetStackEditorDataKey(), !bCurrentShowAdvanced);
+	ToggleShowAdvancedInternal();
+}
+
+void UNiagaraStackItem::ToggleShowAdvancedInternal()
+{
 	RefreshFilteredChildren();
 }
 
diff --git a/Engine/Plugins/FX/Niagara/Source/NiagaraEditor/Private/ViewModels/Stack/NiagaraStackSummaryViewInputCollection.cpp b/Engine/Plugins/FX/Niagara/Source/NiagaraEditor/Private/ViewModels/Stack/NiagaraStackSummaryViewInputCollection.cpp
--- a/Engine/Plugins/FX/Niagara/Source/NiagaraEditor/Private/ViewModels/Stack/NiagaraStackSummaryViewInputCollection.cpp
+++ b/Engine/Plugins/FX/Niagara/Source/NiagaraEditor/Private/ViewModels/Stack/NiagaraStackSummaryViewInputCollection.cpp
@@ -67,8 +67,21 @@ bool UNiagaraStackSummaryViewCollection::GetIsEnabled() const
 	return true;
 }
 
-void UNiagaraStackSummaryViewCollection::RefreshChildrenInternal(const TArray<UNiagaraStackEntry*>& CurrentChildren, TArray<UNiagaraStackEntry*>& NewChildren, TArray<FStackIssue>& NewIssues)
+void UNiagaraStackSummaryViewCollection::RefreshForAdvancedToggle()
 {
+	if (!IsFinalized())
+	{
+		CacheLastActiveSection();
+		// category visibility will be determined by filtered children, so we need to refresh them before updating cached section data
+		RefreshFilteredChildren();
+		UpdateCachedSectionData();
+		// as section data might have changed, we now have to refresh filtered children again
+		RefreshFilteredChildren();
+	}
+}
+
+void UNiagaraStackSummaryViewCollection::RefreshChildrenInternal(const TArray<UNiagaraStackEntry*>& CurrentChildren, TArray<UNiagaraStackEntry*>& NewChildren, TArray<FStackIssue>& NewIssues)
+{	
 	TSharedPtr<FNiagaraEmitterViewModel> ViewModel = GetEmitterViewModel();
 	TSharedPtr<FNiagaraScriptViewModel> ScriptViewModelPinned = ViewModel->GetSharedScriptViewModel();
 	checkf(ScriptViewModelPinned.IsValid(), TEXT("Can not refresh children when the script view model has been deleted."));
@@ -557,6 +570,7 @@ void UNiagaraStackSummaryViewCollection::OnViewStateChanged()
 {
 	if (!IsFinalized())
 	{
+		CacheLastActiveSection();
 		UpdateCachedSectionData();
 		RefreshChildren();
 	}
diff --git a/Engine/Plugins/FX/Niagara/Source/NiagaraEditor/Public/ViewModels/Stack/NiagaraStackEmitterSettingsGroup.h b/Engine/Plugins/FX/Niagara/Source/NiagaraEditor/Public/ViewModels/Stack/NiagaraStackEmitterSettingsGroup.h
--- a/Engine/Plugins/FX/Niagara/Source/NiagaraEditor/Public/ViewModels/Stack/NiagaraStackEmitterSettingsGroup.h
+++ b/Engine/Plugins/FX/Niagara/Source/NiagaraEditor/Public/ViewModels/Stack/NiagaraStackEmitterSettingsGroup.h
@@ -78,15 +78,14 @@ public:
 protected:
 
 	NIAGARAEDITOR_API virtual void RefreshChildrenInternal(const TArray<UNiagaraStackEntry*>& CurrentChildren, TArray<UNiagaraStackEntry*>& NewChildren, TArray<FStackIssue>& NewIssues) override;
-
+	virtual void ToggleShowAdvancedInternal() override;
 private:
 	NIAGARAEDITOR_API void SelectSummaryNodesFromEmitterEditorDataRootNodes(TArray<TSharedRef<IDetailTreeNode>> Source, TArray<TSharedRef<IDetailTreeNode>>* Selected);
-
 private:
 	FVersionedNiagaraEmitterWeakPtr Emitter;
 
 	UPROPERTY()
-	TObjectPtr<UNiagaraStackSummaryViewCollection> FilteredObject;
+	TObjectPtr<UNiagaraStackSummaryViewCollection> SummaryViewCollection;
 };
 
 UCLASS(MinimalAPI)
diff --git a/Engine/Plugins/FX/Niagara/Source/NiagaraEditor/Public/ViewModels/Stack/NiagaraStackFunctionInputCollection.h b/Engine/Plugins/FX/Niagara/Source/NiagaraEditor/Public/ViewModels/Stack/NiagaraStackFunctionInputCollection.h
--- a/Engine/Plugins/FX/Niagara/Source/NiagaraEditor/Public/ViewModels/Stack/NiagaraStackFunctionInputCollection.h
+++ b/Engine/Plugins/FX/Niagara/Source/NiagaraEditor/Public/ViewModels/Stack/NiagaraStackFunctionInputCollection.h
@@ -32,6 +32,8 @@ public:
 	NIAGARAEDITOR_API void SetActiveSection(FText InActiveSection);
 
 	NIAGARAEDITOR_API FText GetTooltipForSection(FString Section) const;
+
+	void CacheLastActiveSection();
 public:
 	static NIAGARAEDITOR_API FText UncategorizedName;
 
diff --git a/Engine/Plugins/FX/Niagara/Source/NiagaraEditor/Public/ViewModels/Stack/NiagaraStackItem.h b/Engine/Plugins/FX/Niagara/Source/NiagaraEditor/Public/ViewModels/Stack/NiagaraStackItem.h
--- a/Engine/Plugins/FX/Niagara/Source/NiagaraEditor/Public/ViewModels/Stack/NiagaraStackItem.h
+++ b/Engine/Plugins/FX/Niagara/Source/NiagaraEditor/Public/ViewModels/Stack/NiagaraStackItem.h
@@ -62,6 +62,8 @@ private:
 
 	NIAGARAEDITOR_API void ToggleShowAdvanced();
 
+	virtual void ToggleShowAdvancedInternal();
+
 protected:
 	FOnModifiedGroupItems ModifiedGroupItemsDelegate;
 	FOnRequestCanPaste RequestCanPasteDelegete;
diff --git a/Engine/Plugins/FX/Niagara/Source/NiagaraEditor/Public/ViewModels/Stack/NiagaraStackSummaryViewInputCollection.h b/Engine/Plugins/FX/Niagara/Source/NiagaraEditor/Public/ViewModels/Stack/NiagaraStackSummaryViewInputCollection.h
--- a/Engine/Plugins/FX/Niagara/Source/NiagaraEditor/Public/ViewModels/Stack/NiagaraStackSummaryViewInputCollection.h
+++ b/Engine/Plugins/FX/Niagara/Source/NiagaraEditor/Public/ViewModels/Stack/NiagaraStackSummaryViewInputCollection.h
@@ -23,7 +23,8 @@ public:
 
 	NIAGARAEDITOR_API virtual FText GetDisplayName() const override;
 	NIAGARAEDITOR_API virtual bool GetIsEnabled() const override;
-	
+
+	void RefreshForAdvancedToggle();
 protected:
 	NIAGARAEDITOR_API virtual void RefreshChildrenInternal(const TArray<UNiagaraStackEntry*>& CurrentChildren, TArray<UNiagaraStackEntry*>& NewChildren, TArray<FStackIssue>& NewIssues) override;
 
diff --git a/Engine/Plugins/FX/Niagara/Source/NiagaraEditorWidgets/Private/Stack/SNiagaraStackFunctionInputCollection.cpp b/Engine/Plugins/FX/Niagara/Source/NiagaraEditorWidgets/Private/Stack/SNiagaraStackFunctionInputCollection.cpp
--- a/Engine/Plugins/FX/Niagara/Source/NiagaraEditorWidgets/Private/Stack/SNiagaraStackFunctionInputCollection.cpp
+++ b/Engine/Plugins/FX/Niagara/Source/NiagaraEditorWidgets/Private/Stack/SNiagaraStackFunctionInputCollection.cpp
@@ -72,7 +72,7 @@ void SNiagaraStackFunctionInputCollection::InputCollectionStructureChanged(ENiag
 
 ECheckBoxState SNiagaraStackFunctionInputCollection::GetSectionCheckState(FText Section) const
 {
-	return Section.IdenticalTo(PropertyCollection->GetActiveSection()) ? ECheckBoxState::Checked : ECheckBoxState::Unchecked;
+	return Section.EqualTo(PropertyCollection->GetActiveSection()) ? ECheckBoxState::Checked : ECheckBoxState::Unchecked;
 }
 
 void SNiagaraStackFunctionInputCollection::OnSectionChecked(ECheckBoxState CheckState, FText Section)
