- Enable half support for linear color data set accessor
- This fixes an issue with linear color reading on compressed systems

#rnx

diff --git a/Engine/Plugins/FX/Niagara/Source/Niagara/Classes/NiagaraDataSetAccessor.h b/Engine/Plugins/FX/Niagara/Source/Niagara/Classes/NiagaraDataSetAccessor.h
--- a/Engine/Plugins/FX/Niagara/Source/Niagara/Classes/NiagaraDataSetAccessor.h
+++ b/Engine/Plugins/FX/Niagara/Source/Niagara/Classes/NiagaraDataSetAccessor.h
@@ -322,11 +322,11 @@ struct FNiagaraDataSetAccessorTypeInfo<FLinearColor>
 {
 	using TAccessorBaseClass = FNiagaraDataSetAccessorFloat<FLinearColor>;
 
-	static constexpr bool bSupportsHalf = false;
+	static constexpr bool bSupportsHalf = true;
 	static constexpr bool bSupportsAlternateType = false;
 	static constexpr int32 NumElements = 4;
 	static const FNiagaraTypeDefinition& GetFloatType() { return FNiagaraTypeDefinition::GetColorDef(); }
-	static const FNiagaraTypeDefinition& GetHalfType() { check(false); return FNiagaraTypeDefinition::GetHalfVec4Def(); }
+	static const FNiagaraTypeDefinition& GetHalfType() { return FNiagaraTypeDefinition::GetHalfVec4Def(); }
 	static const FNiagaraTypeDefinition& GetAlternateType() { return FNiagaraTypeDefinition::GetColorDef(); }
 };
 
