Add a Cook Stat for `UNiagaraSystem::WaitForCompilationComplete` which is on the critical path for cooks.

diff --git a/Engine/Plugins/FX/Niagara/Source/Niagara/Private/NiagaraSystem.cpp b/Engine/Plugins/FX/Niagara/Source/Niagara/Private/NiagaraSystem.cpp
--- a/Engine/Plugins/FX/Niagara/Source/Niagara/Private/NiagaraSystem.cpp
+++ b/Engine/Plugins/FX/Niagara/Source/Niagara/Private/NiagaraSystem.cpp
@@ -37,6 +37,7 @@
 #include "Misc/ScopeExit.h"
 #include "Modules/ModuleManager.h"
 #include "ProfilingDebugging/CookStats.h"
+#include "ProfilingDebugging/ScopedTimers.h"
 #include "UObject/ObjectSaveContext.h"
 #include "UObject/Package.h"
 #include "PipelineStateCache.h"
@@ -57,6 +58,12 @@ DECLARE_CYCLE_STAT(TEXT("Niagara - System - CompileScript_ResetAfter"), STAT_Nia
 namespace NiagaraScriptCookStats
 {
 	extern FCookStats::FDDCResourceUsageStats UsageStats;
+
+	static double NiagaraSystemWaitForCompilationCompleteTime = 0.0;
+	static FCookStatsManager::FAutoRegisterCallback RegisterNiagaraSystemCookStats([](FCookStatsManager::AddStatFuncRef AddStat)
+	{
+		AddStat(TEXT("Niagara"), FCookStatsManager::CreateKeyValueArray(TEXT("NiagaraSystemWaitForCompilationCompleteTime"), NiagaraSystemWaitForCompilationCompleteTime));
+	});
 }
 #endif
 
@@ -2412,6 +2419,7 @@ void UNiagaraSystem::WaitForCompilationComplete(bool bIncludingGPUShaders, bool
 	LLM_SCOPE(ELLMTag::Niagara);
 	TRACE_CPUPROFILER_EVENT_SCOPE(WaitForNiagaraCompilation);
 	TRACE_CPUPROFILER_EVENT_SCOPE_TEXT_ON_CHANNEL(*GetPathName(), NiagaraChannel);
+	COOK_STAT(FScopedDurationTimer DurationTimer(NiagaraScriptCookStats::NiagaraSystemWaitForCompilationCompleteTime));
 
 	if (bNeedsRequestCompile)
 	{
