- Handle case where our render target does not have a valid resource
- Test to see if the render target has been force marked as garbage

#rnx

diff --git a/Engine/Plugins/FX/Niagara/Source/Niagara/Private/NiagaraDataInterfaceRenderTarget2D.cpp b/Engine/Plugins/FX/Niagara/Source/Niagara/Private/NiagaraDataInterfaceRenderTarget2D.cpp
--- a/Engine/Plugins/FX/Niagara/Source/Niagara/Private/NiagaraDataInterfaceRenderTarget2D.cpp
+++ b/Engine/Plugins/FX/Niagara/Source/Niagara/Private/NiagaraDataInterfaceRenderTarget2D.cpp
@@ -386,7 +386,7 @@ void UNiagaraDataInterfaceRenderTarget2D::SetShaderParameters(const FNiagaraData
 	FRDGBuilder& GraphBuilder = Context.GetGraphBuilder();
 
 	// Ensure RDG resources are ready to use
-	if (InstanceData_RT->TransientRDGTexture == nullptr)
+	if (InstanceData_RT->TransientRDGTexture == nullptr && InstanceData_RT->TextureRHI)
 	{
 		InstanceData_RT->TransientRDGTexture = GraphBuilder.FindExternalTexture(InstanceData_RT->TextureRHI);
 		if (InstanceData_RT->TransientRDGTexture == nullptr)
@@ -702,6 +702,11 @@ bool UNiagaraDataInterfaceRenderTarget2D::PerInstanceTickPostSimulate(void* PerI
 	//-TEMP: Until we prune data interface on cook this will avoid consuming memory
 	const bool bValidGpuDataInterface = NiagaraDataInterfaceRenderTargetCommon::GIgnoreCookedOut == 0 || IsUsedWithGPUEmitter();
 
+	if (::IsValid(InstanceData->TargetTexture) == false)
+	{
+		InstanceData->TargetTexture = nullptr;
+	}
+
 	// Do we need to create a new texture?
 	if (bValidGpuDataInterface && !bInheritUserParameterSettings && (InstanceData->TargetTexture == nullptr))
 	{
@@ -744,6 +749,7 @@ bool UNiagaraDataInterfaceRenderTarget2D::PerInstanceTickPostSimulate(void* PerI
 	{
 		FNiagaraDataInterfaceProxyRenderTarget2DProxy* RT_Proxy = GetProxyAs<FNiagaraDataInterfaceProxyRenderTarget2DProxy>();
 		FTextureRenderTargetResource* RT_TargetTexture = InstanceData->TargetTexture ? InstanceData->TargetTexture->GameThread_GetRenderTargetResource() : nullptr;
+		ensureMsgf(RT_TargetTexture != nullptr, TEXT("NiagaraDataInterfaceRenderTarget2D - null RT for system (%s)"), *GetNameSafe(SystemInstance->GetSystem()));
 		ENQUEUE_RENDER_COMMAND(NDIRenderTarget2DUpdate)
 		(
 			[RT_Proxy, RT_InstanceID=SystemInstance->GetId(), RT_InstanceData=*InstanceData, RT_TargetTexture](FRHICommandListImmediate& RHICmdList)
@@ -785,7 +791,10 @@ void FNiagaraDataInterfaceProxyRenderTarget2DProxy::PostStage(const FNDIGpuCompu
 		{
 			ProxyData->bRebuildMips = false;
 			//-TODO:RDG:GPUProfiler:FNiagaraGpuProfileScope GpuProfileScope(RHICmdList, Context, GNiagaraRenderTarget2DGenerateMipsName);
-			NiagaraGenerateMips::GenerateMips(Context.GetGraphBuilder(), ProxyData->TransientRDGTexture, ProxyData->MipMapGenerationType);
+			if (ProxyData->TransientRDGTexture)
+			{
+				NiagaraGenerateMips::GenerateMips(Context.GetGraphBuilder(), ProxyData->TransientRDGTexture, ProxyData->MipMapGenerationType);
+			}
 		}
 	}
 }
@@ -809,7 +818,10 @@ void FNiagaraDataInterfaceProxyRenderTarget2DProxy::PostSimulate(const FNDIGpuCo
 	// We also assume the texture is important for rendering, without discovering renderer bindings we don't really know
 	if (ProxyData->bWroteThisFrame)
 	{
-		Context.GetComputeDispatchInterface().MultiGPUResourceModified(Context.GetGraphBuilder(), ProxyData->TextureRHI, ProxyData->bReadThisFrame, true);
+		if (ProxyData->TransientRDGTexture)
+		{
+			Context.GetComputeDispatchInterface().MultiGPUResourceModified(Context.GetGraphBuilder(), ProxyData->TextureRHI, ProxyData->bReadThisFrame, true);
+		}
 	}
 
 #if NIAGARA_COMPUTEDEBUG_ENABLED && WITH_EDITORONLY_DATA
diff --git a/Engine/Plugins/FX/Niagara/Source/Niagara/Private/NiagaraDataInterfaceRenderTarget2DArray.cpp b/Engine/Plugins/FX/Niagara/Source/Niagara/Private/NiagaraDataInterfaceRenderTarget2DArray.cpp
--- a/Engine/Plugins/FX/Niagara/Source/Niagara/Private/NiagaraDataInterfaceRenderTarget2DArray.cpp
+++ b/Engine/Plugins/FX/Niagara/Source/Niagara/Private/NiagaraDataInterfaceRenderTarget2DArray.cpp
@@ -389,7 +389,7 @@ void UNiagaraDataInterfaceRenderTarget2DArray::SetShaderParameters(const FNiagar
 	FRDGBuilder& GraphBuilder = Context.GetGraphBuilder();
 
 	// Ensure RDG resources are ready to use
-	if (InstanceData_RT->TransientRDGTexture == nullptr)
+	if (InstanceData_RT->TransientRDGTexture == nullptr && InstanceData_RT->RenderTarget)
 	{
 		InstanceData_RT->TransientRDGTexture = GraphBuilder.RegisterExternalTexture(InstanceData_RT->RenderTarget);
 		InstanceData_RT->TransientRDGSRV = GraphBuilder.CreateSRV(InstanceData_RT->TransientRDGTexture);
@@ -666,6 +666,11 @@ bool UNiagaraDataInterfaceRenderTarget2DArray::PerInstanceTickPostSimulate(void*
 	//-TEMP: Until we prune data interface on cook this will avoid consuming memory
 	const bool bValidGpuDataInterface = NiagaraDataInterfaceRenderTargetCommon::GIgnoreCookedOut == 0 || IsUsedWithGPUEmitter();
 
+	if (::IsValid(InstanceData->TargetTexture) == false)
+	{
+		InstanceData->TargetTexture = nullptr;
+	}
+
 	if (bValidGpuDataInterface && !bInheritUserParameterSettings && (InstanceData->TargetTexture == nullptr))
 	{
 		InstanceData->TargetTexture = NewObject<UTextureRenderTarget2DArray>(this);
diff --git a/Engine/Plugins/FX/Niagara/Source/Niagara/Private/NiagaraDataInterfaceRenderTargetCube.cpp b/Engine/Plugins/FX/Niagara/Source/Niagara/Private/NiagaraDataInterfaceRenderTargetCube.cpp
--- a/Engine/Plugins/FX/Niagara/Source/Niagara/Private/NiagaraDataInterfaceRenderTargetCube.cpp
+++ b/Engine/Plugins/FX/Niagara/Source/Niagara/Private/NiagaraDataInterfaceRenderTargetCube.cpp
@@ -354,7 +354,7 @@ void UNiagaraDataInterfaceRenderTargetCube::SetShaderParameters(const FNiagaraDa
 	FRDGBuilder& GraphBuilder = Context.GetGraphBuilder();
 
 	// Ensure RDG resources are ready to use
-	if (InstanceData_RT->TransientRDGTexture == nullptr)
+	if (InstanceData_RT->TransientRDGTexture == nullptr && InstanceData_RT->RenderTarget)
 	{
 		InstanceData_RT->TransientRDGTexture = GraphBuilder.RegisterExternalTexture(InstanceData_RT->RenderTarget);
 		InstanceData_RT->TransientRDGSRV = GraphBuilder.CreateSRV(InstanceData_RT->TransientRDGTexture);
@@ -616,6 +616,11 @@ bool UNiagaraDataInterfaceRenderTargetCube::PerInstanceTickPostSimulate(void* Pe
 	//-TEMP: Until we prune data interface on cook this will avoid consuming memory
 	const bool bValidGpuDataInterface = NiagaraDataInterfaceRenderTargetCommon::GIgnoreCookedOut == 0 || IsUsedWithGPUEmitter();
 
+	if (::IsValid(InstanceData->TargetTexture) == false)
+	{
+		InstanceData->TargetTexture = nullptr;
+	}
+
 	// Do we need to create a new texture?
 	if (bValidGpuDataInterface && !bInheritUserParameterSettings && (InstanceData->TargetTexture == nullptr))
 	{
diff --git a/Engine/Plugins/FX/Niagara/Source/Niagara/Private/NiagaraDataInterfaceRenderTargetVolume.cpp b/Engine/Plugins/FX/Niagara/Source/Niagara/Private/NiagaraDataInterfaceRenderTargetVolume.cpp
--- a/Engine/Plugins/FX/Niagara/Source/Niagara/Private/NiagaraDataInterfaceRenderTargetVolume.cpp
+++ b/Engine/Plugins/FX/Niagara/Source/Niagara/Private/NiagaraDataInterfaceRenderTargetVolume.cpp
@@ -451,7 +451,7 @@ void UNiagaraDataInterfaceRenderTargetVolume::SetShaderParameters(const FNiagara
 	FRDGBuilder& GraphBuilder = Context.GetGraphBuilder();
 
 	// Ensure RDG resources are ready to use
-	if (InstanceData_RT->TransientRDGTexture == nullptr)
+	if (InstanceData_RT->TransientRDGTexture == nullptr && InstanceData_RT->RenderTarget)
 	{
 		InstanceData_RT->TransientRDGTexture = GraphBuilder.RegisterExternalTexture(InstanceData_RT->RenderTarget);
 		InstanceData_RT->TransientRDGSRV = GraphBuilder.CreateSRV(InstanceData_RT->TransientRDGTexture);
@@ -1110,6 +1110,11 @@ bool UNiagaraDataInterfaceRenderTargetVolume::PerInstanceTickPostSimulate(void*
 	//-TEMP: Until we prune data interface on cook this will avoid consuming memory
 	const bool bValidGpuDataInterface = NiagaraDataInterfaceRenderTargetCommon::GIgnoreCookedOut == 0 || IsUsedWithGPUEmitter();
 
+	if (::IsValid(InstanceData->TargetTexture) == false)
+	{
+		InstanceData->TargetTexture = nullptr;
+	}
+
 	// Do we need to create a new texture?
 	if (bValidGpuDataInterface && !bInheritUserParameterSettings && (InstanceData->TargetTexture == nullptr))
 	{
