Moving some Niagara effect wait time to it's own category in CSV reports and out of WorldTickMisc.
#RNX

diff --git a/Engine/Plugins/FX/Niagara/Source/Niagara/Private/NiagaraSystemInstance.cpp b/Engine/Plugins/FX/Niagara/Source/Niagara/Private/NiagaraSystemInstance.cpp
--- a/Engine/Plugins/FX/Niagara/Source/Niagara/Private/NiagaraSystemInstance.cpp
+++ b/Engine/Plugins/FX/Niagara/Source/Niagara/Private/NiagaraSystemInstance.cpp
@@ -2252,6 +2252,7 @@ void FNiagaraSystemInstance::DumpStalledInfo()
 
 void FNiagaraSystemInstance::WaitForConcurrentTickDoNotFinalize(bool bEnsureComplete)
 {
+	CSV_SCOPED_SET_WAIT_STAT(Effects);
 	check(IsInGameThread());
 
 	// Wait for any concurrent ticking for our task
diff --git a/Engine/Plugins/FX/Niagara/Source/Niagara/Private/NiagaraSystemSimulation.cpp b/Engine/Plugins/FX/Niagara/Source/Niagara/Private/NiagaraSystemSimulation.cpp
--- a/Engine/Plugins/FX/Niagara/Source/Niagara/Private/NiagaraSystemSimulation.cpp
+++ b/Engine/Plugins/FX/Niagara/Source/Niagara/Private/NiagaraSystemSimulation.cpp
@@ -1649,6 +1649,7 @@ void FNiagaraSystemSimulation::DumpStalledInfo()
 
 void FNiagaraSystemSimulation::WaitForConcurrentTickComplete(bool bEnsureComplete)
 {
+	CSV_SCOPED_SET_WAIT_STAT(Effects);
 	check(IsInGameThread());
 
 	if (ConcurrentTickGraphEvent.IsValid() && !ConcurrentTickGraphEvent->IsComplete())
