Fix for emitter sourced mesh rendering using MeshIndex binding

diff --git a/Engine/Plugins/FX/Niagara/Source/Niagara/Private/NiagaraRendererMeshes.cpp b/Engine/Plugins/FX/Niagara/Source/Niagara/Private/NiagaraRendererMeshes.cpp
--- a/Engine/Plugins/FX/Niagara/Source/Niagara/Private/NiagaraRendererMeshes.cpp
+++ b/Engine/Plugins/FX/Niagara/Source/Niagara/Private/NiagaraRendererMeshes.cpp
@@ -1264,7 +1264,10 @@ void FNiagaraRendererMeshes::GetDynamicMeshElements(const TArray<const FSceneVie
 			for (int32 MeshIndex = 0; MeshIndex < Meshes.Num(); ++MeshIndex)
 			{
 				// No binding for mesh index or we don't allow culling we will only render the first mesh for all particles
-				if (MeshIndex > 0 && (ParticleMeshRenderData.MeshIndexOffset == INDEX_NONE || (!ParticleMeshRenderData.bNeedsCull && !ParticleMeshRenderData.bUseGPUScene)))
+				const bool bHasValidBinding = ParticleMeshRenderData.MeshIndexOffset != INDEX_NONE
+					|| EmitterModeMeshIndex != INDEX_NONE;
+
+				if (MeshIndex > 0 && (!bHasValidBinding || (!ParticleMeshRenderData.bNeedsCull && !ParticleMeshRenderData.bUseGPUScene)))
 				{
 					break;
 				}
