Attempt to sanitize symbol name when creating a pin for a customhlsl node

diff --git a/Engine/Plugins/FX/Niagara/Source/NiagaraEditor/Private/NiagaraNodeCustomHlsl.cpp b/Engine/Plugins/FX/Niagara/Source/NiagaraEditor/Private/NiagaraNodeCustomHlsl.cpp
--- a/Engine/Plugins/FX/Niagara/Source/NiagaraEditor/Private/NiagaraNodeCustomHlsl.cpp
+++ b/Engine/Plugins/FX/Niagara/Source/NiagaraEditor/Private/NiagaraNodeCustomHlsl.cpp
@@ -369,7 +369,9 @@ void UNiagaraNodeCustomHlsl::OnNewTypedPinAdded(UEdGraphPin*& NewPin)
 		if (Pins[i] != NewPin)
 			Names.Add(Pins[i]->GetFName());
 	}
-	FName Name = FNiagaraUtilities::GetUniqueName(NewPin->GetFName(), Names);
+	FNameBuilder OriginalPinName(NewPin->GetFName());
+	const FName SanitizedName = *FNiagaraHlslTranslator::GetSanitizedSymbolName(OriginalPinName.ToView());
+	FName Name = FNiagaraUtilities::GetUniqueName(SanitizedName, Names);
 	NewPin->PinName = Name;
 	UNiagaraNodeWithDynamicPins::OnNewTypedPinAdded(NewPin);
 	RebuildSignatureFromPins();
