Fix for uniqueness of NPC namespaces

diff --git a/Engine/Plugins/FX/Niagara/Source/Niagara/Classes/NiagaraParameterCollection.h b/Engine/Plugins/FX/Niagara/Source/Niagara/Classes/NiagaraParameterCollection.h
--- a/Engine/Plugins/FX/Niagara/Source/Niagara/Classes/NiagaraParameterCollection.h
+++ b/Engine/Plugins/FX/Niagara/Source/Niagara/Classes/NiagaraParameterCollection.h
@@ -137,6 +137,7 @@ public:
 	//~UObject interface
 #if WITH_EDITORONLY_DATA
 	NIAGARA_API virtual void PostEditChangeProperty(struct FPropertyChangedEvent& PropertyChangedEvent)override;
+	NIAGARA_API virtual void PostDuplicate(EDuplicateMode::Type DuplicateMode) override;
 #endif
 	//~UObject interface
 
diff --git a/Engine/Plugins/FX/Niagara/Source/Niagara/Private/NiagaraParameterCollection.cpp b/Engine/Plugins/FX/Niagara/Source/Niagara/Private/NiagaraParameterCollection.cpp
--- a/Engine/Plugins/FX/Niagara/Source/Niagara/Private/NiagaraParameterCollection.cpp
+++ b/Engine/Plugins/FX/Niagara/Source/Niagara/Private/NiagaraParameterCollection.cpp
@@ -440,6 +440,14 @@ void UNiagaraParameterCollection::PostEditChangeProperty(struct FPropertyChanged
 		OnChangedDelegate.Broadcast();
 	}
 }
+
+void UNiagaraParameterCollection::PostDuplicate(EDuplicateMode::Type DuplicateMode)
+{
+	Super::PostDuplicate(DuplicateMode);
+
+	MakeNamespaceNameUnique();
+}
+
 #endif
 
 int32 UNiagaraParameterCollection::IndexOfParameter(const FNiagaraVariable& Var)
@@ -544,7 +552,8 @@ void UNiagaraParameterCollection::MakeNamespaceNameUnique()
 	TArray<FName> ExistingNames;
  	for (FAssetData& CollectionAsset : CollectionAssets)
  	{
-		if (CollectionAsset.GetFullName() != GetFullName())
+		// skip ourselves - note that ColelctionAsset.GetFullName() uses a fully qualified class name as a prefix in contrast to GetFullName()
+		if (CollectionAsset.GetObjectPathString() != GetPathName())
 		{
 			ExistingNames.Add(CollectionAsset.GetTagValueRef<FName>(GET_MEMBER_NAME_CHECKED(UNiagaraParameterCollection, Namespace)));
 		}
