Prevent PostEditChangeProperty from impacting a niagara graph when it's intended for compilation

diff --git a/Engine/Plugins/FX/Niagara/Source/NiagaraEditor/Private/NiagaraGraph.cpp b/Engine/Plugins/FX/Niagara/Source/NiagaraEditor/Private/NiagaraGraph.cpp
--- a/Engine/Plugins/FX/Niagara/Source/NiagaraEditor/Private/NiagaraGraph.cpp
+++ b/Engine/Plugins/FX/Niagara/Source/NiagaraEditor/Private/NiagaraGraph.cpp
@@ -1527,6 +1527,13 @@ TOptional<FNiagaraScriptVariableData> UNiagaraGraph::GetScriptVariableData(const
 
 void UNiagaraGraph::PostEditChangeProperty(FPropertyChangedEvent& PropertyChangedEvent)
 {
+	// if this is a compilation copy we want to avoid sending any notifications (as we're not really tied to anything anymore).
+	// one area that this can happen is through a Reload of the source system while a compilation is in flight.
+	if (bIsForCompilationOnly)
+	{
+		return;
+	}
+
 	NotifyGraphChanged();
 	RefreshParameterReferences();
 }
