Potential fix for rare crash with stack footer accessing null StackEditorData

diff --git a/Engine/Plugins/FX/Niagara/Source/NiagaraEditor/Private/ViewModels/Stack/NiagaraStackItemFooter.cpp b/Engine/Plugins/FX/Niagara/Source/NiagaraEditor/Private/ViewModels/Stack/NiagaraStackItemFooter.cpp
--- a/Engine/Plugins/FX/Niagara/Source/NiagaraEditor/Private/ViewModels/Stack/NiagaraStackItemFooter.cpp
+++ b/Engine/Plugins/FX/Niagara/Source/NiagaraEditor/Private/ViewModels/Stack/NiagaraStackItemFooter.cpp
@@ -56,6 +56,10 @@ void UNiagaraStackItemFooter::SetOnToggleShowAdvanced(FOnToggleShowAdvanced OnEx
 
 bool UNiagaraStackItemFooter::GetShowAdvanced() const
 {
+	if (IsFinalized())
+	{
+		return false;
+	}
 	return GetStackEditorData().GetStackItemShowAdvanced(OwnerStackItemEditorDataKey, false);
 }
 
