- Add a pool for RDG texture / buffers
- On some platforms artifical barriers were being added (i.e. d3d12) which resulted in slowdown
- In my test for 400 dispatches with a neghbor grid it went from 3.6ms of GPU to 0.33ms

diff --git a/Engine/Plugins/FX/Niagara/Source/Niagara/Private/NiagaraEmptyUAVPool.cpp b/Engine/Plugins/FX/Niagara/Source/Niagara/Private/NiagaraEmptyUAVPool.cpp
--- a/Engine/Plugins/FX/Niagara/Source/Niagara/Private/NiagaraEmptyUAVPool.cpp
+++ b/Engine/Plugins/FX/Niagara/Source/Niagara/Private/NiagaraEmptyUAVPool.cpp
@@ -4,18 +4,21 @@
 #include "NiagaraStats.h"
 #include "NiagaraTypes.h"
 #include "RHICommandList.h"
+#include "RenderGraphBuilder.h"
 
 DECLARE_DWORD_ACCUMULATOR_STAT(TEXT("# EmptyUAVs"), STAT_NiagaraEmptyUAVPool, STATGROUP_Niagara);
 
 FNiagaraEmptyUAVPoolScopedAccess::FNiagaraEmptyUAVPoolScopedAccess(class FNiagaraEmptyUAVPool* InEmptyUAVPool)
 	: EmptyUAVPool(InEmptyUAVPool)
 {
+	check(IsInRenderingThread());
 	check(EmptyUAVPool);
 	++EmptyUAVPool->UAVAccessCounter;
 }
 
 FNiagaraEmptyUAVPoolScopedAccess::~FNiagaraEmptyUAVPoolScopedAccess()
 {
+	check(IsInRenderingThread());
 	--EmptyUAVPool->UAVAccessCounter;
 	if (EmptyUAVPool->UAVAccessCounter == 0)
 	{
@@ -23,6 +26,35 @@ FNiagaraEmptyUAVPoolScopedAccess::~FNiagaraEmptyUAVPoolScopedAccess()
 	}
 }
 
+FNiagaraEmptyRDGUAVPoolScopedAccess::FNiagaraEmptyRDGUAVPoolScopedAccess(class FNiagaraEmptyUAVPool* InEmptyUAVPool)
+	: EmptyUAVPool(InEmptyUAVPool)
+{
+	check(IsInRenderingThread());
+	check(EmptyUAVPool);
+	++EmptyUAVPool->RDGUAVAccessCounter;
+}
+
+FNiagaraEmptyRDGUAVPoolScopedAccess::~FNiagaraEmptyRDGUAVPoolScopedAccess()
+{
+	check(IsInRenderingThread());
+	--EmptyUAVPool->RDGUAVAccessCounter;
+	if (EmptyUAVPool->RDGUAVAccessCounter == 0)
+	{
+		EmptyUAVPool->ResetEmptyRDGUAVPools();
+	}
+}
+
+void FNiagaraEmptyUAVPool::Tick()
+{
+	check(IsInRenderingThread());
+	BufferRDGUAVPool.UAVs.Empty();
+	Texture2DRDGUAVPool.UAVs.Empty();
+	Texture2DArrayRDGUAVPool.UAVs.Empty();
+	Texture3DRDGUAVPool.UAVs.Empty();
+	TextureCubeRDGUAVPool.UAVs.Empty();
+	TextureCubeArrayRDGUAVPool.UAVs.Empty();
+}
+
 FRHIUnorderedAccessView* FNiagaraEmptyUAVPool::GetEmptyUAVFromPool(FRHICommandList& RHICmdList, EPixelFormat Format, ENiagaraEmptyUAVType Type)
 {
 	check(IsInRenderingThread());
@@ -92,6 +124,18 @@ FRHIUnorderedAccessView* FNiagaraEmptyUAVPool::GetEmptyUAVFromPool(FRHICommandLi
 				break;
 			}
 
+
+			case ENiagaraEmptyUAVType::TextureCubeArray:
+			{
+				const FRHITextureCreateDesc Desc =
+					FRHITextureCreateDesc::CreateCubeArray(ResourceName, 1, 1, Format)
+					.SetFlags(ETextureCreateFlags::ShaderResource | ETextureCreateFlags::UAV);
+
+				NewUAV.Texture = RHICreateTexture(Desc);
+				NewUAV.UAV = RHICmdList.CreateUnorderedAccessView(NewUAV.Texture, 0);
+				break;
+			}
+
 			default:
 			{
 				checkNoEntry();
@@ -124,6 +168,128 @@ void FNiagaraEmptyUAVPool::ResetEmptyUAVPools()
 	}
 }
 
+FRDGBufferUAVRef FNiagaraEmptyUAVPool::GetEmptyRDGUAVFromPool(FRDGBuilder& GraphBuilder, EPixelFormat Format)
+{
+	check(IsInRenderingThread());
+	checkf(RDGUAVAccessCounter != 0, TEXT("Accessing Niagara's RDG UAV Pool while not within a scope, this could result in a memory leak!"));
+	check(UE::PixelFormat::HasCapabilities(Format, EPixelFormatCapabilities::TypedUAVStore));
+
+	FBufferRDGUAVPool& Pool = BufferRDGUAVPool;
+	if (Pool.NextFreeIndex == Pool.UAVs.Num())
+	{
+		const uint32 BytesPerElement = GPixelFormats[Format].BlockBytes;
+		Pool.UAVs.Add(
+			GraphBuilder.CreateUAV(
+				GraphBuilder.CreateBuffer(
+					FRDGBufferDesc::CreateBufferDesc(BytesPerElement, 1),
+					TEXT("ENiagaraEmptyUAVType::Buffer")
+				),
+				Format,
+				ERDGUnorderedAccessViewFlags::SkipBarrier
+			)
+		);
+	}
+	return Pool.UAVs[Pool.NextFreeIndex++];
+}
+
+FRDGTextureUAVRef FNiagaraEmptyUAVPool::GetEmptyRDGUAVFromPool(FRDGBuilder& GraphBuilder, EPixelFormat Format, ETextureDimension TextureDimension)
+{
+	check(IsInRenderingThread());
+	checkf(RDGUAVAccessCounter != 0, TEXT("Accessing Niagara's RDG UAV Pool while not within a scope, this could result in a memory leak!"));
+	check(UE::PixelFormat::HasCapabilities(Format, EPixelFormatCapabilities::TypedUAVStore));
+
+	FTextureRDGUAVPool* Pool = nullptr;
+	switch (TextureDimension)
+	{
+		case ETextureDimension::Texture2D:			Pool = &Texture2DRDGUAVPool; break;
+		case ETextureDimension::Texture2DArray:		Pool = &Texture2DArrayRDGUAVPool; break;
+		case ETextureDimension::Texture3D:			Pool = &Texture3DRDGUAVPool; break;
+		case ETextureDimension::TextureCube:		Pool = &TextureCubeRDGUAVPool; break;
+		case ETextureDimension::TextureCubeArray:	Pool = &TextureCubeArrayRDGUAVPool; break;
+		default:									checkNoEntry();	return nullptr;
+	}
+	if (Pool->NextFreeIndex == Pool->UAVs.Num())
+	{
+		switch (TextureDimension)
+		{
+			case ETextureDimension::Texture2D:
+				Pool->UAVs.Add(
+					GraphBuilder.CreateUAV(
+						GraphBuilder.CreateTexture(
+							FRDGTextureDesc::Create2D(FIntPoint(1, 1), Format, FClearValueBinding::Black, ETextureCreateFlags::ShaderResource | ETextureCreateFlags::UAV),
+							TEXT("NiagaraEmptyTextureUAV::Texture2D")
+						),
+						ERDGUnorderedAccessViewFlags::SkipBarrier
+					)
+				);
+				break;
+
+			case ETextureDimension::Texture2DArray:
+				Pool->UAVs.Add(
+					GraphBuilder.CreateUAV(
+						GraphBuilder.CreateTexture(
+							FRDGTextureDesc::Create2DArray(FIntPoint(1, 1), Format, FClearValueBinding::Black, ETextureCreateFlags::ShaderResource | ETextureCreateFlags::UAV, 1),
+							TEXT("NiagaraEmptyTextureUAV::Texture2DArray")
+						),
+						ERDGUnorderedAccessViewFlags::SkipBarrier
+					)
+				);
+				break;
+
+			case ETextureDimension::Texture3D:
+				Pool->UAVs.Add(
+					GraphBuilder.CreateUAV(
+						GraphBuilder.CreateTexture(
+							FRDGTextureDesc::Create3D(FIntVector(1, 1, 1), Format, FClearValueBinding::Black, ETextureCreateFlags::ShaderResource | ETextureCreateFlags::UAV),
+							TEXT("NiagaraEmptyTextureUAV::Texture3D")
+						),
+						ERDGUnorderedAccessViewFlags::SkipBarrier
+					)
+				);
+				break;
+
+			case ETextureDimension::TextureCube:
+				Pool->UAVs.Add(
+					GraphBuilder.CreateUAV(
+						GraphBuilder.CreateTexture(
+							FRDGTextureDesc::CreateCube(1, Format, FClearValueBinding::Black, ETextureCreateFlags::ShaderResource | ETextureCreateFlags::UAV),
+							TEXT("NiagaraEmptyTextureUAV::TextureCube")
+						),
+						ERDGUnorderedAccessViewFlags::SkipBarrier
+					)
+				);
+				break;
+
+			case ETextureDimension::TextureCubeArray:
+				Pool->UAVs.Add(
+					GraphBuilder.CreateUAV(
+						GraphBuilder.CreateTexture(
+							FRDGTextureDesc::CreateCubeArray(1, Format, FClearValueBinding::Black, ETextureCreateFlags::ShaderResource | ETextureCreateFlags::UAV, 1),
+							TEXT("NiagaraEmptyTextureUAV::TextureCubeArray")
+						),
+						ERDGUnorderedAccessViewFlags::SkipBarrier
+					)
+				);
+				break;
+
+			default:
+				checkNoEntry();
+				return nullptr;
+		}
+	}
+	return Pool->UAVs[Pool->NextFreeIndex++];
+}
+
+void FNiagaraEmptyUAVPool::ResetEmptyRDGUAVPools()
+{
+	BufferRDGUAVPool.NextFreeIndex = 0;
+	Texture2DRDGUAVPool.NextFreeIndex = 0;
+	Texture2DArrayRDGUAVPool.NextFreeIndex = 0;
+	Texture3DRDGUAVPool.NextFreeIndex = 0;
+	TextureCubeRDGUAVPool.NextFreeIndex = 0;
+	TextureCubeArrayRDGUAVPool.NextFreeIndex = 0;
+}
+
 FNiagaraEmptyUAVPool::FEmptyUAV::~FEmptyUAV()
 {
 	Buffer.SafeRelease();
diff --git a/Engine/Plugins/FX/Niagara/Source/Niagara/Private/NiagaraGpuComputeDispatch.cpp b/Engine/Plugins/FX/Niagara/Source/Niagara/Private/NiagaraGpuComputeDispatch.cpp
--- a/Engine/Plugins/FX/Niagara/Source/Niagara/Private/NiagaraGpuComputeDispatch.cpp
+++ b/Engine/Plugins/FX/Niagara/Source/Niagara/Private/NiagaraGpuComputeDispatch.cpp
@@ -1521,6 +1521,9 @@ void FNiagaraGpuComputeDispatch::DispatchStage(FRDGBuilder& GraphBuilder, const
 
 	checkf(DispatchNumThreads.X * DispatchNumThreads.Y * DispatchNumThreads.Z > 0, TEXT("DispatchNumThreads(%d, %d, %d) is invalid"), DispatchNumThreads.X, DispatchNumThreads.Y, DispatchNumThreads.Z);
 
+	// Setup our RDG UAV pool access
+	FNiagaraEmptyRDGUAVPoolScopedAccess RDGUAVPoolAccessScope(GetEmptyUAVPool());
+
 	// Set Parameters
 	const FNiagaraShaderScriptParametersMetadata& NiagaraShaderParametersMetadata = InstanceData.Context->GPUScript_RT->GetScriptParametersMetadata_RT();
 	const FShaderParametersMetadata* ShaderParametersMetadata = NiagaraShaderParametersMetadata.ShaderParametersMetadata.Get();
@@ -1891,6 +1894,7 @@ void FNiagaraGpuComputeDispatch::PreInitViews(FRDGBuilder& GraphBuilder, bool bA
 		});
 	}
 
+	EmptyUAVPoolPtr->Tick();
 	GpuReadbackManagerPtr->Tick();
 #if NIAGARA_COMPUTEDEBUG_ENABLED
 	if ( FNiagaraGpuComputeDebug* GpuComputeDebug = GpuComputeDebugPtr.Get() )
diff --git a/Engine/Plugins/FX/Niagara/Source/Niagara/Private/NiagaraGpuComputeDispatchInterface.cpp b/Engine/Plugins/FX/Niagara/Source/Niagara/Private/NiagaraGpuComputeDispatchInterface.cpp
--- a/Engine/Plugins/FX/Niagara/Source/Niagara/Private/NiagaraGpuComputeDispatchInterface.cpp
+++ b/Engine/Plugins/FX/Niagara/Source/Niagara/Private/NiagaraGpuComputeDispatchInterface.cpp
@@ -77,78 +77,12 @@ FRDGTextureSRVRef FNiagaraGpuComputeDispatchInterface::GetBlackTextureSRV(FRDGBu
 
 FRDGTextureUAVRef FNiagaraGpuComputeDispatchInterface::GetEmptyTextureUAV(FRDGBuilder& GraphBuilder, EPixelFormat Format, ETextureDimension TextureDimension) const
 {
-	check(UE::PixelFormat::HasCapabilities(Format, EPixelFormatCapabilities::TypedUAVStore));
-
-	switch (TextureDimension)
-	{
-		case ETextureDimension::Texture2D:
-			return
-				GraphBuilder.CreateUAV(
-					GraphBuilder.CreateTexture(
-						FRDGTextureDesc::Create2D(FIntPoint(1, 1), Format, FClearValueBinding::Black, ETextureCreateFlags::ShaderResource | ETextureCreateFlags::UAV),
-						TEXT("NiagaraEmptyTextureUAV::Texture2D")
-					),
-					ERDGUnorderedAccessViewFlags::SkipBarrier
-				);
-
-		case ETextureDimension::Texture2DArray:
-			return
-				GraphBuilder.CreateUAV(
-					GraphBuilder.CreateTexture(
-						FRDGTextureDesc::Create2DArray(FIntPoint(1, 1), Format, FClearValueBinding::Black, ETextureCreateFlags::ShaderResource | ETextureCreateFlags::UAV, 1),
-						TEXT("NiagaraEmptyTextureUAV::Texture2DArray")
-					),
-					ERDGUnorderedAccessViewFlags::SkipBarrier
-				);
-
-		case ETextureDimension::Texture3D:
-			return
-				GraphBuilder.CreateUAV(
-					GraphBuilder.CreateTexture(
-						FRDGTextureDesc::Create3D(FIntVector(1, 1, 1), Format, FClearValueBinding::Black, ETextureCreateFlags::ShaderResource | ETextureCreateFlags::UAV),
-						TEXT("NiagaraEmptyTextureUAV::Texture3D")
-					),
-					ERDGUnorderedAccessViewFlags::SkipBarrier
-				);
-
-		case ETextureDimension::TextureCube:
-			return
-				GraphBuilder.CreateUAV(
-					GraphBuilder.CreateTexture(
-						FRDGTextureDesc::CreateCube(1, Format, FClearValueBinding::Black, ETextureCreateFlags::ShaderResource | ETextureCreateFlags::UAV),
-						TEXT("NiagaraEmptyTextureUAV::TextureCube")
-					),
-					ERDGUnorderedAccessViewFlags::SkipBarrier
-				);
-
-		case ETextureDimension::TextureCubeArray:
-			return
-				GraphBuilder.CreateUAV(
-					GraphBuilder.CreateTexture(
-						FRDGTextureDesc::CreateCubeArray(1, Format, FClearValueBinding::Black, ETextureCreateFlags::ShaderResource | ETextureCreateFlags::UAV, 1),
-						TEXT("NiagaraEmptyTextureUAV::TextureCubeArray")
-					),
-					ERDGUnorderedAccessViewFlags::SkipBarrier
-				);
-
-		default:
-			checkNoEntry();
-			return nullptr;
-	}
+	return EmptyUAVPoolPtr->GetEmptyRDGUAVFromPool(GraphBuilder, Format, TextureDimension);
 }
 
 FRDGBufferUAVRef FNiagaraGpuComputeDispatchInterface::GetEmptyBufferUAV(FRDGBuilder& GraphBuilder, EPixelFormat Format) const
 {
-	const uint32 BytesPerElement = GPixelFormats[Format].BlockBytes;
-	return
-		GraphBuilder.CreateUAV(
-			GraphBuilder.CreateBuffer(
-				FRDGBufferDesc::CreateBufferDesc(BytesPerElement, 1),
-				TEXT("ENiagaraEmptyUAVType::Buffer")
-			),
-			Format,
-			ERDGUnorderedAccessViewFlags::SkipBarrier
-		);
+	return EmptyUAVPoolPtr->GetEmptyRDGUAVFromPool(GraphBuilder, Format);
 }
 
 FRDGBufferSRVRef FNiagaraGpuComputeDispatchInterface::GetEmptyBufferSRV(FRDGBuilder& GraphBuilder, EPixelFormat Format) const
diff --git a/Engine/Plugins/FX/Niagara/Source/Niagara/Public/NiagaraEmptyUAVPool.h b/Engine/Plugins/FX/Niagara/Source/Niagara/Public/NiagaraEmptyUAVPool.h
--- a/Engine/Plugins/FX/Niagara/Source/Niagara/Public/NiagaraEmptyUAVPool.h
+++ b/Engine/Plugins/FX/Niagara/Source/Niagara/Public/NiagaraEmptyUAVPool.h
@@ -9,14 +9,23 @@
 // but is required for external usage that is outside, i.e. if you are doing some custom dispatch setup
 struct FNiagaraEmptyUAVPoolScopedAccess
 {
+	UE_NONCOPYABLE(FNiagaraEmptyUAVPoolScopedAccess)
 public:
 	explicit FNiagaraEmptyUAVPoolScopedAccess(class FNiagaraEmptyUAVPool* EmptyUAVPool);
-	FNiagaraEmptyUAVPoolScopedAccess(const FNiagaraEmptyUAVPoolScopedAccess&) = delete;
-	FNiagaraEmptyUAVPoolScopedAccess(FNiagaraEmptyUAVPoolScopedAccess&&) = delete;
 	~FNiagaraEmptyUAVPoolScopedAccess();
 
-	FNiagaraEmptyUAVPoolScopedAccess& operator=(const FNiagaraEmptyUAVPoolScopedAccess&) = delete;
-	FNiagaraEmptyUAVPoolScopedAccess& operator=(FNiagaraEmptyUAVPoolScopedAccess&&) = delete;
+private:
+	class FNiagaraEmptyUAVPool* EmptyUAVPool;
+};
+
+// Scoped access for RDG UAVs this is not required when running inside the main dispatch loop
+// but is required for external usage that is outside, i.e. if you are doing some custom dispatch setup
+struct FNiagaraEmptyRDGUAVPoolScopedAccess
+{
+	UE_NONCOPYABLE(FNiagaraEmptyRDGUAVPoolScopedAccess)
+public:
+	explicit FNiagaraEmptyRDGUAVPoolScopedAccess(class FNiagaraEmptyUAVPool* EmptyUAVPool);
+	~FNiagaraEmptyRDGUAVPoolScopedAccess();
 
 private:
 	class FNiagaraEmptyUAVPool* EmptyUAVPool;
@@ -30,6 +39,7 @@ enum class ENiagaraEmptyUAVType
 	Texture2DArray,
 	Texture3D,
 	TextureCube,
+	TextureCubeArray,
 	Num
 };
 
@@ -37,18 +47,35 @@ enum class ENiagaraEmptyUAVType
 class FNiagaraEmptyUAVPool
 {
 	friend struct FNiagaraEmptyUAVPoolScopedAccess;
+	friend struct FNiagaraEmptyRDGUAVPoolScopedAccess;
 
 public:
+	/** Must be called before we start to use the UAV pool for a given scene render. */
+	void Tick();
+
 	/**
-	Grab a temporary empty RW buffer from the pool.
-	Note: When doing this outside of Niagara you must be within a FNiagaraUAVPoolAccessScope.
+	* Grab a temporary empty RW buffer from the pool.
+	* Note: When doing this outside of Niagara you must be within a FNiagaraUAVPoolAccessScope.
 	*/
 	NIAGARA_API FRHIUnorderedAccessView* GetEmptyUAVFromPool(FRHICommandList& RHICmdList, EPixelFormat Format, ENiagaraEmptyUAVType Type);
 
 	/**
-	Returns all used UAVs back to the pool
+	* Grab a temporary empty RDG Buffer UAV from the pool.
+	* Note: When doing this outside of Niagara you must be within a FNiagaraUAVPoolAccessScope.
 	*/
-	NIAGARA_API void ResetEmptyUAVPools();
+	NIAGARA_API FRDGBufferUAVRef GetEmptyRDGUAVFromPool(FRDGBuilder& GraphBuilder, EPixelFormat Format);
+	/**
+	* Grab a temporary empty RDG Texture UAV from the pool.
+	* Note: When doing this outside of Niagara you must be within a FNiagaraUAVPoolAccessScope.
+	*/
+	NIAGARA_API FRDGTextureUAVRef GetEmptyRDGUAVFromPool(FRDGBuilder& GraphBuilder, EPixelFormat Format, ETextureDimension TextureDimension);
+
+protected:
+	/** Returns all used UAVs back to the pool. */
+	void ResetEmptyUAVPools();
+
+	/** Returns all the RDG UAVs back to the pool. */
+	void ResetEmptyRDGUAVPools();
 
 protected:
 	struct FEmptyUAV
@@ -70,4 +97,23 @@ protected:
 
 	uint32 UAVAccessCounter = 0;
 	TMap<EPixelFormat, FEmptyUAVPool> UAVPools[(int)ENiagaraEmptyUAVType::Num];
+
+	struct FBufferRDGUAVPool
+	{
+		int32 NextFreeIndex = 0;
+		TArray<FRDGBufferUAVRef> UAVs;
+	};
+	struct FTextureRDGUAVPool
+	{
+		int32 NextFreeIndex = 0;
+		TArray<FRDGTextureUAVRef> UAVs;
+	};
+
+	uint32 RDGUAVAccessCounter = 0;
+	FBufferRDGUAVPool BufferRDGUAVPool;
+	FTextureRDGUAVPool Texture2DRDGUAVPool;
+	FTextureRDGUAVPool Texture2DArrayRDGUAVPool;
+	FTextureRDGUAVPool Texture3DRDGUAVPool;
+	FTextureRDGUAVPool TextureCubeRDGUAVPool;
+	FTextureRDGUAVPool TextureCubeArrayRDGUAVPool;
 };
