- IWYU fix

#rnx

diff --git a/Engine/Plugins/FX/Niagara/Source/Niagara/Public/NiagaraEmptyUAVPool.h b/Engine/Plugins/FX/Niagara/Source/Niagara/Public/NiagaraEmptyUAVPool.h
--- a/Engine/Plugins/FX/Niagara/Source/Niagara/Public/NiagaraEmptyUAVPool.h
+++ b/Engine/Plugins/FX/Niagara/Source/Niagara/Public/NiagaraEmptyUAVPool.h
@@ -4,6 +4,7 @@
 
 #include "NiagaraCommon.h"
 #include "RHI.h"
+#include "RenderGraphFwd.h"
 
 // Scoped access for UAVs this is not required when running inside the main dispatch loop
 // but is required for external usage that is outside, i.e. if you are doing some custom dispatch setup
