Early out of updating merged scripts if we're not allowed to compile any shaders.  Avoids problems with trying to generate compile ids because we can't access the shader files for hashing

diff --git a/Engine/Plugins/FX/Niagara/Source/Niagara/Private/NiagaraScript.cpp b/Engine/Plugins/FX/Niagara/Source/Niagara/Private/NiagaraScript.cpp
--- a/Engine/Plugins/FX/Niagara/Source/Niagara/Private/NiagaraScript.cpp
+++ b/Engine/Plugins/FX/Niagara/Source/Niagara/Private/NiagaraScript.cpp
@@ -3804,6 +3804,12 @@ bool UNiagaraScript::SynchronizeExecutablesWithCompilation(const UNiagaraScript*
 		return false;
 	}
 
+	if (!AllowShaderCompiling())
+	{
+		// if we're not allowed to compile shaders than don't even bother evaluating the id
+		return false;
+	}
+
 	FNiagaraVMExecutableDataId Id;
 	ComputeVMCompilationId(Id, FGuid());
 
