- Don't allow GC clusters to be built on assets using certain data interfaces

#rnx

diff --git a/Engine/Plugins/FX/Niagara/Source/Niagara/Classes/NiagaraDataInterfaceIntRenderTarget2D.h b/Engine/Plugins/FX/Niagara/Source/Niagara/Classes/NiagaraDataInterfaceIntRenderTarget2D.h
--- a/Engine/Plugins/FX/Niagara/Source/Niagara/Classes/NiagaraDataInterfaceIntRenderTarget2D.h
+++ b/Engine/Plugins/FX/Niagara/Source/Niagara/Classes/NiagaraDataInterfaceIntRenderTarget2D.h
@@ -21,7 +21,8 @@ class UNiagaraDataInterfaceIntRenderTarget2D : public UNiagaraDataInterfaceRWBas
 
 public:
 	NIAGARA_API virtual void PostInitProperties() override;
-	
+	virtual bool CanBeInCluster() const override { return false; }
+
 	//~ UNiagaraDataInterface interface
 	virtual bool CanExecuteOnTarget(ENiagaraSimTarget Target)const override { return true; }
 	NIAGARA_API virtual void GetFunctions(TArray<FNiagaraFunctionSignature>& OutFunctions) override;
diff --git a/Engine/Plugins/FX/Niagara/Source/Niagara/Classes/NiagaraDataInterfaceRenderTarget2D.h b/Engine/Plugins/FX/Niagara/Source/Niagara/Classes/NiagaraDataInterfaceRenderTarget2D.h
--- a/Engine/Plugins/FX/Niagara/Source/Niagara/Classes/NiagaraDataInterfaceRenderTarget2D.h
+++ b/Engine/Plugins/FX/Niagara/Source/Niagara/Classes/NiagaraDataInterfaceRenderTarget2D.h
@@ -92,7 +92,8 @@ class UNiagaraDataInterfaceRenderTarget2D : public UNiagaraDataInterfaceRWBase
 
 public:
 	NIAGARA_API virtual void PostInitProperties() override;
-	
+	virtual bool CanBeInCluster() const override { return false; }
+
 	//~ UNiagaraDataInterface interface
 	// VM functionality
 	virtual bool CanExecuteOnTarget(ENiagaraSimTarget Target)const override { return true; }
diff --git a/Engine/Plugins/FX/Niagara/Source/Niagara/Classes/NiagaraDataInterfaceRenderTarget2DArray.h b/Engine/Plugins/FX/Niagara/Source/Niagara/Classes/NiagaraDataInterfaceRenderTarget2DArray.h
--- a/Engine/Plugins/FX/Niagara/Source/Niagara/Classes/NiagaraDataInterfaceRenderTarget2DArray.h
+++ b/Engine/Plugins/FX/Niagara/Source/Niagara/Classes/NiagaraDataInterfaceRenderTarget2DArray.h
@@ -88,7 +88,8 @@ class UNiagaraDataInterfaceRenderTarget2DArray : public UNiagaraDataInterfaceRWB
 
 public:
 	NIAGARA_API virtual void PostInitProperties() override;
-	
+	virtual bool CanBeInCluster() const override { return false; }
+
 	//~ UNiagaraDataInterface interface
 	// VM functionality
 	virtual bool CanExecuteOnTarget(ENiagaraSimTarget Target)const override { return true; }
diff --git a/Engine/Plugins/FX/Niagara/Source/Niagara/Classes/NiagaraDataInterfaceRenderTargetCube.h b/Engine/Plugins/FX/Niagara/Source/Niagara/Classes/NiagaraDataInterfaceRenderTargetCube.h
--- a/Engine/Plugins/FX/Niagara/Source/Niagara/Classes/NiagaraDataInterfaceRenderTargetCube.h
+++ b/Engine/Plugins/FX/Niagara/Source/Niagara/Classes/NiagaraDataInterfaceRenderTargetCube.h
@@ -81,7 +81,8 @@ class UNiagaraDataInterfaceRenderTargetCube : public UNiagaraDataInterfaceRWBase
 
 public:
 	NIAGARA_API virtual void PostInitProperties() override;
-	
+	virtual bool CanBeInCluster() const override { return false; }
+
 	//~ UNiagaraDataInterface interface
 	// VM functionality
 	virtual bool CanExecuteOnTarget(ENiagaraSimTarget Target)const override { return true; }
diff --git a/Engine/Plugins/FX/Niagara/Source/Niagara/Classes/NiagaraDataInterfaceRenderTargetVolume.h b/Engine/Plugins/FX/Niagara/Source/Niagara/Classes/NiagaraDataInterfaceRenderTargetVolume.h
--- a/Engine/Plugins/FX/Niagara/Source/Niagara/Classes/NiagaraDataInterfaceRenderTargetVolume.h
+++ b/Engine/Plugins/FX/Niagara/Source/Niagara/Classes/NiagaraDataInterfaceRenderTargetVolume.h
@@ -81,7 +81,8 @@ class UNiagaraDataInterfaceRenderTargetVolume : public UNiagaraDataInterfaceRWBa
 
 public:
 	NIAGARA_API virtual void PostInitProperties() override;
-	
+	virtual bool CanBeInCluster() const override { return false; }
+
 	//~ UNiagaraDataInterface interface
 	// VM functionality
 	virtual bool CanExecuteOnTarget(ENiagaraSimTarget Target)const override { return true; }
