Fixed missing package for versioned script details structonscope. This caused an ensure to be triggered everytime an FText within gets modified

diff --git a/Engine/Plugins/FX/Niagara/Source/NiagaraEditor/Private/Customizations/NiagaraScriptDetails.cpp b/Engine/Plugins/FX/Niagara/Source/NiagaraEditor/Private/Customizations/NiagaraScriptDetails.cpp
--- a/Engine/Plugins/FX/Niagara/Source/NiagaraEditor/Private/Customizations/NiagaraScriptDetails.cpp
+++ b/Engine/Plugins/FX/Niagara/Source/NiagaraEditor/Private/Customizations/NiagaraScriptDetails.cpp
@@ -189,6 +189,7 @@ void FNiagaraScriptDetails::CustomizeDetails(IDetailLayoutBuilder& DetailBuilder
 		if (FVersionedNiagaraScriptData* Data = StandaloneScript.Script->GetScriptData(StandaloneScript.Version))
 		{
 			TSharedPtr<FStructOnScope> StructData = MakeShareable(new FStructOnScope(FVersionedNiagaraScriptData::StaticStruct(), (uint8*)Data));
+			StructData->SetPackage(StandaloneScript.Script.GetPackage());
 			TArray<TSharedPtr<IPropertyHandle>> NewHandles = CategoryBuilder.AddAllExternalStructureProperties(StructData.ToSharedRef());
 
 			for(TSharedPtr<IPropertyHandle> NewHandle : NewHandles)
