Heterogeneous Volumes: Add per-volume step factor and shadow step factor to rendering interface

diff --git a/Engine/Plugins/FX/Niagara/Source/Niagara/Private/NiagaraRendererVolumes.cpp b/Engine/Plugins/FX/Niagara/Source/Niagara/Private/NiagaraRendererVolumes.cpp
--- a/Engine/Plugins/FX/Niagara/Source/Niagara/Private/NiagaraRendererVolumes.cpp
+++ b/Engine/Plugins/FX/Niagara/Source/Niagara/Private/NiagaraRendererVolumes.cpp
@@ -48,6 +48,8 @@ namespace NiagaraRendererVolumesLocal
 
 		int32					VolumeResolutionMaxAxis = UNiagaraVolumeRendererProperties::GetDefaultVolumeResolutionMaxAxis();
 		FVector3f				VolumeWorldSpaceSize = UNiagaraVolumeRendererProperties::GetDefaultVolumeWorldSpaceSize();
+		float					StepFactor = UNiagaraVolumeRendererProperties::GetDefaultStepFactor();
+		float					ShadowStepFactor = UNiagaraVolumeRendererProperties::GetDefaultShadowStepFactor();
 		float					LightingDownsampleFactor = UNiagaraVolumeRendererProperties::GetDefaultLightingDownsampleFactor();
 	};
 
@@ -222,6 +224,8 @@ FNiagaraDynamicDataBase* FNiagaraRendererVolumes::GenerateDynamicData(const FNia
 	const FNiagaraParameterStore& ParameterStore = Emitter->GetRendererBoundVariables();
 	VolumeDynamicData->VolumeResolutionMaxAxis = ParameterStore.GetParameterValueOrDefault(RendererProperties->VolumeResolutionMaxAxisBinding.GetParamMapBindableVariable(), DefaultVolumeResolutionMaxAxis);
 	VolumeDynamicData->VolumeWorldSpaceSize = ParameterStore.GetParameterValueOrDefault(RendererProperties->VolumeWorldSpaceSizeBinding.GetParamMapBindableVariable(), DefaultVolumeWorldSpaceSize);
+	VolumeDynamicData->StepFactor = RendererProperties->StepFactor;
+	VolumeDynamicData->ShadowStepFactor = RendererProperties->ShadowStepFactor;
 	VolumeDynamicData->LightingDownsampleFactor = RendererProperties->LightingDownsampleFactor;
 
 	return VolumeDynamicData;
@@ -353,8 +357,8 @@ void FNiagaraRendererVolumes::GetDynamicMeshElements(const TArray<const FSceneVi
 				FMath::CeilToInt(VolumeResolutionV3f.Z));
 
 			HeterogeneousVolumeData->VoxelResolution = VolumeResolution;
-
-
+			HeterogeneousVolumeData->StepFactor = VolumeDynamicData->StepFactor;
+			HeterogeneousVolumeData->ShadowStepFactor = VolumeDynamicData->ShadowStepFactor;
 			HeterogeneousVolumeData->LightingDownsampleFactor = VolumeDynamicData->LightingDownsampleFactor;
 			BatchElement.UserData = HeterogeneousVolumeData;
 
diff --git a/Engine/Plugins/FX/Niagara/Source/Niagara/Public/NiagaraVolumeRendererProperties.h b/Engine/Plugins/FX/Niagara/Source/Niagara/Public/NiagaraVolumeRendererProperties.h
--- a/Engine/Plugins/FX/Niagara/Source/Niagara/Public/NiagaraVolumeRendererProperties.h
+++ b/Engine/Plugins/FX/Niagara/Source/Niagara/Public/NiagaraVolumeRendererProperties.h
@@ -77,6 +77,8 @@ public:
 	static FVector3f GetDefaultVolumeScale() { return FVector3f(1.0f, 1.0f, 1.0f); }
 	static int32 GetDefaultVolumeResolutionMaxAxis() { return 0; }
 	static FVector3f GetDefaultVolumeWorldSpaceSize() { return FVector3f::ZeroVector; }
+	static float GetDefaultStepFactor() { return 1.0f; }
+	static float GetDefaultShadowStepFactor() { return 8.0f; }
 	static float GetDefaultLightingDownsampleFactor() { return 2.0f; }
 
 	/** What material to use for the volume. */
@@ -101,9 +103,15 @@ public:
 	UPROPERTY(EditAnywhere, Category = "Volume Rendering")
 	int32 RendererVisibility = 0;
 
+	UPROPERTY(EditAnywhere, Category = "Volume Rendering")
+	float StepFactor = 1.0f;
+
 	UPROPERTY(EditAnywhere, Category = "Volume Rendering")
 	float LightingDownsampleFactor = 2.0f;
 
+	UPROPERTY(EditAnywhere, Category = "Volume Rendering")
+	float ShadowStepFactor = 8.0f;
+
 	/** Position binding for the center of the volume. */
 	//-TODO: Unhind when binding is supported
 	//UPROPERTY(EditAnywhere, Category = "Bindings")
diff --git a/Engine/Source/Runtime/Engine/Classes/Components/HeterogeneousVolumeComponent.h b/Engine/Source/Runtime/Engine/Classes/Components/HeterogeneousVolumeComponent.h
--- a/Engine/Source/Runtime/Engine/Classes/Components/HeterogeneousVolumeComponent.h
+++ b/Engine/Source/Runtime/Engine/Classes/Components/HeterogeneousVolumeComponent.h
@@ -28,6 +28,12 @@ class UHeterogeneousVolumeComponent : public UMeshComponent
 	UPROPERTY(EditAnywhere, BlueprintReadOnly, Category = Volume)
 	uint32 bAnimate : 1;
 
+	UPROPERTY(EditAnywhere, BlueprintReadOnly, Category = Lighting)
+	float StepFactor;
+
+	UPROPERTY(EditAnywhere, BlueprintReadOnly, Category = Lighting)
+	float ShadowStepFactor;
+
 	UPROPERTY(EditAnywhere, BlueprintReadOnly, Category = Lighting)
 	float LightingDownsampleFactor;
 	
diff --git a/Engine/Source/Runtime/Engine/Private/Components/HeterogeneousVolumeComponent.cpp b/Engine/Source/Runtime/Engine/Private/Components/HeterogeneousVolumeComponent.cpp
--- a/Engine/Source/Runtime/Engine/Private/Components/HeterogeneousVolumeComponent.cpp
+++ b/Engine/Source/Runtime/Engine/Private/Components/HeterogeneousVolumeComponent.cpp
@@ -81,6 +81,8 @@ FHeterogeneousVolumeSceneProxy::FHeterogeneousVolumeSceneProxy(UHeterogeneousVol
 		MaterialInterface = InComponent->MaterialInstanceDynamic;
 	}
 
+	HeterogeneousVolumeData.StepFactor = InComponent->StepFactor;
+	HeterogeneousVolumeData.ShadowStepFactor = InComponent->ShadowStepFactor;
 	HeterogeneousVolumeData.LightingDownsampleFactor = InComponent->LightingDownsampleFactor;
 
 	// Initialize vertex buffer data for a quad
@@ -214,7 +216,9 @@ UHeterogeneousVolumeComponent::UHeterogeneousVolumeComponent(const FObjectInitia
 	VolumeResolution = FIntVector(128);
 	Framerate = 24.0f;
 	bAnimate = false;
-	LightingDownsampleFactor = 1.0f;
+	StepFactor = 1.0f;
+	ShadowStepFactor = 8.0f;
+	LightingDownsampleFactor = 2.0f;
 	Time = 0.0f;
 }
 
diff --git a/Engine/Source/Runtime/Engine/Public/HeterogeneousVolumeInterface.h b/Engine/Source/Runtime/Engine/Public/HeterogeneousVolumeInterface.h
--- a/Engine/Source/Runtime/Engine/Public/HeterogeneousVolumeInterface.h
+++ b/Engine/Source/Runtime/Engine/Public/HeterogeneousVolumeInterface.h
@@ -26,6 +26,8 @@ public:
 	virtual float GetMinimumVoxelSize() const = 0;
 
 	// Lighting
+	virtual float GetStepFactor() const = 0;
+	virtual float GetShadowStepFactor() const = 0;
 	virtual float GetLightingDownsampleFactor() const = 0;
 };
 
@@ -39,6 +41,8 @@ public:
 		, InstanceToLocal(FMatrix::Identity)
 		, VoxelResolution(FIntVector::ZeroValue)
 		, MinimumVoxelSize(0.1)
+		, StepFactor(1.0)
+		, ShadowStepFactor(8.0)
 		, LightingDownsampleFactor(1.0)
 	{}
 	virtual ~FHeterogeneousVolumeData() {}
@@ -57,11 +61,15 @@ public:
 	virtual float GetMinimumVoxelSize() const { return MinimumVoxelSize; }
 
 	// Lighting
+	virtual float GetStepFactor() const { return StepFactor; }
+	virtual float GetShadowStepFactor() const { return ShadowStepFactor; }
 	virtual float GetLightingDownsampleFactor() const { return LightingDownsampleFactor; }
 
 	const FPrimitiveSceneProxy* PrimitiveSceneProxy;
 	FMatrix InstanceToLocal;
 	FIntVector VoxelResolution;
 	float MinimumVoxelSize;
+	float StepFactor;
+	float ShadowStepFactor;
 	float LightingDownsampleFactor;
 };
