- Remove FGCObject from system simulation as it's unused

diff --git a/Engine/Plugins/FX/Niagara/Source/Niagara/Private/NiagaraSystemSimulation.cpp b/Engine/Plugins/FX/Niagara/Source/Niagara/Private/NiagaraSystemSimulation.cpp
--- a/Engine/Plugins/FX/Niagara/Source/Niagara/Private/NiagaraSystemSimulation.cpp
+++ b/Engine/Plugins/FX/Niagara/Source/Niagara/Private/NiagaraSystemSimulation.cpp
@@ -608,10 +608,6 @@ FNiagaraSystemSimulationTickContext::FNiagaraSystemSimulationTickContext(class F
 
 //////////////////////////////////////////////////////////////////////////
 
-void FNiagaraSystemSimulation::AddReferencedObjects(FReferenceCollector& Collector)
-{
-}
-
 FNiagaraSystemSimulation::FNiagaraSystemSimulation()
 	: EffectType(nullptr)
 	, SystemTickGroup(TG_MAX)
diff --git a/Engine/Plugins/FX/Niagara/Source/Niagara/Public/NiagaraSystemSimulation.h b/Engine/Plugins/FX/Niagara/Source/Niagara/Public/NiagaraSystemSimulation.h
--- a/Engine/Plugins/FX/Niagara/Source/Niagara/Public/NiagaraSystemSimulation.h
+++ b/Engine/Plugins/FX/Niagara/Source/Niagara/Public/NiagaraSystemSimulation.h
@@ -4,7 +4,6 @@
 
 #include "Engine/World.h"
 #include "NiagaraParameterCollection.h"
-#include "UObject/GCObject.h"
 #include "NiagaraDataSet.h"
 #include "NiagaraScriptExecutionContext.h"
 #include "NiagaraSystem.h"
@@ -244,21 +243,13 @@ public:
 };
 
 /** Simulation performing all system and emitter scripts for a instances of a UNiagaraSystem in a world. */
-class FNiagaraSystemSimulation : public TSharedFromThis<FNiagaraSystemSimulation, ESPMode::ThreadSafe>, FGCObject
+class FNiagaraSystemSimulation : public TSharedFromThis<FNiagaraSystemSimulation, ESPMode::ThreadSafe>
 {
 	friend FNiagaraSystemSimulationTickContext;
 	friend struct FNiagaraSimCacheHelper;
 	friend class FNiagaraDebugHud;
 
 public:
-	//FGCObject Interface
-	virtual void AddReferencedObjects(FReferenceCollector& Collector)override;
-	virtual FString GetReferencerName() const override
-	{
-		return TEXT("FNiagaraSystemSimulation");
-	}
-	//FGCObject Interface END
-
 	FNiagaraSystemSimulation();
 	~FNiagaraSystemSimulation();
 	bool Init(UNiagaraSystem* InSystem, UWorld* InWorld, bool bInIsSolo, ETickingGroup TickGroup);
