fix case where we have mesh distance field collisions active but no dfs (ie: skel meshes)
update all liquid examples

diff --git a/Engine/Plugins/FX/Niagara/Shaders/Private/NiagaraDataInterfaceRigidMeshCollisionQuery.ush b/Engine/Plugins/FX/Niagara/Shaders/Private/NiagaraDataInterfaceRigidMeshCollisionQuery.ush
--- a/Engine/Plugins/FX/Niagara/Shaders/Private/NiagaraDataInterfaceRigidMeshCollisionQuery.ush
+++ b/Engine/Plugins/FX/Niagara/Shaders/Private/NiagaraDataInterfaceRigidMeshCollisionQuery.ush
@@ -506,12 +506,16 @@ void GetClosestPointMeshDistanceFieldNoNormal_{ParameterName}(in float3 WorldPos
 
 			OutClosestDistance = DistanceToNearestSurfaceForObject(DFIndex, LWCWorldPosition, MaxDistance);			
 		}
-		
+		else
+		{
+			OutClosestDistance = OutClosestDistanceTmp;
+		}
+
 		const float3 PreviousPosition = mul(GetPreviousTransform_{ParameterName}(ElementIndex), float4(CollisionPosition,1.0)).xyz;
 		const float3 CurrentPosition = mul(GetCurrentTransform_{ParameterName}(ElementIndex), float4(CollisionPosition,1.0)).xyz;
 		OutClosestVelocity = ( CurrentPosition - PreviousPosition ) / DeltaTime;
 		OutClosestPosition = PreviousPosition + TimeFraction * (CurrentPosition-PreviousPosition);
-	}	
+	}
 }
 
 #undef BOX_INDEX
