- Add a simple object pool & object referencer for Niagara to the World Manager
- Users of this pool & referencer are expected to handle recycling / unregistering correctly
- Allow clusters to be built for render target data interfaces by moving object referencing into the world manager
- Use object pool for render targets so that we can avoid waiting on GC before resources are collected
- This reduces a lot of memory usage when scrubbing in the editor / sequencer timelines

diff --git a/Engine/Plugins/FX/Niagara/Source/Niagara/Classes/NiagaraDataInterfaceIntRenderTarget2D.h b/Engine/Plugins/FX/Niagara/Source/Niagara/Classes/NiagaraDataInterfaceIntRenderTarget2D.h
--- a/Engine/Plugins/FX/Niagara/Source/Niagara/Classes/NiagaraDataInterfaceIntRenderTarget2D.h
+++ b/Engine/Plugins/FX/Niagara/Source/Niagara/Classes/NiagaraDataInterfaceIntRenderTarget2D.h
@@ -76,7 +76,4 @@ public:
 
 protected:
 	static NIAGARA_API FNiagaraVariableBase ExposedRTVar;
-
-	UPROPERTY(Transient, DuplicateTransient)
-	TMap<uint64, TObjectPtr<UTextureRenderTarget2D>> ManagedRenderTargets;
 };
diff --git a/Engine/Plugins/FX/Niagara/Source/Niagara/Classes/NiagaraDataInterfaceRenderTarget2D.h b/Engine/Plugins/FX/Niagara/Source/Niagara/Classes/NiagaraDataInterfaceRenderTarget2D.h
--- a/Engine/Plugins/FX/Niagara/Source/Niagara/Classes/NiagaraDataInterfaceRenderTarget2D.h
+++ b/Engine/Plugins/FX/Niagara/Source/Niagara/Classes/NiagaraDataInterfaceRenderTarget2D.h
@@ -16,35 +16,22 @@ class UTextureRenderTarget2D;
 
 struct FRenderTarget2DRWInstanceData_GameThread
 {
-	FRenderTarget2DRWInstanceData_GameThread()
-	{
-#if WITH_EDITORONLY_DATA
-		bPreviewTexture = false;
-#endif
-	}
-
 	FIntPoint Size = FIntPoint(EForceInit::ForceInitToZero);
 	ETextureRenderTargetFormat Format = RTF_RGBA16f;
 	TextureFilter Filter = TextureFilter::TF_Default;
 	ENiagaraMipMapGeneration MipMapGeneration = ENiagaraMipMapGeneration::Disabled;
 	ENiagaraMipMapGenerationType MipMapGenerationType = ENiagaraMipMapGenerationType::Linear;
-	
+
+	bool bManagedTexture = false;
 	UTextureRenderTarget2D* TargetTexture = nullptr;
 #if WITH_EDITORONLY_DATA
-	uint32 bPreviewTexture : 1;
+	bool bPreviewTexture = false;
 #endif
 	FNiagaraParameterDirectBinding<UObject*> RTUserParamBinding;
 };
 
 struct FRenderTarget2DRWInstanceData_RenderThread
 {
-	FRenderTarget2DRWInstanceData_RenderThread()
-	{
-#if WITH_EDITORONLY_DATA
-		bPreviewTexture = false;
-#endif
-	}
-
 	FIntPoint Size = FIntPoint(EForceInit::ForceInitToZero);
 	int MipLevels = 0;
 	ENiagaraMipMapGeneration MipMapGeneration = ENiagaraMipMapGeneration::Disabled;
@@ -61,7 +48,7 @@ struct FRenderTarget2DRWInstanceData_RenderThread
 	FRDGTextureUAVRef	TransientRDGUAV = nullptr;
 
 #if WITH_EDITORONLY_DATA
-	uint32 bPreviewTexture : 1;
+	bool bPreviewTexture = false;
 #endif
 #if STATS
 	void UpdateMemoryStats();
@@ -92,7 +79,6 @@ class UNiagaraDataInterfaceRenderTarget2D : public UNiagaraDataInterfaceRWBase
 
 public:
 	NIAGARA_API virtual void PostInitProperties() override;
-	virtual bool CanBeInCluster() const override { return false; }
 
 	//~ UNiagaraDataInterface interface
 	// VM functionality
@@ -179,7 +165,4 @@ protected:
 	static NIAGARA_API FNiagaraVariableBase ExposedRTVar;
 
 	TMap<FNiagaraSystemInstanceID, FRenderTarget2DRWInstanceData_GameThread*> SystemInstancesToProxyData_GT;
-	
-	UPROPERTY(Transient, DuplicateTransient)
-	TMap<uint64, TObjectPtr<UTextureRenderTarget2D>> ManagedRenderTargets;
 };
diff --git a/Engine/Plugins/FX/Niagara/Source/Niagara/Classes/NiagaraDataInterfaceRenderTarget2DArray.h b/Engine/Plugins/FX/Niagara/Source/Niagara/Classes/NiagaraDataInterfaceRenderTarget2DArray.h
--- a/Engine/Plugins/FX/Niagara/Source/Niagara/Classes/NiagaraDataInterfaceRenderTarget2DArray.h
+++ b/Engine/Plugins/FX/Niagara/Source/Niagara/Classes/NiagaraDataInterfaceRenderTarget2DArray.h
@@ -15,33 +15,20 @@ class UTextureRenderTarget2DArray;
 
 struct FRenderTarget2DArrayRWInstanceData_GameThread
 {
-	FRenderTarget2DArrayRWInstanceData_GameThread()
-	{
-#if WITH_EDITORONLY_DATA
-		bPreviewTexture = false;
-#endif
-	}
-
 	FIntVector Size = FIntVector(EForceInit::ForceInitToZero);
 	EPixelFormat Format = EPixelFormat::PF_A16B16G16R16;
 	TextureFilter Filter = TextureFilter::TF_Default;
 
+	bool bManagedTexture = false;
 	UTextureRenderTarget2DArray* TargetTexture = nullptr;
 #if WITH_EDITORONLY_DATA
-	uint32 bPreviewTexture : 1;
+	bool bPreviewTexture = false;
 #endif
 	FNiagaraParameterDirectBinding<UObject*> RTUserParamBinding;
 };
 
 struct FRenderTarget2DArrayRWInstanceData_RenderThread
 {
-	FRenderTarget2DArrayRWInstanceData_RenderThread()
-	{
-#if WITH_EDITORONLY_DATA
-		bPreviewTexture = false;
-#endif
-	}
-
 	FIntVector Size = FIntVector(EForceInit::ForceInitToZero);
 	int MipLevels = 0;
 	bool bWroteThisFrame = false;
@@ -55,7 +42,7 @@ struct FRenderTarget2DArrayRWInstanceData_RenderThread
 	FRDGTextureUAVRef	TransientRDGUAV = nullptr;
 
 #if WITH_EDITORONLY_DATA
-	uint32 bPreviewTexture : 1;
+	bool bPreviewTexture = false;
 #endif
 #if STATS
 	void UpdateMemoryStats();
@@ -88,7 +75,6 @@ class UNiagaraDataInterfaceRenderTarget2DArray : public UNiagaraDataInterfaceRWB
 
 public:
 	NIAGARA_API virtual void PostInitProperties() override;
-	virtual bool CanBeInCluster() const override { return false; }
 
 	//~ UNiagaraDataInterface interface
 	// VM functionality
@@ -162,7 +148,4 @@ public:
 
 protected:
 	static NIAGARA_API FNiagaraVariableBase ExposedRTVar;
-	
-	UPROPERTY(Transient, DuplicateTransient)
-	TMap<uint64, TObjectPtr<UTextureRenderTarget2DArray>> ManagedRenderTargets;
 };
diff --git a/Engine/Plugins/FX/Niagara/Source/Niagara/Classes/NiagaraDataInterfaceRenderTargetCube.h b/Engine/Plugins/FX/Niagara/Source/Niagara/Classes/NiagaraDataInterfaceRenderTargetCube.h
--- a/Engine/Plugins/FX/Niagara/Source/Niagara/Classes/NiagaraDataInterfaceRenderTargetCube.h
+++ b/Engine/Plugins/FX/Niagara/Source/Niagara/Classes/NiagaraDataInterfaceRenderTargetCube.h
@@ -14,33 +14,20 @@ class UTextureRenderTargetCube;
 
 struct FRenderTargetCubeRWInstanceData_GameThread
 {
-	FRenderTargetCubeRWInstanceData_GameThread()
-	{
-#if WITH_EDITORONLY_DATA
-		bPreviewTexture = false;
-#endif
-	}
-
 	int Size = 0;
 	EPixelFormat Format = EPixelFormat::PF_A16B16G16R16;
 	TextureFilter Filter = TextureFilter::TF_Default;
 
+	bool bManagedTexture = false;
 	UTextureRenderTargetCube* TargetTexture = nullptr;
 #if WITH_EDITORONLY_DATA
-	uint32 bPreviewTexture : 1;
+	bool bPreviewTexture = false;
 #endif
 	FNiagaraParameterDirectBinding<UObject*> RTUserParamBinding;
 };
 
 struct FRenderTargetCubeRWInstanceData_RenderThread
 {
-	FRenderTargetCubeRWInstanceData_RenderThread()
-	{
-#if WITH_EDITORONLY_DATA
-		bPreviewTexture = false;
-#endif
-	}
-
 	int Size = 0;
 	int MipLevels = 0;
 	bool bWroteThisFrame = false;
@@ -54,7 +41,7 @@ struct FRenderTargetCubeRWInstanceData_RenderThread
 	FRDGTextureSRVRef	TransientRDGSRV = nullptr;
 	FRDGTextureUAVRef	TransientRDGUAV = nullptr;
 #if WITH_EDITORONLY_DATA
-	uint32 bPreviewTexture : 1;
+	bool bPreviewTexture = false;
 #endif
 #if STATS
 	void UpdateMemoryStats();
@@ -81,7 +68,6 @@ class UNiagaraDataInterfaceRenderTargetCube : public UNiagaraDataInterfaceRWBase
 
 public:
 	NIAGARA_API virtual void PostInitProperties() override;
-	virtual bool CanBeInCluster() const override { return false; }
 
 	//~ UNiagaraDataInterface interface
 	// VM functionality
@@ -157,7 +143,4 @@ protected:
 	//~ UNiagaraDataInterface interface END
 
 	static NIAGARA_API FNiagaraVariableBase ExposedRTVar;
-	
-	UPROPERTY(Transient, DuplicateTransient)
-	TMap<uint64, TObjectPtr<UTextureRenderTargetCube>> ManagedRenderTargets;
 };
diff --git a/Engine/Plugins/FX/Niagara/Source/Niagara/Classes/NiagaraDataInterfaceRenderTargetVolume.h b/Engine/Plugins/FX/Niagara/Source/Niagara/Classes/NiagaraDataInterfaceRenderTargetVolume.h
--- a/Engine/Plugins/FX/Niagara/Source/Niagara/Classes/NiagaraDataInterfaceRenderTargetVolume.h
+++ b/Engine/Plugins/FX/Niagara/Source/Niagara/Classes/NiagaraDataInterfaceRenderTargetVolume.h
@@ -14,33 +14,20 @@ class UTextureRenderTargetVolume;
 
 struct FRenderTargetVolumeRWInstanceData_GameThread
 {
-	FRenderTargetVolumeRWInstanceData_GameThread()
-	{
-#if WITH_EDITORONLY_DATA
-		bPreviewTexture = false;
-#endif
-	}
-
 	FIntVector Size = FIntVector(EForceInit::ForceInitToZero);
 	EPixelFormat Format = EPixelFormat::PF_A16B16G16R16;
 	TextureFilter Filter = TextureFilter::TF_Default;
 
+	bool bManagedTexture = false;
 	UTextureRenderTargetVolume* TargetTexture = nullptr;
 #if WITH_EDITORONLY_DATA
-	uint32 bPreviewTexture : 1;
+	bool bPreviewTexture = false;
 #endif
 	FNiagaraParameterDirectBinding<UObject*> RTUserParamBinding;
 };
 
 struct FRenderTargetVolumeRWInstanceData_RenderThread
 {
-	FRenderTargetVolumeRWInstanceData_RenderThread()
-	{
-#if WITH_EDITORONLY_DATA
-		bPreviewTexture = false;
-#endif
-	}
-
 	FIntVector Size = FIntVector(EForceInit::ForceInitToZero);
 	int MipLevels = 0;
 	bool bWroteThisFrame = false;
@@ -54,7 +41,7 @@ struct FRenderTargetVolumeRWInstanceData_RenderThread
 	FRDGTextureUAVRef	TransientRDGUAV = nullptr;
 
 #if WITH_EDITORONLY_DATA
-	uint32 bPreviewTexture : 1;
+	bool bPreviewTexture = false;
 #endif
 #if STATS
 	void UpdateMemoryStats();
@@ -81,7 +68,6 @@ class UNiagaraDataInterfaceRenderTargetVolume : public UNiagaraDataInterfaceRWBa
 
 public:
 	NIAGARA_API virtual void PostInitProperties() override;
-	virtual bool CanBeInCluster() const override { return false; }
 
 	//~ UNiagaraDataInterface interface
 	// VM functionality
@@ -171,7 +157,4 @@ protected:
 	//~ UNiagaraDataInterface interface END
 
 	static NIAGARA_API FNiagaraVariableBase ExposedRTVar;
-	
-	UPROPERTY(Transient, DuplicateTransient)
-	TMap<uint64, TObjectPtr<UTextureRenderTargetVolume>> ManagedRenderTargets;
 };
diff --git a/Engine/Plugins/FX/Niagara/Source/Niagara/Private/NiagaraDataInterfaceIntRenderTarget2D.cpp b/Engine/Plugins/FX/Niagara/Source/Niagara/Private/NiagaraDataInterfaceIntRenderTarget2D.cpp
--- a/Engine/Plugins/FX/Niagara/Source/Niagara/Private/NiagaraDataInterfaceIntRenderTarget2D.cpp
+++ b/Engine/Plugins/FX/Niagara/Source/Niagara/Private/NiagaraDataInterfaceIntRenderTarget2D.cpp
@@ -73,6 +73,7 @@ struct FNDIIntRenderTarget2DInstanceData_GameThread
 	FVector2D PreviewDisplayRange = FVector2D(0.0f, 255.0f);
 #endif
 
+	bool bManagedTexture = false;
 	UTextureRenderTarget2D* TargetTexture = nullptr;
 	FNiagaraParameterDirectBinding<UObject*> RTUserParamBinding;
 };
@@ -656,6 +657,7 @@ bool UNiagaraDataInterfaceIntRenderTarget2D::InitPerInstanceData(void* PerInstan
 void UNiagaraDataInterfaceIntRenderTarget2D::DestroyPerInstanceData(void* PerInstanceData, FNiagaraSystemInstance* SystemInstance)
 {
 	FNDIIntRenderTarget2DInstanceData_GameThread* InstanceData = static_cast<FNDIIntRenderTarget2DInstanceData_GameThread*>(PerInstanceData);
+	NiagaraDataInterfaceRenderTargetCommon::ReleaseRenderTarget(SystemInstance, InstanceData);
 	InstanceData->~FNDIIntRenderTarget2DInstanceData_GameThread();
 
 	FNDIIntRenderTarget2DProxy* RT_Proxy = GetProxyAs<FNDIIntRenderTarget2DProxy>();
@@ -665,13 +667,6 @@ void UNiagaraDataInterfaceIntRenderTarget2D::DestroyPerInstanceData(void* PerIns
 			RT_Proxy->SystemInstancesToProxyData_RT.Remove(InstanceID);
 		}
 	);
-
-	// Make sure to clear out the reference to the render target if we created one.
-	decltype(ManagedRenderTargets)::ValueType ExistingRenderTarget = nullptr;
-	if ( ManagedRenderTargets.RemoveAndCopyValue(SystemInstance->GetId(), ExistingRenderTarget) && NiagaraDataInterfaceRenderTargetCommon::GReleaseResourceOnRemove )
-	{
-		ExistingRenderTarget->ReleaseResource();
-	}
 }
 
 void UNiagaraDataInterfaceIntRenderTarget2D::GetExposedVariables(TArray<FNiagaraVariableBase>& OutVariables) const
@@ -754,13 +749,8 @@ bool UNiagaraDataInterfaceIntRenderTarget2D::UpdateInstanceTexture(FNiagaraSyste
 			// If the texture has changed remove the old one and release if it's one we created
 			if (InstanceData->TargetTexture != UserTargetTexture)
 			{
+				NiagaraDataInterfaceRenderTargetCommon::ReleaseRenderTarget(SystemInstance, InstanceData);
 				InstanceData->TargetTexture = UserTargetTexture;
-
-				decltype(ManagedRenderTargets)::ValueType ExistingRenderTarget = nullptr;
-				if (ManagedRenderTargets.RemoveAndCopyValue(SystemInstance->GetId(), ExistingRenderTarget) && NiagaraDataInterfaceRenderTargetCommon::GReleaseResourceOnRemove)
-				{
-					ExistingRenderTarget->ReleaseResource();
-				}
 			}
 		}
 		else
@@ -773,7 +763,10 @@ bool UNiagaraDataInterfaceIntRenderTarget2D::UpdateInstanceTexture(FNiagaraSyste
 	bool bHasChanged = false;
 	if (InstanceData->TargetTexture == nullptr)
 	{
-		InstanceData->TargetTexture = NewObject<UTextureRenderTarget2D>(this);
+		if (NiagaraDataInterfaceRenderTargetCommon::CreateRenderTarget<UTextureRenderTarget2D>(SystemInstance, InstanceData) == false)
+		{
+			return false;
+		}
 		InstanceData->TargetTexture->bCanCreateUAV = true;
 		InstanceData->TargetTexture->bAutoGenerateMips = false;
 		InstanceData->TargetTexture->OverrideFormat = InstanceData->Format;
@@ -781,8 +774,6 @@ bool UNiagaraDataInterfaceIntRenderTarget2D::UpdateInstanceTexture(FNiagaraSyste
 		InstanceData->TargetTexture->InitAutoFormat(InstanceData->Size.X, InstanceData->Size.Y);
 		InstanceData->TargetTexture->UpdateResourceImmediate(true);
 
-		ManagedRenderTargets.Add(SystemInstance->GetId()) = InstanceData->TargetTexture;
-
 		bHasChanged = true;
 	}
 	// Do we need to update the existing texture?
diff --git a/Engine/Plugins/FX/Niagara/Source/Niagara/Private/NiagaraDataInterfaceRenderTarget2D.cpp b/Engine/Plugins/FX/Niagara/Source/Niagara/Private/NiagaraDataInterfaceRenderTarget2D.cpp
--- a/Engine/Plugins/FX/Niagara/Source/Niagara/Private/NiagaraDataInterfaceRenderTarget2D.cpp
+++ b/Engine/Plugins/FX/Niagara/Source/Niagara/Private/NiagaraDataInterfaceRenderTarget2D.cpp
@@ -480,11 +480,9 @@ bool UNiagaraDataInterfaceRenderTarget2D::InitPerInstanceData(void* PerInstanceD
 
 void UNiagaraDataInterfaceRenderTarget2D::DestroyPerInstanceData(void* PerInstanceData, FNiagaraSystemInstance* SystemInstance)
 {
-	SystemInstancesToProxyData_GT.Remove(SystemInstance->GetId());
-
-	FRenderTarget2DRWInstanceData_GameThread* InstanceData = static_cast<FRenderTarget2DRWInstanceData_GameThread*>(PerInstanceData);
+	using namespace NDIRenderTarget2DLocal;
 
-	InstanceData->~FRenderTarget2DRWInstanceData_GameThread();
+	SystemInstancesToProxyData_GT.Remove(SystemInstance->GetId());
 
 	FNiagaraDataInterfaceProxyRenderTarget2DProxy* RT_Proxy = GetProxyAs<FNiagaraDataInterfaceProxyRenderTarget2DProxy>();
 	ENQUEUE_RENDER_COMMAND(FNiagaraDIDestroyInstanceData) (
@@ -502,16 +500,11 @@ void UNiagaraDataInterfaceRenderTarget2D::DestroyPerInstanceData(void* PerInstan
 		}
 	);
 
-	// Make sure to clear out the reference to the render target if we created one.
-	using RenderTargetType = decltype(decltype(ManagedRenderTargets)::ElementType::Value);
-	RenderTargetType ExistingRenderTarget = nullptr;
-	if ( ManagedRenderTargets.RemoveAndCopyValue(SystemInstance->GetId(), ExistingRenderTarget) && NiagaraDataInterfaceRenderTargetCommon::GReleaseResourceOnRemove )
-	{
-		ExistingRenderTarget->ReleaseResource();
-	}
+	FRenderTarget2DRWInstanceData_GameThread* InstanceData = static_cast<FRenderTarget2DRWInstanceData_GameThread*>(PerInstanceData);
+	NiagaraDataInterfaceRenderTargetCommon::ReleaseRenderTarget(SystemInstance, InstanceData);
+	InstanceData->~FRenderTarget2DRWInstanceData_GameThread();
 }
 
-
 void UNiagaraDataInterfaceRenderTarget2D::GetExposedVariables(TArray<FNiagaraVariableBase>& OutVariables) const
 {
 	OutVariables.Emplace(ExposedRTVar);
@@ -646,19 +639,16 @@ void UNiagaraDataInterfaceRenderTarget2D::VMSetFormat(FVectorVMExternalFunctionC
 
 bool UNiagaraDataInterfaceRenderTarget2D::PerInstanceTick(void* PerInstanceData, FNiagaraSystemInstance* SystemInstance, float DeltaSeconds)
 {
+	using namespace NDIRenderTarget2DLocal;
+
 	FRenderTarget2DRWInstanceData_GameThread* InstanceData = static_cast<FRenderTarget2DRWInstanceData_GameThread*>(PerInstanceData);
 
 	// Pull from user parameter
 	UTextureRenderTarget2D* UserTargetTexture = InstanceData->RTUserParamBinding.GetValue<UTextureRenderTarget2D>();
 	if (UserTargetTexture && (InstanceData->TargetTexture != UserTargetTexture))
 	{
+		NiagaraDataInterfaceRenderTargetCommon::ReleaseRenderTarget(SystemInstance, InstanceData);
 		InstanceData->TargetTexture = UserTargetTexture;
-
-		TObjectPtr<UTextureRenderTarget2D> ExistingRenderTarget = nullptr;
-		if (ManagedRenderTargets.RemoveAndCopyValue(SystemInstance->GetId(), ExistingRenderTarget) && NiagaraDataInterfaceRenderTargetCommon::GReleaseResourceOnRemove)
-		{
-			ExistingRenderTarget->ReleaseResource();
-		}
 	}
 
 	// Do we inherit the texture parameters from the user supplied texture?
@@ -693,6 +683,8 @@ bool UNiagaraDataInterfaceRenderTarget2D::PerInstanceTick(void* PerInstanceData,
 
 bool UNiagaraDataInterfaceRenderTarget2D::PerInstanceTickPostSimulate(void* PerInstanceData, FNiagaraSystemInstance* SystemInstance, float DeltaSeconds)
 {
+	using namespace NDIRenderTarget2DLocal;
+
 	// Update InstanceData as the texture may have changed
 	FRenderTarget2DRWInstanceData_GameThread* InstanceData = static_cast<FRenderTarget2DRWInstanceData_GameThread*>(PerInstanceData);
 #if WITH_EDITORONLY_DATA
@@ -710,7 +702,10 @@ bool UNiagaraDataInterfaceRenderTarget2D::PerInstanceTickPostSimulate(void* PerI
 	// Do we need to create a new texture?
 	if (bValidGpuDataInterface && !bInheritUserParameterSettings && (InstanceData->TargetTexture == nullptr))
 	{
-		InstanceData->TargetTexture = NewObject<UTextureRenderTarget2D>(this);
+		if (NiagaraDataInterfaceRenderTargetCommon::CreateRenderTarget<UTextureRenderTarget2D>(SystemInstance, InstanceData) == false)
+		{
+			return false;
+		}
 		InstanceData->TargetTexture->bCanCreateUAV = true;
 		InstanceData->TargetTexture->bAutoGenerateMips = InstanceData->MipMapGeneration != ENiagaraMipMapGeneration::Disabled;
 		InstanceData->TargetTexture->RenderTargetFormat = InstanceData->Format;
@@ -718,8 +713,6 @@ bool UNiagaraDataInterfaceRenderTarget2D::PerInstanceTickPostSimulate(void* PerI
 		InstanceData->TargetTexture->Filter = InstanceData->Filter;
 		InstanceData->TargetTexture->InitAutoFormat(InstanceData->Size.X, InstanceData->Size.Y);
 		InstanceData->TargetTexture->UpdateResourceImmediate(true);
-
-		ManagedRenderTargets.Add(SystemInstance->GetId()) = InstanceData->TargetTexture;
 	}
 
 	// Do we need to update the existing texture?
diff --git a/Engine/Plugins/FX/Niagara/Source/Niagara/Private/NiagaraDataInterfaceRenderTarget2DArray.cpp b/Engine/Plugins/FX/Niagara/Source/Niagara/Private/NiagaraDataInterfaceRenderTarget2DArray.cpp
--- a/Engine/Plugins/FX/Niagara/Source/Niagara/Private/NiagaraDataInterfaceRenderTarget2DArray.cpp
+++ b/Engine/Plugins/FX/Niagara/Source/Niagara/Private/NiagaraDataInterfaceRenderTarget2DArray.cpp
@@ -476,8 +476,10 @@ bool UNiagaraDataInterfaceRenderTarget2DArray::InitPerInstanceData(void* PerInst
 
 void UNiagaraDataInterfaceRenderTarget2DArray::DestroyPerInstanceData(void* PerInstanceData, FNiagaraSystemInstance* SystemInstance)
 {
-	FRenderTarget2DArrayRWInstanceData_GameThread* InstanceData = static_cast<FRenderTarget2DArrayRWInstanceData_GameThread*>(PerInstanceData);
+	using namespace NDIRenderTarget2DArrayLocal;
 
+	FRenderTarget2DArrayRWInstanceData_GameThread* InstanceData = static_cast<FRenderTarget2DArrayRWInstanceData_GameThread*>(PerInstanceData);
+	NiagaraDataInterfaceRenderTargetCommon::ReleaseRenderTarget(SystemInstance, InstanceData);
 	InstanceData->~FRenderTarget2DArrayRWInstanceData_GameThread();
 
 	FNiagaraDataInterfaceProxyRenderTarget2DArrayProxy* RT_Proxy = GetProxyAs<FNiagaraDataInterfaceProxyRenderTarget2DArrayProxy>();
@@ -495,14 +497,6 @@ void UNiagaraDataInterfaceRenderTarget2DArray::DestroyPerInstanceData(void* PerI
 			RT_Proxy->SystemInstancesToProxyData_RT.Remove(InstanceID);
 		}
 	);
-
-	// Make sure to clear out the reference to the render target if we created one.
-	using RenderTargetType = decltype(decltype(ManagedRenderTargets)::ElementType::Value);
-	RenderTargetType ExistingRenderTarget = nullptr;
-	if (ManagedRenderTargets.RemoveAndCopyValue(SystemInstance->GetId(), ExistingRenderTarget) && NiagaraDataInterfaceRenderTargetCommon::GReleaseResourceOnRemove)
-	{
-		ExistingRenderTarget->ReleaseResource();
-	}
 }
 
 
@@ -611,19 +605,16 @@ void UNiagaraDataInterfaceRenderTarget2DArray::VMSetFormat(FVectorVMExternalFunc
 
 bool UNiagaraDataInterfaceRenderTarget2DArray::PerInstanceTick(void* PerInstanceData, FNiagaraSystemInstance* SystemInstance, float DeltaSeconds)
 {
+	using namespace NDIRenderTarget2DArrayLocal;
+
 	FRenderTarget2DArrayRWInstanceData_GameThread* InstanceData = static_cast<FRenderTarget2DArrayRWInstanceData_GameThread*>(PerInstanceData);
 
 	// Pull from user parameter
 	UTextureRenderTarget2DArray* UserTargetTexture = InstanceData->RTUserParamBinding.GetValue<UTextureRenderTarget2DArray>();
 	if (UserTargetTexture && (InstanceData->TargetTexture != UserTargetTexture))
 	{
+		NiagaraDataInterfaceRenderTargetCommon::ReleaseRenderTarget(SystemInstance, InstanceData);
 		InstanceData->TargetTexture = UserTargetTexture;
-
-		TObjectPtr<UTextureRenderTarget2DArray> ExistingRenderTarget = nullptr;
-		if (ManagedRenderTargets.RemoveAndCopyValue(SystemInstance->GetId(), ExistingRenderTarget) && NiagaraDataInterfaceRenderTargetCommon::GReleaseResourceOnRemove)
-		{
-			ExistingRenderTarget->ReleaseResource();
-		}
 	}
 
 	// Do we inherit the texture parameters from the user supplied texture?
@@ -657,6 +648,8 @@ bool UNiagaraDataInterfaceRenderTarget2DArray::PerInstanceTick(void* PerInstance
 
 bool UNiagaraDataInterfaceRenderTarget2DArray::PerInstanceTickPostSimulate(void* PerInstanceData, FNiagaraSystemInstance* SystemInstance, float DeltaSeconds)
 {
+	using namespace NDIRenderTarget2DArrayLocal;
+
 	// Update InstanceData as the texture may have changed
 	FRenderTarget2DArrayRWInstanceData_GameThread* InstanceData = static_cast<FRenderTarget2DArrayRWInstanceData_GameThread*>(PerInstanceData);
 #if WITH_EDITORONLY_DATA
@@ -673,15 +666,16 @@ bool UNiagaraDataInterfaceRenderTarget2DArray::PerInstanceTickPostSimulate(void*
 
 	if (bValidGpuDataInterface && !bInheritUserParameterSettings && (InstanceData->TargetTexture == nullptr))
 	{
-		InstanceData->TargetTexture = NewObject<UTextureRenderTarget2DArray>(this);
+		if (NiagaraDataInterfaceRenderTargetCommon::CreateRenderTarget<UTextureRenderTarget2DArray>(SystemInstance, InstanceData) == false)
+		{
+			return false;
+		}
 		InstanceData->TargetTexture->bCanCreateUAV = true;
 		//InstanceData->TargetTexture->bAutoGenerateMips = InstanceData->MipMapGeneration != ENiagaraMipMapGeneration::Disabled;
 		InstanceData->TargetTexture->ClearColor = FLinearColor(0.0, 0, 0, 0);
 		InstanceData->TargetTexture->Filter = InstanceData->Filter;
 		InstanceData->TargetTexture->Init(InstanceData->Size.X, InstanceData->Size.Y, InstanceData->Size.Z, InstanceData->Format);
 		InstanceData->TargetTexture->UpdateResourceImmediate(true);
-
-		ManagedRenderTargets.Add(SystemInstance->GetId()) = InstanceData->TargetTexture;
 	}
 
 	// Do we need to update the existing texture?
diff --git a/Engine/Plugins/FX/Niagara/Source/Niagara/Private/NiagaraDataInterfaceRenderTargetCommon.cpp b/Engine/Plugins/FX/Niagara/Source/Niagara/Private/NiagaraDataInterfaceRenderTargetCommon.cpp
--- a/Engine/Plugins/FX/Niagara/Source/Niagara/Private/NiagaraDataInterfaceRenderTargetCommon.cpp
+++ b/Engine/Plugins/FX/Niagara/Source/Niagara/Private/NiagaraDataInterfaceRenderTargetCommon.cpp
@@ -8,14 +8,6 @@
 
 namespace NiagaraDataInterfaceRenderTargetCommon
 {
-	int32 GReleaseResourceOnRemove = false;
-	static FAutoConsoleVariableRef CVarNiagaraReleaseResourceOnRemove(
-		TEXT("fx.Niagara.RenderTarget.ReleaseResourceOnRemove"),
-		GReleaseResourceOnRemove,
-		TEXT("Releases the render target resource once it is removed from the manager list rather than waiting for a GC."),
-		ECVF_Default
-	);
-
 	//-TEMP: Until we prune data interface on cook this will avoid consuming memory
 	int32 GIgnoreCookedOut = true;
 	static FAutoConsoleVariableRef CVarNiagaraRenderTargetIgnoreCookedOut(
diff --git a/Engine/Plugins/FX/Niagara/Source/Niagara/Private/NiagaraDataInterfaceRenderTargetCommon.h b/Engine/Plugins/FX/Niagara/Source/Niagara/Private/NiagaraDataInterfaceRenderTargetCommon.h
--- a/Engine/Plugins/FX/Niagara/Source/Niagara/Private/NiagaraDataInterfaceRenderTargetCommon.h
+++ b/Engine/Plugins/FX/Niagara/Source/Niagara/Private/NiagaraDataInterfaceRenderTargetCommon.h
@@ -3,14 +3,46 @@
 #pragma once
 
 #include "HAL/Platform.h"
+#include "NiagaraSystemInstance.h"
+#include "NiagaraWorldManager.h"
 
 enum ETextureRenderTargetFormat : int;
 
 namespace NiagaraDataInterfaceRenderTargetCommon
 {
-	extern int32 GReleaseResourceOnRemove;
 	extern int32 GIgnoreCookedOut;
 	extern float GResolutionMultiplier;
 
 	extern bool GetRenderTargetFormat(bool bOverrideFormat, ETextureRenderTargetFormat OverrideFormat, ETextureRenderTargetFormat& OutRenderTargetFormat);
+
+	template<typename TRenderTargetType, typename TInstanceDataType>
+	static bool CreateRenderTarget(FNiagaraSystemInstance* SystemInstance, TInstanceDataType* InstanceData)
+	{
+		if (FNiagaraWorldManager* WorldManager = SystemInstance->GetWorldManager())
+		{
+			InstanceData->bManagedTexture = true;
+			InstanceData->TargetTexture = WorldManager->ObjectPoolGetOrCreate<TRenderTargetType>();
+			if (InstanceData->TargetTexture != nullptr)
+			{
+				WorldManager->AddReferencedObject(InstanceData->TargetTexture);
+				return true;
+			}
+		}
+		return false;
+	}
+
+	template<typename TInstanceDataType>
+	static void ReleaseRenderTarget(FNiagaraSystemInstance* SystemInstance, TInstanceDataType* InstanceData)
+	{
+		if (InstanceData->bManagedTexture && InstanceData->TargetTexture)
+		{
+			if (FNiagaraWorldManager* WorldManager = SystemInstance->GetWorldManager())
+			{
+				WorldManager->ObjectPoolReturn(InstanceData->TargetTexture);
+				WorldManager->RemoveReferencedObject(InstanceData->TargetTexture);
+				InstanceData->TargetTexture = nullptr;
+				InstanceData->bManagedTexture = false;
+			}
+		}
+	}
 }
diff --git a/Engine/Plugins/FX/Niagara/Source/Niagara/Private/NiagaraDataInterfaceRenderTargetCube.cpp b/Engine/Plugins/FX/Niagara/Source/Niagara/Private/NiagaraDataInterfaceRenderTargetCube.cpp
--- a/Engine/Plugins/FX/Niagara/Source/Niagara/Private/NiagaraDataInterfaceRenderTargetCube.cpp
+++ b/Engine/Plugins/FX/Niagara/Source/Niagara/Private/NiagaraDataInterfaceRenderTargetCube.cpp
@@ -439,8 +439,10 @@ bool UNiagaraDataInterfaceRenderTargetCube::InitPerInstanceData(void* PerInstanc
 
 void UNiagaraDataInterfaceRenderTargetCube::DestroyPerInstanceData(void* PerInstanceData, FNiagaraSystemInstance* SystemInstance)
 {
-	FRenderTargetCubeRWInstanceData_GameThread* InstanceData = static_cast<FRenderTargetCubeRWInstanceData_GameThread*>(PerInstanceData);
+	using namespace NDIRenderTargetCubeLocal;
 
+	FRenderTargetCubeRWInstanceData_GameThread* InstanceData = static_cast<FRenderTargetCubeRWInstanceData_GameThread*>(PerInstanceData);
+	NiagaraDataInterfaceRenderTargetCommon::ReleaseRenderTarget(SystemInstance, InstanceData);
 	InstanceData->~FRenderTargetCubeRWInstanceData_GameThread();
 
 	FNiagaraDataInterfaceProxyRenderTargetCubeProxy* RT_Proxy = GetProxyAs<FNiagaraDataInterfaceProxyRenderTargetCubeProxy>();
@@ -458,13 +460,6 @@ void UNiagaraDataInterfaceRenderTargetCube::DestroyPerInstanceData(void* PerInst
 			RT_Proxy->SystemInstancesToProxyData_RT.Remove(InstanceID);
 		}
 	);
-
-	// Make sure to clear out the reference to the render target if we created one.
-	decltype(ManagedRenderTargets)::ValueType ExistingRenderTarget = nullptr;
-	if (ManagedRenderTargets.RemoveAndCopyValue(SystemInstance->GetId(), ExistingRenderTarget) && NiagaraDataInterfaceRenderTargetCommon::GReleaseResourceOnRemove)
-	{
-		ExistingRenderTarget->ReleaseResource();
-	}
 }
 
 
@@ -563,19 +558,16 @@ void UNiagaraDataInterfaceRenderTargetCube::VMSetFormat(FVectorVMExternalFunctio
 
 bool UNiagaraDataInterfaceRenderTargetCube::PerInstanceTick(void* PerInstanceData, FNiagaraSystemInstance* SystemInstance, float DeltaSeconds)
 {
+	using namespace NDIRenderTargetCubeLocal;
+
 	FRenderTargetCubeRWInstanceData_GameThread* InstanceData = static_cast<FRenderTargetCubeRWInstanceData_GameThread*>(PerInstanceData);
 
 	// Pull from user parameter
 	UTextureRenderTargetCube* UserTargetTexture = InstanceData->RTUserParamBinding.GetValue<UTextureRenderTargetCube>();
 	if (UserTargetTexture && (InstanceData->TargetTexture != UserTargetTexture))
 	{
+		NiagaraDataInterfaceRenderTargetCommon::ReleaseRenderTarget(SystemInstance, InstanceData);
 		InstanceData->TargetTexture = UserTargetTexture;
-
-		decltype(ManagedRenderTargets)::ValueType ExistingRenderTarget = nullptr;
-		if (ManagedRenderTargets.RemoveAndCopyValue(SystemInstance->GetId(), ExistingRenderTarget) && NiagaraDataInterfaceRenderTargetCommon::GReleaseResourceOnRemove)
-		{
-			ExistingRenderTarget->ReleaseResource();
-		}
 	}
 
 	// Do we inherit the texture parameters from the user supplied texture?
@@ -607,6 +599,8 @@ bool UNiagaraDataInterfaceRenderTargetCube::PerInstanceTick(void* PerInstanceDat
 
 bool UNiagaraDataInterfaceRenderTargetCube::PerInstanceTickPostSimulate(void* PerInstanceData, FNiagaraSystemInstance* SystemInstance, float DeltaSeconds)
 {
+	using namespace NDIRenderTargetCubeLocal;
+
 	// Update InstanceData as the texture may have changed
 	FRenderTargetCubeRWInstanceData_GameThread* InstanceData = static_cast<FRenderTargetCubeRWInstanceData_GameThread*>(PerInstanceData);
 #if WITH_EDITORONLY_DATA
@@ -624,15 +618,16 @@ bool UNiagaraDataInterfaceRenderTargetCube::PerInstanceTickPostSimulate(void* Pe
 	// Do we need to create a new texture?
 	if (bValidGpuDataInterface && !bInheritUserParameterSettings && (InstanceData->TargetTexture == nullptr))
 	{
-		InstanceData->TargetTexture = NewObject<UTextureRenderTargetCube>(this);
+		if (NiagaraDataInterfaceRenderTargetCommon::CreateRenderTarget<UTextureRenderTargetCube>(SystemInstance, InstanceData) == false)
+		{
+			return false;
+		}
 		InstanceData->TargetTexture->bCanCreateUAV = true;
 		//InstanceData->TargetTexture->bAutoGenerateMips = InstanceData->MipMapGeneration != ENiagaraMipMapGeneration::Disabled;
 		InstanceData->TargetTexture->ClearColor = FLinearColor(0.0, 0, 0, 0);
 		InstanceData->TargetTexture->Filter = InstanceData->Filter;
 		InstanceData->TargetTexture->Init(InstanceData->Size, InstanceData->Format);
 		InstanceData->TargetTexture->UpdateResourceImmediate(true);
-
-		ManagedRenderTargets.Add(SystemInstance->GetId()) = InstanceData->TargetTexture;
 	}
 
 	// Do we need to update the existing texture?
diff --git a/Engine/Plugins/FX/Niagara/Source/Niagara/Private/NiagaraDataInterfaceRenderTargetVolume.cpp b/Engine/Plugins/FX/Niagara/Source/Niagara/Private/NiagaraDataInterfaceRenderTargetVolume.cpp
--- a/Engine/Plugins/FX/Niagara/Source/Niagara/Private/NiagaraDataInterfaceRenderTargetVolume.cpp
+++ b/Engine/Plugins/FX/Niagara/Source/Niagara/Private/NiagaraDataInterfaceRenderTargetVolume.cpp
@@ -538,8 +538,10 @@ bool UNiagaraDataInterfaceRenderTargetVolume::InitPerInstanceData(void* PerInsta
 
 void UNiagaraDataInterfaceRenderTargetVolume::DestroyPerInstanceData(void* PerInstanceData, FNiagaraSystemInstance* SystemInstance)
 {
-	FRenderTargetVolumeRWInstanceData_GameThread* InstanceData = static_cast<FRenderTargetVolumeRWInstanceData_GameThread*>(PerInstanceData);
+	using namespace NDIRenderTargetVolumeLocal;
 
+	FRenderTargetVolumeRWInstanceData_GameThread* InstanceData = static_cast<FRenderTargetVolumeRWInstanceData_GameThread*>(PerInstanceData);
+	NiagaraDataInterfaceRenderTargetCommon::ReleaseRenderTarget(SystemInstance, InstanceData);
 	InstanceData->~FRenderTargetVolumeRWInstanceData_GameThread();
 
 	FNiagaraDataInterfaceProxyRenderTargetVolumeProxy* RT_Proxy = GetProxyAs<FNiagaraDataInterfaceProxyRenderTargetVolumeProxy>();
@@ -557,14 +559,6 @@ void UNiagaraDataInterfaceRenderTargetVolume::DestroyPerInstanceData(void* PerIn
 			RT_Proxy->SystemInstancesToProxyData_RT.Remove(InstanceID);
 		}
 	);
-
-	// Make sure to clear out the reference to the render target if we created one.
-	using RenderTargetType = decltype(decltype(ManagedRenderTargets)::ElementType::Value);
-	RenderTargetType ExistingRenderTarget = nullptr;
-	if (ManagedRenderTargets.RemoveAndCopyValue(SystemInstance->GetId(), ExistingRenderTarget) && NiagaraDataInterfaceRenderTargetCommon::GReleaseResourceOnRemove)
-	{
-		ExistingRenderTarget->ReleaseResource();
-	}
 }
 
 
@@ -1055,19 +1049,16 @@ void UNiagaraDataInterfaceRenderTargetVolume::VMSetFormat(FVectorVMExternalFunct
 
 bool UNiagaraDataInterfaceRenderTargetVolume::PerInstanceTick(void* PerInstanceData, FNiagaraSystemInstance* SystemInstance, float DeltaSeconds)
 {
+	using namespace NDIRenderTargetVolumeLocal;
+
 	FRenderTargetVolumeRWInstanceData_GameThread* InstanceData = static_cast<FRenderTargetVolumeRWInstanceData_GameThread*>(PerInstanceData);
 
 	// Pull from user parameter
 	UTextureRenderTargetVolume* UserTargetTexture = InstanceData->RTUserParamBinding.GetValue<UTextureRenderTargetVolume>();
 	if (UserTargetTexture && (InstanceData->TargetTexture != UserTargetTexture))
 	{
+		NiagaraDataInterfaceRenderTargetCommon::ReleaseRenderTarget(SystemInstance, InstanceData);
 		InstanceData->TargetTexture = UserTargetTexture;
-
-		TObjectPtr<UTextureRenderTargetVolume> ExistingRenderTarget = nullptr;
-		if (ManagedRenderTargets.RemoveAndCopyValue(SystemInstance->GetId(), ExistingRenderTarget) && NiagaraDataInterfaceRenderTargetCommon::GReleaseResourceOnRemove)
-		{
-			ExistingRenderTarget->ReleaseResource();
-		}
 	}
 
 	// Do we inherit the texture parameters from the user supplied texture?
@@ -1101,6 +1092,8 @@ bool UNiagaraDataInterfaceRenderTargetVolume::PerInstanceTick(void* PerInstanceD
 
 bool UNiagaraDataInterfaceRenderTargetVolume::PerInstanceTickPostSimulate(void* PerInstanceData, FNiagaraSystemInstance* SystemInstance, float DeltaSeconds)
 {
+	using namespace NDIRenderTargetVolumeLocal;
+
 	// Update InstanceData as the texture may have changed
 	FRenderTargetVolumeRWInstanceData_GameThread* InstanceData = static_cast<FRenderTargetVolumeRWInstanceData_GameThread*>(PerInstanceData);
 #if WITH_EDITORONLY_DATA
@@ -1118,15 +1111,16 @@ bool UNiagaraDataInterfaceRenderTargetVolume::PerInstanceTickPostSimulate(void*
 	// Do we need to create a new texture?
 	if (bValidGpuDataInterface && !bInheritUserParameterSettings && (InstanceData->TargetTexture == nullptr))
 	{
-		InstanceData->TargetTexture = NewObject<UTextureRenderTargetVolume>(this);
+		if (NiagaraDataInterfaceRenderTargetCommon::CreateRenderTarget<UTextureRenderTargetVolume>(SystemInstance, InstanceData) == false)
+		{
+			return false;
+		}
 		InstanceData->TargetTexture->bCanCreateUAV = true;
 		//InstanceData->TargetTexture->bAutoGenerateMips = InstanceData->MipMapGeneration != ENiagaraMipMapGeneration::Disabled;
 		InstanceData->TargetTexture->ClearColor = FLinearColor(0.0, 0, 0, 0);
 		InstanceData->TargetTexture->Filter = InstanceData->Filter;
 		InstanceData->TargetTexture->Init(InstanceData->Size.X, InstanceData->Size.Y, InstanceData->Size.Z, InstanceData->Format);
 		InstanceData->TargetTexture->UpdateResourceImmediate(true);
-
-		ManagedRenderTargets.Add(SystemInstance->GetId()) = InstanceData->TargetTexture;
 	}
 
 	// Do we need to update the existing texture?
diff --git a/Engine/Plugins/FX/Niagara/Source/Niagara/Private/NiagaraSimpleObjectPool.h b/Engine/Plugins/FX/Niagara/Source/Niagara/Private/NiagaraSimpleObjectPool.h
new file mode 100644
--- /dev/null
+++ b/Engine/Plugins/FX/Niagara/Source/Niagara/Private/NiagaraSimpleObjectPool.h
@@ -0,0 +1,52 @@
+// Copyright Epic Games, Inc. All Rights Reserved.
+
+#pragma once
+
+#include "CoreMinimal.h"
+
+/**
+ * Very simple object pool that will hold weak references to objects that are returned, these can then be recycled rather than creating new objects.
+ * The objects are not strongly referenced and will be lost when garbage collection is run, this is by design.
+ */
+class FNiagaraSimpleObjectPool
+{
+public:
+	UObject* GetOrCreateObject(UClass* ObjectClass, UObject* Outer, bool& bIsExistingObject)
+	{
+		check(IsInGameThread());
+		check(ObjectClass);
+
+		const FName ClassName = ObjectClass->GetFName();
+		if (TArray<TWeakObjectPtr<UObject>>* ObjectPool = ObjectPools.Find(ClassName))
+		{
+			while (ObjectPool->Num() > 0)
+			{
+				if (UObject* Obj = ObjectPool->Pop().Get())
+				{
+					bIsExistingObject = true;
+					return Obj;
+				}
+			}
+		}
+
+		bIsExistingObject = false;
+		return NewObject<UObject>(Outer, ObjectClass);
+	}
+
+	void ReturnObject(UObject* Obj)
+	{
+		check(IsInGameThread());
+		check(::IsValid(Obj));
+
+		const FName ClassName = Obj->GetClass()->GetFName();
+		ObjectPools.FindOrAdd(ClassName).Add(Obj);
+	}
+
+	void PreGarbageCollect()
+	{
+		ObjectPools.Empty();
+	}
+
+private:
+	TMap<FName, TArray<TWeakObjectPtr<UObject>>> ObjectPools;
+};
diff --git a/Engine/Plugins/FX/Niagara/Source/Niagara/Private/NiagaraWorldManager.cpp b/Engine/Plugins/FX/Niagara/Source/Niagara/Private/NiagaraWorldManager.cpp
--- a/Engine/Plugins/FX/Niagara/Source/Niagara/Private/NiagaraWorldManager.cpp
+++ b/Engine/Plugins/FX/Niagara/Source/Niagara/Private/NiagaraWorldManager.cpp
@@ -18,6 +18,7 @@
 #include "NiagaraEffectType.h"
 #include "NiagaraDebugHud.h"
 #include "NiagaraGpuComputeDispatchInterface.h"
+#include "NiagaraSimpleObjectPool.h"
 #include "Particles/FXBudget.h"
 #include "NiagaraCullProxyComponent.h"
 #include "GameDelegates.h"
@@ -47,6 +48,14 @@ static FAutoConsoleVariableRef CVarNiagaraWorldManagerKillUniqueSims(
 	ECVF_Default
 );
 
+static bool GNiagaraWorldManagerObjectPoolEnabled = 1;
+static FAutoConsoleVariableRef CVarNiagaraWorldManagerObjectPoolEnabled(
+	TEXT("fx.Niagara.WorldManager.ObjectPoolEnabled"),
+	GNiagaraWorldManagerObjectPoolEnabled,
+	TEXT("Should we pool objects, yay or nay."),
+	ECVF_Default
+);
+
 static int GNiagaraAllowAsyncWorkToEndOfFrame = 1;
 static FAutoConsoleVariableRef CVarNiagaraAllowAsyncWorkToEndOfFrame(
 	TEXT("fx.Niagara.AllowAsyncWorkToEndOfFrame"),
@@ -472,6 +481,8 @@ void FNiagaraWorldManager::AddReferencedObjects(FReferenceCollector& Collector)
 
 	Collector.AddReferencedObjects(CullProxyMap);
 
+	Collector.AddReferencedObjects(ReferencedObjects);
+
 	DataChannelManager->AddReferencedObjects(Collector);
 }
 
@@ -727,6 +738,12 @@ void FNiagaraWorldManager::PreGarbageCollect()
 			It.RemoveCurrent();
 		}
 	}
+
+	// Clear out the pooled objects
+	if ( ObjectPool )
+	{
+		ObjectPool->PreGarbageCollect();
+	}
 }
 
 void FNiagaraWorldManager::PostReachabilityAnalysis()
@@ -2009,6 +2026,43 @@ void FNiagaraWorldManager::EnqueueGlobalDeferredCallback(TFunction<void()>&& Cal
 	GlobalDeferredCallbacks.Enqueue(Callback);
 }
 
+UObject* FNiagaraWorldManager::ObjectPoolGetOrCreate(UClass* Class, bool& bIsExistingObject)
+{
+	check(IsInGameThread());
+	check(Class != nullptr);
+
+	bIsExistingObject = false;
+	if (GNiagaraWorldManagerObjectPoolEnabled == false)
+	{
+		return NewObject<UObject>(World, Class);
+	}
+
+	if (ObjectPool == nullptr)
+	{
+		ObjectPool.Reset(new FNiagaraSimpleObjectPool());
+	}
+
+	return ObjectPool->GetOrCreateObject(Class, World, bIsExistingObject);
+}
+
+UObject* FNiagaraWorldManager::ObjectPoolGetOrCreate(UClass* Class)
+{
+	bool bIsExistingObject;
+	return ObjectPoolGetOrCreate(Class, bIsExistingObject);
+}
+
+void FNiagaraWorldManager::ObjectPoolReturn(UObject* Obj)
+{
+	check(IsInGameThread());
+	if (::IsValid(Obj) && GNiagaraWorldManagerObjectPoolEnabled)
+	{
+		if (ensure(ObjectPool) && ensure(Obj->GetOuter() == World))
+		{
+			ObjectPool->ReturnObject(Obj);
+		}
+	}
+}
+
 bool FNiagaraWorldManager::IsComponentLocalPlayerLinked(const USceneComponent* InComponent)
 {
 	check(InComponent);
diff --git a/Engine/Plugins/FX/Niagara/Source/Niagara/Public/NiagaraWorldManager.h b/Engine/Plugins/FX/Niagara/Source/Niagara/Public/NiagaraWorldManager.h
--- a/Engine/Plugins/FX/Niagara/Source/Niagara/Public/NiagaraWorldManager.h
+++ b/Engine/Plugins/FX/Niagara/Source/Niagara/Public/NiagaraWorldManager.h
@@ -25,6 +25,7 @@
 #include "NiagaraWorldManager.generated.h"
 
 class FNiagaraDataChannelManager;
+class FNiagaraSimpleObjectPool;
 class UWorld;
 class UNiagaraParameterCollection;
 class UNiagaraParameterCollectionInstance;
@@ -229,6 +230,39 @@ public:
 	*/
 	NIAGARA_API static void EnqueueGlobalDeferredCallback(TFunction<void()>&& Callback);
 
+	/**
+	* Will create a new object or return one from the pool.
+	* Existing objects will not be unregistered / reset / sanitized, they come out exactly as they were put into the pool.
+	* You must hold a valid reference to the object or it will be GCed.
+	* New objects are outered to the world managers world.
+	* EXPERIMENTAL: This API is intended for internal use currently, it is subject to change
+	*/
+	NIAGARA_API UObject* ObjectPoolGetOrCreate(UClass* Class, bool& bIsExistingObject);
+	NIAGARA_API UObject* ObjectPoolGetOrCreate(UClass* Class);
+	
+	/**
+	* Returns an object to the pool
+	* You are expected to do any unregistering of the component and ensuring no references are held.
+	* EXPERIMENTAL: This API is intended for internal use currently, it is subject to change
+	*/
+	NIAGARA_API void ObjectPoolReturn(UObject* Obj);
+
+	template<class T>
+	T* ObjectPoolGetOrCreate() { return CastChecked<T>(ObjectPoolGetOrCreate(T::StaticClass())); }
+	template<class T>
+	T* ObjectPoolGetOrCreate(bool& bIsExistingObject) { return CastChecked<T>(ObjectPoolGetOrCreate(T::StaticClass(), bIsExistingObject)); }
+
+	/**
+	* Objects added here will have a strong reference held.
+	* EXPERIMENTAL: This API is intended for internal use currently, it is subject to change
+	*/
+	void AddReferencedObject(UObject* InObject) { check(!ReferencedObjects.Contains(InObject)); ReferencedObjects.Add(InObject); }
+	/**
+	* Removes an object from the WM references.
+	* EXPERIMENTAL: This API is intended for internal use currently, it is subject to change
+	*/
+	void RemoveReferencedObject(UObject* InObject) { check(ReferencedObjects.Contains(InObject)); ReferencedObjects.RemoveSwap(InObject); }
+
 	/** Is this component in anyway linked to the local player. */
 	static bool IsComponentLocalPlayerLinked(const USceneComponent* Component);
 
@@ -358,11 +392,14 @@ private:
 #endif
 
 	TMap<TObjectPtr<UNiagaraSystem>, TObjectPtr<UNiagaraCullProxyComponent>> CullProxyMap;
+	TArray<TObjectPtr<UObject>> ReferencedObjects;
 
 	/** A global flag for all scalability culling */
 	static ENiagaraScalabilityCullingMode ScalabilityCullingMode;
 
 	TUniquePtr<FNiagaraDataChannelManager> DataChannelManager;
+
+	TUniquePtr<FNiagaraSimpleObjectPool> ObjectPool;
 };
 
 
