- Change ensures about RT RW in a stage to on screen message as it's too hidden from the end user

diff --git a/Engine/Plugins/FX/Niagara/Source/Niagara/Classes/NiagaraDataInterface.h b/Engine/Plugins/FX/Niagara/Source/Niagara/Classes/NiagaraDataInterface.h
--- a/Engine/Plugins/FX/Niagara/Source/Niagara/Classes/NiagaraDataInterface.h
+++ b/Engine/Plugins/FX/Niagara/Source/Niagara/Classes/NiagaraDataInterface.h
@@ -472,6 +472,11 @@ struct FNiagaraDataInterfaceSetShaderParametersContext
 		ShaderStorage = InShaderStorage;
 	}
 
+#if WITH_NIAGARA_DEBUG_EMITTER_NAME
+	/** Formats a string with the simulation & stage name */
+	NIAGARA_API FString GetDebugString() const;
+#endif
+
 private:
 	NIAGARA_API bool IsStructBoundInternal(const void* StructAddress, uint32 StructSize) const;
 	NIAGARA_API uint16 GetParameterIncludedStructInternal(const FShaderParametersMetadata* StructMetadata) const;
diff --git a/Engine/Plugins/FX/Niagara/Source/Niagara/Private/NiagaraDataInterface.cpp b/Engine/Plugins/FX/Niagara/Source/Niagara/Private/NiagaraDataInterface.cpp
--- a/Engine/Plugins/FX/Niagara/Source/Niagara/Private/NiagaraDataInterface.cpp
+++ b/Engine/Plugins/FX/Niagara/Source/Niagara/Private/NiagaraDataInterface.cpp
@@ -462,5 +462,19 @@ bool FNiagaraDataInterfaceSetShaderParametersContext::IsIterationStage() const
 	return ComputeInstanceData.IsIterationStage(DataInterfaceProxy, SimStageData.StageIndex);
 }
 
+#if WITH_NIAGARA_DEBUG_EMITTER_NAME
+FString FNiagaraDataInterfaceSetShaderParametersContext::GetDebugString() const
+{
+	TStringBuilder<128> Builder;
+	Builder.Append(ComputeInstanceData.Context ? ComputeInstanceData.Context->GetDebugSimName() : TEXT("Unknown"));
+	if (SimStageData.StageMetaData)
+	{
+		Builder.AppendChar(TEXT(','));
+		SimStageData.StageMetaData->SimulationStageName.AppendString(Builder);
+	}
+	return Builder.ToString();
+}
+#endif
+
 #undef LOCTEXT_NAMESPACE
 
diff --git a/Engine/Plugins/FX/Niagara/Source/Niagara/Private/NiagaraDataInterfaceRenderTarget2D.cpp b/Engine/Plugins/FX/Niagara/Source/Niagara/Private/NiagaraDataInterfaceRenderTarget2D.cpp
--- a/Engine/Plugins/FX/Niagara/Source/Niagara/Private/NiagaraDataInterfaceRenderTarget2D.cpp
+++ b/Engine/Plugins/FX/Niagara/Source/Niagara/Private/NiagaraDataInterfaceRenderTarget2D.cpp
@@ -1,5 +1,6 @@
 // Copyright Epic Games, Inc. All Rights Reserved.
 #include "NiagaraDataInterfaceRenderTarget2D.h"
+#include "Engine/Engine.h"
 #include "NiagaraCompileHashVisitor.h"
 #include "TextureResource.h"
 #include "CanvasTypes.h"
@@ -421,7 +422,6 @@ void UNiagaraDataInterfaceRenderTarget2D::SetShaderParameters(const FNiagaraData
 	
 	if (bRTRead)
 	{
-		ensureMsgf(bRTWrite == false, TEXT("RenderTarget DataInterface is both wrote and read from in the same stage, this is not allowed, read will be invalid"));
 		if (bRTWrite == false && InstanceData_RT->TransientRDGSRV)
 		{
 			InstanceData_RT->bReadThisFrame = true;
@@ -429,6 +429,9 @@ void UNiagaraDataInterfaceRenderTarget2D::SetShaderParameters(const FNiagaraData
 		}
 		else
 		{
+		#if WITH_NIAGARA_DEBUG_EMITTER_NAME
+			GEngine->AddOnScreenDebugMessage(uint64(this), 1.f, FColor::White, *FString::Printf(TEXT("RenderTarget is read and wrote in the same stage, this is not allowed, read will be invalid. (%s)"), *Context.GetDebugString()));
+		#endif
 			Parameters->Texture = Context.GetComputeDispatchInterface().GetBlackTextureSRV(GraphBuilder, ETextureDimension::Texture2D);
 		}
 
diff --git a/Engine/Plugins/FX/Niagara/Source/Niagara/Private/NiagaraDataInterfaceRenderTarget2DArray.cpp b/Engine/Plugins/FX/Niagara/Source/Niagara/Private/NiagaraDataInterfaceRenderTarget2DArray.cpp
--- a/Engine/Plugins/FX/Niagara/Source/Niagara/Private/NiagaraDataInterfaceRenderTarget2DArray.cpp
+++ b/Engine/Plugins/FX/Niagara/Source/Niagara/Private/NiagaraDataInterfaceRenderTarget2DArray.cpp
@@ -1,5 +1,6 @@
 // Copyright Epic Games, Inc. All Rights Reserved.
 #include "NiagaraDataInterfaceRenderTarget2DArray.h"
+#include "Engine/Engine.h"
 #include "Engine/TextureRenderTarget2D.h"
 #include "TextureResource.h"
 #include "NiagaraCompileHashVisitor.h"
@@ -419,7 +420,6 @@ void UNiagaraDataInterfaceRenderTarget2DArray::SetShaderParameters(const FNiagar
 	
 	if (bRTRead)
 	{
-		ensureMsgf(bRTWrite == false, TEXT("RenderTarget DataInterface is both wrote and read from in the same stage, this is not allowed, read will be invalid"));
 		if (bRTWrite == false && InstanceData_RT->RenderTarget.IsValid())
 		{
 			InstanceData_RT->bReadThisFrame = true;
@@ -427,6 +427,9 @@ void UNiagaraDataInterfaceRenderTarget2DArray::SetShaderParameters(const FNiagar
 		}
 		else
 		{
+		#if WITH_NIAGARA_DEBUG_EMITTER_NAME
+			GEngine->AddOnScreenDebugMessage(uint64(this), 1.f, FColor::White, *FString::Printf(TEXT("RenderTarget is read and wrote in the same stage, this is not allowed, read will be invalid. (%s)"), *Context.GetDebugString()));
+		#endif
 			Parameters->Texture = Context.GetComputeDispatchInterface().GetBlackTextureSRV(GraphBuilder, ETextureDimension::Texture2DArray);
 		}
 
diff --git a/Engine/Plugins/FX/Niagara/Source/Niagara/Private/NiagaraDataInterfaceRenderTargetCube.cpp b/Engine/Plugins/FX/Niagara/Source/Niagara/Private/NiagaraDataInterfaceRenderTargetCube.cpp
--- a/Engine/Plugins/FX/Niagara/Source/Niagara/Private/NiagaraDataInterfaceRenderTargetCube.cpp
+++ b/Engine/Plugins/FX/Niagara/Source/Niagara/Private/NiagaraDataInterfaceRenderTargetCube.cpp
@@ -1,5 +1,6 @@
 // Copyright Epic Games, Inc. All Rights Reserved.
 #include "NiagaraDataInterfaceRenderTargetCube.h"
+#include "Engine/Engine.h"
 #include "Engine/TextureRenderTarget2D.h"
 #include "TextureResource.h"
 #include "Engine/TextureRenderTargetCube.h"
@@ -384,7 +385,6 @@ void UNiagaraDataInterfaceRenderTargetCube::SetShaderParameters(const FNiagaraDa
 
 	if (bRTRead)
 	{
-		ensureMsgf(bRTWrite == false, TEXT("RenderTarget DataInterface is both wrote and read from in the same stage, this is not allowed, read will be invalid"));
 		if (bRTWrite == false && InstanceData_RT->RenderTarget.IsValid())
 		{
 			InstanceData_RT->bReadThisFrame = true;
@@ -392,6 +392,9 @@ void UNiagaraDataInterfaceRenderTargetCube::SetShaderParameters(const FNiagaraDa
 		}
 		else
 		{
+		#if WITH_NIAGARA_DEBUG_EMITTER_NAME
+			GEngine->AddOnScreenDebugMessage(uint64(this), 1.f, FColor::White, *FString::Printf(TEXT("RenderTarget is read and wrote in the same stage, this is not allowed, read will be invalid. (%s)"), *Context.GetDebugString()));
+		#endif
 			Parameters->Texture = Context.GetComputeDispatchInterface().GetBlackTextureSRV(GraphBuilder, ETextureDimension::TextureCube);
 		}
 
diff --git a/Engine/Plugins/FX/Niagara/Source/Niagara/Private/NiagaraDataInterfaceRenderTargetVolume.cpp b/Engine/Plugins/FX/Niagara/Source/Niagara/Private/NiagaraDataInterfaceRenderTargetVolume.cpp
--- a/Engine/Plugins/FX/Niagara/Source/Niagara/Private/NiagaraDataInterfaceRenderTargetVolume.cpp
+++ b/Engine/Plugins/FX/Niagara/Source/Niagara/Private/NiagaraDataInterfaceRenderTargetVolume.cpp
@@ -1,5 +1,6 @@
 // Copyright Epic Games, Inc. All Rights Reserved.
 #include "NiagaraDataInterfaceRenderTargetVolume.h"
+#include "Engine/Engine.h"
 #include "Engine/TextureRenderTarget2D.h"
 #include "TextureResource.h"
 #include "NiagaraCompileHashVisitor.h"
@@ -481,7 +482,6 @@ void UNiagaraDataInterfaceRenderTargetVolume::SetShaderParameters(const FNiagara
 
 	if (bRTRead)
 	{
-		ensureMsgf(bRTWrite == false, TEXT("RenderTarget DataInterface is both wrote and read from in the same stage, this is not allowed, read will be invalid"));
 		if (bRTWrite == false && InstanceData_RT->RenderTarget.IsValid())
 		{
 			InstanceData_RT->bReadThisFrame = true;
@@ -489,6 +489,9 @@ void UNiagaraDataInterfaceRenderTargetVolume::SetShaderParameters(const FNiagara
 		}
 		else
 		{
+		#if WITH_NIAGARA_DEBUG_EMITTER_NAME
+			GEngine->AddOnScreenDebugMessage(uint64(this), 1.f, FColor::White, *FString::Printf(TEXT("RenderTarget is read and wrote in the same stage, this is not allowed, read will be invalid. (%s)"), *Context.GetDebugString()));
+		#endif
 			Parameters->Texture = Context.GetComputeDispatchInterface().GetBlackTextureSRV(GraphBuilder, ETextureDimension::Texture3D);
 		}
 
diff --git a/Engine/Plugins/FX/Niagara/Source/Niagara/Private/NiagaraGpuComputeDispatch.cpp b/Engine/Plugins/FX/Niagara/Source/Niagara/Private/NiagaraGpuComputeDispatch.cpp
--- a/Engine/Plugins/FX/Niagara/Source/Niagara/Private/NiagaraGpuComputeDispatch.cpp
+++ b/Engine/Plugins/FX/Niagara/Source/Niagara/Private/NiagaraGpuComputeDispatch.cpp
@@ -791,13 +791,7 @@ void FNiagaraGpuComputeDispatch::PrepareTicksForProxy(FRHICommandListImmediate&
 					if (bRaisedWarningThisFrame == false)
 					{
 						bRaisedWarningThisFrame = true;
-						AsyncTask(
-							ENamedThreads::GameThread,
-							[MessageId=uint64(this), DebugSimName=ComputeContext->GetDebugSimFName()]()
-							{
-								GEngine->AddOnScreenDebugMessage(MessageId, 1.f, FColor::White, *FString::Printf(TEXT("GPU Simulation(%s) will not show in preview mode, as we may sample from wrong SceneTextures buffer."), *DebugSimName.ToString()));
-							}
-						);
+						GEngine->AddOnScreenDebugMessage(uint64(this), 1.f, FColor::White, *FString::Printf(TEXT("GPU Simulation(%s) will not show in preview mode, as we may sample from wrong SceneTextures buffer."), *ComputeContext->GetDebugSimFName().ToString()));
 					}
 					continue;
 				}
