Moved NDC BeginFrame to world tick start to include writes from net update.
Debug Util CVars for forcing latent reads and TG dependencies for NDC Read DI.
No longer forcing NDC island handlers to tick in LastDemotable.

diff --git a/Engine/Plugins/FX/Niagara/Source/Niagara/Internal/DataInterface/NiagaraDataInterfaceDataChannelRead.h b/Engine/Plugins/FX/Niagara/Source/Niagara/Internal/DataInterface/NiagaraDataInterfaceDataChannelRead.h
--- a/Engine/Plugins/FX/Niagara/Source/Niagara/Internal/DataInterface/NiagaraDataInterfaceDataChannelRead.h
+++ b/Engine/Plugins/FX/Niagara/Source/Niagara/Internal/DataInterface/NiagaraDataInterfaceDataChannelRead.h
@@ -140,7 +140,10 @@ public:
 	NIAGARA_API virtual void PostStageTick(FNDICpuPostStageContext& Context)override;
 	NIAGARA_API virtual void ProvidePerInstanceDataForRenderThread(void* DataForRenderThread, void* PerInstanceData, const FNiagaraSystemInstanceID& SystemInstance) override;
 	NIAGARA_API virtual void GetEmitterDependencies(UNiagaraSystem* Asset, TArray<FVersionedNiagaraEmitter>& Dependencies) const;
-	
+
+	virtual bool HasTickGroupPrereqs() const override;
+	virtual ETickingGroup CalculateTickGroup(const void* PerInstanceData) const override;
+
 	//We cannot overlap frames as we must correctly sync up with the data channel manager on Begin/End frame etc.
 	virtual bool PostSimulateCanOverlapFrames() const { return false; }
 	//UNiagaraDataInterface Interface
diff --git a/Engine/Plugins/FX/Niagara/Source/Niagara/Private/DataInterface/NiagaraDataInterfaceDataChannelRead.cpp b/Engine/Plugins/FX/Niagara/Source/Niagara/Private/DataInterface/NiagaraDataInterfaceDataChannelRead.cpp
--- a/Engine/Plugins/FX/Niagara/Source/Niagara/Private/DataInterface/NiagaraDataInterfaceDataChannelRead.cpp
+++ b/Engine/Plugins/FX/Niagara/Source/Niagara/Private/DataInterface/NiagaraDataInterfaceDataChannelRead.cpp
@@ -36,6 +36,24 @@ DECLARE_CYCLE_STAT(TEXT("NDIDataChannelRead Spawn"), STAT_NDIDataChannelRead_Spa
 DECLARE_CYCLE_STAT(TEXT("NDIDataChannelRead Tick"), STAT_NDIDataChannelRead_Tick, STATGROUP_NiagaraDataChannels);
 DECLARE_CYCLE_STAT(TEXT("NDIDataChannelRead PostTick"), STAT_NDIDataChannelRead_PostTick, STATGROUP_NiagaraDataChannels);
 
+static int GNDCReadForceTG = -1;
+static FAutoConsoleVariableRef CVarNDCReadForceTG(
+	TEXT("fx.Niagara.DataChannels.ForceReadTickGroup"),
+	GNDCReadForceTG,
+	TEXT("When >= 0 this will force Niagara systems with NDC read DIs to tick in the given Tick Group."),
+	ECVF_Default
+);
+
+static bool GNDCReadForcePrevFrame = false;
+static FAutoConsoleVariableRef CVarNDCReadForcePrevFrame(
+	TEXT("fx.Niagara.DataChannels.ForceReadPrevFrame"),
+	GNDCReadForcePrevFrame,
+	TEXT("When true this will force Niagara systems with NDC read DIs to read from the previous frame."),
+	ECVF_Default
+);
+
+
+
 namespace NDIDataChannelReadLocal
 {
 	static const TCHAR* CommonShaderFile = TEXT("/Plugin/FX/Niagara/Private/DataChannel/NiagaraDataInterfaceDataChannelCommon.ush");
@@ -457,7 +475,9 @@ void UNiagaraDataInterfaceDataChannelRead::ProvidePerInstanceDataForRenderThread
 
 	//Always update the dataset, this may change without triggering a full update if it's layout is the same.
 	TargetData->ChannelDataRTProxy = SourceData.DataChannelData ? SourceData.DataChannelData->GetRTProxy() : nullptr;
-	TargetData->bReadPrevFrame = bReadCurrentFrame == false;
+
+	bool bReadPrevFrame = bReadCurrentFrame == false || GNDCReadForcePrevFrame;
+	TargetData->bReadPrevFrame = bReadPrevFrame;
 
 	if (SourceData.bUpdateRTData && INiagaraModule::DataChannelsEnabled())
 	{
@@ -530,6 +550,24 @@ void UNiagaraDataInterfaceDataChannelRead::GetEmitterDependencies(UNiagaraSystem
 // 	}
 }
 
+bool UNiagaraDataInterfaceDataChannelRead::HasTickGroupPrereqs() const
+{
+	if (GNDCReadForceTG >= 0 && GNDCReadForceTG < (int32)ETickingGroup::TG_MAX)
+	{
+		return true;
+	}
+	return false;
+}
+
+ETickingGroup UNiagaraDataInterfaceDataChannelRead::CalculateTickGroup(const void* PerInstanceData) const
+{
+	if(GNDCReadForceTG >= 0 && GNDCReadForceTG < (int32)ETickingGroup::TG_MAX)
+	{
+		return (ETickingGroup)GNDCReadForceTG;
+	}
+	return NiagaraFirstTickGroup; 
+}
+
 #if WITH_EDITORONLY_DATA
 
 void UNiagaraDataInterfaceDataChannelRead::PostCompile()
@@ -866,7 +904,8 @@ void UNiagaraDataInterfaceDataChannelRead::Num(FVectorVMExternalFunctionContext&
 
 	FNDIOutputParam<int32> OutNum(Context);
 
-	FNiagaraDataBuffer* Buffer = InstData->GetReadBufferCPU(bReadCurrentFrame == false);
+	bool bReadPrevFrame = bReadCurrentFrame == false || GNDCReadForcePrevFrame;
+	FNiagaraDataBuffer* Buffer = InstData->GetReadBufferCPU(bReadPrevFrame);
 
 	int32 Num = 0;
 	if (Buffer && INiagaraModule::DataChannelsEnabled())
@@ -892,7 +931,8 @@ void UNiagaraDataInterfaceDataChannelRead::Read(FVectorVMExternalFunctionContext
 	const FNDIDataChannel_FunctionToDataSetBinding* BindingInfo = InstData->FuncToDataSetBindingInfo.IsValidIndex(FuncIdx) ? InstData->FuncToDataSetBindingInfo[FuncIdx].Get() : nullptr;
 	FNDIVariadicOutputHandler<16> VariadicOutputs(Context, BindingInfo);//TODO: Make static / avoid allocation
 
-	FNiagaraDataBuffer* Data = InstData->GetReadBufferCPU(bReadCurrentFrame == false);
+	bool bReadPrevFrame = bReadCurrentFrame == false || GNDCReadForcePrevFrame;
+	FNiagaraDataBuffer* Data = InstData->GetReadBufferCPU(bReadPrevFrame);
 	if (Data && BindingInfo && INiagaraModule::DataChannelsEnabled())
  	{
  		//FString Label = TEXT("NDIDataChannelRead::Read() - ");
@@ -958,7 +998,8 @@ void UNiagaraDataInterfaceDataChannelRead::Consume(FVectorVMExternalFunctionCont
 
 	//TODO: Optimize for constant bConsume.
 	//TODO: Optimize for long runs of bConsume==true;
-	FNiagaraDataBuffer* Data = InstData->GetReadBufferCPU(bReadCurrentFrame == false);
+	bool bReadPrevFrame = bReadCurrentFrame == false || GNDCReadForcePrevFrame;
+	FNiagaraDataBuffer* Data = InstData->GetReadBufferCPU(bReadPrevFrame);
 	if (Data && BindingInfo && INiagaraModule::DataChannelsEnabled())
 	{
 		for (int32 i = 0; i < Context.GetNumInstances(); ++i)
@@ -1039,7 +1080,9 @@ void UNiagaraDataInterfaceDataChannelRead::SpawnFromSpawnInfo(FVectorVMExternalF
 
 	FNiagaraEmitterInstance* EmittterInst = InstData->EmitterInstance;
 	FNiagaraDataChannelData* DataChannelData = InstData->DataChannelData.Get();
-	FNiagaraDataBuffer* Data = DataChannelData ? DataChannelData->GetCPUData(!bReadCurrentFrame) : nullptr;
+
+	bool bReadPrevFrame = bReadCurrentFrame == false || GNDCReadForcePrevFrame;
+	FNiagaraDataBuffer* Data = DataChannelData ? DataChannelData->GetCPUData(bReadPrevFrame) : nullptr;
 
 	if ( INiagaraModule::DataChannelsEnabled() && Data && EmittterInst )
 	{
@@ -1074,6 +1117,8 @@ void UNiagaraDataInterfaceDataChannelRead::SpawnConditional(FVectorVMExternalFun
 {
 	SCOPE_CYCLE_COUNTER(STAT_NDIDataChannelRead_Spawn);
 
+	//UE_LOG(LogNiagara, Warning, TEXT("UNiagaraDataInterfaceDataChannelRead::SpawnConditional - %u"), FuncIdx);
+
 	//This should only be called from emitter scripts and since it has per instance data then we process them individually.
 	check(Context.GetNumInstances() == 1);
 
@@ -1092,7 +1137,9 @@ void UNiagaraDataInterfaceDataChannelRead::SpawnConditional(FVectorVMExternalFun
 
 	FNiagaraEmitterInstance* EmittterInst = InstData->EmitterInstance;
 	FNiagaraDataChannelData* DataChannelData = InstData->DataChannelData.Get();
-	FNiagaraDataBuffer* Data = DataChannelData ? DataChannelData->GetCPUData(!bReadCurrentFrame) : nullptr;
+
+	bool bReadPrevFrame = bReadCurrentFrame == false || GNDCReadForcePrevFrame;
+	FNiagaraDataBuffer* Data = DataChannelData ? DataChannelData->GetCPUData(bReadPrevFrame) : nullptr;
 
 	bool bSpawn = INiagaraModule::DataChannelsEnabled() && Data && Data->GetNumInstances() > 0 && EmittterInst && EmittterInst->IsActive() && InEnabled.GetAndAdvance();
 	if(bSpawn)
diff --git a/Engine/Plugins/FX/Niagara/Source/Niagara/Private/NiagaraDataChannel_Global.cpp b/Engine/Plugins/FX/Niagara/Source/Niagara/Private/NiagaraDataChannel_Global.cpp
--- a/Engine/Plugins/FX/Niagara/Source/Niagara/Private/NiagaraDataChannel_Global.cpp
+++ b/Engine/Plugins/FX/Niagara/Source/Niagara/Private/NiagaraDataChannel_Global.cpp
@@ -51,11 +51,13 @@ void UNiagaraDataChannelHandler_Global::Init(const UNiagaraDataChannel* InChanne
 
 void UNiagaraDataChannelHandler_Global::BeginFrame(float DeltaTime, FNiagaraWorldManager* OwningWorld)
 {
+	Super::BeginFrame(DeltaTime, OwningWorld);
 	Data->BeginFrame(this);
 }
 
 void UNiagaraDataChannelHandler_Global::EndFrame(float DeltaTime, FNiagaraWorldManager* OwningWorld)
 {
+	Super::EndFrame(DeltaTime, OwningWorld);
 	Data->EndFrame(this);
 }
 
diff --git a/Engine/Plugins/FX/Niagara/Source/Niagara/Private/NiagaraDataChannel_Islands.cpp b/Engine/Plugins/FX/Niagara/Source/Niagara/Private/NiagaraDataChannel_Islands.cpp
--- a/Engine/Plugins/FX/Niagara/Source/Niagara/Private/NiagaraDataChannel_Islands.cpp
+++ b/Engine/Plugins/FX/Niagara/Source/Niagara/Private/NiagaraDataChannel_Islands.cpp
@@ -82,6 +82,10 @@ void UNiagaraDataChannelHandler_Islands::EndFrame(float DeltaTime, FNiagaraWorld
 
 void UNiagaraDataChannelHandler_Islands::Tick(float DeltaTime, ETickingGroup TickGroup, FNiagaraWorldManager* OwningWorld)
 {
+	Super::Tick(DeltaTime, TickGroup, OwningWorld);
+
+	//static UEnum* TGEnum = StaticEnum<ETickingGroup>();
+	//UE_LOG(LogNiagara,Warning, TEXT("UNiagaraDataChannelHandler_Islands::Tick - %s"), *TGEnum->GetDisplayNameTextByValue((int32)TickGroup).ToString());
 	const UNiagaraDataChannel_Islands* IslandChannel = GetChannelTyped<const UNiagaraDataChannel_Islands>();
 	for (auto It = ActiveIslands.CreateIterator(); It; ++It)
 	{
@@ -304,7 +308,6 @@ void FNDCIsland::OnAcquired(FVector Location)
 			{
 				NiagaraSystems.Add(NewComp);
 				NewComp->SetSystemFixedBounds(HandlerSystemBounds);
-				NewComp->SetTickBehavior(ENiagaraTickBehavior::ForceTickLast);
 				NewComp->Activate();
 			}
 		}
diff --git a/Engine/Plugins/FX/Niagara/Source/Niagara/Private/NiagaraWorldManager.cpp b/Engine/Plugins/FX/Niagara/Source/Niagara/Private/NiagaraWorldManager.cpp
--- a/Engine/Plugins/FX/Niagara/Source/Niagara/Private/NiagaraWorldManager.cpp
+++ b/Engine/Plugins/FX/Niagara/Source/Niagara/Private/NiagaraWorldManager.cpp
@@ -252,6 +252,7 @@ FDelegateHandle FNiagaraWorldManager::OnWorldCleanupHandle;
 FDelegateHandle FNiagaraWorldManager::OnPostWorldCleanupHandle;
 FDelegateHandle FNiagaraWorldManager::OnPreWorldFinishDestroyHandle;
 FDelegateHandle FNiagaraWorldManager::OnWorldBeginTearDownHandle;
+FDelegateHandle FNiagaraWorldManager::TickWorldStartHandle;
 FDelegateHandle FNiagaraWorldManager::TickWorldHandle;
 FDelegateHandle FNiagaraWorldManager::OnWorldPreActorTickHandle;
 FDelegateHandle FNiagaraWorldManager::OnWorldPreSendAllEndOfFrameUpdatesHandle;
@@ -401,6 +402,7 @@ void FNiagaraWorldManager::OnStartup()
 	OnPostWorldCleanupHandle = FWorldDelegates::OnPostWorldCleanup.AddStatic(&FNiagaraWorldManager::OnPostWorldCleanup);
 	OnPreWorldFinishDestroyHandle = FWorldDelegates::OnPreWorldFinishDestroy.AddStatic(&FNiagaraWorldManager::OnPreWorldFinishDestroy);
 	OnWorldBeginTearDownHandle = FWorldDelegates::OnWorldBeginTearDown.AddStatic(&FNiagaraWorldManager::OnWorldBeginTearDown);
+	TickWorldStartHandle = FWorldDelegates::OnWorldTickStart.AddStatic(&FNiagaraWorldManager::TickWorldStart);
 	TickWorldHandle = FWorldDelegates::OnWorldPostActorTick.AddStatic(&FNiagaraWorldManager::TickWorld);
 	OnWorldPreActorTickHandle = FWorldDelegates::OnWorldPreActorTick.AddStatic(&FNiagaraWorldManager::OnWorldPreActorTick);
 	OnWorldPreSendAllEndOfFrameUpdatesHandle = FWorldDelegates::OnWorldPreSendAllEndOfFrameUpdates.AddLambda(
@@ -916,6 +918,11 @@ void FNiagaraWorldManager::DestroyAllSystemSimulations(class UNiagaraSystem* Sys
 	}
 }
 
+void FNiagaraWorldManager::TickWorldStart(UWorld* World, ELevelTick TickType, float DeltaSeconds)
+{
+	Get(World)->TickStart(DeltaSeconds);
+}
+
 void FNiagaraWorldManager::TickWorld(UWorld* World, ELevelTick TickType, float DeltaSeconds)
 {
 	Get(World)->PostActorTick(DeltaSeconds);
@@ -966,9 +973,17 @@ void FNiagaraWorldManager::OnRefreshOwnerAllowsScalability()
 	}
 }
 
+void FNiagaraWorldManager::TickStart(float DeltaSeconds)
+{
+	CSV_SCOPED_TIMING_STAT_EXCLUSIVE(Effects);
+	LLM_SCOPE(ELLMTag::Niagara);
+	SCOPE_CYCLE_COUNTER(STAT_NiagaraWorldManTick);
+
+	DataChannelManager->BeginFrame(DeltaSeconds);
+}
+
 void FNiagaraWorldManager::PreActorTick(ELevelTick InLevelTick, float InDeltaSeconds)
 {	
-	DataChannelManager->BeginFrame(InDeltaSeconds);
 	NiagaraWorldManagerInternal::ExecuteGlobalDeferredCallbacks();
 }
 
@@ -1044,6 +1059,7 @@ void FNiagaraWorldManager::PostActorTick(float DeltaSeconds)
 	// Tick debug HUD for the world
 	if (NiagaraDebugHud != nullptr)
 	{
+		QUICK_SCOPE_CYCLE_COUNTER(STAT_NiagaraDebugHUD);
 		NiagaraDebugHud->GatherSystemInfo();
 	}
 #endif
diff --git a/Engine/Plugins/FX/Niagara/Source/Niagara/Public/NiagaraWorldManager.h b/Engine/Plugins/FX/Niagara/Source/Niagara/Public/NiagaraWorldManager.h
--- a/Engine/Plugins/FX/Niagara/Source/Niagara/Public/NiagaraWorldManager.h
+++ b/Engine/Plugins/FX/Niagara/Source/Niagara/Public/NiagaraWorldManager.h
@@ -122,6 +122,9 @@ public:
 	void MarkSimulationForPostActorWork(FNiagaraSystemSimulation* SystemSimulation);
 	void MarkSimulationsForEndOfFrameWait(FNiagaraSystemSimulation* SystemSimulation);
 
+	/** Called at the very beginning of world ticking. */
+	void TickStart(float DeltaSeconds);
+
 	void Tick(ETickingGroup TickGroup, float DeltaSeconds, ELevelTick TickType, ENamedThreads::Type CurrentThread, const FGraphEventRef& MyCompletionGraphEvent);
 
 	void PreActorTick(ELevelTick InLevelTick, float InDeltaSeconds);
@@ -297,6 +300,9 @@ private:
 
 	static void OnWorldPreActorTick(UWorld* InWorld, ELevelTick InLevelTick, float InDeltaSeconds);
 
+	// Callback for when a world is ticked.
+	static void TickWorldStart(UWorld* World, ELevelTick TickType, float DeltaSeconds);
+
 	// Callback for when a world is ticked.
 	static void TickWorld(UWorld* World, ELevelTick TickType, float DeltaSeconds);
 
@@ -334,6 +340,7 @@ private:
 	static FDelegateHandle OnPostWorldCleanupHandle;
 	static FDelegateHandle OnPreWorldFinishDestroyHandle;
 	static FDelegateHandle OnWorldBeginTearDownHandle;
+	static FDelegateHandle TickWorldStartHandle;
 	static FDelegateHandle TickWorldHandle;
 	static FDelegateHandle OnWorldPreActorTickHandle;
 	static FDelegateHandle OnWorldPreSendAllEndOfFrameUpdatesHandle;
