- Validate the copy destination
- This avoid potential crashes if users RT parameters are bound to multiple different data interfaces

diff --git a/Engine/Plugins/FX/Niagara/Source/Niagara/Private/NiagaraRenderGraphUtils.cpp b/Engine/Plugins/FX/Niagara/Source/Niagara/Private/NiagaraRenderGraphUtils.cpp
--- a/Engine/Plugins/FX/Niagara/Source/Niagara/Private/NiagaraRenderGraphUtils.cpp
+++ b/Engine/Plugins/FX/Niagara/Source/Niagara/Private/NiagaraRenderGraphUtils.cpp
@@ -154,6 +154,20 @@ FRDGTextureUAVRef FNiagaraPooledRWTexture::GetOrCreateUAV(FRDGBuilder& GraphBuil
 
 void FNiagaraPooledRWTexture::CopyToTexture(FRDGBuilder& GraphBuilder, FRHITexture* DestinationTextureRHI, const TCHAR* NameIfNotRegistered)
 {
+	const FPooledRenderTargetDesc& TextureDesc = Texture->GetDesc();
+	const FRHITextureDesc& DestTextureDesc = DestinationTextureRHI->GetDesc();
+	if (!ensureMsgf(
+			(DestTextureDesc.Dimension == ETextureDimension::Texture3D) &&
+			(DestTextureDesc.Format == TextureDesc.Format) &&
+			(DestTextureDesc.Extent == TextureDesc.Extent) &&
+			(DestTextureDesc.Depth == TextureDesc.Depth),
+			TEXT("FNiagaraPooledRWTexture::CopyToTexture - Destination does not match copy will fail")
+		)
+	)
+	{
+		return;
+	}
+
 	FRDGTexture* SourceTexture = GetOrCreateTexture(GraphBuilder);
 	FRDGTexture* DestinationTexture = GraphBuilder.FindExternalTexture(DestinationTextureRHI);
 	if (DestinationTexture == nullptr)
