- Fix for UniqueID being shared incorrectly between spawn groups on the GPU

diff --git a/Engine/Plugins/FX/Niagara/Source/NiagaraEditor/Private/NiagaraHlslTranslator.cpp b/Engine/Plugins/FX/Niagara/Source/NiagaraEditor/Private/NiagaraHlslTranslator.cpp
--- a/Engine/Plugins/FX/Niagara/Source/NiagaraEditor/Private/NiagaraHlslTranslator.cpp
+++ b/Engine/Plugins/FX/Niagara/Source/NiagaraEditor/Private/NiagaraHlslTranslator.cpp
@@ -3667,7 +3667,7 @@ void TNiagaraHlslTranslator<GraphBridge>::DefineMainGPUFunctions(
 		HlslOutput.Append(TEXT("		ReadDataSets(Context);\n"));
 		if (bParticleSpawnStage == true)
 		{
-			HlslOutput.Append(TEXT("		Context.MapSpawn.Particles.UniqueID = Engine_Emitter_TotalSpawnedParticles + ExecIndex();\n"));
+			HlslOutput.Append(TEXT("		Context.MapSpawn.Particles.UniqueID = Engine_Emitter_TotalSpawnedParticles + GLinearThreadId - GSpawnStartInstance;\n"));
 			HlslOutput.Append(TEXT("		ConditionalInterpolateParameters(Context);\n"));
 			HlslOutput.Append(TEXT("		SimulateMapSpawn(Context);\n"));
 		}
