- Validate the export type & size when copying grid data out
- This fixes a potential gpu hang if the destination texture was not setup correctly or accidentally shared inside the GPU sim

#rnx

diff --git a/Engine/Plugins/FX/Niagara/Source/Niagara/Private/NiagaraDataInterfaceGrid2DCollection.cpp b/Engine/Plugins/FX/Niagara/Source/Niagara/Private/NiagaraDataInterfaceGrid2DCollection.cpp
--- a/Engine/Plugins/FX/Niagara/Source/Niagara/Private/NiagaraDataInterfaceGrid2DCollection.cpp
+++ b/Engine/Plugins/FX/Niagara/Source/Niagara/Private/NiagaraDataInterfaceGrid2DCollection.cpp
@@ -195,6 +195,12 @@ struct FNiagaraGrid2DLegacyTiled2DInfo
 
 	void CopyTo2D(FRDGBuilder& GraphBuilder, FRDGTexture* SourceTexture, FRHITexture* DestinationTextureRHI, const TCHAR* NameIfNotRegistered) const
 	{
+		const FRHITextureDesc& DestTextureDesc = DestinationTextureRHI->GetDesc();
+		if (ensureMsgf(DestTextureDesc.Dimension != ETextureDimension::Texture2D, TEXT("Destination for Grid2D copy is invalid type (%d), copy will fail"), DestTextureDesc.Dimension))
+		{
+			return;
+		}
+
 		FRDGTexture* DestinationTexture = GraphBuilder.FindExternalTexture(DestinationTextureRHI);
 		if (DestinationTexture == nullptr)
 		{
@@ -209,7 +215,12 @@ struct FNiagaraGrid2DLegacyTiled2DInfo
 			CopyInfo.DestPosition.X = (iAttribute % NumTiles.X) * NumCells.X;
 			CopyInfo.DestPosition.Y = (iAttribute / NumTiles.X) * NumCells.Y;
 			CopyInfo.DestPosition.Z = 0;
-			AddCopyTexturePass(GraphBuilder, SourceTexture, DestinationTexture, CopyInfo);
+			if (ensureMsgf(
+					(CopyInfo.DestPosition.X + CopyInfo.Size.X <= DestTextureDesc.Extent.X) && (CopyInfo.DestPosition.Y + CopyInfo.Size.Y <= DestTextureDesc.Extent.Y),
+					TEXT("Destination for Grid2D can not hold all the attributes, copy will fail.")))
+			{
+				AddCopyTexturePass(GraphBuilder, SourceTexture, DestinationTexture, CopyInfo);
+			}
 		}
 	}
 
@@ -217,6 +228,12 @@ struct FNiagaraGrid2DLegacyTiled2DInfo
 	{
 		check(Src != nullptr && Dst != nullptr);
 
+		const FRHITextureDesc& DestTextureDesc = Dst->GetDesc();
+		if (ensureMsgf(DestTextureDesc.Dimension != ETextureDimension::Texture2D, TEXT("Destination for Grid2D copy is invalid type (%d), copy will fail"), DestTextureDesc.Dimension))
+		{
+			return;
+		}
+
 		FRHITransitionInfo TransitionsBefore[] = {
 			FRHITransitionInfo(Src, ERHIAccess::SRVMask, ERHIAccess::CopySrc),
 			FRHITransitionInfo(Dst, ERHIAccess::SRVMask, ERHIAccess::CopyDest)
@@ -231,7 +248,12 @@ struct FNiagaraGrid2DLegacyTiled2DInfo
 			CopyInfo.DestPosition.X = (iAttribute % NumTiles.X) * NumCells.X;
 			CopyInfo.DestPosition.Y = (iAttribute / NumTiles.X) * NumCells.Y;
 			CopyInfo.DestPosition.Z = 0;
-			RHICmdList.CopyTexture(Src, Dst, CopyInfo);
+			if (ensureMsgf(
+				(CopyInfo.DestPosition.X + CopyInfo.Size.X <= DestTextureDesc.Extent.X) && (CopyInfo.DestPosition.Y + CopyInfo.Size.Y <= DestTextureDesc.Extent.Y),
+				TEXT("Destination for Grid2D can not hold all the attributes, copy will fail.")))
+			{
+				RHICmdList.CopyTexture(Src, Dst, CopyInfo);
+			}
 		}
 
 		FRHITransitionInfo TransitionsAfter[] = {
