Show restricted content stack errors for scratch pads that use things like custom hlsl nodes

diff --git a/Engine/Plugins/FX/Niagara/Source/NiagaraEditor/Private/NiagaraEditorSettings.cpp b/Engine/Plugins/FX/Niagara/Source/NiagaraEditor/Private/NiagaraEditorSettings.cpp
--- a/Engine/Plugins/FX/Niagara/Source/NiagaraEditor/Private/NiagaraEditorSettings.cpp
+++ b/Engine/Plugins/FX/Niagara/Source/NiagaraEditor/Private/NiagaraEditorSettings.cpp
@@ -773,6 +773,42 @@ bool UNiagaraEditorSettings::IsAllowedAssetObjectByClassUsageInternal(const UObj
 	return true;
 }
 
+bool UNiagaraEditorSettings::IsAllowedObjectByClassUsageInternal(const UObject& InObject, TSet<const UObject*>& CheckedObjects) const
+{
+	CheckedObjects.Add(&InObject);
+
+	TArray<UObject*> ObjectsWithOuter;
+	GetObjectsWithOuter(&InObject, ObjectsWithOuter, true);
+	
+	for (const UObject* ObjectInPackage : ObjectsWithOuter)
+	{
+		if (ShouldTrackClassUsage(ObjectInPackage->GetClass()) && IsAllowedClass(ObjectInPackage->GetClass()) == false)
+		{
+			if (GbLogFoundButNotAllowedAssets)
+			{
+				UE_LOG(LogNiagaraEditor, Log, TEXT("Asset %s is not allowed due to object %s with class %s which is not allowed in this editor context."),
+					*InObject.GetPackage()->GetPathName(), *ObjectInPackage->GetPathName(), *ObjectInPackage->GetClass()->GetClassPathName().ToString());
+			}
+			return false;
+		}
+
+		const UNiagaraNode* NiagaraNode = Cast<UNiagaraNode>(ObjectInPackage);
+		if (NiagaraNode != nullptr &&
+			NiagaraNode->GetReferencedAsset() != nullptr &&
+			OnShouldFilterAssetByClassUsage.Execute(FTopLevelAssetPath(NiagaraNode->GetReferencedAsset()->GetPathName())) &&
+			CheckedObjects.Contains(NiagaraNode->GetReferencedAsset()) == false)
+		{
+			if (IsAllowedAssetObjectByClassUsageInternal(*NiagaraNode->GetReferencedAsset(), CheckedObjects) == false)
+			{
+				UE_LOG(LogNiagaraEditor, Log, TEXT("Asset %s is not allowed due to referenced asset %s which is not allowed in this editor context."),
+					*InObject.GetPackage()->GetPathName(), *NiagaraNode->GetReferencedAsset()->GetPathName());
+				return false;
+			}
+		}
+	}
+	return true;
+}
+
 bool UNiagaraEditorSettings::IsAllowedAssetByClassUsage(const FAssetData& InAssetData) const
 {
 	if (OnShouldFilterAssetByClassUsage.IsBound() == false ||
@@ -831,6 +867,12 @@ bool UNiagaraEditorSettings::IsAllowedAssetByClassUsage(const FAssetData& InAsse
 
 bool UNiagaraEditorSettings::IsAllowedAssetObjectByClassUsage(const UObject& InAssetObject) const
 {
+	TSet<const UObject*> CheckedAssetObjects;
+	if (OnIsClassAllowedDelegate.IsBound() && InAssetObject.IsAsset() == false)
+	{
+		return IsAllowedObjectByClassUsageInternal(InAssetObject, CheckedAssetObjects);
+	}
+	
 	if (OnShouldFilterAssetByClassUsage.IsBound() == false ||
 		OnIsClassAllowedDelegate.IsBound() == false ||
 		OnShouldFilterAssetByClassUsage.Execute(FTopLevelAssetPath(InAssetObject.GetPathName())) == false)
@@ -838,7 +880,6 @@ bool UNiagaraEditorSettings::IsAllowedAssetObjectByClassUsage(const UObject& InA
 		return true;
 	}
 
-	TSet<const UObject*> CheckedAssetObjects;
 	return IsAllowedAssetObjectByClassUsageInternal(InAssetObject, CheckedAssetObjects);
 }
 
diff --git a/Engine/Plugins/FX/Niagara/Source/NiagaraEditor/Public/NiagaraEditorSettings.h b/Engine/Plugins/FX/Niagara/Source/NiagaraEditor/Public/NiagaraEditorSettings.h
--- a/Engine/Plugins/FX/Niagara/Source/NiagaraEditor/Public/NiagaraEditorSettings.h
+++ b/Engine/Plugins/FX/Niagara/Source/NiagaraEditor/Public/NiagaraEditorSettings.h
@@ -428,6 +428,7 @@ public:
 	NIAGARAEDITOR_API bool GetUpdateStackValuesOnCommitOnly() const;
 
 private:
+	NIAGARAEDITOR_API bool IsAllowedObjectByClassUsageInternal(const UObject& InObject, TSet<const UObject*>& CheckedObjects) const;
 	NIAGARAEDITOR_API bool IsAllowedAssetObjectByClassUsageInternal(const UObject& InAssetObject, TSet<const UObject*>& CheckedAssetObjects) const;
 
 	NIAGARAEDITOR_API void SetupNamespaceMetadata();
