PR #10579: Removing ambiguous equality operator

#rnx

diff --git a/Engine/Plugins/FX/Niagara/Source/Niagara/Public/NiagaraTypes.h b/Engine/Plugins/FX/Niagara/Source/Niagara/Public/NiagaraTypes.h
--- a/Engine/Plugins/FX/Niagara/Source/Niagara/Public/NiagaraTypes.h
+++ b/Engine/Plugins/FX/Niagara/Source/Niagara/Public/NiagaraTypes.h
@@ -1924,25 +1924,6 @@ struct FNiagaraVariable : public FNiagaraVariableBase
 		: FNiagaraVariableBase(InType, InName)
 	{
 	}
-	
-	/** Check if Name and Type definition are the same. The actual stored value is not checked here.*/
-	bool operator==(const FNiagaraVariable& Other)const
-	{
-		//-TODO: Should this check the value???
-		return Name == Other.Name && TypeDefHandle == Other.TypeDefHandle;
-	}
-
-	/** Check if Name and Type definition are the same. The actual stored value is not checked here.*/
-	bool operator==(const FNiagaraVariableBase& Other)const
-	{
-		return FNiagaraVariableBase(*this) == Other;
-	}
-
-	/** Check if Name and Type definition are not the same. The actual stored value is not checked here.*/
-	bool operator!=(const FNiagaraVariable& Other)const
-	{
-		return !(*this == Other);
-	}
 
 	/** Checks if the types match and either both variables are uninitialized or both hold exactly the same data.*/
 	bool HoldsSameData(const FNiagaraVariable& Other) const
