- Particle reader updated to use emitter binding to match other data interfaces
- This removes the customization and allows for 'self' to be used

diff --git a/Engine/Plugins/FX/Niagara/Source/Niagara/Classes/NiagaraDataInterfaceParticleRead.h b/Engine/Plugins/FX/Niagara/Source/Niagara/Classes/NiagaraDataInterfaceParticleRead.h
--- a/Engine/Plugins/FX/Niagara/Source/Niagara/Classes/NiagaraDataInterfaceParticleRead.h
+++ b/Engine/Plugins/FX/Niagara/Source/Niagara/Classes/NiagaraDataInterfaceParticleRead.h
@@ -5,6 +5,7 @@
 #include "NiagaraDataInterfaceRW.h"
 #include "NiagaraCommon.h"
 #include "NiagaraEmitterInstance.h"
+#include "NiagaraDataInterfaceEmitterBinding.h"
 #include "NiagaraDataInterfaceParticleRead.generated.h"
 
 UCLASS(EditInlineNew, Category = "ParticleRead", CollapseCategories, meta = (DisplayName = "Particle Attribute Reader"), MinimalAPI)
@@ -28,11 +29,13 @@ class UNiagaraDataInterfaceParticleRead : public UNiagaraDataInterfaceRWBase
 	END_SHADER_PARAMETER_STRUCT()
 
 public:
-	UPROPERTY(EditAnywhere, Category = "ParticleRead")
-	FString EmitterName;
+	/** Selects which emitter the data interface will bind to, i.e the emitter we are contained within or a named emitter. */
+	UPROPERTY(EditAnywhere, Category = "Emitter")
+	FNiagaraDataInterfaceEmitterBinding EmitterBinding;
 
 	//UObject Interface
-	NIAGARA_API virtual void PostInitProperties()override;
+	NIAGARA_API virtual void PostInitProperties() override;
+	NIAGARA_API virtual void PostLoad() override;
 #if WITH_EDITOR
 	NIAGARA_API virtual void PostEditChangeProperty(struct FPropertyChangedEvent& PropertyChangedEvent);
 #endif
@@ -100,4 +103,9 @@ protected:
 	NIAGARA_API void GetPersistentIDFunctions(TArray<FNiagaraFunctionSignature>& OutFunctions);
 	NIAGARA_API void GetIndexFunctions(TArray<FNiagaraFunctionSignature>& OutFunctions);
 	NIAGARA_API virtual bool CopyToInternal(UNiagaraDataInterface* Destination) const override;
+
+#if WITH_EDITORONLY_DATA
+	UPROPERTY()
+	FString EmitterName_DEPRECATED;
+#endif
 };
diff --git a/Engine/Plugins/FX/Niagara/Source/Niagara/Private/DataInterface/NiagaraDataInterfaceEmitterProperties.cpp b/Engine/Plugins/FX/Niagara/Source/Niagara/Private/DataInterface/NiagaraDataInterfaceEmitterProperties.cpp
--- a/Engine/Plugins/FX/Niagara/Source/Niagara/Private/DataInterface/NiagaraDataInterfaceEmitterProperties.cpp
+++ b/Engine/Plugins/FX/Niagara/Source/Niagara/Private/DataInterface/NiagaraDataInterfaceEmitterProperties.cpp
@@ -332,7 +332,7 @@ void UNiagaraDataInterfaceEmitterProperties::GetFeedback(UNiagaraSystem* Asset,
 	// See if we are resolve the source emitter
 	if ( EmitterBinding.BindingMode == ENiagaraDataInterfaceEmitterBindingMode::Other )
 	{
-		UNiagaraEmitter* NiagaraEmitter = EmitterBinding.Resolve(Asset);
+		UNiagaraEmitter* NiagaraEmitter = EmitterBinding.Resolve(this);
 		if (NiagaraEmitter == nullptr)
 		{
 			OutWarnings.Emplace(
diff --git a/Engine/Plugins/FX/Niagara/Source/Niagara/Private/NiagaraDataInterfaceEmitterBinding.cpp b/Engine/Plugins/FX/Niagara/Source/Niagara/Private/NiagaraDataInterfaceEmitterBinding.cpp
--- a/Engine/Plugins/FX/Niagara/Source/Niagara/Private/NiagaraDataInterfaceEmitterBinding.cpp
+++ b/Engine/Plugins/FX/Niagara/Source/Niagara/Private/NiagaraDataInterfaceEmitterBinding.cpp
@@ -7,7 +7,7 @@
 
 #include UE_INLINE_GENERATED_CPP_BY_NAME(NiagaraDataInterfaceEmitterBinding)
 
-FNiagaraEmitterInstance* FNiagaraDataInterfaceEmitterBinding::Resolve(FNiagaraSystemInstance* SystemInstance, UNiagaraDataInterface* DataInterface)
+FNiagaraEmitterInstance* FNiagaraDataInterfaceEmitterBinding::Resolve(const FNiagaraSystemInstance* SystemInstance, const UNiagaraDataInterface* DataInterface) const
 {
 	if (BindingMode == ENiagaraDataInterfaceEmitterBindingMode::Self)
 	{
@@ -95,27 +95,44 @@ FNiagaraEmitterInstance* FNiagaraDataInterfaceEmitterBinding::Resolve(FNiagaraSy
 	return nullptr;
 }
 
-UNiagaraEmitter* FNiagaraDataInterfaceEmitterBinding::Resolve(UNiagaraSystem* NiagaraSystem)
+const FNiagaraEmitterHandle* FNiagaraDataInterfaceEmitterBinding::ResolveHandle(const UNiagaraDataInterface* DataInterface) const
 {
-	if ( BindingMode == ENiagaraDataInterfaceEmitterBindingMode::Other )
+	UNiagaraSystem* OwnerSystem = DataInterface->GetTypedOuter<UNiagaraSystem>();
+	if (OwnerSystem == nullptr)
 	{
-		if (EmitterName.IsNone())
+		return nullptr;
+	}
+
+	if (BindingMode == ENiagaraDataInterfaceEmitterBindingMode::Self)
+	{
+		if (UNiagaraEmitter* OwnerEmitter = DataInterface->GetTypedOuter<UNiagaraEmitter>())
 		{
-			return nullptr;
+			for (const FNiagaraEmitterHandle& EmitterHandle : OwnerSystem->GetEmitterHandles())
+			{
+				if (OwnerEmitter == EmitterHandle.GetInstance().Emitter)
+				{
+					return &EmitterHandle;
+				}
+			}
 		}
-
-		FNameBuilder EmitterNameString;
-		EmitterName.ToString(EmitterNameString);
-		FStringView EmitterNameStringView = EmitterNameString.ToView();
-
-		for ( const FNiagaraEmitterHandle& EmitterHandle : NiagaraSystem->GetEmitterHandles() )
+	}
+	else if (BindingMode == ENiagaraDataInterfaceEmitterBindingMode::Other)
+	{
+		if (!EmitterName.IsNone())
 		{
-			if (UNiagaraEmitter* NiagaraEmitter = EmitterHandle.GetInstance().Emitter)
+			FNameBuilder EmitterNameString;
+			EmitterName.ToString(EmitterNameString);
+			FStringView EmitterNameStringView = EmitterNameString.ToView();
+
+			for (const FNiagaraEmitterHandle& EmitterHandle : OwnerSystem->GetEmitterHandles())
 			{
-				//-TODO: UniqueEmitterName should probably be a FName?
-				if (EmitterNameStringView.Equals(NiagaraEmitter->GetUniqueEmitterName()))
+				if (UNiagaraEmitter* NiagaraEmitter = EmitterHandle.GetInstance().Emitter)
 				{
-					return NiagaraEmitter;
+					//-TODO: UniqueEmitterName should probably be a FName?
+					if (EmitterNameStringView.Equals(NiagaraEmitter->GetUniqueEmitterName()))
+					{
+						return &EmitterHandle;
+					}
 				}
 			}
 		}
@@ -123,3 +140,14 @@ UNiagaraEmitter* FNiagaraDataInterfaceEmitterBinding::Resolve(UNiagaraSystem* Ni
 	return nullptr;
 }
 
+UNiagaraEmitter* FNiagaraDataInterfaceEmitterBinding::Resolve(const UNiagaraDataInterface* DataInterface) const
+{
+	const FNiagaraEmitterHandle* EmitterHandle = ResolveHandle(DataInterface);
+	return EmitterHandle ? EmitterHandle->GetInstance().Emitter : nullptr;
+}
+
+FString FNiagaraDataInterfaceEmitterBinding::ResolveUniqueName(const UNiagaraDataInterface* DataInterface) const
+{
+	UNiagaraEmitter* ResolvedEmitter = Resolve(DataInterface);
+	return ResolvedEmitter ? ResolvedEmitter->GetUniqueEmitterName() : FString();
+}
diff --git a/Engine/Plugins/FX/Niagara/Source/Niagara/Private/NiagaraDataInterfaceParticleRead.cpp b/Engine/Plugins/FX/Niagara/Source/Niagara/Private/NiagaraDataInterfaceParticleRead.cpp
--- a/Engine/Plugins/FX/Niagara/Source/Niagara/Private/NiagaraDataInterfaceParticleRead.cpp
+++ b/Engine/Plugins/FX/Niagara/Source/Niagara/Private/NiagaraDataInterfaceParticleRead.cpp
@@ -396,6 +396,22 @@ void UNiagaraDataInterfaceParticleRead::PostInitProperties()
 	}
 }
 
+void UNiagaraDataInterfaceParticleRead::PostLoad()
+{
+	Super::PostLoad();
+
+#if WITH_EDITORONLY_DATA
+	if (!EmitterName_DEPRECATED.IsEmpty())
+	{
+		// We aren't going to search to see if this should be self, as we don't know the user intention and it could be confusing if we switched to this mode
+		EmitterBinding.BindingMode = ENiagaraDataInterfaceEmitterBindingMode::Other;
+		EmitterBinding.EmitterName = FName(*EmitterName_DEPRECATED);
+
+		EmitterName_DEPRECATED.Empty();
+	}
+#endif
+}
+
 #if WITH_EDITOR
 void UNiagaraDataInterfaceParticleRead::PostEditChangeProperty(struct FPropertyChangedEvent& PropertyChangedEvent)
 {
@@ -407,7 +423,7 @@ void UNiagaraDataInterfaceParticleRead::PostEditChangeProperty(struct FPropertyC
 		PropertyName = PropertyChangedEvent.Property->GetFName();
 	}
 
-	if (PropertyName == GET_MEMBER_NAME_CHECKED(UNiagaraDataInterfaceParticleRead, EmitterName))
+	if (PropertyName == GET_MEMBER_NAME_CHECKED(UNiagaraDataInterfaceParticleRead, EmitterBinding))
 	{
 		UNiagaraSystem::RecomputeExecutionOrderForDataInterface(this);
 	}
@@ -416,7 +432,7 @@ void UNiagaraDataInterfaceParticleRead::PostEditChangeProperty(struct FPropertyC
 
 bool UNiagaraDataInterfaceParticleRead::HasInternalAttributeReads(const UNiagaraEmitter* OwnerEmitter, const UNiagaraEmitter* Provider) const
 {
-	return Provider && Provider->GetUniqueEmitterName() == EmitterName;
+	return Provider == EmitterBinding.Resolve(this);
 }
 
 bool UNiagaraDataInterfaceParticleRead::InitPerInstanceData(void* PerInstanceData, FNiagaraSystemInstance* SystemInstance)
@@ -424,13 +440,17 @@ bool UNiagaraDataInterfaceParticleRead::InitPerInstanceData(void* PerInstanceDat
 	FNDIParticleRead_InstanceData* PIData = new (PerInstanceData) FNDIParticleRead_InstanceData;
 	PIData->SystemInstance = SystemInstance;
 	PIData->EmitterInstance = nullptr;
-	for (TSharedPtr<FNiagaraEmitterInstance, ESPMode::ThreadSafe> EmitterInstance : SystemInstance->GetEmitters())
+
+	if ( UNiagaraEmitter* OwnerEmitter = EmitterBinding.Resolve(this) )
 	{
-		const FVersionedNiagaraEmitter& CachedEmitter = EmitterInstance->GetCachedEmitter();
-		if (CachedEmitter.Emitter && (EmitterName == CachedEmitter.Emitter->GetUniqueEmitterName()))
+		for (TSharedPtr<FNiagaraEmitterInstance, ESPMode::ThreadSafe> EmitterInstance : SystemInstance->GetEmitters())
 		{
-			PIData->EmitterInstance = EmitterInstance.Get();
-			break;
+			const FVersionedNiagaraEmitter& CachedEmitter = EmitterInstance->GetCachedEmitter();
+			if (OwnerEmitter == CachedEmitter.Emitter)
+			{
+				PIData->EmitterInstance = EmitterInstance.Get();
+				break;
+			}
 		}
 	}
 
@@ -438,14 +458,14 @@ bool UNiagaraDataInterfaceParticleRead::InitPerInstanceData(void* PerInstanceDat
 	{
 		if (FNiagaraUtilities::LogVerboseWarnings())
 		{
-			UE_LOG(LogNiagara, Warning, TEXT("Source emitter '%s' not found. System: %s"), *EmitterName, *GetFullNameSafe(SystemInstance->GetSystem()));
+			UE_LOG(LogNiagara, Warning, TEXT("Source emitter '%s' not found. System: %s"), *EmitterBinding.ResolveUniqueName(this), *GetFullNameSafe(SystemInstance->GetSystem()));
 		}
 	}
 
 	FString DebugSourceName;
 	if ( UNiagaraSystem* NiagaraSystem = PIData->SystemInstance->GetSystem() )
 	{
-		DebugSourceName = FString::Printf(TEXT("%s.%s"), *NiagaraSystem->GetName(), *EmitterName);
+		DebugSourceName = FString::Printf(TEXT("%s.%s"), *NiagaraSystem->GetName(), *EmitterBinding.ResolveUniqueName(this));
 	}
 
 	FNiagaraDataInterfaceProxyParticleRead* ThisProxy = GetProxyAs<FNiagaraDataInterfaceProxyParticleRead>();
@@ -1175,7 +1195,7 @@ void UNiagaraDataInterfaceParticleRead::GetVMExternalFunction(const FVMExternalF
 	{
 		UNiagaraSystem* NiagaraSystem = PIData->SystemInstance ? PIData->SystemInstance->GetSystem() : nullptr;
 		UNiagaraEmitter* NiagaraEmitter = PIData->EmitterInstance ? PIData->EmitterInstance->GetCachedEmitter().Emitter : nullptr;
-		UE_LOG(LogNiagara, Warning, TEXT("ParticleRead: Failed to '%s' attribute '%s' System '%s' Emitter '%s'! Check that the attribute is named correctly."), *BindingInfo.Name.ToString(), *AttributeToRead.ToString(), *GetNameSafe(NiagaraSystem), *EmitterName);
+		UE_LOG(LogNiagara, Warning, TEXT("ParticleRead: Failed to '%s' attribute '%s' System '%s' Emitter '%s'! Check that the attribute is named correctly."), *BindingInfo.Name.ToString(), *AttributeToRead.ToString(), *GetNameSafe(NiagaraSystem), *EmitterBinding.ResolveUniqueName(this));
 	}
 }
 
@@ -1668,7 +1688,7 @@ bool UNiagaraDataInterfaceParticleRead::Equals(const UNiagaraDataInterface* Othe
 	{
 		return false;
 	}
-	return CastChecked<UNiagaraDataInterfaceParticleRead>(Other)->EmitterName == EmitterName;
+	return CastChecked<UNiagaraDataInterfaceParticleRead>(Other)->EmitterBinding == EmitterBinding;
 }
 
 bool UNiagaraDataInterfaceParticleRead::CopyToInternal(UNiagaraDataInterface* Destination) const
@@ -1677,7 +1697,7 @@ bool UNiagaraDataInterfaceParticleRead::CopyToInternal(UNiagaraDataInterface* De
 	{
 		return false;
 	}
-	CastChecked<UNiagaraDataInterfaceParticleRead>(Destination)->EmitterName = EmitterName;
+	CastChecked<UNiagaraDataInterfaceParticleRead>(Destination)->EmitterBinding = EmitterBinding;
 	return true;
 }
 
@@ -2369,21 +2389,16 @@ void UNiagaraDataInterfaceParticleRead::GetFeedback(UNiagaraSystem* Asset, UNiag
 	}
 
 	FVersionedNiagaraEmitter FoundSourceEmitter;
-	for (const FNiagaraEmitterHandle& EmitterHandle : Asset->GetEmitterHandles())
+	if (const FNiagaraEmitterHandle* EmitterHandle = EmitterBinding.ResolveHandle(this))
 	{
-		FVersionedNiagaraEmitter EmitterInstance = EmitterHandle.GetInstance();
-		if (EmitterInstance.Emitter && EmitterInstance.Emitter->GetUniqueEmitterName() == EmitterName)
-		{
-			FoundSourceEmitter = EmitterInstance;
-			break; 
-		}
+		FoundSourceEmitter = EmitterHandle->GetInstance();
 	}
 
 	if (!FoundSourceEmitter.Emitter)
 	{
 		Warnings.Emplace(
 			LOCTEXT("SourceEmitterNotFound", "Source emitter was not found."),
-			FText::Format(LOCTEXT("SourceEmitterNotFoundSummary", "Source emitter '{0}' could not be found"), FText::FromString(EmitterName)),
+			FText::Format(LOCTEXT("SourceEmitterNotFoundSummary", "Source emitter '{0}' could not be found"), FText::FromString(EmitterBinding.ResolveUniqueName(this))),
 			FNiagaraDataInterfaceFix()
 		);
 	}
@@ -2551,22 +2566,17 @@ void UNiagaraDataInterfaceParticleRead::GetEmitterDependencies(UNiagaraSystem* A
 		return;
 	}
 
-	for (const FNiagaraEmitterHandle& EmitterHandle : Asset->GetEmitterHandles())
+	if (const FNiagaraEmitterHandle* ResolvedEmitterHandle = EmitterBinding.ResolveHandle(this))
 	{
-		FVersionedNiagaraEmitter EmitterInstance = EmitterHandle.GetInstance();
-		if (EmitterInstance.Emitter && EmitterInstance.Emitter->GetUniqueEmitterName() == EmitterName)
-		{
-			Dependencies.Add(EmitterInstance);
-			return;
-		}
+		Dependencies.Add(ResolvedEmitterHandle->GetInstance());
 	}
 }
 
 void UNiagaraDataInterfaceParticleRead::GetEmitterReferencesByName(TArray<FString>& EmitterReferences) const
 {
-	if (!EmitterName.IsEmpty())
+	if (const UNiagaraEmitter* ResolvedEmitter = EmitterBinding.Resolve(this))
 	{
-		EmitterReferences.AddUnique(EmitterName);
+		EmitterReferences.AddUnique(ResolvedEmitter->GetUniqueEmitterName());
 	}
 }
 
diff --git a/Engine/Plugins/FX/Niagara/Source/Niagara/Public/NiagaraDataInterfaceEmitterBinding.h b/Engine/Plugins/FX/Niagara/Source/Niagara/Public/NiagaraDataInterfaceEmitterBinding.h
--- a/Engine/Plugins/FX/Niagara/Source/Niagara/Public/NiagaraDataInterfaceEmitterBinding.h
+++ b/Engine/Plugins/FX/Niagara/Source/Niagara/Public/NiagaraDataInterfaceEmitterBinding.h
@@ -6,7 +6,9 @@
 #include "NiagaraDataInterfaceEmitterBinding.generated.h"
 
 class UNiagaraEmitter;
-class UNiagaraSystem;
+struct FNiagaraEmitterHandle;
+class FNiagaraSystemInstance;
+class FNiagaraEmitterInstance;
 
 UENUM()
 enum class ENiagaraDataInterfaceEmitterBindingMode
@@ -29,9 +31,16 @@ struct FNiagaraDataInterfaceEmitterBinding
 	FName EmitterName;
 
 	/** Resolves the emitter binding, this can return nullptr if we failed to resolve */
-	FNiagaraEmitterInstance* Resolve(FNiagaraSystemInstance* SystemInstance, UNiagaraDataInterface* DataInterface);
-	/** Resolves the emitter binding, only returns a valid result when in Named mode as we can not determine the source otherwise. */
-	UNiagaraEmitter* Resolve(UNiagaraSystem* NiagaraSystem);
+	FNiagaraEmitterInstance* Resolve(const FNiagaraSystemInstance* SystemInstance, const UNiagaraDataInterface* DataInterface) const;
+
+	/** Resolves the emitter binding for a data interface. */
+	UNiagaraEmitter* Resolve(const UNiagaraDataInterface* DataInterface) const;
+
+	/** Resolves the binding to an FNiagaraEmitterHandle or nullptr if it's invalid. */
+	const FNiagaraEmitterHandle* ResolveHandle(const UNiagaraDataInterface* DataInterface) const;
+
+	/** Resolves the emitter name */
+	FString ResolveUniqueName(const UNiagaraDataInterface* DataInterface) const;
 
 	bool operator==(const FNiagaraDataInterfaceEmitterBinding& Other) const
 	{
diff --git a/Engine/Plugins/FX/Niagara/Source/NiagaraEditorWidgets/Private/DetailCustomizations/NiagaraDataInterfaceParticleReadDetails.cpp b/Engine/Plugins/FX/Niagara/Source/NiagaraEditorWidgets/Private/DetailCustomizations/NiagaraDataInterfaceParticleReadDetails.cpp
deleted file mode 100644
diff --git a/Engine/Plugins/FX/Niagara/Source/NiagaraEditorWidgets/Private/DetailCustomizations/NiagaraDataInterfaceParticleReadDetails.h b/Engine/Plugins/FX/Niagara/Source/NiagaraEditorWidgets/Private/DetailCustomizations/NiagaraDataInterfaceParticleReadDetails.h
deleted file mode 100644
diff --git a/Engine/Plugins/FX/Niagara/Source/NiagaraEditorWidgets/Private/NiagaraEditorWidgetsModule.cpp b/Engine/Plugins/FX/Niagara/Source/NiagaraEditorWidgets/Private/NiagaraEditorWidgetsModule.cpp
--- a/Engine/Plugins/FX/Niagara/Source/NiagaraEditorWidgets/Private/NiagaraEditorWidgetsModule.cpp
+++ b/Engine/Plugins/FX/Niagara/Source/NiagaraEditorWidgets/Private/NiagaraEditorWidgetsModule.cpp
@@ -7,7 +7,6 @@
 #include "DetailCustomizations/NiagaraDataInterfaceDetails.h"
 #include "DetailCustomizations/NiagaraDataInterfaceGrid2DCollectionDetails.h"
 #include "DetailCustomizations/NiagaraDataInterfaceGrid3DCollectionDetails.h"
-#include "DetailCustomizations/NiagaraDataInterfaceParticleReadDetails.h"
 #include "DetailCustomizations/NiagaraDataInterfaceSkeletalMeshDetails.h"
 #include "DetailCustomizations/NiagaraDataInterfaceStaticMeshDetails.h"
 #include "DetailCustomizations/NiagaraDataInterfaceMeshRendererInfoDetails.h"
@@ -132,7 +131,6 @@ void FNiagaraEditorWidgetsModule::StartupModule()
 	PropertyModule.RegisterCustomClassLayout("NiagaraDataInterfaceVectorCurve", FOnGetDetailCustomizationInstance::CreateStatic(&FNiagaraDataInterfaceVectorCurveDetails::MakeInstance));
 	PropertyModule.RegisterCustomClassLayout("NiagaraDataInterfaceVector4Curve", FOnGetDetailCustomizationInstance::CreateStatic(&FNiagaraDataInterfaceVector4CurveDetails::MakeInstance));
 	PropertyModule.RegisterCustomClassLayout("NiagaraDataInterfaceColorCurve", FOnGetDetailCustomizationInstance::CreateStatic(&FNiagaraDataInterfaceColorCurveDetails::MakeInstance));
-	PropertyModule.RegisterCustomClassLayout("NiagaraDataInterfaceParticleRead", FOnGetDetailCustomizationInstance::CreateStatic(&FNiagaraDataInterfaceParticleReadDetails::MakeInstance));
 	PropertyModule.RegisterCustomClassLayout("NiagaraDataInterfaceSkeletalMesh", FOnGetDetailCustomizationInstance::CreateStatic(&FNiagaraDataInterfaceSkeletalMeshDetails::MakeInstance));
 	PropertyModule.RegisterCustomClassLayout("NiagaraDataInterfaceStaticMesh", FOnGetDetailCustomizationInstance::CreateStatic(&FNiagaraDataInterfaceStaticMeshDetails::MakeInstance));
 	PropertyModule.RegisterCustomClassLayout("NiagaraDataInterfaceGrid2DCollection", FOnGetDetailCustomizationInstance::CreateStatic(&FNiagaraDataInterfaceGrid2DCollectionDetails::MakeInstance));
