Niagara - Unbind the system finished delegate in post unregister components, since it's registered in post register components.  This prevents unwanted references when copy/pasting between level instances.

diff --git a/Engine/Plugins/FX/Niagara/Source/Niagara/Private/NiagaraActor.cpp b/Engine/Plugins/FX/Niagara/Source/Niagara/Private/NiagaraActor.cpp
--- a/Engine/Plugins/FX/Niagara/Source/Niagara/Private/NiagaraActor.cpp
+++ b/Engine/Plugins/FX/Niagara/Source/Niagara/Private/NiagaraActor.cpp
@@ -81,6 +81,17 @@ void ANiagaraActor::PostRegisterAllComponents()
 	}
 }
 
+void ANiagaraActor::PostUnregisterAllComponents()
+{
+	Super::PostUnregisterAllComponents();
+
+	// Clear Notification Delegate
+	if (NiagaraComponent)
+	{
+		NiagaraComponent->OnSystemFinished.RemoveAll(this);
+	}
+}
+
 void ANiagaraActor::SetDestroyOnSystemFinish(bool bShouldDestroyOnSystemFinish)
 {
 	bDestroyOnSystemFinish = bShouldDestroyOnSystemFinish ? 1 : 0;  
diff --git a/Engine/Plugins/FX/Niagara/Source/Niagara/Public/NiagaraActor.h b/Engine/Plugins/FX/Niagara/Source/Niagara/Public/NiagaraActor.h
--- a/Engine/Plugins/FX/Niagara/Source/Niagara/Public/NiagaraActor.h
+++ b/Engine/Plugins/FX/Niagara/Source/Niagara/Public/NiagaraActor.h
@@ -19,6 +19,8 @@ public:
 
 	NIAGARA_API virtual void PostRegisterAllComponents() override;
 
+	NIAGARA_API virtual void PostUnregisterAllComponents() override;
+
 	/** Set true for this actor to self-destruct when the Niagara system finishes, false otherwise */
 	UFUNCTION(BlueprintCallable, Category=NiagaraActor)
 	NIAGARA_API void SetDestroyOnSystemFinish(bool bShouldDestroyOnSystemFinish);
