Added more options & settings to Niagara viewport config, maintaining state across session & instances, like view mode, auto exposure settings, real time etc.

Split viewport settings into
Shared settings saved in config
Per asset settings saved in system editor data (doesn't work for emitters yet)

Location & Rotation & Orbit mode are saved in per asset settings. Post loaded orbit mode into new struct with a new version.
The typical settings (view mode, exposure etc.) in shared settings

diff --git a/Engine/Plugins/FX/Niagara/Source/NiagaraEditor/Private/NiagaraEditorSettings.cpp b/Engine/Plugins/FX/Niagara/Source/NiagaraEditor/Private/NiagaraEditorSettings.cpp
--- a/Engine/Plugins/FX/Niagara/Source/NiagaraEditor/Private/NiagaraEditorSettings.cpp
+++ b/Engine/Plugins/FX/Niagara/Source/NiagaraEditor/Private/NiagaraEditorSettings.cpp
@@ -307,84 +307,84 @@ void UNiagaraEditorSettings::BuildCachedPlaybackSpeeds() const
 
 bool UNiagaraEditorSettings::IsShowGridInViewport() const
 {
-	return bShowGridInViewport;
+	return ViewportSettings.bShowGridInViewport;
 }
 
 void UNiagaraEditorSettings::SetShowGridInViewport(bool bInShowGridInViewport)
 {
-	if (this->bShowGridInViewport != bInShowGridInViewport)
+	if (this->ViewportSettings.bShowGridInViewport != bInShowGridInViewport)
 	{
-		this->bShowGridInViewport = bInShowGridInViewport;
+		this->ViewportSettings.bShowGridInViewport = bInShowGridInViewport;
 		SaveConfig();
 	}
 }
 
 bool UNiagaraEditorSettings::IsShowInstructionsCount() const
 {
-	return bShowInstructionsCount;
+	return ViewportSettings.bShowInstructionsCount;
 }
 
 void UNiagaraEditorSettings::SetShowInstructionsCount(bool bInShowInstructionsCount)
 {
-	if (this->bShowInstructionsCount != bInShowInstructionsCount)
+	if (this->ViewportSettings.bShowInstructionsCount != bInShowInstructionsCount)
 	{
-		this->bShowInstructionsCount = bInShowInstructionsCount;
+		this->ViewportSettings.bShowInstructionsCount = bInShowInstructionsCount;
 		SaveConfig();
 	}
 }
 
 bool UNiagaraEditorSettings::IsShowParticleCountsInViewport() const
 {
-	return bShowParticleCountsInViewport;
+	return ViewportSettings.bShowParticleCountsInViewport;
 }
 
 void UNiagaraEditorSettings::SetShowParticleCountsInViewport(bool bInShowParticleCountsInViewport)
 {
-	if (this->bShowParticleCountsInViewport != bInShowParticleCountsInViewport)
+	if (this->ViewportSettings.bShowParticleCountsInViewport != bInShowParticleCountsInViewport)
 	{
-		this->bShowParticleCountsInViewport = bInShowParticleCountsInViewport;
+		this->ViewportSettings.bShowParticleCountsInViewport = bInShowParticleCountsInViewport;
 		SaveConfig();
 	}
 }
 
 bool UNiagaraEditorSettings::IsShowEmitterExecutionOrder() const
 {
-	return bShowEmitterExecutionOrder;
+	return ViewportSettings.bShowEmitterExecutionOrder;
 }
 
 void UNiagaraEditorSettings::SetShowEmitterExecutionOrder(bool bInShowEmitterExecutionOrder)
 {
-	if (this->bShowEmitterExecutionOrder != bInShowEmitterExecutionOrder)
+	if (this->ViewportSettings.bShowEmitterExecutionOrder != bInShowEmitterExecutionOrder)
 	{
-		this->bShowEmitterExecutionOrder = bInShowEmitterExecutionOrder;
+		this->ViewportSettings.bShowEmitterExecutionOrder = bInShowEmitterExecutionOrder;
 		SaveConfig();		
 	}
 }
 
 bool UNiagaraEditorSettings::IsShowGpuTickInformation() const
 {
-	return bShowGpuTickInformation;
+	return ViewportSettings.bShowGpuTickInformation;
 }
 
 void UNiagaraEditorSettings::SetShowGpuTickInformation(bool bInShowGpuTickInformation)
 {
-	if (bShowGpuTickInformation != bInShowGpuTickInformation)
+	if (ViewportSettings.bShowGpuTickInformation != bInShowGpuTickInformation)
 	{
-		bShowGpuTickInformation = bInShowGpuTickInformation;
+		ViewportSettings.bShowGpuTickInformation = bInShowGpuTickInformation;
 		SaveConfig();
 	}
 }
 
 bool UNiagaraEditorSettings::IsShowMemoryInfo() const
 {
-	return bShowMemoryInfo;
+	return ViewportSettings.bShowMemoryInfo;
 }
 
 void UNiagaraEditorSettings::SetShowMemoryInfo(bool bInShowInfo)
 {
-	if (bShowMemoryInfo != bInShowInfo)
+	if (ViewportSettings.bShowMemoryInfo != bInShowInfo)
 	{
-		bShowMemoryInfo = bInShowInfo;
+		ViewportSettings.bShowMemoryInfo = bInShowInfo;
 		SaveConfig();
 	}
 }
@@ -664,6 +664,12 @@ void UNiagaraEditorSettings::PostEditChangeProperty(FPropertyChangedEvent& Prope
 	}
 }
 
+void UNiagaraEditorSettings::SetViewportSharedSettings(const FNiagaraViewportSharedSettings& InViewportSharedSettings)
+{
+	ViewportSettings = InViewportSharedSettings;
+	SaveConfig();
+}
+
 UNiagaraEditorSettings::FOnNiagaraEditorSettingsChanged& UNiagaraEditorSettings::OnSettingsChanged()
 {
 	return GetMutableDefault<UNiagaraEditorSettings>()->SettingsChangedDelegate;
diff --git a/Engine/Plugins/FX/Niagara/Source/NiagaraEditor/Private/NiagaraSystemEditorData.cpp b/Engine/Plugins/FX/Niagara/Source/NiagaraEditor/Private/NiagaraSystemEditorData.cpp
--- a/Engine/Plugins/FX/Niagara/Source/NiagaraEditor/Private/NiagaraSystemEditorData.cpp
+++ b/Engine/Plugins/FX/Niagara/Source/NiagaraEditor/Private/NiagaraSystemEditorData.cpp
@@ -12,6 +12,7 @@
 #include "EdGraphSchema_NiagaraSystemOverview.h"
 #include "NiagaraConstants.h"
 #include "NiagaraScriptVariable.h"
+#include "NiagaraSettings.h"
 #include "EdGraph/EdGraph.h"
 
 #include UE_INLINE_GENERATED_CPP_BY_NAME(NiagaraSystemEditorData)
@@ -71,6 +72,11 @@ UNiagaraSystemEditorData::UNiagaraSystemEditorData(const FObjectInitializer& Obj
 	PlaybackRangeMin = 0;
 	PlaybackRangeMax = 10;
 	bSystemIsPlaceholder = false;
+	
+	if(const UNiagaraSettings* Settings = GetDefault<UNiagaraSettings>())
+	{
+		SetUseOrbitMode(Settings->bSystemViewportInOrbitMode);
+	}
 }
 
 void UNiagaraSystemEditorData::PostInitProperties()
@@ -440,3 +446,18 @@ bool UNiagaraSystemEditorData::RemoveUserScriptVariable(FNiagaraVariable Variabl
 			return ScriptVariable->Variable == Variable;
 		}) > 0;
 }
+
+const FNiagaraPerAssetViewportSettings& UNiagaraSystemEditorData::GetAssetViewportSettings() const
+{
+	return AssetViewportSettings;
+}
+
+void UNiagaraSystemEditorData::SetAssetViewportSettings(FNiagaraPerAssetViewportSettings InSettings)
+{
+	AssetViewportSettings = InSettings;
+}
+
+void UNiagaraSystemEditorData::SetUseOrbitMode(bool bInUseOrbitMode)
+{
+	AssetViewportSettings.bUseOrbitMode = bInUseOrbitMode;
+}
diff --git a/Engine/Plugins/FX/Niagara/Source/NiagaraEditor/Private/Toolkits/SystemToolkitModes/NiagaraSystemToolkitModeBase.cpp b/Engine/Plugins/FX/Niagara/Source/NiagaraEditor/Private/Toolkits/SystemToolkitModes/NiagaraSystemToolkitModeBase.cpp
--- a/Engine/Plugins/FX/Niagara/Source/NiagaraEditor/Private/Toolkits/SystemToolkitModes/NiagaraSystemToolkitModeBase.cpp
+++ b/Engine/Plugins/FX/Niagara/Source/NiagaraEditor/Private/Toolkits/SystemToolkitModes/NiagaraSystemToolkitModeBase.cpp
@@ -483,7 +483,7 @@ TSharedRef<SDockTab> FNiagaraSystemToolkitModeBase::SpawnTab_Viewport(const FSpa
 
 	if(!Toolkit->Viewport.IsValid())
 	{
-		Toolkit->Viewport = SNew(SNiagaraSystemViewport)
+		Toolkit->Viewport = SNew(SNiagaraSystemViewport, Toolkit->GetSystemViewModel().ToSharedRef())
 			.OnThumbnailCaptured(Toolkit.Get(), &FNiagaraSystemToolkit::OnThumbnailCaptured)
 			.Sequencer(Toolkit->GetSystemViewModel()->GetSequencer());
 	}
diff --git a/Engine/Plugins/FX/Niagara/Source/NiagaraEditor/Private/ViewModels/NiagaraSystemViewModel.cpp b/Engine/Plugins/FX/Niagara/Source/NiagaraEditor/Private/ViewModels/NiagaraSystemViewModel.cpp
--- a/Engine/Plugins/FX/Niagara/Source/NiagaraEditor/Private/ViewModels/NiagaraSystemViewModel.cpp
+++ b/Engine/Plugins/FX/Niagara/Source/NiagaraEditor/Private/ViewModels/NiagaraSystemViewModel.cpp
@@ -1054,6 +1054,8 @@ void FNiagaraSystemViewModel::NotifyPreSave()
 		System->RequestCompile(false);
 		System->WaitForCompilationComplete(true);
 	}
+
+	OnPreSaveDelegate.Broadcast();
 }
 
 void FNiagaraSystemViewModel::NotifyPostSave()
@@ -1078,6 +1080,11 @@ void FNiagaraSystemViewModel::NotifyPreClose()
 	OnPreCloseDelegate.Broadcast();
 }
 
+FSimpleMulticastDelegate& FNiagaraSystemViewModel::OnPreSave()
+{
+	return OnPreSaveDelegate;
+}
+
 FNiagaraSystemViewModel::FOnPreClose& FNiagaraSystemViewModel::OnPreClose()
 {
 	return OnPreCloseDelegate;
diff --git a/Engine/Plugins/FX/Niagara/Source/NiagaraEditor/Private/Widgets/SNiagaraSystemViewport.cpp b/Engine/Plugins/FX/Niagara/Source/NiagaraEditor/Private/Widgets/SNiagaraSystemViewport.cpp
--- a/Engine/Plugins/FX/Niagara/Source/NiagaraEditor/Private/Widgets/SNiagaraSystemViewport.cpp
+++ b/Engine/Plugins/FX/Niagara/Source/NiagaraEditor/Private/Widgets/SNiagaraSystemViewport.cpp
@@ -40,14 +40,17 @@
 class UNiagaraSystemEditorData;
 
 /** Viewport Client for the preview viewport */
-class FNiagaraSystemViewportClient : public FEditorViewportClient
+class FNiagaraSystemViewportClient : public FEditorViewportClient, public TSharedFromThis<FNiagaraSystemViewportClient>
 {
 public:
 	DECLARE_DELEGATE_OneParam(FOnScreenShotCaptured, UTexture2D*);
 
 public:
 	FNiagaraSystemViewportClient(FAdvancedPreviewScene& InPreviewScene, const TSharedRef<SNiagaraSystemViewport>& InNiagaraEditorViewport, FOnScreenShotCaptured InOnScreenShotCaptured);
+	virtual ~FNiagaraSystemViewportClient() override;
 
+	void InitDelegates();
+	
 	// FEditorViewportClient interface
 	virtual FLinearColor GetBackgroundColor() const override;
 	virtual void Tick(float DeltaSeconds) override;
@@ -76,6 +79,12 @@ public:
 
 	FOnScreenShotCaptured OnScreenShotCaptured;
 
+private:
+	void SaveSharedViewportSettingsToConfig() const;
+	void LoadSharedSettingsFromConfig();
+	void SavePerAssetViewportSettings() const;
+	void LoadPerAssetViewportSettings();
+
 private:
 	bool bUpdateViewportFocus = false;
 
@@ -89,7 +98,7 @@ FNiagaraSystemViewportClient::FNiagaraSystemViewportClient(FAdvancedPreviewScene
 	, OnScreenShotCaptured(InOnScreenShotCaptured)
 {
 	NiagaraViewportPtr = InNiagaraEditorViewport;
-
+	
 	// Setup defaults for the common draw helper.
 	DrawHelper.bDrawPivot = false;
 	DrawHelper.bDrawWorldBox = false;
@@ -106,15 +115,35 @@ FNiagaraSystemViewportClient::FNiagaraSystemViewportClient(FAdvancedPreviewScene
 	EngineShowFlags.SetSnap(0);
 
 	OverrideNearClipPlane(1.0f);
-	SetOrbitModeFromSettings();
+
 	bCaptureScreenShot = false;
 
 	bDrawAxesGame = true;
 
 	//This seems to be needed to get the correct world time in the preview.
 	FNiagaraSystemViewportClient::SetIsSimulateInEditorViewport(true);
+
+	LoadSharedSettingsFromConfig();
+	LoadPerAssetViewportSettings();
 }
 
+FNiagaraSystemViewportClient::~FNiagaraSystemViewportClient()
+{
+	SaveSharedViewportSettingsToConfig();
+
+	if(NiagaraViewportPtr.IsValid() && NiagaraViewportPtr.Pin()->GetSystemViewModel().IsValid())
+	{
+		NiagaraViewportPtr.Pin()->GetSystemViewModel().Pin()->OnPreSave().RemoveAll(this);
+	}
+}
+
+void FNiagaraSystemViewportClient::InitDelegates()
+{
+	if(NiagaraViewportPtr.IsValid() && NiagaraViewportPtr.Pin()->GetSystemViewModel().IsValid())
+	{
+		NiagaraViewportPtr.Pin()->GetSystemViewModel().Pin()->OnPreSave().AddSP(this, &FNiagaraSystemViewportClient::SavePerAssetViewportSettings);
+	}
+}
 
 void FNiagaraSystemViewportClient::Tick(float DeltaSeconds)
 {
@@ -490,9 +519,9 @@ void FNiagaraSystemViewportClient::SetOrbitModeFromSettings()
 	check(Settings);
 
 	UNiagaraSystemEditorData* EditorData = GetSystemEditorData();
-	if (EditorData && EditorData->bSetOrbitModeByAsset)
+	if (EditorData)
 	{
-		bUsingOrbitCamera = EditorData->bSystemViewportInOrbitMode;
+		bUsingOrbitCamera = EditorData->GetAssetViewportSettings().bUseOrbitMode;
 	}
 	else
 	{
@@ -551,10 +580,103 @@ UNiagaraSystemEditorData* FNiagaraSystemViewportClient::GetSystemEditorData() co
 	return nullptr;
 }
 
+void FNiagaraSystemViewportClient::SaveSharedViewportSettingsToConfig() const
+{
+	if(GUnrealEd && GetDefault<UNiagaraEditorSettings>())
+	{
+		const FEngineShowFlags& EditorShowFlagsToSave = IsInGameView() ? LastEngineShowFlags : EngineShowFlags;
+		const FEngineShowFlags& GameShowFlagsToSave = IsInGameView() ? EngineShowFlags : LastEngineShowFlags;
+		
+		FNiagaraViewportSharedSettings ViewportSharedSettings = GetDefault<UNiagaraEditorSettings>()->GetViewportSharedSettings();
+		ViewportSharedSettings.ViewportType = ViewportType;
+		ViewportSharedSettings.PerspViewModeIndex = GetPerspViewMode();
+		ViewportSharedSettings.OrthoViewModeIndex = GetOrthoViewMode();
+		ViewportSharedSettings.EditorShowFlagsString = EditorShowFlagsToSave.ToString();
+		ViewportSharedSettings.GameShowFlagsString = GameShowFlagsToSave.ToString();
+		ViewportSharedSettings.ExposureSettings = ExposureSettings;
+		ViewportSharedSettings.FOVAngle = FOVAngle;
+		ViewportSharedSettings.bShowOnScreenStats = ShouldShowStats();
+		SaveRealtimeStateToConfig(ViewportSharedSettings.bIsRealtime);
+		
+		GetMutableDefault<UNiagaraEditorSettings>()->SetViewportSharedSettings(ViewportSharedSettings);
+	}
+}
+
+void FNiagaraSystemViewportClient::LoadSharedSettingsFromConfig()
+{
+	const FNiagaraViewportSharedSettings& ViewportSharedSettings = GetDefault<UNiagaraEditorSettings>()->GetViewportSharedSettings();
+
+	FEngineShowFlags EditorShowFlags(ESFIM_Editor);
+	FEngineShowFlags GameShowFlags(ESFIM_Game);
+
+	if(!ViewportSharedSettings.EditorShowFlagsString.IsEmpty())
+	{
+		EditorShowFlags.SetFromString(*ViewportSharedSettings.EditorShowFlagsString);
+	}
+
+	if(!ViewportSharedSettings.GameShowFlagsString.IsEmpty())
+	{
+		GameShowFlags.SetFromString(*ViewportSharedSettings.GameShowFlagsString);
+	}
+
+	if(ViewportSharedSettings.ViewportType == LVT_Perspective)
+	{
+		ApplyViewMode(ViewportSharedSettings.PerspViewModeIndex, true, EditorShowFlags);
+		ApplyViewMode(ViewportSharedSettings.PerspViewModeIndex, true, GameShowFlags);
+	}
+	else
+	{
+		ApplyViewMode(ViewportSharedSettings.OrthoViewModeIndex, false, EditorShowFlags);
+		ApplyViewMode(ViewportSharedSettings.OrthoViewModeIndex, false, GameShowFlags);
+	}
+
+	// Disabling some features for orthographic views. 
+	if(ViewportSharedSettings.ViewportType != LVT_Perspective)
+	{
+		EditorShowFlags.MotionBlur = 0;
+		EditorShowFlags.Fog = 0;
+		EditorShowFlags.SetDepthOfField(false);
+		GameShowFlags.MotionBlur = 0;
+		GameShowFlags.Fog = 0;
+		GameShowFlags.SetDepthOfField(false);
+	}
+	
+	ViewportType = ViewportSharedSettings.ViewportType;
+	EngineShowFlags = EditorShowFlags;
+	LastEngineShowFlags = GameShowFlags;
+	ExposureSettings = ViewportSharedSettings.ExposureSettings;
+	SetRealtime(ViewportSharedSettings.bIsRealtime);
+	SetShowStats(ViewportSharedSettings.bShowOnScreenStats);
+	ViewFOV = FOVAngle = ViewportSharedSettings.FOVAngle;
+	SetViewModes(ViewportSharedSettings.PerspViewModeIndex, ViewportSharedSettings.OrthoViewModeIndex );
+}
+
+void FNiagaraSystemViewportClient::SavePerAssetViewportSettings() const
+{
+	if(UNiagaraSystemEditorData* SystemEditorData = GetSystemEditorData())
+	{
+		FNiagaraPerAssetViewportSettings CurrentAssetSettings = SystemEditorData->GetAssetViewportSettings();
+		CurrentAssetSettings.ViewLocation = GetViewLocation();
+		CurrentAssetSettings.ViewRotation = GetViewRotation();
+		SystemEditorData->SetAssetViewportSettings(CurrentAssetSettings);
+	}
+}
+
+void FNiagaraSystemViewportClient::LoadPerAssetViewportSettings()
+{
+	SetOrbitModeFromSettings();
+	
+	const FNiagaraPerAssetViewportSettings& AssetViewportSettings = NiagaraViewportPtr.Pin()->GetSystemViewModel().Pin()->GetEditorData().GetAssetViewportSettings();
+	SetViewLocation(AssetViewportSettings.ViewLocation);
+	SetViewRotation(AssetViewportSettings.ViewRotation);
+}
+
 //////////////////////////////////////////////////////////////////////////
 
-void SNiagaraSystemViewport::Construct(const FArguments& InArgs)
+void SNiagaraSystemViewport::Construct(const FArguments& InArgs, TSharedRef<FNiagaraSystemViewModel> InSystemViewModel)
 {
+	SystemViewModel = InSystemViewModel;
+	
 	const UNiagaraEditorSettings* Settings = GetDefault<UNiagaraEditorSettings>();
 	
 	DrawFlags = 0;
@@ -656,8 +778,7 @@ void SNiagaraSystemViewport::ToggleOrbit()
 		UNiagaraSystemEditorData* EditorData = Cast<UNiagaraSystemEditorData>(PreviewComponent->GetAsset()->GetEditorData());
 		if (EditorData)
 		{
-			EditorData->bSystemViewportInOrbitMode = bNewOrbitSetting;
-			EditorData->bSetOrbitModeByAsset = true;
+			EditorData->SetUseOrbitMode(bNewOrbitSetting);
 		}
 	}
 }
@@ -904,14 +1025,16 @@ TSharedRef<FEditorViewportClient> SNiagaraSystemViewport::MakeEditorViewportClie
 	SystemViewportClient = MakeShareable( new FNiagaraSystemViewportClient(*AdvancedPreviewScene.Get(), SharedThis(this),
 		FNiagaraSystemViewportClient::FOnScreenShotCaptured::CreateSP(this, &SNiagaraSystemViewport::OnScreenShotCaptured) ) );
 
-	SystemViewportClient->SetViewLocation( FVector::ZeroVector );
-	SystemViewportClient->SetViewRotation( FRotator::ZeroRotator );
+	SystemViewportClient->InitDelegates();
+	
+	// Settings that aren't kept in the config are 
+	SystemViewportClient->SetAllowCinematicControl(true);
+	// SystemViewportClient->SetViewLocation( FVector::ZeroVector );
+	// SystemViewportClient->SetViewRotation( FRotator::ZeroRotator );
 	SystemViewportClient->SetViewLocationForOrbiting( FVector::ZeroVector );
-	SystemViewportClient->bSetListenerPosition = false;
-
-	SystemViewportClient->SetRealtime( true );
 	SystemViewportClient->VisibilityDelegate.BindSP( this, &SNiagaraSystemViewport::IsVisible );
-
+	SystemViewportClient->bSetListenerPosition = false;
+	
 	return SystemViewportClient.ToSharedRef();
 }
 
diff --git a/Engine/Plugins/FX/Niagara/Source/NiagaraEditor/Private/Widgets/SNiagaraSystemViewport.h b/Engine/Plugins/FX/Niagara/Source/NiagaraEditor/Private/Widgets/SNiagaraSystemViewport.h
--- a/Engine/Plugins/FX/Niagara/Source/NiagaraEditor/Private/Widgets/SNiagaraSystemViewport.h
+++ b/Engine/Plugins/FX/Niagara/Source/NiagaraEditor/Private/Widgets/SNiagaraSystemViewport.h
@@ -16,6 +16,7 @@
 #include "Particles/ParticlePerfStatsManager.h"
 #include "NiagaraPerfBaseline.h"
 
+class FNiagaraSystemViewModel;
 class UNiagaraComponent;
 class FNiagaraSystemEditorViewportClient;
 class FNiagaraSystemInstance;
@@ -36,7 +37,7 @@ public:
 		SLATE_ARGUMENT(TWeakPtr<ISequencer>, Sequencer)
 	SLATE_END_ARGS()
 	
-	void Construct(const FArguments& InArgs);
+	void Construct(const FArguments& InArgs, TSharedRef<FNiagaraSystemViewModel> InSystemViewModel);
 	~SNiagaraSystemViewport();
 	
 	virtual void AddReferencedObjects( FReferenceCollector& Collector ) override;
@@ -57,7 +58,8 @@ public:
 	bool bShowBackground;
 
 	TSharedRef<class FAdvancedPreviewScene> GetPreviewScene() { return AdvancedPreviewScene.ToSharedRef(); }
-
+	TWeakPtr<FNiagaraSystemViewModel> GetSystemViewModel() { return SystemViewModel; }
+	
 	/** The material editor has been added to a tab */
 	void OnAddedToTab( const TSharedRef<SDockTab>& OwnerTab );
 	
@@ -101,7 +103,8 @@ public:
 	float GetMotionRadius() const { return MotionRadius; }
 	void SetMotionRadius(float Radius) { MotionRadius = Radius; }
 
-protected:
+
+protected:	
 	/** SEditorViewport interface */
 	virtual TSharedRef<FEditorViewportClient> MakeEditorViewportClient() override;
 	virtual TSharedPtr<SWidget> MakeViewportToolbar() override;
@@ -113,7 +116,6 @@ protected:
 
 private:
 	bool IsVisible() const override;
-
 	void OnScreenShotCaptured(UTexture2D* ScreenShot);
 
 private:
@@ -121,6 +123,8 @@ private:
 	TWeakPtr<SDockTab> ParentTab;
 
 	TWeakPtr<ISequencer> Sequencer = nullptr;
+
+	TWeakPtr<FNiagaraSystemViewModel> SystemViewModel;
 	
 	/** Preview Scene - uses advanced preview settings */
 	TSharedPtr<class FAdvancedPreviewScene> AdvancedPreviewScene;
diff --git a/Engine/Plugins/FX/Niagara/Source/NiagaraEditor/Public/NiagaraEditorSettings.h b/Engine/Plugins/FX/Niagara/Source/NiagaraEditor/Public/NiagaraEditorSettings.h
--- a/Engine/Plugins/FX/Niagara/Source/NiagaraEditor/Public/NiagaraEditorSettings.h
+++ b/Engine/Plugins/FX/Niagara/Source/NiagaraEditor/Public/NiagaraEditorSettings.h
@@ -2,12 +2,15 @@
 
 #pragma once
 
+#include "CoreMinimal.h"
 #include "NiagaraSystem.h"
 #include "NiagaraEmitter.h"
 #include "NiagaraScript.h"
 #include "NiagaraValidationRuleSet.h"
 #include "Engine/DeveloperSettings.h"
 #include "NiagaraSpawnShortcut.h"
+#include "Editor/UnrealEdTypes.h"
+#include "Viewports.h"
 #include "NiagaraEditorSettings.generated.h"
 
 class UCurveFloat;
@@ -260,6 +263,75 @@ struct FNiagaraParameterPanelSectionStorage
 	TArray<FGuid> ExpandedCategories;
 };
 
+/** Contains all the settings that should be shared across sessions. */
+USTRUCT()
+struct FNiagaraViewportSharedSettings
+{
+	GENERATED_USTRUCT_BODY()
+
+	FNiagaraViewportSharedSettings() {}
+
+	/** The viewport type */
+	UPROPERTY(config)
+	TEnumAsByte<ELevelViewportType> ViewportType = LVT_Perspective;
+
+	/* View mode to set when this viewport is of type LVT_Perspective. */
+	UPROPERTY(config)
+	TEnumAsByte<EViewModeIndex> PerspViewModeIndex = VMI_Lit;
+
+	/* View mode to set when this viewport is not of type LVT_Perspective. */
+	UPROPERTY(config)
+	TEnumAsByte<EViewModeIndex> OrthoViewModeIndex = VMI_BrushWireframe;
+
+	/**
+	 * A set of flags that determines visibility for various scene elements (FEngineShowFlags), converted to string form.
+	 * These have to be saved as strings since FEngineShowFlags is too complex for UHT to parse correctly.
+	 */
+	UPROPERTY(config)
+	FString EditorShowFlagsString;
+
+	/**
+	 * A set of flags that determines visibility for various scene elements (FEngineShowFlags), converted to string form.
+	 * These have to be saved as strings since FEngineShowFlags is too complex for UHT to parse correctly.
+	 */
+	UPROPERTY(config)
+	FString GameShowFlagsString;
+
+	/** Setting to allow designers to override the automatic expose. */
+	UPROPERTY(config)
+	FExposureSettings ExposureSettings;
+
+	/* Field of view angle for the viewport. */
+	UPROPERTY(config)
+	float FOVAngle = EditorViewportDefs::DefaultPerspectiveFOVAngle;
+
+	/* Whether this viewport is updating in real-time. */
+	UPROPERTY(config)
+	bool bIsRealtime = true;
+	
+	/* Whether viewport statistics should be shown. */
+	UPROPERTY(config)
+	bool bShowOnScreenStats = true;
+
+	UPROPERTY(config)
+	bool bShowGridInViewport = false;
+
+	UPROPERTY(config)
+	bool bShowInstructionsCount = false;
+	
+	UPROPERTY(config)
+	bool bShowParticleCountsInViewport = false;
+
+	UPROPERTY(config)
+	bool bShowEmitterExecutionOrder = false;
+
+	UPROPERTY(config)
+	bool bShowGpuTickInformation = false;
+
+	UPROPERTY(config)
+	bool bShowMemoryInfo = false;
+};
+
 FORCEINLINE uint32 GetTypeHash(const FNiagaraNamespaceMetadata& NamespaceMetaData)
 {
 	return GetTypeHash(NamespaceMetaData.GetGuid());
@@ -386,8 +458,12 @@ public:
 	// END UDeveloperSettings Interface
 
 	NIAGARAEDITOR_API virtual void PostEditChangeProperty(FPropertyChangedEvent& PropertyChangedEvent) override;
+	
+	const FNiagaraViewportSharedSettings& GetViewportSharedSettings() const { return ViewportSettings; }
+	void SetViewportSharedSettings(const FNiagaraViewportSharedSettings& InViewportSharedSettings);
 
-	DECLARE_MULTICAST_DELEGATE_TwoParams(FOnNiagaraEditorSettingsChanged, const FString&, const UNiagaraEditorSettings*);
+	DECLARE_MULTICAST_DELEGATE_TwoParams(FOnNiagaraEditorSettingsChanged, const FString&,
+	                                     const UNiagaraEditorSettings*);
 
 	/** Gets a multicast delegate which is called whenever one of the parameters in this settings object changes. */
 	static NIAGARAEDITOR_API FOnNiagaraEditorSettingsChanged& OnSettingsChanged();
@@ -427,6 +503,12 @@ public:
 
 	NIAGARAEDITOR_API bool GetUpdateStackValuesOnCommitOnly() const;
 
+	NIAGARAEDITOR_API bool GetUseAutoExposure() const;
+	NIAGARAEDITOR_API void SetUseAutoExposure(bool bInUseAutoExposure);
+	
+	NIAGARAEDITOR_API float GetExposureValue() const;
+	NIAGARAEDITOR_API void SetAutoExposureValue(bool bInUseAutoExposure);
+	
 private:
 	NIAGARAEDITOR_API bool IsAllowedObjectByClassUsageInternal(const UObject& InObject, TSet<const UObject*>& CheckedObjects) const;
 	NIAGARAEDITOR_API bool IsAllowedAssetObjectByClassUsageInternal(const UObject& InAssetObject, TSet<const UObject*>& CheckedAssetObjects) const;
@@ -503,27 +585,13 @@ private:
 
 	UPROPERTY(config, EditAnywhere, Category = Niagara)
 	TArray<FNiagaraCurveTemplate> CurveTemplates;
-
-	UPROPERTY(config)
-	bool bShowGridInViewport;
-
-	UPROPERTY(config)
-	bool bShowInstructionsCount;
 	
 	UPROPERTY(config)
-	bool bShowParticleCountsInViewport;
-
-	UPROPERTY(config)
-	bool bShowEmitterExecutionOrder;
-
-	UPROPERTY(config)
-	bool bShowGpuTickInformation;
-
-	UPROPERTY(config)
-	bool bShowMemoryInfo;
-
+	FNiagaraViewportSharedSettings ViewportSettings;
+	
 	UPROPERTY(config)
 	TArray<FNiagaraParameterPanelSectionStorage> SystemParameterPanelSectionData;
+	
 
 	FOnIsClassAllowed OnIsClassAllowedDelegate;
 	FOnIsClassPathAllowed OnIsClassPathAllowedDelegate;
diff --git a/Engine/Plugins/FX/Niagara/Source/NiagaraEditor/Public/NiagaraSystemEditorData.h b/Engine/Plugins/FX/Niagara/Source/NiagaraEditor/Public/NiagaraSystemEditorData.h
--- a/Engine/Plugins/FX/Niagara/Source/NiagaraEditor/Public/NiagaraSystemEditorData.h
+++ b/Engine/Plugins/FX/Niagara/Source/NiagaraEditor/Public/NiagaraSystemEditorData.h
@@ -45,6 +45,23 @@ private:
 	TArray<FGuid> ChildEmitterHandleIds;
 };
 
+/** View settings that are saved per asset and aren't shared between different Niagara viewports. */
+USTRUCT()
+struct FNiagaraPerAssetViewportSettings
+{
+	GENERATED_BODY()
+
+public:
+	UPROPERTY()
+	FVector ViewLocation = FVector::ZeroVector;
+
+	UPROPERTY()
+	FRotator ViewRotation = FRotator::ZeroRotator;
+
+	UPROPERTY()
+	bool bUseOrbitMode = true;
+};
+
 /** Editor only data for systems. */
 UCLASS(MinimalAPI)
 class UNiagaraSystemEditorData : public UNiagaraEditorDataBase
@@ -102,12 +119,9 @@ public:
 	NIAGARAEDITOR_API bool RenameUserScriptVariable(FNiagaraVariable OldVariable, FName NewName);
 	NIAGARAEDITOR_API bool RemoveUserScriptVariable(FNiagaraVariable Variable);
 
-	// If true then the preview viewport's orbit setting is saved in the asset data
-	UPROPERTY()
-	bool bSetOrbitModeByAsset = false;
-
-	UPROPERTY()
-	bool bSystemViewportInOrbitMode = true;
+	NIAGARAEDITOR_API const FNiagaraPerAssetViewportSettings& GetAssetViewportSettings() const;
+	NIAGARAEDITOR_API void SetAssetViewportSettings(FNiagaraPerAssetViewportSettings InSettings);
+	NIAGARAEDITOR_API void SetUseOrbitMode(bool bInUseOrbitMode);
 	
 	/** Contains the root ids for organizing user parameters. */
 	UPROPERTY()
@@ -139,6 +153,9 @@ private:
 	UPROPERTY()
 	FNiagaraGraphViewSettings OverviewGraphViewSettings;
 
+	UPROPERTY()
+	FNiagaraPerAssetViewportSettings AssetViewportSettings;
+	
 	UPROPERTY()
 	bool bSystemIsPlaceholder;
 
diff --git a/Engine/Plugins/FX/Niagara/Source/NiagaraEditor/Public/ViewModels/NiagaraSystemViewModel.h b/Engine/Plugins/FX/Niagara/Source/NiagaraEditor/Public/ViewModels/NiagaraSystemViewModel.h
--- a/Engine/Plugins/FX/Niagara/Source/NiagaraEditor/Public/ViewModels/NiagaraSystemViewModel.h
+++ b/Engine/Plugins/FX/Niagara/Source/NiagaraEditor/Public/ViewModels/NiagaraSystemViewModel.h
@@ -352,6 +352,7 @@ public:
 
 	NIAGARAEDITOR_API void NotifyPreClose();
 
+	NIAGARAEDITOR_API FSimpleMulticastDelegate& OnPreSave();
 	NIAGARAEDITOR_API FOnPreClose& OnPreClose();
 
 	FOnRequestFocusTab& OnRequestFocusTab();
@@ -664,6 +665,9 @@ private:
 	/** A multicast delegate which is called whenever the system has been compiled. */
 	FOnSystemCompiled OnSystemCompiledDelegate;
 
+	/** A multicast delegate which is called before the asset is saved. Useful to notify other systems that need to transfer data into the asset before saving is done.. */
+	FSimpleMulticastDelegate OnPreSaveDelegate;
+	
 	/** A multicast delegate which is called whenever this has been notified it's owner will be closing. */
 	FOnPreClose OnPreCloseDelegate;
 
