- Niagara task priorities can be set per tick group
- This helps to reduce wait times if we execute a simulation later and allows us to fit around physics tasks

diff --git a/Engine/Plugins/FX/Niagara/Source/Niagara/Private/NiagaraModule.cpp b/Engine/Plugins/FX/Niagara/Source/Niagara/Private/NiagaraModule.cpp
--- a/Engine/Plugins/FX/Niagara/Source/Niagara/Private/NiagaraModule.cpp
+++ b/Engine/Plugins/FX/Niagara/Source/Niagara/Private/NiagaraModule.cpp
@@ -14,6 +14,7 @@
 #include "NiagaraSettings.h"
 #include "NiagaraDataChannelManager.h"
 #include "NiagaraSimCache.h"
+#include "NiagaraSimulationTaskPriority.h"
 #include "NiagaraUseOpenVDB.h"
 #include "NiagaraSystem.h"
 #include "UObject/Class.h"
@@ -293,6 +294,8 @@ void INiagaraModule::StartupModule()
 
 	FNiagaraWorldManager::OnStartup();
 
+	NiagaraSimulationTaskPriority::Initialize();
+
 #if WITH_NIAGARA_DEBUGGER
 	DebuggerClient = MakePimpl<FNiagaraDebuggerClient>();
 #endif
diff --git a/Engine/Plugins/FX/Niagara/Source/Niagara/Private/NiagaraSimulationTaskPriority.cpp b/Engine/Plugins/FX/Niagara/Source/Niagara/Private/NiagaraSimulationTaskPriority.cpp
new file mode 100644
--- /dev/null
+++ b/Engine/Plugins/FX/Niagara/Source/Niagara/Private/NiagaraSimulationTaskPriority.cpp
@@ -0,0 +1,188 @@
+// Copyright Epic Games, Inc. All Rights Reserved.
+
+#include "NiagaraSimulationTaskPriority.h"
+#include "NiagaraCommon.h"
+#include "HAL/IConsoleManager.h"
+#include "Particles/ParticlePerfStats.h"
+
+namespace NiagaraSimulationTaskPriority
+{
+	// Adjusts the priorty when capturing performance tests to avoid tasks being preempted as much
+	static int32 GPerfCapturePriority = 1;
+	static FAutoConsoleVariableRef CVarPerfCapturePriority(
+		TEXT("fx.Niagara.TaskPriorities.PerfCapturePriority"),
+		GPerfCapturePriority,
+		TEXT("Priority to use if performance captures are enabled.  Reduces the amount of context switching for Niagara to make performance measurements more reliable.")
+		TEXT("Default is 1, set to -1 to not override the default priorities."),
+		ECVF_Default
+	);
+
+	// Post actor tick priority for Niagara tasks
+	static int32 GPostActorTickPriority = 1;
+	static FAutoConsoleVariableRef CVarPostActorTickPriority(
+		TEXT("fx.Niagara.TaskPriorities.PostActorTickPriority"),
+		GPostActorTickPriority,
+		TEXT("Any tasks we execute, such as spawning, in post actor tick will run at this priority."),
+		ECVF_Default | ECVF_Scalability
+	);
+
+	// Translation from a Niagara Priority into a Task Priority
+	static FAutoConsoleTaskPriority GTaskPriorities[] =
+	{
+		//																														Thread Priority (w HiPri Thread)			Task Priority (w HiPri Thread)				Task Priority (w Normal Thread)
+		FAutoConsoleTaskPriority(TEXT("fx.Niagara.TaskPriorities.High"),		TEXT("Task Priority When Set to High"),			ENamedThreads::HighThreadPriority,			ENamedThreads::HighTaskPriority,			ENamedThreads::HighTaskPriority),
+		FAutoConsoleTaskPriority(TEXT("fx.Niagara.TaskPriorities.Normal"),		TEXT("Task Priority When Set to Normal"),		ENamedThreads::HighThreadPriority,			ENamedThreads::NormalTaskPriority,			ENamedThreads::NormalTaskPriority),
+		FAutoConsoleTaskPriority(TEXT("fx.Niagara.TaskPriorities.Low"),			TEXT("Task Priority When Set to Low"),			ENamedThreads::NormalThreadPriority,		ENamedThreads::HighTaskPriority,			ENamedThreads::NormalTaskPriority),
+		FAutoConsoleTaskPriority(TEXT("fx.Niagara.TaskPriorities.Background"),	TEXT("Task Priority When Set to Background"),	ENamedThreads::BackgroundThreadPriority,	ENamedThreads::NormalTaskPriority,			ENamedThreads::NormalTaskPriority),
+	};
+
+	static int32 GTickGroupPriorities[TG_MAX] = {};
+	static FString GTickGroupPrioritiesString;
+	static void UpdateTickGroupPriorities(IConsoleVariable*)
+	{
+		Initialize();
+
+		TArray<FString> TGPairs;
+		GTickGroupPrioritiesString.ParseIntoArray(TGPairs, TEXT(","));
+
+		UEnum* TickingGroupEnum = StaticEnum<ETickingGroup>();
+		check(TickingGroupEnum);
+		for (const FString& TGPair : TGPairs)
+		{
+			int32 SepIndex;
+			TGPair.FindChar(':', SepIndex);
+			if (SepIndex == INDEX_NONE || TGPair.Len() < SepIndex + 1)
+			{
+				continue;
+			}
+			
+			const FString TickGroupString = TGPair.Mid(0, SepIndex);
+			const int32 TickGroup = TickingGroupEnum->GetIndexByNameString(TickGroupString);
+			if (TickGroup >= 0 && TickGroup < TG_MAX)
+			{
+				int32 Priority = FCString::Atoi(*TGPair.Mid(SepIndex + 1));
+				Priority = FMath::Clamp(Priority, 0, UE_ARRAY_COUNT(GTaskPriorities) - 1);
+				GTickGroupPriorities[TickGroup] = Priority;
+			}
+		}
+	}
+
+	static FAutoConsoleVariableRef CVarTickGroupPriorities(
+		TEXT("fx.Niagara.TaskPriorities.TickGroupPriority"),
+		GTickGroupPrioritiesString,
+		TEXT("Set tick group priories for Niagara.")
+		TEXT("For example, TG_PrePhysics:2,TG_DuringPhysics:2"),
+		FConsoleVariableDelegate::CreateStatic(UpdateTickGroupPriorities),
+		ECVF_Scalability | ECVF_Default
+	);
+
+	static FAutoConsoleCommand CCmdDumpPriorities(
+		TEXT("fx.Niagara.TaskPriorities.Dump"),
+		TEXT("Dump currently set priorities"),
+		FConsoleCommandDelegate::CreateLambda(
+			[]()
+			{
+				auto DumpPriority =
+					[](int32 NiagaraPriority, const ENamedThreads::Type TaskThread, const TCHAR* TaskName)
+					{
+						UE_LOG(LogNiagara, Log, TEXT("%s = %d = Thread Priority(%d) Task Priority(%d)"), TaskName, NiagaraPriority, ENamedThreads::GetThreadPriorityIndex(TaskThread), ENamedThreads::GetTaskPriority(TaskThread));
+					};
+
+				UE_LOG(LogNiagara, Log, TEXT("=== Niagara Task Priorities"));
+				DumpPriority(GPostActorTickPriority, GetPostActorTickPriority(), TEXT("PostActorTickPriority"));
+
+				UEnum* TickingGroupEnum = StaticEnum<ETickingGroup>();
+				check(TickingGroupEnum);
+				for (int32 i = 0; i < UE_ARRAY_COUNT(GTickGroupPriorities); ++i)
+				{
+					DumpPriority(
+						GTickGroupPriorities[i],
+						GetTickGroupPriority(ETickingGroup(i)),
+						*TickingGroupEnum->GetNameStringByValue(i)
+					);
+				}
+			}
+		)
+	);
+
+	static FAutoConsoleCommandWithWorldAndArgs CCmdRunTest(
+		TEXT("fx.Niagara.TaskPriorities.RunTest"),
+		TEXT("Run a test set of priorites"),
+		FConsoleCommandWithWorldAndArgsDelegate::CreateLambda(
+			[](const TArray<FString>& Args, UWorld* World)
+			{
+				const int32 TestSet = Args.Num() > 0 ? FCString::Atoi(*Args[0]) : 0;
+				switch (TestSet)
+				{
+					// Current
+					case 0:
+						GPostActorTickPriority = 1;
+						for ( int32& Pri : GTickGroupPriorities )
+						{
+							Pri = 2;
+						}
+						break;
+					case 1:
+						GPostActorTickPriority = 1;
+						GTickGroupPriorities[TG_PrePhysics] = 2;
+						GTickGroupPriorities[TG_StartPhysics] = 2;
+						GTickGroupPriorities[TG_DuringPhysics] = 2;
+						break;
+					case 2:
+						GPostActorTickPriority = 0;
+						GTickGroupPriorities[TG_PrePhysics] = 2;
+						GTickGroupPriorities[TG_StartPhysics] = 2;
+						GTickGroupPriorities[TG_DuringPhysics] = 2;
+						break;
+					case 3:
+						GPostActorTickPriority = 0;
+						GTickGroupPriorities[TG_PrePhysics] = 2;
+						GTickGroupPriorities[TG_StartPhysics] = 2;
+						GTickGroupPriorities[TG_DuringPhysics] = 2;
+						GTickGroupPriorities[TG_LastDemotable] = 0;
+						break;
+				}
+			}
+		)
+	);
+
+	bool GIsInitialized = false;
+	void Initialize()
+	{
+		if ( GIsInitialized )
+		{
+			return;
+		}
+		GIsInitialized = true;
+
+		for ( int32& Pri : GTickGroupPriorities)
+		{
+			Pri = 1;
+		}
+	}
+
+	ENamedThreads::Type GetPostActorTickPriority()
+	{
+		const int32 Priority = FMath::Clamp(GPostActorTickPriority, 0, (int32)UE_ARRAY_COUNT(GTaskPriorities) - 1);
+		return GTaskPriorities[Priority].Get();
+	}
+
+	ENamedThreads::Type GetTickGroupPriority(ETickingGroup TickGroup)
+	{
+		int32 Priority = 0;
+	#if WITH_PARTICLE_PERF_STATS
+		// If we are profiling particle performance make sure we don't get context switched due to lower priority as that will confuse the results
+		// Leave low pri if we're just gathering world stats but for per system or per component stats we should use high pri.
+		if (GPerfCapturePriority >= 0 && (FParticlePerfStats::GetGatherSystemStats() || FParticlePerfStats::GetGatherComponentStats()))
+		{
+			Priority = GPerfCapturePriority;
+		}
+		else
+	#endif
+		{
+			Priority = GTickGroupPriorities[TickGroup];
+		}
+		Priority = FMath::Clamp(Priority, 0, (int32)UE_ARRAY_COUNT(GTaskPriorities) - 1);
+		return GTaskPriorities[Priority].Get();
+	}
+}
diff --git a/Engine/Plugins/FX/Niagara/Source/Niagara/Private/NiagaraSimulationTaskPriority.h b/Engine/Plugins/FX/Niagara/Source/Niagara/Private/NiagaraSimulationTaskPriority.h
new file mode 100644
--- /dev/null
+++ b/Engine/Plugins/FX/Niagara/Source/Niagara/Private/NiagaraSimulationTaskPriority.h
@@ -0,0 +1,11 @@
+// Copyright Epic Games, Inc. All Rights Reserved.
+
+#include "Async/TaskGraphInterfaces.h"
+#include "Engine/EngineBaseTypes.h"
+
+namespace NiagaraSimulationTaskPriority
+{
+	extern void Initialize();
+	extern ENamedThreads::Type GetPostActorTickPriority();
+	extern ENamedThreads::Type GetTickGroupPriority(ETickingGroup TickGroup);
+}
diff --git a/Engine/Plugins/FX/Niagara/Source/Niagara/Private/NiagaraSystemSimulation.cpp b/Engine/Plugins/FX/Niagara/Source/Niagara/Private/NiagaraSystemSimulation.cpp
--- a/Engine/Plugins/FX/Niagara/Source/Niagara/Private/NiagaraSystemSimulation.cpp
+++ b/Engine/Plugins/FX/Niagara/Source/Niagara/Private/NiagaraSystemSimulation.cpp
@@ -5,6 +5,7 @@
 #include "NiagaraModule.h"
 #include "NiagaraTypes.h"
 #include "NiagaraParameterCollection.h"
+#include "NiagaraSimulationTaskPriority.h"
 #include "NiagaraSystemInstance.h"
 #include "NiagaraSystemGpuComputeProxy.h"
 #include "NiagaraConstants.h"
@@ -242,94 +243,6 @@ namespace NiagaraSystemSimulationLocal
 #endif
 }
 
-//////////////////////////////////////////////////////////////////////////
-// Task priorities for simulation tasks
-
-static FAutoConsoleTaskPriority GNiagaraTaskPriorities[] =
-{
-	//																														Thread Priority (w HiPri Thread)			Task Priority (w HiPri Thread)				Task Priority (w Normal Thread)
-	FAutoConsoleTaskPriority(TEXT("fx.Niagara.TaskPriorities.High"),		TEXT("Task Priority When Set to High"),			ENamedThreads::HighThreadPriority,			ENamedThreads::HighTaskPriority,			ENamedThreads::HighTaskPriority),
-	FAutoConsoleTaskPriority(TEXT("fx.Niagara.TaskPriorities.Normal"),		TEXT("Task Priority When Set to Normal"),		ENamedThreads::HighThreadPriority,			ENamedThreads::NormalTaskPriority,			ENamedThreads::NormalTaskPriority),
-	FAutoConsoleTaskPriority(TEXT("fx.Niagara.TaskPriorities.Low"),			TEXT("Task Priority When Set to Low"),			ENamedThreads::NormalThreadPriority,		ENamedThreads::HighTaskPriority,			ENamedThreads::NormalTaskPriority),
-	FAutoConsoleTaskPriority(TEXT("fx.Niagara.TaskPriorities.Background"),	TEXT("Task Priority When Set to Background"),	ENamedThreads::BackgroundThreadPriority,	ENamedThreads::NormalTaskPriority,			ENamedThreads::NormalTaskPriority),
-};
-
-static int32 GNiagaraSystemSimulationSpawnPendingTaskPri = 1;
-static FAutoConsoleVariableRef CVarNiagaraSystemSimulationSpawnPendingTaskPri(
-	TEXT("fx.Niagara.TaskPriority.SystemSimulationSpawnPendingTask"),
-	GNiagaraSystemSimulationSpawnPendingTaskPri,
-	TEXT("Task priority to use for Niagara System Simulation Spawning Pending Task"),
-	ECVF_Default
-);
-
-static int32 GNiagaraSystemSimulationTaskPri = 1;
-static FAutoConsoleVariableRef CVarNiagaraSystemSimulationTaskPri(
-	TEXT("fx.Niagara.TaskPriority.SystemSimulationTask"),
-	GNiagaraSystemSimulationTaskPri,
-	TEXT("Task priority to use for Niagara System Simulation Task"),
-	ECVF_Default
-);
-
-static int32 GNiagaraSystemInstanceTaskPri = 1;
-static FAutoConsoleVariableRef CVarNiagaraSystemInstanceTaskPri(
-	TEXT("fx.Niagara.TaskPriority.SystemInstanceTask"),
-	GNiagaraSystemInstanceTaskPri,
-	TEXT("Task priority to use for Niagara System Instance Task"),
-	ECVF_Default
-);
-
-static int32 GNiagaraSystemSimulationWaitAllTaskPri = 1;
-static FAutoConsoleVariableRef CVarNiagaraSystemSimulationWaitAllTaskPri(
-	TEXT("fx.Niagara.TaskPriority.SystemSimulationWaitAll"),
-	GNiagaraSystemSimulationWaitAllTaskPri,
-	TEXT("Task priority to use for Niagara System Simulation Wait All Task"),
-	ECVF_Default
-);
-
-static int32 GAllowHighPriorityForPerfTests = 1;
-static FAutoConsoleVariableRef CVarAllowHighPriorityForPerfTests(
-	TEXT("fx.Niagara.TaskPriority.AllowHighPriPerfTests"),
-	GAllowHighPriorityForPerfTests,
-	TEXT("Allow Niagara to pump up to high task priority when running performance tests. Reduces the context switching of Niagara tasks but can increase overall frame time when Niagara blocks GT work like Physics."),
-	ECVF_Default
-);
-
-ENamedThreads::Type GetNiagaraTaskPriority(int32 Priority)
-{
-#if WITH_PARTICLE_PERF_STATS
-	// If we are profiling particle performance make sure we don't get context switched due to lower priority as that will confuse the results
-	// Leave low pri if we're just gathering world stats but for per system or per component stats we should use high pri.
-	if (GAllowHighPriorityForPerfTests && (FParticlePerfStats::GetGatherSystemStats() || FParticlePerfStats::GetGatherComponentStats()))
-	{
-		return GNiagaraTaskPriorities[1].Get();
-	}
-#endif
-	Priority = FMath::Clamp(Priority, 0, (int32)UE_ARRAY_COUNT(GNiagaraTaskPriorities) - 1);
-	return GNiagaraTaskPriorities[Priority].Get();
-}
-
-static FAutoConsoleCommand CCmdNiagaraDumpPriorities(
-	TEXT("fx.Niagara.TaskPriority.Dump"),
-	TEXT("Dump currently set priorities"),
-	FConsoleCommandDelegate::CreateLambda(
-		[]()
-		{
-			auto DumpPriority =
-				[](int32 Priority, const TCHAR* TaskName)
-				{
-					const ENamedThreads::Type TaskThread = GetNiagaraTaskPriority(Priority);
-					UE_LOG(LogNiagara, Log, TEXT("%s = %d = Thread Priority(%d) Task Priority(%d)"), TaskName, Priority, ENamedThreads::GetThreadPriorityIndex(TaskThread), ENamedThreads::GetTaskPriority(TaskThread));
-				};
-
-			
-			UE_LOG(LogNiagara, Log, TEXT("=== Niagara Task Priorities"));
-			DumpPriority(GNiagaraSystemSimulationTaskPri, TEXT("NiagaraSystemSimulationTask"));
-			DumpPriority(GNiagaraSystemInstanceTaskPri, TEXT("NiagaraSystemInstanceTask"));
-			DumpPriority(GNiagaraSystemSimulationWaitAllTaskPri, TEXT("NiagaraSystemSimulationWaitAllTask"));
-		}
-	)
-);
-
 //////////////////////////////////////////////////////////////////////////
 
 #if WITH_PER_COMPONENT_PARTICLE_PERF_STATS
@@ -369,12 +282,13 @@ struct FNiagaraSystemSimulationTickConcurrentTask
 	FNiagaraSystemSimulationTickConcurrentTask(FNiagaraSystemSimulationTickContext InContext, FGraphEventRef& CompletionGraphEvent)
 		: Context(InContext)
 	{
+		TaskThread = NiagaraSimulationTaskPriority::GetTickGroupPriority(Context.Owner->GetTickGroup());
 		CompletionTask = TGraphTask<FNiagaraSystemSimulationAllWorkCompleteTask>::CreateTask(nullptr, ENamedThreads::GameThread).ConstructAndHold(Context.CompletionEvents);
 		CompletionGraphEvent = CompletionTask->GetCompletionEvent();
 	}
 
 	FORCEINLINE TStatId GetStatId() const { RETURN_QUICK_DECLARE_CYCLE_STAT(FNiagaraSystemSimulationTickConcurrentTask, STATGROUP_TaskGraphTasks); }
-	ENamedThreads::Type GetDesiredThread() { return GetNiagaraTaskPriority(GNiagaraSystemSimulationTaskPri); }
+	ENamedThreads::Type GetDesiredThread() { return TaskThread; }
 	static ESubsequentsMode::Type GetSubsequentsMode() { return ESubsequentsMode::TrackSubsequents; }
 
 	void DoTask(ENamedThreads::Type CurrentThread, const FGraphEventRef& MyCompletionGraphEvent)
@@ -388,6 +302,7 @@ struct FNiagaraSystemSimulationTickConcurrentTask
 		CompletionTask->Unlock();
 	}
 
+	ENamedThreads::Type TaskThread;
 	FNiagaraSystemSimulationTickContext Context;
 	TGraphTask<FNiagaraSystemSimulationAllWorkCompleteTask>* CompletionTask = nullptr;
 };
@@ -399,12 +314,13 @@ struct FNiagaraSystemSimulationSpawnConcurrentTask
 	FNiagaraSystemSimulationSpawnConcurrentTask(FNiagaraSystemSimulationTickContext InContext, FGraphEventRef& CompletionGraphEvent)
 		: Context(InContext)
 	{
+		TaskThread = NiagaraSimulationTaskPriority::GetPostActorTickPriority();
 		CompletionTask = TGraphTask<FNiagaraSystemSimulationAllWorkCompleteTask>::CreateTask(nullptr, ENamedThreads::GameThread).ConstructAndHold(Context.CompletionEvents);
 		CompletionGraphEvent = CompletionTask->GetCompletionEvent();
 	}
 
 	FORCEINLINE TStatId GetStatId() const { RETURN_QUICK_DECLARE_CYCLE_STAT(FNiagaraSystemSimulationSpawnConcurrentTask, STATGROUP_TaskGraphTasks); }
-	ENamedThreads::Type GetDesiredThread() { return GetNiagaraTaskPriority(GNiagaraSystemSimulationSpawnPendingTaskPri); }
+	ENamedThreads::Type GetDesiredThread() { return TaskThread; }
 	static ESubsequentsMode::Type GetSubsequentsMode() { return ESubsequentsMode::TrackSubsequents; }
 
 	void DoTask(ENamedThreads::Type CurrentThread, const FGraphEventRef& MyCompletionGraphEvent)
@@ -420,6 +336,7 @@ struct FNiagaraSystemSimulationSpawnConcurrentTask
 		CompletionTask->Unlock();
 	}
 
+	ENamedThreads::Type TaskThread;
 	FNiagaraSystemSimulationTickContext Context;
 	TGraphTask<FNiagaraSystemSimulationAllWorkCompleteTask>* CompletionTask = nullptr;
 };
@@ -433,10 +350,11 @@ struct FNiagaraSystemInstanceTickConcurrentTask
 		, Batch(InBatch)
 		, WorldContext(InWorldContext)
 	{
+		TaskThread = NiagaraSimulationTaskPriority::GetTickGroupPriority(SystemSimulation->GetTickGroup());
 	}
 
 	FORCEINLINE TStatId GetStatId() const { RETURN_QUICK_DECLARE_CYCLE_STAT(FNiagaraSystemInstanceTickConcurrentTask, STATGROUP_TaskGraphTasks); }
-	ENamedThreads::Type GetDesiredThread() { return GetNiagaraTaskPriority(GNiagaraSystemInstanceTaskPri); }
+	ENamedThreads::Type GetDesiredThread() { return TaskThread; }
 	static ESubsequentsMode::Type GetSubsequentsMode() { return ESubsequentsMode::TrackSubsequents; }
 
 	void DoTask(ENamedThreads::Type CurrentThread, const FGraphEventRef& MyCompletionGraphEvent)
@@ -485,6 +403,7 @@ struct FNiagaraSystemInstanceTickConcurrentTask
 		}
 	}
 
+	ENamedThreads::Type TaskThread;
 	FNiagaraSystemSimulation* SystemSimulation = nullptr;
 	FNiagaraSystemTickBatch Batch;
 	UWorld* WorldContext = nullptr;
