- Add missing pragma once

#rnx

diff --git a/Engine/Plugins/FX/Niagara/Source/Niagara/Private/NiagaraSimulationTaskPriority.h b/Engine/Plugins/FX/Niagara/Source/Niagara/Private/NiagaraSimulationTaskPriority.h
--- a/Engine/Plugins/FX/Niagara/Source/Niagara/Private/NiagaraSimulationTaskPriority.h
+++ b/Engine/Plugins/FX/Niagara/Source/Niagara/Private/NiagaraSimulationTaskPriority.h
@@ -1,5 +1,7 @@
 // Copyright Epic Games, Inc. All Rights Reserved.
 
+#pragma once
+
 #include "Async/TaskGraphInterfaces.h"
 #include "Engine/EngineBaseTypes.h"
 
