- Remove some unused function

#rnx

diff --git a/Engine/Plugins/FX/Niagara/Source/Niagara/Private/NiagaraRenderer.cpp b/Engine/Plugins/FX/Niagara/Source/Niagara/Private/NiagaraRenderer.cpp
--- a/Engine/Plugins/FX/Niagara/Source/Niagara/Private/NiagaraRenderer.cpp
+++ b/Engine/Plugins/FX/Niagara/Source/Niagara/Private/NiagaraRenderer.cpp
@@ -204,27 +204,6 @@ FRHIShaderResourceView* FNiagaraRenderer::GetDummyUInt4Buffer()
 	return DummyUInt4Buffer.SRV;
 }
 
-FRHIShaderResourceView* FNiagaraRenderer::GetDummyTextureReadBuffer2D()
-{
-	check(IsInRenderingThread());
-	static TGlobalResource<FNiagaraEmptyTextureSRV> DummyTextureReadBuffer2D(PF_R32_FLOAT, TEXT("NiagaraRenderer::DummyTextureReadBuffer2D"), ETextureDimension::Texture2D);
-	return DummyTextureReadBuffer2D.SRV;
-}
-
-FRHIShaderResourceView* FNiagaraRenderer::GetDummyTextureReadBuffer2DArray()
-{
-	check(IsInRenderingThread());
-	static TGlobalResource<FNiagaraEmptyTextureSRV> DummyTextureReadBuffer2DArray(PF_R32_FLOAT, TEXT("NiagaraRenderer::DummyTextureReadBuffer2DArray"), ETextureDimension::Texture2DArray);
-	return DummyTextureReadBuffer2DArray.SRV;
-}
-
-FRHIShaderResourceView* FNiagaraRenderer::GetDummyTextureReadBuffer3D()
-{
-	check(IsInRenderingThread());
-	static TGlobalResource<FNiagaraEmptyTextureSRV> DummyTextureReadBuffer3D(PF_R32_FLOAT, TEXT("NiagaraRenderer::DummyTextureReadBuffer3D"), ETextureDimension::Texture3D);
-	return DummyTextureReadBuffer3D.SRV;
-}
-
 FRHIShaderResourceView* FNiagaraRenderer::GetDummyHalfBuffer()
 {
 	return DummyHalfBuffer.SRV;
diff --git a/Engine/Plugins/FX/Niagara/Source/Niagara/Public/NiagaraRenderer.h b/Engine/Plugins/FX/Niagara/Source/Niagara/Public/NiagaraRenderer.h
--- a/Engine/Plugins/FX/Niagara/Source/Niagara/Public/NiagaraRenderer.h
+++ b/Engine/Plugins/FX/Niagara/Source/Niagara/Public/NiagaraRenderer.h
@@ -131,9 +131,6 @@ public:
 	static NIAGARA_API FRHIShaderResourceView* GetDummyUIntBuffer();
 	static NIAGARA_API FRHIShaderResourceView* GetDummyUInt2Buffer();
 	static NIAGARA_API FRHIShaderResourceView* GetDummyUInt4Buffer();
-	static NIAGARA_API FRHIShaderResourceView* GetDummyTextureReadBuffer2D();
-	static NIAGARA_API FRHIShaderResourceView* GetDummyTextureReadBuffer2DArray();
-	static NIAGARA_API FRHIShaderResourceView* GetDummyTextureReadBuffer3D();
 	static NIAGARA_API FRHIShaderResourceView* GetDummyHalfBuffer();
 
 	FORCEINLINE ENiagaraSimTarget GetSimTarget() const { return SimTarget; }
