- Niagara Landscape DI fix for uint vs int for CachedPhysMatTexture

#rnx

diff --git a/Engine/Plugins/FX/Niagara/Shaders/Private/NiagaraDataInterfaceLandscape.ush b/Engine/Plugins/FX/Niagara/Shaders/Private/NiagaraDataInterfaceLandscape.ush
--- a/Engine/Plugins/FX/Niagara/Shaders/Private/NiagaraDataInterfaceLandscape.ush
+++ b/Engine/Plugins/FX/Niagara/Shaders/Private/NiagaraDataInterfaceLandscape.ush
@@ -48,7 +48,7 @@ float4x4			{ParameterName}_CachedHeightTextureUvToWorldTransform;
 float4				{ParameterName}_CachedHeightTextureUvScaleBias;
 float2				{ParameterName}_CachedHeightTextureGridSize;
 SamplerState		{ParameterName}_PointClampedSampler;
-Texture2D<int>		{ParameterName}_CachedPhysMatTexture;
+Texture2D<uint>		{ParameterName}_CachedPhysMatTexture;
 int2				{ParameterName}_CachedPhysMatTextureDimension;
 
 void GetBaseColor_{ParameterName}(float3 InWorldPos, out float3 OutBaseColor, out bool OutIsValid)
@@ -217,7 +217,7 @@ void GetPhysicalMaterialIndex_{ParameterName}(float3 InWorldPos, out int OutInde
 	Uv = Uv * {ParameterName}_CachedHeightTextureUvScaleBias.xy + {ParameterName}_CachedHeightTextureUvScaleBias.zw;
 
 	int3 SampleIndex = int3(Uv * {ParameterName}_CachedPhysMatTextureDimension, 0);
-	OutIndex = {ParameterName}_CachedPhysMatTexture.Load(SampleIndex).x;
+	OutIndex = (int){ParameterName}_CachedPhysMatTexture.Load(SampleIndex).x;
 
 	if (OutIndex >= 255)
 	{
diff --git a/Engine/Plugins/FX/Niagara/Source/Niagara/Private/NiagaraDataInterfaceLandscape.cpp b/Engine/Plugins/FX/Niagara/Source/Niagara/Private/NiagaraDataInterfaceLandscape.cpp
--- a/Engine/Plugins/FX/Niagara/Source/Niagara/Private/NiagaraDataInterfaceLandscape.cpp
+++ b/Engine/Plugins/FX/Niagara/Source/Niagara/Private/NiagaraDataInterfaceLandscape.cpp
@@ -85,7 +85,7 @@ namespace NiagaraDataInterfaceLandscape
 		SHADER_PARAMETER(FVector4f,					CachedHeightTextureUvScaleBias)
 		SHADER_PARAMETER(FVector2f,					CachedHeightTextureGridSize)
 		SHADER_PARAMETER_SAMPLER(SamplerState,		PointClampedSampler)
-		SHADER_PARAMETER_SRV(Texture2D<int>,		CachedPhysMatTexture)
+		SHADER_PARAMETER_TEXTURE(Texture2D<uint>,	CachedPhysMatTexture)
 		SHADER_PARAMETER(FIntPoint,					CachedPhysMatTextureDimension)
 	END_SHADER_PARAMETER_STRUCT()
 }
@@ -112,7 +112,7 @@ public:
 	void ReleaseSourceData();
 
 	FRHIShaderResourceView* GetHeightTexture() const { return HeightTexture.SRV; }
-	FRHIShaderResourceView* GetPhysMatTexture() const { return PhysMatTexture.SRV; }
+	FRHITexture* GetPhysMatTexture() const { return PhysMatTexture.Buffer; }
 	FIntPoint GetDimensions() const { return CellCount; }
 
 	TArray<float>& EditHeightValues(int32 SampleCount);
@@ -526,7 +526,7 @@ struct FNDILandscapeData_RenderThread
 		if (TextureResources)
 		{
 			FRHIShaderResourceView* HeightTextureSrv = TextureResources->GetHeightTexture();
-			FRHIShaderResourceView* PhysMatTextureSrv = TextureResources->GetPhysMatTexture();
+			FRHITexture* PhysMatTextureSrv = TextureResources->GetPhysMatTexture();
 
 			if (HeightTextureSrv || PhysMatTextureSrv)
 			{
@@ -555,7 +555,7 @@ struct FNDILandscapeData_RenderThread
 				else
 				{
 					const FIntPoint PhysMatDimensions(ForceInitToZero);
-					ShaderParameters->CachedPhysMatTexture = GBlackTextureWithSRV->ShaderResourceViewRHI;
+					ShaderParameters->CachedPhysMatTexture = GBlackUintTexture->TextureRHI;
 					ShaderParameters->CachedPhysMatTextureDimension = PhysMatDimensions;
 				}
 
@@ -577,7 +577,7 @@ struct FNDILandscapeData_RenderThread
 		ShaderParameters->CachedHeightTextureGridSize = FVector2f::ZeroVector;
 		ShaderParameters->CachedHeightTextureEnabled = 0;
 
-		ShaderParameters->CachedPhysMatTexture = GBlackTextureWithSRV->ShaderResourceViewRHI;
+		ShaderParameters->CachedPhysMatTexture = GBlackUintTexture->TextureRHI;
 	}
 };
 
