- Handle issue where landscape actor is force killed

diff --git a/Engine/Plugins/FX/Niagara/Source/Niagara/Private/NiagaraDataInterfaceLandscape.cpp b/Engine/Plugins/FX/Niagara/Source/Niagara/Private/NiagaraDataInterfaceLandscape.cpp
--- a/Engine/Plugins/FX/Niagara/Source/Niagara/Private/NiagaraDataInterfaceLandscape.cpp
+++ b/Engine/Plugins/FX/Niagara/Source/Niagara/Private/NiagaraDataInterfaceLandscape.cpp
@@ -1237,39 +1237,42 @@ void UNiagaraDataInterfaceLandscape::ProvidePerInstanceDataForRenderThread(void*
 		TargetData->CachedHeightTextureGridSize = SourceResource.TextureWorldGridSize;
 	}
 
-	if (SourceData.BaseColorVirtualTextureIndex != INDEX_NONE)
+	if (ALandscape* SourceDataLandscape = SourceData.Landscape.Get())
 	{
-		if (const URuntimeVirtualTexture* VirtualTexture = SourceData.Landscape->RuntimeVirtualTextures[SourceData.BaseColorVirtualTextureIndex])
+		if (SourceData.BaseColorVirtualTextureIndex != INDEX_NONE)
 		{
-			TargetData->BaseColorVirtualTexture = VirtualTexture;
-			TargetData->BaseColorVirtualTextureSRGB = SourceData.BaseColorVirtualTextureSRGB;
-			TargetData->BaseColorVirtualTextureWorldToUvParameters[0] = VirtualTexture->GetUniformParameter(ERuntimeVirtualTextureShaderUniform_WorldToUVTransform0);
-			TargetData->BaseColorVirtualTextureWorldToUvParameters[1] = VirtualTexture->GetUniformParameter(ERuntimeVirtualTextureShaderUniform_WorldToUVTransform1);
-			TargetData->BaseColorVirtualTextureWorldToUvParameters[2] = VirtualTexture->GetUniformParameter(ERuntimeVirtualTextureShaderUniform_WorldToUVTransform2);
+			if (const URuntimeVirtualTexture* VirtualTexture = SourceDataLandscape->RuntimeVirtualTextures[SourceData.BaseColorVirtualTextureIndex])
+			{
+				TargetData->BaseColorVirtualTexture = VirtualTexture;
+				TargetData->BaseColorVirtualTextureSRGB = SourceData.BaseColorVirtualTextureSRGB;
+				TargetData->BaseColorVirtualTextureWorldToUvParameters[0] = VirtualTexture->GetUniformParameter(ERuntimeVirtualTextureShaderUniform_WorldToUVTransform0);
+				TargetData->BaseColorVirtualTextureWorldToUvParameters[1] = VirtualTexture->GetUniformParameter(ERuntimeVirtualTextureShaderUniform_WorldToUVTransform1);
+				TargetData->BaseColorVirtualTextureWorldToUvParameters[2] = VirtualTexture->GetUniformParameter(ERuntimeVirtualTextureShaderUniform_WorldToUVTransform2);
+			}
 		}
-	}
 
-	if (SourceData.HeightVirtualTextureIndex != INDEX_NONE)
-	{
-		if (const URuntimeVirtualTexture* VirtualTexture = SourceData.Landscape->RuntimeVirtualTextures[SourceData.HeightVirtualTextureIndex])
+		if (SourceData.HeightVirtualTextureIndex != INDEX_NONE)
 		{
-			TargetData->HeightVirtualTexture = VirtualTexture;
-			TargetData->HeightVirtualTextureWorldToUvParameters[0] = VirtualTexture->GetUniformParameter(ERuntimeVirtualTextureShaderUniform_WorldToUVTransform0);
-			TargetData->HeightVirtualTextureWorldToUvParameters[1] = VirtualTexture->GetUniformParameter(ERuntimeVirtualTextureShaderUniform_WorldToUVTransform1);
-			TargetData->HeightVirtualTextureWorldToUvParameters[2] = VirtualTexture->GetUniformParameter(ERuntimeVirtualTextureShaderUniform_WorldToUVTransform2);
-			TargetData->HeightVirtualTextureWorldToUvParameters[3] = VirtualTexture->GetUniformParameter(ERuntimeVirtualTextureShaderUniform_WorldHeightUnpack);
+			if (const URuntimeVirtualTexture* VirtualTexture = SourceDataLandscape->RuntimeVirtualTextures[SourceData.HeightVirtualTextureIndex])
+			{
+				TargetData->HeightVirtualTexture = VirtualTexture;
+				TargetData->HeightVirtualTextureWorldToUvParameters[0] = VirtualTexture->GetUniformParameter(ERuntimeVirtualTextureShaderUniform_WorldToUVTransform0);
+				TargetData->HeightVirtualTextureWorldToUvParameters[1] = VirtualTexture->GetUniformParameter(ERuntimeVirtualTextureShaderUniform_WorldToUVTransform1);
+				TargetData->HeightVirtualTextureWorldToUvParameters[2] = VirtualTexture->GetUniformParameter(ERuntimeVirtualTextureShaderUniform_WorldToUVTransform2);
+				TargetData->HeightVirtualTextureWorldToUvParameters[3] = VirtualTexture->GetUniformParameter(ERuntimeVirtualTextureShaderUniform_WorldHeightUnpack);
+			}
 		}
-	}
 
-	if (SourceData.NormalVirtualTextureIndex != INDEX_NONE)
-	{
-		if (const URuntimeVirtualTexture* VirtualTexture = SourceData.Landscape->RuntimeVirtualTextures[SourceData.NormalVirtualTextureIndex])
+		if (SourceData.NormalVirtualTextureIndex != INDEX_NONE)
 		{
-			TargetData->NormalVirtualTexture = VirtualTexture;
-			TargetData->NormalVirtualTextureWorldToUvParameters[0] = VirtualTexture->GetUniformParameter(ERuntimeVirtualTextureShaderUniform_WorldToUVTransform0);
-			TargetData->NormalVirtualTextureWorldToUvParameters[1] = VirtualTexture->GetUniformParameter(ERuntimeVirtualTextureShaderUniform_WorldToUVTransform1);
-			TargetData->NormalVirtualTextureWorldToUvParameters[2] = VirtualTexture->GetUniformParameter(ERuntimeVirtualTextureShaderUniform_WorldToUVTransform2);
-			TargetData->NormalVirtualTextureMode = SourceData.NormalVirtualTextureMode;
+			if (const URuntimeVirtualTexture* VirtualTexture = SourceDataLandscape->RuntimeVirtualTextures[SourceData.NormalVirtualTextureIndex])
+			{
+				TargetData->NormalVirtualTexture = VirtualTexture;
+				TargetData->NormalVirtualTextureWorldToUvParameters[0] = VirtualTexture->GetUniformParameter(ERuntimeVirtualTextureShaderUniform_WorldToUVTransform0);
+				TargetData->NormalVirtualTextureWorldToUvParameters[1] = VirtualTexture->GetUniformParameter(ERuntimeVirtualTextureShaderUniform_WorldToUVTransform1);
+				TargetData->NormalVirtualTextureWorldToUvParameters[2] = VirtualTexture->GetUniformParameter(ERuntimeVirtualTextureShaderUniform_WorldToUVTransform2);
+				TargetData->NormalVirtualTextureMode = SourceData.NormalVirtualTextureMode;
+			}
 		}
 	}
 }
