Fix NiagaraDataChannelManager causing worlds to leak on map load.

diff --git a/Engine/Plugins/FX/Niagara/Source/Niagara/Private/NiagaraDataChannelManager.cpp b/Engine/Plugins/FX/Niagara/Source/Niagara/Private/NiagaraDataChannelManager.cpp
--- a/Engine/Plugins/FX/Niagara/Source/Niagara/Private/NiagaraDataChannelManager.cpp
+++ b/Engine/Plugins/FX/Niagara/Source/Niagara/Private/NiagaraDataChannelManager.cpp
@@ -23,14 +23,24 @@ void FNiagaraDataChannelManager::AddReferencedObjects(FReferenceCollector& Colle
 
 void FNiagaraDataChannelManager::RefreshDataChannels()
 {
-	for (TObjectIterator<UNiagaraDataChannel> DCIt; DCIt; ++DCIt)
+	if (bIsCleanedUp)
 	{
-		InitDataChannel(*DCIt, false);
+		return;
 	}
+
+	UNiagaraDataChannel::ForEachDataChannel([&](UNiagaraDataChannel* DataChannel)
+	{
+		if (DataChannel->HasAnyFlags(RF_ClassDefaultObject) == false)
+		{
+			InitDataChannel(DataChannel, false);
+		}
+	});
 }
 
 void FNiagaraDataChannelManager::Init()
 {
+	check(!bIsCleanedUp);
+
 	//Initialize any existing data channels, more may be initialized later as they are loaded.
 	UNiagaraDataChannel::ForEachDataChannel([&](UNiagaraDataChannel* DataChannel)
 	{
@@ -44,6 +54,7 @@ void FNiagaraDataChannelManager::Init()
 void FNiagaraDataChannelManager::Cleanup()
 {
 	Channels.Empty();
+	bIsCleanedUp = true;
 }
 
 void FNiagaraDataChannelManager::BeginFrame(float DeltaSeconds)
diff --git a/Engine/Plugins/FX/Niagara/Source/Niagara/Private/NiagaraDataChannelManager.h b/Engine/Plugins/FX/Niagara/Source/Niagara/Private/NiagaraDataChannelManager.h
--- a/Engine/Plugins/FX/Niagara/Source/Niagara/Private/NiagaraDataChannelManager.h
+++ b/Engine/Plugins/FX/Niagara/Source/Niagara/Private/NiagaraDataChannelManager.h
@@ -47,4 +47,5 @@ private:
 
 	/** Runtime handlers for each DataChannel channel. */
 	TMap<TWeakObjectPtr<const UNiagaraDataChannel>, TObjectPtr<UNiagaraDataChannelHandler>> Channels;
+	bool bIsCleanedUp = false;
 };
