Allow FNiagaraDataChannelManager to be reinitialized after cleaup.

diff --git a/Engine/Plugins/FX/Niagara/Source/Niagara/Private/NiagaraDataChannelManager.cpp b/Engine/Plugins/FX/Niagara/Source/Niagara/Private/NiagaraDataChannelManager.cpp
--- a/Engine/Plugins/FX/Niagara/Source/Niagara/Private/NiagaraDataChannelManager.cpp
+++ b/Engine/Plugins/FX/Niagara/Source/Niagara/Private/NiagaraDataChannelManager.cpp
@@ -39,7 +39,7 @@ void FNiagaraDataChannelManager::RefreshDataChannels()
 
 void FNiagaraDataChannelManager::Init()
 {
-	check(!bIsCleanedUp);
+	bIsCleanedUp = false;
 
 	//Initialize any existing data channels, more may be initialized later as they are loaded.
 	UNiagaraDataChannel::ForEachDataChannel([&](UNiagaraDataChannel* DataChannel)
