- Allow users to select ribbon material back facing culling vs shape mode

diff --git a/Engine/Plugins/FX/Niagara/Source/Niagara/Private/NiagaraRendererRibbons.cpp b/Engine/Plugins/FX/Niagara/Source/Niagara/Private/NiagaraRendererRibbons.cpp
--- a/Engine/Plugins/FX/Niagara/Source/Niagara/Private/NiagaraRendererRibbons.cpp
+++ b/Engine/Plugins/FX/Niagara/Source/Niagara/Private/NiagaraRendererRibbons.cpp
@@ -1098,10 +1098,10 @@ void FNiagaraRendererRibbons::GetDynamicRayTracingInstances(FRayTracingMaterialG
 }
 #endif
 
-void FNiagaraRendererRibbons::GenerateShapeStateMultiPlane(FNiagaraRibbonShapeGeometryData& State, int32 MultiPlaneCount, int32 WidthSegmentationCount, bool bEnableAccurateGeometry)
+void FNiagaraRendererRibbons::GenerateShapeStateMultiPlane(FNiagaraRibbonShapeGeometryData& State, int32 MultiPlaneCount, int32 WidthSegmentationCount, bool bEnableAccurateGeometry, bool bUseMaterialBackfaceCulling)
 {
 	State.Shape = ENiagaraRibbonShapeMode::MultiPlane;
-	State.bDisableBackfaceCulling = !bEnableAccurateGeometry;
+	State.bDisableBackfaceCulling = !bEnableAccurateGeometry && !bUseMaterialBackfaceCulling;
 	State.bShouldFlipNormalToView = !bEnableAccurateGeometry;
 	State.TrianglesPerSegment = 2 * MultiPlaneCount * WidthSegmentationCount * (bEnableAccurateGeometry? 2 : 1);
 	State.NumVerticesInSlice = MultiPlaneCount * (WidthSegmentationCount + 1) * (bEnableAccurateGeometry ? 2 : 1);
@@ -1164,10 +1164,10 @@ void FNiagaraRendererRibbons::GenerateShapeStateMultiPlane(FNiagaraRibbonShapeGe
 	}
 }
 
-void FNiagaraRendererRibbons::GenerateShapeStateTube(FNiagaraRibbonShapeGeometryData& State, int32 TubeSubdivisions)
+void FNiagaraRendererRibbons::GenerateShapeStateTube(FNiagaraRibbonShapeGeometryData& State, int32 TubeSubdivisions, bool bUseMaterialBackfaceCulling)
 {
 	State.Shape = ENiagaraRibbonShapeMode::Tube;
-	State.bDisableBackfaceCulling = true;
+	State.bDisableBackfaceCulling = !bUseMaterialBackfaceCulling;
 	State.bShouldFlipNormalToView = false;
 	State.TrianglesPerSegment = 2 * TubeSubdivisions;
 	State.NumVerticesInSlice = TubeSubdivisions + 1;
@@ -1193,10 +1193,10 @@ void FNiagaraRendererRibbons::GenerateShapeStateTube(FNiagaraRibbonShapeGeometry
 	}
 }
 
-void FNiagaraRendererRibbons::GenerateShapeStateCustom(FNiagaraRibbonShapeGeometryData& State, const TArray<FNiagaraRibbonShapeCustomVertex>& CustomVertices)
+void FNiagaraRendererRibbons::GenerateShapeStateCustom(FNiagaraRibbonShapeGeometryData& State, const TArray<FNiagaraRibbonShapeCustomVertex>& CustomVertices, bool bUseMaterialBackfaceCulling)
 {
 	State.Shape = ENiagaraRibbonShapeMode::Custom;
-	State.bDisableBackfaceCulling = true;
+	State.bDisableBackfaceCulling = !bUseMaterialBackfaceCulling;
 	State.bShouldFlipNormalToView = false;
 	State.TrianglesPerSegment = 2 * CustomVertices.Num();
 	State.NumVerticesInSlice = CustomVertices.Num() + 1;
@@ -1232,10 +1232,10 @@ void FNiagaraRendererRibbons::GenerateShapeStateCustom(FNiagaraRibbonShapeGeomet
 	}
 }
 
-void FNiagaraRendererRibbons::GenerateShapeStatePlane(FNiagaraRibbonShapeGeometryData& State, int32 WidthSegmentationCount)
+void FNiagaraRendererRibbons::GenerateShapeStatePlane(FNiagaraRibbonShapeGeometryData& State, int32 WidthSegmentationCount, bool bUseMaterialBackfaceCulling)
 {
 	State.Shape = ENiagaraRibbonShapeMode::Plane;
-	State.bDisableBackfaceCulling = true;
+	State.bDisableBackfaceCulling = !bUseMaterialBackfaceCulling;
 	State.bShouldFlipNormalToView = false;
 	State.TrianglesPerSegment = 2 * WidthSegmentationCount;
 	State.NumVerticesInSlice = WidthSegmentationCount + 1;
@@ -1264,19 +1264,19 @@ void FNiagaraRendererRibbons::InitializeShape(const UNiagaraRibbonRendererProper
 {
 	if (Properties->Shape == ENiagaraRibbonShapeMode::Custom && Properties->CustomVertices.Num() > 2)
 	{
-		GenerateShapeStateCustom(ShapeState, Properties->CustomVertices);
+		GenerateShapeStateCustom(ShapeState, Properties->CustomVertices, Properties->bUseMaterialBackfaceCulling);
 	}
 	else if (Properties->Shape == ENiagaraRibbonShapeMode::Tube && Properties->TubeSubdivisions > 2 && Properties->TubeSubdivisions <= 16)
 	{
-		GenerateShapeStateTube(ShapeState, Properties->TubeSubdivisions);
+		GenerateShapeStateTube(ShapeState, Properties->TubeSubdivisions, Properties->bUseMaterialBackfaceCulling);
 	}
 	else if (Properties->Shape == ENiagaraRibbonShapeMode::MultiPlane && Properties->MultiPlaneCount > 1 && Properties->MultiPlaneCount <= 16)
 	{
-		GenerateShapeStateMultiPlane(ShapeState, Properties->MultiPlaneCount, Properties->WidthSegmentationCount, Properties->bEnableAccurateGeometry);
+		GenerateShapeStateMultiPlane(ShapeState, Properties->MultiPlaneCount, Properties->WidthSegmentationCount, Properties->bEnableAccurateGeometry, Properties->bUseMaterialBackfaceCulling);
 	}
 	else
 	{
-		GenerateShapeStatePlane(ShapeState, Properties->WidthSegmentationCount);
+		GenerateShapeStatePlane(ShapeState, Properties->WidthSegmentationCount, Properties->bUseMaterialBackfaceCulling);
 	}	
 }
 
diff --git a/Engine/Plugins/FX/Niagara/Source/Niagara/Private/NiagaraRibbonRendererProperties.cpp b/Engine/Plugins/FX/Niagara/Source/Niagara/Private/NiagaraRibbonRendererProperties.cpp
--- a/Engine/Plugins/FX/Niagara/Source/Niagara/Private/NiagaraRibbonRendererProperties.cpp
+++ b/Engine/Plugins/FX/Niagara/Source/Niagara/Private/NiagaraRibbonRendererProperties.cpp
@@ -65,6 +65,7 @@ UNiagaraRibbonRendererProperties::UNiagaraRibbonRendererProperties()
 	, MaxNumRibbons(0)
 	, Shape(ENiagaraRibbonShapeMode::Plane)
 	, bEnableAccurateGeometry(false)
+	, bUseMaterialBackfaceCulling(false)
 	, bUseGPUInit(false)
 	, bUseConstantFactor(false)
 	, bScreenSpaceTessellation(true)
diff --git a/Engine/Plugins/FX/Niagara/Source/Niagara/Public/NiagaraRendererRibbons.h b/Engine/Plugins/FX/Niagara/Source/Niagara/Public/NiagaraRendererRibbons.h
--- a/Engine/Plugins/FX/Niagara/Source/Niagara/Public/NiagaraRendererRibbons.h
+++ b/Engine/Plugins/FX/Niagara/Source/Niagara/Public/NiagaraRendererRibbons.h
@@ -317,15 +317,14 @@ public:
 
 protected:
 	
-	static NIAGARA_API void GenerateShapeStateMultiPlane(FNiagaraRibbonShapeGeometryData& State, int32 MultiPlaneCount, int32 WidthSegmentationCount, bool bEnableAccurateGeometry);
-	static NIAGARA_API void GenerateShapeStateTube(FNiagaraRibbonShapeGeometryData& State, int32 TubeSubdivisions);
-	static NIAGARA_API void GenerateShapeStateCustom(FNiagaraRibbonShapeGeometryData& State, const TArray<FNiagaraRibbonShapeCustomVertex>& CustomVertices);
-	static NIAGARA_API void GenerateShapeStatePlane(FNiagaraRibbonShapeGeometryData& State, int32 WidthSegmentationCount);	
+	static NIAGARA_API void GenerateShapeStateMultiPlane(FNiagaraRibbonShapeGeometryData& State, int32 MultiPlaneCount, int32 WidthSegmentationCount, bool bEnableAccurateGeometry, bool bUseMaterialBackfaceCulling);
+	static NIAGARA_API void GenerateShapeStateTube(FNiagaraRibbonShapeGeometryData& State, int32 TubeSubdivisions, bool bUseMaterialBackfaceCulling);
+	static NIAGARA_API void GenerateShapeStateCustom(FNiagaraRibbonShapeGeometryData& State, const TArray<FNiagaraRibbonShapeCustomVertex>& CustomVertices, bool bUseMaterialBackfaceCulling);
+	static NIAGARA_API void GenerateShapeStatePlane(FNiagaraRibbonShapeGeometryData& State, int32 WidthSegmentationCount, bool bUseMaterialBackfaceCulling);
 	NIAGARA_API void InitializeShape(const UNiagaraRibbonRendererProperties* Properties);
 	
 	NIAGARA_API void InitializeTessellation(const UNiagaraRibbonRendererProperties* Properties);
 	
-	
 	template<typename IntType>
 	static void CalculateUVScaleAndOffsets(const FNiagaraRibbonUVSettings& UVSettings, const TArray<IntType>& RibbonIndices, const TArray<FVector4f>& RibbonTangentsAndDistances, const FNiagaraDataSetReaderFloat<float>& NormalizedAgeReader,
 		int32 StartIndex, int32 EndIndex, int32 NumSegments, float TotalLength, float& OutUScale, float& OutUOffset, float& OutUDistributionScaler);
diff --git a/Engine/Plugins/FX/Niagara/Source/Niagara/Public/NiagaraRibbonRendererProperties.h b/Engine/Plugins/FX/Niagara/Source/Niagara/Public/NiagaraRibbonRendererProperties.h
--- a/Engine/Plugins/FX/Niagara/Source/Niagara/Public/NiagaraRibbonRendererProperties.h
+++ b/Engine/Plugins/FX/Niagara/Source/Niagara/Public/NiagaraRibbonRendererProperties.h
@@ -305,6 +305,10 @@ public:
 	UPROPERTY(EditAnywhere, Category = "Ribbon Shape", meta = (EditCondition = "Shape != ENiagaraRibbonShapeMode::Plane", EditConditionHides))
 	uint8 bEnableAccurateGeometry : 1;
 
+	/** When enabled the ribbons renderer will not override how backface culling works depending on shape type, but instad use the material culling mode */
+	UPROPERTY(EditAnywhere, Category = "Ribbon Shape")
+	uint8 bUseMaterialBackfaceCulling : 1;
+
 	/**
 	*	Whether we use the CPU or GPU to generate ribbon geometry for CPU systems.
 	*	GPU systems will always use a fully GPU initialization pipeline,
