- Add SourceMode to object reader, allows you to easily read from the attached actor
- Add support for invoking BP get functions to read properties
- Fix issue with alignment on FVector4's

diff --git a/Engine/Plugins/FX/Niagara/Source/Niagara/Private/DataInterface/NiagaraDataInterfaceUObjectPropertyReader.cpp b/Engine/Plugins/FX/Niagara/Source/Niagara/Private/DataInterface/NiagaraDataInterfaceUObjectPropertyReader.cpp
--- a/Engine/Plugins/FX/Niagara/Source/Niagara/Private/DataInterface/NiagaraDataInterfaceUObjectPropertyReader.cpp
+++ b/Engine/Plugins/FX/Niagara/Source/Niagara/Private/DataInterface/NiagaraDataInterfaceUObjectPropertyReader.cpp
@@ -157,11 +157,12 @@ namespace NDIUObjectPropertyReaderLocal
 		if (InstanceData_GT->PropertyGetters[GetterIndex].PropertyCopyFunction != nullptr)
 		{
 			const uint32 DataOffset = InstanceData_GT->PropertyGetters[GetterIndex].DataOffset;
-			const TInputType* Value = reinterpret_cast<const TInputType*>(InstanceData_GT->PropertyData.GetData() + DataOffset);
+			TInputType Value;
+			FMemory::Memcpy(&Value, InstanceData_GT->PropertyData.GetData() + DataOffset, sizeof(TInputType));
 			for (int32 i = 0; i < Context.GetNumInstances(); ++i)
 			{
 				SuccessValue.SetAndAdvance(true);
-				OutValue.SetAndAdvance(*Value);
+				OutValue.SetAndAdvance(Value);
 			}
 		}
 		else
@@ -174,6 +175,24 @@ namespace NDIUObjectPropertyReaderLocal
 		}
 	}
 
+
+	template<typename TNiagaraType, typename TPropertyType>
+	void CopyPropertyHelper(const void* InPropertyValue, void* Destination)
+	{
+		const TPropertyType PropertyValue(*reinterpret_cast<const TPropertyType*>(InPropertyValue));
+		const TNiagaraType NiagaraValue = TNiagaraType(PropertyValue);
+		FMemory::Memcpy(Destination, &NiagaraValue, sizeof(TNiagaraType));
+	}
+
+	template<typename TNiagaraType, typename TScriptType>
+	void InvokeGetterHelper(UObject* BoundObject, UFunction* BoundFunction, void* Destination, const TScriptType& DefaultValue)
+	{
+		TScriptType ReturnValue = DefaultValue;
+		BoundObject->ProcessEvent(BoundFunction, &ReturnValue);
+		const TNiagaraType NiagaraValue = TNiagaraType(ReturnValue);
+		FMemory::Memcpy(Destination, &NiagaraValue, sizeof(TNiagaraType));
+	}
+
 	template<typename TType>
 	struct FTypeHelper : std::false_type
 	{
@@ -189,11 +208,23 @@ namespace NDIUObjectPropertyReaderLocal
 		{
 			if (const FFloatProperty* FloatProperty = CastField<const FFloatProperty>(InProperty))
 			{
-				return [PropertyAddress](const FNiagaraLWCConverter& LwcConverter, void* DestAddress) { *reinterpret_cast<float*>(DestAddress) = *reinterpret_cast<const float*>(PropertyAddress); };
+				return [PropertyAddress](const FNiagaraLWCConverter& LwcConverter, void* DestAddress) { CopyPropertyHelper<float, float>(PropertyAddress, DestAddress); };
 			}
 			else if (const FDoubleProperty* DoubleProperty = CastField<const FDoubleProperty>(InProperty))
 			{
-				return [PropertyAddress](const FNiagaraLWCConverter& LwcConverter, void* DestAddress) { *reinterpret_cast<float*>(DestAddress) = float(*reinterpret_cast<const double*>(PropertyAddress)); };
+				return [PropertyAddress](const FNiagaraLWCConverter& LwcConverter, void* DestAddress) { CopyPropertyHelper<float, double>(PropertyAddress, DestAddress); };
+			}
+			return nullptr;
+		}
+		static TFunction<void(const FNiagaraLWCConverter&, void*)> GetGetterFunction(const FProperty* InProperty, UObject* BoundObject, UFunction* BoundFunction)
+		{
+			if (CastField<const FFloatProperty>(InProperty))
+			{
+				return [BoundObject, BoundFunction](const FNiagaraLWCConverter& LwcConverter, void* DestAddress) { BoundObject->ProcessEvent(BoundFunction, DestAddress); };
+			}
+			else if (CastField<const FDoubleProperty>(InProperty))
+			{
+				return [BoundObject, BoundFunction](const FNiagaraLWCConverter& LwcConverter, void* DestAddress) { InvokeGetterHelper<float, double>(BoundObject, BoundFunction, DestAddress, 0.0); };
 			}
 			return nullptr;
 		}
@@ -213,7 +244,18 @@ namespace NDIUObjectPropertyReaderLocal
 			{
 				if ( StructProperty->Struct == TBaseStructure<FVector2D>::Get() )
 				{
-					return [PropertyAddress](const FNiagaraLWCConverter& LwcConverter, void* DestAddress) { *reinterpret_cast<FVector2f*>(DestAddress) = FVector2f(*reinterpret_cast<const FVector2D*>(PropertyAddress)); };
+					return [PropertyAddress](const FNiagaraLWCConverter& LwcConverter, void* DestAddress) { CopyPropertyHelper<FVector2f, FVector2D>(PropertyAddress, DestAddress); };
+				}
+			}
+			return nullptr;
+		}
+		static TFunction<void(const FNiagaraLWCConverter&, void*)> GetGetterFunction(const FProperty* InProperty, UObject* BoundObject, UFunction* BoundFunction)
+		{
+			if (const FStructProperty* StructProperty = CastField<FStructProperty>(InProperty))
+			{
+				if (StructProperty->Struct == TBaseStructure<FVector2D>::Get())
+				{
+					return [BoundObject, BoundFunction](const FNiagaraLWCConverter& LwcConverter, void* DestAddress) { InvokeGetterHelper<FVector2f, FVector2D>(BoundObject, BoundFunction, DestAddress, FVector2D::ZeroVector); };
 				}
 			}
 			return nullptr;
@@ -234,7 +276,18 @@ namespace NDIUObjectPropertyReaderLocal
 			{
 				if (StructProperty->Struct == TBaseStructure<FVector>::Get())
 				{
-					return [PropertyAddress](const FNiagaraLWCConverter& LwcConverter, void* DestAddress) { *reinterpret_cast<FVector3f*>(DestAddress) = FVector3f(*reinterpret_cast<const FVector*>(PropertyAddress)); };
+					return [PropertyAddress](const FNiagaraLWCConverter& LwcConverter, void* DestAddress) { CopyPropertyHelper<FVector3f, FVector>(PropertyAddress, DestAddress); };
+				}
+			}
+			return nullptr;
+		}
+		static TFunction<void(const FNiagaraLWCConverter&, void*)> GetGetterFunction(const FProperty* InProperty, UObject* BoundObject, UFunction* BoundFunction)
+		{
+			if (const FStructProperty* StructProperty = CastField<FStructProperty>(InProperty))
+			{
+				if (StructProperty->Struct == TBaseStructure<FVector>::Get())
+				{
+					return [BoundObject, BoundFunction](const FNiagaraLWCConverter& LwcConverter, void* DestAddress) { InvokeGetterHelper<FVector3f, FVector>(BoundObject, BoundFunction, DestAddress, FVector::ZeroVector); };
 				}
 			}
 			return nullptr;
@@ -248,14 +301,25 @@ namespace NDIUObjectPropertyReaderLocal
 	{
 		static FName GetFunctionName() { return FName("GetVec4Property"); }
 		static FNiagaraTypeDefinition GetTypeDef() { return FNiagaraTypeDefinition::GetVec4Def(); }
-		static void VMFunction(FVectorVMExternalFunctionContext& Context, const uint32 GetterIndex) { VMReadData<FVector4f, FVector4f>(Context, GetterIndex, FVector3f::ZeroVector); }
+		static void VMFunction(FVectorVMExternalFunctionContext& Context, const uint32 GetterIndex) { VMReadData<FVector4f, FVector4f>(Context, GetterIndex, FVector4f::Zero()); }
 		static TFunction<void(const FNiagaraLWCConverter&, void*)> GetCopyFunction(const FProperty* InProperty, const void* PropertyAddress)
 		{
 			if (const FStructProperty* StructProperty = CastField<FStructProperty>(InProperty))
 			{
 				if (StructProperty->Struct == TBaseStructure<FVector4>::Get())
 				{
-					return [PropertyAddress](const FNiagaraLWCConverter& LwcConverter, void* DestAddress) { *reinterpret_cast<FVector4f*>(DestAddress) = FVector4f(*reinterpret_cast<const FVector4*>(PropertyAddress)); };
+					return [PropertyAddress](const FNiagaraLWCConverter& LwcConverter, void* DestAddress) { CopyPropertyHelper<FVector4f, FVector4>(PropertyAddress, DestAddress); };
+				}
+			}
+			return nullptr;
+		}
+		static TFunction<void(const FNiagaraLWCConverter&, void*)> GetGetterFunction(const FProperty* InProperty, UObject* BoundObject, UFunction* BoundFunction)
+		{
+			if (const FStructProperty* StructProperty = CastField<FStructProperty>(InProperty))
+			{
+				if (StructProperty->Struct == TBaseStructure<FVector4>::Get())
+				{
+					return [BoundObject, BoundFunction](const FNiagaraLWCConverter& LwcConverter, void* DestAddress) { InvokeGetterHelper<FVector4f, FVector4>(BoundObject, BoundFunction, DestAddress, FVector4::Zero()); };
 				}
 			}
 			return nullptr;
@@ -269,18 +333,33 @@ namespace NDIUObjectPropertyReaderLocal
 	{
 		static FName GetFunctionName() { return FName("GetColorProperty"); }
 		static FNiagaraTypeDefinition GetTypeDef() { return FNiagaraTypeDefinition::GetColorDef(); }
-		static void VMFunction(FVectorVMExternalFunctionContext& Context, const uint32 GetterIndex) { VMReadData<FLinearColor, FLinearColor>(Context, GetterIndex, FVector3f::ZeroVector); }
+		static void VMFunction(FVectorVMExternalFunctionContext& Context, const uint32 GetterIndex) { VMReadData<FLinearColor, FLinearColor>(Context, GetterIndex, FLinearColor::Black); }
 		static TFunction<void(const FNiagaraLWCConverter&, void*)> GetCopyFunction(const FProperty* InProperty, const void* PropertyAddress)
 		{
 			if (const FStructProperty* StructProperty = CastField<FStructProperty>(InProperty))
 			{
 				if (StructProperty->Struct == TBaseStructure<FColor>::Get())
 				{
-					return [PropertyAddress](const FNiagaraLWCConverter& LwcConverter, void* DestAddress) { *reinterpret_cast<FLinearColor*>(DestAddress) = FLinearColor(*reinterpret_cast<const FColor*>(PropertyAddress)); };
+					return [PropertyAddress](const FNiagaraLWCConverter& LwcConverter, void* DestAddress) { CopyPropertyHelper<FLinearColor, FColor>(PropertyAddress, DestAddress); };
 				}
 				else if (StructProperty->Struct == TBaseStructure<FLinearColor>::Get())
 				{
-					return [PropertyAddress](const FNiagaraLWCConverter& LwcConverter, void* DestAddress) { *reinterpret_cast<FLinearColor*>(DestAddress) = *reinterpret_cast<const FLinearColor*>(PropertyAddress); };
+					return [PropertyAddress](const FNiagaraLWCConverter& LwcConverter, void* DestAddress) { CopyPropertyHelper<FLinearColor, FLinearColor>(PropertyAddress, DestAddress); };
+				}
+			}
+			return nullptr;
+		}
+		static TFunction<void(const FNiagaraLWCConverter&, void*)> GetGetterFunction(const FProperty* InProperty, UObject* BoundObject, UFunction* BoundFunction)
+		{
+			if (const FStructProperty* StructProperty = CastField<FStructProperty>(InProperty))
+			{
+				if (StructProperty->Struct == TBaseStructure<FColor>::Get())
+				{
+					return [BoundObject, BoundFunction](const FNiagaraLWCConverter& LwcConverter, void* DestAddress) { InvokeGetterHelper<FLinearColor, FColor>(BoundObject, BoundFunction, DestAddress, FColor::Black); };
+				}
+				else if (StructProperty->Struct == TBaseStructure<FLinearColor>::Get())
+				{
+					return [BoundObject, BoundFunction](const FNiagaraLWCConverter& LwcConverter, void* DestAddress) { InvokeGetterHelper<FLinearColor, FLinearColor>(BoundObject, BoundFunction, DestAddress, FLinearColor::Black); };
 				}
 			}
 			return nullptr;
@@ -306,8 +385,26 @@ namespace NDIUObjectPropertyReaderLocal
 			}
 			return nullptr;
 		}
+		static TFunction<void(const FNiagaraLWCConverter&, void*)> GetGetterFunction(const FProperty* InProperty, UObject* BoundObject, UFunction* BoundFunction)
+		{
+			if (const FStructProperty* StructProperty = CastField<FStructProperty>(InProperty))
+			{
+				if (StructProperty->Struct == TBaseStructure<FVector>::Get())
+				{
+					return
+						[BoundObject, BoundFunction](const FNiagaraLWCConverter& LwcConverter, void* DestAddress)
+						{
+							FVector ReturnValue = FVector::ZeroVector;
+							BoundObject->ProcessEvent(BoundFunction, &ReturnValue);
+							const FVector3f NiagaraValue = LwcConverter.ConvertWorldToSimulationPosition(ReturnValue);
+							FMemory::Memcpy(DestAddress, &NiagaraValue, sizeof(FVector3f));
+						};
+				}
+			}
+			return nullptr;
+		}
 		static constexpr TCHAR const* HlslBufferType = TEXT("float3");
-		static constexpr TCHAR const* HlslBufferRead = TEXT("float3(asfloat(BUFFER[OFFSET + 0)], asfloat(BUFFER[OFFSET + 1]), asfloat(BUFFER[OFFSET + 2])) : float3(0.0f, 0.0f, 0.0f)");
+		static constexpr TCHAR const* HlslBufferRead = TEXT("float3(asfloat(BUFFER[OFFSET + 0]), asfloat(BUFFER[OFFSET + 1]), asfloat(BUFFER[OFFSET + 2])) : float3(0.0f, 0.0f, 0.0f)");
 	};
 
 	template<>
@@ -325,23 +422,52 @@ namespace NDIUObjectPropertyReaderLocal
 
 			if (InProperty->IsA<const FByteProperty>())
 			{
-				return [PropertyAddress](const FNiagaraLWCConverter& LwcConverter, void* DestAddress) { *reinterpret_cast<int32*>(DestAddress) = *reinterpret_cast<const uint8*>(PropertyAddress); };
+				return [PropertyAddress](const FNiagaraLWCConverter& LwcConverter, void* DestAddress) { CopyPropertyHelper<int32, uint8>(PropertyAddress, DestAddress); };
+			}
+			else if (InProperty->IsA<const FUInt16Property>())
+			{
+				return [PropertyAddress](const FNiagaraLWCConverter& LwcConverter, void* DestAddress) { CopyPropertyHelper<int32, uint16>(PropertyAddress, DestAddress); };
+			}
+			else if (InProperty->IsA<const FUInt32Property>())
+			{
+				return [PropertyAddress](const FNiagaraLWCConverter& LwcConverter, void* DestAddress) { CopyPropertyHelper<int32, uint32>(PropertyAddress, DestAddress); };
+			}
+			else if (InProperty->IsA<const FInt16Property>())
+			{
+				return [PropertyAddress](const FNiagaraLWCConverter& LwcConverter, void* DestAddress) { CopyPropertyHelper<int32, int16>(PropertyAddress, DestAddress); };
+			}
+			else if (InProperty->IsA<const FIntProperty>())
+			{
+				return [PropertyAddress](const FNiagaraLWCConverter& LwcConverter, void* DestAddress) { CopyPropertyHelper<int32, int32>(PropertyAddress, DestAddress); };
+			}
+			return nullptr;
+		}
+		static TFunction<void(const FNiagaraLWCConverter&, void*)> GetGetterFunction(const FProperty* InProperty, UObject* BoundObject, UFunction* BoundFunction)
+		{
+			if (const FEnumProperty* EnumProperty = CastField<FEnumProperty>(InProperty))
+			{
+				InProperty = EnumProperty->GetUnderlyingProperty();
+			}
+
+			if (InProperty->IsA<const FByteProperty>())
+			{
+				return [BoundObject, BoundFunction](const FNiagaraLWCConverter& LwcConverter, void* DestAddress) { InvokeGetterHelper<int32, uint8>(BoundObject, BoundFunction, DestAddress, 0); };
 			}
 			else if (InProperty->IsA<const FUInt16Property>())
 			{
-				return [PropertyAddress](const FNiagaraLWCConverter& LwcConverter, void* DestAddress) { *reinterpret_cast<int32*>(DestAddress) = *reinterpret_cast<const uint16*>(PropertyAddress); };
+				return [BoundObject, BoundFunction](const FNiagaraLWCConverter& LwcConverter, void* DestAddress) { InvokeGetterHelper<int32, uint16>(BoundObject, BoundFunction, DestAddress, 0); };
 			}
 			else if (InProperty->IsA<const FUInt32Property>())
 			{
-				return [PropertyAddress](const FNiagaraLWCConverter& LwcConverter, void* DestAddress) { *reinterpret_cast<int32*>(DestAddress) = *reinterpret_cast<const uint32*>(PropertyAddress); };
+				return [BoundObject, BoundFunction](const FNiagaraLWCConverter& LwcConverter, void* DestAddress) { InvokeGetterHelper<int32, uint32>(BoundObject, BoundFunction, DestAddress, 0); };
 			}
 			else if (InProperty->IsA<const FInt16Property>())
 			{
-				return [PropertyAddress](const FNiagaraLWCConverter& LwcConverter, void* DestAddress) { *reinterpret_cast<int32*>(DestAddress) = *reinterpret_cast<const int16*>(PropertyAddress); };
+				return [BoundObject, BoundFunction](const FNiagaraLWCConverter& LwcConverter, void* DestAddress) { InvokeGetterHelper<int32, uint16>(BoundObject, BoundFunction, DestAddress, 0); };
 			}
 			else if (InProperty->IsA<const FIntProperty>())
 			{
-				return [PropertyAddress](const FNiagaraLWCConverter& LwcConverter, void* DestAddress) { *reinterpret_cast<int32*>(DestAddress) = *reinterpret_cast<const int32*>(PropertyAddress); };
+				return [BoundObject, BoundFunction](const FNiagaraLWCConverter& LwcConverter, void* DestAddress) { InvokeGetterHelper<int32, int32>(BoundObject, BoundFunction, DestAddress, 0); };
 			}
 			return nullptr;
 		}
@@ -363,6 +489,14 @@ namespace NDIUObjectPropertyReaderLocal
 			}
 			return nullptr;
 		}
+		static TFunction<void(const FNiagaraLWCConverter&, void*)> GetGetterFunction(const FProperty* InProperty, UObject* BoundObject, UFunction* BoundFunction)
+		{
+			if (const FBoolProperty* BoolProperty = CastField<const FBoolProperty>(InProperty))
+			{
+				return [BoundObject, BoundFunction](const FNiagaraLWCConverter& LwcConverter, void* DestAddress) { InvokeGetterHelper<FNiagaraBool, bool>(BoundObject, BoundFunction, DestAddress, false); };
+			}
+			return nullptr;
+		}
 		static constexpr TCHAR const* HlslBufferType = TEXT("bool");
 		static constexpr TCHAR const* HlslBufferRead = TEXT("asint(BUFFER[OFFSET]) != 0 : 0");
 	};
@@ -392,6 +526,43 @@ namespace NDIUObjectPropertyReaderLocal
 		UStruct* ObjectClass = ObjectBinding->GetClass();
 		void* ObjectAddress = ObjectBinding;
 
+		// Look for a getter function, i.e. one with a single return value
+		if (UFunction* GetterFunction = ObjectBinding->FindFunction(PropertyGetter.Variable.GetName()))
+		{
+			FProperty* ReturnProperty = nullptr;
+			if (GetterFunction->NumParms == 1 && GetterFunction->PropertyLink != nullptr && GetterFunction->PropertyLink->HasAnyPropertyFlags(CPF_OutParm))
+			{
+				ReturnProperty = GetterFunction->PropertyLink;
+			}
+			//-TODO: Support native function path??
+			//else if (GetterFunction->NumParms == 0)
+			//{
+			//	ReturnProperty = GetterFunction->GetReturnProperty();
+			//}
+
+			if (ReturnProperty)
+			{
+				#define NDI_PROPERTY_TYPE(TYPE) \
+					else if (PropertyGetter.Variable.GetType() == FTypeHelper<TYPE>::GetTypeDef()) \
+					{ \
+						PropertyGetter.PropertyCopyFunction = FTypeHelper<TYPE>::GetGetterFunction(ReturnProperty, ObjectBinding, GetterFunction); \
+						if (PropertyGetter.PropertyCopyFunction == nullptr && FNiagaraUtilities::LogVerboseWarnings()) \
+						{ \
+							const FStructProperty* StructProperty = CastField<FStructProperty>(ReturnProperty); \
+							const FString PropertyType = StructProperty && StructProperty->Struct ? StructProperty->Struct->GetName() : ReturnProperty->GetClass()->GetName(); \
+							UE_LOG(LogNiagara, Warning, TEXT("Could convert return property for function '%s' type '%s' into expected Niagara type '%s'"), *PropertyGetter.Variable.GetName().ToString(), *PropertyType, *PropertyGetter.Variable.GetType().GetName()); \
+						} \
+					} \
+
+					if (false) { } NDI_PROPERTY_TYPES
+				#undef NDI_PROPERTY_TYPE
+			}
+
+			// Complete as BP can not have both a function and variable of the same type
+			return;
+		}
+
+		// Look for a property
 		while ( true )
 		{
 			FStringView PropertyName = VariablePath;
@@ -482,6 +653,7 @@ bool UNiagaraDataInterfaceUObjectPropertyReader::Equals(const UNiagaraDataInterf
 
 	auto* OtherTyped = CastChecked<const UNiagaraDataInterfaceUObjectPropertyReader>(Other);
 	return
+		OtherTyped->SourceMode == SourceMode &&
 		OtherTyped->UObjectParameterBinding == UObjectParameterBinding &&
 		OtherTyped->PropertyRemap == PropertyRemap &&
 		OtherTyped->SourceActor == SourceActor &&
@@ -496,6 +668,7 @@ bool UNiagaraDataInterfaceUObjectPropertyReader::CopyToInternal(UNiagaraDataInte
 	}
 
 	auto* DestinationTyped = CastChecked<UNiagaraDataInterfaceUObjectPropertyReader>(Destination);
+	DestinationTyped->SourceMode = SourceMode;
 	DestinationTyped->UObjectParameterBinding = UObjectParameterBinding;
 	DestinationTyped->PropertyRemap = PropertyRemap;
 	DestinationTyped->SourceActor = SourceActor;
@@ -609,13 +782,22 @@ bool UNiagaraDataInterfaceUObjectPropertyReader::PerInstanceTick(void* PerInstan
 		return true;
 	}
 
-	// Do we need to rebind our properties?
-	UObject* ObjectBinding = InstanceData_GT->UObjectBinding.GetValue();
-	if ( ObjectBinding == nullptr )
+	// Resolve our source object
+	UObject* ObjectBinding = nullptr;
+	if (SourceMode != ENDIObjectPropertyReaderSourceMode::AttachParentActor)
+	{
+		ObjectBinding = InstanceData_GT->UObjectBinding.GetValue();
+		ObjectBinding = ObjectBinding ? ObjectBinding : SourceActor.Get();
+	}
+	if (SourceMode != ENDIObjectPropertyReaderSourceMode::Binding && ObjectBinding == nullptr)
 	{
-		ObjectBinding = SourceActor.Get();
+		if (USceneComponent* AttachComponent = SystemInstance->GetAttachComponent())
+		{
+			ObjectBinding = AttachComponent->GetAttachParentActor();
+		}
 	}
 
+	// Do we need to rebind the reader?
 	if ( InstanceData_GT->WeakUObject.Get() != ObjectBinding )
 	{
 		InstanceData_GT->WeakUObject = ObjectBinding;
diff --git a/Engine/Plugins/FX/Niagara/Source/Niagara/Private/DataInterface/NiagaraDataInterfaceUObjectPropertyReader.h b/Engine/Plugins/FX/Niagara/Source/Niagara/Private/DataInterface/NiagaraDataInterfaceUObjectPropertyReader.h
--- a/Engine/Plugins/FX/Niagara/Source/Niagara/Private/DataInterface/NiagaraDataInterfaceUObjectPropertyReader.h
+++ b/Engine/Plugins/FX/Niagara/Source/Niagara/Private/DataInterface/NiagaraDataInterfaceUObjectPropertyReader.h
@@ -8,6 +8,18 @@
 
 struct FNiagaraDataInterfaceGeneratedFunction;
 
+UENUM()
+enum class ENDIObjectPropertyReaderSourceMode : uint8
+{
+	/** Source object is found in the order of object binding, source actor. */
+	Binding,
+	/** Source object is found use the attached parent actor. */
+	AttachParentActor,
+	/** Source object is found in the order of object binding, source actor, then attached parent actor. */
+	BindingThenAttachParentActor,
+};
+
+
 USTRUCT()
 struct FNiagaraUObjectPropertyReaderRemap
 {
@@ -49,15 +61,19 @@ class UNiagaraDataInterfaceUObjectPropertyReader : public UNiagaraDataInterface
 	END_SHADER_PARAMETER_STRUCT()
 
 public:
-	/** User parameter Object binding to read properties from. */
+	/** Determines how we should select the source object we read from. */
 	UPROPERTY(EditAnywhere, Category = "UObjectPropertyReader")
+	ENDIObjectPropertyReaderSourceMode SourceMode = ENDIObjectPropertyReaderSourceMode::Binding;
+
+	/** User parameter Object binding to read properties from. */
+	UPROPERTY(EditAnywhere, Category = "UObjectPropertyReader", meta = (EditCondition = "SourceMode != ENDIObjectPropertyReaderSourceMode::AttachParentActor", EditConditionHides))
 	FNiagaraUserParameterBinding UObjectParameterBinding;
 
 	UPROPERTY(EditAnywhere, Category = "UObjectPropertyReader")
 	TArray<FNiagaraUObjectPropertyReaderRemap> PropertyRemap;
 
 	/** Optional source actor to use, if the user parameter binding is valid this will be ignored. */
-	UPROPERTY(EditAnywhere, Category = "UObjectPropertyReader")
+	UPROPERTY(EditAnywhere, Category = "UObjectPropertyReader", meta = (EditCondition = "SourceMode != ENDIObjectPropertyReaderSourceMode::AttachParentActor", EditConditionHides))
 	TSoftObjectPtr<AActor> SourceActor;
 
 	/**
