Niagara Vulkan nDisplay crash fix.  Exit CopyToMultiViewCountBuffer if AllocatedInstanceCounts is zero, as this leads to an attempt to copy between NULL resources, which crashes on Vulkan (D3D platforms treat this as a NOP and don't crash).

#rnx

diff --git a/Engine/Plugins/FX/Niagara/Source/Niagara/Private/NiagaraGPUInstanceCountManager.cpp b/Engine/Plugins/FX/Niagara/Source/Niagara/Private/NiagaraGPUInstanceCountManager.cpp
--- a/Engine/Plugins/FX/Niagara/Source/Niagara/Private/NiagaraGPUInstanceCountManager.cpp
+++ b/Engine/Plugins/FX/Niagara/Source/Niagara/Private/NiagaraGPUInstanceCountManager.cpp
@@ -654,24 +654,27 @@ bool FNiagaraGPUInstanceCountManager::HasPendingGPUReadback() const
 
 void FNiagaraGPUInstanceCountManager::CopyToMultiViewCountBuffer(FRHICommandListImmediate& RHICmdList)
 {
-	// Need to copy on all GPUs
-	SCOPED_GPU_MASK(RHICmdList, FRHIGPUMask::All());
-
-	// Set AllocatedInstanceCounts and copy CountBuffer
-	if (MultiViewAllocatedInstanceCounts != AllocatedInstanceCounts)
-	{
-		MultiViewAllocatedInstanceCounts = AllocatedInstanceCounts;
-		MultiViewCountBuffer.Initialize(RHICmdList, TEXT("NiagaraGPUInstanceCounts"), sizeof(uint32), MultiViewAllocatedInstanceCounts, EPixelFormat::PF_R32_UINT, ERHIAccess::UAVCompute, BUF_Static | BUF_SourceCopy);
-	}
-	else
+	if (AllocatedInstanceCounts > 0)
 	{
-		RHICmdList.Transition(FRHITransitionInfo(MultiViewCountBuffer.UAV, kCountBufferDefaultState, ERHIAccess::UAVCompute));
-	}
+		// Need to copy on all GPUs
+		SCOPED_GPU_MASK(RHICmdList, FRHIGPUMask::All());
+
+		// Set AllocatedInstanceCounts and copy CountBuffer
+		if (MultiViewAllocatedInstanceCounts != AllocatedInstanceCounts)
+		{
+			MultiViewAllocatedInstanceCounts = AllocatedInstanceCounts;
+			MultiViewCountBuffer.Initialize(RHICmdList, TEXT("NiagaraGPUInstanceCounts"), sizeof(uint32), MultiViewAllocatedInstanceCounts, EPixelFormat::PF_R32_UINT, ERHIAccess::UAVCompute, BUF_Static | BUF_SourceCopy);
+		}
+		else
+		{
+			RHICmdList.Transition(FRHITransitionInfo(MultiViewCountBuffer.UAV, kCountBufferDefaultState, ERHIAccess::UAVCompute));
+		}
 
-	FRHIUnorderedAccessView* UAVs[] = { MultiViewCountBuffer.UAV };
-	int32 UsedIndexCounts[] = { MultiViewAllocatedInstanceCounts };
-	CopyUIntBufferToTargets(RHICmdList, FeatureLevel, CountBuffer.SRV, UAVs, UsedIndexCounts, 0, UE_ARRAY_COUNT(UAVs));
+		FRHIUnorderedAccessView* UAVs[] = { MultiViewCountBuffer.UAV };
+		int32 UsedIndexCounts[] = { MultiViewAllocatedInstanceCounts };
+		CopyUIntBufferToTargets(RHICmdList, FeatureLevel, CountBuffer.SRV, UAVs, UsedIndexCounts, 0, UE_ARRAY_COUNT(UAVs));
 
-	RHICmdList.Transition(FRHITransitionInfo(MultiViewCountBuffer.UAV, ERHIAccess::UAVCompute, kCountBufferDefaultState));
+		RHICmdList.Transition(FRHITransitionInfo(MultiViewCountBuffer.UAV, ERHIAccess::UAVCompute, kCountBufferDefaultState));
+	}
 }
 
