- Change IsVisible for pooled components to only check the attached components level collection rather than the attached component
- This mirrors what Cascade does and fixes an issue with gameplay cues that are attached as they are sometimes attached to hidden components but we don't actually want to hide our component

diff --git a/Engine/Plugins/FX/Niagara/Source/Niagara/Private/NiagaraComponent.cpp b/Engine/Plugins/FX/Niagara/Source/Niagara/Private/NiagaraComponent.cpp
--- a/Engine/Plugins/FX/Niagara/Source/Niagara/Private/NiagaraComponent.cpp
+++ b/Engine/Plugins/FX/Niagara/Source/Niagara/Private/NiagaraComponent.cpp
@@ -2261,13 +2261,20 @@ void UNiagaraComponent::OnChildDetached(USceneComponent* ChildComponent)
 
 bool UNiagaraComponent::IsVisible() const
 {
-	if (PoolingMethod == ENCPoolMethod::None || GetAttachParent() == nullptr)
+	// Pooled components are owned by the persistent level, so we need to check their current attached parent's level visibility, otherwise they ignore sublevel visibility
+	// Note we do not test the attached parents visibility as that doesn't match how Cascade works, and can cause visibility issues with gameplay cues who often attach to hidden components
+	if (PoolingMethod != ENCPoolMethod::None)
 	{
-		return Super::IsVisible();
+		if (USceneComponent* OwnerComponent = GetAttachParent())
+		{
+			if (OwnerComponent->CachedLevelCollection && !OwnerComponent->CachedLevelCollection->IsVisible())
+			{
+				return false;
+			}
+		}
 	}
 
-	// pooled components are owned by the persistent level, so we need to check their current attached parent visibility, otherwise they ignore sublevel visibility
-	return Super::IsVisible() && GetAttachParent()->IsVisible();
+	return Super::IsVisible();
 }
 
 void UNiagaraComponent::SetTickBehavior(ENiagaraTickBehavior NewTickBehavior)
