- Overload RequiresGameThreadEndOfFrameRecreate for Niagara as we are safe to run async

diff --git a/Engine/Plugins/FX/Niagara/Source/Niagara/Public/NiagaraComponent.h b/Engine/Plugins/FX/Niagara/Source/Niagara/Public/NiagaraComponent.h
--- a/Engine/Plugins/FX/Niagara/Source/Niagara/Public/NiagaraComponent.h
+++ b/Engine/Plugins/FX/Niagara/Source/Niagara/Public/NiagaraComponent.h
@@ -74,6 +74,8 @@ public:
 
 	NIAGARA_API virtual FName GetFNameForStatID() const override;
 
+	virtual bool RequiresGameThreadEndOfFrameRecreate() const override { return false; }
+
 private:
 	UPROPERTY(EditAnywhere, Category="Niagara", meta = (DisplayName = "Niagara System Asset"))
 	TObjectPtr<UNiagaraSystem> Asset;
