- Change UsedBy / UsedWith to refer to Script rather than Emitter as it included the system scripts
- Deprecate the various combinations we had to prefer these functions

diff --git a/Engine/Plugins/FX/Niagara/Source/Niagara/Classes/NiagaraDataInterface.h b/Engine/Plugins/FX/Niagara/Source/Niagara/Classes/NiagaraDataInterface.h
--- a/Engine/Plugins/FX/Niagara/Source/Niagara/Classes/NiagaraDataInterface.h
+++ b/Engine/Plugins/FX/Niagara/Source/Niagara/Classes/NiagaraDataInterface.h
@@ -681,12 +681,6 @@ public:
 	virtual bool HasTickGroupPrereqs() const { return false; }
 	virtual ETickingGroup CalculateTickGroup(const void* PerInstanceData) const { return NiagaraFirstTickGroup; }
 
-	/** Used to determine if we need to create CPU resources for the emitter. */
-	NIAGARA_API bool IsUsedWithCPUEmitter() const;
-
-	/** Used to determine if we need to create GPU resources for the emitter. */
-	NIAGARA_API bool IsUsedWithGPUEmitter() const;
-
 	/** Determines if this type definition matches to a known data interface type.*/
 	static NIAGARA_API bool IsDataInterfaceType(const FNiagaraTypeDefinition& TypeDef);
 
@@ -809,7 +803,7 @@ protected:
 public:
 	void PushToRenderThread()
 	{
-		if (bUsedByGPUEmitter && bRenderDataDirty)
+		if (bUsedWithGPUScript && bRenderDataDirty)
 		{
 			PushToRenderThreadImpl();
 			bRenderDataDirty = false;
@@ -822,21 +816,32 @@ public:
 		PushToRenderThread();
 	}
 
-	void SetUsedByCPUEmitter(bool bUsed = true)
+	void SetUsedWithCPUScript(bool bUsed = true)
 	{
 		check(IsInGameThread());
-		bUsedByCPUEmitter = bUsed;
+		bUsedWithCPUScript = bUsed;
 	}
 
-	void SetUsedByGPUEmitter(bool bUsed = true)
+	void SetUsedWithGPUScript(bool bUsed = true)
 	{
 		check(IsInGameThread());
-		bUsedByGPUEmitter = bUsed;
+		bUsedWithGPUScript = bUsed;
 		PushToRenderThread();
 	}
 
-	bool IsUsedByCPUEmitter() const { return bUsedByCPUEmitter; }
-	bool IsUsedByGPUEmitter() const { return bUsedByGPUEmitter; }
+	UE_DEPRECATED(5.3, "Please update your code to use IsUsedWithCPUScript")
+	bool IsUsedByCPUEmitter() const { return bUsedWithCPUScript; }
+	UE_DEPRECATED(5.3, "Please update your code to use IsUsedWithGPUScript")
+	bool IsUsedByGPUEmitter() const { return bUsedWithGPUScript; }
+	UE_DEPRECATED(5.3, "Please update your code to use IsUsedWithCPUScript")
+	NIAGARA_API bool IsUsedWithCPUEmitter() const { return bUsedWithCPUScript; }
+	UE_DEPRECATED(5.3, "Please update your code to use IsUsedWithGPUScript")
+	NIAGARA_API bool IsUsedWithGPUEmitter() const { return bUsedWithGPUScript; }
+
+	/** Used to determine if we need to create CPU resources for the emitter. */
+	bool IsUsedWithCPUScript() const { return bUsedWithCPUScript; }
+	/** Used to determine if we need to create GPU resources for the emitter. */
+	bool IsUsedWithGPUScript() const { return bUsedWithGPUScript; }
 
 protected:
 	template<typename T>
@@ -860,8 +865,8 @@ protected:
 	TUniquePtr<FNiagaraDataInterfaceProxy> Proxy;
 
 	uint32 bRenderDataDirty : 1;
-	uint32 bUsedByCPUEmitter : 1;
-	uint32 bUsedByGPUEmitter : 1;
+	uint32 bUsedWithCPUScript : 1;
+	uint32 bUsedWithGPUScript : 1;
 
 private:
 #if WITH_EDITOR
diff --git a/Engine/Plugins/FX/Niagara/Source/Niagara/Classes/NiagaraDataInterfaceArrayImpl.h b/Engine/Plugins/FX/Niagara/Source/Niagara/Classes/NiagaraDataInterfaceArrayImpl.h
--- a/Engine/Plugins/FX/Niagara/Source/Niagara/Classes/NiagaraDataInterfaceArrayImpl.h
+++ b/Engine/Plugins/FX/Niagara/Source/Niagara/Classes/NiagaraDataInterfaceArrayImpl.h
@@ -430,7 +430,7 @@ struct FNDIArrayProxyImpl : public INDIArrayProxyBase
 	void CachePropertiesFromOwner()
 	{
 		bShouldSyncToGpu = FNiagaraUtilities::ShouldSyncCpuToGpu(Owner->GpuSyncMode);
-		bShouldSyncToCpu = FNiagaraUtilities::ShouldSyncGpuToCpu(Owner->GpuSyncMode) && Owner->IsUsedByCPUEmitter();
+		bShouldSyncToCpu = FNiagaraUtilities::ShouldSyncGpuToCpu(Owner->GpuSyncMode) && Owner->IsUsedWithCPUScript();
 	}
 
 	//////////////////////////////////////////////////////////////////////////
@@ -1076,7 +1076,7 @@ struct FNDIArrayProxyImpl : public INDIArrayProxyBase
 
 		PerInstanceData_GameThread.Emplace(SystemInstance->GetId(), InstanceData_GT);
 
-		if ( FNDIArrayImplHelper<TArrayType>::bSupportsGPU && Owner->IsUsedWithGPUEmitter() )
+		if ( FNDIArrayImplHelper<TArrayType>::bSupportsGPU && Owner->IsUsedWithGPUScript() )
 		{
 			bool bRWGpuArray = false;
 			FNiagaraDataInterfaceUtilities::ForEachGpuFunction(
@@ -1105,7 +1105,7 @@ struct FNDIArrayProxyImpl : public INDIArrayProxyBase
 	{
 		auto* InstanceData_GT = reinterpret_cast<FNDIArrayInstanceData_GameThread<TArrayType>*>(InPerInstanceData);
 
-		if ( FNDIArrayImplHelper<TArrayType>::bSupportsGPU && Owner->IsUsedWithGPUEmitter() )
+		if ( FNDIArrayImplHelper<TArrayType>::bSupportsGPU && Owner->IsUsedWithGPUScript() )
 		{
 			ENQUEUE_RENDER_COMMAND(FNDIArrayProxyImpl_RemoveProxy)
 			(
diff --git a/Engine/Plugins/FX/Niagara/Source/Niagara/Private/DataInterface/NiagaraDataInterfaceDynamicMesh.cpp b/Engine/Plugins/FX/Niagara/Source/Niagara/Private/DataInterface/NiagaraDataInterfaceDynamicMesh.cpp
--- a/Engine/Plugins/FX/Niagara/Source/Niagara/Private/DataInterface/NiagaraDataInterfaceDynamicMesh.cpp
+++ b/Engine/Plugins/FX/Niagara/Source/Niagara/Private/DataInterface/NiagaraDataInterfaceDynamicMesh.cpp
@@ -1048,7 +1048,7 @@ bool UNiagaraDataInterfaceDynamicMesh::InitPerInstanceData(void* PerInstanceData
 	FNDIInstanceData_GameThread* InstanceData = new(PerInstanceData) FNDIInstanceData_GameThread();
 	InstanceData->SystemInstanceID = SystemInstance->GetId();
 	InstanceData->bClearTrianglesPerFrame = bClearTrianglesPerFrame;
-	InstanceData->bGpuUsesDynamicAllocation = IsUsedWithGPUEmitter();		//-OPT: We can improve this by looking for used functions
+	InstanceData->bGpuUsesDynamicAllocation = IsUsedWithGPUScript();		//-OPT: We can improve this by looking for used functions
 	InstanceData->LocalBounds = DefaultBounds;
 
 	InstanceData->Sections.SetNum(Sections.Num());
diff --git a/Engine/Plugins/FX/Niagara/Source/Niagara/Private/DataInterface/NiagaraDataInterfaceEmitterProperties.cpp b/Engine/Plugins/FX/Niagara/Source/Niagara/Private/DataInterface/NiagaraDataInterfaceEmitterProperties.cpp
--- a/Engine/Plugins/FX/Niagara/Source/Niagara/Private/DataInterface/NiagaraDataInterfaceEmitterProperties.cpp
+++ b/Engine/Plugins/FX/Niagara/Source/Niagara/Private/DataInterface/NiagaraDataInterfaceEmitterProperties.cpp
@@ -128,7 +128,7 @@ bool UNiagaraDataInterfaceEmitterProperties::InitPerInstanceData(void* PerInstan
 		InstanceData_GT->bLocalSpace = InstanceData_GT->EmitterInstance->GetCachedEmitterData()->bLocalSpace;
 	}
 
-	if ( IsUsedByGPUEmitter() )
+	if ( IsUsedWithGPUScript() )
 	{
 		// Initialize render side instance data
 		ENQUEUE_RENDER_COMMAND(NDIEmitter_InitRT)
@@ -151,7 +151,7 @@ void UNiagaraDataInterfaceEmitterProperties::DestroyPerInstanceData(void* PerIns
 	FInstanceData_GameThread* InstanceData_GT = reinterpret_cast<FInstanceData_GameThread*>(PerInstanceData);
 	InstanceData_GT->~FInstanceData_GameThread();
 
-	if ( IsUsedByGPUEmitter() )
+	if ( IsUsedWithGPUScript() )
 	{
 		ENQUEUE_RENDER_COMMAND(NDIEmitter_InitRT)
 		(
diff --git a/Engine/Plugins/FX/Niagara/Source/Niagara/Private/DataInterface/NiagaraDataInterfaceSceneCapture2D.cpp b/Engine/Plugins/FX/Niagara/Source/Niagara/Private/DataInterface/NiagaraDataInterfaceSceneCapture2D.cpp
--- a/Engine/Plugins/FX/Niagara/Source/Niagara/Private/DataInterface/NiagaraDataInterfaceSceneCapture2D.cpp
+++ b/Engine/Plugins/FX/Niagara/Source/Niagara/Private/DataInterface/NiagaraDataInterfaceSceneCapture2D.cpp
@@ -727,7 +727,7 @@ bool UNiagaraDataInterfaceSceneCapture2D::PerInstanceTickPostSimulate(void* PerI
 
 	// We need to update the texture ideally this is part of the tick but that isn't possible
 	// since the resource could have been deleted / changed before the ticks are processed
-	if (IsUsedWithGPUEmitter())
+	if (IsUsedWithGPUScript())
 	{
 		ENQUEUE_RENDER_COMMAND(NDISceneCapture2D_UpdateInstance)
 		(
diff --git a/Engine/Plugins/FX/Niagara/Source/Niagara/Private/DataInterface/NiagaraDataInterfaceSimCacheReader.cpp b/Engine/Plugins/FX/Niagara/Source/Niagara/Private/DataInterface/NiagaraDataInterfaceSimCacheReader.cpp
--- a/Engine/Plugins/FX/Niagara/Source/Niagara/Private/DataInterface/NiagaraDataInterfaceSimCacheReader.cpp
+++ b/Engine/Plugins/FX/Niagara/Source/Niagara/Private/DataInterface/NiagaraDataInterfaceSimCacheReader.cpp
@@ -244,7 +244,7 @@ void UNiagaraDataInterfaceSimCacheReader::DestroyPerInstanceData(void* PerInstan
 	FInstanceData_GameThread* InstanceData_GT = reinterpret_cast<FInstanceData_GameThread*>(PerInstanceData);
 	InstanceData_GT->~FInstanceData_GameThread();
 
-	if ( IsUsedByGPUEmitter() )
+	if ( IsUsedWithGPUScript() )
 	{
 		ENQUEUE_RENDER_COMMAND(NDISimCacheReader_DestroyRT)
 		(
@@ -317,7 +317,7 @@ bool UNiagaraDataInterfaceSimCacheReader::PerInstanceTick(void* PerInstanceData,
 	else if (InstanceSimCache == nullptr)
 	{
 		InstanceData_GT->EmitterIndex = INDEX_NONE;
-		if ( IsUsedByGPUEmitter() )
+		if ( IsUsedWithGPUScript() )
 		{
 			ENQUEUE_RENDER_COMMAND(NDISimCacheReader_DestroyRT)
 			(
diff --git a/Engine/Plugins/FX/Niagara/Source/Niagara/Private/DataInterface/NiagaraDataInterfaceSimpleCounter.cpp b/Engine/Plugins/FX/Niagara/Source/Niagara/Private/DataInterface/NiagaraDataInterfaceSimpleCounter.cpp
--- a/Engine/Plugins/FX/Niagara/Source/Niagara/Private/DataInterface/NiagaraDataInterfaceSimpleCounter.cpp
+++ b/Engine/Plugins/FX/Niagara/Source/Niagara/Private/DataInterface/NiagaraDataInterfaceSimpleCounter.cpp
@@ -212,7 +212,7 @@ bool UNiagaraDataInterfaceSimpleCounter::InitPerInstanceData(void* PerInstanceDa
 	FNDISimpleCounterInstanceData_GameThread* InstanceData_GT = new(PerInstanceData) FNDISimpleCounterInstanceData_GameThread();
 	InstanceData_GT->Counter = InitialValue;
 
-	if ( IsUsedWithGPUEmitter() )
+	if ( IsUsedWithGPUScript() )
 	{
 		FNDISimpleCounterProxy* Proxy_GT = GetProxyAs<FNDISimpleCounterProxy>();
 		Proxy_GT->PerInstanceData_GameThread.Add(SystemInstance->GetId(), InstanceData_GT);
@@ -235,7 +235,7 @@ void UNiagaraDataInterfaceSimpleCounter::DestroyPerInstanceData(void* PerInstanc
 	auto* InstanceData = reinterpret_cast<FNDISimpleCounterInstanceData_GameThread*>(PerInstanceData);
 	InstanceData->~FNDISimpleCounterInstanceData_GameThread();
 
-	if ( IsUsedWithGPUEmitter() )
+	if ( IsUsedWithGPUScript() )
 	{
 		FNDISimpleCounterProxy* Proxy_GT = GetProxyAs<FNDISimpleCounterProxy>();
 		Proxy_GT->PerInstanceData_GameThread.Remove(SystemInstance->GetId());
diff --git a/Engine/Plugins/FX/Niagara/Source/Niagara/Private/DataInterface/NiagaraDataInterfaceStaticMesh.cpp b/Engine/Plugins/FX/Niagara/Source/Niagara/Private/DataInterface/NiagaraDataInterfaceStaticMesh.cpp
--- a/Engine/Plugins/FX/Niagara/Source/Niagara/Private/DataInterface/NiagaraDataInterfaceStaticMesh.cpp
+++ b/Engine/Plugins/FX/Niagara/Source/Niagara/Private/DataInterface/NiagaraDataInterfaceStaticMesh.cpp
@@ -885,7 +885,7 @@ namespace NDIStaticMeshLocal
 			else if (!StaticMesh->bAllowCPUAccess)
 			{
 				//-TODO: This will be modified with the refactor to allow GPU to be used when CPU is also used
-				if (Interface->IsUsedWithGPUEmitter() && !Interface->IsUsedWithCPUEmitter())
+				if (Interface->IsUsedWithGPUScript() && !Interface->IsUsedWithCPUScript())
 				{
 					if (!FNiagaraUtilities::AreBufferSRVsAlwaysCreated(GMaxRHIShaderPlatform))
 					{
@@ -908,7 +908,7 @@ namespace NDIStaticMeshLocal
 
 					if (bNeedsCpuAccess)
 					{
-						if (Interface->IsUsedWithCPUEmitter())
+						if (Interface->IsUsedWithCPUScript())
 						{
 							UE_LOG(LogNiagara, Log, TEXT("NiagaraStaticMeshDataInterface used by CPU emitter and does not allow CPU access. System: %s, Mesh: %s"), *GetFullNameSafe(SystemInstance->GetSystem()), *GetFullNameSafe(StaticMesh));
 						}
@@ -1054,7 +1054,7 @@ namespace NDIStaticMeshLocal
 
 				const bool MeshValid = IsValid(StaticMesh);
 				const bool SupportUvMappingCpu = UsedByCpuUvMapping && MeshValid;
-				const bool SupportUvMappingGpu = UsedByGpuUvMapping && MeshValid && Interface->IsUsedWithGPUEmitter();
+				const bool SupportUvMappingGpu = UsedByGpuUvMapping && MeshValid && Interface->IsUsedWithGPUScript();
 
 				UvMappingUsage = FMeshUvMappingUsage(SupportUvMappingCpu, SupportUvMappingGpu);
 
@@ -1763,7 +1763,7 @@ bool UNiagaraDataInterfaceStaticMesh::InitPerInstanceData(void* PerInstanceData,
 	const bool bIsValid = InstanceData->Init(this, SystemInstance);
 
 	// Create render thread data?
-	if (bIsValid && IsUsedWithGPUEmitter() )
+	if (bIsValid && IsUsedWithGPUScript() )
 	{
 		TUniquePtr<NDIStaticMeshLocal::FGpuInitializeData> GpuInitializeData = MakeUnique<NDIStaticMeshLocal::FGpuInitializeData>();
 		GpuInitializeData->RenderProxy = GetProxyAs<NDIStaticMeshLocal::FRenderProxy>();
@@ -1872,7 +1872,7 @@ void UNiagaraDataInterfaceStaticMesh::DestroyPerInstanceData(void* PerInstanceDa
 	InstanceData->Release();
 	InstanceData->~FInstanceData_GameThread();
 
-	if ( IsUsedWithGPUEmitter() )
+	if ( IsUsedWithGPUScript() )
 	{
 		NDIStaticMeshLocal::FRenderProxy* Proxy_RT = GetProxyAs<NDIStaticMeshLocal::FRenderProxy>();
 		FNiagaraSystemInstanceID InstanceID_RT = SystemInstance->GetId();
diff --git a/Engine/Plugins/FX/Niagara/Source/Niagara/Private/DataInterface/NiagaraDataInterfaceUObjectPropertyReader.cpp b/Engine/Plugins/FX/Niagara/Source/Niagara/Private/DataInterface/NiagaraDataInterfaceUObjectPropertyReader.cpp
--- a/Engine/Plugins/FX/Niagara/Source/Niagara/Private/DataInterface/NiagaraDataInterfaceUObjectPropertyReader.cpp
+++ b/Engine/Plugins/FX/Niagara/Source/Niagara/Private/DataInterface/NiagaraDataInterfaceUObjectPropertyReader.cpp
@@ -685,7 +685,7 @@ bool UNiagaraDataInterfaceUObjectPropertyReader::InitPerInstanceData(void* PerIn
 	InstanceData_GT->UObjectBinding.Init(SystemInstance->GetInstanceParameters(), UObjectParameterBinding.Parameter);
 	InstanceData_GT->ChangeId = ChangeId;
 
-	if ( IsUsedByGPUEmitter() )
+	if ( IsUsedWithGPUScript() )
 	{
 		TArray<uint32> GpuFunctionToPropertyRemap;
 
@@ -743,7 +743,7 @@ void UNiagaraDataInterfaceUObjectPropertyReader::DestroyPerInstanceData(void* Pe
 	FInstanceData_GameThread* InstanceData_GT = reinterpret_cast<FInstanceData_GameThread*>(PerInstanceData);
 	InstanceData_GT->~FInstanceData_GameThread();
 
-	if ( IsUsedByGPUEmitter() )
+	if ( IsUsedWithGPUScript() )
 	{
 		ENQUEUE_RENDER_COMMAND(NDIUObjectPropertyReader_InitRT)
 		(
diff --git a/Engine/Plugins/FX/Niagara/Source/Niagara/Private/NiagaraComponent.cpp b/Engine/Plugins/FX/Niagara/Source/Niagara/Private/NiagaraComponent.cpp
--- a/Engine/Plugins/FX/Niagara/Source/Niagara/Private/NiagaraComponent.cpp
+++ b/Engine/Plugins/FX/Niagara/Source/Niagara/Private/NiagaraComponent.cpp
@@ -3630,8 +3630,8 @@ void UNiagaraComponent::SetOverrideParameterStoreValue(const FNiagaraVariableBas
 		UNiagaraDataInterface* OriginalDI = InValue.GetDataInterface();
 		UNiagaraDataInterface* DuplicatedDI = DuplicateObject(OriginalDI, this);
 		OverrideParameters.SetDataInterface(DuplicatedDI, InKey);
-		DuplicatedDI->SetUsedByCPUEmitter(OriginalDI->IsUsedWithCPUEmitter());
-		DuplicatedDI->SetUsedByGPUEmitter(OriginalDI->IsUsedWithGPUEmitter());
+		DuplicatedDI->SetUsedWithCPUScript(OriginalDI->IsUsedWithCPUScript());
+		DuplicatedDI->SetUsedWithGPUScript(OriginalDI->IsUsedWithGPUScript());
 	}
 	else if (InKey.IsUObject())
 	{
diff --git a/Engine/Plugins/FX/Niagara/Source/Niagara/Private/NiagaraDataInterface.cpp b/Engine/Plugins/FX/Niagara/Source/Niagara/Private/NiagaraDataInterface.cpp
--- a/Engine/Plugins/FX/Niagara/Source/Niagara/Private/NiagaraDataInterface.cpp
+++ b/Engine/Plugins/FX/Niagara/Source/Niagara/Private/NiagaraDataInterface.cpp
@@ -90,8 +90,8 @@ void FNDIStageTickHandler::PostStageTick(FNiagaraSystemInstance* Instance, float
 UNiagaraDataInterface::UNiagaraDataInterface(FObjectInitializer const& ObjectInitializer)
 {
 	bRenderDataDirty = false;
-	bUsedByCPUEmitter = false;
-	bUsedByGPUEmitter = false;
+	bUsedWithCPUScript = false;
+	bUsedWithGPUScript = false;
 }
 
 UNiagaraDataInterface::~UNiagaraDataInterface()
@@ -195,16 +195,6 @@ bool UNiagaraDataInterface::Equals(const UNiagaraDataInterface* Other) const
 	return true;
 }
 
-bool UNiagaraDataInterface::IsUsedWithCPUEmitter() const
-{
-	return bUsedByCPUEmitter;
-}
-
-bool UNiagaraDataInterface::IsUsedWithGPUEmitter() const
-{
-	return bUsedByGPUEmitter;
-}
-
 bool UNiagaraDataInterface::IsDataInterfaceType(const FNiagaraTypeDefinition& TypeDef)
 {
 	const UClass* Class = TypeDef.GetClass();
diff --git a/Engine/Plugins/FX/Niagara/Source/Niagara/Private/NiagaraDataInterfaceIntRenderTarget2D.cpp b/Engine/Plugins/FX/Niagara/Source/Niagara/Private/NiagaraDataInterfaceIntRenderTarget2D.cpp
--- a/Engine/Plugins/FX/Niagara/Source/Niagara/Private/NiagaraDataInterfaceIntRenderTarget2D.cpp
+++ b/Engine/Plugins/FX/Niagara/Source/Niagara/Private/NiagaraDataInterfaceIntRenderTarget2D.cpp
@@ -607,7 +607,7 @@ bool UNiagaraDataInterfaceIntRenderTarget2D::InitPerInstanceData(void* PerInstan
 
 	FNDIIntRenderTarget2DInstanceData_GameThread* InstanceData = new (PerInstanceData) FNDIIntRenderTarget2DInstanceData_GameThread();
 
-	if (NiagaraDataInterfaceRenderTargetCommon::GIgnoreCookedOut && !IsUsedWithGPUEmitter())
+	if (NiagaraDataInterfaceRenderTargetCommon::GIgnoreCookedOut && !IsUsedWithGPUScript())
 	{
 		return true;
 	}
@@ -694,7 +694,7 @@ int32 UNiagaraDataInterfaceIntRenderTarget2D::PerInstanceDataSize() const
 bool UNiagaraDataInterfaceIntRenderTarget2D::PerInstanceTickPostSimulate(void* PerInstanceData, FNiagaraSystemInstance* SystemInstance, float DeltaSeconds)
 {
 	//-TEMP: Until we prune data interface on cook this will avoid consuming memory
-	if (NiagaraDataInterfaceRenderTargetCommon::GIgnoreCookedOut && !IsUsedWithGPUEmitter())
+	if (NiagaraDataInterfaceRenderTargetCommon::GIgnoreCookedOut && !IsUsedWithGPUScript())
 	{
 		return false;
 	}
diff --git a/Engine/Plugins/FX/Niagara/Source/Niagara/Private/NiagaraDataInterfaceRenderTarget2D.cpp b/Engine/Plugins/FX/Niagara/Source/Niagara/Private/NiagaraDataInterfaceRenderTarget2D.cpp
--- a/Engine/Plugins/FX/Niagara/Source/Niagara/Private/NiagaraDataInterfaceRenderTarget2D.cpp
+++ b/Engine/Plugins/FX/Niagara/Source/Niagara/Private/NiagaraDataInterfaceRenderTarget2D.cpp
@@ -454,7 +454,7 @@ bool UNiagaraDataInterfaceRenderTarget2D::InitPerInstanceData(void* PerInstanceD
 	SystemInstancesToProxyData_GT.Emplace(SystemInstance->GetId(), InstanceData);
 
 	//-TEMP: Until we prune data interface on cook this will avoid consuming memory
-	const bool bValidGpuDataInterface = NiagaraDataInterfaceRenderTargetCommon::GIgnoreCookedOut == 0 || IsUsedWithGPUEmitter();
+	const bool bValidGpuDataInterface = NiagaraDataInterfaceRenderTargetCommon::GIgnoreCookedOut == 0 || IsUsedWithGPUScript();
 
 	ETextureRenderTargetFormat RenderTargetFormat;
 	if ( NiagaraDataInterfaceRenderTargetCommon::GetRenderTargetFormat(bOverrideFormat, OverrideRenderTargetFormat, RenderTargetFormat) == false )
@@ -695,7 +695,7 @@ bool UNiagaraDataInterfaceRenderTarget2D::PerInstanceTickPostSimulate(void* PerI
 #endif
 
 	//-TEMP: Until we prune data interface on cook this will avoid consuming memory
-	const bool bValidGpuDataInterface = NiagaraDataInterfaceRenderTargetCommon::GIgnoreCookedOut == 0 || IsUsedWithGPUEmitter();
+	const bool bValidGpuDataInterface = NiagaraDataInterfaceRenderTargetCommon::GIgnoreCookedOut == 0 || IsUsedWithGPUScript();
 
 	if (::IsValid(InstanceData->TargetTexture) == false)
 	{
diff --git a/Engine/Plugins/FX/Niagara/Source/Niagara/Private/NiagaraDataInterfaceRenderTarget2DArray.cpp b/Engine/Plugins/FX/Niagara/Source/Niagara/Private/NiagaraDataInterfaceRenderTarget2DArray.cpp
--- a/Engine/Plugins/FX/Niagara/Source/Niagara/Private/NiagaraDataInterfaceRenderTarget2DArray.cpp
+++ b/Engine/Plugins/FX/Niagara/Source/Niagara/Private/NiagaraDataInterfaceRenderTarget2DArray.cpp
@@ -451,7 +451,7 @@ bool UNiagaraDataInterfaceRenderTarget2DArray::InitPerInstanceData(void* PerInst
 	FRenderTarget2DArrayRWInstanceData_GameThread* InstanceData = new (PerInstanceData) FRenderTarget2DArrayRWInstanceData_GameThread();
 
 	//-TEMP: Until we prune data interface on cook this will avoid consuming memory
-	const bool bValidGpuDataInterface = NiagaraDataInterfaceRenderTargetCommon::GIgnoreCookedOut == 0 || IsUsedWithGPUEmitter();
+	const bool bValidGpuDataInterface = NiagaraDataInterfaceRenderTargetCommon::GIgnoreCookedOut == 0 || IsUsedWithGPUScript();
 
 	ETextureRenderTargetFormat RenderTargetFormat;
 	if (NiagaraDataInterfaceRenderTargetCommon::GetRenderTargetFormat(bOverrideFormat, OverrideRenderTargetFormat, RenderTargetFormat) == false)
@@ -660,7 +660,7 @@ bool UNiagaraDataInterfaceRenderTarget2DArray::PerInstanceTickPostSimulate(void*
 #endif
 
 	//-TEMP: Until we prune data interface on cook this will avoid consuming memory
-	const bool bValidGpuDataInterface = NiagaraDataInterfaceRenderTargetCommon::GIgnoreCookedOut == 0 || IsUsedWithGPUEmitter();
+	const bool bValidGpuDataInterface = NiagaraDataInterfaceRenderTargetCommon::GIgnoreCookedOut == 0 || IsUsedWithGPUScript();
 
 	if (::IsValid(InstanceData->TargetTexture) == false)
 	{
diff --git a/Engine/Plugins/FX/Niagara/Source/Niagara/Private/NiagaraDataInterfaceRenderTargetCube.cpp b/Engine/Plugins/FX/Niagara/Source/Niagara/Private/NiagaraDataInterfaceRenderTargetCube.cpp
--- a/Engine/Plugins/FX/Niagara/Source/Niagara/Private/NiagaraDataInterfaceRenderTargetCube.cpp
+++ b/Engine/Plugins/FX/Niagara/Source/Niagara/Private/NiagaraDataInterfaceRenderTargetCube.cpp
@@ -416,7 +416,7 @@ bool UNiagaraDataInterfaceRenderTargetCube::InitPerInstanceData(void* PerInstanc
 	FRenderTargetCubeRWInstanceData_GameThread* InstanceData = new (PerInstanceData) FRenderTargetCubeRWInstanceData_GameThread();
 
 	//-TEMP: Until we prune data interface on cook this will avoid consuming memory
-	const bool bValidGpuDataInterface = NiagaraDataInterfaceRenderTargetCommon::GIgnoreCookedOut == 0 || IsUsedWithGPUEmitter();
+	const bool bValidGpuDataInterface = NiagaraDataInterfaceRenderTargetCommon::GIgnoreCookedOut == 0 || IsUsedWithGPUScript();
 
 	ETextureRenderTargetFormat RenderTargetFormat;
 	if (NiagaraDataInterfaceRenderTargetCommon::GetRenderTargetFormat(bOverrideFormat, OverrideRenderTargetFormat, RenderTargetFormat) == false)
@@ -611,7 +611,7 @@ bool UNiagaraDataInterfaceRenderTargetCube::PerInstanceTickPostSimulate(void* Pe
 #endif
 
 	//-TEMP: Until we prune data interface on cook this will avoid consuming memory
-	const bool bValidGpuDataInterface = NiagaraDataInterfaceRenderTargetCommon::GIgnoreCookedOut == 0 || IsUsedWithGPUEmitter();
+	const bool bValidGpuDataInterface = NiagaraDataInterfaceRenderTargetCommon::GIgnoreCookedOut == 0 || IsUsedWithGPUScript();
 
 	if (::IsValid(InstanceData->TargetTexture) == false)
 	{
diff --git a/Engine/Plugins/FX/Niagara/Source/Niagara/Private/NiagaraDataInterfaceRenderTargetVolume.cpp b/Engine/Plugins/FX/Niagara/Source/Niagara/Private/NiagaraDataInterfaceRenderTargetVolume.cpp
--- a/Engine/Plugins/FX/Niagara/Source/Niagara/Private/NiagaraDataInterfaceRenderTargetVolume.cpp
+++ b/Engine/Plugins/FX/Niagara/Source/Niagara/Private/NiagaraDataInterfaceRenderTargetVolume.cpp
@@ -513,7 +513,7 @@ bool UNiagaraDataInterfaceRenderTargetVolume::InitPerInstanceData(void* PerInsta
 	FRenderTargetVolumeRWInstanceData_GameThread* InstanceData = new (PerInstanceData) FRenderTargetVolumeRWInstanceData_GameThread();
 
 	//-TEMP: Until we prune data interface on cook this will avoid consuming memory
-	const bool bValidGpuDataInterface = NiagaraDataInterfaceRenderTargetCommon::GIgnoreCookedOut == 0 || IsUsedWithGPUEmitter();
+	const bool bValidGpuDataInterface = NiagaraDataInterfaceRenderTargetCommon::GIgnoreCookedOut == 0 || IsUsedWithGPUScript();
 
 	ETextureRenderTargetFormat RenderTargetFormat;
 	if (NiagaraDataInterfaceRenderTargetCommon::GetRenderTargetFormat(bOverrideFormat, OverrideRenderTargetFormat, RenderTargetFormat) == false)
@@ -1104,7 +1104,7 @@ bool UNiagaraDataInterfaceRenderTargetVolume::PerInstanceTickPostSimulate(void*
 #endif
 
 	//-TEMP: Until we prune data interface on cook this will avoid consuming memory
-	const bool bValidGpuDataInterface = NiagaraDataInterfaceRenderTargetCommon::GIgnoreCookedOut == 0 || IsUsedWithGPUEmitter();
+	const bool bValidGpuDataInterface = NiagaraDataInterfaceRenderTargetCommon::GIgnoreCookedOut == 0 || IsUsedWithGPUScript();
 
 	if (::IsValid(InstanceData->TargetTexture) == false)
 	{
diff --git a/Engine/Plugins/FX/Niagara/Source/Niagara/Private/NiagaraDataInterfaceSkeletalMesh.cpp b/Engine/Plugins/FX/Niagara/Source/Niagara/Private/NiagaraDataInterfaceSkeletalMesh.cpp
--- a/Engine/Plugins/FX/Niagara/Source/Niagara/Private/NiagaraDataInterfaceSkeletalMesh.cpp
+++ b/Engine/Plugins/FX/Niagara/Source/Niagara/Private/NiagaraDataInterfaceSkeletalMesh.cpp
@@ -1748,7 +1748,7 @@ bool FNDISkeletalMesh_InstanceData::Init(UNiagaraDataInterfaceSkeletalMesh* Inte
 
 		const bool MeshValid = SkeletalMesh.IsValid();
 		const bool SupportUvMappingCpu = UsedByCpuUvMapping && MeshValid;
-		const bool SupportUvMappingGpu = UsedByGpuUvMapping && MeshValid && Interface->IsUsedWithGPUEmitter();
+		const bool SupportUvMappingGpu = UsedByGpuUvMapping && MeshValid && Interface->IsUsedWithGPUScript();
 
 		FMeshUvMappingUsage UvMappingUsage(SupportUvMappingCpu, SupportUvMappingGpu);
 
@@ -1775,7 +1775,7 @@ bool FNDISkeletalMesh_InstanceData::Init(UNiagaraDataInterfaceSkeletalMesh* Inte
 
 		const bool MeshValid = SkeletalMesh.IsValid();
 		const bool SupportConnectivityCpu = UsedByCpuConnectivity && MeshValid;
-		const bool SupportConnectivityGpu = UsedByGpuConnectivity && MeshValid && Interface->IsUsedWithGPUEmitter();
+		const bool SupportConnectivityGpu = UsedByGpuConnectivity && MeshValid && Interface->IsUsedWithGPUScript();
 
 		FSkeletalMeshConnectivityUsage ConnectivityUsage(SupportConnectivityCpu, SupportConnectivityGpu);
 
@@ -1972,7 +1972,7 @@ bool FNDISkeletalMesh_InstanceData::Init(UNiagaraDataInterfaceSkeletalMesh* Inte
 			}
 		}
 
-		if (Interface->IsUsedWithGPUEmitter())
+		if (Interface->IsUsedWithGPUScript())
 		{
 			GPUSkinBoneInfluenceType BoneInfluenceType = SkinWeightBuffer->GetBoneInfluenceType();
 			bUnlimitedBoneInfluences = (BoneInfluenceType == GPUSkinBoneInfluenceType::UnlimitedBoneInfluence);
diff --git a/Engine/Plugins/FX/Niagara/Source/Niagara/Private/NiagaraDataInterfaceSpline.cpp b/Engine/Plugins/FX/Niagara/Source/Niagara/Private/NiagaraDataInterfaceSpline.cpp
--- a/Engine/Plugins/FX/Niagara/Source/Niagara/Private/NiagaraDataInterfaceSpline.cpp
+++ b/Engine/Plugins/FX/Niagara/Source/Niagara/Private/NiagaraDataInterfaceSpline.cpp
@@ -641,7 +641,7 @@ bool UNiagaraDataInterfaceSpline::PerInstanceTick(void* PerInstanceData, FNiagar
 		InstData->DefaultUpVector = SplineComponent->DefaultUpVector;
 		InstData->LwcConverter = SystemInstance->GetLWCConverter();
 
-		bool bShouldBuildLUT = (bUseLUT || IsUsedWithGPUEmitter()) && InstData->SplineLUT.MaxIndex < 0;
+		bool bShouldBuildLUT = (bUseLUT || IsUsedWithGPUScript()) && InstData->SplineLUT.MaxIndex < 0;
 		
 		if (InstData->SplineCurvesVersion != SplineComponent->SplineCurves.Version)
 		{
@@ -650,17 +650,17 @@ bool UNiagaraDataInterfaceSpline::PerInstanceTick(void* PerInstanceData, FNiagar
 			InstData->bSyncedGPUCopy = false;
 			InstData->SplineLUT.Reset();
 
-			bShouldBuildLUT = bUseLUT || IsUsedWithGPUEmitter();
+			bShouldBuildLUT = bUseLUT || IsUsedWithGPUScript();
 		}
 		
-		bool bShouldSyncToGPU = IsUsedWithGPUEmitter() && !InstData->bSyncedGPUCopy && InstData->SplineLUT.MaxIndex != INDEX_NONE;
+		bool bShouldSyncToGPU = IsUsedWithGPUScript() && !InstData->bSyncedGPUCopy && InstData->SplineLUT.MaxIndex != INDEX_NONE;
 		
 		// We must build the LUT if this is for GPU regardless of settings
 		if (bShouldBuildLUT)
 		{
 			InstData->SplineLUT.BuildLUT(InstData->SplineCurves, bUseLUT? NumLUTSteps : 256/*Default the LUT to a reasonable value if it's not specifically enabled*/);
 
-			bShouldSyncToGPU = IsUsedWithGPUEmitter();				
+			bShouldSyncToGPU = IsUsedWithGPUScript();
 		}
 
 
diff --git a/Engine/Plugins/FX/Niagara/Source/Niagara/Private/NiagaraSystemInstance.cpp b/Engine/Plugins/FX/Niagara/Source/Niagara/Private/NiagaraSystemInstance.cpp
--- a/Engine/Plugins/FX/Niagara/Source/Niagara/Private/NiagaraSystemInstance.cpp
+++ b/Engine/Plugins/FX/Niagara/Source/Niagara/Private/NiagaraSystemInstance.cpp
@@ -1435,7 +1435,7 @@ void FNiagaraSystemInstance::InitDataInterfaces()
 	//Now the interfaces in the simulations are all correct, we can build the per instance data table.
 	int32 InstanceDataSize = 0;
 	DataInterfaceInstanceDataOffsets.Empty();
-	auto CalcInstDataSize = [&](const FNiagaraParameterStore& ParamStore, bool bIsCPUSimulation, bool bIsGPUSimulation, bool bSearchInstanceParams)
+	auto CalcInstDataSize = [&](const FNiagaraParameterStore& ParamStore, bool bIsCPUScript, bool bIsGPUScript, bool bSearchInstanceParams)
 	{
 		const TArrayView<const FNiagaraVariableWithOffset> Params = ParamStore.ReadParameterVariables();
 		const TArray<UNiagaraDataInterface*>& Interfaces = ParamStore.GetDataInterfaces();
@@ -1478,18 +1478,18 @@ void FNiagaraSystemInstance::InitDataInterfaces()
 						bDataInterfacesHaveTickPrereqs = Interface->HasTickGroupPrereqs();
 					}
 
-					if (bIsGPUSimulation)
+					if (bIsGPUScript)
 					{
-						Interface->SetUsedByGPUEmitter(true);
+						Interface->SetUsedWithGPUScript(true);
 						if(FNiagaraDataInterfaceProxy* Proxy = Interface->GetProxy())
 						{
 							// We need to store the name of each DI source variable here so that we can look it up later when looking for the iteration interface.
 							Proxy->SourceDIName = Var.GetName();
 						}
 					}
-					else if (bIsCPUSimulation)
+					else if (bIsCPUScript)
 					{
-						Interface->SetUsedByCPUEmitter(true);
+						Interface->SetUsedWithCPUScript(true);
 					}
 				}
 			}
@@ -1497,8 +1497,8 @@ void FNiagaraSystemInstance::InitDataInterfaces()
 	};
 
 	CalcInstDataSize(InstanceParameters, false, false, false);//This probably should be a proper exec context.
-		CalcInstDataSize(SystemSimulation->GetSpawnExecutionContext()->Parameters, true, false, true);
-		CalcInstDataSize(SystemSimulation->GetUpdateExecutionContext()->Parameters, true, false, true);
+	CalcInstDataSize(SystemSimulation->GetSpawnExecutionContext()->Parameters, true, false, true);
+	CalcInstDataSize(SystemSimulation->GetUpdateExecutionContext()->Parameters, true, false, true);
 
 	//Iterate over interfaces to get size for table and clear their interface bindings.
 	for (const TSharedRef<FNiagaraEmitterInstance, ESPMode::ThreadSafe>& Simulation : Emitters)
@@ -1554,7 +1554,7 @@ void FNiagaraSystemInstance::InitDataInterfaces()
 				PostTickDataInterfaces.Add(i);
 			}
 
-			if (Interface->IsUsedWithGPUEmitter())
+			if (Interface->IsUsedWithGPUScript())
 			{
 				const int32 GPUDataSize = Interface->PerInstanceDataPassedToRenderThreadSize();
 				if (GPUDataSize > 0)
