- Update for deprecated API call

#rnx

diff --git a/Engine/Plugins/FX/Niagara/Source/Niagara/Private/DataInterface/NiagaraDataInterfaceConsoleVariable.cpp b/Engine/Plugins/FX/Niagara/Source/Niagara/Private/DataInterface/NiagaraDataInterfaceConsoleVariable.cpp
--- a/Engine/Plugins/FX/Niagara/Source/Niagara/Private/DataInterface/NiagaraDataInterfaceConsoleVariable.cpp
+++ b/Engine/Plugins/FX/Niagara/Source/Niagara/Private/DataInterface/NiagaraDataInterfaceConsoleVariable.cpp
@@ -172,7 +172,7 @@ bool UNiagaraDataInterfaceConsoleVariable::InitPerInstanceData(void* PerInstance
 
 	FInstanceData_GameThread* InstanceData_GT = new(PerInstanceData) FInstanceData_GameThread();
 
-	if ( IsUsedByGPUEmitter() )
+	if ( IsUsedWithGPUScript() )
 	{
 		// We shouldn't need to do this per init, we should be able to cache once and once only
 		for (const TSharedRef<FNiagaraEmitterInstance, ESPMode::ThreadSafe>& EmitterInstance : SystemInstance->GetEmitters())
@@ -214,7 +214,7 @@ void UNiagaraDataInterfaceConsoleVariable::DestroyPerInstanceData(void* PerInsta
 	FInstanceData_GameThread* InstanceData_GT = reinterpret_cast<FInstanceData_GameThread*>(PerInstanceData);
 	InstanceData_GT->~FInstanceData_GameThread();
 
-	if ( IsUsedByGPUEmitter() )
+	if ( IsUsedWithGPUScript() )
 	{
 		ENQUEUE_RENDER_COMMAND(NDIConsoleVariable_InitRT)
 		(
