- Don't force last tick if the camera DI is only used with GPU scripts

diff --git a/Engine/Plugins/FX/Niagara/Source/Niagara/Private/NiagaraDataInterfaceCamera.cpp b/Engine/Plugins/FX/Niagara/Source/Niagara/Private/NiagaraDataInterfaceCamera.cpp
--- a/Engine/Plugins/FX/Niagara/Source/Niagara/Private/NiagaraDataInterfaceCamera.cpp
+++ b/Engine/Plugins/FX/Niagara/Source/Niagara/Private/NiagaraDataInterfaceCamera.cpp
@@ -549,9 +549,7 @@ void UNiagaraDataInterfaceCamera::CalculateParticleDistances(FVectorVMExternalFu
 
 ETickingGroup UNiagaraDataInterfaceCamera::CalculateTickGroup(const void* PerInstanceData) const
 {
-	if (!bRequireCurrentFrameData)
-		return NiagaraFirstTickGroup;
-	return ETickingGroup::TG_PostUpdateWork;
+	return IsUsedWithCPUScript() && bRequireCurrentFrameData ? ETickingGroup::TG_PostUpdateWork : NiagaraFirstTickGroup;
 }
 
 #if WITH_EDITOR	
