- Ensure Niagara Systems can not be garbage collected until any render commands have been flushed
- This is a speculative fix for a low repro issue

diff --git a/Engine/Plugins/FX/Niagara/Source/Niagara/Classes/NiagaraSystem.h b/Engine/Plugins/FX/Niagara/Source/Niagara/Classes/NiagaraSystem.h
--- a/Engine/Plugins/FX/Niagara/Source/Niagara/Classes/NiagaraSystem.h
+++ b/Engine/Plugins/FX/Niagara/Source/Niagara/Classes/NiagaraSystem.h
@@ -233,6 +233,7 @@ public:
 	static NIAGARA_API void DeclareConstructClasses(TArray<FTopLevelAssetPath>& OutConstructClasses, const UClass* SpecificSubclass);
 #endif
 	NIAGARA_API virtual void BeginDestroy() override;
+	NIAGARA_API virtual bool IsReadyForFinishDestroy() override;
 	PRAGMA_DISABLE_DEPRECATION_WARNINGS // Suppress compiler warning on override of deprecated function
 	UE_DEPRECATED(5.0, "Use version that takes FObjectPreSaveContext instead.")
 	NIAGARA_API virtual void PreSave(const class ITargetPlatform* TargetPlatform) override;
@@ -1014,6 +1015,8 @@ protected:
 
 	struct FStaticBuffersDeletor { void operator()(FNiagaraSystemStaticBuffers* Ptr) const; };
 	TUniquePtr<FNiagaraSystemStaticBuffers, FStaticBuffersDeletor> StaticBuffers;
+
+	FRenderCommandFence WaitRenderCommandsFence;
 };
 
 FORCEINLINE void UNiagaraSystem::RegisterActiveInstance()
diff --git a/Engine/Plugins/FX/Niagara/Source/Niagara/Private/NiagaraSystem.cpp b/Engine/Plugins/FX/Niagara/Source/Niagara/Private/NiagaraSystem.cpp
--- a/Engine/Plugins/FX/Niagara/Source/Niagara/Private/NiagaraSystem.cpp
+++ b/Engine/Plugins/FX/Niagara/Source/Niagara/Private/NiagaraSystem.cpp
@@ -236,6 +236,15 @@ void UNiagaraSystem::BeginDestroy()
 #endif
 
 	FNiagaraWorldManager::DestroyAllSystemSimulations(this);
+
+	// Ensure we wait for any potential render commands to finish executing before we GC the NiagaraSystem
+	WaitRenderCommandsFence.BeginFence();
+}
+
+bool UNiagaraSystem::IsReadyForFinishDestroy()
+{
+	const bool bReady = Super::IsReadyForFinishDestroy();
+	return bReady && WaitRenderCommandsFence.IsFenceComplete();
 }
 
 void UNiagaraSystem::PreSave(const class ITargetPlatform* TargetPlatform)
