Rework of NDC impact FX loading to avoid any sync loads and just fail instead.

Some misc tidying of code

diff --git a/Engine/Plugins/FX/Niagara/Source/Niagara/Private/NiagaraDataChannel_Islands.cpp b/Engine/Plugins/FX/Niagara/Source/Niagara/Private/NiagaraDataChannel_Islands.cpp
--- a/Engine/Plugins/FX/Niagara/Source/Niagara/Private/NiagaraDataChannel_Islands.cpp
+++ b/Engine/Plugins/FX/Niagara/Source/Niagara/Private/NiagaraDataChannel_Islands.cpp
@@ -12,6 +12,64 @@
 
 #include "DrawDebugHelpers.h"
 
+//////////////////////////////////////////////////////////////////////////
+
+void UNiagaraDataChannel_Islands::PostLoad()
+{
+	Super::PostLoad();
+	
+	//For now immediately load handler systems but in future we should look at deferring the load until a channel is used.
+	//Sub classes may want to load differently so we'll keep this a soft ptr.
+	AsyncLoadSystems();
+}
+
+TConstArrayView<TObjectPtr<UNiagaraSystem>> UNiagaraDataChannel_Islands::GetSystems()const
+{
+	check(IsInGameThread());
+	return SystemsInternal;
+}
+
+void UNiagaraDataChannel_Islands::AsyncLoadSystems()const
+{
+	if(SystemsInternal.Num() != Systems.Num() && Systems.Num() > 0)
+	{
+		TArray<FSoftObjectPath> Requests;
+		for (const TSoftObjectPtr<UNiagaraSystem>& SoftSys : Systems)
+		{
+			if (SoftSys.IsPending())
+			{
+				Requests.Add(SoftSys.ToSoftObjectPath());
+			}
+		}
+
+		if (Requests.Num() > 0)
+		{
+			UAssetManager::GetStreamableManager().RequestAsyncLoad(Requests
+				, FStreamableDelegate::CreateUObject(this, &UNiagaraDataChannel_Islands::PostLoadSystems)
+				, FStreamableManager::AsyncLoadHighPriority
+				, false
+				, false
+				, TEXT("NiagaraDataChannelIsland_HandlerSystems"));
+		}
+		else
+		{
+			PostLoadSystems();
+		}
+	}
+}
+
+void UNiagaraDataChannel_Islands::PostLoadSystems()const
+{
+	check(IsInGameThread());
+	SystemsInternal.Reset(Systems.Num());
+	for(auto& SoftSys : Systems)
+	{
+		SystemsInternal.Add(SoftSys.Get());
+	}
+}
+
+//////////////////////////////////////////////////////////////////////////
+
 UNiagaraDataChannelHandler* UNiagaraDataChannel_Islands::CreateHandler(UWorld* OwningWorld)const
 {
 	UNiagaraDataChannelHandler* NewHandler = NewObject<UNiagaraDataChannelHandler_Islands>(OwningWorld);
@@ -37,18 +95,9 @@ void UNiagaraDataChannelHandler_Islands::Init(const UNiagaraDataChannel* InChann
 	Super::Init(InChannel);
 	if (const UNiagaraDataChannel_Islands* IslandChannel = CastChecked<UNiagaraDataChannel_Islands>(InChannel))
 	{
-		//For now immediately load handler systems but in future we should look at deferring the load until a channel is used.
-		for (const TSoftObjectPtr<UNiagaraSystem>& SoftSys : IslandChannel->Systems)
-		{
-			UAssetManager::GetStreamableManager().RequestAsyncLoad(SoftSys.ToSoftObjectPath()
-				, FStreamableDelegate()
-				, FStreamableManager::AsyncLoadHighPriority
-				, false
-				, false
-				, TEXT("NiagaraDataChannelIsland_HandlerSystems"));
-		}
+		IslandChannel->AsyncLoadSystems();
 
-		int32 InitialPoolCount = IslandChannel->IslandPoolSize;
+		int32 InitialPoolCount = IslandChannel->GetIslandPoolSize();
 		ActiveIslands.Reserve(InitialPoolCount);
 		FreeIslands.Reserve(InitialPoolCount);
 		IslandPool.SetNum(InitialPoolCount);
@@ -92,7 +141,7 @@ void UNiagaraDataChannelHandler_Islands::Tick(float DeltaTime, ETickingGroup Tic
 		FNDCIsland& Island = IslandPool[*It];
 		Island.Tick();
 
-		if(IslandChannel->DebugDrawSettings.ShowBounds())
+		if(IslandChannel->GetDebugDrawSettings().ShowBounds())
 		{
 			Island.DebugDrawBounds();
 		}
@@ -173,8 +222,8 @@ FNDCIsland* UNiagaraDataChannelHandler_Islands::FindOrCreateIsland(const FNiagar
 		//Assuming that overall active island count will be low.
 		//If this is not the case then we'll want to add an acceleration structure to speed up this search.
 		FNDCIsland* IslandToUse = nullptr;
-		FVector MaxExtents = IslandChannel->MaxExtents;
-		FVector PerElementExtents = IslandChannel->PerElementExtents;
+		FVector MaxExtents = IslandChannel->GetMaxExtents();
+		FVector PerElementExtents = IslandChannel->GetPerElementExtents();
 		for (int32 i : ActiveIslands)
 		{
 			FNDCIsland& Island = IslandPool[i];
@@ -263,38 +312,38 @@ void FNDCIsland::Init(UNiagaraDataChannelHandler_Islands* InOwner)
 	const UNiagaraDataChannel_Islands* Channel = CastChecked<UNiagaraDataChannel_Islands>(Owner->GetDataChannel());
 
 	Bounds.Origin = FVector::ZeroVector;
-	Bounds.BoxExtent = Channel->InitialExtents;
+	Bounds.BoxExtent = Channel->GetInitialExtents();
 }
 
 void FNDCIsland::OnAcquired(FVector Location)
 {
 	const UNiagaraDataChannel_Islands* Channel = CastChecked<UNiagaraDataChannel_Islands>(Owner->GetDataChannel());
 
-	if(Channel->Mode == ENiagraDataChannel_IslandMode::AlignedStatic)
+	if(Channel->GetMode() == ENiagraDataChannel_IslandMode::AlignedStatic)
 	{
 		//Align our location down to the nearest lower multiple of the max island size to avoid overlapping islands.
 		//TODO: make our padding extents come from the owning component or the handler systems?
-		FVector ActualMaxSize = Channel->MaxExtents + Channel->PerElementExtents;
-		FVector BoxSize = Channel->MaxExtents * 2;
+		FVector ActualMaxSize = Channel->GetMaxExtents() + Channel->GetPerElementExtents();
+		FVector BoxSize = Channel->GetMaxExtents() * 2;
 		FVector BoxN = Location / BoxSize;
 		BoxN.X = FMath::RoundToNegativeInfinity(BoxN.X);
 		BoxN.Y = FMath::RoundToNegativeInfinity(BoxN.Y);
 		BoxN.Z = FMath::RoundToNegativeInfinity(BoxN.Z);
-		Bounds.Origin = BoxN * BoxSize + Channel->MaxExtents;
+		Bounds.Origin = BoxN * BoxSize + Channel->GetMaxExtents();
 		Bounds.BoxExtent = ActualMaxSize;
 	}
 	else
 	{
 		Bounds.Origin = Location;
-		Bounds.BoxExtent = Channel->InitialExtents;
+		Bounds.BoxExtent = Channel->GetInitialExtents();
 	}
 
 	Data->SetLwcTile(FLargeWorldRenderScalar::GetTileFor(Bounds.Origin));
 	//Spawn our handler systems.
 	FBox HandlerSystemBounds(-Bounds.BoxExtent, Bounds.BoxExtent);
-	for (auto& SoftSys : Channel->Systems)
+	for (const TObjectPtr<UNiagaraSystem>& Sys : Channel->GetSystems())
 	{
-		if (UNiagaraSystem* Sys = SoftSys.LoadSynchronous())
+		if(Sys)
 		{
 			FFXSystemSpawnParameters SpawnParams;
 			SpawnParams.bAutoActivate = false;//We must activate AFTER adding this component to our handle system array.
@@ -312,10 +361,10 @@ void FNDCIsland::OnAcquired(FVector Location)
 			}
 		}
 	}
-	
-	if (Channel->Mode == ENiagraDataChannel_IslandMode::Dynamic)
+
+	if (Channel->GetMode() == ENiagraDataChannel_IslandMode::Dynamic)
 	{
-		ensure(TryGrow(Location, Channel->PerElementExtents, Channel->MaxExtents));
+		ensure(TryGrow(Location, Channel->GetPerElementExtents(), Channel->GetMaxExtents()));
 	}
 }
 
diff --git a/Engine/Plugins/FX/Niagara/Source/Niagara/Public/NiagaraDataChannel_Islands.h b/Engine/Plugins/FX/Niagara/Source/Niagara/Public/NiagaraDataChannel_Islands.h
--- a/Engine/Plugins/FX/Niagara/Source/Niagara/Public/NiagaraDataChannel_Islands.h
+++ b/Engine/Plugins/FX/Niagara/Source/Niagara/Public/NiagaraDataChannel_Islands.h
@@ -111,10 +111,27 @@ class UNiagaraDataChannel_Islands : public UNiagaraDataChannel
 {
 	GENERATED_BODY()
 
-public: 
+public:
+
+	//UObject Interface
+	virtual void PostLoad()override;
+	//UObject Interface End
 
 	NIAGARA_API virtual UNiagaraDataChannelHandler* CreateHandler(UWorld* OwningWorld)const override;
 
+	void AsyncLoadSystems()const;
+
+	ENiagraDataChannel_IslandMode GetMode()const { return Mode; }
+	FVector GetInitialExtents()const { return InitialExtents; }
+	FVector GetMaxExtents()const { return MaxExtents; }
+	FVector GetPerElementExtents()const { return PerElementExtents; }
+
+	TConstArrayView<TObjectPtr<UNiagaraSystem>> GetSystems()const;
+	int32 GetIslandPoolSize()const { return IslandPoolSize; }
+	const FNDCIslandDebugDrawSettings& GetDebugDrawSettings()const { return DebugDrawSettings; }
+
+private:
+
 	/** Controls how islands are placed and sized. */
 	UPROPERTY(EditAnywhere, Category = "Islands")
 	ENiagraDataChannel_IslandMode Mode = ENiagraDataChannel_IslandMode::AlignedStatic;
@@ -158,6 +175,11 @@ public:
 
 	UPROPERTY(EditAnywhere, Category = "Debug Rendering")
 	FNDCIslandDebugDrawSettings DebugDrawSettings;
+
+	void PostLoadSystems()const;
+
+	UPROPERTY(Transient)
+	mutable TArray<TObjectPtr<UNiagaraSystem>> SystemsInternal;
 };
 
 UCLASS(Experimental, BlueprintType, MinimalAPI)
@@ -194,4 +216,4 @@ private:
 
 	/** Initializes and adds a new island to the ActiveIslands list. Either retrieving from the free pool of existing inactive islands or creating a new one. */
 	int32 ActivateNewIsland(FVector Location);
-};
+};
\ No newline at end of file
