Fix linker error due to classes no longer being full exported
#rnx

diff --git a/Engine/Plugins/FX/Niagara/Source/Niagara/Public/NiagaraDataChannel_Islands.h b/Engine/Plugins/FX/Niagara/Source/Niagara/Public/NiagaraDataChannel_Islands.h
--- a/Engine/Plugins/FX/Niagara/Source/Niagara/Public/NiagaraDataChannel_Islands.h
+++ b/Engine/Plugins/FX/Niagara/Source/Niagara/Public/NiagaraDataChannel_Islands.h
@@ -114,7 +114,7 @@ class UNiagaraDataChannel_Islands : public UNiagaraDataChannel
 public:
 
 	//UObject Interface
-	virtual void PostLoad()override;
+	NIAGARA_API virtual void PostLoad()override;
 	//UObject Interface End
 
 	NIAGARA_API virtual UNiagaraDataChannelHandler* CreateHandler(UWorld* OwningWorld)const override;
