- Modify visibility display for debug hud to call IsVisible and that can include things like the owner's level collection visibility
- Show any active SimCache information on HUD also

diff --git a/Engine/Plugins/FX/Niagara/Source/Niagara/Private/NiagaraDebugHud.cpp b/Engine/Plugins/FX/Niagara/Source/Niagara/Private/NiagaraDebugHud.cpp
--- a/Engine/Plugins/FX/Niagara/Source/Niagara/Private/NiagaraDebugHud.cpp
+++ b/Engine/Plugins/FX/Niagara/Source/Niagara/Private/NiagaraDebugHud.cpp
@@ -13,6 +13,7 @@
 #include "NiagaraGpuComputeDispatchInterface.h"
 #include "NiagaraMeshRendererProperties.h"
 #include "NiagaraScript.h"
+#include "NiagaraSimCache.h"
 #include "NiagaraSpriteRendererProperties.h"
 #include "NiagaraSystem.h"
 #include "NiagaraSystemGpuComputeProxy.h"
@@ -2843,18 +2844,19 @@ void FNiagaraDebugHud::DrawComponents(FNiagaraWorldManager* WorldManager, UCanva
 						StringBuilder.Append(TEXT("\n"));
 					}
 
-					if ( NiagaraComponent->bHiddenInGame || !NiagaraComponent->GetVisibleFlag() || (OwnerActor && OwnerActor->IsHidden()) )
+					if ( NiagaraComponent->bHiddenInGame || !NiagaraComponent->IsVisible() || (OwnerActor && OwnerActor->IsHidden()) )
 					{
-						StringBuilder.Appendf(TEXT("HiddenInGame(%d) Visible(%d)"), NiagaraComponent->bHiddenInGame, NiagaraComponent->GetVisibleFlag());
+						StringBuilder.Appendf(TEXT("HiddenInGame(%d) IsVisible(%d) GetVisibleFlag(%d)"), NiagaraComponent->bHiddenInGame, NiagaraComponent->IsVisible(), NiagaraComponent->GetVisibleFlag());
 						if ( OwnerActor )
 						{
 							StringBuilder.Appendf(TEXT(" OwnerActorHidden(%d)"), OwnerActor->IsHidden());
 						}
 						StringBuilder.Append(TEXT("\n"));
 					}
-					if ( !NiagaraComponent->GetVisibleFlag() )
+
+					if (UNiagaraSimCache* SimCache = NiagaraComponent->GetSimCache())
 					{
-						StringBuilder.Append(TEXT("Visibile - false\n"));
+						StringBuilder.Appendf(TEXT("SimCache(%s)\n"), *GetFullNameSafe(SimCache));
 					}
 
 					if (bIsActive)
