- Add ability to append to an existing simulation cache
- Frames that are before the current simulation age / tick count will be culled
- BeginAppend can be called vs BeginWrite and it will handle the initial setup if the cache is invalid
- Update capture library with a new call so that we can specify all the parameters, such are removing timeouts, should we append, etc.

diff --git a/Engine/Plugins/FX/Niagara/Source/Niagara/Classes/NiagaraSimCache.h b/Engine/Plugins/FX/Niagara/Source/Niagara/Classes/NiagaraSimCache.h
--- a/Engine/Plugins/FX/Niagara/Source/Niagara/Classes/NiagaraSimCache.h
+++ b/Engine/Plugins/FX/Niagara/Source/Niagara/Classes/NiagaraSimCache.h
@@ -190,10 +190,8 @@ struct FNiagaraSimCacheFrame
 	UPROPERTY()
 	float SimulationAge = 0.0f;
 
-#if WITH_EDITORONLY_DATA
 	UPROPERTY()
 	int32 SimulationTickCount = 0;
-#endif
 
 	UPROPERTY()
 	FNiagaraSimCacheSystemFrame SystemData;
@@ -400,10 +398,15 @@ public:
 	UFUNCTION(BlueprintCallable, Category=NiagaraSimCache)
 	ENiagaraSimCacheAttributeCaptureMode GetAttributeCaptureMode() const { return CreateParameters.AttributeCaptureMode; }
 
-	NIAGARA_API bool BeginWrite(FNiagaraSimCacheCreateParameters InCreateParameters, UNiagaraComponent * NiagaraComponent);
+	NIAGARA_API bool BeginWrite(FNiagaraSimCacheCreateParameters InCreateParameters, UNiagaraComponent* NiagaraComponent);
 	NIAGARA_API bool BeginWrite(FNiagaraSimCacheCreateParameters InCreateParameters, UNiagaraComponent* NiagaraComponent, FNiagaraSimCacheFeedbackContext& FeedbackContext);
+
+	NIAGARA_API bool BeginAppend(FNiagaraSimCacheCreateParameters InCreateParameters, UNiagaraComponent* NiagaraComponent);
+	NIAGARA_API bool BeginAppend(FNiagaraSimCacheCreateParameters InCreateParameters, UNiagaraComponent* NiagaraComponent, FNiagaraSimCacheFeedbackContext& FeedbackContext);
+
 	NIAGARA_API bool WriteFrame(UNiagaraComponent* NiagaraComponent);
 	NIAGARA_API bool WriteFrame(UNiagaraComponent* NiagaraComponent, FNiagaraSimCacheFeedbackContext& FeedbackContext);
+
 	NIAGARA_API bool EndWrite();
 
 	NIAGARA_API bool CanRead(UNiagaraSystem* NiagaraSystem);
diff --git a/Engine/Plugins/FX/Niagara/Source/Niagara/Classes/NiagaraSimCacheCapture.cpp b/Engine/Plugins/FX/Niagara/Source/Niagara/Classes/NiagaraSimCacheCapture.cpp
--- a/Engine/Plugins/FX/Niagara/Source/Niagara/Classes/NiagaraSimCacheCapture.cpp
+++ b/Engine/Plugins/FX/Niagara/Source/Niagara/Classes/NiagaraSimCacheCapture.cpp
@@ -7,6 +7,14 @@
 #include "NiagaraComponent.h"
 #include "NiagaraSimCache.h"
 
+FNiagaraSimCacheCapture::~FNiagaraSimCacheCapture()
+{
+	if (TickerHandle.IsValid())
+	{
+		FinishCapture();
+	}
+}
+
 void FNiagaraSimCacheCapture::FinishCapture()
 {
 	if (TickerHandle.IsValid())
@@ -15,39 +23,43 @@ void FNiagaraSimCacheCapture::FinishCapture()
 	}
 
 	UNiagaraSimCache* CaptureSimCache = WeakCaptureSimCache.Get();
-
 	if ( CaptureSimCache != nullptr )
 	{
 		CaptureSimCache->EndWrite();
 	}
 	CaptureComplete.Broadcast(CaptureSimCache);
-
 	WeakCaptureComponent = nullptr;
 }
 
-void FNiagaraSimCacheCapture::CaptureNiagaraSimCache(UNiagaraSimCache* SimCache,
-                                                     FNiagaraSimCacheCreateParameters CreateParameters, UNiagaraComponent* NiagaraComponent, FNiagaraSimCacheCaptureParameters InCaptureParameters)
+void FNiagaraSimCacheCapture::CaptureNiagaraSimCache(UNiagaraSimCache* SimCache, FNiagaraSimCacheCreateParameters CreateParameters, UNiagaraComponent* NiagaraComponent, FNiagaraSimCacheCaptureParameters InCaptureParameters)
 {
 	WeakCaptureSimCache = SimCache;
 	WeakCaptureComponent = NiagaraComponent;
 	CaptureParameters = InCaptureParameters;
+	CaptureParameters.Sanitize();
 	CaptureFrameCounter = 0;
-
 	if(SimCache == nullptr || NiagaraComponent == nullptr)
 	{
 		FinishCapture();
 		return;
 	}
 	
-	SimCache->BeginWrite(CreateParameters, NiagaraComponent);
+	if (CaptureParameters.bAppendToSimCache)
+	{
+		SimCache->BeginAppend(CreateParameters, NiagaraComponent);
+	}
+	else
+	{
+		SimCache->BeginWrite(CreateParameters, NiagaraComponent);
+	}
 
 	// In this path we are going to manually advance the simulation until we complete
-	if (CaptureParameters.bManuallyAdvanceSimulation)
+	if (CaptureParameters.bCaptureAllFramesImmediatly)
 	{
 		while (WeakCaptureComponent != nullptr && NiagaraComponent != nullptr)
 		{
-			NiagaraComponent->AdvanceSimulation(1, CaptureParameters.AdvanceDeltaTime);
-			OnFrameTick(CaptureParameters.AdvanceDeltaTime);
+			NiagaraComponent->AdvanceSimulation(1, CaptureParameters.ImmediateCaptureDeltaTime);
+			OnFrameTick(CaptureParameters.ImmediateCaptureDeltaTime);
 		}
 	}
 	// In this path we are going to monitor the simulation and capture each frame
@@ -109,7 +121,7 @@ bool FNiagaraSimCacheCapture::OnFrameTick(float DeltaTime)
 			}
 
 			// Make sure we don't keep this alive forever, if we didn't managed to capture anything in 10 ticks something has probably gone wrong so bail
-			if (TimeOutCounter++ > 10)
+			if (CaptureParameters.bUseTimeout && (TimeOutCounter++ > CaptureParameters.TimeoutFrameCount))
 			{
 				UE_LOG(LogNiagara, Warning, TEXT("SimCache Write has failed too many times, abandoning capturing for (%s)"), *GetFullNameSafe(CaptureSimCache));
 				FinishCapture();
@@ -123,7 +135,7 @@ bool FNiagaraSimCacheCapture::OnFrameTick(float DeltaTime)
 
 	// Have we recorded all the frames we need?
 	// Note: the -1 is because T0 was the initial frame
-	if ( (CaptureParameters.NumFrames > 0) && (CaptureFrameCounter > (CaptureParameters.CaptureRate * (CaptureParameters.NumFrames - 1))) )
+	if ( CaptureParameters.bCaptureFixedNumberOfFrames && (CaptureFrameCounter > (CaptureParameters.CaptureRate * (CaptureParameters.NumFrames - 1))) )
 	{
 		FinishCapture();
 		return true;
diff --git a/Engine/Plugins/FX/Niagara/Source/Niagara/Classes/NiagaraSimCacheCapture.h b/Engine/Plugins/FX/Niagara/Source/Niagara/Classes/NiagaraSimCacheCapture.h
--- a/Engine/Plugins/FX/Niagara/Source/Niagara/Classes/NiagaraSimCacheCapture.h
+++ b/Engine/Plugins/FX/Niagara/Source/Niagara/Classes/NiagaraSimCacheCapture.h
@@ -15,30 +15,60 @@ struct FNiagaraSimCacheCaptureParameters
 {
 	GENERATED_BODY()
 
+	FNiagaraSimCacheCaptureParameters()
+		: bAppendToSimCache(false)
+		, bCaptureFixedNumberOfFrames(true)
+		, bUseTimeout(true)
+		, bCaptureAllFramesImmediatly(false)
+	{
+
+	}
+
+	/** When enabled we will append to the existing simulation cache rather than destroying the existing contents. */
+	UPROPERTY(EditAnywhere, BlueprintReadWrite, Category = "SimCache")
+	uint32 bAppendToSimCache : 1;
+
+	/** When enabled we capture NumFrames number of frames, otherwise the capture will continue until cancelled or the simulation is complete. */
+	UPROPERTY(EditAnywhere, BlueprintReadWrite, Category = "SimCache")//, meta = (InlineEditConditionToggle))
+	uint32 bCaptureFixedNumberOfFrames : 1;
+
 	/**
 	 The max number of frames to capture. The capture will stop when the simulation completes or we have rendered this many frames, whichever happens first.
 	 Set to 0 to capture until simulation completes.
 	 **/
-	UPROPERTY(EditAnywhere, BlueprintReadWrite, Category = "SimCache")
-	int32 NumFrames = 0;
+	UPROPERTY(EditAnywhere, BlueprintReadWrite, Category = "SimCache", meta = (EditCondition = "bCaptureFixedNumberOfFrames", ClampMin = "1"))
+	int32 NumFrames = 16;
 
-	/**
-	 Allows for reducing the frequency of captured frames in relation to the simulation's frames. The ratio is 1 / CaptureRate, so a CaptureRate of 2 would captures frames 0, 2, 4, etc.
-	 */	
-	UPROPERTY(EditAnywhere, BlueprintReadWrite, Category = "SimCache")
+	/** Allows for reducing the frequency of captured frames in relation to the simulation's frames. The ratio is 1 / CaptureRate, so a CaptureRate of 2 would captures frames 0, 2, 4, etc. */	
+	UPROPERTY(EditAnywhere, BlueprintReadWrite, Category = "SimCache", meta = (ClampMin = "1"))
 	int32 CaptureRate = 1;
 
-	/**
-	 Controls manually advancing the simulation in a loop vs reading from it every frame. If bAdvanceSimulation is true, this capture will manually advance the simulation
-	 */
-	UPROPERTY(EditAnywhere, BlueprintReadWrite, Category = "SimCache")
-	bool bManuallyAdvanceSimulation = false;
+	/** When enabled the capture will time out if we reach the defined number of frames without anything new to capture. */
+	UPROPERTY(EditAnywhere, BlueprintReadWrite, Category = "SimCache")//, meta = (InlineEditConditionToggle))
+	uint32 bUseTimeout : 1;
 
-	/*
-	 The delta time in seconds to use when manually advancing the simulation. Defaults to 1/60th of a second (0.01666).
-	 */
-	UPROPERTY(EditAnywhere, BlueprintReadWrite, Category = "SimCache", meta = (EditCondition = "bManuallyAdvanceSimulation"))
-	float AdvanceDeltaTime = 0.01666f;
+	/** When we fail to capture a new frame after this many frames the capture will complete automatically. */
+	UPROPERTY(EditAnywhere, BlueprintReadWrite, Category = "SimCache", meta = (EditCondition = "bUseTimeout", ClampMin = "1"))
+	int32 TimeoutFrameCount = 10;
+
+	/**
+	When enabled we will capture all the requested frames immediatly.
+	This will capture the simulation outside of the main work tick, i.e. if you request a 10 frame capture we will loop capturing 10 frames on the capture call rather than over 10 world ticks.
+	*/
+	UPROPERTY(EditAnywhere, BlueprintReadWrite, Category = "SimCache")//, meta = (InlineEditConditionToggle))
+	uint32 bCaptureAllFramesImmediatly : 1;
+
+	/** The delta time in seconds to use when manually advancing the simulation.Defaults to 1 / 60th of a second(0.01666). */
+	UPROPERTY(EditAnywhere, BlueprintReadWrite, Category = "SimCache", meta = (EditCondition = "bCaptureAllFramesImmediatly", ClampMin = "0.0001"))
+	float ImmediateCaptureDeltaTime = 0.01666f;
+
+	void Sanitize()
+	{
+		NumFrames = FMath::Max(NumFrames, 1);
+		CaptureRate = FMath::Max(CaptureRate, 1);
+		TimeoutFrameCount = FMath::Max(TimeoutFrameCount, 1);
+		ImmediateCaptureDeltaTime = FMath::Max(ImmediateCaptureDeltaTime, 0.0001f);
+	}
 };
 
 class FNiagaraSimCacheCapture
@@ -46,9 +76,10 @@ class FNiagaraSimCacheCapture
 public:
 	DECLARE_MULTICAST_DELEGATE_OneParam(FOnCaptureComplete, UNiagaraSimCache*);
 
+	NIAGARA_API ~FNiagaraSimCacheCapture();
+
 	FOnCaptureComplete& OnCaptureComplete() {return CaptureComplete;}
 	
-	
 	// Captures a niagara sim cache. The caller is responsible for managing the lifetime of the provided component and sim cache. 
 	NIAGARA_API void CaptureNiagaraSimCache(UNiagaraSimCache* SimCache, FNiagaraSimCacheCreateParameters CreateParameters, UNiagaraComponent* NiagaraComponent, FNiagaraSimCacheCaptureParameters CaptureParameters);
 
@@ -60,8 +91,9 @@ public:
 	*/
 	NIAGARA_API bool CaptureCurrentFrameImmediate(UNiagaraSimCache* SimCache, FNiagaraSimCacheCreateParameters CreateParameters, UNiagaraComponent* NiagaraComponent, UNiagaraSimCache*& OutSimCache, bool bAdvanceSimulation=false, float AdvanceDeltaTime=0.01666f);
 
-private:
+	NIAGARA_API void FinishCapture();
 
+private:
 	TWeakObjectPtr<UNiagaraSimCache> WeakCaptureSimCache;
 
 	TWeakObjectPtr<UNiagaraComponent> WeakCaptureComponent;
@@ -77,6 +109,4 @@ private:
 	int32 TimeOutCounter = 0;
 	
 	NIAGARA_API bool OnFrameTick(float DeltaTime);
-
-	NIAGARA_API void FinishCapture();
 };
diff --git a/Engine/Plugins/FX/Niagara/Source/Niagara/Classes/NiagaraSimCacheFunctionLibrary.h b/Engine/Plugins/FX/Niagara/Source/Niagara/Classes/NiagaraSimCacheFunctionLibrary.h
--- a/Engine/Plugins/FX/Niagara/Source/Niagara/Classes/NiagaraSimCacheFunctionLibrary.h
+++ b/Engine/Plugins/FX/Niagara/Source/Niagara/Classes/NiagaraSimCacheFunctionLibrary.h
@@ -6,12 +6,12 @@
 #include "NiagaraSimCacheCapture.h"
 
 #include "Kismet/BlueprintFunctionLibrary.h"
-#include "Kismet/BlueprintAsyncActionBase.h"
+#include "Engine/CancellableAsyncAction.h"
 
 #include "NiagaraSimCacheFunctionLibrary.generated.h"
 
 UCLASS()
-class UAsyncNiagaraCaptureSimCache : public UBlueprintAsyncActionBase
+class UAsyncNiagaraCaptureSimCache : public UCancellableAsyncAction
 {
 	GENERATED_BODY()
 
@@ -27,13 +27,13 @@ public:
 	UPROPERTY(BlueprintAssignable)
 	FOnCaptureComplete CaptureComplete;
 
+	bool bIsRunning = false;
 	FNiagaraSimCacheCapture SimCacheCapture;
-
 	FNiagaraSimCacheCaptureParameters CaptureParameters;
-
 	FNiagaraSimCacheCreateParameters CreateParameters;
 
 	virtual void Activate() override;
+	virtual void Cancel() override;
 	virtual void SetReadyToDestroy() override;
 
 	/**
@@ -54,8 +54,14 @@ public:
 	UFUNCTION(BlueprintCallable, meta=(BlueprintInternalUseOnly="true", Category=NiagaraSimCache, AdvancedDisplay = "CaptureRate, bAdvanceSimulation, AdvanceDeltaTime"))
 	static UAsyncNiagaraCaptureSimCache* CaptureNiagaraSimCacheUntilComplete(UNiagaraSimCache* SimCache, FNiagaraSimCacheCreateParameters CreateParameters, UNiagaraComponent* NiagaraComponent, UNiagaraSimCache*& OutSimCache, int32 CaptureRate = 1, bool bAdvanceSimulation=false, float AdvanceDeltaTime=0.01666f);
 
+	/**
+	Capture a simulation cache with customizable capture parameters.
+	*/
+	UFUNCTION(BlueprintCallable, meta = (BlueprintInternalUseOnly = "true", Category = NiagaraSimCache))
+	static UAsyncNiagaraCaptureSimCache* CaptureNiagaraSimCache(UNiagaraSimCache* SimCache, FNiagaraSimCacheCreateParameters CreateParameters, UNiagaraComponent* NiagaraComponent, FNiagaraSimCacheCaptureParameters CaptureParameters, UNiagaraSimCache*& OutSimCache);
+
 private:
-	static UAsyncNiagaraCaptureSimCache* CaptureNiagaraSimCacheImpl(UNiagaraSimCache* SimCache, FNiagaraSimCacheCreateParameters CreateParameters, UNiagaraComponent* NiagaraComponent, UNiagaraSimCache*& OutSimCache, FNiagaraSimCacheCaptureParameters CaptureParameters);
+	static UAsyncNiagaraCaptureSimCache* CaptureNiagaraSimCacheImpl(UNiagaraSimCache* SimCache, const FNiagaraSimCacheCreateParameters& CreateParameters, UNiagaraComponent* NiagaraComponent, const FNiagaraSimCacheCaptureParameters& CaptureParameters, UNiagaraSimCache*& OutSimCache);
 
 	void CaptureFinished(UNiagaraSimCache* CapturedSimCache);
 };
diff --git a/Engine/Plugins/FX/Niagara/Source/Niagara/Private/NiagaraSimCache.cpp b/Engine/Plugins/FX/Niagara/Source/Niagara/Private/NiagaraSimCache.cpp
--- a/Engine/Plugins/FX/Niagara/Source/Niagara/Private/NiagaraSimCache.cpp
+++ b/Engine/Plugins/FX/Niagara/Source/Niagara/Private/NiagaraSimCache.cpp
@@ -257,6 +257,76 @@ bool UNiagaraSimCache::BeginWrite(FNiagaraSimCacheCreateParameters InCreateParam
 	return true;
 }
 
+bool UNiagaraSimCache::BeginAppend(FNiagaraSimCacheCreateParameters InCreateParameters, UNiagaraComponent* NiagaraComponent)
+{
+	FNiagaraSimCacheFeedbackContext FeedbackContext;
+	return BeginAppend(InCreateParameters, NiagaraComponent, FeedbackContext);
+}
+
+bool UNiagaraSimCache::BeginAppend(FNiagaraSimCacheCreateParameters InCreateParameters, UNiagaraComponent* NiagaraComponent, FNiagaraSimCacheFeedbackContext& FeedbackContext)
+{
+	check(PendingCommandsInFlight == 0);
+
+	OnCacheBeginWrite.Broadcast(this);
+
+	FNiagaraSimCacheHelper Helper(NiagaraComponent);
+	if (Helper.HasValidSimulation() == false)
+	{
+		//FeedbackContext.Errors.Emplace(FString::Printf(TEXT("No valid simulation data for %s"), *GetPathNameSafe(Helper.NiagaraSystem)));
+		return false;
+	}
+
+	if (!CacheGuid.IsValid() || SoftNiagaraSystem.Get() != Helper.NiagaraSystem)
+	{
+		return BeginWrite(InCreateParameters, NiagaraComponent);
+	}
+
+	// Are we clearing out all frames
+	const float SimulationAge = Helper.SystemInstance->GetAge();
+	const int32 SimulationTickCount = Helper.SystemInstance->GetTickCount();
+	if (SimulationAge <= StartSeconds || SimulationTickCount == 0)
+	{
+		CacheFrames.Empty();
+	}
+	// Cull any frames that have a later start time or tick count
+	else if (SimulationAge <= StartSeconds + DurationSeconds)
+	{
+		int32 CullFrame = CacheFrames.Num();
+		while (CullFrame > 0)
+		{
+			const FNiagaraSimCacheFrame& Frame = CacheFrames[CullFrame - 1];
+			if (Frame.SimulationAge < SimulationAge && Frame.SimulationTickCount < SimulationTickCount)
+			{
+				break;
+			}
+			--CullFrame;
+		}
+
+		if (CullFrame != CacheFrames.Num())
+		{
+			CacheFrames.SetNum(CullFrame);
+			if (CacheFrames.Num() > 0)
+			{
+				DurationSeconds = CacheFrames.Last().SimulationAge - StartSeconds;
+				CaptureTickCount = CacheFrames.Last().SimulationTickCount;
+			}
+		}
+	}
+
+	if (CacheFrames.Num() == 0)
+	{
+		StartSeconds = 0.0f;
+		DurationSeconds = 0.0f;
+		CaptureTickCount = INDEX_NONE;
+	}
+	else
+	{
+		CaptureTickCount = SimulationTickCount - 1;
+	}
+
+	return true;
+}
+
 bool UNiagaraSimCache::WriteFrame(UNiagaraComponent* NiagaraComponent)
 {
 	FNiagaraSimCacheFeedbackContext FeedbackContext;
@@ -430,7 +500,6 @@ bool UNiagaraSimCache::EndWrite()
 		SoftNiagaraSystem.Reset();
 	}
 
-
 	if (DataInterfaceStorage.IsEmpty() == false)
 	{
 		bool bDataInterfacesSucess = true;
diff --git a/Engine/Plugins/FX/Niagara/Source/Niagara/Private/NiagaraSimCacheFunctionLibrary.cpp b/Engine/Plugins/FX/Niagara/Source/Niagara/Private/NiagaraSimCacheFunctionLibrary.cpp
--- a/Engine/Plugins/FX/Niagara/Source/Niagara/Private/NiagaraSimCacheFunctionLibrary.cpp
+++ b/Engine/Plugins/FX/Niagara/Source/Niagara/Private/NiagaraSimCacheFunctionLibrary.cpp
@@ -11,45 +11,65 @@ void UAsyncNiagaraCaptureSimCache::Activate()
 {
 	Super::Activate();
 
+	bIsRunning = true;
 	SimCacheCapture.OnCaptureComplete().AddUObject(this, &UAsyncNiagaraCaptureSimCache::CaptureFinished);
-	
 	SimCacheCapture.CaptureNiagaraSimCache(CaptureSimCache, CreateParameters, CaptureComponent, CaptureParameters);
-	
+}
+
+void UAsyncNiagaraCaptureSimCache::Cancel()
+{
+	if (bIsRunning)
+	{
+		SimCacheCapture.FinishCapture();
+	}
+
+	Super::Cancel();
 }
 
 void UAsyncNiagaraCaptureSimCache::SetReadyToDestroy()
 {
 	Super::SetReadyToDestroy();
 
-	if ( CaptureSimCache != nullptr )
+	if (bIsRunning)
 	{
-		CaptureSimCache->EndWrite();
+		bIsRunning = false;
+		CaptureComplete.Broadcast(::IsValid(CaptureSimCache) ? CaptureSimCache->IsCacheValid() : false);
+		CaptureComponent = nullptr;
 	}
-	CaptureComplete.Broadcast(CaptureSimCache ? CaptureSimCache->IsCacheValid() : false);
+}
+
+UAsyncNiagaraCaptureSimCache* UAsyncNiagaraCaptureSimCache::CaptureNiagaraSimCacheMultiFrame(UNiagaraSimCache* SimCache, FNiagaraSimCacheCreateParameters InCreateParameters, UNiagaraComponent* NiagaraComponent, UNiagaraSimCache*& OutSimCache, int32 NumFrames, int32 CaptureRate, bool bAdvanceSimulation, float AdvanceDeltaTime)
+{
+	FNiagaraSimCacheCaptureParameters InCaptureParameters;
+	InCaptureParameters.NumFrames = FMath::Max(1, NumFrames);
+	InCaptureParameters.CaptureRate = FMath::Max(1, CaptureRate);
+	InCaptureParameters.bCaptureAllFramesImmediatly = bAdvanceSimulation;
+	InCaptureParameters.ImmediateCaptureDeltaTime = AdvanceDeltaTime;
 
-	CaptureComponent = nullptr;
+	return CaptureNiagaraSimCacheImpl(SimCache, InCreateParameters, NiagaraComponent, InCaptureParameters, OutSimCache);
 }
 
-UAsyncNiagaraCaptureSimCache* UAsyncNiagaraCaptureSimCache::CaptureNiagaraSimCacheMultiFrame(UNiagaraSimCache* SimCache, FNiagaraSimCacheCreateParameters CreateParameters, UNiagaraComponent* NiagaraComponent, UNiagaraSimCache*& OutSimCache, int32 NumFrames, int32 CaptureRate, bool bAdvanceSimulation, float AdvanceDeltaTime)
+UAsyncNiagaraCaptureSimCache* UAsyncNiagaraCaptureSimCache::CaptureNiagaraSimCacheUntilComplete(UNiagaraSimCache* SimCache, FNiagaraSimCacheCreateParameters InCreateParameters, UNiagaraComponent* NiagaraComponent, UNiagaraSimCache*& OutSimCache, int32 CaptureRate, bool bAdvanceSimulation, float AdvanceDeltaTime)
 {
-	FNiagaraSimCacheCaptureParameters CaptureParameters;
-	
-	CaptureParameters.NumFrames = FMath::Max(1, NumFrames);
-	CaptureParameters.CaptureRate = FMath::Max(1, CaptureRate);
-	CaptureParameters.bManuallyAdvanceSimulation = bAdvanceSimulation;
-	CaptureParameters.AdvanceDeltaTime = AdvanceDeltaTime;
-
-	return CaptureNiagaraSimCacheImpl(SimCache, CreateParameters, NiagaraComponent, OutSimCache, CaptureParameters);
+	FNiagaraSimCacheCaptureParameters InCaptureParameters;
+	InCaptureParameters.NumFrames = 0;
+	InCaptureParameters.CaptureRate = FMath::Max(1, CaptureRate);
+	InCaptureParameters.bCaptureAllFramesImmediatly = bAdvanceSimulation;
+	InCaptureParameters.ImmediateCaptureDeltaTime = AdvanceDeltaTime;
+
+	return CaptureNiagaraSimCacheImpl(SimCache, InCreateParameters, NiagaraComponent, InCaptureParameters, OutSimCache);
 }
 
-UAsyncNiagaraCaptureSimCache* UAsyncNiagaraCaptureSimCache::CaptureNiagaraSimCacheImpl(UNiagaraSimCache* SimCache,
-	FNiagaraSimCacheCreateParameters CreateParameters, UNiagaraComponent* NiagaraComponent, UNiagaraSimCache*& OutSimCache,
-	FNiagaraSimCacheCaptureParameters InCaptureParameters)
+UAsyncNiagaraCaptureSimCache* UAsyncNiagaraCaptureSimCache::CaptureNiagaraSimCache(UNiagaraSimCache* SimCache, FNiagaraSimCacheCreateParameters InCreateParameters, UNiagaraComponent* NiagaraComponent, FNiagaraSimCacheCaptureParameters InCaptureParameters, UNiagaraSimCache*& OutSimCache)
 {
-	UAsyncNiagaraCaptureSimCache* CaptureAction = NewObject<UAsyncNiagaraCaptureSimCache>();
+	return CaptureNiagaraSimCacheImpl(SimCache, InCreateParameters, NiagaraComponent, InCaptureParameters, OutSimCache);
+}
 
+UAsyncNiagaraCaptureSimCache* UAsyncNiagaraCaptureSimCache::CaptureNiagaraSimCacheImpl(UNiagaraSimCache* SimCache, const FNiagaraSimCacheCreateParameters& InCreateParameters, UNiagaraComponent* NiagaraComponent, const FNiagaraSimCacheCaptureParameters& InCaptureParameters, UNiagaraSimCache*& OutSimCache)
+{
+	UAsyncNiagaraCaptureSimCache* CaptureAction = NewObject<UAsyncNiagaraCaptureSimCache>();
 	CaptureAction->CaptureParameters = InCaptureParameters;
-	CaptureAction->CreateParameters = CreateParameters;
+	CaptureAction->CreateParameters = InCreateParameters;
 	CaptureAction->CaptureSimCache = SimCache;
 	CaptureAction->CaptureComponent = NiagaraComponent;
 	OutSimCache = SimCache;
@@ -57,19 +77,6 @@ UAsyncNiagaraCaptureSimCache* UAsyncNiagaraCaptureSimCache::CaptureNiagaraSimCac
 	return CaptureAction;
 }
 
-UAsyncNiagaraCaptureSimCache* UAsyncNiagaraCaptureSimCache::CaptureNiagaraSimCacheUntilComplete(UNiagaraSimCache* SimCache, FNiagaraSimCacheCreateParameters CreateParameters, UNiagaraComponent* NiagaraComponent, UNiagaraSimCache*& OutSimCache, int32 CaptureRate, bool bAdvanceSimulation, float AdvanceDeltaTime)
-{
-	
-	FNiagaraSimCacheCaptureParameters CaptureParameters;
-	
-	CaptureParameters.NumFrames = 0;
-	CaptureParameters.CaptureRate = FMath::Max(1, CaptureRate);
-	CaptureParameters.bManuallyAdvanceSimulation = bAdvanceSimulation;
-	CaptureParameters.AdvanceDeltaTime = AdvanceDeltaTime;
-
-	return CaptureNiagaraSimCacheImpl(SimCache, CreateParameters, NiagaraComponent, OutSimCache, CaptureParameters);
-}
-
 void UAsyncNiagaraCaptureSimCache::CaptureFinished(UNiagaraSimCache* InCapturedSimCache)
 {
 	CaptureSimCache = InCapturedSimCache;
diff --git a/Engine/Plugins/FX/Niagara/Source/NiagaraEditor/Private/Customizations/NiagaraComponentDetails.cpp b/Engine/Plugins/FX/Niagara/Source/NiagaraEditor/Private/Customizations/NiagaraComponentDetails.cpp
--- a/Engine/Plugins/FX/Niagara/Source/NiagaraEditor/Private/Customizations/NiagaraComponentDetails.cpp
+++ b/Engine/Plugins/FX/Niagara/Source/NiagaraEditor/Private/Customizations/NiagaraComponentDetails.cpp
@@ -367,7 +367,7 @@ FReply FNiagaraComponentDetails::OnCaptureSelectedSystem()
 		CaptureParameters.CaptureRate = 1;
 		// A value of 0 would capture until completion. Since this can apply broadly, looping systems would not finish capturing, so don't go below 1.
 		CaptureParameters.NumFrames = FMath::Max(1, GetDefault<UNiagaraSettings>()->QuickSimCacheCaptureFrameCount);
-		CaptureParameters.bManuallyAdvanceSimulation = false;
+		CaptureParameters.bCaptureAllFramesImmediatly = false;
 		
 		if(NiagaraComponentToUse->GetWorld())
 		{
diff --git a/Engine/Plugins/FX/Niagara/Source/NiagaraEditor/Private/Widgets/SNiagaraDebugCaptureView.cpp b/Engine/Plugins/FX/Niagara/Source/NiagaraEditor/Private/Widgets/SNiagaraDebugCaptureView.cpp
--- a/Engine/Plugins/FX/Niagara/Source/NiagaraEditor/Private/Widgets/SNiagaraDebugCaptureView.cpp
+++ b/Engine/Plugins/FX/Niagara/Source/NiagaraEditor/Private/Widgets/SNiagaraDebugCaptureView.cpp
@@ -344,8 +344,8 @@ void SNiagaraDebugCaptureView::OnMultiFrameSelected()
 		
 		FNiagaraSimCacheCaptureParameters CaptureParameters;
 		CaptureParameters.NumFrames = NumFrames;
-		CaptureParameters.bManuallyAdvanceSimulation = true;
-		CaptureParameters.AdvanceDeltaTime = 0.01666f;
+		CaptureParameters.bCaptureAllFramesImmediatly = true;
+		CaptureParameters.ImmediateCaptureDeltaTime = 0.01666f;
 		
 		bIsCaptureActive = true;
 
