- Remove some edit conditions as they didn't work well inside BP

#rnx

diff --git a/Engine/Plugins/FX/Niagara/Source/Niagara/Classes/NiagaraSimCacheCapture.h b/Engine/Plugins/FX/Niagara/Source/Niagara/Classes/NiagaraSimCacheCapture.h
--- a/Engine/Plugins/FX/Niagara/Source/Niagara/Classes/NiagaraSimCacheCapture.h
+++ b/Engine/Plugins/FX/Niagara/Source/Niagara/Classes/NiagaraSimCacheCapture.h
@@ -29,14 +29,14 @@ struct FNiagaraSimCacheCaptureParameters
 	uint32 bAppendToSimCache : 1;
 
 	/** When enabled we capture NumFrames number of frames, otherwise the capture will continue until cancelled or the simulation is complete. */
-	UPROPERTY(EditAnywhere, BlueprintReadWrite, Category = "SimCache")//, meta = (InlineEditConditionToggle))
+	UPROPERTY(EditAnywhere, BlueprintReadWrite, Category = "SimCache")
 	uint32 bCaptureFixedNumberOfFrames : 1;
 
 	/**
 	 The max number of frames to capture. The capture will stop when the simulation completes or we have rendered this many frames, whichever happens first.
 	 Set to 0 to capture until simulation completes.
 	 **/
-	UPROPERTY(EditAnywhere, BlueprintReadWrite, Category = "SimCache", meta = (EditCondition = "bCaptureFixedNumberOfFrames", ClampMin = "1"))
+	UPROPERTY(EditAnywhere, BlueprintReadWrite, Category = "SimCache", meta = (ClampMin = "1"))
 	int32 NumFrames = 16;
 
 	/** Allows for reducing the frequency of captured frames in relation to the simulation's frames. The ratio is 1 / CaptureRate, so a CaptureRate of 2 would captures frames 0, 2, 4, etc. */	
@@ -44,22 +44,22 @@ struct FNiagaraSimCacheCaptureParameters
 	int32 CaptureRate = 1;
 
 	/** When enabled the capture will time out if we reach the defined number of frames without anything new to capture. */
-	UPROPERTY(EditAnywhere, BlueprintReadWrite, Category = "SimCache")//, meta = (InlineEditConditionToggle))
+	UPROPERTY(EditAnywhere, BlueprintReadWrite, Category = "SimCache")
 	uint32 bUseTimeout : 1;
 
 	/** When we fail to capture a new frame after this many frames the capture will complete automatically. */
-	UPROPERTY(EditAnywhere, BlueprintReadWrite, Category = "SimCache", meta = (EditCondition = "bUseTimeout", ClampMin = "1"))
+	UPROPERTY(EditAnywhere, BlueprintReadWrite, Category = "SimCache", meta = (ClampMin = "1"))
 	int32 TimeoutFrameCount = 10;
 
 	/**
 	When enabled we will capture all the requested frames immediatly.
 	This will capture the simulation outside of the main work tick, i.e. if you request a 10 frame capture we will loop capturing 10 frames on the capture call rather than over 10 world ticks.
 	*/
-	UPROPERTY(EditAnywhere, BlueprintReadWrite, Category = "SimCache")//, meta = (InlineEditConditionToggle))
+	UPROPERTY(EditAnywhere, BlueprintReadWrite, Category = "SimCache")
 	uint32 bCaptureAllFramesImmediatly : 1;
 
 	/** The delta time in seconds to use when manually advancing the simulation.Defaults to 1 / 60th of a second(0.01666). */
-	UPROPERTY(EditAnywhere, BlueprintReadWrite, Category = "SimCache", meta = (EditCondition = "bCaptureAllFramesImmediatly", ClampMin = "0.0001"))
+	UPROPERTY(EditAnywhere, BlueprintReadWrite, Category = "SimCache", meta = (ClampMin = "0.0001"))
 	float ImmediateCaptureDeltaTime = 0.01666f;
 
 	void Sanitize()
