Fixed a non-deterministic crash that could happen when parameter definition assets were not fully loaded/unloaded.

diff --git a/Engine/Plugins/FX/Niagara/Source/Niagara/Private/NiagaraParameterDefinitionsSubscriber.cpp b/Engine/Plugins/FX/Niagara/Source/Niagara/Private/NiagaraParameterDefinitionsSubscriber.cpp
--- a/Engine/Plugins/FX/Niagara/Source/Niagara/Private/NiagaraParameterDefinitionsSubscriber.cpp
+++ b/Engine/Plugins/FX/Niagara/Source/Niagara/Private/NiagaraParameterDefinitionsSubscriber.cpp
@@ -183,6 +183,10 @@ void INiagaraParameterDefinitionsSubscriber::SynchronizeWithParameterDefinitions
 		AssetRegistryModule.GetRegistry().GetAssetsByClass(FTopLevelAssetPath(TEXT("/Script/NiagaraEditor"), TEXT("NiagaraParameterDefinitions")), ParameterDefinitionsAssetData);
 		for (const FAssetData& ParameterDefinitionsAssetDatum : ParameterDefinitionsAssetData)
 		{
+			if (!ParameterDefinitionsAssetDatum.IsUAsset() || !ParameterDefinitionsAssetDatum.IsAssetLoaded())
+			{
+				continue;
+			}
 			UNiagaraParameterDefinitionsBase* ParameterDefinitions = Cast<UNiagaraParameterDefinitionsBase>(ParameterDefinitionsAssetDatum.GetAsset());
 			if (ParameterDefinitions == nullptr)
 			{
diff --git a/Engine/Plugins/FX/Niagara/Source/NiagaraEditor/Private/NiagaraEditorUtilities.cpp b/Engine/Plugins/FX/Niagara/Source/NiagaraEditor/Private/NiagaraEditorUtilities.cpp
--- a/Engine/Plugins/FX/Niagara/Source/NiagaraEditor/Private/NiagaraEditorUtilities.cpp
+++ b/Engine/Plugins/FX/Niagara/Source/NiagaraEditor/Private/NiagaraEditorUtilities.cpp
@@ -4716,6 +4716,10 @@ TArray<const UNiagaraScriptVariable*> FNiagaraParameterDefinitionsUtilities::Fin
 	const TArray<TWeakObjectPtr<UNiagaraParameterDefinitions>>& CachedParameterDefinitionsAssets = NiagaraEditorModule.GetCachedParameterDefinitionsAssets();
 	for (const TWeakObjectPtr<UNiagaraParameterDefinitions>& CachedParameterDefinitionsAsset : CachedParameterDefinitionsAssets)
 	{
+		if (!CachedParameterDefinitionsAsset.IsValid())
+		{
+			continue;
+		}
 		for (const UNiagaraScriptVariable* ScriptVar : CachedParameterDefinitionsAsset->GetParametersConst())
 		{
 			if (ScriptVar->Variable.GetName() == ParameterName)
diff --git a/Engine/Plugins/FX/Niagara/Source/NiagaraEditor/Public/NiagaraEditorModule.h b/Engine/Plugins/FX/Niagara/Source/NiagaraEditor/Public/NiagaraEditorModule.h
--- a/Engine/Plugins/FX/Niagara/Source/NiagaraEditor/Public/NiagaraEditorModule.h
+++ b/Engine/Plugins/FX/Niagara/Source/NiagaraEditor/Public/NiagaraEditorModule.h
@@ -300,6 +300,7 @@ private:
 				{
 					if (AssetType* Asset = Cast<AssetType>(AssetDatum.GetAsset()))
 					{
+						Asset->ConditionalPostLoad();
 						CachedAssets.Add(MakeWeakObjectPtr(Asset));
 					}
 				}
