Update Niagara landscape DI to collect texture resources during a render command from PerInstanceTick instead of during SetShaderParameters and moves SRV parameters over to textures.  This is all in hopes of cleaning up resource lifetime issues that have been sieen.

diff --git a/Engine/Plugins/FX/Niagara/Source/Niagara/Private/NiagaraDataInterfaceLandscape.cpp b/Engine/Plugins/FX/Niagara/Source/Niagara/Private/NiagaraDataInterfaceLandscape.cpp
--- a/Engine/Plugins/FX/Niagara/Source/Niagara/Private/NiagaraDataInterfaceLandscape.cpp
+++ b/Engine/Plugins/FX/Niagara/Source/Niagara/Private/NiagaraDataInterfaceLandscape.cpp
@@ -48,7 +48,7 @@ namespace NiagaraDataInterfaceLandscape
 	static const TCHAR* TemplateShaderFile = TEXT("/Plugin/FX/Niagara/Private/NiagaraDataInterfaceLandscape.ush");
 
 	BEGIN_SHADER_PARAMETER_STRUCT(FShaderParameters,)
-		SHADER_PARAMETER_SRV(Texture2D,				BaseColorVirtualTexture)
+		SHADER_PARAMETER_TEXTURE(Texture2D,			BaseColorVirtualTexture)
 		SHADER_PARAMETER_TEXTURE(Texture2D<uint4>,	BaseColorVirtualTexturePageTable)
 		SHADER_PARAMETER_SAMPLER(SamplerState,		BaseColorVirtualTextureSampler)
 		SHADER_PARAMETER(FMatrix44f,				BaseColorVirtualTextureWorldToUvTransform)
@@ -57,7 +57,7 @@ namespace NiagaraDataInterfaceLandscape
 		SHADER_PARAMETER(FUintVector4,				BaseColorVirtualTexturePackedUniform0)
 		SHADER_PARAMETER(FUintVector4,				BaseColorVirtualTexturePackedUniform1)
 		SHADER_PARAMETER(FUintVector4,				BaseColorVirtualTextureUniforms)
-		SHADER_PARAMETER_SRV(Texture2D,				HeightVirtualTexture)
+		SHADER_PARAMETER_TEXTURE(Texture2D,			HeightVirtualTexture)
 		SHADER_PARAMETER_TEXTURE(Texture2D<uint4>,	HeightVirtualTexturePageTable)
 		SHADER_PARAMETER_SAMPLER(SamplerState,		HeightVirtualTextureSampler)
 		SHADER_PARAMETER(FMatrix44f,				HeightVirtualTextureWorldToUvTransform)
@@ -65,8 +65,8 @@ namespace NiagaraDataInterfaceLandscape
 		SHADER_PARAMETER(FUintVector4,				HeightVirtualTexturePackedUniform0)
 		SHADER_PARAMETER(FUintVector4,				HeightVirtualTexturePackedUniform1)
 		SHADER_PARAMETER(FUintVector4,				HeightVirtualTextureUniforms)
-		SHADER_PARAMETER_SRV(Texture2D,				NormalVirtualTexture0)
-		SHADER_PARAMETER_SRV(Texture2D,				NormalVirtualTexture1)
+		SHADER_PARAMETER_TEXTURE(Texture2D,			NormalVirtualTexture0)
+		SHADER_PARAMETER_TEXTURE(Texture2D,			NormalVirtualTexture1)
 		SHADER_PARAMETER_TEXTURE(Texture2D<uint4>,	NormalVirtualTexturePageTable)
 		SHADER_PARAMETER_SAMPLER(SamplerState,		NormalVirtualTexture0Sampler)
 		SHADER_PARAMETER_SAMPLER(SamplerState,		NormalVirtualTexture1Sampler)
@@ -77,7 +77,7 @@ namespace NiagaraDataInterfaceLandscape
 		SHADER_PARAMETER(FUintVector4,				NormalVirtualTextureUniforms1)
 		SHADER_PARAMETER(int32,						NormalVirtualTextureUnpackMode)
 		SHADER_PARAMETER(uint32,					NormalVirtualTextureEnabled)
-		SHADER_PARAMETER_SRV(Texture2D,				CachedHeightTexture)
+		SHADER_PARAMETER_TEXTURE(Texture2D,			CachedHeightTexture)
 		SHADER_PARAMETER_SAMPLER(SamplerState,		CachedHeightTextureSampler)
 		SHADER_PARAMETER(uint32,					CachedHeightTextureEnabled)
 		SHADER_PARAMETER(FMatrix44f,				CachedHeightTextureWorldToUvTransform)
@@ -111,7 +111,7 @@ public:
 
 	void ReleaseSourceData();
 
-	FRHIShaderResourceView* GetHeightTexture() const { return HeightTexture.SRV; }
+	FRHITexture* GetHeightTexture() const { return HeightTexture.Buffer; }
 	FRHITexture* GetPhysMatTexture() const { return PhysMatTexture.Buffer; }
 	FIntPoint GetDimensions() const { return CellCount; }
 
@@ -240,23 +240,21 @@ struct FNDILandscapeData_GameThread
 	}
 };
 
-// Landscape data used on the render thread
-struct FNDILandscapeData_RenderThread
+struct FNDILandscapeData_GameToRenderThread
 {
-	TWeakObjectPtr<const URuntimeVirtualTexture> BaseColorVirtualTexture;
-	TWeakObjectPtr<const URuntimeVirtualTexture> HeightVirtualTexture;
-	TWeakObjectPtr<const URuntimeVirtualTexture> NormalVirtualTexture;
 	const FLandscapeTextureResource* TextureResources = nullptr;
-	bool BaseColorVirtualTextureSRGB = false;
-	FVector4 BaseColorVirtualTextureWorldToUvParameters[3];
-	FVector4 HeightVirtualTextureWorldToUvParameters[4];
-	FVector4 NormalVirtualTextureWorldToUvParameters[3];
+
 	FMatrix CachedHeightTextureWorldToUvTransform = FMatrix::Identity;
 	FMatrix CachedHeightTextureUvToWorldTransform = FMatrix::Identity;
 	FVector4 CachedHeightTextureUvScaleBias = FVector4(ForceInitToZero);
 	FVector2D CachedHeightTextureGridSize = FVector2D(1.0f, 1.0f);
-	ERuntimeVirtualTextureMaterialType NormalVirtualTextureMode = ERuntimeVirtualTextureMaterialType::Count;
 
+	bool BaseColorVirtualTextureSRGB = false;
+};
+
+// Landscape data used on the render thread
+struct FNDILandscapeData_RenderThread
+{
 	enum class ENormalUnpackType : int32
 	{
 		None = 0,
@@ -265,64 +263,129 @@ struct FNDILandscapeData_RenderThread
 		B5G6R5
 	};
 
-	bool SetBaseColorVirtualTextureParameters(NiagaraDataInterfaceLandscape::FShaderParameters* ShaderParameters) const
+	struct FVirtualTextureLayer
 	{
-		const IAllocatedVirtualTexture* BaseColorAllocatedTexture = BaseColorVirtualTexture.IsValid()
-			? BaseColorVirtualTexture->GetAllocatedVirtualTexture()
-			: nullptr;
+		bool IsValid() const
+		{
+			return TextureRef.IsValid();
+		}
 
-		if (!BaseColorAllocatedTexture)
+		void Reset()
 		{
-			return false;
+			TextureRef = nullptr;
+			TextureUniforms = FUintVector4(0, 0, 0, 0);
+		}
+
+		void Update(const URuntimeVirtualTexture* VirtualTexture, uint32 LayerIndex)
+		{
+			Reset();
+
+			if (VirtualTexture && ::IsValid(VirtualTexture))
+			{
+				if (IAllocatedVirtualTexture* AllocatedTexture = VirtualTexture->GetAllocatedVirtualTexture())
+				{
+					if (FRHITexture* PhysicalTexture = AllocatedTexture->GetPhysicalTexture(LayerIndex))
+					{
+						TextureRef = PhysicalTexture;
+						AllocatedTexture->GetPackedUniform(&TextureUniforms, LayerIndex);
+					}
+				}
+			}
 		}
 
-		// todo - need to figure out a way to confirm that this is in fact the best/only option for the basecolor
-		constexpr uint32 BaseColorVirtualTextureLayerIndex = 0;
-		constexpr uint32 BaseColorVirtualTexturePageIndex = 0;
-		constexpr bool BaseColorVirtualTextureSrgb = false;
+		FTextureRHIRef TextureRef;
+		FUintVector4 TextureUniforms = FUintVector4(0, 0, 0, 0);
+	};
 
-		FRHIShaderResourceView* PhysicalTextureSrv = BaseColorAllocatedTexture->GetPhysicalTextureSRV(BaseColorVirtualTextureLayerIndex, BaseColorVirtualTextureSrgb);
-		if (!PhysicalTextureSrv)
+	struct FVirtualTexturePage
+	{
+		bool IsValid() const
 		{
-			return false;
+			return PageTableRef.IsValid();
 		}
 
-		FRHITexture* PageTableTexture = nullptr;
+		void Reset()
+		{
+			PageTableRef = nullptr;
+			PageTableUniforms[0] = FUintVector4(0, 0, 0, 0);
+			PageTableUniforms[1] = FUintVector4(0, 0, 0, 0);
+
+			WorldToUvParameters[0] = FVector4::Zero();
+			WorldToUvParameters[1] = FVector4::Zero();
+			WorldToUvParameters[2] = FVector4::Zero();
+			WorldToUvParameters[3] = FVector4::Zero();
+		}
 
-		if (FRHITexture* PageTable = BaseColorAllocatedTexture->GetPageTableTexture(BaseColorVirtualTexturePageIndex))
+		void Update(const URuntimeVirtualTexture* VirtualTexture, uint32 PageTableIndex, bool bIncludeWorldToUv, bool bIncludeHeightUnpack)
 		{
-			if (FRHITextureReference* TextureReference = PageTable->GetTextureReference())
+			Reset();
+
+			if (VirtualTexture && ::IsValid(VirtualTexture))
 			{
-				PageTableTexture = TextureReference->GetReferencedTexture();
+				if (IAllocatedVirtualTexture* AllocatedTexture = VirtualTexture->GetAllocatedVirtualTexture())
+				{
+					if (FRHITexture* PageTable = AllocatedTexture->GetPageTableTexture(PageTableIndex))
+					{
+						if (FRHITextureReference* TextureReference = PageTable->GetTextureReference())
+						{
+							PageTableRef = TextureReference->GetReferencedTexture();
+							AllocatedTexture->GetPackedPageTableUniform(PageTableUniforms);
+
+							if (bIncludeWorldToUv)
+							{
+								WorldToUvParameters[0] = VirtualTexture->GetUniformParameter(ERuntimeVirtualTextureShaderUniform_WorldToUVTransform0);
+								WorldToUvParameters[1] = VirtualTexture->GetUniformParameter(ERuntimeVirtualTextureShaderUniform_WorldToUVTransform1);
+								WorldToUvParameters[2] = VirtualTexture->GetUniformParameter(ERuntimeVirtualTextureShaderUniform_WorldToUVTransform2);
+							}
+
+							if (bIncludeHeightUnpack)
+							{
+								WorldToUvParameters[3] = VirtualTexture->GetUniformParameter(ERuntimeVirtualTextureShaderUniform_WorldHeightUnpack);
+							}
+						}
+					}
+				}
 			}
 		}
 
-		if (!PageTableTexture)
+		FTextureRHIRef PageTableRef;
+		FUintVector4 PageTableUniforms[2];
+
+		FVector4 WorldToUvParameters[4];
+	};
+
+	FVirtualTexturePage BaseColorVirtualPage;
+	FVirtualTextureLayer BaseColorVirtualLayer;
+
+	FVirtualTexturePage HeightVirtualPage;
+	FVirtualTextureLayer HeightVirtualLayer;
+
+	FVirtualTexturePage NormalVirtualPage;
+	FVirtualTextureLayer NormalVirtualLayer0;
+	FVirtualTextureLayer NormalVirtualLayer1;
+	ENormalUnpackType NormalUnpackMode;
+
+	FNDILandscapeData_GameToRenderThread LandscapeData;
+
+	bool SetBaseColorVirtualTextureParameters(NiagaraDataInterfaceLandscape::FShaderParameters* ShaderParameters) const
+	{
+		if (!BaseColorVirtualPage.IsValid() || !BaseColorVirtualLayer.IsValid())
 		{
 			return false;
 		}
 
-		FMatrix44f WorldToUvTransform(
-			FVector3f((FVector4f)BaseColorVirtualTextureWorldToUvParameters[0]),
-			FVector3f((FVector4f)BaseColorVirtualTextureWorldToUvParameters[1]),
-			FVector3f((FVector4f)BaseColorVirtualTextureWorldToUvParameters[2]),
+		ShaderParameters->BaseColorVirtualTexture = BaseColorVirtualLayer.TextureRef;
+		ShaderParameters->BaseColorVirtualTexturePageTable = BaseColorVirtualPage.PageTableRef;
+		ShaderParameters->BaseColorVirtualTextureWorldToUvTransform = FMatrix44f(
+			FVector3f((FVector4f)BaseColorVirtualPage.WorldToUvParameters[0]),
+			FVector3f((FVector4f)BaseColorVirtualPage.WorldToUvParameters[1]),
+			FVector3f((FVector4f)BaseColorVirtualPage.WorldToUvParameters[2]),
 			FVector3f(0, 0, 0));
-
-		ShaderParameters->BaseColorVirtualTexture = PhysicalTextureSrv;
-		ShaderParameters->BaseColorVirtualTexturePageTable = PageTableTexture;
-		ShaderParameters->BaseColorVirtualTextureWorldToUvTransform = WorldToUvTransform;
-		ShaderParameters->BaseColorVirtualTextureSRGB = BaseColorVirtualTextureSRGB ? 1 : 0;
+		ShaderParameters->BaseColorVirtualTextureSRGB = LandscapeData.BaseColorVirtualTextureSRGB;
 		ShaderParameters->BaseColorVirtualTextureEnabled = 1;
-
-		FUintVector4 BaseColorVirtualTexturePageTableUniforms[2];
-		FUintVector4 BaseColorVirtualTextureUniforms;
-
-		BaseColorAllocatedTexture->GetPackedPageTableUniform(BaseColorVirtualTexturePageTableUniforms);
-		BaseColorAllocatedTexture->GetPackedUniform(&BaseColorVirtualTextureUniforms, BaseColorVirtualTextureLayerIndex);
-		ShaderParameters->BaseColorVirtualTexturePackedUniform0 = BaseColorVirtualTexturePageTableUniforms[0];
-		ShaderParameters->BaseColorVirtualTexturePackedUniform1 = BaseColorVirtualTexturePageTableUniforms[1];
-
-		ShaderParameters->BaseColorVirtualTextureUniforms = BaseColorVirtualTextureUniforms;
+		ShaderParameters->BaseColorVirtualTexturePackedUniform0 = BaseColorVirtualPage.PageTableUniforms[0];
+		ShaderParameters->BaseColorVirtualTexturePackedUniform1 = BaseColorVirtualPage.PageTableUniforms[1];
+		ShaderParameters->BaseColorVirtualTextureUniforms = BaseColorVirtualLayer.TextureUniforms;
 
 		return true;
 	}
@@ -331,7 +394,7 @@ struct FNDILandscapeData_RenderThread
 	{
 		const FUintVector4 DummyUint4(ForceInitToZero);
 
-		ShaderParameters->BaseColorVirtualTexture = GBlackTextureWithSRV->ShaderResourceViewRHI;
+		ShaderParameters->BaseColorVirtualTexture = GBlackTexture->TextureRHI;
 		ShaderParameters->BaseColorVirtualTexturePageTable = GBlackUintTexture->TextureRHI;
 		ShaderParameters->BaseColorVirtualTextureWorldToUvTransform = FMatrix44f::Identity;
 		ShaderParameters->BaseColorVirtualTextureSRGB = 0;
@@ -343,61 +406,23 @@ struct FNDILandscapeData_RenderThread
 
 	bool SetHeightVirtualTextureParameters(NiagaraDataInterfaceLandscape::FShaderParameters* ShaderParameters) const
 	{
-		const IAllocatedVirtualTexture* HeightAllocatedTexture = HeightVirtualTexture.IsValid()
-			? HeightVirtualTexture->GetAllocatedVirtualTexture()
-			: nullptr;
-
-		if (!HeightAllocatedTexture)
-		{
-			return false;
-		}
-
-		// todo - need to figure out a way to confirm that this is in fact the best/only option for the height
-		constexpr uint32 HeightVirtualTextureLayerIndex = 0;
-		constexpr uint32 HeightVirtualTexturePageIndex = 0;
-		constexpr bool HeightVirtualTextureSrgb = false;
-
-		FRHIShaderResourceView* PhysicalTextureSrv = HeightAllocatedTexture->GetPhysicalTextureSRV(HeightVirtualTextureLayerIndex, HeightVirtualTextureSrgb);
-		if (!PhysicalTextureSrv)
-		{
-			return false;
-		}
-
-		FRHITexture* PageTableTexture = nullptr;
-
-		if (FRHITexture* PageTable = HeightAllocatedTexture->GetPageTableTexture(HeightVirtualTexturePageIndex))
-		{
-			if (FRHITextureReference* TextureReference = PageTable->GetTextureReference())
-			{
-				PageTableTexture = TextureReference->GetReferencedTexture();
-			}
-		}
-
-		if (!PageTableTexture)
+		if (!HeightVirtualPage.IsValid() || !HeightVirtualLayer.IsValid())
 		{
 			return false;
 		}
 
-		FMatrix44f WorldToUvTransform(
-			FVector3f( (FVector4f)HeightVirtualTextureWorldToUvParameters[0]),
-			FVector3f( (FVector4f)HeightVirtualTextureWorldToUvParameters[1]),
-			FVector3f( (FVector4f)HeightVirtualTextureWorldToUvParameters[2]),
-			FVector3f( (FVector4f)HeightVirtualTextureWorldToUvParameters[3]));
+		ShaderParameters->HeightVirtualTexture = HeightVirtualLayer.TextureRef;
+		ShaderParameters->HeightVirtualTexturePageTable = HeightVirtualPage.PageTableRef;
+		ShaderParameters->HeightVirtualTextureWorldToUvTransform = FMatrix44f(
+			FVector3f((FVector4f)HeightVirtualPage.WorldToUvParameters[0]),
+			FVector3f((FVector4f)HeightVirtualPage.WorldToUvParameters[1]),
+			FVector3f((FVector4f)HeightVirtualPage.WorldToUvParameters[2]),
+			FVector3f((FVector4f)HeightVirtualPage.WorldToUvParameters[3]));
 
-		ShaderParameters->HeightVirtualTexture = PhysicalTextureSrv;
-		ShaderParameters->HeightVirtualTexturePageTable = PageTableTexture;
-		ShaderParameters->HeightVirtualTextureWorldToUvTransform = WorldToUvTransform;
 		ShaderParameters->HeightVirtualTextureEnabled = 1;
-
-		FUintVector4 HeightVirtualTexturePageTableUniforms[2];
-		FUintVector4 HeightVirtualTextureUniforms;
-
-		HeightAllocatedTexture->GetPackedPageTableUniform(HeightVirtualTexturePageTableUniforms);
-		HeightAllocatedTexture->GetPackedUniform(&HeightVirtualTextureUniforms, HeightVirtualTextureLayerIndex);
-		ShaderParameters->HeightVirtualTexturePackedUniform0 = HeightVirtualTexturePageTableUniforms[0];
-		ShaderParameters->HeightVirtualTexturePackedUniform1 = HeightVirtualTexturePageTableUniforms[1];
-
-		ShaderParameters->HeightVirtualTextureUniforms = HeightVirtualTextureUniforms;
+		ShaderParameters->HeightVirtualTexturePackedUniform0 = HeightVirtualPage.PageTableUniforms[0];
+		ShaderParameters->HeightVirtualTexturePackedUniform1 = HeightVirtualPage.PageTableUniforms[1];
+		ShaderParameters->HeightVirtualTextureUniforms = HeightVirtualLayer.TextureUniforms;
 
 		return true;
 	}
@@ -406,7 +431,7 @@ struct FNDILandscapeData_RenderThread
 	{
 		const FUintVector4 DummyUint4(ForceInitToZero);
 
-		ShaderParameters->HeightVirtualTexture = GBlackTextureWithSRV->ShaderResourceViewRHI;
+		ShaderParameters->HeightVirtualTexture = GBlackTexture->TextureRHI;
 		ShaderParameters->HeightVirtualTexturePageTable = GBlackUintTexture->TextureRHI;
 		ShaderParameters->HeightVirtualTextureWorldToUvTransform = FMatrix44f::Identity;
 		ShaderParameters->HeightVirtualTextureEnabled = 0;
@@ -417,90 +442,26 @@ struct FNDILandscapeData_RenderThread
 
 	bool SetNormalVirtualTextureParameters(NiagaraDataInterfaceLandscape::FShaderParameters* ShaderParameters) const
 	{
-		const IAllocatedVirtualTexture* NormalAllocatedTexture = NormalVirtualTexture.IsValid()
-			? NormalVirtualTexture->GetAllocatedVirtualTexture()
-			: nullptr;
-
-		if (!NormalAllocatedTexture)
-		{
-			return false;
-		}
-
-		constexpr uint32 NormalVirtualTexturePageIndex = 0;
-
-		FRHITexture* PageTableTexture = nullptr;
-
-		if (FRHITexture* PageTable = NormalAllocatedTexture->GetPageTableTexture(NormalVirtualTexturePageIndex))
-		{
-			if (FRHITextureReference* TextureReference = PageTable->GetTextureReference())
-			{
-				PageTableTexture = TextureReference->GetReferencedTexture();
-			}
-		}
-
-		if (!PageTableTexture)
+		if (!NormalVirtualPage.IsValid() || !NormalVirtualLayer0.IsValid() || !NormalVirtualLayer1.IsValid())
 		{
 			return false;
 		}
 
-		FRHIShaderResourceView* PhysicalTextureSrv[2] = { GBlackTextureWithSRV->ShaderResourceViewRHI, GBlackTextureWithSRV->ShaderResourceViewRHI };
-		FUintVector4 NormalVirtualTextureUniforms[2];
-
-		switch (NormalVirtualTextureMode)
-		{
-			case ERuntimeVirtualTextureMaterialType::BaseColor_Normal_Roughness:
-				PhysicalTextureSrv[0] = NormalAllocatedTexture->GetPhysicalTextureSRV(0, false);
-				NormalAllocatedTexture->GetPackedUniform(&NormalVirtualTextureUniforms[0], 0);
-
-				PhysicalTextureSrv[1] = NormalAllocatedTexture->GetPhysicalTextureSRV(1, false);
-				NormalAllocatedTexture->GetPackedUniform(&NormalVirtualTextureUniforms[1], 1);
-
-				ShaderParameters->NormalVirtualTextureUnpackMode = (int32)ENormalUnpackType::B5G6R5;
-				break;
-			case ERuntimeVirtualTextureMaterialType::BaseColor_Normal_Specular:
-				PhysicalTextureSrv[0] = NormalAllocatedTexture->GetPhysicalTextureSRV(0, false);
-				NormalAllocatedTexture->GetPackedUniform(&NormalVirtualTextureUniforms[0], 0);
-
-				PhysicalTextureSrv[1] = NormalAllocatedTexture->GetPhysicalTextureSRV(1, false);
-				NormalAllocatedTexture->GetPackedUniform(&NormalVirtualTextureUniforms[1], 1);
-
-				ShaderParameters->NormalVirtualTextureUnpackMode = (int32)ENormalUnpackType::BC3BC3;
-				break;
-
-			case ERuntimeVirtualTextureMaterialType::BaseColor_Normal_Specular_YCoCg:
-			case ERuntimeVirtualTextureMaterialType::BaseColor_Normal_Specular_Mask_YCoCg:
-				PhysicalTextureSrv[0] = NormalAllocatedTexture->GetPhysicalTextureSRV(1, false);
-				NormalAllocatedTexture->GetPackedUniform(&NormalVirtualTextureUniforms[0], 1);
-
-				PhysicalTextureSrv[1] = NormalAllocatedTexture->GetPhysicalTextureSRV(2, false);
-				NormalAllocatedTexture->GetPackedUniform(&NormalVirtualTextureUniforms[1], 2);
-
-				ShaderParameters->NormalVirtualTextureUnpackMode = (int32)ENormalUnpackType::BC5BC1;
-				break;
-
-			default:
-				return false;
-		}
-
-		FMatrix44f WorldToUvTransform(
-			FPlane4f( (FVector4f)NormalVirtualTextureWorldToUvParameters[0]),
-			FPlane4f( (FVector4f)NormalVirtualTextureWorldToUvParameters[1]),
-			FPlane4f( (FVector4f)NormalVirtualTextureWorldToUvParameters[2]),
+		ShaderParameters->NormalVirtualTexture0 = NormalVirtualLayer0.TextureRef;
+		ShaderParameters->NormalVirtualTexture1 = NormalVirtualLayer1.TextureRef;
+		ShaderParameters->NormalVirtualTexturePageTable = NormalVirtualPage.PageTableRef;
+		ShaderParameters->NormalVirtualTextureWorldToUvTransform = FMatrix44f(
+			FPlane4f((FVector4f)NormalVirtualPage.WorldToUvParameters[0]),
+			FPlane4f((FVector4f)NormalVirtualPage.WorldToUvParameters[1]),
+			FPlane4f((FVector4f)NormalVirtualPage.WorldToUvParameters[2]),
 			FPlane4f(0.0f, 0.0f, 0.0f, 1.0f));
 
-		ShaderParameters->NormalVirtualTexture0 = PhysicalTextureSrv[0];
-		ShaderParameters->NormalVirtualTexture1 = PhysicalTextureSrv[1];
-		ShaderParameters->NormalVirtualTexturePageTable = PageTableTexture;
-		ShaderParameters->NormalVirtualTextureWorldToUvTransform = WorldToUvTransform;
 		ShaderParameters->NormalVirtualTextureEnabled = 1;
-
-		FUintVector4 NormalVirtualTexturePageTableUniforms[2];
-		NormalAllocatedTexture->GetPackedPageTableUniform(NormalVirtualTexturePageTableUniforms);
-		ShaderParameters->NormalVirtualTexturePackedUniform0 = NormalVirtualTexturePageTableUniforms[0];
-		ShaderParameters->NormalVirtualTexturePackedUniform1 = NormalVirtualTexturePageTableUniforms[1];
-
-		ShaderParameters->NormalVirtualTextureUniforms0 = NormalVirtualTextureUniforms[0];
-		ShaderParameters->NormalVirtualTextureUniforms1 = NormalVirtualTextureUniforms[1];
+		ShaderParameters->NormalVirtualTextureUnpackMode = (int32)NormalUnpackMode;
+		ShaderParameters->NormalVirtualTexturePackedUniform0 = NormalVirtualPage.PageTableUniforms[0];
+		ShaderParameters->NormalVirtualTexturePackedUniform1 = NormalVirtualPage.PageTableUniforms[1];
+		ShaderParameters->NormalVirtualTextureUniforms0 = NormalVirtualLayer0.TextureUniforms;
+		ShaderParameters->NormalVirtualTextureUniforms1 = NormalVirtualLayer1.TextureUniforms;
 
 		return true;
 	}
@@ -509,8 +470,8 @@ struct FNDILandscapeData_RenderThread
 	{
 		const FUintVector4 DummyUint4(ForceInitToZero);
 
-		ShaderParameters->NormalVirtualTexture0 =  GBlackTextureWithSRV->ShaderResourceViewRHI;
-		ShaderParameters->NormalVirtualTexture1 = GBlackTextureWithSRV->ShaderResourceViewRHI;
+		ShaderParameters->NormalVirtualTexture0 =  GBlackTexture->TextureRHI;
+		ShaderParameters->NormalVirtualTexture1 = GBlackTexture->TextureRHI;
 		ShaderParameters->NormalVirtualTexturePageTable = GBlackUintTexture->TextureRHI;
 		ShaderParameters->NormalVirtualTextureWorldToUvTransform = FMatrix44f::Identity;
 		ShaderParameters->NormalVirtualTextureEnabled = 0;
@@ -523,33 +484,33 @@ struct FNDILandscapeData_RenderThread
 
 	bool SetCachedHeightTextureParameters(NiagaraDataInterfaceLandscape::FShaderParameters* ShaderParameters) const
 	{
-		if (TextureResources)
+		if (LandscapeData.TextureResources)
 		{
-			FRHIShaderResourceView* HeightTextureSrv = TextureResources->GetHeightTexture();
-			FRHITexture* PhysMatTextureSrv = TextureResources->GetPhysMatTexture();
+			FRHITexture* HeightTexture = LandscapeData.TextureResources->GetHeightTexture();
+			FRHITexture* PhysMatTexture = LandscapeData.TextureResources->GetPhysMatTexture();
 
-			if (HeightTextureSrv || PhysMatTextureSrv)
+			if (HeightTexture || PhysMatTexture)
 			{
-				ShaderParameters->CachedHeightTextureWorldToUvTransform = (FMatrix44f)CachedHeightTextureWorldToUvTransform;
-				ShaderParameters->CachedHeightTextureUvToWorldTransform = (FMatrix44f)CachedHeightTextureUvToWorldTransform;
-				ShaderParameters->CachedHeightTextureUvScaleBias = (FVector4f)CachedHeightTextureUvScaleBias;
-				ShaderParameters->CachedHeightTextureGridSize = (FVector2f)CachedHeightTextureGridSize;
+				ShaderParameters->CachedHeightTextureWorldToUvTransform = (FMatrix44f)LandscapeData.CachedHeightTextureWorldToUvTransform;
+				ShaderParameters->CachedHeightTextureUvToWorldTransform = (FMatrix44f)LandscapeData.CachedHeightTextureUvToWorldTransform;
+				ShaderParameters->CachedHeightTextureUvScaleBias = (FVector4f)LandscapeData.CachedHeightTextureUvScaleBias;
+				ShaderParameters->CachedHeightTextureGridSize = (FVector2f)LandscapeData.CachedHeightTextureGridSize;
 
-				if (HeightTextureSrv)
+				if (HeightTexture)
 				{
-					ShaderParameters->CachedHeightTexture = HeightTextureSrv;
+					ShaderParameters->CachedHeightTexture = HeightTexture;
 					ShaderParameters->CachedHeightTextureEnabled = 1;
 				}
 				else
 				{
-					ShaderParameters->CachedHeightTexture = GBlackTextureWithSRV->ShaderResourceViewRHI;
+					ShaderParameters->CachedHeightTexture = GBlackTexture->TextureRHI;
 					ShaderParameters->CachedHeightTextureEnabled = 0;
 				}
 
-				if (PhysMatTextureSrv)
+				if (PhysMatTexture)
 				{
-					const FIntPoint PhysMatDimensions(TextureResources->GetDimensions());
-					ShaderParameters->CachedPhysMatTexture = PhysMatTextureSrv;
+					const FIntPoint PhysMatDimensions(LandscapeData.TextureResources->GetDimensions());
+					ShaderParameters->CachedPhysMatTexture = PhysMatTexture;
 					ShaderParameters->CachedPhysMatTextureDimension = PhysMatDimensions;
 				}
 				else
@@ -570,7 +531,7 @@ struct FNDILandscapeData_RenderThread
 	{
 		const FVector4f DummyVector4(ForceInitToZero);
 
-		ShaderParameters->CachedHeightTexture = GBlackTextureWithSRV->ShaderResourceViewRHI;
+		ShaderParameters->CachedHeightTexture = GBlackTexture->TextureRHI;
 		ShaderParameters->CachedHeightTextureWorldToUvTransform = FMatrix44f::Identity;
 		ShaderParameters->CachedHeightTextureUvToWorldTransform = FMatrix44f::Identity;
 		ShaderParameters->CachedHeightTextureUvScaleBias = DummyVector4;
@@ -601,13 +562,66 @@ struct FNiagaraDataInterfaceProxyLandscape : public FNiagaraDataInterfaceProxy
 {
 	virtual void ConsumePerInstanceDataFromGameThread(void* PerInstanceData, const FNiagaraSystemInstanceID& Instance) override
 	{
-		const FNDILandscapeData_RenderThread& SourceData = *reinterpret_cast<const FNDILandscapeData_RenderThread*>(PerInstanceData);
-		*SystemInstancesToProxyData_RT.Find(Instance) = SourceData;
+		const FNDILandscapeData_GameToRenderThread& SourceData = *reinterpret_cast<const FNDILandscapeData_GameToRenderThread*>(PerInstanceData);
+		SystemInstancesToProxyData_RT.FindOrAdd(Instance).LandscapeData = SourceData;
 	}
 
 	virtual int32 PerInstanceDataPassedToRenderThreadSize() const override
 	{
-		return sizeof(FNDILandscapeData_RenderThread);
+		return sizeof(FNDILandscapeData_GameToRenderThread);
+	}
+
+	void UpdateProxy_RT(
+		FNiagaraSystemInstanceID ID,
+		const URuntimeVirtualTexture* BaseColorVirtualTexture,
+		const URuntimeVirtualTexture* HeightVirtualTexture,
+		const URuntimeVirtualTexture* NormalVirtualTexture,
+		ERuntimeVirtualTextureMaterialType NormalVirtualTextureMode)
+	{
+		if (FNDILandscapeData_RenderThread* Proxy_RT = SystemInstancesToProxyData_RT.Find(ID))
+		{
+			// todo - need to figure out a way to confirm that this is in fact the best/only option for the page/layer indices
+			constexpr uint32 BaseColorVirtualTextureLayerIndex = 0;
+			constexpr uint32 BaseColorVirtualTexturePageIndex = 0;
+			Proxy_RT->BaseColorVirtualLayer.Update(BaseColorVirtualTexture, BaseColorVirtualTextureLayerIndex);
+			Proxy_RT->BaseColorVirtualPage.Update(BaseColorVirtualTexture, BaseColorVirtualTexturePageIndex, true /*bIncludeWorldToUv*/, false /*bIncludeHeightUnpack*/);
+
+			constexpr uint32 HeightVirtualTextureLayerIndex = 0;
+			constexpr uint32 HeightVirtualTexturePageIndex = 0;
+			Proxy_RT->HeightVirtualLayer.Update(HeightVirtualTexture, HeightVirtualTextureLayerIndex);
+			Proxy_RT->HeightVirtualPage.Update(HeightVirtualTexture, HeightVirtualTexturePageIndex, true /*bIncludeWorldToUv*/, true /*bIncludeHeightUnpack*/);
+
+			constexpr uint32 NormalVirtualTexturePageIndex = 0;
+			Proxy_RT->NormalVirtualPage.Update(NormalVirtualTexture, NormalVirtualTexturePageIndex, true /*bIncludeWorldToUv*/, false /*bIncludeHeightUnpack*/);
+
+			switch (NormalVirtualTextureMode)
+			{
+			case ERuntimeVirtualTextureMaterialType::BaseColor_Normal_Roughness:
+				Proxy_RT->NormalVirtualLayer0.Update(NormalVirtualTexture, 0);
+				Proxy_RT->NormalVirtualLayer1.Update(NormalVirtualTexture, 1);
+				Proxy_RT->NormalUnpackMode = FNDILandscapeData_RenderThread::ENormalUnpackType::B5G6R5;
+				break;
+			case ERuntimeVirtualTextureMaterialType::BaseColor_Normal_Specular:
+				Proxy_RT->NormalVirtualLayer0.Update(NormalVirtualTexture, 0);
+				Proxy_RT->NormalVirtualLayer1.Update(NormalVirtualTexture, 1);
+				Proxy_RT->NormalUnpackMode = FNDILandscapeData_RenderThread::ENormalUnpackType::BC3BC3;
+				break;
+
+			case ERuntimeVirtualTextureMaterialType::BaseColor_Normal_Specular_YCoCg:
+			case ERuntimeVirtualTextureMaterialType::BaseColor_Normal_Specular_Mask_YCoCg:
+				Proxy_RT->NormalVirtualLayer0.Update(NormalVirtualTexture, 1);
+				Proxy_RT->NormalVirtualLayer1.Update(NormalVirtualTexture, 2);
+				Proxy_RT->NormalUnpackMode = FNDILandscapeData_RenderThread::ENormalUnpackType::BC5BC1;
+				break;
+
+			default:
+				Proxy_RT->NormalVirtualPage.Reset();
+				Proxy_RT->NormalVirtualLayer0.Reset();
+				Proxy_RT->NormalVirtualLayer1.Reset();
+				Proxy_RT->NormalUnpackMode = FNDILandscapeData_RenderThread::ENormalUnpackType::None;
+				break;
+			}
+		}
 	}
 
 	TMap<FNiagaraSystemInstanceID, FNDILandscapeData_RenderThread> SystemInstancesToProxyData_RT;
@@ -1223,7 +1237,7 @@ bool UNiagaraDataInterfaceLandscape::UpgradeFunctionCall(FNiagaraFunctionSignatu
 void UNiagaraDataInterfaceLandscape::ProvidePerInstanceDataForRenderThread(void* DataForRenderThread, void* PerInstanceData, const FNiagaraSystemInstanceID& SystemInstance)
 {
 	const FNDILandscapeData_GameThread& SourceData = *reinterpret_cast<const FNDILandscapeData_GameThread*>(PerInstanceData);
-	FNDILandscapeData_RenderThread* TargetData = new(DataForRenderThread) FNDILandscapeData_RenderThread();
+	FNDILandscapeData_GameToRenderThread* TargetData = new(DataForRenderThread) FNDILandscapeData_GameToRenderThread();
 
 	if (SourceData.SharedResourceHandle)
 	{
@@ -1237,44 +1251,7 @@ void UNiagaraDataInterfaceLandscape::ProvidePerInstanceDataForRenderThread(void*
 		TargetData->CachedHeightTextureGridSize = SourceResource.TextureWorldGridSize;
 	}
 
-	if (ALandscape* SourceDataLandscape = SourceData.Landscape.Get())
-	{
-		if (SourceData.BaseColorVirtualTextureIndex != INDEX_NONE)
-		{
-			if (const URuntimeVirtualTexture* VirtualTexture = SourceDataLandscape->RuntimeVirtualTextures[SourceData.BaseColorVirtualTextureIndex])
-			{
-				TargetData->BaseColorVirtualTexture = VirtualTexture;
-				TargetData->BaseColorVirtualTextureSRGB = SourceData.BaseColorVirtualTextureSRGB;
-				TargetData->BaseColorVirtualTextureWorldToUvParameters[0] = VirtualTexture->GetUniformParameter(ERuntimeVirtualTextureShaderUniform_WorldToUVTransform0);
-				TargetData->BaseColorVirtualTextureWorldToUvParameters[1] = VirtualTexture->GetUniformParameter(ERuntimeVirtualTextureShaderUniform_WorldToUVTransform1);
-				TargetData->BaseColorVirtualTextureWorldToUvParameters[2] = VirtualTexture->GetUniformParameter(ERuntimeVirtualTextureShaderUniform_WorldToUVTransform2);
-			}
-		}
-
-		if (SourceData.HeightVirtualTextureIndex != INDEX_NONE)
-		{
-			if (const URuntimeVirtualTexture* VirtualTexture = SourceDataLandscape->RuntimeVirtualTextures[SourceData.HeightVirtualTextureIndex])
-			{
-				TargetData->HeightVirtualTexture = VirtualTexture;
-				TargetData->HeightVirtualTextureWorldToUvParameters[0] = VirtualTexture->GetUniformParameter(ERuntimeVirtualTextureShaderUniform_WorldToUVTransform0);
-				TargetData->HeightVirtualTextureWorldToUvParameters[1] = VirtualTexture->GetUniformParameter(ERuntimeVirtualTextureShaderUniform_WorldToUVTransform1);
-				TargetData->HeightVirtualTextureWorldToUvParameters[2] = VirtualTexture->GetUniformParameter(ERuntimeVirtualTextureShaderUniform_WorldToUVTransform2);
-				TargetData->HeightVirtualTextureWorldToUvParameters[3] = VirtualTexture->GetUniformParameter(ERuntimeVirtualTextureShaderUniform_WorldHeightUnpack);
-			}
-		}
-
-		if (SourceData.NormalVirtualTextureIndex != INDEX_NONE)
-		{
-			if (const URuntimeVirtualTexture* VirtualTexture = SourceDataLandscape->RuntimeVirtualTextures[SourceData.NormalVirtualTextureIndex])
-			{
-				TargetData->NormalVirtualTexture = VirtualTexture;
-				TargetData->NormalVirtualTextureWorldToUvParameters[0] = VirtualTexture->GetUniformParameter(ERuntimeVirtualTextureShaderUniform_WorldToUVTransform0);
-				TargetData->NormalVirtualTextureWorldToUvParameters[1] = VirtualTexture->GetUniformParameter(ERuntimeVirtualTextureShaderUniform_WorldToUVTransform1);
-				TargetData->NormalVirtualTextureWorldToUvParameters[2] = VirtualTexture->GetUniformParameter(ERuntimeVirtualTextureShaderUniform_WorldToUVTransform2);
-				TargetData->NormalVirtualTextureMode = SourceData.NormalVirtualTextureMode;
-			}
-		}
-	}
+	TargetData->BaseColorVirtualTextureSRGB = SourceData.BaseColorVirtualTextureSRGB;
 }
 
 #if WITH_EDITORONLY_DATA
@@ -1451,6 +1428,38 @@ bool UNiagaraDataInterfaceLandscape::PerInstanceTick(void* PerInstanceData, FNia
 		{
 			InstanceData->SharedResourceHandle = FNDI_Landscape_SharedResourceHandle();
 		}
+
+		const URuntimeVirtualTexture* BaseColorVirtualTexture = nullptr;
+		const URuntimeVirtualTexture* HeightVirtualTexture = nullptr;
+		const URuntimeVirtualTexture* NormalVirtualTexture = nullptr;
+
+		if (ALandscape* SourceDataLandscape = InstanceData->Landscape.Get())
+		{
+			if (SourceDataLandscape->RuntimeVirtualTextures.IsValidIndex(InstanceData->BaseColorVirtualTextureIndex))
+			{
+				BaseColorVirtualTexture = SourceDataLandscape->RuntimeVirtualTextures[InstanceData->BaseColorVirtualTextureIndex];
+			}
+			if (SourceDataLandscape->RuntimeVirtualTextures.IsValidIndex(InstanceData->HeightVirtualTextureIndex))
+			{
+				HeightVirtualTexture = SourceDataLandscape->RuntimeVirtualTextures[InstanceData->HeightVirtualTextureIndex];
+			}
+			if (SourceDataLandscape->RuntimeVirtualTextures.IsValidIndex(InstanceData->NormalVirtualTextureIndex))
+			{
+				NormalVirtualTexture = SourceDataLandscape->RuntimeVirtualTextures[InstanceData->NormalVirtualTextureIndex];
+			}
+		}
+
+		FNiagaraDataInterfaceProxyLandscape* RT_Proxy = GetProxyAs<FNiagaraDataInterfaceProxyLandscape>();
+		ENQUEUE_RENDER_COMMAND(FNiagaraDIUpdateProxy) (
+			[RT_Proxy,
+			BaseColorVirtualTexture,
+			HeightVirtualTexture,
+			NormalVirtualTexture,
+			NormalVirtualTextureMode = InstanceData->NormalVirtualTextureMode,
+			InstanceID = SystemInstance->GetId()](FRHICommandListImmediate& CmdList)
+		{
+			RT_Proxy->UpdateProxy_RT(InstanceID, BaseColorVirtualTexture, HeightVirtualTexture, NormalVirtualTexture, NormalVirtualTextureMode);
+		});
 	}
 
 	return ResetInstance;
