- Niagara object reader fix for float vs uint on buffer type

#rnx

diff --git a/Engine/Plugins/FX/Niagara/Source/Niagara/Private/DataInterface/NiagaraDataInterfaceUObjectPropertyReader.cpp b/Engine/Plugins/FX/Niagara/Source/Niagara/Private/DataInterface/NiagaraDataInterfaceUObjectPropertyReader.cpp
--- a/Engine/Plugins/FX/Niagara/Source/Niagara/Private/DataInterface/NiagaraDataInterfaceUObjectPropertyReader.cpp
+++ b/Engine/Plugins/FX/Niagara/Source/Niagara/Private/DataInterface/NiagaraDataInterfaceUObjectPropertyReader.cpp
@@ -124,7 +124,7 @@ namespace NDIUObjectPropertyReaderLocal
 					const uint32 NumElements = InstanceData_FromGT->PropertyData.Num() / 4;
 					if (NumElements > 0)
 					{
-						InstanceData_RT->PropertyData.Initialize(RHICmdList, TEXT("NiagaraUObjectPropertyReader"), sizeof(float), NumElements, EPixelFormat::PF_R32_FLOAT);
+						InstanceData_RT->PropertyData.Initialize(RHICmdList, TEXT("NiagaraUObjectPropertyReader"), sizeof(float), NumElements, EPixelFormat::PF_R32_UINT);
 					}
 				}
 				if (InstanceData_RT->PropertyData.NumBytes > 0)
