- Add a control to log warnings for Niagara buffer validation errors
- Allow used to control the truncation limit or none for validation errors

diff --git a/Engine/Plugins/FX/Niagara/Source/Niagara/Private/NiagaraDebugHud.cpp b/Engine/Plugins/FX/Niagara/Source/Niagara/Private/NiagaraDebugHud.cpp
--- a/Engine/Plugins/FX/Niagara/Source/Niagara/Private/NiagaraDebugHud.cpp
+++ b/Engine/Plugins/FX/Niagara/Source/Niagara/Private/NiagaraDebugHud.cpp
@@ -130,7 +130,9 @@ namespace NiagaraDebugLocal
 			}
 		),
 		MakeTuple(TEXT("ValidateSystemSimulationDataBuffers="), TEXT("Enable or disable validation on system data buffers"), [](FString Arg) {Settings.bValidateSystemSimulationDataBuffers = FCString::Atoi(*Arg) != 0; }),
-		MakeTuple(TEXT("bValidateParticleDataBuffers="), TEXT("Enable or disable validation on particle data buffers"), [](FString Arg) {Settings.bValidateParticleDataBuffers = FCString::Atoi(*Arg) != 0; }),
+		MakeTuple(TEXT("ValidateParticleDataBuffers="), TEXT("Enable or disable validation on particle data buffers"), [](FString Arg) {Settings.bValidateParticleDataBuffers = FCString::Atoi(*Arg) != 0; }),
+		MakeTuple(TEXT("ValidationLogErrors="), TEXT("When enabled validation errors will be logged as we as on screen"), [](FString Arg) {Settings.bValidationLogErrors = FCString::Atoi(*Arg) != 0; }),
+		MakeTuple(TEXT("ValidationAttributeDisplayTruncate="), TEXT("When > 0 limits the number of attributes we log"), [](FString Arg) {Settings.ValidationAttributeDisplayTruncate = FCString::Atoi(*Arg); }),
 
 		MakeTuple(TEXT("OverviewEnabled="), TEXT("Enable or disable the main overview display"), [](FString Arg) {Settings.bOverviewEnabled = FCString::Atoi(*Arg) != 0; }),
 		MakeTuple(TEXT("OverviewMode="), TEXT("Change the mode of the debug overivew"), [](FString Arg) {Settings.OverviewMode = (ENiagaraDebugHUDOverviewMode)FCString::Atoi(*Arg); }),
@@ -2603,7 +2605,7 @@ void FNiagaraDebugHud::DrawValidation(class FNiagaraWorldManager* WorldManager,
 			for (auto EmitterIt=ErrorInfo.ParticleVariablesWithErrors.CreateConstIterator(); EmitterIt; ++EmitterIt)
 			{
 				const TArray<FName>& EmitterVariables = EmitterIt.Value();
-				const int32 NumVariables = FMath::Min(EmitterVariables.Num(), 3);
+				const int32 NumVariables = Settings.ValidationAttributeDisplayTruncate > 0 ? FMath::Min(EmitterVariables.Num(), Settings.ValidationAttributeDisplayTruncate) : EmitterVariables.Num();
 				if (NumVariables > 0)
 				{
 					for (int32 iVariable = 0; iVariable < NumVariables; ++iVariable)
@@ -2623,6 +2625,7 @@ void FNiagaraDebugHud::DrawValidation(class FNiagaraWorldManager* WorldManager,
 						ErrorString.Append(TEXT(", ..."));
 					}
 					ErrorString.Append(TEXT("\n"));
+
 				}
 			}
 		}
@@ -2641,6 +2644,11 @@ void FNiagaraDebugHud::DrawValidation(class FNiagaraWorldManager* WorldManager,
 			TextLocation.Y += fAdvanceHeight;
 
 			DrawCanvas->DrawShadowedString(TextLocation.X, TextLocation.Y, ErrorString.ToString(), Font, Settings.MessageErrorTextColor);
+
+			if (Settings.bValidationLogErrors)
+			{
+				UE_LOG(LogNiagara, Warning, TEXT("Validation Errors - %s"), ErrorString.ToString());
+			}
 		}
 	}
 }
diff --git a/Engine/Plugins/FX/Niagara/Source/Niagara/Public/NiagaraDebuggerCommon.h b/Engine/Plugins/FX/Niagara/Source/Niagara/Public/NiagaraDebuggerCommon.h
--- a/Engine/Plugins/FX/Niagara/Source/Niagara/Public/NiagaraDebuggerCommon.h
+++ b/Engine/Plugins/FX/Niagara/Source/Niagara/Public/NiagaraDebuggerCommon.h
@@ -441,7 +441,7 @@ struct FNiagaraDebugHUDSettingsData
 	i.e. we will look for NaN or other invalidate data  inside it
 	Note: This will have an impact on performance.
 	*/
-	UPROPERTY(EditAnywhere, Category = "Debug General")
+	UPROPERTY(EditAnywhere, Category = "Debug Validation")
 	bool bValidateSystemSimulationDataBuffers = false;
 
 	/**
@@ -449,9 +449,24 @@ struct FNiagaraDebugHUDSettingsData
 	i.e. we will look for NaN or other invalidate data  inside it
 	Note: This will have an impact on performance.
 	*/
-	UPROPERTY(EditAnywhere, Category = "Debug General")
+	UPROPERTY(EditAnywhere, Category = "Debug Validation")
 	bool bValidateParticleDataBuffers = false;
 
+	/**
+	When enabled all validation errors will be sent to the log as warnings.
+	This can be useful to try and narrow down the exact source of an invalid value in the data buffers as often
+	they will cascade from the first frame where one is generated into other attributes in the subsequent frames.
+	*/
+	UPROPERTY(EditAnywhere, Category = "Debug Validation")
+	bool bValidationLogErrors = false;
+
+	/**
+	When > 0 this is the maximum number of attributes we will display that contain a NaN,
+	there could be more but the display will be truncated to this amount.  This is to reduce generating long strings.
+	*/
+	UPROPERTY(EditAnywhere, Category = "Debug Validation")
+	int32 ValidationAttributeDisplayTruncate = 3;
+
 	/** When enabled the overview display will be enabled. */
 	UPROPERTY(EditAnywhere, Category = "Debug Overview", meta = (DisplayName = "Debug Overview Enabled"))
 	bool bOverviewEnabled = false;
