- Reset SKM data interface if the source component is unregistered
- Reading from an unregistered skeletal mesh will have no ComponentSpaceTransforms and hit an ensure

diff --git a/Engine/Plugins/FX/Niagara/Source/Niagara/Private/NiagaraDataInterfaceSkeletalMesh.cpp b/Engine/Plugins/FX/Niagara/Source/Niagara/Private/NiagaraDataInterfaceSkeletalMesh.cpp
--- a/Engine/Plugins/FX/Niagara/Source/Niagara/Private/NiagaraDataInterfaceSkeletalMesh.cpp
+++ b/Engine/Plugins/FX/Niagara/Source/Niagara/Private/NiagaraDataInterfaceSkeletalMesh.cpp
@@ -2029,9 +2029,9 @@ bool FNDISkeletalMesh_InstanceData::Init(UNiagaraDataInterfaceSkeletalMesh* Inte
 
 bool FNDISkeletalMesh_InstanceData::ResetRequired(UNiagaraDataInterfaceSkeletalMesh* Interface, FNiagaraSystemInstance* SystemInstance) const
 {
-	// Reset if the scene component we've cached has been invalidated
+	// Reset if the scene component we've cached has been invalidated or unregistered
 	USceneComponent* Comp = SceneComponent.Get();
-	if (bComponentValid && !Comp)
+	if (bComponentValid && (!Comp || !Comp->IsRegistered()))
 	{
 		return true;
 	}
