- Cache GDF information and use when tick flushing
- This fixes issues like scrubbing backwards in sequencer with distance field collisions enabled

diff --git a/Engine/Plugins/FX/Niagara/Source/Niagara/Private/NiagaraDataInterfaceCollisionQuery.cpp b/Engine/Plugins/FX/Niagara/Source/Niagara/Private/NiagaraDataInterfaceCollisionQuery.cpp
--- a/Engine/Plugins/FX/Niagara/Source/Niagara/Private/NiagaraDataInterfaceCollisionQuery.cpp
+++ b/Engine/Plugins/FX/Niagara/Source/Niagara/Private/NiagaraDataInterfaceCollisionQuery.cpp
@@ -8,6 +8,7 @@
 #include "NiagaraDistanceFieldHelper.h"
 #include "NiagaraEmitter.h"
 #include "NiagaraGpuComputeDispatchInterface.h"
+#include "NiagaraGpuComputeDispatch.h"
 #include "NiagaraShaderParametersBuilder.h"
 #include "NiagaraSystem.h"
 #include "NiagaraTypes.h"
@@ -566,11 +567,7 @@ void UNiagaraDataInterfaceCollisionQuery::SetShaderParameters(const FNiagaraData
 	FGlobalDistanceFieldParameters2* ShaderGDFParameters = Context.GetParameterIncludedStruct<FGlobalDistanceFieldParameters2>();
 	if (Context.IsStructBound(ShaderGDFParameters))
 	{
-		TConstStridedView<FSceneView> SimulationSceneViews = Context.GetComputeDispatchInterface().GetSimulationSceneViews();
-		FNiagaraDistanceFieldHelper::SetGlobalDistanceFieldParameters(
-			UE::FXRenderingUtils::GetGlobalDistanceFieldParameterData(SimulationSceneViews),
-			*ShaderGDFParameters
-		);
+		FNiagaraDistanceFieldHelper::SetGlobalDistanceFieldParameters(Context.GetComputeDispatchInterface().GetGlobalDistanceFieldData(), *ShaderGDFParameters);
 	}
 }
 
diff --git a/Engine/Plugins/FX/Niagara/Source/Niagara/Private/NiagaraGpuComputeDispatch.cpp b/Engine/Plugins/FX/Niagara/Source/Niagara/Private/NiagaraGpuComputeDispatch.cpp
--- a/Engine/Plugins/FX/Niagara/Source/Niagara/Private/NiagaraGpuComputeDispatch.cpp
+++ b/Engine/Plugins/FX/Niagara/Source/Niagara/Private/NiagaraGpuComputeDispatch.cpp
@@ -404,7 +404,7 @@ void FNiagaraGpuComputeDispatch::ProcessPendingTicksFlush(FRHICommandListImmedia
 			// Make a temporary ViewInfo
 			//-TODO: We could gather some more information here perhaps?
 			FSceneViewFamilyContext ViewFamily(
-				FSceneViewFamily::ConstructionValues(nullptr, nullptr, FEngineShowFlags(ESFIM_Game))
+				FSceneViewFamily::ConstructionValues(nullptr, GetSceneInterface(), FEngineShowFlags(ESFIM_Game))
 				.SetTime(CachedViewInitOptions.GameTime)
 				.SetGammaCorrection(CachedViewInitOptions.GammaCorrection)
 			);
@@ -1822,6 +1822,7 @@ void FNiagaraGpuComputeDispatch::PreInitViews(FRDGBuilder& GraphBuilder, bool bA
 {
 	SCOPE_CYCLE_COUNTER(STAT_NiagaraGPUDispatchSetup_RT);
 	RDG_CSV_STAT_EXCLUSIVE_SCOPE(GraphBuilder, Niagara);
+	LLM_SCOPE(ELLMTag::Niagara);
 
 	bIsFirstViewFamily = CurrentFamily ? CurrentFamily == ViewFamilies[0] : true;
 	bIsLastViewFamily = CurrentFamily ? CurrentFamily == ViewFamilies.Last() : true;
@@ -1897,7 +1898,12 @@ void FNiagaraGpuComputeDispatch::PreInitViews(FRDGBuilder& GraphBuilder, bool bA
 	}
 #endif
 
-	LLM_SCOPE(ELLMTag::Niagara);
+	// If we during a scene render ensure remove the GDF cache data	to avoid holding onto GDF textures that could be reallocated
+	if (IsOutsideSceneRenderer() == false)
+	{
+		CachedGDFData = FCachedDistanceFieldData();
+	}
+	
 	TotalDispatchesThisFrame = 0;
 
 	// Add pass to begin the gpu profiler frame
@@ -1992,6 +1998,28 @@ void FNiagaraGpuComputeDispatch::PostRenderOpaque(FRDGBuilder& GraphBuilder, TCo
 
 	bAllowGPUParticleUpdate = bAllowGPUParticleUpdate && Views.Num() > 0 && Views[0].AllowGPUParticleUpdate();
 
+	// If we are during a scene render cache the GDF information for any potential tick flushing
+	// This fixes issues with distance field collisions (for example) where we may be scrubbing backwards and are force to flush to avoid TDRs / RDG pass limits
+	// The alternative to this would be to query the FSceneInterface but we currently do not have a clean way to do this as the GDF clipmaps are built transiently plus the ViewStates do not exist
+	if (IsOutsideSceneRenderer() == false)
+	{
+		if (const FGlobalDistanceFieldParameterData* SceneGDFData = UE::FXRenderingUtils::GetGlobalDistanceFieldParameterData(Views))
+		{
+			CachedGDFData.bCacheValid			= true;
+			CachedGDFData.PageAtlasTexture		= SceneGDFData->PageAtlasTexture;
+			CachedGDFData.CoverageAtlasTexture	= SceneGDFData->CoverageAtlasTexture;
+			CachedGDFData.PageObjectGridBuffer	= SceneGDFData->PageObjectGridBuffer;
+			CachedGDFData.PageTableTexture		= SceneGDFData->PageTableTexture;
+			CachedGDFData.MipTexture			= SceneGDFData->MipTexture;
+			CachedGDFData.GDFParameterData		= *SceneGDFData;
+		}
+		else
+		{
+			CachedGDFData						= FCachedDistanceFieldData();
+		}
+	}
+	TGuardValue<bool> GuardGDFValidForPass(CachedGDFData.bValidForPass, CachedGDFData.bCacheValid);
+
 	// Cache view information which will be used if we have to flush simulation commands in the future
 	if ( bAllowGPUParticleUpdate && Views.IsValidIndex(0) )
 	{
@@ -2174,6 +2202,11 @@ bool FNiagaraGpuComputeDispatch::AddSortedGPUSimulation(FNiagaraGPUSortInfo& Sor
 	}
 }
 
+const FGlobalDistanceFieldParameterData* FNiagaraGpuComputeDispatch::GetGlobalDistanceFieldData() const
+{
+	return CachedGDFData.bValidForPass ? &CachedGDFData.GDFParameterData : nullptr;
+}
+
 void FNiagaraGpuComputeDispatch::GenerateSortKeys(FRHICommandListImmediate& RHICmdList, int32 BatchId, int32 NumElementsInBatch, EGPUSortFlags Flags, FRHIUnorderedAccessView* KeysUAV, FRHIUnorderedAccessView* ValuesUAV)
 {
 	const bool bHighPrecision = EnumHasAnyFlags(Flags, EGPUSortFlags::HighPrecisionKeys);
diff --git a/Engine/Plugins/FX/Niagara/Source/Niagara/Private/NiagaraGpuComputeDispatch.h b/Engine/Plugins/FX/Niagara/Source/Niagara/Private/NiagaraGpuComputeDispatch.h
--- a/Engine/Plugins/FX/Niagara/Source/Niagara/Private/NiagaraGpuComputeDispatch.h
+++ b/Engine/Plugins/FX/Niagara/Source/Niagara/Private/NiagaraGpuComputeDispatch.h
@@ -8,6 +8,7 @@ the same VectorVM byte code / compute shader code
 #pragma once
 
 #include "FXSystem.h"
+#include "GlobalDistanceFieldParameters.h"
 #include "NiagaraCommon.h"
 #include "NiagaraComputeExecutionContext.h"
 #include "NiagaraGPUProfiler.h"
@@ -141,6 +142,7 @@ public:
 	virtual void ProcessDebugReadbacks(FRHICommandListImmediate& RHICmdList, bool bWaitCompletion) override;
 
 	virtual bool AddSortedGPUSimulation(FNiagaraGPUSortInfo& SortInfo) override;
+	virtual const FGlobalDistanceFieldParameterData* GetGlobalDistanceFieldData() const override;
 
 	void ResetDataInterfaces(FRDGBuilder& GraphBuilder, const FNiagaraGPUSystemTick& Tick, const FNiagaraComputeInstanceData& InstanceData) const;
 	void SetDataInterfaceParameters(FRDGBuilder& GraphBuilder, const FNiagaraGPUSystemTick& Tick, const FNiagaraComputeInstanceData& InstanceData, const FNiagaraShaderRef& ComputeShader, const FNiagaraSimStageData& SimStageData, const FNiagaraShaderScriptParametersMetadata& NiagaraShaderParametersMetadata, uint8* ParametersStructure) const;
@@ -274,9 +276,22 @@ private:
 		FMatrix		ViewRotationMatrix	= FMatrix::Identity;
 		FMatrix		ProjectionMatrix	= FMatrix::Identity;
 	};
-
 	FCachedViewInitOptions CachedViewInitOptions;
 
+	// Cached info for distance fields
+	struct FCachedDistanceFieldData
+	{
+		bool								bCacheValid = false;
+		bool								bValidForPass = false;
+		FTextureRHIRef						PageAtlasTexture;
+		FTextureRHIRef						CoverageAtlasTexture;
+		TRefCountPtr<FRDGPooledBuffer>		PageObjectGridBuffer;
+		FTextureRHIRef						PageTableTexture;
+		FTextureRHIRef						MipTexture;
+		FGlobalDistanceFieldParameterData	GDFParameterData;
+	};
+	FCachedDistanceFieldData			CachedGDFData;
+
 #if WITH_EDITOR
 	bool bRaisedWarningThisFrame = false;
 #endif
diff --git a/Engine/Plugins/FX/Niagara/Source/Niagara/Public/NiagaraGpuComputeDispatchInterface.h b/Engine/Plugins/FX/Niagara/Source/Niagara/Public/NiagaraGpuComputeDispatchInterface.h
--- a/Engine/Plugins/FX/Niagara/Source/Niagara/Public/NiagaraGpuComputeDispatchInterface.h
+++ b/Engine/Plugins/FX/Niagara/Source/Niagara/Public/NiagaraGpuComputeDispatchInterface.h
@@ -9,6 +9,7 @@
 #include "NiagaraGPUInstanceCountManager.h"
 #include "FXSystem.h"
 
+class FGlobalDistanceFieldParameterData;
 class FRDGBuilder;
 class FRDGExternalAccessQueue;
 class FSceneView;
@@ -80,6 +81,12 @@ public:
 	*/
 	TConstStridedView<FSceneView> GetSimulationSceneViews() const { return SimulationSceneViews; }
 
+	/**
+	* Get access to the global distance field data
+	* This will return nullptr if you attempt to access at an invalid point (i.e. before GDF is prepared or the GDF is not available)
+	*/
+	NIAGARA_API virtual const FGlobalDistanceFieldParameterData* GetGlobalDistanceFieldData() const = 0;
+
 	/** Get access to the instance count manager. */
 	FORCEINLINE FNiagaraGPUInstanceCountManager& GetGPUInstanceCounterManager() { check(IsInRenderingThread()); return GPUInstanceCounterManager; }
 	FORCEINLINE const FNiagaraGPUInstanceCountManager& GetGPUInstanceCounterManager() const { check(IsInRenderingThread()); return GPUInstanceCounterManager; }
