Fixed a bug where some inputs couldn't get focus any more

diff --git a/Engine/Plugins/FX/Niagara/Source/Niagara/Public/NiagaraTypes.h b/Engine/Plugins/FX/Niagara/Source/Niagara/Public/NiagaraTypes.h
--- a/Engine/Plugins/FX/Niagara/Source/Niagara/Public/NiagaraTypes.h
+++ b/Engine/Plugins/FX/Niagara/Source/Niagara/Public/NiagaraTypes.h
@@ -712,6 +712,9 @@ struct FNiagaraInputParameterCustomization
 	
 	UPROPERTY(EditAnywhere, Category="Customization", meta=(EditCondition="WidgetType == ENiagaraInputWidgetType::EnumStyle", EditConditionHides))
 	TArray<FNiagaraWidgetNamedIntegerInputValue> EnumStyleDropdownValues;
+
+	UPROPERTY()
+	bool bBroadcastValueChangesOnCommitOnly = false;
 	
 	/** If true then the input is also displayed and editable as a 3d widget in the viewport (vector and transform types only). */
 	//UPROPERTY(EditAnywhere, Category="Customization")
diff --git a/Engine/Plugins/FX/Niagara/Source/NiagaraEditor/Private/Customizations/NiagaraParameterBindingCustomization.cpp b/Engine/Plugins/FX/Niagara/Source/NiagaraEditor/Private/Customizations/NiagaraParameterBindingCustomization.cpp
--- a/Engine/Plugins/FX/Niagara/Source/NiagaraEditor/Private/Customizations/NiagaraParameterBindingCustomization.cpp
+++ b/Engine/Plugins/FX/Niagara/Source/NiagaraEditor/Private/Customizations/NiagaraParameterBindingCustomization.cpp
@@ -133,7 +133,9 @@ void FNiagaraParameterBindingCustomization::CustomizeHeader(TSharedRef<IProperty
 		TSharedPtr<INiagaraEditorTypeUtilities, ESPMode::ThreadSafe> TypeEditorUtilities = NiagaraEditorModule.GetTypeUtilities(DefaultParameter.GetType());
 		if (TypeEditorUtilities.IsValid())
 		{
-			DefaultValueParameterEditor = TypeEditorUtilities->CreateParameterEditor(ParameterBinding->ResolvedParameter.GetType());
+			FNiagaraInputParameterCustomization CustomizationOptions;
+			CustomizationOptions.bBroadcastValueChangesOnCommitOnly = true; // each broadcast usually forces a recompile, so we only want to do it on commits
+			DefaultValueParameterEditor = TypeEditorUtilities->CreateParameterEditor(ParameterBinding->ResolvedParameter.GetType(), EUnit::Unspecified, CustomizationOptions);
 			DefaultValueParameterEditor->UpdateInternalValueFromStruct(DefaultValueStructOnScope.ToSharedRef());
 			DefaultValueParameterEditor->SetOnValueChanged(SNiagaraParameterEditor::FOnValueChange::CreateSP(this, &FNiagaraParameterBindingCustomization::OnValueChanged));
 
diff --git a/Engine/Plugins/FX/Niagara/Source/NiagaraEditor/Private/TypeEditorUtilities/NiagaraFloatTypeEditorUtilities.cpp b/Engine/Plugins/FX/Niagara/Source/NiagaraEditor/Private/TypeEditorUtilities/NiagaraFloatTypeEditorUtilities.cpp
--- a/Engine/Plugins/FX/Niagara/Source/NiagaraEditor/Private/TypeEditorUtilities/NiagaraFloatTypeEditorUtilities.cpp
+++ b/Engine/Plugins/FX/Niagara/Source/NiagaraEditor/Private/TypeEditorUtilities/NiagaraFloatTypeEditorUtilities.cpp
@@ -160,7 +160,7 @@ public:
 				.OnEndSliderMovement(this, &SNiagaraFloatParameterEditor::EndSliderMovement)
 				.TypeInterface(GetTypeInterface<float>(DisplayUnit))
 				.AllowSpin(true)
-				.BroadcastValueChangesPerKey(!NiagaraSettings->GetUpdateStackValuesOnCommitOnly())
+				.BroadcastValueChangesPerKey(!NiagaraSettings->GetUpdateStackValuesOnCommitOnly() && !WidgetCustomization.bBroadcastValueChangesOnCommitOnly)
 				.LabelPadding(FMargin(3))
 				.Delta(WidgetCustomization.bHasStepWidth ? WidgetCustomization.StepWidth : 0)
 				.LabelLocation(SNumericEntryBox<float>::ELabelLocation::Inside)
diff --git a/Engine/Plugins/FX/Niagara/Source/NiagaraEditor/Private/TypeEditorUtilities/NiagaraIntegerTypeEditorUtilities.cpp b/Engine/Plugins/FX/Niagara/Source/NiagaraEditor/Private/TypeEditorUtilities/NiagaraIntegerTypeEditorUtilities.cpp
--- a/Engine/Plugins/FX/Niagara/Source/NiagaraEditor/Private/TypeEditorUtilities/NiagaraIntegerTypeEditorUtilities.cpp
+++ b/Engine/Plugins/FX/Niagara/Source/NiagaraEditor/Private/TypeEditorUtilities/NiagaraIntegerTypeEditorUtilities.cpp
@@ -166,7 +166,7 @@ public:
 				.OnEndSliderMovement(this, &SNiagaraIntegerParameterEditor::EndSliderMovement)
 				.TypeInterface(GetTypeInterface<int32>(DisplayUnit))
 				.AllowSpin(true)
-				.BroadcastValueChangesPerKey(!GetDefault<UNiagaraEditorSettings>()->GetUpdateStackValuesOnCommitOnly())
+				.BroadcastValueChangesPerKey(!GetDefault<UNiagaraEditorSettings>()->GetUpdateStackValuesOnCommitOnly() && !WidgetCustomization.bBroadcastValueChangesOnCommitOnly)
 				.LabelPadding(FMargin(3))
 				.LabelLocation(SNumericEntryBox<int32>::ELabelLocation::Inside)
 				.Label()
diff --git a/Engine/Plugins/FX/Niagara/Source/NiagaraEditor/Private/TypeEditorUtilities/NiagaraVectorTypeEditorUtilities.cpp b/Engine/Plugins/FX/Niagara/Source/NiagaraEditor/Private/TypeEditorUtilities/NiagaraVectorTypeEditorUtilities.cpp
--- a/Engine/Plugins/FX/Niagara/Source/NiagaraEditor/Private/TypeEditorUtilities/NiagaraVectorTypeEditorUtilities.cpp
+++ b/Engine/Plugins/FX/Niagara/Source/NiagaraEditor/Private/TypeEditorUtilities/NiagaraVectorTypeEditorUtilities.cpp
@@ -81,7 +81,7 @@ private:
 		.OnBeginSliderMovement(this, &SNiagaraVectorParameterEditorBase::BeginSliderMovement)
 		.OnEndSliderMovement(this, &SNiagaraVectorParameterEditorBase::EndSliderMovement)
 		.TypeInterface(GetTypeInterface<float>(DisplayUnit))
-		.BroadcastValueChangesPerKey(!GetDefault<UNiagaraEditorSettings>()->GetUpdateStackValuesOnCommitOnly())
+		.BroadcastValueChangesPerKey(!GetDefault<UNiagaraEditorSettings>()->GetUpdateStackValuesOnCommitOnly() && !WidgetCustomization.bBroadcastValueChangesOnCommitOnly)
 		.AllowSpin(true)
 		.LabelPadding( FMargin(3) )
 		.LabelLocation( SNumericEntryBox<float>::ELabelLocation::Inside )
