Fix an engine test failure caused by incomplete system data after object duplication

diff --git a/Engine/Plugins/FX/Niagara/Source/Niagara/Classes/NiagaraSystem.h b/Engine/Plugins/FX/Niagara/Source/Niagara/Classes/NiagaraSystem.h
--- a/Engine/Plugins/FX/Niagara/Source/Niagara/Classes/NiagaraSystem.h
+++ b/Engine/Plugins/FX/Niagara/Source/Niagara/Classes/NiagaraSystem.h
@@ -226,7 +226,8 @@ public:
 	//~ UObject interface
 	NIAGARA_API void PostInitProperties();
 	NIAGARA_API virtual void Serialize(FArchive& Ar) override;
-	NIAGARA_API virtual void PostLoad() override; 
+	NIAGARA_API virtual void PostLoad() override;
+	NIAGARA_API virtual void PostDuplicate(bool bDuplicateForPIE) override;
 	virtual bool CanBeClusterRoot() const override { return true; }
 	NIAGARA_API virtual void GetResourceSizeEx(FResourceSizeEx& CumulativeResourceSize) override;
 #if WITH_EDITORONLY_DATA
diff --git a/Engine/Plugins/FX/Niagara/Source/Niagara/Private/NiagaraSystem.cpp b/Engine/Plugins/FX/Niagara/Source/Niagara/Private/NiagaraSystem.cpp
--- a/Engine/Plugins/FX/Niagara/Source/Niagara/Private/NiagaraSystem.cpp
+++ b/Engine/Plugins/FX/Niagara/Source/Niagara/Private/NiagaraSystem.cpp
@@ -1040,6 +1040,16 @@ void UNiagaraSystem::PostLoad()
 	PrecachePSOs();
 }
 
+void UNiagaraSystem::PostDuplicate(bool bDuplicateForPIE)
+{
+	Super::PostDuplicate(bDuplicateForPIE);
+#if WITH_EDITOR
+	// The assumption here is that whatever is being duplicated was already postloaded - but duplication clears out the non-uproperty fields.
+	// Also, this only needs to be done in editor, otherwise postload does everything in one go.
+	EnsureFullyLoaded();
+#endif
+}
+
 void UNiagaraSystem::GetResourceSizeEx(FResourceSizeEx& CumulativeResourceSize)
 {
 	Super::GetResourceSizeEx(CumulativeResourceSize);
