- Add debug HUD overview sort mode

diff --git a/Engine/Plugins/FX/Niagara/Source/Niagara/Private/NiagaraDebugHud.cpp b/Engine/Plugins/FX/Niagara/Source/Niagara/Private/NiagaraDebugHud.cpp
--- a/Engine/Plugins/FX/Niagara/Source/Niagara/Private/NiagaraDebugHud.cpp
+++ b/Engine/Plugins/FX/Niagara/Source/Niagara/Private/NiagaraDebugHud.cpp
@@ -1996,7 +1996,12 @@ void FNiagaraDebugHud::DrawOverview(class FNiagaraWorldManager* WorldManager, FC
 	if (Settings.bOverviewEnabled && Settings.OverviewMode != ENiagaraDebugHUDOverviewMode::GpuComputePerformance)
 	{
 		TextLocation.Y += fAdvanceHeight;		
-		uint32 NumLines = 1;
+
+		// Filter out what we won't display on the HUD and calculate number of lines
+		using FSortedPerSystemDebugInfo = TPair<FName, const FSystemDebugInfo&>;
+		TArray<FSortedPerSystemDebugInfo> SortedPerSystemDebugInfo;
+		SortedPerSystemDebugInfo.Reserve(PerSystemDebugInfo.Num());
+
 		for (const auto& Pair : PerSystemDebugInfo)
 		{
 			const FSystemDebugInfo& SystemInfo = Pair.Value;
@@ -2006,9 +2011,25 @@ void FNiagaraDebugHud::DrawOverview(class FNiagaraWorldManager* WorldManager, FC
 				continue;
 			}
 
-			++NumLines;
+			SortedPerSystemDebugInfo.Emplace(Pair);
+		}
+
+		// Sort our display list, if enabled
+		switch (Settings.OverviewSortMode)
+		{
+			case ENiagaraDebugHUDDOverviewSort::Name:				Algo::Sort(SortedPerSystemDebugInfo, [](const FSortedPerSystemDebugInfo& Lhs, const FSortedPerSystemDebugInfo& Rhs) -> bool { return Lhs.Key.LexicalLess(Rhs.Key); });	break;
+			case ENiagaraDebugHUDDOverviewSort::NumberRegistered:	Algo::Sort(SortedPerSystemDebugInfo, [](const FSortedPerSystemDebugInfo& Lhs, const FSortedPerSystemDebugInfo& Rhs) -> bool { return Lhs.Value.TotalRegistered    > Rhs.Value.TotalRegistered; }); break;
+			case ENiagaraDebugHUDDOverviewSort::NumberActive:		Algo::Sort(SortedPerSystemDebugInfo, [](const FSortedPerSystemDebugInfo& Lhs, const FSortedPerSystemDebugInfo& Rhs) -> bool { return Lhs.Value.TotalActive        > Rhs.Value.TotalActive; }); break;
+			case ENiagaraDebugHUDDOverviewSort::NumberScalability:	Algo::Sort(SortedPerSystemDebugInfo, [](const FSortedPerSystemDebugInfo& Lhs, const FSortedPerSystemDebugInfo& Rhs) -> bool { return Lhs.Value.TotalScalability   > Rhs.Value.TotalScalability; });	break;
+			case ENiagaraDebugHUDDOverviewSort::MemoryUsage:		Algo::Sort(SortedPerSystemDebugInfo, [](const FSortedPerSystemDebugInfo& Lhs, const FSortedPerSystemDebugInfo& Rhs) -> bool { return Lhs.Value.TotalBytes         > Rhs.Value.TotalBytes; });	break;
+			case ENiagaraDebugHUDDOverviewSort::RecentlyVisibilty:	Algo::Sort(SortedPerSystemDebugInfo, [](const FSortedPerSystemDebugInfo& Lhs, const FSortedPerSystemDebugInfo& Rhs) -> bool { return Lhs.Value.FramesSinceVisible < Rhs.Value.FramesSinceVisible; });	break;
+
+			default:
+				break;
 		}
-		DrawCanvas->DrawTile(TextLocation.X - 1.0f, TextLocation.Y - 1.0f, TotalOverviewWidth + 1.0f, 2.0f + (float(NumLines) * fAdvanceHeight), 0.0f, 0.0f, 0.0f, 0.0f, BackgroundColor);
+
+		// Draw background + headers
+		DrawCanvas->DrawTile(TextLocation.X - 1.0f, TextLocation.Y - 1.0f, TotalOverviewWidth + 1.0f, 2.0f + (float(SortedPerSystemDebugInfo.Num() + 1) * fAdvanceHeight), 0.0f, 0.0f, 0.0f, 0.0f, BackgroundColor);
 
 		float SystemDataY = TextLocation.Y;
 		for (FOverviewColumn& Column : OverviewColumns)
@@ -2018,17 +2039,10 @@ void FNiagaraDebugHud::DrawOverview(class FNiagaraWorldManager* WorldManager, FC
 
 		TextLocation.Y += fAdvanceHeight;
 
-		//TODO: Pull out to an array to allow sorting?
-
-		for (auto it = PerSystemDebugInfo.CreateConstIterator(); it; ++it)
+		// Draw each row
+		for (const FSortedPerSystemDebugInfo& Pair : SortedPerSystemDebugInfo)
 		{
-			const auto& SystemInfo = it->Value;
-			if ((SystemInfo.FramesSinceVisible >= Settings.PerfHistoryFrames) ||
-				(Settings.bOverviewShowFilteredSystemOnly && !SystemInfo.bPassesSystemFilter))
-			{
-				continue;
-			}
-
+			const FSystemDebugInfo& SystemInfo = Pair.Value;
 			for (FOverviewColumn& Column : OverviewColumns)
 			{
 				Column.DrawSystemData(DrawCanvas, Font, TextLocation.X, TextLocation.Y, SystemInfo);
@@ -2056,9 +2070,9 @@ void FNiagaraDebugHud::DrawOverview(class FNiagaraWorldManager* WorldManager, FC
 			Graph.Draw(Settings.PerfGraphAxisColor, BackgroundColor);
 
 			//Add each line to the graph.
-			for (auto& Pair : PerSystemDebugInfo)
+			for (const auto& Pair : PerSystemDebugInfo)
 			{
-				FSystemDebugInfo& SysInfo = Pair.Value;
+				const FSystemDebugInfo& SysInfo = Pair.Value;
 				if (SysInfo.PerfStats && SysInfo.bPassesSystemFilter)
 				{
 					TArray<double> Frames;
diff --git a/Engine/Plugins/FX/Niagara/Source/Niagara/Public/NiagaraDebuggerCommon.h b/Engine/Plugins/FX/Niagara/Source/Niagara/Public/NiagaraDebuggerCommon.h
--- a/Engine/Plugins/FX/Niagara/Source/Niagara/Public/NiagaraDebuggerCommon.h
+++ b/Engine/Plugins/FX/Niagara/Source/Niagara/Public/NiagaraDebuggerCommon.h
@@ -406,6 +406,23 @@ enum class ENiagaraDebugHUDPerfSampleMode
 	PerInstanceAverage,	
 };
 
+UENUM()
+enum class ENiagaraDebugHUDDOverviewSort
+{
+	/** Lexical sort on system name */
+	Name,
+	/** Sort by the number currently registered. */
+	NumberRegistered,
+	/** Sort by the number currently active. */
+	NumberActive,
+	/** Sort by the number currently managed by scalability. */
+	NumberScalability,
+	/** Sort by approximate memory usage. */
+	MemoryUsage,
+	/** Sort by when the system was recently visible. */
+	RecentlyVisibilty,
+};
+
 /** Settings for Niagara debug HUD. Contained in it's own struct so that we can pass it whole in a message to the debugger client. */
 USTRUCT()
 struct FNiagaraDebugHUDSettingsData
@@ -478,6 +495,9 @@ struct FNiagaraDebugHUDSettingsData
 	UPROPERTY(EditAnywhere, Category = "Debug Overview", meta = (DisplayName = "Debug Overview Mode"))
 	ENiagaraDebugHUDOverviewMode OverviewMode = ENiagaraDebugHUDOverviewMode::Overview;
 
+	UPROPERTY(EditAnywhere, Category = "Debug Overview", meta = (DisplayName = "Debug Overview Sort Mode"))
+	ENiagaraDebugHUDDOverviewSort OverviewSortMode = ENiagaraDebugHUDDOverviewSort::Name;
+
 	/** Overview display font to use. */
 	UPROPERTY(EditAnywhere, Category = "Debug Overview", meta = (DisplayName = "Debug Overview Font", EditCondition = "bOverviewEnabled"))
 	ENiagaraDebugHudFont OverviewFont = ENiagaraDebugHudFont::Normal;
