- Sanitize the tick grouip to ensure no OOB access

#rnx

diff --git a/Engine/Plugins/FX/Niagara/Source/Niagara/Private/NiagaraSimulationTaskPriority.cpp b/Engine/Plugins/FX/Niagara/Source/Niagara/Private/NiagaraSimulationTaskPriority.cpp
--- a/Engine/Plugins/FX/Niagara/Source/Niagara/Private/NiagaraSimulationTaskPriority.cpp
+++ b/Engine/Plugins/FX/Niagara/Source/Niagara/Private/NiagaraSimulationTaskPriority.cpp
@@ -180,7 +180,8 @@ namespace NiagaraSimulationTaskPriority
 		else
 	#endif
 		{
-			Priority = GTickGroupPriorities[TickGroup];
+			const int32 TickGroupIndex = FMath::Clamp(int32(Priority), 0, (int32)UE_ARRAY_COUNT(GTickGroupPriorities) - 1);
+			Priority = GTickGroupPriorities[TickGroupIndex];
 		}
 		Priority = FMath::Clamp(Priority, 0, (int32)UE_ARRAY_COUNT(GTaskPriorities) - 1);
 		return GTaskPriorities[Priority].Get();
