Added shader.build.cs file.
Updated build.cs files to use new shader module.

diff --git a/Engine/Plugins/Experimental/GPULightmass/Source/GPULightmass/GPULightmass.Build.cs b/Engine/Plugins/Experimental/GPULightmass/Source/GPULightmass/GPULightmass.Build.cs
--- a/Engine/Plugins/Experimental/GPULightmass/Source/GPULightmass/GPULightmass.Build.cs
+++ b/Engine/Plugins/Experimental/GPULightmass/Source/GPULightmass/GPULightmass.Build.cs
@@ -6,11 +6,12 @@ namespace UnrealBuildTool.Rules
 	{
 		public GPULightmass(ReadOnlyTargetRules Target) : base(Target)
 		{
-			PublicIncludePaths.AddRange(
-				new string[] {
-					"../Shaders/Shared"
+			PublicIncludePathModuleNames.AddRange(
+				new string[]
+				{
+					"Shaders"
 				}
-				);
+			);
 
 			PrivateIncludePaths.AddRange(
 				new string[]
diff --git a/Engine/Plugins/FX/Niagara/Source/Niagara/Niagara.Build.cs b/Engine/Plugins/FX/Niagara/Source/Niagara/Niagara.Build.cs
--- a/Engine/Plugins/FX/Niagara/Source/Niagara/Niagara.Build.cs
+++ b/Engine/Plugins/FX/Niagara/Source/Niagara/Niagara.Build.cs
@@ -8,8 +8,6 @@ public class Niagara : ModuleRules
     public Niagara(ReadOnlyTargetRules Target) : base(Target)
 	{
 		NumIncludedBytesPerUnityCPPOverride = 491520; // best unity size found from using UBT ProfileUnitySizes mode
-
-		PrivateIncludePaths.Add("../../../../Shaders/Shared");
 		
 		PrivateDependencyModuleNames.AddRange(
             new string[] {
@@ -43,15 +41,12 @@ public class Niagara : ModuleRules
                 "RHI",
                 "VectorVM",
             }
-        );
+		);
 
-        PrivateIncludePathModuleNames.AddRange(
-                new string[] {
-                        "DerivedDataCache",
-                });
-
-        PrivateIncludePaths.AddRange(
-            new string[] {
+		PrivateIncludePathModuleNames.AddRange(
+			new string[] {
+				"DerivedDataCache",
+				"Shaders"
 			});
 
         if (Target.bBuildEditor == true)
diff --git a/Engine/Plugins/Runtime/GeometryCache/Source/GeometryCache/GeometryCache.Build.cs b/Engine/Plugins/Runtime/GeometryCache/Source/GeometryCache/GeometryCache.Build.cs
--- a/Engine/Plugins/Runtime/GeometryCache/Source/GeometryCache/GeometryCache.Build.cs
+++ b/Engine/Plugins/Runtime/GeometryCache/Source/GeometryCache/GeometryCache.Build.cs
@@ -6,8 +6,6 @@ public class GeometryCache : ModuleRules
 {
 	public GeometryCache(ReadOnlyTargetRules Target) : base(Target)
 	{
-		PrivateIncludePaths.Add("../../../../Shaders/Shared");
-
         PublicDependencyModuleNames.AddRange(
 			new string[] {
 				"Core",
@@ -24,8 +22,9 @@ public class GeometryCache : ModuleRules
 		);
 
         PublicIncludePathModuleNames.Add("TargetPlatform");
+		PrivateIncludePathModuleNames.Add("Shaders");
 
-        if (Target.bBuildEditor)
+		if (Target.bBuildEditor)
         {
             PublicIncludePathModuleNames.Add("GeometryCacheEd");
             DynamicallyLoadedModuleNames.Add("GeometryCacheEd");
diff --git a/Engine/Plugins/Runtime/HairStrands/Source/HairStrandsCore/HairStrandsCore.Build.cs b/Engine/Plugins/Runtime/HairStrands/Source/HairStrandsCore/HairStrandsCore.Build.cs
--- a/Engine/Plugins/Runtime/HairStrands/Source/HairStrandsCore/HairStrandsCore.Build.cs
+++ b/Engine/Plugins/Runtime/HairStrands/Source/HairStrandsCore/HairStrandsCore.Build.cs
@@ -6,12 +6,6 @@ namespace UnrealBuildTool.Rules
 	{
 		public HairStrandsCore(ReadOnlyTargetRules Target) : base(Target)
 		{
-			// Include Renderer/Private to have access to default resources
-			PrivateIncludePaths.AddRange(
-				new string[] {
-					EngineDirectory + "/Shaders/Shared",
-				});
-
 			PrivateDependencyModuleNames.AddRange(
 				new string[]
 				{
@@ -34,6 +28,12 @@ namespace UnrealBuildTool.Rules
 					"OptimusCore",
 				});
 
+			PrivateIncludePathModuleNames.AddRange(
+				new string[]
+				{
+					"Shaders",
+				});
+
 			if (Target.bBuildEditor == true)
 			{
 				PrivateDependencyModuleNames.AddRange(
diff --git a/Engine/Plugins/Runtime/ProceduralMeshComponent/Source/ProceduralMeshComponent/ProceduralMeshComponent.Build.cs b/Engine/Plugins/Runtime/ProceduralMeshComponent/Source/ProceduralMeshComponent/ProceduralMeshComponent.Build.cs
--- a/Engine/Plugins/Runtime/ProceduralMeshComponent/Source/ProceduralMeshComponent/ProceduralMeshComponent.Build.cs
+++ b/Engine/Plugins/Runtime/ProceduralMeshComponent/Source/ProceduralMeshComponent/ProceduralMeshComponent.Build.cs
@@ -6,8 +6,6 @@ namespace UnrealBuildTool.Rules
 	{
 		public ProceduralMeshComponent(ReadOnlyTargetRules Target) : base(Target)
 		{
-			PrivateIncludePaths.Add("../../../../Shaders/Shared");
-
 			PublicDependencyModuleNames.AddRange(
 				new string[]
 				{
@@ -21,6 +19,12 @@ namespace UnrealBuildTool.Rules
 					"PhysicsCore"
 				}
 				);
+
+			PrivateIncludePathModuleNames.AddRange(
+				new string[]
+				{
+					"Shaders",
+				});
 		}
 	}
 }
diff --git a/Engine/Shaders/Shaders.Build.cs b/Engine/Shaders/Shaders.Build.cs
new file mode 100644
--- /dev/null
+++ b/Engine/Shaders/Shaders.Build.cs
@@ -0,0 +1,19 @@
+// Copyright Epic Games, Inc. All Rights Reserved.
+
+using System.IO;
+
+namespace UnrealBuildTool.Rules
+{
+	public class Shaders : ModuleRules
+	{
+		public Shaders(ReadOnlyTargetRules Target) : base(Target)
+		{
+			// External so we don't build a DLL for this, it is header only.
+			Type = ModuleType.External;
+			ModuleIncludePathWarningLevel = WarningLevel.Off;
+
+			PublicIncludePaths.Add(Path.Combine(ModuleDirectory, "Public"));
+			PublicIncludePaths.Add(Path.Combine(ModuleDirectory, "Shared"));
+		}
+	}
+}
diff --git a/Engine/Source/Developer/Windows/ShaderFormatD3D/ShaderFormatD3D.Build.cs b/Engine/Source/Developer/Windows/ShaderFormatD3D/ShaderFormatD3D.Build.cs
--- a/Engine/Source/Developer/Windows/ShaderFormatD3D/ShaderFormatD3D.Build.cs
+++ b/Engine/Source/Developer/Windows/ShaderFormatD3D/ShaderFormatD3D.Build.cs
@@ -7,11 +7,14 @@ public class ShaderFormatD3D : ModuleRules
 {
 	public ShaderFormatD3D(ReadOnlyTargetRules Target) : base(Target)
 	{
-		PrivateIncludePathModuleNames.Add("TargetPlatform");
-		PrivateIncludePathModuleNames.Add("D3D11RHI");
-		PrivateIncludePathModuleNames.Add("D3D12RHI");
-
-		PrivateIncludePaths.Add("../Shaders/Shared");
+		PrivateIncludePathModuleNames.AddRange(
+			new string[]
+			{
+				"D3D11RHI",
+				"D3D12RHI",
+				"Shaders",
+				"TargetPlatform",
+			});
 
         PrivateDependencyModuleNames.AddRange(
 			new string[] {
diff --git a/Engine/Source/Programs/UnrealBuildTool/System/RulesCompiler.cs b/Engine/Source/Programs/UnrealBuildTool/System/RulesCompiler.cs
--- a/Engine/Source/Programs/UnrealBuildTool/System/RulesCompiler.cs
+++ b/Engine/Source/Programs/UnrealBuildTool/System/RulesCompiler.cs
@@ -142,6 +142,7 @@ namespace UnrealBuildTool
 					AddEngineModuleRulesWithContext(SourceDirectory, "Developer", DefaultModuleContext, UHTModuleType.EngineDeveloper, ModuleFileToContext);
 					AddEngineModuleRulesWithContext(SourceDirectory, "Editor", DefaultModuleContext, UHTModuleType.EngineEditor, ModuleFileToContext);
 					AddEngineModuleRulesWithContext(SourceDirectory, "ThirdParty", DefaultModuleContext, UHTModuleType.EngineThirdParty, ModuleFileToContext);
+					AddEngineModuleRulesWithContext(RootDirectory, "Shaders", DefaultModuleContext, UHTModuleType.EngineThirdParty, ModuleFileToContext);
 				}
 			}
 
diff --git a/Engine/Source/Runtime/D3D12RHI/D3D12RHI.Build.cs b/Engine/Source/Runtime/D3D12RHI/D3D12RHI.Build.cs
--- a/Engine/Source/Runtime/D3D12RHI/D3D12RHI.Build.cs
+++ b/Engine/Source/Runtime/D3D12RHI/D3D12RHI.Build.cs
@@ -8,8 +8,6 @@ public class D3D12RHI : ModuleRules
 {
 	public D3D12RHI(ReadOnlyTargetRules Target) : base(Target)
 	{
-		PrivateIncludePaths.Add("../Shaders/Shared");
-
 		PrivateDependencyModuleNames.AddAll(
 			"CoreUObject",
 			"Engine",
@@ -18,6 +16,10 @@ public class D3D12RHI : ModuleRules
 			"TraceLog"
 		);
 
+		PrivateIncludePathModuleNames.AddAll(
+			"Shaders"
+		);
+
 		PublicDependencyModuleNames.AddAll(
 			"Core",
 			"RHI"
diff --git a/Engine/Source/Runtime/Engine/Engine.Build.cs b/Engine/Source/Runtime/Engine/Engine.Build.cs
--- a/Engine/Source/Runtime/Engine/Engine.Build.cs
+++ b/Engine/Source/Runtime/Engine/Engine.Build.cs
@@ -9,8 +9,6 @@ public class Engine : ModuleRules
 	{
 		NumIncludedBytesPerUnityCPPOverride = 589824; // best unity size found from using UBT ProfileUnitySizes mode
 
-		PublicIncludePaths.Add("../Shaders/Shared");
-		
 		PrivatePCHHeaderFile = "Private/EnginePrivatePCH.h";
 
 		SharedPCHHeaderFile = "Public/EngineSharedPCH.h";
@@ -26,6 +24,7 @@ public class Engine : ModuleRules
 				"PacketHandler", 
 				"Renderer",
 				"RHI",
+				"Shaders"
 			}
 		);
 
diff --git a/Engine/Source/Runtime/GeometryFramework/GeometryFramework.Build.cs b/Engine/Source/Runtime/GeometryFramework/GeometryFramework.Build.cs
--- a/Engine/Source/Runtime/GeometryFramework/GeometryFramework.Build.cs
+++ b/Engine/Source/Runtime/GeometryFramework/GeometryFramework.Build.cs
@@ -8,8 +8,11 @@ public class GeometryFramework : ModuleRules
 	{
 		//PCHUsage = ModuleRules.PCHUsageMode.UseExplicitOrSharedPCHs;
 
-		// this is required for raytracing suppport?
-		PublicIncludePaths.Add("../Shaders/Shared");
+		PublicIncludePathModuleNames.AddRange(
+			new string[] {
+				"Shaders" // this is required for raytracing suppport?
+			}
+		);
 		
 		// These are required to register filters to the Level Editor Outliner
 		if (Target.Type == TargetType.Editor)
diff --git a/Engine/Source/Runtime/Landscape/Landscape.Build.cs b/Engine/Source/Runtime/Landscape/Landscape.Build.cs
--- a/Engine/Source/Runtime/Landscape/Landscape.Build.cs
+++ b/Engine/Source/Runtime/Landscape/Landscape.Build.cs
@@ -7,16 +7,11 @@ public class Landscape : ModuleRules
 {
 	public Landscape(ReadOnlyTargetRules Target) : base(Target)
 	{
-		PrivateIncludePaths.AddRange(
-			new string[] {
-				"../Shaders/Shared"
-            }
-		);
-
 		PrivateIncludePathModuleNames.AddRange(
 			new string[] {
 				"TargetPlatform",
 				"DerivedDataCache",
+				"Shaders"
 			}
 		);
 
diff --git a/Engine/Source/Runtime/RenderCore/RenderCore.Build.cs b/Engine/Source/Runtime/RenderCore/RenderCore.Build.cs
--- a/Engine/Source/Runtime/RenderCore/RenderCore.Build.cs
+++ b/Engine/Source/Runtime/RenderCore/RenderCore.Build.cs
@@ -9,11 +9,9 @@ public class RenderCore : ModuleRules
 {
 	public RenderCore(ReadOnlyTargetRules Target) : base(Target)
 	{
-		PrivateIncludePaths.Add("../Shaders/Shared");
-
 		PublicDependencyModuleNames.AddRange(new string[] { "RHI" });
 
-		PrivateIncludePathModuleNames.AddRange(new string[] { "TargetPlatform" });
+		PrivateIncludePathModuleNames.AddRange(new string[] { "Shaders", "TargetPlatform" });
 
 		// JSON is used for the asset info in the shader library and dumping out frames.
 		PrivateDependencyModuleNames.Add("Json");
diff --git a/Engine/Source/Runtime/Renderer/Renderer.Build.cs b/Engine/Source/Runtime/Renderer/Renderer.Build.cs
--- a/Engine/Source/Runtime/Renderer/Renderer.Build.cs
+++ b/Engine/Source/Runtime/Renderer/Renderer.Build.cs
@@ -8,11 +8,11 @@ public class Renderer : ModuleRules
 	{
 		PrivateIncludePaths.AddRange(
 			new string[] {
-				EngineDirectory + "/Shaders/Shared",
 				EngineDirectory + "/Shaders/Private", // For HaltonUtilities.ush
-				}
-			);
+			}
+		);
 
+		PrivateIncludePathModuleNames.Add("Shaders");
 		PublicDependencyModuleNames.Add("Core");
         PublicDependencyModuleNames.Add("Engine");
 
diff --git a/Engine/Source/Runtime/Windows/D3D11RHI/D3D11RHI.Build.cs b/Engine/Source/Runtime/Windows/D3D11RHI/D3D11RHI.Build.cs
--- a/Engine/Source/Runtime/Windows/D3D11RHI/D3D11RHI.Build.cs
+++ b/Engine/Source/Runtime/Windows/D3D11RHI/D3D11RHI.Build.cs
@@ -8,7 +8,7 @@ public class D3D11RHI : ModuleRules
 {
 	public D3D11RHI(ReadOnlyTargetRules Target) : base(Target)
 	{
-		PrivateIncludePaths.Add("../Shaders/Shared");
+		PrivateIncludePathModuleNames.AddRange(new string[] { "Shaders" });
 
 		PrivateDependencyModuleNames.AddAll(
 			"CoreUObject",
