Make sure to abort Niagara async compilation if edits to the system (like deleting emitters) occured before we were able to capture the precompile data

diff --git a/Engine/Plugins/FX/Niagara/Source/Niagara/Private/NiagaraAsyncCompile.cpp b/Engine/Plugins/FX/Niagara/Source/Niagara/Private/NiagaraAsyncCompile.cpp
--- a/Engine/Plugins/FX/Niagara/Source/Niagara/Private/NiagaraAsyncCompile.cpp
+++ b/Engine/Plugins/FX/Niagara/Source/Niagara/Private/NiagaraAsyncCompile.cpp
@@ -296,6 +296,15 @@ void FNiagaraAsyncCompileTask::PrecompileData()
 			ComputedPrecompileDuplicateData = SystemPrecompileDuplicateData->GetDependentRequest(EmitterIndex);
 		}
 	}
+
+	// check that we have valid PrecompileDuplicateData, if we don't we'll abort the task.  This is likely due to
+	// the system being edited, like emitters being deleted, within the window of the compile being requested and
+	// the precompile being issued
+	if (ComputedPrecompileDuplicateData == nullptr)
+	{
+		AbortTask();
+		return;
+	}
 }
 
 void FNiagaraAsyncCompileTask::StartCompileJob()
