- Fixed hierarchy sections not being rearrangeable.
This is because the inline edit text was taking up focus
- Fixed missing 'Delete' context action when right clicking sections

diff --git a/Engine/Plugins/FX/Niagara/Source/NiagaraEditor/Private/ViewModels/HierarchyEditor/NiagaraHierarchyCommands.cpp b/Engine/Plugins/FX/Niagara/Source/NiagaraEditor/Private/ViewModels/HierarchyEditor/NiagaraHierarchyCommands.cpp
--- a/Engine/Plugins/FX/Niagara/Source/NiagaraEditor/Private/ViewModels/HierarchyEditor/NiagaraHierarchyCommands.cpp
+++ b/Engine/Plugins/FX/Niagara/Source/NiagaraEditor/Private/ViewModels/HierarchyEditor/NiagaraHierarchyCommands.cpp
@@ -6,7 +6,7 @@
 
 void FNiagaraHierarchyEditorCommands::RegisterCommands()
 {
-	UI_COMMAND(Delete, "Delete", "Delete current selection", EUserInterfaceActionType::Button, FInputChord(EKeys::Delete))
+	UI_COMMAND(DeleteSection, "Delete", "Delete currently selected section", EUserInterfaceActionType::Button, FInputChord(EKeys::Delete))
 }
 
 #undef LOCTEXT_NAMESPACE
diff --git a/Engine/Plugins/FX/Niagara/Source/NiagaraEditor/Private/Widgets/SNiagaraHierarchy.cpp b/Engine/Plugins/FX/Niagara/Source/NiagaraEditor/Private/Widgets/SNiagaraHierarchy.cpp
--- a/Engine/Plugins/FX/Niagara/Source/NiagaraEditor/Private/Widgets/SNiagaraHierarchy.cpp
+++ b/Engine/Plugins/FX/Niagara/Source/NiagaraEditor/Private/Widgets/SNiagaraHierarchy.cpp
@@ -39,26 +39,19 @@ TSharedRef<SWidget> SummonContextMenu(TArray<TSharedPtr<FNiagaraHierarchyItemVie
 
 	if(ensure(BaseSection != nullptr))
 	{
-		if(Items.Num() == 1 && Items[0]->CanRename())
+		if(Items.Num() == 1)
 		{
 			BaseSection->AddMenuEntry(FGenericCommands::Get().Rename);
 		}
 
-		bool bCanDeleteAll = true;
-		for(TSharedPtr<FNiagaraHierarchyItemViewModelBase>& Item : Items)
+		if(Items.Num() == 1 && Items[0]->GetData()->IsA<UNiagaraHierarchySection>())
 		{
-			if(!Item->CanDelete())
-			{
-				bCanDeleteAll = false;
-				break;
-			}
-		}
-
-		if(bCanDeleteAll)
-		{
-			BaseSection->AddMenuEntry(FNiagaraHierarchyEditorCommands::Get().Delete);
+			BaseSection->AddMenuEntry(FNiagaraHierarchyEditorCommands::Get().DeleteSection);
 		}
-
+		
+		// the generic delete command handles hierarchy items (not sections)
+		BaseSection->AddMenuEntry(FGenericCommands::Get().Delete);
+		
 		if(Items.Num() == 1 && bFromHierarchy == false)
 		{
 			if(UNiagaraHierarchyItemBase* FoundItem = ViewModel->GetHierarchyRoot()->FindChildWithIdentity(Items[0]->GetData()->GetPersistentIdentity(), true))
@@ -166,7 +159,7 @@ void SNiagaraHierarchySection::Construct(const FArguments& InArgs, TSharedPtr<FN
 	
 	if(SectionViewModel != nullptr)
 	{
-		SectionViewModel->GetOnRequestRename().BindSP(this, &SNiagaraHierarchySection::EnterEditingMode);
+		SectionViewModel->GetOnRequestRename().BindSP(this, &SNiagaraHierarchySection::TryEnterEditingMode);
 
 		SDropTarget::FArguments LeftDropTargetArgs;
 		SDropTarget::FArguments OntoDropTargetArgs;
@@ -206,7 +199,7 @@ void SNiagaraHierarchySection::Construct(const FArguments& InArgs, TSharedPtr<FN
 				.Content()
 				[
 					SAssignNew(CheckBox, SCheckBox)
-					//.Visibility(EVisibility::HitTestInvisible)
+					.Visibility(EVisibility::HitTestInvisible)
 					.Style(FAppStyle::Get(), "DetailsView.SectionButton")
 					.OnCheckStateChanged(this, &SNiagaraHierarchySection::OnSectionCheckChanged)
 					.IsChecked(this, &SNiagaraHierarchySection::GetSectionCheckState)
@@ -222,6 +215,7 @@ void SNiagaraHierarchySection::Construct(const FArguments& InArgs, TSharedPtr<FN
 						+ SHorizontalBox::Slot()
 						[
 							SAssignNew(InlineEditableTextBlock, SInlineEditableTextBlock)
+							.Visibility(EVisibility::HitTestInvisible)
 							.Text(this, &SNiagaraHierarchySection::GetText)
 							.OnTextCommitted(this, &SNiagaraHierarchySection::OnRenameSection)
 							.OnVerifyTextChanged(this, &SNiagaraHierarchySection::OnVerifySectionRename)
@@ -267,7 +261,7 @@ SNiagaraHierarchySection::~SNiagaraHierarchySection()
 	SectionViewModel.Reset();
 }
 
-void SNiagaraHierarchySection::EnterEditingMode() const
+void SNiagaraHierarchySection::TryEnterEditingMode() const
 {
 	if(SectionViewModel.IsValid() && SectionViewModel->CanRename())
 	{
@@ -358,7 +352,7 @@ FReply SNiagaraHierarchySection::OnMouseButtonDown(const FGeometry& MyGeometry,
 			return FReply::Handled();
 		}
 		else if(MouseEvent.IsMouseButtonDown(EKeys::LeftMouseButton))
-		{
+		{			
 			OnSectionActivatedDelegate.ExecuteIfBound(SectionViewModel);
 			return FReply::Handled().DetectDrag(AsShared(), EKeys::LeftMouseButton).SetUserFocus(AsShared());	
 		}
@@ -382,7 +376,7 @@ FReply SNiagaraHierarchySection::OnMouseButtonUp(const FGeometry& MyGeometry, co
 				FPopupTransitionEffect(FPopupTransitionEffect::ContextMenu)
 			);
 				
-			OnSectionActivatedDelegate.ExecuteIfBound(SectionViewModel);
+			OnSectionActivatedDelegate.ExecuteIfBound(SectionViewModel);			
 			return FReply::Handled();
 		}
 	}
@@ -779,10 +773,20 @@ void SNiagaraHierarchy::Construct(const FArguments& InArgs, TObjectPtr<UNiagaraH
 		FExecuteAction::CreateSP(this, &SNiagaraHierarchy::RequestRenameSelectedItem),
 		FCanExecuteAction::CreateSP(this, &SNiagaraHierarchy::CanRequestRenameSelectedItem));
 
-	HierarchyViewModel->GetCommands()->MapAction(FNiagaraHierarchyEditorCommands::Get().Delete,
-		FExecuteAction::CreateSP(this, &SNiagaraHierarchy::DeleteSelectedItems),
-		FCanExecuteAction::CreateSP(this, &SNiagaraHierarchy::CanDeleteSelectedItems),
-		FIsActionChecked(), FIsActionButtonVisible::CreateSP(this, &SNiagaraHierarchy::CanDeleteSelectedItems));
+	HierarchyViewModel->GetCommands()->MapAction(FGenericCommands::Get().Delete,
+		FExecuteAction::CreateSP(this, &SNiagaraHierarchy::DeleteSelectedHierarchyItems),
+		FCanExecuteAction::CreateSP(this, &SNiagaraHierarchy::CanDeleteSelectedHierarchyItems),
+		FIsActionChecked(), FIsActionButtonVisible::CreateSP(this, &SNiagaraHierarchy::CanDeleteSelectedHierarchyItems));
+
+	// HierarchyViewModel->GetCommands()->MapAction(FNiagaraHierarchyEditorCommands::Get().RenameSection,
+	// 	FExecuteAction::CreateSP(this, &SNiagaraHierarchy::DeleteActiveSection),
+	// 	FCanExecuteAction::CreateSP(this, &SNiagaraHierarchy::CanDeleteActiveSection),
+	// 	FIsActionChecked(), FIsActionButtonVisible::CreateSP(this, &SNiagaraHierarchy::CanDeleteActiveSection));
+	
+	HierarchyViewModel->GetCommands()->MapAction(FNiagaraHierarchyEditorCommands::Get().DeleteSection,
+		FExecuteAction::CreateSP(this, &SNiagaraHierarchy::DeleteActiveSection),
+		FCanExecuteAction::CreateSP(this, &SNiagaraHierarchy::CanDeleteActiveSection),
+		FIsActionChecked(), FIsActionButtonVisible::CreateSP(this, &SNiagaraHierarchy::CanDeleteActiveSection));
 	
 	HierarchyViewModel->ForceFullRefresh();
 }
@@ -1257,13 +1261,13 @@ void SNiagaraHierarchy::DeleteItems(TArray<TSharedPtr<FNiagaraHierarchyItemViewM
 	HierarchyViewModel->DeleteItemsWithIdentities(DeletionIdentities);
 }
 
-void SNiagaraHierarchy::DeleteSelectedItems() const
+void SNiagaraHierarchy::DeleteSelectedHierarchyItems() const
 {	
 	TArray<TSharedPtr<FNiagaraHierarchyItemViewModelBase>> SelectedItems = HierarchyTreeView->GetSelectedItems();
 	DeleteItems(SelectedItems);
 }
 
-bool SNiagaraHierarchy::CanDeleteSelectedItems() const
+bool SNiagaraHierarchy::CanDeleteSelectedHierarchyItems() const
 {
 	TArray<TSharedPtr<FNiagaraHierarchyItemViewModelBase>> SelectedItems = HierarchyTreeView->GetSelectedItems();
 
@@ -1281,6 +1285,19 @@ bool SNiagaraHierarchy::CanDeleteSelectedItems() const
 	return false;
 }
 
+void SNiagaraHierarchy::DeleteActiveSection() const
+{
+	if(TSharedPtr<FNiagaraHierarchySectionViewModel> ActiveHierarchySectionViewModel = HierarchyViewModel->GetActiveHierarchySection())
+	{
+		DeleteItems({ActiveHierarchySectionViewModel});
+	}
+}
+
+bool SNiagaraHierarchy::CanDeleteActiveSection() const
+{
+	return HierarchyViewModel->GetActiveHierarchySection() != nullptr;
+}
+
 void SNiagaraHierarchy::OnItemAdded(TSharedPtr<FNiagaraHierarchyItemViewModelBase> AddedItem)
 {
 	// when a new item is created (opposed to dragged & dropped from source view, i.e. only categories so far)
diff --git a/Engine/Plugins/FX/Niagara/Source/NiagaraEditor/Public/ViewModels/HierarchyEditor/NiagaraHierarchyCommands.h b/Engine/Plugins/FX/Niagara/Source/NiagaraEditor/Public/ViewModels/HierarchyEditor/NiagaraHierarchyCommands.h
--- a/Engine/Plugins/FX/Niagara/Source/NiagaraEditor/Public/ViewModels/HierarchyEditor/NiagaraHierarchyCommands.h
+++ b/Engine/Plugins/FX/Niagara/Source/NiagaraEditor/Public/ViewModels/HierarchyEditor/NiagaraHierarchyCommands.h
@@ -20,6 +20,6 @@ public:
 	}
 
 	virtual void RegisterCommands() override;	
-
-	TSharedPtr<FUICommandInfo> Delete;
+	
+	TSharedPtr<FUICommandInfo> DeleteSection;
 };
diff --git a/Engine/Plugins/FX/Niagara/Source/NiagaraEditor/Public/Widgets/SNiagaraHierarchy.h b/Engine/Plugins/FX/Niagara/Source/NiagaraEditor/Public/Widgets/SNiagaraHierarchy.h
--- a/Engine/Plugins/FX/Niagara/Source/NiagaraEditor/Public/Widgets/SNiagaraHierarchy.h
+++ b/Engine/Plugins/FX/Niagara/Source/NiagaraEditor/Public/Widgets/SNiagaraHierarchy.h
@@ -37,7 +37,7 @@ public:
 	bool OnCanAcceptDrop(TSharedPtr<FDragDropOperation> DragDropOperation, EItemDropZone ItemDropZone) const;
 	FReply OnDroppedOn(const FGeometry&, const FDragDropEvent& DragDropEvent, EItemDropZone DropZone) const;
 
-	void EnterEditingMode() const;
+	void TryEnterEditingMode() const;
 
 	TSharedPtr<struct FNiagaraHierarchySectionViewModel> GetSectionViewModel();
 private:
@@ -133,8 +133,11 @@ private:
 	void ClearSourceItems() const;
 
 	void DeleteItems(TArray<TSharedPtr<FNiagaraHierarchyItemViewModelBase>> ItemsToDelete) const;
-	void DeleteSelectedItems() const;
-	bool CanDeleteSelectedItems() const;
+	void DeleteSelectedHierarchyItems() const;
+	bool CanDeleteSelectedHierarchyItems() const;
+
+	void DeleteActiveSection() const;
+	bool CanDeleteActiveSection() const;
 
 	void OnItemAdded(TSharedPtr<FNiagaraHierarchyItemViewModelBase> AddedItem);
 	void OnHierarchySectionActivated(TSharedPtr<FNiagaraHierarchySectionViewModel> Section);
