Fix landscape DI for LWC.  Supplies tile information for sampling the VT and cached heights as well as corrects the cached height acquisition for large worlds.

diff --git a/Engine/Plugins/FX/Niagara/Shaders/Private/NiagaraDataInterfaceLandscape.ush b/Engine/Plugins/FX/Niagara/Shaders/Private/NiagaraDataInterfaceLandscape.ush
--- a/Engine/Plugins/FX/Niagara/Shaders/Private/NiagaraDataInterfaceLandscape.ush
+++ b/Engine/Plugins/FX/Niagara/Shaders/Private/NiagaraDataInterfaceLandscape.ush
@@ -14,6 +14,7 @@
 Texture2D			{ParameterName}_BaseColorVirtualTexture;
 Texture2D<uint4>	{ParameterName}_BaseColorVirtualTexturePageTable;
 SamplerState		{ParameterName}_BaseColorVirtualTextureSampler;
+float3				{ParameterName}_BaseColorVirtualTextureLWCTile;
 float4x4			{ParameterName}_BaseColorVirtualTextureWorldToUvTransform;
 uint				{ParameterName}_BaseColorVirtualTextureSRGB;
 uint				{ParameterName}_BaseColorVirtualTextureEnabled;
@@ -23,6 +24,7 @@ uint4				{ParameterName}_BaseColorVirtualTextureUniforms;
 Texture2D			{ParameterName}_HeightVirtualTexture;
 Texture2D<uint4>	{ParameterName}_HeightVirtualTexturePageTable;
 SamplerState		{ParameterName}_HeightVirtualTextureSampler;
+float3				{ParameterName}_HeightVirtualTextureLWCTile;
 float4x4			{ParameterName}_HeightVirtualTextureWorldToUvTransform;
 uint				{ParameterName}_HeightVirtualTextureEnabled;
 uint4				{ParameterName}_HeightVirtualTexturePackedUniform0;
@@ -33,6 +35,7 @@ Texture2D			{ParameterName}_NormalVirtualTexture1;
 Texture2D<uint4>	{ParameterName}_NormalVirtualTexturePageTable;
 SamplerState		{ParameterName}_NormalVirtualTexture0Sampler;
 SamplerState		{ParameterName}_NormalVirtualTexture1Sampler;
+float3				{ParameterName}_NormalVirtualTextureLWCTile;
 float4x4			{ParameterName}_NormalVirtualTextureWorldToUvTransform;
 uint4				{ParameterName}_NormalVirtualTexturePackedUniform0;
 uint4				{ParameterName}_NormalVirtualTexturePackedUniform1;
@@ -43,6 +46,7 @@ uint				{ParameterName}_NormalVirtualTextureEnabled;
 Texture2D			{ParameterName}_CachedHeightTexture;
 SamplerState		{ParameterName}_CachedHeightTextureSampler;
 uint				{ParameterName}_CachedHeightTextureEnabled;
+float3				{ParameterName}_CachedHeightTextureLWCTile;
 float4x4			{ParameterName}_CachedHeightTextureWorldToUvTransform;
 float4x4			{ParameterName}_CachedHeightTextureUvToWorldTransform;
 float4				{ParameterName}_CachedHeightTextureUvScaleBias;
@@ -50,17 +54,30 @@ float2				{ParameterName}_CachedHeightTextureGridSize;
 SamplerState		{ParameterName}_PointClampedSampler;
 Texture2D<uint>		{ParameterName}_CachedPhysMatTexture;
 int2				{ParameterName}_CachedPhysMatTextureDimension;
+float3				{ParameterName}_SystemLWCTile;
+
+float2 GetCachedHeightTextureUv(float3 InWorldPos)
+{
+	FLWCVector4 LwcWorldPos = MakeLWCVector4({ParameterName}_SystemLWCTile, float4(InWorldPos, 1.0f));
+	FLWCInverseMatrix LwcWorldToUv = MakeLWCInverseMatrix({ParameterName}_CachedHeightTextureLWCTile, {ParameterName}_CachedHeightTextureWorldToUvTransform);
+
+	float2 Uv = LWCMultiply(LwcWorldPos, LwcWorldToUv).xy;
+	Uv = Uv * {ParameterName}_CachedHeightTextureUvScaleBias.xy + {ParameterName}_CachedHeightTextureUvScaleBias.zw;
+
+	return Uv;
+}
 
 void GetBaseColor_{ParameterName}(float3 InWorldPos, out float3 OutBaseColor, out bool OutIsValid)
 {
-	// todo[mg] convert world pos and transform origin to LWC pos
 	if ({ParameterName}_BaseColorVirtualTextureEnabled != 0)
 	{
 		float3 UnpackBaseColorOrigin = {ParameterName}_BaseColorVirtualTextureWorldToUvTransform[0].xyz;
 		float3 UnpackBaseColorU = {ParameterName}_BaseColorVirtualTextureWorldToUvTransform[1].xyz;
 		float3 UnpackBaseColorV = {ParameterName}_BaseColorVirtualTextureWorldToUvTransform[2].xyz;
 
-		float2 SampleUv = VirtualTextureWorldToUV(LWCPromote(InWorldPos), LWCPromote(UnpackBaseColorOrigin), UnpackBaseColorU, UnpackBaseColorV);
+		FLWCVector3 LwcWorldPos = MakeLWCVector3({ParameterName}_SystemLWCTile, InWorldPos);
+		FLWCVector3 LwcBaseColorOrigin = MakeLWCVector3({ParameterName}_BaseColorVirtualTextureLWCTile, UnpackBaseColorOrigin);
+		float2 SampleUv = VirtualTextureWorldToUV(LwcWorldPos, LwcBaseColorOrigin, UnpackBaseColorU, UnpackBaseColorV);
 
 		VTPageTableResult PageTable = TextureLoadVirtualPageTableLevel(
 			{ParameterName}_BaseColorVirtualTexturePageTable,
@@ -89,7 +106,6 @@ void GetBaseColor_{ParameterName}(float3 InWorldPos, out float3 OutBaseColor, ou
 
 void GetHeight_{ParameterName}(float3 InWorldPos, out float OutHeight, out bool OutIsValid)
 {
-	// todo[mg] convert world pos and transform origin to LWC pos
 	if ({ParameterName}_HeightVirtualTextureEnabled != 0)
 	{
 		float3 UnpackHeightOrigin = {ParameterName}_HeightVirtualTextureWorldToUvTransform[0].xyz;
@@ -97,7 +113,9 @@ void GetHeight_{ParameterName}(float3 InWorldPos, out float OutHeight, out bool
 		float3 UnpackHeightV = {ParameterName}_HeightVirtualTextureWorldToUvTransform[2].xyz;
 		float2 UnpackHeightScaleBias = {ParameterName}_HeightVirtualTextureWorldToUvTransform[3].xy;
 
-		float2 SampleUv = VirtualTextureWorldToUV(LWCPromote(InWorldPos), LWCPromote(UnpackHeightOrigin), UnpackHeightU, UnpackHeightV);
+		FLWCVector3 LwcWorldPos = MakeLWCVector3({ParameterName}_SystemLWCTile, InWorldPos);
+		FLWCVector3 LwcHeightOrigin = MakeLWCVector3({ParameterName}_HeightVirtualTextureLWCTile, UnpackHeightOrigin);
+		float2 SampleUv = VirtualTextureWorldToUV(LwcWorldPos, LwcHeightOrigin, UnpackHeightU, UnpackHeightV);
 
 		VTPageTableResult PageTable = TextureLoadVirtualPageTableLevel(
 			{ParameterName}_HeightVirtualTexturePageTable,
@@ -119,8 +137,8 @@ void GetHeight_{ParameterName}(float3 InWorldPos, out float OutHeight, out bool
 	}
 	else if ({ParameterName}_CachedHeightTextureEnabled != 0)
 	{
-		float2 Uv = mul(float4(InWorldPos, 1.0f), {ParameterName}_CachedHeightTextureWorldToUvTransform).xy;
-		Uv = Uv * {ParameterName}_CachedHeightTextureUvScaleBias.xy + {ParameterName}_CachedHeightTextureUvScaleBias.zw;
+		float2 Uv = GetCachedHeightTextureUv(InWorldPos);
+
 		OutHeight = {ParameterName}_CachedHeightTexture.SampleLevel({ParameterName}_CachedHeightTextureSampler, Uv, 0).x;
 		OutIsValid = true;
 	}
@@ -133,14 +151,15 @@ void GetHeight_{ParameterName}(float3 InWorldPos, out float OutHeight, out bool
 
 void GetWorldNormal_{ParameterName}(float3 InWorldPos, out float3 OutNormal, out bool OutIsValid)
 {
-	// todo[mg] convert world pos and transform origin to LWC pos
 	if ({ParameterName}_NormalVirtualTextureEnabled != 0)
 	{
 		float3 UnpackNormalOrigin = {ParameterName}_NormalVirtualTextureWorldToUvTransform[0].xyz;
 		float3 UnpackNormalU = {ParameterName}_NormalVirtualTextureWorldToUvTransform[1].xyz;
 		float3 UnpackNormalV = {ParameterName}_NormalVirtualTextureWorldToUvTransform[2].xyz;
 
-		float2 SampleUv = VirtualTextureWorldToUV(LWCPromote(InWorldPos), LWCPromote(UnpackNormalOrigin), UnpackNormalU, UnpackNormalV);
+		FLWCVector3 LwcWorldPos = MakeLWCVector3({ParameterName}_SystemLWCTile, InWorldPos);
+		FLWCVector3 LwcNormalOrigin = MakeLWCVector3({ParameterName}_NormalVirtualTextureLWCTile, UnpackNormalOrigin);
+		float2 SampleUv = VirtualTextureWorldToUV(LwcWorldPos, LwcNormalOrigin, UnpackNormalU, UnpackNormalV);
 
 		VTPageTableResult PageTable = TextureLoadVirtualPageTableLevel(
 			{ParameterName}_NormalVirtualTexturePageTable,
@@ -183,8 +202,7 @@ void GetWorldNormal_{ParameterName}(float3 InWorldPos, out float3 OutNormal, out
 	}
 	else if ({ParameterName}_CachedHeightTextureEnabled != 0)
 	{
-		float2 Uv = mul(float4(InWorldPos, 1.0f), {ParameterName}_CachedHeightTextureWorldToUvTransform).xy;
-		Uv = Uv * {ParameterName}_CachedHeightTextureUvScaleBias.xy + {ParameterName}_CachedHeightTextureUvScaleBias.zw;
+		float2 Uv = GetCachedHeightTextureUv(InWorldPos);
 
 		float Left = {ParameterName}_CachedHeightTexture.SampleLevel({ParameterName}_PointClampedSampler, Uv, 0, int2(-1,0)).x;
 		float Right = {ParameterName}_CachedHeightTexture.SampleLevel({ParameterName}_PointClampedSampler, Uv, 0, int2(1,0)).x;
@@ -212,9 +230,7 @@ void GetWorldNormal_{ParameterName}(float3 InWorldPos, out float3 OutNormal, out
 
 void GetPhysicalMaterialIndex_{ParameterName}(float3 InWorldPos, out int OutIndex, out bool OutIsValid)
 {
-	// todo[mg] convert world pos to LWC pos
-	float2 Uv = mul(float4(InWorldPos, 1.0f), {ParameterName}_CachedHeightTextureWorldToUvTransform).xy;
-	Uv = Uv * {ParameterName}_CachedHeightTextureUvScaleBias.xy + {ParameterName}_CachedHeightTextureUvScaleBias.zw;
+	float2 Uv = GetCachedHeightTextureUv(InWorldPos);
 
 	int3 SampleIndex = int3(Uv * {ParameterName}_CachedPhysMatTextureDimension, 0);
 	OutIndex = (int){ParameterName}_CachedPhysMatTexture.Load(SampleIndex).x;
diff --git a/Engine/Plugins/FX/Niagara/Source/Niagara/Private/NiagaraDataInterfaceLandscape.cpp b/Engine/Plugins/FX/Niagara/Source/Niagara/Private/NiagaraDataInterfaceLandscape.cpp
--- a/Engine/Plugins/FX/Niagara/Source/Niagara/Private/NiagaraDataInterfaceLandscape.cpp
+++ b/Engine/Plugins/FX/Niagara/Source/Niagara/Private/NiagaraDataInterfaceLandscape.cpp
@@ -51,6 +51,7 @@ namespace NiagaraDataInterfaceLandscape
 		SHADER_PARAMETER_TEXTURE(Texture2D,			BaseColorVirtualTexture)
 		SHADER_PARAMETER_TEXTURE(Texture2D<uint4>,	BaseColorVirtualTexturePageTable)
 		SHADER_PARAMETER_SAMPLER(SamplerState,		BaseColorVirtualTextureSampler)
+		SHADER_PARAMETER(FVector3f,					BaseColorVirtualTextureLWCTile)
 		SHADER_PARAMETER(FMatrix44f,				BaseColorVirtualTextureWorldToUvTransform)
 		SHADER_PARAMETER(uint32,					BaseColorVirtualTextureSRGB)
 		SHADER_PARAMETER(uint32,					BaseColorVirtualTextureEnabled)
@@ -60,6 +61,7 @@ namespace NiagaraDataInterfaceLandscape
 		SHADER_PARAMETER_TEXTURE(Texture2D,			HeightVirtualTexture)
 		SHADER_PARAMETER_TEXTURE(Texture2D<uint4>,	HeightVirtualTexturePageTable)
 		SHADER_PARAMETER_SAMPLER(SamplerState,		HeightVirtualTextureSampler)
+		SHADER_PARAMETER(FVector3f,					HeightVirtualTextureLWCTile)
 		SHADER_PARAMETER(FMatrix44f,				HeightVirtualTextureWorldToUvTransform)
 		SHADER_PARAMETER(uint32,					HeightVirtualTextureEnabled)
 		SHADER_PARAMETER(FUintVector4,				HeightVirtualTexturePackedUniform0)
@@ -70,6 +72,7 @@ namespace NiagaraDataInterfaceLandscape
 		SHADER_PARAMETER_TEXTURE(Texture2D<uint4>,	NormalVirtualTexturePageTable)
 		SHADER_PARAMETER_SAMPLER(SamplerState,		NormalVirtualTexture0Sampler)
 		SHADER_PARAMETER_SAMPLER(SamplerState,		NormalVirtualTexture1Sampler)
+		SHADER_PARAMETER(FVector3f,					NormalVirtualTextureLWCTile)
 		SHADER_PARAMETER(FMatrix44f,				NormalVirtualTextureWorldToUvTransform)
 		SHADER_PARAMETER(FUintVector4,				NormalVirtualTexturePackedUniform0)
 		SHADER_PARAMETER(FUintVector4,				NormalVirtualTexturePackedUniform1)
@@ -80,6 +83,7 @@ namespace NiagaraDataInterfaceLandscape
 		SHADER_PARAMETER_TEXTURE(Texture2D,			CachedHeightTexture)
 		SHADER_PARAMETER_SAMPLER(SamplerState,		CachedHeightTextureSampler)
 		SHADER_PARAMETER(uint32,					CachedHeightTextureEnabled)
+		SHADER_PARAMETER(FVector3f,					CachedHeightTextureLWCTile)
 		SHADER_PARAMETER(FMatrix44f,				CachedHeightTextureWorldToUvTransform)
 		SHADER_PARAMETER(FMatrix44f,				CachedHeightTextureUvToWorldTransform)
 		SHADER_PARAMETER(FVector4f,					CachedHeightTextureUvScaleBias)
@@ -87,6 +91,7 @@ namespace NiagaraDataInterfaceLandscape
 		SHADER_PARAMETER_SAMPLER(SamplerState,		PointClampedSampler)
 		SHADER_PARAMETER_TEXTURE(Texture2D<uint>,	CachedPhysMatTexture)
 		SHADER_PARAMETER(FIntPoint,					CachedPhysMatTextureDimension)
+		SHADER_PARAMETER(FVector3f,					SystemLWCTile)
 	END_SHADER_PARAMETER_STRUCT()
 }
 
@@ -178,6 +183,7 @@ public:
 	void Release();
 
 	TUniquePtr<FLandscapeTextureResource> LandscapeTextures;
+	FVector3f LandscapeLWCTile = FVector3f::ZeroVector;
 	FMatrix ActorToWorldTransform = FMatrix::Identity;
 	FMatrix WorldToActorTransform = FMatrix::Identity;
 	FVector4 UvScaleBias = FVector4(1.0f, 1.0f, 0.0f, 0.0f);
@@ -244,9 +250,10 @@ struct FNDILandscapeData_GameToRenderThread
 {
 	const FLandscapeTextureResource* TextureResources = nullptr;
 
+	FVector3f CachedHeightTextureLWCTile = FVector3f::ZeroVector;
 	FMatrix CachedHeightTextureWorldToUvTransform = FMatrix::Identity;
 	FMatrix CachedHeightTextureUvToWorldTransform = FMatrix::Identity;
-	FVector4 CachedHeightTextureUvScaleBias = FVector4(ForceInitToZero);
+	FVector4 CachedHeightTextureUvScaleBias = FVector4::Zero();
 	FVector2D CachedHeightTextureGridSize = FVector2D(1.0f, 1.0f);
 
 	bool BaseColorVirtualTextureSRGB = false;
@@ -376,8 +383,12 @@ struct FNDILandscapeData_RenderThread
 
 		ShaderParameters->BaseColorVirtualTexture = BaseColorVirtualLayer.TextureRef;
 		ShaderParameters->BaseColorVirtualTexturePageTable = BaseColorVirtualPage.PageTableRef;
+
+		FLargeWorldRenderPosition BaseColorVirtualTextureOrigin(BaseColorVirtualPage.WorldToUvParameters[0]);
+		ShaderParameters->BaseColorVirtualTextureLWCTile = BaseColorVirtualTextureOrigin.GetTile();
+
 		ShaderParameters->BaseColorVirtualTextureWorldToUvTransform = FMatrix44f(
-			FVector3f((FVector4f)BaseColorVirtualPage.WorldToUvParameters[0]),
+			BaseColorVirtualTextureOrigin.GetOffset(),
 			FVector3f((FVector4f)BaseColorVirtualPage.WorldToUvParameters[1]),
 			FVector3f((FVector4f)BaseColorVirtualPage.WorldToUvParameters[2]),
 			FVector3f(0, 0, 0));
@@ -413,8 +424,12 @@ struct FNDILandscapeData_RenderThread
 
 		ShaderParameters->HeightVirtualTexture = HeightVirtualLayer.TextureRef;
 		ShaderParameters->HeightVirtualTexturePageTable = HeightVirtualPage.PageTableRef;
+
+		FLargeWorldRenderPosition HeightVirtualTextureOrigin(HeightVirtualPage.WorldToUvParameters[0]);
+		ShaderParameters->HeightVirtualTextureLWCTile = HeightVirtualTextureOrigin.GetTile();
+
 		ShaderParameters->HeightVirtualTextureWorldToUvTransform = FMatrix44f(
-			FVector3f((FVector4f)HeightVirtualPage.WorldToUvParameters[0]),
+			HeightVirtualTextureOrigin.GetOffset(),
 			FVector3f((FVector4f)HeightVirtualPage.WorldToUvParameters[1]),
 			FVector3f((FVector4f)HeightVirtualPage.WorldToUvParameters[2]),
 			FVector3f((FVector4f)HeightVirtualPage.WorldToUvParameters[3]));
@@ -450,8 +465,12 @@ struct FNDILandscapeData_RenderThread
 		ShaderParameters->NormalVirtualTexture0 = NormalVirtualLayer0.TextureRef;
 		ShaderParameters->NormalVirtualTexture1 = NormalVirtualLayer1.TextureRef;
 		ShaderParameters->NormalVirtualTexturePageTable = NormalVirtualPage.PageTableRef;
+
+		FLargeWorldRenderPosition NormalVirtualTextureOrigin(NormalVirtualPage.WorldToUvParameters[0]);
+		ShaderParameters->NormalVirtualTextureLWCTile = NormalVirtualTextureOrigin.GetTile();
+
 		ShaderParameters->NormalVirtualTextureWorldToUvTransform = FMatrix44f(
-			FPlane4f((FVector4f)NormalVirtualPage.WorldToUvParameters[0]),
+			FPlane4f(NormalVirtualTextureOrigin.GetOffset()),
 			FPlane4f((FVector4f)NormalVirtualPage.WorldToUvParameters[1]),
 			FPlane4f((FVector4f)NormalVirtualPage.WorldToUvParameters[2]),
 			FPlane4f(0.0f, 0.0f, 0.0f, 1.0f));
@@ -491,6 +510,7 @@ struct FNDILandscapeData_RenderThread
 
 			if (HeightTexture || PhysMatTexture)
 			{
+				ShaderParameters->CachedHeightTextureLWCTile = LandscapeData.CachedHeightTextureLWCTile;
 				ShaderParameters->CachedHeightTextureWorldToUvTransform = (FMatrix44f)LandscapeData.CachedHeightTextureWorldToUvTransform;
 				ShaderParameters->CachedHeightTextureUvToWorldTransform = (FMatrix44f)LandscapeData.CachedHeightTextureUvToWorldTransform;
 				ShaderParameters->CachedHeightTextureUvScaleBias = (FVector4f)LandscapeData.CachedHeightTextureUvScaleBias;
@@ -529,9 +549,10 @@ struct FNDILandscapeData_RenderThread
 
 	static void SetCachedHeightTextureParameters_Defaults(NiagaraDataInterfaceLandscape::FShaderParameters* ShaderParameters)
 	{
-		const FVector4f DummyVector4(ForceInitToZero);
+		const FVector4f DummyVector4 = FVector4f::Zero();
 
 		ShaderParameters->CachedHeightTexture = GBlackTexture->TextureRHI;
+		ShaderParameters->CachedHeightTextureLWCTile = FVector3f::ZeroVector;
 		ShaderParameters->CachedHeightTextureWorldToUvTransform = FMatrix44f::Identity;
 		ShaderParameters->CachedHeightTextureUvToWorldTransform = FMatrix44f::Identity;
 		ShaderParameters->CachedHeightTextureUvScaleBias = DummyVector4;
@@ -864,7 +885,13 @@ void FNDI_Landscape_SharedResource::Initialize()
 			(0.5f - RegionVertexBase.X) * UvScale.X,
 			(0.5f - RegionVertexBase.Y) * UvScale.Y);
 
-		ActorToWorldTransform = ResourceKey.Source->GetTransform().ToMatrixWithScale();
+		FTransform LandscapeTransform = ResourceKey.Source->GetTransform();
+		FLargeWorldRenderPosition LandscapeTransformOrigin(LandscapeTransform.GetLocation());
+
+		LandscapeLWCTile = LandscapeTransformOrigin.GetTile();
+		LandscapeTransform.SetLocation(FVector(LandscapeTransformOrigin.GetOffset()));
+
+		ActorToWorldTransform = LandscapeTransform.ToMatrixWithScale();
 		WorldToActorTransform = ActorToWorldTransform.Inverse();
 		TextureWorldGridSize = FVector2D(ResourceKey.Source->GetTransform().GetScale3D());
 
@@ -990,7 +1017,11 @@ FNDI_Landscape_SharedResourceHandle FNDI_Landscape_GeneratedData::GetLandscapeDa
 	// we want to use the bounds of the system to figure out which cells of the landscape that we need to handle
 	const int32 MaxLandscapeRegionCount = LandscapeInfo->XYtoCollisionComponentMap.Num();
 
-	const FBox SystemWorldBounds = SystemInstance.GetLocalBounds().TransformBy(SystemInstance.GetWorldTransform());
+	const FVector LWCTileOffset = FVector(SystemInstance.GetLWCTile()) * FLargeWorldRenderScalar::GetTileSize();
+	FBox SystemWorldBounds = SystemInstance.GetLocalBounds().TransformBy(SystemInstance.GetWorldTransform());
+	SystemWorldBounds.Min += LWCTileOffset;
+	SystemWorldBounds.Max += LWCTileOffset;
+
 	const FTransform& LandscapeActorToWorld = Landscape->LandscapeActorToWorld();
 	const FVector SystemMinInLandscape = LandscapeActorToWorld.InverseTransformPosition(SystemWorldBounds.Min);
 	const FVector SystemMaxInLandscape = LandscapeActorToWorld.InverseTransformPosition(SystemWorldBounds.Max);
@@ -1004,17 +1035,19 @@ FNDI_Landscape_SharedResourceHandle FNDI_Landscape_GeneratedData::GetLandscapeDa
 	SystemBoundsInLandscape.Min = SystemBoundsInLandscape.Min.ComponentMax(FVector::ZeroVector);
 	SystemBoundsInLandscape.Max = SystemBoundsInLandscape.Max.ComponentMax(FVector::ZeroVector);
 
+	// next rescale based on the quad size
+	const double QuadSizeScaleFactor = 1.0 / ((double)Landscape->ComponentSizeQuads);
+	SystemBoundsInLandscape.Min *= QuadSizeScaleFactor;
+	SystemBoundsInLandscape.Max *= QuadSizeScaleFactor;
+
 	// truncate to integers
 	const FVector MaxIntValue = FVector(double(TNumericLimits<int32>::Max()));
-	SystemBoundsInLandscape.Min = SystemBoundsInLandscape.Min.ComponentMin(MaxIntValue); // LWC Precision Loss 
+	SystemBoundsInLandscape.Min = SystemBoundsInLandscape.Min.ComponentMin(MaxIntValue);
 	SystemBoundsInLandscape.Max = SystemBoundsInLandscape.Max.ComponentMin(MaxIntValue);
 
-	// scale it by the quad size
-	const double QuadSizeScaleFactor = 1.0 / ((double)Landscape->ComponentSizeQuads);
-
 	const FIntRect SystemRect = FIntRect(
 		FIntPoint(int32(SystemBoundsInLandscape.Min.X), int32(SystemBoundsInLandscape.Min.Y)),
-		FIntPoint(int32(SystemBoundsInLandscape.Max.X), int32(SystemBoundsInLandscape.Max.Y))).Scale(QuadSizeScaleFactor);
+		FIntPoint(int32(SystemBoundsInLandscape.Max.X), int32(SystemBoundsInLandscape.Max.Y)));
 
 	// for obnoxiously large system bounds we need to guard against potential overflow on the number of cells
 	const int32 MaxSystemWidth = FMath::Clamp<int32>(SystemRect.Max.X - SystemRect.Min.X, 0, MaxLandscapeRegionCount);
@@ -1246,6 +1279,7 @@ void UNiagaraDataInterfaceLandscape::ProvidePerInstanceDataForRenderThread(void*
 		TargetData->TextureResources = SourceResource.LandscapeTextures.Get();
 		
 		TargetData->CachedHeightTextureUvScaleBias = SourceResource.UvScaleBias;
+		TargetData->CachedHeightTextureLWCTile = SourceResource.LandscapeLWCTile;
 		TargetData->CachedHeightTextureWorldToUvTransform = SourceResource.WorldToActorTransform;
 		TargetData->CachedHeightTextureUvToWorldTransform = SourceResource.ActorToWorldTransform;
 		TargetData->CachedHeightTextureGridSize = SourceResource.TextureWorldGridSize;
@@ -1350,6 +1384,9 @@ void UNiagaraDataInterfaceLandscape::SetShaderParameters(const FNiagaraDataInter
 	{
 		FNDILandscapeData_RenderThread::SetCachedHeightTextureParameters_Defaults(ShaderParameters);
 	}
+
+	// System tile for LWC
+	ShaderParameters->SystemLWCTile = Context.GetSystemLWCTile();
 }
 
 bool UNiagaraDataInterfaceLandscape::InitPerInstanceData(void* PerInstanceData, FNiagaraSystemInstance* SystemInstance)
