Adapted BP Node UI for Niagara
It differs as it uses a HorizontalBox instead of an Overlay for input pins, title, output pins.

The problem with the overlay is that pin default values can grow 'into' the title

To ensure the middle part stays centered and the pin default values have enough space, left & right pin boxes have their sizes synced to the bigger of the two.

diff --git a/Engine/Plugins/FX/Niagara/Source/NiagaraEditor/Private/NiagaraEditorCommon.cpp b/Engine/Plugins/FX/Niagara/Source/NiagaraEditor/Private/NiagaraEditorCommon.cpp
--- a/Engine/Plugins/FX/Niagara/Source/NiagaraEditor/Private/NiagaraEditorCommon.cpp
+++ b/Engine/Plugins/FX/Niagara/Source/NiagaraEditor/Private/NiagaraEditorCommon.cpp
@@ -142,6 +142,8 @@ void FNiagaraOpInfo::Init()
 		FNiagaraOpInfo* Op = &OpInfos[Idx];
 		Op->Category = NumericCategory;
 		Op->FriendlyName = NSLOCTEXT("NiagaraOpInfo", "Add Name", "Add");
+		Op->CompactName = FText::FromString(TEXT("+"));
+		Op->CompactNameFontSizeOverride = 25.f;
 		Op->Description = NSLOCTEXT("NiagaraOpInfo", "Add Desc", "Result = A + B");
 		Op->Keywords = FText::FromString(TEXT("+"));
 		Op->Inputs.Add(FNiagaraOpInOutInfo(A, NumericType, AText, AText, DefaultStr_Zero));
@@ -193,6 +195,8 @@ void FNiagaraOpInfo::Init()
 		Op = &OpInfos[Idx];
 		Op->Category = NumericCategory;
 		Op->FriendlyName = NSLOCTEXT("NiagaraOpInfo", "Subtract Name", "Subtract");
+		Op->CompactName = FText::FromString(TEXT("-"));
+		Op->CompactNameFontSizeOverride = 28.f;
 		Op->Description = NSLOCTEXT("NiagaraOpInfo", "Subtract Desc", "Result = A - B");
 		Op->Keywords = FText::FromString(TEXT("-"));
 		Op->Inputs.Add(FNiagaraOpInOutInfo(A, NumericType, AText, AText, DefaultStr_Zero));
@@ -248,6 +252,8 @@ void FNiagaraOpInfo::Init()
 		Op = &OpInfos[Idx];
 		Op->Category = NumericCategory;
 		Op->FriendlyName = NSLOCTEXT("NiagaraOpInfo", "Multiply Name", "Multiply");
+		Op->CompactName = FText::FromString(TEXT("\xD7"));
+		Op->CompactNameFontSizeOverride = 25.f;
 		Op->Description = NSLOCTEXT("NiagaraOpInfo", "Multiply Desc", "Result = A * B");
 		Op->Keywords = FText::FromString(TEXT("*"));
 		Op->Inputs.Add(FNiagaraOpInOutInfo(A, NumericType, AText, AText, DefaultStr_One));
@@ -280,6 +286,8 @@ void FNiagaraOpInfo::Init()
 		Op = &OpInfos[Idx];
 		Op->Category = NumericCategory;
 		Op->FriendlyName = NSLOCTEXT("NiagaraOpInfo", "Divide Name", "Divide");
+		Op->CompactName = FText::FromString(TEXT("\xF7"));
+		Op->CompactNameFontSizeOverride = 25.f;
 		Op->Description = NSLOCTEXT("NiagaraOpInfo", "Divide Desc", "Result = A / B");
 		Op->Keywords = FText::FromString(TEXT("/"));
 		Op->Inputs.Add(FNiagaraOpInOutInfo(A, NumericType, AText, AText, DefaultStr_One));
@@ -312,6 +320,11 @@ void FNiagaraOpInfo::Init()
 		Op = &OpInfos[Idx];
 		Op->Category = NumericCategory;
 		Op->FriendlyName = NSLOCTEXT("NiagaraOpInfo", "MultiplyAdd Name", "MultiplyAdd");
+		FString CompactNameString = TEXT("(A");
+		CompactNameString.Append(TEXT("\xD7"));
+		CompactNameString.Append(TEXT("B)+C"));
+		Op->CompactName = FText::FromString(CompactNameString);
+		Op->bShowPinNamesInCompactMode = true;
 		Op->Description = NSLOCTEXT("NiagaraOpInfo", "MultiplyAdd Desc", "Result = (A * B) + C");
 		Op->Inputs.Add(FNiagaraOpInOutInfo(A, NumericType, AText, AText, DefaultStr_One));
 		Op->Inputs.Add(FNiagaraOpInOutInfo(B, NumericType, BText, BText, DefaultStr_One));
@@ -348,6 +361,8 @@ void FNiagaraOpInfo::Init()
 		Op = &OpInfos[Idx];
 		Op->Category = NumericCategory;
 		Op->FriendlyName = NSLOCTEXT("NiagaraOpInfo", "Lerp Name", "Lerp");
+		Op->CompactName = FText::FromString(TEXT("Lerp"));
+		Op->bShowPinNamesInCompactMode = true;
 		Op->Description = NSLOCTEXT("NiagaraOpInfo", "Lerp Desc", "Result = (A * (1 - C)) + (B * C)");
 		Op->Keywords = FText::FromString(TEXT("lerp"));
 		Op->Inputs.Add(FNiagaraOpInOutInfo(A, NumericType, AText, AText, DefaultStr_Zero));
@@ -385,6 +400,7 @@ void FNiagaraOpInfo::Init()
 		Op = &OpInfos[Idx];
 		Op->Category = NumericCategory;
 		Op->FriendlyName = NSLOCTEXT("NiagaraOpInfo", "Reciprocal Fast Name", "Reciprocal Fast");
+		Op->CompactName = FText::FromString("Rcp Fast");
 		Op->Description = NSLOCTEXT("NiagaraOpInfo", "Reciprocal Fast Desc", "12-bits of accuracy, but faster. Result = 1 / A using Newton/Raphson approximation.");
 		Op->Inputs.Add(FNiagaraOpInOutInfo(A, NumericType, AText, AText, DefaultStr_One));
 		Op->Outputs.Add(FNiagaraOpInOutInfo(Result, NumericType, ResultText, ResultText, DefaultStr_One, TEXT("rcp({0})")));
@@ -396,6 +412,7 @@ void FNiagaraOpInfo::Init()
 		Op = &OpInfos[Idx];
 		Op->Category = NumericCategory;
 		Op->FriendlyName = NSLOCTEXT("NiagaraOpInfo", "Reciprocal Name", "Reciprocal");
+		Op->CompactName = FText::FromString("Rcp");
 		Op->Description = NSLOCTEXT("NiagaraOpInfo", "Reciprocal Desc", "More accurate than Reciprocal Fast. Result = 1 / A");
 		Op->Inputs.Add(FNiagaraOpInOutInfo(A, NumericType, AText, AText, DefaultStr_One));
 		Op->Outputs.Add(FNiagaraOpInOutInfo(Result, NumericType, ResultText, ResultText, DefaultStr_One, TEXT("Reciprocal({0})")));
@@ -407,6 +424,7 @@ void FNiagaraOpInfo::Init()
 		Op = &OpInfos[Idx];
 		Op->Category = NumericCategory;
 		Op->FriendlyName = NSLOCTEXT("NiagaraOpInfo", "Reciprocal Sqrt Name", "Reciprocal Sqrt");
+		Op->CompactName = FText::FromString("Rcp Sqrt");
 		Op->Description = NSLOCTEXT("NiagaraOpInfo", "Reciprocal Sqrt Desc", "Result = 1 / sqrt(A)");
 		Op->Inputs.Add(FNiagaraOpInOutInfo(A, NumericType, AText, AText, DefaultStr_One));
 		Op->Outputs.Add(FNiagaraOpInOutInfo(Result, NumericType, ResultText, ResultText, DefaultStr_One, TEXT("rsqrt({0})")));
@@ -418,6 +436,7 @@ void FNiagaraOpInfo::Init()
 		Op = &OpInfos[Idx];
 		Op->Category = NumericCategory;
 		Op->FriendlyName = NSLOCTEXT("NiagaraOpInfo", "Sqrt Name", "Sqrt");
+		Op->CompactName = FText::FromString("Sqrt");
 		Op->Description = NSLOCTEXT("NiagaraOpInfo", "Sqrt Desc", "Result = sqrt(A)");
 		Op->Keywords = FText::FromString(TEXT("sqrt"));
 		Op->Inputs.Add(FNiagaraOpInOutInfo(A, NumericType, AText, AText, DefaultStr_One));
@@ -430,6 +449,8 @@ void FNiagaraOpInfo::Init()
 		Op = &OpInfos[Idx];
 		Op->Category = NumericCategory;
 		Op->FriendlyName = NSLOCTEXT("NiagaraOpInfo", "One Minus Name", "One Minus");
+		Op->CompactName = FText::FromString(TEXT("1-A"));
+		Op->bShowPinNamesInCompactMode = true;
 		Op->Description = NSLOCTEXT("NiagaraOpInfo", "One Minus Desc", "Result = 1 - A");
 		Op->Keywords = FText::FromString(TEXT("1-x"));
 		Op->Inputs.Add(FNiagaraOpInOutInfo(A, NumericType, AText, AText, DefaultStr_One));
@@ -450,6 +471,8 @@ void FNiagaraOpInfo::Init()
 		Op = &OpInfos[Idx];
 		Op->Category = NumericCategory;
 		Op->FriendlyName = NSLOCTEXT("NiagaraOpInfo", "Negate Name", "Negate");
+		Op->CompactName = FText::FromString("-A");
+		Op->bShowPinNamesInCompactMode = true;
 		Op->Description = NSLOCTEXT("NiagaraOpInfo", "Negate Desc", "Result = -A");
 		Op->Inputs.Add(FNiagaraOpInOutInfo(A, NumericType, AText, AText, DefaultStr_One));
 		Op->Outputs.Add(FNiagaraOpInOutInfo(Result, NumericType, ResultText, ResultText, DefaultStr_One, TEXT("-({0})")));
@@ -469,6 +492,7 @@ void FNiagaraOpInfo::Init()
 		Op = &OpInfos[Idx];
 		Op->Category = NumericCategory;
 		Op->FriendlyName = NSLOCTEXT("NiagaraOpInfo", "Abs Name", "Abs");
+		Op->CompactName = FText::FromString("Abs");
 		Op->Description = NSLOCTEXT("NiagaraOpInfo", "Abs Desc", "Result = abs(A)");
 		Op->Inputs.Add(FNiagaraOpInOutInfo(A, NumericType, AText, AText, DefaultStr_One));
 		Op->Outputs.Add(FNiagaraOpInOutInfo(Result, NumericType, ResultText, ResultText, DefaultStr_One, TEXT("abs({0})")));
@@ -480,6 +504,7 @@ void FNiagaraOpInfo::Init()
 		Op = &OpInfos[Idx];
 		Op->Category = NumericCategory;
 		Op->FriendlyName = NSLOCTEXT("NiagaraOpInfo", "Exp Name", "Exp");
+		Op->CompactName = FText::FromString("Exp");
 		Op->Description = NSLOCTEXT("NiagaraOpInfo", "Exp Desc", "Result = exp(A)");
 		Op->Inputs.Add(FNiagaraOpInOutInfo(A, NumericType, AText, AText, DefaultStr_One));
 		Op->Outputs.Add(FNiagaraOpInOutInfo(Result, NumericType, ResultText, ResultText, DefaultStr_One, TEXT("exp({0})")));
@@ -491,6 +516,7 @@ void FNiagaraOpInfo::Init()
 		Op = &OpInfos[Idx];
 		Op->Category = NumericCategory;
 		Op->FriendlyName = NSLOCTEXT("NiagaraOpInfo", "Exp2 Name", "Exp2");
+		Op->CompactName = FText::FromString("Exp2");
 		Op->Description = NSLOCTEXT("NiagaraOpInfo", "Exp2 Desc", "Result = exp2(A)");
 		Op->Inputs.Add(FNiagaraOpInOutInfo(A, NumericType, AText, AText, DefaultStr_One));
 		Op->Outputs.Add(FNiagaraOpInOutInfo(Result, NumericType, ResultText, ResultText, DefaultStr_One, TEXT("exp2({0})")));
@@ -502,6 +528,7 @@ void FNiagaraOpInfo::Init()
 		Op = &OpInfos[Idx];
 		Op->Category = NumericCategory;
 		Op->FriendlyName = NSLOCTEXT("NiagaraOpInfo", "Log Name", "Log");
+		Op->CompactName = FText::FromString("Log");
 		Op->Description = NSLOCTEXT("NiagaraOpInfo", "Log Desc", "Result = log(A)");
 		Op->Inputs.Add(FNiagaraOpInOutInfo(A, NumericType, AText, AText, DefaultStr_One));
 		Op->Outputs.Add(FNiagaraOpInOutInfo(Result, NumericType, ResultText, ResultText, DefaultStr_One, TEXT("log({0})")));
@@ -513,6 +540,7 @@ void FNiagaraOpInfo::Init()
 		Op = &OpInfos[Idx];
 		Op->Category = NumericCategory;
 		Op->FriendlyName = NSLOCTEXT("NiagaraOpInfo", "Log2 Name", "Log2");
+		Op->CompactName = FText::FromString("Log2");
 		Op->Description = NSLOCTEXT("NiagaraOpInfo", "Log2 Desc", "Result = log2(A)");
 		Op->Keywords = FText::FromString(TEXT("log2"));
 		Op->Inputs.Add(FNiagaraOpInOutInfo(A, NumericType, AText, AText, DefaultStr_One));
@@ -544,6 +572,7 @@ void FNiagaraOpInfo::Init()
 		Op = &OpInfos[Idx];
 		Op->Category = NumericCategory;
 		Op->FriendlyName = NSLOCTEXT("NiagaraOpInfo", "SinRad Name", "Sine(Radians)");
+		Op->CompactName = FText::FromString("Sin(R)");
 		Op->Description = NSLOCTEXT("NiagaraOpInfo", "SinRad Desc", "Result = sin(AngleInRadians)");
 		Op->Keywords = FText::FromString(TEXT("sine"));
 		Op->Inputs.Add(FNiagaraOpInOutInfo(A, NumericType, AText, AText, DefaultStr_One));
@@ -556,6 +585,7 @@ void FNiagaraOpInfo::Init()
 		Op = &OpInfos[Idx];
 		Op->Category = NumericCategory;
 		Op->FriendlyName = NSLOCTEXT("NiagaraOpInfo", "SinDeg Name", "Sine(Degrees)");
+		Op->CompactName = FText::FromString("Sin(D)");
 		Op->Description = NSLOCTEXT("NiagaraOpInfo", "SinDeg Desc", "Result = sin(AngleInDegrees*DegreesToRadians)");
 		Op->Keywords = FText::FromString(TEXT("sine"));
 		Op->Inputs.Add(FNiagaraOpInOutInfo(A, NumericType, AText, AText, DefaultStr_One));
@@ -580,6 +610,7 @@ void FNiagaraOpInfo::Init()
 		Op = &OpInfos[Idx];
 		Op->Category = NumericCategory;
 		Op->FriendlyName = NSLOCTEXT("NiagaraOpInfo", "CosRad Name", "Cosine(Radians)");
+		Op->CompactName = FText::FromString("Cos(R)");
 		Op->Description = NSLOCTEXT("NiagaraOpInfo", "CosRad Desc", "Result = cos(AngleInRadians)");
 		Op->Keywords = FText::FromString(TEXT("Cosine"));
 		Op->Inputs.Add(FNiagaraOpInOutInfo(A, NumericType, AText, AText, DefaultStr_One));
@@ -592,6 +623,7 @@ void FNiagaraOpInfo::Init()
 		Op = &OpInfos[Idx];
 		Op->Category = NumericCategory;
 		Op->FriendlyName = NSLOCTEXT("NiagaraOpInfo", "CosDeg Name", "Cosine(Degrees)");
+		Op->CompactName = FText::FromString("Cos(D)");
 		Op->Description = NSLOCTEXT("NiagaraOpInfo", "CosDeg Desc", "Result = cos(AngleInDegrees*DegreesToRadians)");
 		Op->Keywords = FText::FromString(TEXT("Cosine"));
 		Op->Inputs.Add(FNiagaraOpInOutInfo(A, NumericType, AText, AText, DefaultStr_One));
@@ -616,6 +648,7 @@ void FNiagaraOpInfo::Init()
 		Op = &OpInfos[Idx];
 		Op->Category = NumericCategory;
 		Op->FriendlyName = NSLOCTEXT("NiagaraOpInfo", "TanRad Name", "Tangent(Radians)");
+		Op->CompactName = FText::FromString("Tan(R)");
 		Op->Description = NSLOCTEXT("NiagaraOpInfo", "TanRad Desc", "Result = tan(AngleInRadians)");
 		Op->Keywords = FText::FromString(TEXT("Tangent"));
 		Op->Inputs.Add(FNiagaraOpInOutInfo(A, NumericType, AText, AText, DefaultStr_One));
@@ -628,6 +661,7 @@ void FNiagaraOpInfo::Init()
 		Op = &OpInfos[Idx];
 		Op->Category = NumericCategory;
 		Op->FriendlyName = NSLOCTEXT("NiagaraOpInfo", "TanDeg Name", "Tangent(Degrees)");
+		Op->CompactName = FText::FromString("Tan(D)");
 		Op->Description = NSLOCTEXT("NiagaraOpInfo", "TanDeg Desc", "Result = tan(AngleInDegrees*DegreesToRadians)");
 		Op->Keywords = FText::FromString(TEXT("Tangent"));
 		Op->Inputs.Add(FNiagaraOpInOutInfo(A, NumericType, AText, AText, DefaultStr_One));
@@ -654,6 +688,7 @@ void FNiagaraOpInfo::Init()
 		Op = &OpInfos[Idx];
 		Op->Category = NumericCategory;
 		Op->FriendlyName = NSLOCTEXT("NiagaraOpInfo", "ArcSineRad Name", "ArcSine(Radians)");
+		Op->CompactName = FText::FromString("ArcSin(R)");
 		Op->Description = NSLOCTEXT("NiagaraOpInfo", "ArcSineRad Desc", "Result = asin(A)");
 		Op->Keywords = FText::FromString(TEXT("ArcSine"));
 		Op->Inputs.Add(FNiagaraOpInOutInfo(A, NumericType, AText, AText, DefaultStr_One));
@@ -666,6 +701,7 @@ void FNiagaraOpInfo::Init()
 		Op = &OpInfos[Idx];
 		Op->Category = NumericCategory;
 		Op->FriendlyName = NSLOCTEXT("NiagaraOpInfo", "ArcSineDeg Name", "ArcSine(Degrees)");
+		Op->CompactName = FText::FromString("ArcSin(D)");
 		Op->Description = NSLOCTEXT("NiagaraOpInfo", "ArcSineDeg Desc", "Result = asin(A)*RadiansToDegrees");
 		Op->Keywords = FText::FromString(TEXT("ArcSine"));
 		Op->Inputs.Add(FNiagaraOpInOutInfo(A, NumericType, AText, AText, DefaultStr_One));
@@ -678,6 +714,7 @@ void FNiagaraOpInfo::Init()
 		Op = &OpInfos[Idx];
 		Op->Category = NumericCategory;
 		Op->FriendlyName = NSLOCTEXT("NiagaraOpInfo", "Pi Name", "PI");
+		Op->CompactName = FText::FromString(UTF8_TO_TCHAR("\xCF\x80"));
 		Op->Description = NSLOCTEXT("NiagaraOpInfo", "Pi Desc", "The constant PI");
 		Op->Keywords = FText::FromString(TEXT("pi"));
 		Op->Outputs.Add(FNiagaraOpInOutInfo(Result, FNiagaraTypeDefinition::GetFloatDef(), ResultText, ResultText, DefaultStr_One, TEXT("PI")));
@@ -688,6 +725,7 @@ void FNiagaraOpInfo::Init()
 		Op = &OpInfos[Idx];
 		Op->Category = NumericCategory;
 		Op->FriendlyName = NSLOCTEXT("NiagaraOpInfo", "Two Pi Name", "TWO_PI");
+		Op->CompactName = FText::FormatOrdered(FText::FromString("2 {0}"), FText::FromString(UTF8_TO_TCHAR("\xCF\x80")));
 		Op->Description = NSLOCTEXT("NiagaraOpInfo", "Two Pi Desc", "The constant PI * 2");
 		Op->Keywords = FText::FromString(TEXT("pi"));
 		Op->Outputs.Add(FNiagaraOpInOutInfo(Result, FNiagaraTypeDefinition::GetFloatDef(), ResultText, ResultText, DefaultStr_One, TEXT("TWO_PI")));
@@ -712,6 +750,7 @@ void FNiagaraOpInfo::Init()
 		Op = &OpInfos[Idx];
 		Op->Category = NumericCategory;
 		Op->FriendlyName = NSLOCTEXT("NiagaraOpInfo", "ArcCosineRad Name", "ArcCosine(Radians)");
+		Op->CompactName = FText::FromString("ArcCos(R)");
 		Op->Description = NSLOCTEXT("NiagaraOpInfo", "ArcCosineRad Desc", "Result = acos(A)");
 		Op->Keywords = FText::FromString(TEXT("ArcCosine"));
 		Op->Inputs.Add(FNiagaraOpInOutInfo(A, NumericType, AText, AText, DefaultStr_One));
@@ -724,6 +763,7 @@ void FNiagaraOpInfo::Init()
 		Op = &OpInfos[Idx];
 		Op->Category = NumericCategory;
 		Op->FriendlyName = NSLOCTEXT("NiagaraOpInfo", "ArcCosineDeg Name", "ArcCosine(Degrees)");
+		Op->CompactName = FText::FromString("ArcCos(D)");
 		Op->Description = NSLOCTEXT("NiagaraOpInfo", "ArcCosineDeg Desc", "Result = acos(A)*RadiansToDegrees");
 		Op->Keywords = FText::FromString(TEXT("ArcCosine"));
 		Op->Inputs.Add(FNiagaraOpInOutInfo(A, NumericType, AText, AText, DefaultStr_One));
@@ -750,6 +790,7 @@ void FNiagaraOpInfo::Init()
 		Op = &OpInfos[Idx];
 		Op->Category = NumericCategory;
 		Op->FriendlyName = NSLOCTEXT("NiagaraOpInfo", "ArcTangentRad Name", "ArcTangent(Radians)");
+		Op->CompactName = FText::FromString("ArcTan(R)");
 		Op->Description = NSLOCTEXT("NiagaraOpInfo", "ArcTangentRad Desc", "Result = atan(A)");
 		Op->Keywords = FText::FromString(TEXT("ArcTangent"));
 		Op->Inputs.Add(FNiagaraOpInOutInfo(A, NumericType, AText, AText, DefaultStr_One));
@@ -762,6 +803,7 @@ void FNiagaraOpInfo::Init()
 		Op = &OpInfos[Idx];
 		Op->Category = NumericCategory;
 		Op->FriendlyName = NSLOCTEXT("NiagaraOpInfo", "ArcTangentDeg Name", "ArcTangent(Degrees)");
+		Op->CompactName = FText::FromString("ArcTan(D)");
 		Op->Description = NSLOCTEXT("NiagaraOpInfo", "ArcTangentDeg Desc", "Result = atan(A)*RadiansToDegrees");
 		Op->Keywords = FText::FromString(TEXT("ArcTangent"));
 		Op->Inputs.Add(FNiagaraOpInOutInfo(A, NumericType, AText, AText, DefaultStr_One));
@@ -811,7 +853,8 @@ void FNiagaraOpInfo::Init()
 		Op = &OpInfos[Idx];
 		Op->Category = NumericCategory;
 		Op->FriendlyName = NSLOCTEXT("NiagaraOpInfo", "Degrees To Radians", "DegreesToRadians");
-		Op->Description = NSLOCTEXT("NiagaraOpInfo", "Degrees To Radians Desc", "DegreesToRadians(A)");
+		Op->CompactName = FText::FromString("DegToRad");
+		Op->Description = NSLOCTEXT("NiagaraOpInfo", "Degrees To Radians Desc", "Degrees to Radians");
 		Op->Inputs.Add(FNiagaraOpInOutInfo(A, FNiagaraTypeDefinition::GetFloatDef(), AText, AText, DefaultStr_One));
 		Op->Outputs.Add(FNiagaraOpInOutInfo(Result, FNiagaraTypeDefinition::GetFloatDef(), ResultText, ResultText, DefaultStr_One, TEXT("(PI/180.0f)*({0})")));
 		Op->BuildName(TEXT("DegreesToRadians"), CategoryName);
@@ -821,7 +864,8 @@ void FNiagaraOpInfo::Init()
 		Op = &OpInfos[Idx];
 		Op->Category = NumericCategory;
 		Op->FriendlyName = NSLOCTEXT("NiagaraOpInfo", "Radians To Degrees", "RadiansToDegrees");
-		Op->Description = NSLOCTEXT("NiagaraOpInfo", "RadiansToDegrees Desc", "RadiansToDegrees(A)");
+		Op->CompactName = FText::FromString("RadToDeg");
+		Op->Description = NSLOCTEXT("NiagaraOpInfo", "RadiansToDegrees Desc", "Radians to Degrees");
 		Op->Inputs.Add(FNiagaraOpInOutInfo(A, FNiagaraTypeDefinition::GetFloatDef(), AText, AText, DefaultStr_One));
 		Op->Outputs.Add(FNiagaraOpInOutInfo(Result, FNiagaraTypeDefinition::GetFloatDef(), ResultText, ResultText, DefaultStr_One, TEXT("(180.0f/PI)*({0})")));
 		Op->BuildName(TEXT("RadiansToDegrees"), CategoryName);
@@ -831,6 +875,7 @@ void FNiagaraOpInfo::Init()
 		Op = &OpInfos[Idx];
 		Op->Category = NumericCategory;
 		Op->FriendlyName = NSLOCTEXT("NiagaraOpInfo", "Ceil Name", "Ceil");
+		Op->CompactName = FText::FromString("Ceil");
 		Op->Description = NSLOCTEXT("NiagaraOpInfo", "Ceil Desc", "Rounds A to the nearest integer higher than A.");
 		Op->Inputs.Add(FNiagaraOpInOutInfo(A, NumericType, AText, AText, DefaultStr_One));
 		Op->Outputs.Add(FNiagaraOpInOutInfo(Result, NumericType, ResultText, ResultText, DefaultStr_One, TEXT("ceil({0})")));
@@ -841,6 +886,7 @@ void FNiagaraOpInfo::Init()
 		Op = &OpInfos[Idx];
 		Op->Category = NumericCategory;
 		Op->FriendlyName = NSLOCTEXT("NiagaraOpInfo", "Floor Name", "Floor");
+		Op->CompactName = FText::FromString("Floor");
 		Op->Description = NSLOCTEXT("NiagaraOpInfo", "Floor Desc", "Rounds A to the nearest integer lower than A.");
 		Op->Inputs.Add(FNiagaraOpInOutInfo(A, NumericType, AText, AText, DefaultStr_One));
 		Op->Outputs.Add(FNiagaraOpInOutInfo(Result, NumericType, ResultText, ResultText, DefaultStr_One, TEXT("floor({0})")));
@@ -851,6 +897,7 @@ void FNiagaraOpInfo::Init()
 		Op = &OpInfos[Idx];
 		Op->Category = NumericCategory;
 		Op->FriendlyName = NSLOCTEXT("NiagaraOpInfo", "Round Name", "Round");
+		Op->CompactName = FText::FromString("Round");
 		Op->Description = NSLOCTEXT("NiagaraOpInfo", "Round Desc", "Rounds A to the nearest integer.");
 		Op->Inputs.Add(FNiagaraOpInOutInfo(A, NumericType, AText, AText, DefaultStr_One));
 		Op->Outputs.Add(FNiagaraOpInOutInfo(Result, NumericType, ResultText, ResultText, DefaultStr_One, TEXT("round({0})")));
@@ -861,7 +908,9 @@ void FNiagaraOpInfo::Init()
 		Op = &OpInfos[Idx];
 		Op->Category = NumericCategory;
 		Op->FriendlyName = NSLOCTEXT("NiagaraOpInfo", "Fmod Name", "Modulo");
-		Op->Description = NSLOCTEXT("NiagaraOpInfo", "Fmod Desc", "Result = A % B");
+		Op->CompactName = FText::FromString("%");
+		Op->Description = NSLOCTEXT("NiagaraOpInfo", "Fmod Desc", "Result = A % B (Modulo)");
+		Op->CompactNameFontSizeOverride = 25.f;
 		Op->Keywords = FText::FromString(TEXT("%"));
 		Op->Inputs.Add(FNiagaraOpInOutInfo(A, NumericType, AText, AText, DefaultStr_One));
 		Op->Inputs.Add(FNiagaraOpInOutInfo(B, NumericType, BText, BText, DefaultStr_One));
@@ -887,6 +936,7 @@ void FNiagaraOpInfo::Init()
 		Op = &OpInfos[Idx];
 		Op->Category = NumericCategory;
 		Op->FriendlyName = NSLOCTEXT("NiagaraOpInfo", "Frac Name", "Frac");
+		Op->CompactName = FText::FromString("Frac");
 		Op->Description = NSLOCTEXT("NiagaraOpInfo", "Frac Desc", "Result = frac(A)");
 		Op->Inputs.Add(FNiagaraOpInOutInfo(A, NumericType, AText, AText, DefaultStr_One));
 		Op->Outputs.Add(FNiagaraOpInOutInfo(Result, NumericType, ResultText, ResultText, DefaultStr_One, TEXT("frac({0})")));
@@ -897,6 +947,7 @@ void FNiagaraOpInfo::Init()
 		Op = &OpInfos[Idx];
 		Op->Category = NumericCategory;
 		Op->FriendlyName = NSLOCTEXT("NiagaraOpInfo", "Trunc Name", "Trunc");
+		Op->CompactName = FText::FromString("Trunc");
 		Op->Description = NSLOCTEXT("NiagaraOpInfo", "Trunc Desc", "Result = trunc(A)");
 		Op->Inputs.Add(FNiagaraOpInOutInfo(A, NumericType, AText, AText, DefaultStr_One));
 		Op->Outputs.Add(FNiagaraOpInOutInfo(Result, NumericType, ResultText, ResultText, DefaultStr_One, TEXT("trunc({0})")));
@@ -939,6 +990,7 @@ void FNiagaraOpInfo::Init()
 		Op = &OpInfos[Idx];
 		Op->Category = NumericCategory;
 		Op->FriendlyName = NSLOCTEXT("NiagaraOpInfo", "Min Name", "Min");
+		Op->CompactName = FText::FromString("Min");
 		Op->Description = NSLOCTEXT("NiagaraOpInfo", "Min Desc", "Result = min(A, B)");
 		Op->Inputs.Add(FNiagaraOpInOutInfo(A, NumericType, AText, AText, DefaultStr_One));
 		Op->Inputs.Add(FNiagaraOpInOutInfo(B, NumericType, BText, BText, DefaultStr_One));
@@ -973,6 +1025,7 @@ void FNiagaraOpInfo::Init()
 		Op = &OpInfos[Idx];
 		Op->Category = NumericCategory;
 		Op->FriendlyName = NSLOCTEXT("NiagaraOpInfo", "Max Name", "Max");
+		Op->CompactName = FText::FromString("Max");
 		Op->Description = NSLOCTEXT("NiagaraOpInfo", "Max Desc", "Result = max(A, B)");
 		Op->Inputs.Add(FNiagaraOpInOutInfo(A, NumericType, AText, AText, DefaultStr_One));
 		Op->Inputs.Add(FNiagaraOpInOutInfo(B, NumericType, BText, BText, DefaultStr_One));
@@ -1019,6 +1072,7 @@ void FNiagaraOpInfo::Init()
 		Op = &OpInfos[Idx];
 		Op->Category = NumericCategory;
 		Op->FriendlyName = NSLOCTEXT("NiagaraOpInfo", "Sign Name", "Sign");
+		Op->CompactName = FText::FromString("Sign");
 		Op->Description = NSLOCTEXT("NiagaraOpInfo", "Sign Desc", "Result = sign(A)");
 		Op->Inputs.Add(FNiagaraOpInOutInfo(A, NumericType, AText, AText, DefaultStr_One));
 		Op->Outputs.Add(FNiagaraOpInOutInfo(Result, NumericType, ResultText, ResultText, DefaultStr_One, TEXT("sign({0})")));
@@ -1030,6 +1084,7 @@ void FNiagaraOpInfo::Init()
 		Op = &OpInfos[Idx];
 		Op->Category = NumericCategory;
 		Op->FriendlyName = NSLOCTEXT("NiagaraOpInfo", "Step Name", "Step");
+		Op->CompactName = FText::FromString("Step");
 		Op->Description = NSLOCTEXT("NiagaraOpInfo", "Step Desc", "Result = step(0, A)");
 		Op->Inputs.Add(FNiagaraOpInOutInfo(A, NumericType, AText, AText, DefaultStr_One));
 		Op->Outputs.Add(FNiagaraOpInOutInfo(Result, NumericType, ResultText, ResultText, DefaultStr_One, TEXT("step(0, {0})")));
@@ -1052,6 +1107,7 @@ void FNiagaraOpInfo::Init()
 		Op = &OpInfos[Idx];
 		Op->Category = NumericCategory;
 		Op->FriendlyName = NSLOCTEXT("NiagaraOpInfo", "Dot Name", "Dot");
+		Op->CompactName = FText::FromString("Dot");
 		Op->Description = NSLOCTEXT("NiagaraOpInfo", "Dot Desc", "Dot product of two vectors.");
 		Op->Inputs.Add(FNiagaraOpInOutInfo(A, NumericType, AText, AText, DefaultStr_One));
 		Op->Inputs.Add(FNiagaraOpInOutInfo(B, NumericType, BText, BText, DefaultStr_One));
@@ -1076,6 +1132,7 @@ void FNiagaraOpInfo::Init()
 		Op = &OpInfos[Idx];
 		Op->Category = NumericCategory;
 		Op->FriendlyName = NSLOCTEXT("NiagaraOpInfo", "Length Name", "Length");
+		Op->CompactName = FText::FromString("Len");
 		Op->Description = NSLOCTEXT("NiagaraOpInfo", "Length Desc", "Returns the length of the passed value.");
 		Op->Inputs.Add(FNiagaraOpInOutInfo(A, NumericType, AText, AText, DefaultStr_One));
 		Op->Outputs.Add(FNiagaraOpInOutInfo(Result, NumericType, ResultText, ResultText, DefaultStr_One, TEXT("length({0})")));
@@ -1240,6 +1297,7 @@ void FNiagaraOpInfo::Init()
 		Op = &OpInfos[Idx];
 		Op->Category = NumericCategory;
 		Op->FriendlyName = NSLOCTEXT("NiagaraOpInfo", "CmpLT Name", "Less Than");
+		Op->CompactName = FText::FromString("<");
 		Op->Description = NSLOCTEXT("NiagaraOpInfo", "CmpLT Desc", "Result = A < B");
 		Op->Keywords = FText::FromString(TEXT("<"));
 		Op->Inputs.Add(FNiagaraOpInOutInfo(A, NumericType, AText, AText, DefaultStr_Zero));
@@ -1264,6 +1322,7 @@ void FNiagaraOpInfo::Init()
 		Op = &OpInfos[Idx];
 		Op->Category = NumericCategory;
 		Op->FriendlyName = NSLOCTEXT("NiagaraOpInfo", "CmpLE Name", "Less Than Or Equal");
+		Op->CompactName = FText::FromString("<=");
 		Op->Description = NSLOCTEXT("NiagaraOpInfo", "CmpLE Desc", "Result = A <= B");
 		Op->Keywords = FText::FromString(TEXT("<="));
 		Op->Inputs.Add(FNiagaraOpInOutInfo(A, NumericType, AText, AText, DefaultStr_Zero));
@@ -1287,6 +1346,7 @@ void FNiagaraOpInfo::Init()
 		Op = &OpInfos[Idx];
 		Op->Category = NumericCategory;
 		Op->FriendlyName = NSLOCTEXT("NiagaraOpInfo", "CmpGT Name", "Greater Than");
+		Op->CompactName = FText::FromString(">");
 		Op->Description = NSLOCTEXT("NiagaraOpInfo", "CmpGT Desc", "Result = A > B");
 		Op->Keywords = FText::FromString(TEXT(">"));
 		Op->Inputs.Add(FNiagaraOpInOutInfo(A, NumericType, AText, AText, DefaultStr_Zero));
@@ -1309,6 +1369,7 @@ void FNiagaraOpInfo::Init()
 		Op = &OpInfos[Idx];
 		Op->Category = NumericCategory;
 		Op->FriendlyName = NSLOCTEXT("NiagaraOpInfo", "CmpGE Name", "Greater Than Or Equal");
+		Op->CompactName = FText::FromString(">=");
 		Op->Description = NSLOCTEXT("NiagaraOpInfo", "CmpGE Desc", "Result = A >= B");
 		Op->Keywords = FText::FromString(TEXT(">="));
 		Op->Inputs.Add(FNiagaraOpInOutInfo(A, NumericType, AText, AText, DefaultStr_Zero));
@@ -1330,6 +1391,7 @@ void FNiagaraOpInfo::Init()
 		Op = &OpInfos[Idx];
 		Op->Category = NumericCategory;
 		Op->FriendlyName = NSLOCTEXT("NiagaraOpInfo", "CmpEQ Name", "Equal");
+		Op->CompactName = FText::FromString("==");
 		Op->Description = NSLOCTEXT("NiagaraOpInfo", "CmpEQ Desc", "Result = A == B");
 		Op->Keywords = FText::FromString(TEXT("=="));
 		Op->Inputs.Add(FNiagaraOpInOutInfo(A, NumericType, AText, AText, DefaultStr_Zero));
@@ -1354,6 +1416,7 @@ void FNiagaraOpInfo::Init()
 		Op = &OpInfos[Idx];
 		Op->Category = NumericCategory;
 		Op->FriendlyName = NSLOCTEXT("NiagaraOpInfo", "CmpNEQ Name", "Not Equal");
+		Op->CompactName = FText::FromString("!=");
 		Op->Description = NSLOCTEXT("NiagaraOpInfo", "CmpNEQ Desc", "Result = A != B");
 		Op->Keywords = FText::FromString(TEXT("!="));
 		Op->Inputs.Add(FNiagaraOpInOutInfo(A, NumericType, AText, AText, DefaultStr_Zero));
@@ -1595,6 +1658,7 @@ void FNiagaraOpInfo::Init()
 	Op = &OpInfos[Idx];
 	Op->Category = BoolCategory;
 	Op->FriendlyName = NSLOCTEXT("NiagaraOpInfo", "LogicEq Name", "Bool Equal");
+	Op->CompactName = FText::FromString("==");
 	Op->Description = NSLOCTEXT("NiagaraOpInfo", "LogicEq Desc", "Result = A == B");
 	Op->Keywords = FText::FromString(TEXT("=="));
 	Op->Inputs.Add(FNiagaraOpInOutInfo(A, BoolType, AText, AText, Default_BoolZero));
@@ -1616,6 +1680,7 @@ void FNiagaraOpInfo::Init()
 	Op = &OpInfos[Idx];
 	Op->Category = BoolCategory;
 	Op->FriendlyName = NSLOCTEXT("NiagaraOpInfo", "LogicNEq Name", "Bool Not Equal");
+	Op->CompactName = FText::FromString("!=");
 	Op->Description = NSLOCTEXT("NiagaraOpInfo", "LogicNEq Desc", "Result = A != B");
 	Op->Keywords = FText::FromString(TEXT("!="));
 	Op->Inputs.Add(FNiagaraOpInOutInfo(A, BoolType, AText, AText, Default_BoolZero));
@@ -1769,6 +1834,7 @@ void FNiagaraOpInfo::Init()
 	Op = &OpInfos[Idx];
 	Op->Category = NumericCategory;
 	Op->FriendlyName = NSLOCTEXT("NiagaraOpInfo", "Vector Cross Name", "Cross");
+	Op->CompactName = FText::FromString("Cross");
 	Op->Description = NSLOCTEXT("NiagaraOpInfo", "Vector Cross Desc", "Cross product of two vectors.");
 	Op->Inputs.Add(FNiagaraOpInOutInfo(A, FNiagaraTypeDefinition::GetVec3Def(), AText, AText, Default_VectorX));
 	Op->Inputs.Add(FNiagaraOpInOutInfo(B, FNiagaraTypeDefinition::GetVec3Def(), BText, BText, Default_VectorY));
diff --git a/Engine/Plugins/FX/Niagara/Source/NiagaraEditor/Private/NiagaraNodeOp.cpp b/Engine/Plugins/FX/Niagara/Source/NiagaraEditor/Private/NiagaraNodeOp.cpp
--- a/Engine/Plugins/FX/Niagara/Source/NiagaraEditor/Private/NiagaraNodeOp.cpp
+++ b/Engine/Plugins/FX/Niagara/Source/NiagaraEditor/Private/NiagaraNodeOp.cpp
@@ -533,6 +533,40 @@ void UNiagaraNodeOp::PinConnectionListChanged(UEdGraphPin* Pin)
 	}
 }
 
+FText UNiagaraNodeOp::GetCompactTitle() const
+{
+	const FNiagaraOpInfo* OpInfo = FNiagaraOpInfo::GetOpInfo(OpName);
+	if(OpInfo == nullptr)
+	{
+		// this will effectively turn off compact mode
+		return FText::GetEmpty();
+	}
+	
+	return OpInfo->CompactName; 
+}
+
+bool UNiagaraNodeOp::ShouldShowPinNamesInCompactMode()
+{
+	const FNiagaraOpInfo* OpInfo = FNiagaraOpInfo::GetOpInfo(OpName);
+	if(OpInfo == nullptr)
+	{
+		return false;
+	}
+	
+	return OpInfo->bShowPinNamesInCompactMode; 
+}
+
+TOptional<float> UNiagaraNodeOp::GetCompactModeFontSizeOverride() const
+{
+	const FNiagaraOpInfo* OpInfo = FNiagaraOpInfo::GetOpInfo(OpName);
+	if(OpInfo == nullptr)
+	{
+		return {};
+	}
+	
+	return OpInfo->CompactNameFontSizeOverride;
+}
+
 void UNiagaraNodeOp::HandleStaticInputPinUpgrade(UEdGraphPin* Pin)
 {
 	const FNiagaraOpInfo* OpInfo = FNiagaraOpInfo::GetOpInfo(OpName);
diff --git a/Engine/Plugins/FX/Niagara/Source/NiagaraEditor/Private/SNiagaraGraphPinNumeric.cpp b/Engine/Plugins/FX/Niagara/Source/NiagaraEditor/Private/SNiagaraGraphPinNumeric.cpp
--- a/Engine/Plugins/FX/Niagara/Source/NiagaraEditor/Private/SNiagaraGraphPinNumeric.cpp
+++ b/Engine/Plugins/FX/Niagara/Source/NiagaraEditor/Private/SNiagaraGraphPinNumeric.cpp
@@ -1,8 +1,47 @@
 // Copyright Epic Games, Inc. All Rights Reserved.
 
 #include "SNiagaraGraphPinNumeric.h"
+#include "EdGraphSchema_Niagara.h"
 
 void SNiagaraGraphPinNumeric::Construct(const FArguments& InArgs, UEdGraphPin* InGraphPinObj)
 {
 	SGraphPin::Construct(SGraphPin::FArguments(), InGraphPinObj);
-}
\ No newline at end of file
+}
+
+FSlateColor SNiagaraGraphPinNumeric::GetPinColor() const
+{
+	FSlateColor OriginalColor = SGraphPin::GetPinColor();
+	
+	if(UEdGraphPin* SourcePin = GetPinObj())
+	{
+		if(!SourcePin->IsPendingKill() && SourcePin->LinkedTo.Num() == 1)
+		{
+			if(UEdGraphPin* ConnectedPin = SourcePin->LinkedTo[0])
+			{
+				FNiagaraTypeDefinition ConnectedPiNType = UEdGraphSchema_Niagara::PinTypeToTypeDefinition(ConnectedPin->PinType);
+				
+				if (bIsDiffHighlighted)
+				{
+					return FSlateColor(FLinearColor(0.9f, 0.2f, 0.15f));
+				}
+				if (SourcePin->bOrphanedPin)
+				{
+					return FSlateColor(FLinearColor::Red);
+				}
+				if (const UEdGraphSchema* Schema = SourcePin->GetSchema())
+				{
+					if (!GetPinObj()->GetOwningNode()->IsNodeEnabled() || GetPinObj()->GetOwningNode()->IsDisplayAsDisabledForced() || !IsEditingEnabled() || GetPinObj()->GetOwningNode()->IsNodeUnrelated())
+					{
+						return Schema->GetPinTypeColor(ConnectedPin->PinType) * FLinearColor(1.0f, 1.0f, 1.0f, 0.5f);
+					}
+
+					return Schema->GetPinTypeColor(ConnectedPin->PinType) * PinColorModifier;
+				}
+
+				return FLinearColor::White;
+			}
+		}
+	}
+
+	return OriginalColor;
+}
diff --git a/Engine/Plugins/FX/Niagara/Source/NiagaraEditor/Private/SNiagaraGraphPinNumeric.h b/Engine/Plugins/FX/Niagara/Source/NiagaraEditor/Private/SNiagaraGraphPinNumeric.h
--- a/Engine/Plugins/FX/Niagara/Source/NiagaraEditor/Private/SNiagaraGraphPinNumeric.h
+++ b/Engine/Plugins/FX/Niagara/Source/NiagaraEditor/Private/SNiagaraGraphPinNumeric.h
@@ -15,4 +15,7 @@ public:
 	SLATE_END_ARGS()
 
 	void Construct(const FArguments& InArgs, UEdGraphPin* InGraphPinObj);
+
+private:
+	virtual FSlateColor GetPinColor() const override;
 };
\ No newline at end of file
diff --git a/Engine/Plugins/FX/Niagara/Source/NiagaraEditor/Private/Widgets/SNiagaraGraphNode.cpp b/Engine/Plugins/FX/Niagara/Source/NiagaraEditor/Private/Widgets/SNiagaraGraphNode.cpp
--- a/Engine/Plugins/FX/Niagara/Source/NiagaraEditor/Private/Widgets/SNiagaraGraphNode.cpp
+++ b/Engine/Plugins/FX/Niagara/Source/NiagaraEditor/Private/Widgets/SNiagaraGraphNode.cpp
@@ -1,13 +1,23 @@
 // Copyright Epic Games, Inc. All Rights Reserved.
 
 #include "SNiagaraGraphNode.h"
+
+#include "EdGraphSchema_Niagara.h"
 #include "GraphEditorSettings.h"
+#include "IDocumentation.h"
 #include "Rendering/DrawElements.h"
 #include "SGraphPin.h"
 #include "NiagaraEditorSettings.h"
+#include "SCommentBubble.h"
+#include "TutorialMetaData.h"
+#include "Widgets/Images/SLayeredImage.h"
+#include "Widgets/SToolTip.h"
 
 #define LOCTEXT_NAMESPACE "SNiagaraGraphNode"
 
+const FSlateBrush* SNiagaraGraphNode::CachedInnerIcon = nullptr;
+const FSlateBrush* SNiagaraGraphNode::CachedOuterIcon = nullptr;
+
 SNiagaraGraphNode::SNiagaraGraphNode() : SGraphNode()
 {
 	NiagaraNode = nullptr;
@@ -24,12 +34,14 @@ SNiagaraGraphNode::~SNiagaraGraphNode()
 void SNiagaraGraphNode::Construct(const FArguments& InArgs, UEdGraphNode* InGraphNode)
 {
 	GraphNode = InGraphNode;
+	NiagaraNode = Cast<UNiagaraNode>(InGraphNode);
+	CompactTitle = NiagaraNode->GetCompactTitle();
+	bShowPinNamesInCompactMode = NiagaraNode->ShouldShowPinNamesInCompactMode();
+	CompactNodeTitleFontSizeOverride = NiagaraNode->GetCompactModeFontSizeOverride();
 	RegisterNiagaraGraphNode(InGraphNode);
-	UpdateGraphNode();
-	
+	UpdateGraphNode();	
 }
 
-
 void SNiagaraGraphNode::HandleNiagaraNodeChanged(UNiagaraNode* InNode)
 {
 	check(InNode == NiagaraNode);
@@ -45,7 +57,14 @@ void SNiagaraGraphNode::RegisterNiagaraGraphNode(UEdGraphNode* InNode)
 void SNiagaraGraphNode::UpdateGraphNode()
 {
 	check(NiagaraNode.IsValid());
-	SGraphNode::UpdateGraphNode();
+	if(ShouldDrawCompact())
+	{
+		UpdateGraphNodeCompact();
+	}
+	else
+	{
+		SGraphNode::UpdateGraphNode();
+	}
 	LastSyncedNodeChangeId = NiagaraNode->GetChangeId();
 }
 
@@ -79,4 +98,317 @@ void SNiagaraGraphNode::UpdateErrorInfo()
 	}
 }
 
-#undef LOCTEXT_NAMESPACE
\ No newline at end of file
+void SNiagaraGraphNode::CreatePinWidgets()
+{
+	// the default implementation will create default pin widgets for all our logical pins
+	SGraphNode::CreatePinWidgets();
+
+	TSet<TSharedRef<SWidget>> AllPins;
+	GetPins(AllPins);
+
+	LoadCachedIcons();
+
+	// then we override the pin widget images for wildcards & numerics
+	for(TSharedRef<SWidget>& Widget : AllPins)
+	{
+		if(TSharedPtr<SGraphPin> Pin = StaticCastSharedRef<SGraphPin>(Widget))
+		{
+			UEdGraphPin* SourcePin = Pin->GetPinObj();
+			
+			// Split pins should be drawn as normal pins, the inner properties are not promotable
+			if(!SourcePin || SourcePin->ParentPin)
+			{
+				continue;
+			}
+
+			if(SourcePin->PinType == UEdGraphSchema_Niagara::TypeDefinitionToPinType(FNiagaraTypeDefinition::GetWildcardDef()) || SourcePin->PinType == UEdGraphSchema_Niagara::TypeDefinitionToPinType(FNiagaraTypeDefinition::GetGenericNumericDef()))
+			{
+				if(TSharedPtr<SLayeredImage> PinImage = StaticCastSharedPtr<SLayeredImage>(Pin->GetPinImageWidget()))
+				{
+					// Set the image to use the outer icon, which will be the connect pin type color
+					PinImage->SetLayerBrush(0, CachedOuterIcon);
+
+					// Set the inner image to be wildcard color, which is grey by default
+					PinImage->AddLayer(CachedInnerIcon, GetDefault<UGraphEditorSettings>()->WildcardPinTypeColor);
+				}
+			}			
+		}
+	}
+}
+
+void SNiagaraGraphNode::Tick(const FGeometry& AllottedGeometry, const double InCurrentTime, const float InDeltaTime)
+{
+	SGraphNode::Tick(AllottedGeometry, InCurrentTime, InDeltaTime);
+
+	if(CompactTitle.IsEmpty() == false)
+	{
+		FVector2D LeftNodeBoxSize = LeftNodeBox->GetDesiredSize();
+		FVector2D RightNodeBoxSize = RightNodeBox->GetDesiredSize();
+
+		if(LeftNodeBoxSize.X > RightNodeBoxSize.X)
+		{
+			CompactNodeMaxPinBoxX = LeftNodeBoxSize.X;
+		}
+		else
+		{
+			CompactNodeMaxPinBoxX = RightNodeBoxSize.X;
+		}
+	}
+}
+
+void SNiagaraGraphNode::UpdateGraphNodeCompact()
+{
+	InputPins.Empty();
+	OutputPins.Empty();
+
+	// error handling set-up
+	SetupErrorReporting();
+
+	// Reset variables that are going to be exposed, in case we are refreshing an already setup node.
+	RightNodeBox.Reset();
+	LeftNodeBox.Reset();
+
+ 	TSharedPtr<SToolTip> NodeToolTip = SNew( SToolTip );
+	if (!SWidget::GetToolTip().IsValid())
+	{
+		NodeToolTip = IDocumentation::Get()->CreateToolTip( TAttribute< FText >( this, &SGraphNode::GetNodeTooltip ), NULL, GraphNode->GetDocumentationLink(), GraphNode->GetDocumentationExcerptName() );
+		SetToolTip(NodeToolTip);
+	}
+
+	// Setup a meta tag for this node
+	FGraphNodeMetaData TagMeta(TEXT("Graphnode"));
+	PopulateMetaTag(&TagMeta);
+	
+	TSharedPtr<SNodeTitle> NodeTitle = SNew(SNodeTitle, GraphNode)
+		.Text(this, &SNiagaraGraphNode::GetNodeCompactTitle);
+
+	TSharedRef<SHorizontalBox> NodeContent = SNew(SHorizontalBox);
+	
+	// Default to "pure" styling, where we can just center the pins vertically
+	// since don't need to worry about alignment with other nodes
+	float PinPaddingTop = 0.f;
+
+	EVerticalAlignment PinVerticalAlignment = VAlign_Center;
+	
+	// If at least 1 input pin exists, we want to sync left & right pin box sizes so that the title stays in the middle
+	// If we don't have an input, like for constant getters or BP-like variable getters, we don't want to enlarge the left pin box
+	bool bAnyInputExists = false;
+	for(UEdGraphPin* Pin : GraphNode->GetAllPins())
+	{
+		if(Pin->Direction == EGPD_Input && Pin->bHidden == false)
+		{
+			bAnyInputExists = true;
+			bSyncPinBoxSizes = true;
+			break;
+		}
+	}
+
+	NodeContent->AddSlot()
+		.HAlign(HAlign_Left)
+		.VAlign(PinVerticalAlignment)
+		.AutoWidth()
+		.Padding(/* left */ 0.f, PinPaddingTop, 5.f, /* bottom */ 0.f)
+		[
+			SNew(SBox)
+			.MinDesiredWidth(this, &SNiagaraGraphNode::GetCompactMinDesiredPinBoxSize)
+			[
+				// LEFT
+				SAssignNew(LeftNodeBox, SVerticalBox)
+			]
+		];
+	
+	NodeContent->AddSlot()
+		.HAlign(HAlign_Center)
+		.VAlign(VAlign_Center)
+		.AutoWidth()
+		.Padding(bAnyInputExists ? 5.f : 25.f, 0.f, 5.f, 0.f)
+		[
+			// MIDDLE
+			SNew(SVerticalBox)
+			+SVerticalBox::Slot()
+			.HAlign(HAlign_Center)
+			.VAlign(VAlign_Center)
+			.AutoHeight()
+			[
+				SNew(STextBlock)
+					.TextStyle( FAppStyle::Get(), "Graph.CompactNode.Title" )
+					.Text( NodeTitle.Get(), &SNodeTitle::GetHeadTitle )
+					.WrapTextAt(128.0f)
+					.Font(GetCompactNodeTitleFont())
+			]
+			+SVerticalBox::Slot()
+			.AutoHeight()
+			[
+				NodeTitle.ToSharedRef()
+			]
+		];
+	
+	NodeContent->AddSlot()
+		.HAlign(HAlign_Right)
+		.VAlign(PinVerticalAlignment)
+		.AutoWidth()
+		.Padding(5.f, PinPaddingTop, 0.f, 0.f)
+		[
+			SNew(SBox)
+			.MinDesiredWidth(this, &SNiagaraGraphNode::GetCompactMinDesiredPinBoxSize)
+			[
+				// RIGHT
+				SAssignNew(RightNodeBox, SVerticalBox)
+			]
+		];
+
+	//
+	//             ______________________
+	//            | (>) L |      | R (>) |
+	//            | (>) E |      | I (>) |
+	//            | (>) F |   +  | G (>) |
+	//            | (>) T |      | H (>) |
+	//            |       |      | T (>) |
+	//            |_______|______|_______|
+	//
+	this->ContentScale.Bind( this, &SGraphNode::GetContentScale );
+	
+	TSharedRef<SVerticalBox> InnerVerticalBox =
+		SNew(SVerticalBox)
+		+SVerticalBox::Slot()
+		[
+			// NODE CONTENT AREA
+			SNew( SOverlay)
+			+SOverlay::Slot()
+			[
+				SNew(SImage)
+				.Image( FAppStyle::GetBrush("Graph.VarNode.Body") )
+			]
+			+ SOverlay::Slot()
+			[
+				SNew(SImage)
+				.Image( FAppStyle::GetBrush("Graph.VarNode.Gloss") )
+			]
+			+SOverlay::Slot()
+			.Padding( FMargin(0,3) )
+			[
+				NodeContent
+			]
+		];
+	
+	TSharedPtr<SWidget> EnabledStateWidget = GetEnabledStateWidget();
+	if (EnabledStateWidget.IsValid())
+	{
+		InnerVerticalBox->AddSlot()
+			.AutoHeight()
+			.HAlign(HAlign_Fill)
+			.VAlign(VAlign_Top)
+			.Padding(FMargin(2, 0))
+			[
+				EnabledStateWidget.ToSharedRef()
+			];
+	}
+
+	InnerVerticalBox->AddSlot()
+		.AutoHeight()
+		.Padding( FMargin(5.0f, 1.0f) )
+		[
+			ErrorReporting->AsWidget()
+		];
+
+	this->GetOrAddSlot( ENodeZone::Center )
+	.HAlign(HAlign_Center)
+	.VAlign(VAlign_Center)
+	[
+		InnerVerticalBox
+	];
+
+	CreatePinWidgets();
+
+	// We hide input pin labels by default, unless the node specified we want to show them in compact mode
+	if(bShowPinNamesInCompactMode == false)
+	{
+		for (auto InputPin: this->InputPins)
+		{
+			if (InputPin->GetPinObj()->ParentPin == nullptr)
+			{
+				InputPin->SetShowLabel(false);
+			}
+		}
+
+	}
+
+	// Output pin names are generally hidden here as compact mode is generally used for simple operations that have only 1 output
+	for (auto OutputPin : this->OutputPins)
+	{
+		if (OutputPin->GetPinObj()->ParentPin == nullptr)
+		{
+			OutputPin->SetShowLabel(false);
+		}
+	}
+	
+	// Create comment bubble
+	TSharedPtr<SCommentBubble> CommentBubble;
+	const FSlateColor CommentColor = GetDefault<UGraphEditorSettings>()->DefaultCommentNodeTitleColor;
+
+	SAssignNew( CommentBubble, SCommentBubble )
+	.GraphNode( GraphNode )
+	.Text( this, &SGraphNode::GetNodeComment )
+	.OnTextCommitted( this, &SGraphNode::OnCommentTextCommitted )
+	.ColorAndOpacity( CommentColor )
+	.AllowPinning( true )
+	.EnableTitleBarBubble( true )
+	.EnableBubbleCtrls( true )
+	.GraphLOD( this, &SGraphNode::GetCurrentLOD );
+
+	GetOrAddSlot( ENodeZone::TopCenter )
+	.SlotOffset( TAttribute<FVector2D>( CommentBubble.Get(), &SCommentBubble::GetOffset ))
+	.SlotSize( TAttribute<FVector2D>( CommentBubble.Get(), &SCommentBubble::GetSize ))
+	.AllowScaling( TAttribute<bool>( CommentBubble.Get(), &SCommentBubble::IsScalingAllowed ))
+	.VAlign( VAlign_Top )
+	[
+		CommentBubble.ToSharedRef()
+	];
+
+	CreateInputSideAddButton(LeftNodeBox);
+	CreateOutputSideAddButton(RightNodeBox);
+}
+
+FOptionalSize SNiagaraGraphNode::GetCompactMinDesiredPinBoxSize() const
+{
+	// Uncomment this if we want pin box sizes to sync again. Will make the right side grow to the same size as the left side, wasting space, but keeping the title central
+	// if(bSyncPinBoxSizes)
+	// {
+	// 	return CompactNodeMaxPinBoxX;
+	// }
+
+	return FOptionalSize();
+}
+
+TAttribute<FSlateFontInfo> SNiagaraGraphNode::GetCompactNodeTitleFont()
+{
+	const FTextBlockStyle& CompactStyle = FAppStyle::GetWidgetStyle<FTextBlockStyle>("Graph.CompactNode.Title");
+	
+	if(CompactNodeTitleFontSizeOverride.IsSet())
+	{
+		FSlateFontInfo FontInfo = CompactStyle.Font;
+		FontInfo.Size = CompactNodeTitleFontSizeOverride.GetValue();
+		return FontInfo;
+	}
+
+	return CompactStyle.Font;
+}
+
+void SNiagaraGraphNode::LoadCachedIcons()
+{
+	static const FName PromotableTypeOuterName("Kismet.VariableList.PromotableTypeOuterIcon");
+	static const FName PromotableTypeInnerName("Kismet.VariableList.PromotableTypeInnerIcon");
+
+	// Outer ring icons
+	if(!CachedOuterIcon)
+	{
+		CachedOuterIcon = FAppStyle::GetBrush(PromotableTypeOuterName);
+	}
+
+	if(!CachedInnerIcon)
+	{
+		CachedInnerIcon = FAppStyle::GetBrush(PromotableTypeInnerName);
+	}
+}
+
+#undef LOCTEXT_NAMESPACE
diff --git a/Engine/Plugins/FX/Niagara/Source/NiagaraEditor/Private/Widgets/SNiagaraGraphNode.h b/Engine/Plugins/FX/Niagara/Source/NiagaraEditor/Private/Widgets/SNiagaraGraphNode.h
--- a/Engine/Plugins/FX/Niagara/Source/NiagaraEditor/Private/Widgets/SNiagaraGraphNode.h
+++ b/Engine/Plugins/FX/Niagara/Source/NiagaraEditor/Private/Widgets/SNiagaraGraphNode.h
@@ -22,10 +22,34 @@ public:
 	virtual void CreateInputSideAddButton(TSharedPtr<SVerticalBox> InputBox) override;
 	virtual void CreateOutputSideAddButton(TSharedPtr<SVerticalBox> OutputBox) override;
 	virtual void UpdateErrorInfo() override;
+	virtual void CreatePinWidgets() override;
 
 protected:
+	virtual void Tick(const FGeometry& AllottedGeometry, const double InCurrentTime, const float InDeltaTime) override;
+
+	void UpdateGraphNodeCompact();
+	virtual bool ShouldDrawCompact() const { return CompactTitle.IsEmpty() == false; }
+	
+	FText GetNodeCompactTitle() const { return CompactTitle; }
+	FOptionalSize GetCompactMinDesiredPinBoxSize() const;
+	/** To allow customization of the node title font size */
+	TAttribute<FSlateFontInfo> GetCompactNodeTitleFont();
+
+	void LoadCachedIcons();
+
 	void RegisterNiagaraGraphNode(UEdGraphNode* InNode);
 	void HandleNiagaraNodeChanged(UNiagaraNode* InNode);
 	TWeakObjectPtr<UNiagaraNode> NiagaraNode;
 	FGuid LastSyncedNodeChangeId;
+	
+	FText CompactTitle;
+	bool bShowPinNamesInCompactMode = false;
+	TOptional<float> CompactNodeTitleFontSizeOverride;
+	
+	static const FSlateBrush* CachedOuterIcon;
+	static const FSlateBrush* CachedInnerIcon;
+
+	// We keep track of the bigger of the two pin boxes so that the smaller one can scale up, so that the middle part is centered
+	double CompactNodeMaxPinBoxX;
+	bool bSyncPinBoxSizes = false;
 };
\ No newline at end of file
diff --git a/Engine/Plugins/FX/Niagara/Source/NiagaraEditor/Public/NiagaraEditorCommon.h b/Engine/Plugins/FX/Niagara/Source/NiagaraEditor/Public/NiagaraEditorCommon.h
--- a/Engine/Plugins/FX/Niagara/Source/NiagaraEditor/Public/NiagaraEditorCommon.h
+++ b/Engine/Plugins/FX/Niagara/Source/NiagaraEditor/Public/NiagaraEditorCommon.h
@@ -27,12 +27,15 @@ public:
 		, bSupportsAddedInputs(false)
 		, bNumericsCanBeIntegers(true)
 		, bNumericsCanBeFloats(true)
+		, bShowPinNamesInCompactMode(false)
 		, bSupportsStaticResolution(false)
 	{}
 
 	FName Name;
 	FText Category;
 	FText FriendlyName;
+	FText CompactName;
+	TOptional<float> CompactNameFontSizeOverride;
 	FText Description;
 	FText Keywords;
 	ENiagaraNumericOutputTypeSelectionMode NumericOuputTypeSelectionMode;
@@ -49,6 +52,9 @@ public:
 	/** If float pins are allowed on this op's numeric pins. */
 	bool bNumericsCanBeFloats;
 
+	/** Only applicable if CompactName is set as that will activate compact node mode. */
+	bool bShowPinNamesInCompactMode;
+
 	/** 
 	* The format that can generate the hlsl for the given number of inputs.
 	* Used the placeholder {A} and {B} to chain the inputs together.
diff --git a/Engine/Plugins/FX/Niagara/Source/NiagaraEditor/Public/NiagaraNode.h b/Engine/Plugins/FX/Niagara/Source/NiagaraEditor/Public/NiagaraNode.h
--- a/Engine/Plugins/FX/Niagara/Source/NiagaraEditor/Public/NiagaraNode.h
+++ b/Engine/Plugins/FX/Niagara/Source/NiagaraEditor/Public/NiagaraNode.h
@@ -61,6 +61,11 @@ public:
 	NIAGARAEDITOR_API virtual bool CanCreateUnderSpecifiedSchema(const UEdGraphSchema* Schema) const override;
 	//~ End EdGraphNode Interface
 
+	/** If this does not return an empty title, compact mode will be activated putting the text into the center of the node. Should be short. */
+	NIAGARAEDITOR_API virtual FText GetCompactTitle() const { return FText::GetEmpty(); }
+	NIAGARAEDITOR_API virtual bool ShouldShowPinNamesInCompactMode() { return false; }
+	NIAGARAEDITOR_API virtual TOptional<float> GetCompactModeFontSizeOverride() const { return {}; }
+
 	/** Virtual function to allow for custom widgets in the input or output box */
 	NIAGARAEDITOR_API virtual void AddWidgetsToInputBox(TSharedPtr<SVerticalBox> InputBox);
 	NIAGARAEDITOR_API virtual void AddWidgetsToOutputBox(TSharedPtr<SVerticalBox> OutputBox);
@@ -74,7 +79,7 @@ public:
 
 	/** Gets the asset referenced by this node, or nullptr if there isn't one. */
 	virtual UObject* GetReferencedAsset() const { return nullptr; }
-	virtual void OpenReferencedAsset() const {};
+	virtual void OpenReferencedAsset() const {}
 
 	/** Refreshes the node due to external changes, e.g. the underlying function changed for a function call node. Return true if the graph changed.*/
 	virtual bool RefreshFromExternalChanges() { return false; }
diff --git a/Engine/Plugins/FX/Niagara/Source/NiagaraEditor/Public/NiagaraNodeOp.h b/Engine/Plugins/FX/Niagara/Source/NiagaraEditor/Public/NiagaraNodeOp.h
--- a/Engine/Plugins/FX/Niagara/Source/NiagaraEditor/Public/NiagaraNodeOp.h
+++ b/Engine/Plugins/FX/Niagara/Source/NiagaraEditor/Public/NiagaraNodeOp.h
@@ -65,7 +65,9 @@ public:
 	virtual ENiagaraNumericOutputTypeSelectionMode GetNumericOutputTypeSelectionMode() const override;
 	virtual bool GenerateCompileHashForClassMembers(const UClass* InClass, FNiagaraCompileHashVisitor* InVisitor) const override;
 	virtual void BuildParameterMapHistory(FNiagaraParameterMapHistoryBuilder& OutHistory, bool bRecursive = true, bool bFilterForCompilation = true) const;
-
+	virtual FText GetCompactTitle() const override;
+	virtual bool ShouldShowPinNamesInCompactMode() override;
+	virtual TOptional<float> GetCompactModeFontSizeOverride() const override;
 	//~ End UNiagaraNode Interface
 
 	//~ Begin UNiagaraNodeWithDynamicPins Interface
