#UE Make GetCachedHeightTextureUv have a unique name per parameter

diff --git a/Engine/Plugins/FX/Niagara/Shaders/Private/NiagaraDataInterfaceLandscape.ush b/Engine/Plugins/FX/Niagara/Shaders/Private/NiagaraDataInterfaceLandscape.ush
--- a/Engine/Plugins/FX/Niagara/Shaders/Private/NiagaraDataInterfaceLandscape.ush
+++ b/Engine/Plugins/FX/Niagara/Shaders/Private/NiagaraDataInterfaceLandscape.ush
@@ -56,7 +56,7 @@ Texture2D<uint>		{ParameterName}_CachedPhysMatTexture;
 int2				{ParameterName}_CachedPhysMatTextureDimension;
 float3				{ParameterName}_SystemLWCTile;
 
-float2 GetCachedHeightTextureUv(float3 InWorldPos)
+float2 GetCachedHeightTextureUv_{ParameterName}(float3 InWorldPos)
 {
 	FLWCVector4 LwcWorldPos = MakeLWCVector4({ParameterName}_SystemLWCTile, float4(InWorldPos, 1.0f));
 	FLWCInverseMatrix LwcWorldToUv = MakeLWCInverseMatrix({ParameterName}_CachedHeightTextureLWCTile, {ParameterName}_CachedHeightTextureWorldToUvTransform);
@@ -137,7 +137,7 @@ void GetHeight_{ParameterName}(float3 InWorldPos, out float OutHeight, out bool
 	}
 	else if ({ParameterName}_CachedHeightTextureEnabled != 0)
 	{
-		float2 Uv = GetCachedHeightTextureUv(InWorldPos);
+		float2 Uv = GetCachedHeightTextureUv_{ParameterName}(InWorldPos);
 
 		OutHeight = {ParameterName}_CachedHeightTexture.SampleLevel({ParameterName}_CachedHeightTextureSampler, Uv, 0).x;
 		OutIsValid = true;
@@ -202,7 +202,7 @@ void GetWorldNormal_{ParameterName}(float3 InWorldPos, out float3 OutNormal, out
 	}
 	else if ({ParameterName}_CachedHeightTextureEnabled != 0)
 	{
-		float2 Uv = GetCachedHeightTextureUv(InWorldPos);
+		float2 Uv = GetCachedHeightTextureUv_{ParameterName}(InWorldPos);
 
 		float Left = {ParameterName}_CachedHeightTexture.SampleLevel({ParameterName}_PointClampedSampler, Uv, 0, int2(-1,0)).x;
 		float Right = {ParameterName}_CachedHeightTexture.SampleLevel({ParameterName}_PointClampedSampler, Uv, 0, int2(1,0)).x;
@@ -230,7 +230,7 @@ void GetWorldNormal_{ParameterName}(float3 InWorldPos, out float3 OutNormal, out
 
 void GetPhysicalMaterialIndex_{ParameterName}(float3 InWorldPos, out int OutIndex, out bool OutIsValid)
 {
-	float2 Uv = GetCachedHeightTextureUv(InWorldPos);
+	float2 Uv = GetCachedHeightTextureUv_{ParameterName}(InWorldPos);
 
 	int3 SampleIndex = int3(Uv * {ParameterName}_CachedPhysMatTextureDimension, 0);
 	OutIndex = (int){ParameterName}_CachedPhysMatTexture.Load(SampleIndex).x;
