Adjustments to new Niagara compact mode UI

diff --git a/Engine/Plugins/FX/Niagara/Source/NiagaraEditor/Private/NiagaraEditorCommon.cpp b/Engine/Plugins/FX/Niagara/Source/NiagaraEditor/Private/NiagaraEditorCommon.cpp
--- a/Engine/Plugins/FX/Niagara/Source/NiagaraEditor/Private/NiagaraEditorCommon.cpp
+++ b/Engine/Plugins/FX/Niagara/Source/NiagaraEditor/Private/NiagaraEditorCommon.cpp
@@ -143,7 +143,7 @@ void FNiagaraOpInfo::Init()
 		Op->Category = NumericCategory;
 		Op->FriendlyName = NSLOCTEXT("NiagaraOpInfo", "Add Name", "Add");
 		Op->CompactName = FText::FromString(TEXT("+"));
-		Op->CompactNameFontSizeOverride = 25.f;
+		Op->CompactNameFontSizeOverride = 30.f;
 		Op->Description = NSLOCTEXT("NiagaraOpInfo", "Add Desc", "Result = A + B");
 		Op->Keywords = FText::FromString(TEXT("+"));
 		Op->Inputs.Add(FNiagaraOpInOutInfo(A, NumericType, AText, AText, DefaultStr_Zero));
@@ -253,7 +253,7 @@ void FNiagaraOpInfo::Init()
 		Op->Category = NumericCategory;
 		Op->FriendlyName = NSLOCTEXT("NiagaraOpInfo", "Multiply Name", "Multiply");
 		Op->CompactName = FText::FromString(TEXT("\xD7"));
-		Op->CompactNameFontSizeOverride = 25.f;
+		Op->CompactNameFontSizeOverride = 30.f;
 		Op->Description = NSLOCTEXT("NiagaraOpInfo", "Multiply Desc", "Result = A * B");
 		Op->Keywords = FText::FromString(TEXT("*"));
 		Op->Inputs.Add(FNiagaraOpInOutInfo(A, NumericType, AText, AText, DefaultStr_One));
@@ -412,7 +412,8 @@ void FNiagaraOpInfo::Init()
 		Op = &OpInfos[Idx];
 		Op->Category = NumericCategory;
 		Op->FriendlyName = NSLOCTEXT("NiagaraOpInfo", "Reciprocal Name", "Reciprocal");
-		Op->CompactName = FText::FromString("Rcp");
+		Op->CompactName = FText::FromString("Reciprocal");
+		Op->CompactNameFontSizeOverride = 18.f;
 		Op->Description = NSLOCTEXT("NiagaraOpInfo", "Reciprocal Desc", "More accurate than Reciprocal Fast. Result = 1 / A");
 		Op->Inputs.Add(FNiagaraOpInOutInfo(A, NumericType, AText, AText, DefaultStr_One));
 		Op->Outputs.Add(FNiagaraOpInOutInfo(Result, NumericType, ResultText, ResultText, DefaultStr_One, TEXT("Reciprocal({0})")));
@@ -1121,6 +1122,8 @@ void FNiagaraOpInfo::Init()
 		Op = &OpInfos[Idx];
 		Op->Category = NumericCategory;
 		Op->FriendlyName = NSLOCTEXT("NiagaraOpInfo", "Normalize Name", "Normalize");
+		Op->CompactName = FText::FromString("Normalize");
+		Op->CompactNameFontSizeOverride = 18.f;
 		Op->Description = NSLOCTEXT("NiagaraOpInfo", "Normalize Desc", "Normalizes the passed value.");
 		Op->Inputs.Add(FNiagaraOpInOutInfo(A, NumericType, AText, AText, DefaultStr_One));
 		Op->Outputs.Add(FNiagaraOpInOutInfo(Result, NumericType, ResultText, ResultText, DefaultStr_One, TEXT("normalize({0})")));
diff --git a/Engine/Plugins/FX/Niagara/Source/NiagaraEditor/Private/Widgets/SNiagaraGraphNode.cpp b/Engine/Plugins/FX/Niagara/Source/NiagaraEditor/Private/Widgets/SNiagaraGraphNode.cpp
--- a/Engine/Plugins/FX/Niagara/Source/NiagaraEditor/Private/Widgets/SNiagaraGraphNode.cpp
+++ b/Engine/Plugins/FX/Niagara/Source/NiagaraEditor/Private/Widgets/SNiagaraGraphNode.cpp
@@ -136,26 +136,6 @@ void SNiagaraGraphNode::CreatePinWidgets()
 	}
 }
 
-void SNiagaraGraphNode::Tick(const FGeometry& AllottedGeometry, const double InCurrentTime, const float InDeltaTime)
-{
-	SGraphNode::Tick(AllottedGeometry, InCurrentTime, InDeltaTime);
-
-	if(CompactTitle.IsEmpty() == false)
-	{
-		FVector2D LeftNodeBoxSize = LeftNodeBox->GetDesiredSize();
-		FVector2D RightNodeBoxSize = RightNodeBox->GetDesiredSize();
-
-		if(LeftNodeBoxSize.X > RightNodeBoxSize.X)
-		{
-			CompactNodeMaxPinBoxX = LeftNodeBoxSize.X;
-		}
-		else
-		{
-			CompactNodeMaxPinBoxX = RightNodeBoxSize.X;
-		}
-	}
-}
-
 void SNiagaraGraphNode::UpdateGraphNodeCompact()
 {
 	InputPins.Empty();
@@ -183,10 +163,6 @@ void SNiagaraGraphNode::UpdateGraphNodeCompact()
 		.Text(this, &SNiagaraGraphNode::GetNodeCompactTitle);
 
 	TSharedRef<SHorizontalBox> NodeContent = SNew(SHorizontalBox);
-	
-	// Default to "pure" styling, where we can just center the pins vertically
-	// since don't need to worry about alignment with other nodes
-	float PinPaddingTop = 0.f;
 
 	EVerticalAlignment PinVerticalAlignment = VAlign_Center;
 	
@@ -198,7 +174,6 @@ void SNiagaraGraphNode::UpdateGraphNodeCompact()
 		if(Pin->Direction == EGPD_Input && Pin->bHidden == false)
 		{
 			bAnyInputExists = true;
-			bSyncPinBoxSizes = true;
 			break;
 		}
 	}
@@ -207,14 +182,10 @@ void SNiagaraGraphNode::UpdateGraphNodeCompact()
 		.HAlign(HAlign_Left)
 		.VAlign(PinVerticalAlignment)
 		.AutoWidth()
-		.Padding(/* left */ 0.f, PinPaddingTop, 5.f, /* bottom */ 0.f)
+		.Padding(/* left */ 0.f, 0.f, 5.f, /* bottom */ 0.f)
 		[
-			SNew(SBox)
-			.MinDesiredWidth(this, &SNiagaraGraphNode::GetCompactMinDesiredPinBoxSize)
-			[
-				// LEFT
-				SAssignNew(LeftNodeBox, SVerticalBox)
-			]
+			// LEFT
+			SAssignNew(LeftNodeBox, SVerticalBox)
 		];
 	
 	NodeContent->AddSlot()
@@ -247,14 +218,10 @@ void SNiagaraGraphNode::UpdateGraphNodeCompact()
 		.HAlign(HAlign_Right)
 		.VAlign(PinVerticalAlignment)
 		.AutoWidth()
-		.Padding(5.f, PinPaddingTop, 0.f, 0.f)
+		.Padding(5.f, 0.f, 0.f, 0.f)
 		[
-			SNew(SBox)
-			.MinDesiredWidth(this, &SNiagaraGraphNode::GetCompactMinDesiredPinBoxSize)
-			[
-				// RIGHT
-				SAssignNew(RightNodeBox, SVerticalBox)
-			]
+			// RIGHT
+			SAssignNew(RightNodeBox, SVerticalBox)
 		];
 
 	//
@@ -369,17 +336,6 @@ void SNiagaraGraphNode::UpdateGraphNodeCompact()
 	CreateOutputSideAddButton(RightNodeBox);
 }
 
-FOptionalSize SNiagaraGraphNode::GetCompactMinDesiredPinBoxSize() const
-{
-	// Uncomment this if we want pin box sizes to sync again. Will make the right side grow to the same size as the left side, wasting space, but keeping the title central
-	// if(bSyncPinBoxSizes)
-	// {
-	// 	return CompactNodeMaxPinBoxX;
-	// }
-
-	return FOptionalSize();
-}
-
 TAttribute<FSlateFontInfo> SNiagaraGraphNode::GetCompactNodeTitleFont()
 {
 	const FTextBlockStyle& CompactStyle = FAppStyle::GetWidgetStyle<FTextBlockStyle>("Graph.CompactNode.Title");
diff --git a/Engine/Plugins/FX/Niagara/Source/NiagaraEditor/Private/Widgets/SNiagaraGraphNode.h b/Engine/Plugins/FX/Niagara/Source/NiagaraEditor/Private/Widgets/SNiagaraGraphNode.h
--- a/Engine/Plugins/FX/Niagara/Source/NiagaraEditor/Private/Widgets/SNiagaraGraphNode.h
+++ b/Engine/Plugins/FX/Niagara/Source/NiagaraEditor/Private/Widgets/SNiagaraGraphNode.h
@@ -25,13 +25,10 @@ public:
 	virtual void CreatePinWidgets() override;
 
 protected:
-	virtual void Tick(const FGeometry& AllottedGeometry, const double InCurrentTime, const float InDeltaTime) override;
-
 	void UpdateGraphNodeCompact();
 	virtual bool ShouldDrawCompact() const { return CompactTitle.IsEmpty() == false; }
 	
 	FText GetNodeCompactTitle() const { return CompactTitle; }
-	FOptionalSize GetCompactMinDesiredPinBoxSize() const;
 	/** To allow customization of the node title font size */
 	TAttribute<FSlateFontInfo> GetCompactNodeTitleFont();
 
@@ -48,8 +45,4 @@ protected:
 	
 	static const FSlateBrush* CachedOuterIcon;
 	static const FSlateBrush* CachedInnerIcon;
-
-	// We keep track of the bigger of the two pin boxes so that the smaller one can scale up, so that the middle part is centered
-	double CompactNodeMaxPinBoxX;
-	bool bSyncPinBoxSizes = false;
 };
\ No newline at end of file
