- Add an option to control occlusion queries for Niagara
- Currently it's hard coded to look for custom depth enabled or not, but sometimes you still want occlusion queries to run even though custom depth is enabled

diff --git a/Engine/Plugins/FX/Niagara/Source/Niagara/Private/NiagaraComponent.cpp b/Engine/Plugins/FX/Niagara/Source/Niagara/Private/NiagaraComponent.cpp
--- a/Engine/Plugins/FX/Niagara/Source/Niagara/Private/NiagaraComponent.cpp
+++ b/Engine/Plugins/FX/Niagara/Source/Niagara/Private/NiagaraComponent.cpp
@@ -186,6 +186,7 @@ bool FitsIntoFloat(FVector InValue)
 
 FNiagaraSceneProxy::FNiagaraSceneProxy(UNiagaraComponent* InComponent)
 	: FPrimitiveSceneProxy(InComponent, InComponent->GetAsset() ? InComponent->GetAsset()->GetFName() : FName())
+	, OcclusionQueryMode(InComponent->GetOcclusionQueryMode())
 {
 	if (FNiagaraSystemInstanceControllerConstPtr SystemInstanceController = InComponent->GetSystemInstanceController())
 	{
@@ -419,6 +420,19 @@ void FNiagaraSceneProxy::SetRenderingEnabled(bool bInRenderingEnabled)
 	}
 }
 
+bool FNiagaraSceneProxy::CanBeOccluded() const
+{
+	switch (OcclusionQueryMode)
+	{
+		case ENiagaraOcclusionQueryMode::AlwaysEnabled:		return true;
+		case ENiagaraOcclusionQueryMode::AlwaysDisabled:	return false;
+
+		// TODO account for MaterialRelevance.bDisableDepthTest and MaterialRelevance.bPostMotionBlurTranslucency as well
+		//case ENiagaraOcclusionQueryMode::Default:
+		default:											return !ShouldRenderCustomDepth();
+	}
+}
+
 void FNiagaraSceneProxy::GetDynamicMeshElements(const TArray<const FSceneView*>& Views, const FSceneViewFamily& ViewFamily, uint32 VisibilityMap, FMeshElementCollector& Collector) const
 {
 	SCOPE_CYCLE_COUNTER(STAT_NiagaraOverview_RT);
@@ -3967,7 +3981,16 @@ void UNiagaraComponent::SetAsset(UNiagaraSystem* InAsset, bool bResetExistingOve
 	}
 }
 
-void UNiagaraComponent::SetForceSolo(bool bInForceSolo) 
+void UNiagaraComponent::SetOcclusionQueryMode(ENiagaraOcclusionQueryMode Mode)
+{
+	if (Mode != OcclusionQueryMode)
+	{
+		OcclusionQueryMode = Mode;
+		MarkRenderStateDirty();
+	}
+}
+
+void UNiagaraComponent::SetForceSolo(bool bInForceSolo)
 { 
 	if (bForceSolo != bInForceSolo)
 	{
diff --git a/Engine/Plugins/FX/Niagara/Source/Niagara/Private/NiagaraDebugHud.cpp b/Engine/Plugins/FX/Niagara/Source/Niagara/Private/NiagaraDebugHud.cpp
--- a/Engine/Plugins/FX/Niagara/Source/Niagara/Private/NiagaraDebugHud.cpp
+++ b/Engine/Plugins/FX/Niagara/Source/Niagara/Private/NiagaraDebugHud.cpp
@@ -2897,6 +2897,18 @@ void FNiagaraDebugHud::DrawComponents(FNiagaraWorldManager* WorldManager, UCanva
 							StringBuilder.Appendf(TEXT("TickGroup - %s\n"), *TickingGroupEnum->GetNameStringByValue(SystemInstance->CalculateTickGroup()));
 						}
 						BuildGpuHudInformation(StringBuilder, NiagaraComponent, SystemInstance, World->GetFeatureLevel());
+
+						StringBuilder.Appendf(TEXT("LastRenderTime - %5.2f"), NiagaraComponent->GetLastRenderTime());
+						if (NiagaraComponent->bRenderCustomDepth)
+						{
+							StringBuilder.Append(TEXT(" bRenderCustomDepth"));
+						}
+						switch (NiagaraComponent->GetOcclusionQueryMode())
+						{
+							case ENiagaraOcclusionQueryMode::AlwaysEnabled:		StringBuilder.Append(TEXT(" OcclusionQuery(AlwaysEnabled)")); break;
+							case ENiagaraOcclusionQueryMode::AlwaysDisabled:	StringBuilder.Append(TEXT(" OcclusionQuery(AlwaysDisable)")); break;
+						}
+						StringBuilder.Append(TEXT("\n"));
 					}
 
 					int64 TotalBytes = SystemInstanceController->GetTotalBytesUsed();
diff --git a/Engine/Plugins/FX/Niagara/Source/Niagara/Public/NiagaraComponent.h b/Engine/Plugins/FX/Niagara/Source/Niagara/Public/NiagaraComponent.h
--- a/Engine/Plugins/FX/Niagara/Source/Niagara/Public/NiagaraComponent.h
+++ b/Engine/Plugins/FX/Niagara/Source/Niagara/Public/NiagaraComponent.h
@@ -36,6 +36,17 @@ enum class ENCPoolMethod : uint8;
 using FNiagaraSystemInstanceControllerPtr = TSharedPtr<FNiagaraSystemInstanceController, ESPMode::ThreadSafe>;
 using FNiagaraSystemInstanceControllerConstPtr = TSharedPtr<const FNiagaraSystemInstanceController, ESPMode::ThreadSafe>;
 
+UENUM()
+enum class ENiagaraOcclusionQueryMode : uint8
+{
+	/** Allow Niagara to determine if we allow occlusion queries to be enabled or not. */
+	Default,
+	/** Occlusion queries will always be enabled for Niagara. */
+	AlwaysEnabled,
+	/** Occlusion queries will always be disabled for Niagara. */
+	AlwaysDisabled,
+};
+
 // Called when the particle system is done
 DECLARE_DYNAMIC_MULTICAST_DELEGATE_OneParam(FOnNiagaraSystemFinished, class UNiagaraComponent*, PSystem);
 
@@ -204,6 +215,11 @@ public:
 	/** How to handle pooling for this component instance. */
 	ENCPoolMethod PoolingMethod;
 
+protected:
+	UPROPERTY(EditAnywhere, AdvancedDisplay, BlueprintReadOnly, Category = Rendering)
+	ENiagaraOcclusionQueryMode OcclusionQueryMode = ENiagaraOcclusionQueryMode::Default;
+
+public:
 	NIAGARA_API virtual void Activate(bool bReset = false) override;
 	NIAGARA_API virtual void Deactivate() override;
 	NIAGARA_API virtual void DeactivateImmediate() override;
@@ -215,7 +231,6 @@ public:
 	NIAGARA_API bool IsComplete() const;
 
 private:
-
 	//Internal versions that can be called from the scalability code.
 	//These will behave as expected but will keep the component registered with the scalability manager.
 	NIAGARA_API void ActivateInternal(bool bReset, bool bIsScalabilityCull);
@@ -276,6 +291,12 @@ public:
 	UFUNCTION(BlueprintCallable, Category = Niagara, meta = (DisplayName = "Get Niagara System Asset"))
 	UNiagaraSystem* GetAsset() const { return Asset; }
 
+	UFUNCTION(BlueprintCallable, Category = Niagara)
+	void SetOcclusionQueryMode(ENiagaraOcclusionQueryMode Mode);
+
+	UFUNCTION(BlueprintCallable, Category = Niagara)
+	ENiagaraOcclusionQueryMode GetOcclusionQueryMode() const { return OcclusionQueryMode; }
+
 	UFUNCTION(BlueprintCallable, Category = Niagara, meta = (DisplayName = "Set Forced Solo Mode"))
 	NIAGARA_API void SetForceSolo(bool bInForceSolo);
 
diff --git a/Engine/Plugins/FX/Niagara/Source/Niagara/Public/NiagaraSceneProxy.h b/Engine/Plugins/FX/Niagara/Source/Niagara/Public/NiagaraSceneProxy.h
--- a/Engine/Plugins/FX/Niagara/Source/Niagara/Public/NiagaraSceneProxy.h
+++ b/Engine/Plugins/FX/Niagara/Source/Niagara/Public/NiagaraSceneProxy.h
@@ -9,6 +9,7 @@
 class FNiagaraGpuComputeDispatchInterface;
 class FNiagaraSystemRenderData;
 class UNiagaraComponent;
+enum class ENiagaraOcclusionQueryMode : uint8;
 
 /**
 * Scene proxy for drawing niagara particle simulations.
@@ -75,14 +76,7 @@ private:
 
 	NIAGARA_API virtual void GetDynamicMeshElements(const TArray<const FSceneView*>& Views, const FSceneViewFamily& ViewFamily, uint32 VisibilityMap, FMeshElementCollector& Collector) const override;
 
-
-	virtual bool CanBeOccluded() const override
-	{
-		// TODO account for MaterialRelevance.bDisableDepthTest and MaterialRelevance.bPostMotionBlurTranslucency as well
-		return !ShouldRenderCustomDepth();
-	}
-
-
+	NIAGARA_API virtual bool CanBeOccluded() const override;
 	/** Callback from the renderer to gather simple lights that this proxy wants renderered. */
 	NIAGARA_API virtual void GatherSimpleLights(const FSceneViewFamily& ViewFamily, FSimpleLightArray& OutParticleLights) const override;
 
@@ -104,5 +98,7 @@ private:
 	TStatId SystemStatID;
 
 	FDynamicData DynamicData;
+
+	ENiagaraOcclusionQueryMode OcclusionQueryMode;
 };
 
