Changed 'solo emitter' action to take into account emitter selection in all instances (solo button in overview graph & in timeline)
For this purpose the emitter that triggered the toggle has to be treated like selected, even if it isn't, since it's possible to click the isolate button without an emitter actually being selected.

Behavior is as follows:
Toggles isolation by default for all selected (includes origin emitter) emitters, except if there is a mismatch in isolation state among selected emitters.
If there is a mismatch, all selected emitters will be isolated regardless of previous state.

Examples:
If you click Isolate with 2 selected non-isolate emitters -> 2 emitters will be isolated
If you click Isolate again with 2 of them selected -> 2 emitters stop being isolated
If you now click on a third emitter that is already isolated and you click Isolate on any of them -> 3 emitters will now be isolated (3rd emitter keeps isolation state)

diff --git a/Engine/Plugins/FX/Niagara/Source/NiagaraEditor/Private/Sequencer/NiagaraSequence/NiagaraEmitterTrackEditor.cpp b/Engine/Plugins/FX/Niagara/Source/NiagaraEditor/Private/Sequencer/NiagaraSequence/NiagaraEmitterTrackEditor.cpp
--- a/Engine/Plugins/FX/Niagara/Source/NiagaraEditor/Private/Sequencer/NiagaraSequence/NiagaraEmitterTrackEditor.cpp
+++ b/Engine/Plugins/FX/Niagara/Source/NiagaraEditor/Private/Sequencer/NiagaraSequence/NiagaraEmitterTrackEditor.cpp
@@ -208,11 +208,7 @@ private:
 		TArray<FGuid> EmitterIdsToIsolate;
 		if (EmitterTrack.IsValid())
 		{
-			if (EmitterTrack->GetEmitterHandleViewModel()->GetIsIsolated() == false)
-			{
-				EmitterIdsToIsolate.Add(EmitterTrack->GetEmitterHandleViewModel()->GetId());
-			}
-			EmitterTrack->GetSystemViewModel().IsolateEmitters(EmitterIdsToIsolate);
+			EmitterTrack->GetSystemViewModel().ToggleIsolateEmitterAndSelectedEmitters(EmitterTrack->GetEmitterHandleViewModel()->GetId());
 		}
 		return FReply::Handled();
 	}
diff --git a/Engine/Plugins/FX/Niagara/Source/NiagaraEditor/Private/ViewModels/NiagaraSystemViewModel.cpp b/Engine/Plugins/FX/Niagara/Source/NiagaraEditor/Private/ViewModels/NiagaraSystemViewModel.cpp
--- a/Engine/Plugins/FX/Niagara/Source/NiagaraEditor/Private/ViewModels/NiagaraSystemViewModel.cpp
+++ b/Engine/Plugins/FX/Niagara/Source/NiagaraEditor/Private/ViewModels/NiagaraSystemViewModel.cpp
@@ -2269,6 +2269,58 @@ void FNiagaraSystemViewModel::IsolateSelectedEmitters()
 	IsolateEmitters(EmitterHandles);
 }
 
+void FNiagaraSystemViewModel::ToggleIsolateEmitterAndSelectedEmitters(FGuid OriginEmitter)
+{
+	TArray<FGuid> SelectedEmitterHandles = GetSelectionViewModel()->GetSelectedEmitterHandleIds();
+	// we treat the origin emitter as if he was selected
+	SelectedEmitterHandles.AddUnique(OriginEmitter);
+
+	bool bAnyIsolationMismatch = false;
+
+	bool bContainsIsolatedEmitter = false;
+	bool bContainsNonIsolatedEmitter = false;
+	for (TSharedRef<FNiagaraEmitterHandleViewModel> EmitterHandle : EmitterHandleViewModels)
+	{
+		if (SelectedEmitterHandles.Contains(EmitterHandle->GetId()))
+		{
+			if(EmitterHandle->GetEmitterHandle()->IsIsolated())
+			{
+				bContainsIsolatedEmitter |= true;
+			}
+			else
+			{
+				bContainsNonIsolatedEmitter |= true;
+			}
+
+			if(bContainsNonIsolatedEmitter && bContainsIsolatedEmitter)
+			{
+				bAnyIsolationMismatch = true;
+				break;
+			}
+		}
+	}
+
+	// if emitter isolation state is mismatched, it means we want to isolate all currently selected emitters 
+	if(bAnyIsolationMismatch)
+	{
+		IsolateEmitters(SelectedEmitterHandles);
+	}
+	// if all currently selected emitters have the same state, we want to toggle it instead
+	else
+	{
+		// all emitters are isolated, which means we want to turn emitter isolation off
+		if(bContainsIsolatedEmitter)
+		{
+			IsolateEmitters({});
+		}
+		// all emitters aren't isolated, so we want to turn emitter isolation on
+		else if(bContainsNonIsolatedEmitter)
+		{
+			IsolateEmitters(SelectedEmitterHandles);
+		}
+	}
+}
+
 void FNiagaraSystemViewModel::DisableSelectedEmitters()
 {
 	const TArray<FGuid> EmitterHandles = GetSelectionViewModel()->GetSelectedEmitterHandleIds();
diff --git a/Engine/Plugins/FX/Niagara/Source/NiagaraEditor/Public/ViewModels/NiagaraSystemViewModel.h b/Engine/Plugins/FX/Niagara/Source/NiagaraEditor/Public/ViewModels/NiagaraSystemViewModel.h
--- a/Engine/Plugins/FX/Niagara/Source/NiagaraEditor/Public/ViewModels/NiagaraSystemViewModel.h
+++ b/Engine/Plugins/FX/Niagara/Source/NiagaraEditor/Public/ViewModels/NiagaraSystemViewModel.h
@@ -284,8 +284,14 @@ public:
 	 * @param ReinitMode Defines whether the system should reinitialize and pull in changes or reset and keep its current state.
 	 */
 	void ResetSystem(ETimeResetMode TimeResetMode, EMultiResetMode MultiResetMode, EReinitMode ReinitMode);
-	
-	void IsolateSelectedEmitters();
+
+	/** This is specifically for UI with some unique behavior.
+	 *  If at least 1 selected emitter is currently not in solo-mode, we assume the user wants to override the 'solo selection' so that all currently selected emitters enter solo mode
+	 *  Since the UI functions can be activated without selection (i.e. clicking 'Isolate' in a overview node won't select the emitter), we always have to specify the guid of the emitter that triggered the call.
+	 *  This is so we can work on all selected emitters + the emitter we clicked Isolate on.
+	 */
+	NIAGARAEDITOR_API void ToggleIsolateEmitterAndSelectedEmitters(FGuid OriginEmitter);
+	NIAGARAEDITOR_API void IsolateSelectedEmitters();
 	void DisableSelectedEmitters();
 	
 	/** Compiles the spawn and update scripts. */
diff --git a/Engine/Plugins/FX/Niagara/Source/NiagaraEditorWidgets/Private/SNiagaraOverviewStackNode.cpp b/Engine/Plugins/FX/Niagara/Source/NiagaraEditorWidgets/Private/SNiagaraOverviewStackNode.cpp
--- a/Engine/Plugins/FX/Niagara/Source/NiagaraEditorWidgets/Private/SNiagaraOverviewStackNode.cpp
+++ b/Engine/Plugins/FX/Niagara/Source/NiagaraEditorWidgets/Private/SNiagaraOverviewStackNode.cpp
@@ -806,9 +806,9 @@ FReply SNiagaraOverviewStackNode::OnToggleIsolateButtonClicked()
 	TSharedPtr<FNiagaraEmitterHandleViewModel> EmitterHandleViewModel = EmitterHandleViewModelWeak.Pin();
 	if (EmitterHandleViewModel.IsValid())
 	{
-		bool bShouldBeIsolated = !EmitterHandleViewModel->GetIsIsolated();
-		EmitterHandleViewModel->SetIsIsolated(bShouldBeIsolated);
+		EmitterHandleViewModel->GetOwningSystemViewModel()->ToggleIsolateEmitterAndSelectedEmitters(EmitterHandleViewModel->GetId());
 	}
+
 	return FReply::Handled();
 }
 
