- Fix for mesh renderer MeshBoundsScale being unused
- Fix for incorrect bounds scaling, it was assuming centered bounds

diff --git a/Engine/Plugins/FX/Niagara/Source/Niagara/Private/NiagaraRendererMeshes.cpp b/Engine/Plugins/FX/Niagara/Source/Niagara/Private/NiagaraRendererMeshes.cpp
--- a/Engine/Plugins/FX/Niagara/Source/Niagara/Private/NiagaraRendererMeshes.cpp
+++ b/Engine/Plugins/FX/Niagara/Source/Niagara/Private/NiagaraRendererMeshes.cpp
@@ -73,6 +73,17 @@ struct FNiagaraDynamicDataMesh : public FNiagaraDynamicDataBase
 
 //////////////////////////////////////////////////////////////////////////
 
+static FBox ScaleFBox(const FBox InBox, const FVector3f InMeshScale, const FVector3f InBoundsScale)
+{
+	const FVector MeshScale = FVector(InMeshScale);
+	const FVector BoundsScale = FVector(InBoundsScale);
+	const FVector MeshCenter = InBox.GetCenter() * MeshScale;
+	const FVector MeshExtent = InBox.GetExtent() * MeshScale;
+	return FBox(MeshCenter - (MeshExtent * BoundsScale), MeshCenter + (MeshExtent * BoundsScale));
+}
+
+//////////////////////////////////////////////////////////////////////////
+
 FNiagaraRendererMeshes::FNiagaraRendererMeshes(ERHIFeatureLevel::Type FeatureLevel, const UNiagaraRendererProperties* Props, const FNiagaraEmitterInstance* Emitter)
 	: FNiagaraRenderer(FeatureLevel, Props, Emitter)
 	, MaterialParamValidMask(0)
@@ -83,6 +94,7 @@ FNiagaraRendererMeshes::FNiagaraRendererMeshes(ERHIFeatureLevel::Type FeatureLev
 	const UNiagaraMeshRendererProperties* Properties = CastChecked<const UNiagaraMeshRendererProperties>(Props);
 	SourceMode = Properties->SourceMode;
 	FacingMode = Properties->FacingMode;
+	MeshBoundsScale = FVector3f(Properties->MeshBoundsScale);
 	bLockedAxisEnable = Properties->bLockedAxisEnable;
 	LockedAxis = FVector3f(Properties->LockedAxis);
 	LockedAxisSpace = Properties->LockedAxisSpace;
@@ -597,9 +609,7 @@ void FNiagaraRendererMeshes::InitializeSortInfo(const FParticleMeshRenderData& P
 void FNiagaraRendererMeshes::PreparePerMeshData(FParticleMeshRenderData& ParticleMeshRenderData, const FNiagaraMeshVertexFactory& VertexFactory, const FNiagaraSceneProxy& SceneProxy, const FMeshData& MeshData) const
 {
 	// Calculate pivot offset / culling sphere
-	FBox MeshLocalBounds = MeshData.RenderableMesh->GetLocalBounds();
-	MeshLocalBounds.Min *= FVector(MeshData.Scale);
-	MeshLocalBounds.Max *= FVector(MeshData.Scale);
+	const FBox MeshLocalBounds = ScaleFBox(MeshData.RenderableMesh->GetLocalBounds(), MeshData.Scale, MeshBoundsScale);
 	ParticleMeshRenderData.CullingSphere.Center = MeshLocalBounds.GetCenter();
 	ParticleMeshRenderData.CullingSphere.W = MeshLocalBounds.GetExtent().Length();
 
@@ -977,9 +987,7 @@ void FNiagaraRendererMeshes::SetupElementForGPUScene(
 	{
 		FMemory::Memzero(&GPUSceneRes.GPUWriteParams, sizeof(GPUSceneRes.GPUWriteParams));
 
-		FBox LocalBounds = MeshData.RenderableMesh->GetLocalBounds();
-		LocalBounds.Min *= FVector(MeshData.Scale);
-		LocalBounds.Max *= FVector(MeshData.Scale);
+		const FBox LocalBounds = ScaleFBox(MeshData.RenderableMesh->GetLocalBounds(), MeshData.Scale, MeshBoundsScale);
 
 		GPUSceneRes.GPUWriteParams.Common 					= CommonParameters;
 		GPUSceneRes.GPUWriteParams.ParticleCount 			= NumInstances;
@@ -1106,7 +1114,8 @@ void FNiagaraRendererMeshes::CreateMeshBatchForSection(
 #endif
 	MeshBatch.DepthPriorityGroup = (ESceneDepthPriorityGroup)SceneProxy.GetDepthPriorityGroup(&View);
 
-	FBox LocalBounds = MeshData.RenderableMesh->GetLocalBounds();
+	const FBox LocalBounds = ScaleFBox(MeshData.RenderableMesh->GetLocalBounds(), FVector3f::OneVector, MeshBoundsScale);
+
 	FMeshBatchElement& BatchElement = MeshBatch.Elements[0];
 	if (ParticleMeshRenderData.bUseGPUScene)
 	{
diff --git a/Engine/Plugins/FX/Niagara/Source/Niagara/Public/NiagaraRendererMeshes.h b/Engine/Plugins/FX/Niagara/Source/Niagara/Public/NiagaraRendererMeshes.h
--- a/Engine/Plugins/FX/Niagara/Source/Niagara/Public/NiagaraRendererMeshes.h
+++ b/Engine/Plugins/FX/Niagara/Source/Niagara/Public/NiagaraRendererMeshes.h
@@ -177,6 +177,8 @@ private:
 	uint32 bSubImageBlend : 1;
 	FVector2f SubImageSize = FVector2f::ZeroVector;
 
+	FVector3f MeshBoundsScale = FVector3f::OneVector;
+
 	FVector3f LockedAxis = FVector3f::ZeroVector;
 	ENiagaraMeshLockedAxisSpace LockedAxisSpace;
 
