Added an arbitrary limit of 15 nested categories

diff --git a/Engine/Plugins/FX/Niagara/Source/NiagaraEditor/Private/ViewModels/HierarchyEditor/NiagaraHierarchyViewModelBase.cpp b/Engine/Plugins/FX/Niagara/Source/NiagaraEditor/Private/ViewModels/HierarchyEditor/NiagaraHierarchyViewModelBase.cpp
--- a/Engine/Plugins/FX/Niagara/Source/NiagaraEditor/Private/ViewModels/HierarchyEditor/NiagaraHierarchyViewModelBase.cpp
+++ b/Engine/Plugins/FX/Niagara/Source/NiagaraEditor/Private/ViewModels/HierarchyEditor/NiagaraHierarchyViewModelBase.cpp
@@ -733,6 +733,16 @@ const TArray<TSharedPtr<FNiagaraHierarchyItemViewModelBase>>& FNiagaraHierarchyI
 	return FilteredChildren;
 }
 
+int32 FNiagaraHierarchyItemViewModelBase::GetHierarchyDepth() const
+{
+	if(Parent.IsValid())
+	{
+		return 1 + Parent.Pin()->GetHierarchyDepth();
+	}
+
+	return 0;
+}
+
 void FNiagaraHierarchyItemViewModelBase::AddChildFilter(FOnFilterChild InFilterChild)
 {
 	if(ensure(InFilterChild.IsBound()))
@@ -1483,6 +1493,15 @@ void FNiagaraHierarchyCategoryViewModel::OnDroppedOnInternal(TSharedPtr<FNiagara
 
 void UNiagaraHierarchyViewModelBase::AddCategory(TSharedPtr<FNiagaraHierarchyItemViewModelBase> CategoryParent) const
 {
+	int32 HierarchyDepth = CategoryParent->GetHierarchyDepth();
+	if(HierarchyDepth > 15)
+	{
+		FNotificationInfo Info(LOCTEXT("TooManyNestedCategoriesToastText", "We currently only allow a hierarchy depth of 15."));
+		Info.ExpireDuration = 4.f;
+		FSlateNotificationManager::Get().AddNotification(Info);
+		return;
+	}
+	
 	TSharedPtr<FNiagaraHierarchyItemViewModelBase> ViewModel = CategoryParent->AddNewItem(UNiagaraHierarchyCategory::StaticClass());
 
 	if(ensure(ViewModel.IsValid()))
diff --git a/Engine/Plugins/FX/Niagara/Source/NiagaraEditor/Public/ViewModels/HierarchyEditor/NiagaraHierarchyViewModelBase.h b/Engine/Plugins/FX/Niagara/Source/NiagaraEditor/Public/ViewModels/HierarchyEditor/NiagaraHierarchyViewModelBase.h
--- a/Engine/Plugins/FX/Niagara/Source/NiagaraEditor/Public/ViewModels/HierarchyEditor/NiagaraHierarchyViewModelBase.h
+++ b/Engine/Plugins/FX/Niagara/Source/NiagaraEditor/Public/ViewModels/HierarchyEditor/NiagaraHierarchyViewModelBase.h
@@ -558,7 +558,9 @@ struct FNiagaraHierarchyItemViewModelBase : TSharedFromThis<FNiagaraHierarchyIte
 	
 	template<class DataClass, class ViewModelChildClass>
 	void GetChildrenViewModelsForType(TArray<TSharedPtr<ViewModelChildClass>>& OutChildren, bool bRecursive = false);
-	
+
+	/** Returns the hierarchy depth via number of parents above. */
+	int32 GetHierarchyDepth() const;
 	NIAGARAEDITOR_API bool HasParent(TSharedPtr<FNiagaraHierarchyItemViewModelBase> ParentCandidate, bool bRecursive = false);
 
 	NIAGARAEDITOR_API void AddChild(TSharedPtr<FNiagaraHierarchyItemViewModelBase> Item);
