Added DI 'inline context actions' API for function call nodes.

Using a DI node context provider, it is possible to specify to use these options, override text or icon for the menu button in the node header.
With a separate function you can add options to it

diff --git a/Engine/Plugins/FX/Niagara/Source/NiagaraEditor/Private/NiagaraActions.cpp b/Engine/Plugins/FX/Niagara/Source/NiagaraEditor/Private/NiagaraActions.cpp
--- a/Engine/Plugins/FX/Niagara/Source/NiagaraEditor/Private/NiagaraActions.cpp
+++ b/Engine/Plugins/FX/Niagara/Source/NiagaraEditor/Private/NiagaraActions.cpp
@@ -683,6 +683,34 @@ void INiagaraDataInterfaceNodeActionProvider::GetNodeContextMenuActions(UClass*
 	}
 }
 
+void INiagaraDataInterfaceNodeActionProvider::GetInlineNodeContextMenuActions(UClass* DIClass, UToolMenu* ToolMenu)
+{
+	if (ensure(DIClass))
+	{
+		while (DIClass && DIClass != UObject::StaticClass())
+		{
+			if (TUniquePtr<INiagaraDataInterfaceNodeActionProvider>* Provider = RegisteredActionProviders.Find(DIClass->GetFName()))
+			{
+				(*Provider)->GetInlineNodeContextMenuActionsImpl(ToolMenu);
+			}
+			DIClass = DIClass->GetSuperClass();
+		}
+	}
+}
+
+INiagaraDataInterfaceNodeActionProvider::FInlineMenuDisplayOptions INiagaraDataInterfaceNodeActionProvider::GetInlineMenuDisplayOptions(UClass* DIClass, UEdGraphNode* Source)
+{
+	if(ensure(DIClass))
+	{
+		if (TUniquePtr<INiagaraDataInterfaceNodeActionProvider>* Provider = RegisteredActionProviders.Find(DIClass->GetFName()))
+		{
+			return (*Provider)->GetInlineMenuDisplayOptionsImpl(DIClass, Source);
+		}
+	}
+
+	return {};
+}
+
 void INiagaraDataInterfaceNodeActionProvider::CollectAddPinActions(UClass* DIClass, FNiagaraMenuActionCollector& Collector, UEdGraphPin* AddPin)
 {
 	if (ensure(DIClass))
@@ -777,52 +805,44 @@ void FNiagaraDataInterfaceNodeActionProvider_DataChannelRead::GetNodeContextMenu
 	{
 		return;
 	}
-
-	FToolMenuSection& Section = Menu->AddSection("DataChannelRead", InitForDataChannelHeaderText);
-
+	
 	UNiagaraNodeFunctionCall* FuncNode = CastChecked<UNiagaraNodeFunctionCall>(Context->Node);
 	TWeakObjectPtr<UNiagaraNodeFunctionCall> WeakNode = FuncNode;
 
-	auto CreateNodeContextMenu = [WeakNode](UToolMenu* InNewToolMenu)
+	// as the lambda gives us the tool menu, it's important to use that instead of the menu we are being passed in from the function
+	auto CreateNodeContextMenu = [WeakNode](UToolMenu* InToolMenu)
 	{
-		FToolMenuSection& SubSection = InNewToolMenu->AddSection("InitForDataChannelSection", InitForDataChannelSectionText);
-
-		auto InitForDataChannelSection = [&SubSection, WeakNode](UNiagaraDataChannel* DataChannel)
-		{
-			if (DataChannel->HasAnyFlags(RF_ClassDefaultObject) == false)
-			{
-				TWeakObjectPtr<UNiagaraDataChannel> WeakChannel = DataChannel;
-				auto CreateDataChannelActionEntry = [WeakChannel, WeakNode]()
-				{
-					UNiagaraDataChannel* Channel = WeakChannel.Get();
-					UNiagaraNodeFunctionCall* Node = WeakNode.Get();
+		AddDataChannelInitActions(InToolMenu);
+	};
 
-					if (Channel && Node)
-					{
-						Node->RemoveAllDynamicPins();
-						TConstArrayView<FNiagaraVariable> ChannelVars = Channel->GetVariables();
-						for (const FNiagaraVariable& Var : ChannelVars)
-						{
-							FNiagaraTypeDefinition Type = Var.GetType();
-							if (Type.IsEnum() == false)
-							{
-								Type = FNiagaraTypeDefinition(FNiagaraTypeHelper::GetSWCStruct(Var.GetType().GetScriptStruct()));
-							}
-							FNiagaraVariable SWCVar(Type, Var.GetName());
-							Node->AddParameter(SWCVar, EEdGraphPinDirection::EGPD_Output);
-						}
-					}
-				};
+	FToolMenuSection& Section = Menu->AddSection("DataChannelRead", InitForDataChannelHeaderText);
+	Section.AddSubMenu("InitForDataChannelMenu", InitForDataChannelMenuText, InitForDataChannelMenuTooltipText,	FNewToolMenuDelegate::CreateLambda(CreateNodeContextMenu));
+}
 
-				SubSection.AddMenuEntry(NAME_None, FText::FromString(DataChannel->GetAsset()->GetName()), FText::FromString(DataChannel->GetAsset()->GetName()), FSlateIcon(FAppStyle::GetAppStyleSetName(), "Icons.Import"),
-					FUIAction(FExecuteAction::CreateLambda(CreateDataChannelActionEntry), FCanExecuteAction()));
-			}
-		};
+void FNiagaraDataInterfaceNodeActionProvider_DataChannelRead::GetInlineNodeContextMenuActionsImpl(UToolMenu* ToolMenu) const
+{
+	AddDataChannelInitActions(ToolMenu);
+}
 
-		UNiagaraDataChannel::ForEachDataChannel(InitForDataChannelSection);
-	};
+INiagaraDataInterfaceNodeActionProvider::FInlineMenuDisplayOptions FNiagaraDataInterfaceNodeActionProvider_DataChannelRead::GetInlineMenuDisplayOptionsImpl(UClass* DIClass, UEdGraphNode* Source) const
+{
+	UNiagaraNodeFunctionCall* FunctionCall = Cast<UNiagaraNodeFunctionCall>(Source);
+	if(FunctionCall == nullptr)
+	{
+		return {};
+	}
+	
+	if(FunctionCall->Signature.Name == TEXT("Num"))
+	{
+		return {};
+	}
+	
+	FInlineMenuDisplayOptions DisplayOptions;
+	DisplayOptions.bDisplayInline = true;
+	DisplayOptions.DisplayBrush = FAppStyle::GetBrush("Icons.Edit");
+	DisplayOptions.TooltipText = LOCTEXT("InitWithDataChannel", "Initialize the read node using the selected Niagara Data Channel asset.");
 
-	Section.AddSubMenu("InitForDataChannelMenu", InitForDataChannelMenuText, InitForDataChannelMenuTooltipText,	FNewToolMenuDelegate::CreateLambda(CreateNodeContextMenu));
+	return DisplayOptions;
 }
 
 void FNiagaraDataInterfaceNodeActionProvider_DataChannelRead::CollectAddPinActionsImpl(FNiagaraMenuActionCollector& Collector, UEdGraphPin* AddPin)const
@@ -880,4 +900,52 @@ void FNiagaraDataInterfaceNodeActionProvider_DataChannelRead::CollectAddPinActio
 	UNiagaraDataChannel::ForEachDataChannel(GatherAddPinsForChannel);
 }
 
-#undef LOCTEXT_NAMESPACE
\ No newline at end of file
+void FNiagaraDataInterfaceNodeActionProvider_DataChannelRead::AddDataChannelInitActions(UToolMenu* ToolMenu)
+{
+	UGraphNodeContextMenuContext* Context = ToolMenu->FindContext<UGraphNodeContextMenuContext>();
+
+	if(Context == nullptr || Context->Node == nullptr)
+	{
+		return;
+	}
+	
+	UNiagaraNodeFunctionCall* FuncNode = CastChecked<UNiagaraNodeFunctionCall>(Context->Node);
+	TWeakObjectPtr<UNiagaraNodeFunctionCall> WeakNode = FuncNode;
+
+	FToolMenuSection& MenuSection = ToolMenu->AddSection("DataChannelRead", NiagaraActionsLocal::InitForDataChannelHeaderText);
+	auto InitForDataChannelSection = [&MenuSection, WeakNode](UNiagaraDataChannel* DataChannel)
+	{
+		if (DataChannel->HasAnyFlags(RF_ClassDefaultObject) == false)
+		{
+			TWeakObjectPtr<UNiagaraDataChannel> WeakChannel = DataChannel;
+			auto CreateDataChannelActionEntry = [WeakChannel, WeakNode]()
+			{
+				UNiagaraDataChannel* Channel = WeakChannel.Get();
+				UNiagaraNodeFunctionCall* Node = WeakNode.Get();
+
+				if (Channel && Node)
+				{
+					Node->RemoveAllDynamicPins();
+					TConstArrayView<FNiagaraVariable> ChannelVars = Channel->GetVariables();
+					for (const FNiagaraVariable& Var : ChannelVars)
+					{
+						FNiagaraTypeDefinition Type = Var.GetType();
+						if (Type.IsEnum() == false)
+						{
+							Type = FNiagaraTypeDefinition(FNiagaraTypeHelper::GetSWCStruct(Var.GetType().GetScriptStruct()));
+						}
+						FNiagaraVariable SWCVar(Type, Var.GetName());
+						Node->AddParameter(SWCVar, EEdGraphPinDirection::EGPD_Output);
+					}
+				}
+			};
+
+			MenuSection.AddMenuEntry(NAME_None, FText::FromString(DataChannel->GetAsset()->GetName()), FText::FromString(DataChannel->GetAsset()->GetName()), FSlateIcon(FAppStyle::GetAppStyleSetName(), "Icons.Import"),
+				FUIAction(FExecuteAction::CreateLambda(CreateDataChannelActionEntry), FCanExecuteAction()));
+		}
+	};
+
+	UNiagaraDataChannel::ForEachDataChannel(InitForDataChannelSection);
+}
+
+#undef LOCTEXT_NAMESPACE
diff --git a/Engine/Plugins/FX/Niagara/Source/NiagaraEditor/Private/NiagaraNode.cpp b/Engine/Plugins/FX/Niagara/Source/NiagaraEditor/Private/NiagaraNode.cpp
--- a/Engine/Plugins/FX/Niagara/Source/NiagaraEditor/Private/NiagaraNode.cpp
+++ b/Engine/Plugins/FX/Niagara/Source/NiagaraEditor/Private/NiagaraNode.cpp
@@ -550,7 +550,7 @@ void UNiagaraNode::AddWidgetsToOutputBox(TSharedPtr<SVerticalBox> OutputBox)
 }
 
 void UNiagaraNode::GetNodeContextMenuActions(class UToolMenu* Menu, class UGraphNodeContextMenuContext* Context) const
-{
+{	
 	{
 		FToolMenuSection& Section = Menu->AddSection("Alignment");
 		Section.AddSubMenu(
@@ -584,6 +584,11 @@ bool UNiagaraNode::CanCreateUnderSpecifiedSchema(const UEdGraphSchema* Schema) c
 	return Schema->IsA<UEdGraphSchema_Niagara>();
 }
 
+TSharedRef<SWidget> UNiagaraNode::CreateTitleRightWidget()
+{
+	return SNullWidget::NullWidget;
+}
+
 void UNiagaraNode::MarkNodeRequiresSynchronization(FString Reason, bool bRaiseGraphNeedsRecompile)
 {
 	Modify();
diff --git a/Engine/Plugins/FX/Niagara/Source/NiagaraEditor/Private/NiagaraNodeFunctionCall.cpp b/Engine/Plugins/FX/Niagara/Source/NiagaraEditor/Private/NiagaraNodeFunctionCall.cpp
--- a/Engine/Plugins/FX/Niagara/Source/NiagaraEditor/Private/NiagaraNodeFunctionCall.cpp
+++ b/Engine/Plugins/FX/Niagara/Source/NiagaraEditor/Private/NiagaraNodeFunctionCall.cpp
@@ -29,6 +29,7 @@
 #include "NiagaraNodeStaticSwitch.h"
 #include "NiagaraSettings.h"
 #include "ScopedTransaction.h"
+#include "Widgets/Input/SComboButton.h"
 
 #include UE_INLINE_GENERATED_CPP_BY_NAME(NiagaraNodeFunctionCall)
 
@@ -1202,6 +1203,81 @@ void UNiagaraNodeFunctionCall::GetNodeContextMenuActions(class UToolMenu* Menu,
 	}
 }
 
+TSharedRef<SWidget> UNiagaraNodeFunctionCall::CreateTitleRightWidget()
+{
+	if ( FunctionScript != nullptr || Signature.Inputs.Num() == 0 )
+	{
+		return SNullWidget::NullWidget;
+	}
+
+	if ( Signature.Inputs[0].GetType().IsDataInterface() == false )
+	{
+		return SNullWidget::NullWidget;
+	}
+
+	if ( UClass* DIClass = Cast<UClass>(Signature.Inputs[0].GetType().GetClass()) )
+	{
+		FString MenuBaseName = "GraphEditor.Niagara.InlinePrimaryActions.";
+		FString MenuName = MenuBaseName.Append(DIClass->GetName());
+
+		if(UToolMenus::Get()->IsMenuRegistered(FName(*MenuName)) == false)
+		{
+			UToolMenu* InlinePrimaryMenu = UToolMenus::Get()->RegisterMenu(FName(*MenuName));
+			InlinePrimaryMenu->AddDynamicSection("GetInlineNodeContextActions", FNewToolMenuDelegate::CreateLambda([DIClass, this](UToolMenu* InMenu)
+			{
+				// it's important to pass in the menu from the lambda parameter instead of using the original menu
+				INiagaraDataInterfaceNodeActionProvider::GetInlineNodeContextMenuActions(DIClass, InMenu);
+			}));
+		}
+
+		TSharedPtr<SWidget> ButtonContentWidget = nullptr;
+
+		INiagaraDataInterfaceNodeActionProvider::FInlineMenuDisplayOptions DisplayOptions = INiagaraDataInterfaceNodeActionProvider::GetInlineMenuDisplayOptions(DIClass, this);
+
+		if(DisplayOptions.bDisplayInline == false)
+		{
+			return SNullWidget::NullWidget;
+		}
+		
+		if(DisplayOptions.DisplayBrush != nullptr)
+		{
+			ButtonContentWidget = SNew(SImage)
+				.Image(DisplayOptions.DisplayBrush);
+		}
+		else
+		{
+			ButtonContentWidget = SNew(STextBlock)
+				.Text(DisplayOptions.DisplayName.IsEmpty() ? LOCTEXT("NodeInlineOptionsMenuLabel", "Options") : DisplayOptions.DisplayName);
+		}
+
+		TSharedRef<SWidget> Result = SNew(SComboButton)
+		.ComboButtonStyle(&FAppStyle::GetWidgetStyle<FComboButtonStyle>("SimpleComboButton"))
+		.OnGetMenuContent(FOnGetContent::CreateLambda([this, MenuName, DIClass]()
+		{
+			UGraphNodeContextMenuContext* ContextObject = NewObject<UGraphNodeContextMenuContext>();
+			ContextObject->Init(this->GetGraph(), this, nullptr, false);
+			FToolMenuContext Context(ContextObject);
+
+			// we generate the menu instead of using FindMenu as this will cause the 'instance' of this menu to be created, with dynamic sections being generated.
+			UToolMenu* GeneratedMenu = UToolMenus::Get()->GenerateMenu(FName(*MenuName), Context);
+			return UToolMenus::Get()->GenerateWidget(GeneratedMenu);
+		}))
+		.ButtonContent()
+		[
+			ButtonContentWidget.ToSharedRef()
+		];
+
+		if(DisplayOptions.TooltipText.IsEmpty() == false)
+		{
+			Result->SetToolTipText(DisplayOptions.TooltipText);
+		}
+		
+		return Result;
+	}
+
+	return SNullWidget::NullWidget;
+}
+
 void UNiagaraNodeFunctionCall::CollectAddPinActions(FNiagaraMenuActionCollector& Collector, UEdGraphPin* AddPin)const
 {
 	Super::CollectAddPinActions(Collector, AddPin);
diff --git a/Engine/Plugins/FX/Niagara/Source/NiagaraEditor/Private/Widgets/SNiagaraGraphNode.cpp b/Engine/Plugins/FX/Niagara/Source/NiagaraEditor/Private/Widgets/SNiagaraGraphNode.cpp
--- a/Engine/Plugins/FX/Niagara/Source/NiagaraEditor/Private/Widgets/SNiagaraGraphNode.cpp
+++ b/Engine/Plugins/FX/Niagara/Source/NiagaraEditor/Private/Widgets/SNiagaraGraphNode.cpp
@@ -136,6 +136,16 @@ void SNiagaraGraphNode::CreatePinWidgets()
 	}
 }
 
+TSharedRef<SWidget> SNiagaraGraphNode::CreateTitleRightWidget()
+{
+	if(NiagaraNode.IsValid())
+	{
+		return NiagaraNode->CreateTitleRightWidget();
+	}
+
+	return SGraphNode::CreateTitleRightWidget();
+}
+
 void SNiagaraGraphNode::UpdateGraphNodeCompact()
 {
 	InputPins.Empty();
diff --git a/Engine/Plugins/FX/Niagara/Source/NiagaraEditor/Private/Widgets/SNiagaraGraphNode.h b/Engine/Plugins/FX/Niagara/Source/NiagaraEditor/Private/Widgets/SNiagaraGraphNode.h
--- a/Engine/Plugins/FX/Niagara/Source/NiagaraEditor/Private/Widgets/SNiagaraGraphNode.h
+++ b/Engine/Plugins/FX/Niagara/Source/NiagaraEditor/Private/Widgets/SNiagaraGraphNode.h
@@ -23,6 +23,7 @@ public:
 	virtual void CreateOutputSideAddButton(TSharedPtr<SVerticalBox> OutputBox) override;
 	virtual void UpdateErrorInfo() override;
 	virtual void CreatePinWidgets() override;
+	virtual TSharedRef<SWidget> CreateTitleRightWidget() override;
 
 protected:
 	void UpdateGraphNodeCompact();
diff --git a/Engine/Plugins/FX/Niagara/Source/NiagaraEditor/Public/NiagaraActions.h b/Engine/Plugins/FX/Niagara/Source/NiagaraEditor/Public/NiagaraActions.h
--- a/Engine/Plugins/FX/Niagara/Source/NiagaraEditor/Public/NiagaraActions.h
+++ b/Engine/Plugins/FX/Niagara/Source/NiagaraEditor/Public/NiagaraActions.h
@@ -388,7 +388,16 @@ public:
 class INiagaraDataInterfaceNodeActionProvider
 {
 public:
-
+	/** If a node displays its options 'inline' (title right widget), you can configure the way the combo button content is displayed and give it a tooltip. */
+	struct FInlineMenuDisplayOptions
+	{
+		/** This needs to be true to have any effect, as it will determine widget generation in the title right area of the node. */
+		bool bDisplayInline = false;
+		FText DisplayName;
+		const FSlateBrush* DisplayBrush = nullptr;
+		FText TooltipText;
+	};
+	
 	INiagaraDataInterfaceNodeActionProvider() = default;
 	INiagaraDataInterfaceNodeActionProvider(INiagaraDataInterfaceNodeActionProvider&) = delete;
 	INiagaraDataInterfaceNodeActionProvider(INiagaraDataInterfaceNodeActionProvider&&) = delete;
@@ -399,8 +408,12 @@ public:
 	/** Allows DIs to add context menu actions for the whole node. */
 	virtual void GetNodeContextMenuActionsImpl(UToolMenu* Menu, UGraphNodeContextMenuContext* Context, FNiagaraFunctionSignature Signature) const {}
 
+	virtual void GetInlineNodeContextMenuActionsImpl(UToolMenu* ToolMenu) const {}
+
+	virtual FInlineMenuDisplayOptions GetInlineMenuDisplayOptionsImpl(UClass* DIClass, UEdGraphNode* Source) const { return {}; }
+	
 	/** Allows DIs to add actions for add pins on the DI function call nodes. */
-	virtual void  CollectAddPinActionsImpl(FNiagaraMenuActionCollector& Collector, UEdGraphPin* AddPin) const {}
+	virtual void CollectAddPinActionsImpl(FNiagaraMenuActionCollector& Collector, UEdGraphPin* AddPin) const {}
 
 	template<typename DIClass, typename ActionProviderClass>
 	static void Register();
@@ -413,6 +426,10 @@ public:
 
 	static NIAGARAEDITOR_API void GetNodeContextMenuActions(UClass* DIClass, UToolMenu* Menu, UGraphNodeContextMenuContext* Context, FNiagaraFunctionSignature Signature);
 
+	static NIAGARAEDITOR_API void GetInlineNodeContextMenuActions(UClass* DIClass, UToolMenu* ToolMenu);
+	
+	static NIAGARAEDITOR_API FInlineMenuDisplayOptions GetInlineMenuDisplayOptions(UClass* DIClass, UEdGraphNode* Source);
+	
 	template<typename DIClass>
 	static void CollectAddPinActions(FNiagaraMenuActionCollector& Collector, UEdGraphPin* AddPin);
 
@@ -469,5 +486,10 @@ class FNiagaraDataInterfaceNodeActionProvider_DataChannelRead : public INiagaraD
 public:
 
 	virtual void GetNodeContextMenuActionsImpl(UToolMenu* Menu, UGraphNodeContextMenuContext* Context, FNiagaraFunctionSignature Signature) const override;
+	virtual void GetInlineNodeContextMenuActionsImpl(UToolMenu* ToolMenu) const override;
+	virtual FInlineMenuDisplayOptions GetInlineMenuDisplayOptionsImpl(UClass* DIClass, UEdGraphNode* Source) const override;
 	virtual void CollectAddPinActionsImpl(FNiagaraMenuActionCollector& Collector, UEdGraphPin* AddPin)const override;
+
+private:
+	static void AddDataChannelInitActions(UToolMenu* ToolMenu);
 };
diff --git a/Engine/Plugins/FX/Niagara/Source/NiagaraEditor/Public/NiagaraNode.h b/Engine/Plugins/FX/Niagara/Source/NiagaraEditor/Public/NiagaraNode.h
--- a/Engine/Plugins/FX/Niagara/Source/NiagaraEditor/Public/NiagaraNode.h
+++ b/Engine/Plugins/FX/Niagara/Source/NiagaraEditor/Public/NiagaraNode.h
@@ -61,6 +61,8 @@ public:
 	NIAGARAEDITOR_API virtual bool CanCreateUnderSpecifiedSchema(const UEdGraphSchema* Schema) const override;
 	//~ End EdGraphNode Interface
 
+	virtual TSharedRef<SWidget> CreateTitleRightWidget();
+	
 	/** If this does not return an empty title, compact mode will be activated putting the text into the center of the node. Should be short. */
 	NIAGARAEDITOR_API virtual FText GetCompactTitle() const { return FText::GetEmpty(); }
 	NIAGARAEDITOR_API virtual bool ShouldShowPinNamesInCompactMode() { return false; }
diff --git a/Engine/Plugins/FX/Niagara/Source/NiagaraEditor/Public/NiagaraNodeFunctionCall.h b/Engine/Plugins/FX/Niagara/Source/NiagaraEditor/Public/NiagaraNodeFunctionCall.h
--- a/Engine/Plugins/FX/Niagara/Source/NiagaraEditor/Public/NiagaraNodeFunctionCall.h
+++ b/Engine/Plugins/FX/Niagara/Source/NiagaraEditor/Public/NiagaraNodeFunctionCall.h
@@ -115,6 +115,7 @@ public:
 	virtual void GatherExternalDependencyData(ENiagaraScriptUsage InUsage, const FGuid& InUsageId, TArray<FNiagaraCompileHash>& InReferencedCompileHashes, TArray<FString>& InReferencedObjs) const override;
 	virtual void UpdateCompileHashForNode(FSHA1& HashState) const override;
 	virtual void GetNodeContextMenuActions(class UToolMenu* Menu, class UGraphNodeContextMenuContext* Context) const override;
+	virtual TSharedRef<SWidget> CreateTitleRightWidget() override;
 	//End UNiagaraNode interface
 
 	virtual void UpdateReferencedStaticsHashForNode(FSHA1& HashState) const;
