- Update BufferFormatToPixelFormat & BufferFormatToRenderTargetFormat to take in expected capabilities and allow them to fallback when formats are not available
- Grid Data Interfaces will now verify that the returned information is valid and fail to run if not

diff --git a/Engine/Plugins/FX/Niagara/Source/Niagara/Classes/NiagaraDataInterfaceGrid2DCollection.h b/Engine/Plugins/FX/Niagara/Source/Niagara/Classes/NiagaraDataInterfaceGrid2DCollection.h
--- a/Engine/Plugins/FX/Niagara/Source/Niagara/Classes/NiagaraDataInterfaceGrid2DCollection.h
+++ b/Engine/Plugins/FX/Niagara/Source/Niagara/Classes/NiagaraDataInterfaceGrid2DCollection.h
@@ -24,7 +24,7 @@ struct FGrid2DCollectionRWInstanceData_GameThread
 	int32 NumAttributes = 0;
 	FVector2D CellSize = FVector2D::ZeroVector;
 	FVector2D WorldBBoxSize = FVector2D::ZeroVector;
-	EPixelFormat PixelFormat = EPixelFormat::PF_R32_FLOAT;
+	TOptional<EPixelFormat> PixelFormat;
 #if WITH_EDITORONLY_DATA
 	bool bPreviewGrid = false;
 	FIntVector4 PreviewAttribute = FIntVector4(INDEX_NONE, INDEX_NONE, INDEX_NONE, INDEX_NONE);
@@ -59,7 +59,7 @@ struct FGrid2DCollectionRWInstanceData_RenderThread
 	int32 NumAttributes = 0;
 	FVector2D CellSize = FVector2D::ZeroVector;
 	FVector2D WorldBBoxSize = FVector2D::ZeroVector;
-	EPixelFormat PixelFormat = EPixelFormat::PF_R32_FLOAT;
+	TOptional<EPixelFormat> PixelFormat;
 
 	TArray<TUniquePtr<FGrid2DBuffer>, TInlineAllocator<2>> Buffers;
 	FGrid2DBuffer* CurrentData = nullptr;
diff --git a/Engine/Plugins/FX/Niagara/Source/Niagara/Classes/NiagaraDataInterfaceGrid3DCollection.h b/Engine/Plugins/FX/Niagara/Source/Niagara/Classes/NiagaraDataInterfaceGrid3DCollection.h
--- a/Engine/Plugins/FX/Niagara/Source/Niagara/Classes/NiagaraDataInterfaceGrid3DCollection.h
+++ b/Engine/Plugins/FX/Niagara/Source/Niagara/Classes/NiagaraDataInterfaceGrid3DCollection.h
@@ -26,7 +26,7 @@ struct FGrid3DCollectionRWInstanceData_GameThread
 	int32 TotalNumNamedAttributes = 0;
 	FVector CellSize = FVector::ZeroVector;
 	FVector WorldBBoxSize = FVector::ZeroVector;
-	EPixelFormat PixelFormat = EPixelFormat::PF_R32_FLOAT;
+	TOptional<EPixelFormat> PixelFormat;
 	bool UseRGBATexture = false;
 #if WITH_EDITOR
 	bool bPreviewGrid = false;
@@ -64,7 +64,7 @@ struct FGrid3DCollectionRWInstanceData_RenderThread
 	int32 TotalNumNamedAttributes = 0;
 	FVector CellSize = FVector::ZeroVector;
 	FVector WorldBBoxSize = FVector::ZeroVector;
-	EPixelFormat PixelFormat = EPixelFormat::PF_R32_FLOAT;
+	TOptional<EPixelFormat> PixelFormat;
 	bool UseRGBATexture = false;
 
 	TArray<int32> AttributeIndices;
diff --git a/Engine/Plugins/FX/Niagara/Source/Niagara/Private/NiagaraCommon.cpp b/Engine/Plugins/FX/Niagara/Source/Niagara/Private/NiagaraCommon.cpp
--- a/Engine/Plugins/FX/Niagara/Source/Niagara/Private/NiagaraCommon.cpp
+++ b/Engine/Plugins/FX/Niagara/Source/Niagara/Private/NiagaraCommon.cpp
@@ -1159,28 +1159,106 @@ FString FNiagaraUtilities::SystemInstanceIDToString(FNiagaraSystemInstanceID ID)
 	return FString(Buffer);
 }
 
-EPixelFormat FNiagaraUtilities::BufferFormatToPixelFormat(ENiagaraGpuBufferFormat NiagaraFormat)
+TOptional<EPixelFormat> FNiagaraUtilities::BufferFormatToPixelFormat(ENiagaraGpuBufferFormat NiagaraFormat, EPixelFormatCapabilities RequiredCapabilities, int NumberOfChannels)
 {
-	switch (NiagaraFormat)
+	if (!ensure(NumberOfChannels >= 1 && NumberOfChannels <= 4))
 	{
-		case ENiagaraGpuBufferFormat::Float:					return EPixelFormat::PF_R32_FLOAT;
-		case ENiagaraGpuBufferFormat::HalfFloat:				return EPixelFormat::PF_R16F;
-		case ENiagaraGpuBufferFormat::UnsignedNormalizedByte:	return EPixelFormat::PF_R8;
+		UE_LOG(LogNiagara, Error, TEXT("Invalid number of channels %d"), NumberOfChannels);
+		return TOptional<EPixelFormat>();
 	}
-	UE_LOG(LogNiagara, Error, TEXT("NiagaraFormat(%d) is invalid, returning float format"), int(NiagaraFormat));
-	return EPixelFormat::PF_R32_FLOAT;
+
+	static_assert(int(ENiagaraGpuBufferFormat::Float) == 0);
+	static_assert(int(ENiagaraGpuBufferFormat::HalfFloat) == 1);
+	static_assert(int(ENiagaraGpuBufferFormat::UnsignedNormalizedByte) == 2);
+	static const EPixelFormat FormatTable[4][int(ENiagaraGpuBufferFormat::Max)] =
+	{
+		{EPixelFormat::PF_R32_FLOAT,		EPixelFormat::PF_R16F,		EPixelFormat::PF_R8},
+		{EPixelFormat::PF_G32R32F,			EPixelFormat::PF_G16R16F,	EPixelFormat::PF_R8G8},
+		{EPixelFormat::PF_A32B32G32R32F,	EPixelFormat::PF_FloatRGBA,	EPixelFormat::PF_R8G8B8A8},
+		{EPixelFormat::PF_A32B32G32R32F,	EPixelFormat::PF_FloatRGBA,	EPixelFormat::PF_R8G8B8A8},
+	};
+
+	static const TPair<EPixelFormat, EPixelFormat> FormatFallbacks[] =
+	{
+		TPairInitializer(EPixelFormat::PF_R8,			EPixelFormat::PF_R16F),
+		TPairInitializer(EPixelFormat::PF_R16F,			EPixelFormat::PF_R32_FLOAT),
+		TPairInitializer(EPixelFormat::PF_R32_FLOAT,	EPixelFormat::PF_G32R32F),
+
+		TPairInitializer(EPixelFormat::PF_R8G8,			EPixelFormat::PF_G16R16F),
+		TPairInitializer(EPixelFormat::PF_G16R16F,		EPixelFormat::PF_G32R32F),
+		TPairInitializer(EPixelFormat::PF_G32R32F,		EPixelFormat::PF_A32B32G32R32F),
+
+		TPairInitializer(EPixelFormat::PF_R8G8B8A8,		EPixelFormat::PF_FloatRGBA),
+		TPairInitializer(EPixelFormat::PF_FloatRGBA,	EPixelFormat::PF_A32B32G32R32F),
+	};
+
+	EPixelFormat CurrentFormat = FormatTable[NumberOfChannels - 1][int(NiagaraFormat)];
+	while (true)
+	{
+		if (UE::PixelFormat::HasCapabilities(CurrentFormat, RequiredCapabilities))
+		{
+			return CurrentFormat;
+		}
+
+		bool bKeepTrying = false;
+		for (const TPair<EPixelFormat, EPixelFormat>& Fallback : FormatFallbacks)
+		{
+			if (Fallback.Key == CurrentFormat)
+			{
+				CurrentFormat = Fallback.Value;
+				bKeepTrying = true;
+				break;
+			}
+		}
+		if (bKeepTrying == false)
+		{
+			break;
+		}
+	};
+
+	UE_LOG(LogNiagara, Error, TEXT("Could not find a valid EPixelFormat for NiagaraFormat(%d) Channels(%d) Caps(%08x)"), int(NiagaraFormat), NumberOfChannels, int(RequiredCapabilities));
+	return TOptional<EPixelFormat>();
 }
 
-ETextureRenderTargetFormat FNiagaraUtilities::BufferFormatToRenderTargetFormat(ENiagaraGpuBufferFormat NiagaraFormat)
+TOptional<ETextureRenderTargetFormat> FNiagaraUtilities::BufferFormatToRenderTargetFormat(ENiagaraGpuBufferFormat NiagaraFormat, EPixelFormatCapabilities RequiredCapabilities)
 {
 	switch (NiagaraFormat)
 	{
-		case ENiagaraGpuBufferFormat::Float:					return ETextureRenderTargetFormat::RTF_R32f;
-		case ENiagaraGpuBufferFormat::HalfFloat:				return ETextureRenderTargetFormat::RTF_R16f;
-		case ENiagaraGpuBufferFormat::UnsignedNormalizedByte:	return ETextureRenderTargetFormat::RTF_R8;
+		case ENiagaraGpuBufferFormat::UnsignedNormalizedByte:
+			if (UE::PixelFormat::HasCapabilities(EPixelFormat::PF_R8, RequiredCapabilities))
+			{
+				return ETextureRenderTargetFormat::RTF_R8;
+			}
+			// Intentional fall through to next format to try
+
+		case ENiagaraGpuBufferFormat::HalfFloat:
+			if (UE::PixelFormat::HasCapabilities(EPixelFormat::PF_R16F, RequiredCapabilities))
+			{
+				return ETextureRenderTargetFormat::RTF_R16f;
+			}
+			// Intentional fall through to next format to try
+
+		case ENiagaraGpuBufferFormat::Float:
+			if (UE::PixelFormat::HasCapabilities(EPixelFormat::PF_R32_FLOAT, RequiredCapabilities))
+			{
+				return ETextureRenderTargetFormat::RTF_R32f;
+			}
+			// Intentional fall through to failure
+
+		default:
+			UE_LOG(LogNiagara, Error, TEXT("Could not find a valid ETextureRenderTargetFormat for NiagaraFormat(%d) Caps(%08x)"), int(NiagaraFormat), int(RequiredCapabilities));
+			return TOptional<ETextureRenderTargetFormat>();
 	}
-	UE_LOG(LogNiagara, Error, TEXT("NiagaraFormat(%d) is invalid, returning float format"), int(NiagaraFormat));
-	return ETextureRenderTargetFormat::RTF_R32f;
+}
+
+EPixelFormat FNiagaraUtilities::BufferFormatToPixelFormat(ENiagaraGpuBufferFormat NiagaraFormat)
+{
+	return BufferFormatToPixelFormat(NiagaraFormat, EPixelFormatCapabilities::TypedUAVLoad | EPixelFormatCapabilities::TypedUAVStore).Get(EPixelFormat::PF_R32_FLOAT);
+}
+
+ETextureRenderTargetFormat FNiagaraUtilities::BufferFormatToRenderTargetFormat(ENiagaraGpuBufferFormat NiagaraFormat)
+{
+	return BufferFormatToRenderTargetFormat(NiagaraFormat, EPixelFormatCapabilities::TypedUAVLoad | EPixelFormatCapabilities::TypedUAVStore).Get(ETextureRenderTargetFormat::RTF_R32f);
 }
 
 FString FNiagaraUtilities::SanitizeNameForObjectsAndPackages(const FString& InName)
diff --git a/Engine/Plugins/FX/Niagara/Source/Niagara/Private/NiagaraDataInterfaceGrid2DCollection.cpp b/Engine/Plugins/FX/Niagara/Source/Niagara/Private/NiagaraDataInterfaceGrid2DCollection.cpp
--- a/Engine/Plugins/FX/Niagara/Source/Niagara/Private/NiagaraDataInterfaceGrid2DCollection.cpp
+++ b/Engine/Plugins/FX/Niagara/Source/Niagara/Private/NiagaraDataInterfaceGrid2DCollection.cpp
@@ -46,6 +46,8 @@ BEGIN_SHADER_PARAMETER_STRUCT(FNDIGrid2DShaderParameters, )
 	SHADER_PARAMETER_RDG_TEXTURE_UAV(RWTexture2DArray<float>,	OutputGrid)
 END_SHADER_PARAMETER_STRUCT()
 
+static constexpr EPixelFormatCapabilities GNDIGrid2DFomatCaps = EPixelFormatCapabilities::TypedUAVLoad | EPixelFormatCapabilities::TypedUAVStore | EPixelFormatCapabilities::Texture2D;
+
 const FString UNiagaraDataInterfaceGrid2DCollection::GridName(TEXT("_Grid"));
 const FString UNiagaraDataInterfaceGrid2DCollection::OutputGridName(TEXT("_OutputGrid"));
 const FString UNiagaraDataInterfaceGrid2DCollection::SamplerName(TEXT("_GridSampler"));
@@ -1817,7 +1819,7 @@ bool UNiagaraDataInterfaceGrid2DCollection::InitPerInstanceData(void* PerInstanc
 		BufferFormat = ENiagaraGpuBufferFormat(GNiagaraGrid2DOverrideFormat);
 	}
 
-	InstanceData->PixelFormat = FNiagaraUtilities::BufferFormatToPixelFormat(BufferFormat);
+	InstanceData->PixelFormat = FNiagaraUtilities::BufferFormatToPixelFormat(BufferFormat, GNDIGrid2DFomatCaps);
 
 	if (!FMath::IsNearlyEqual(GNiagaraGrid2DResolutionMultiplier, 1.0f))
 	{
@@ -2381,35 +2383,35 @@ bool UNiagaraDataInterfaceGrid2DCollection::PerInstanceTickPostSimulate(void* Pe
 				BufferFormat = ENiagaraGpuBufferFormat(GNiagaraGrid2DOverrideFormat);
 			}
 
-			InstanceData->UpdateTargetTexture(BufferFormat);
+			bNeedsReset = InstanceData->UpdateTargetTexture(BufferFormat);
 		}
 
 		// Push Updates to Proxy.
 		FNiagaraDataInterfaceProxyGrid2DCollectionProxy* RT_Proxy = GetProxyAs<FNiagaraDataInterfaceProxyGrid2DCollectionProxy>();
 		ENQUEUE_RENDER_COMMAND(FUpdateData)(
 			[RT_Resource=InstanceData->TargetTexture ? InstanceData->TargetTexture->GetResource() : nullptr, RT_Proxy, InstanceID=SystemInstance->GetId(), RT_InstanceData=*InstanceData](FRHICommandListImmediate& RHICmdList)
-		{
-			check(RT_Proxy->SystemInstancesToProxyData_RT.Contains(InstanceID));
-			FGrid2DCollectionRWInstanceData_RenderThread* TargetData = RT_Proxy->SystemInstancesToProxyData_RT.Find(InstanceID);
+			{
+				check(RT_Proxy->SystemInstancesToProxyData_RT.Contains(InstanceID));
+				FGrid2DCollectionRWInstanceData_RenderThread* TargetData = RT_Proxy->SystemInstancesToProxyData_RT.Find(InstanceID);
 
-			TargetData->NumCells = RT_InstanceData.NumCells;
-			TargetData->NumAttributes = RT_InstanceData.NumAttributes;
-			TargetData->CellSize = RT_InstanceData.CellSize;
+				TargetData->NumCells = RT_InstanceData.NumCells;
+				TargetData->NumAttributes = RT_InstanceData.NumAttributes;
+				TargetData->CellSize = RT_InstanceData.CellSize;
 
-			TargetData->Buffers.Empty();
-			TargetData->CurrentData = nullptr;
-			TargetData->DestinationData = nullptr;
+				TargetData->Buffers.Empty();
+				TargetData->CurrentData = nullptr;
+				TargetData->DestinationData = nullptr;
 
-			if (RT_Resource && RT_Resource->TextureRHI.IsValid())
-			{
-				TargetData->RenderTargetToCopyTo = RT_Resource->TextureRHI;
-			}
-			else
-			{
-				TargetData->RenderTargetToCopyTo = nullptr;
+				if (RT_Resource && RT_Resource->TextureRHI.IsValid())
+				{
+					TargetData->RenderTargetToCopyTo = RT_Resource->TextureRHI;
+				}
+				else
+				{
+					TargetData->RenderTargetToCopyTo = nullptr;
+				}
 			}
-		});
-
+		);
 	}
 
 	return false;
@@ -2473,14 +2475,17 @@ bool FGrid2DCollectionRWInstanceData_GameThread::UpdateTargetTexture(ENiagaraGpu
 	{
 		if ( UTextureRenderTarget2DArray* TargetTextureArray = Cast<UTextureRenderTarget2DArray>(TargetTexture) )
 		{
-			const EPixelFormat RenderTargetFormat = FNiagaraUtilities::BufferFormatToPixelFormat(BufferFormat);
-			if (TargetTextureArray->SizeX != NumCells.X || TargetTextureArray->SizeY != NumCells.Y || TargetTextureArray->Slices != NumAttributes || TargetTextureArray->OverrideFormat != RenderTargetFormat)
+			const TOptional<EPixelFormat> RenderTargetFormat = FNiagaraUtilities::BufferFormatToPixelFormat(BufferFormat, GNDIGrid2DFomatCaps);
+			if (RenderTargetFormat.IsSet())
 			{
-				TargetTextureArray->OverrideFormat = RenderTargetFormat;
-				TargetTextureArray->ClearColor = FLinearColor(0.5, 0, 0, 0);
-				TargetTextureArray->InitAutoFormat(NumCells.X, NumCells.Y, NumAttributes);
-				TargetTextureArray->UpdateResourceImmediate(true);
-				return true;
+				if (TargetTextureArray->SizeX != NumCells.X || TargetTextureArray->SizeY != NumCells.Y || TargetTextureArray->Slices != NumAttributes || TargetTextureArray->OverrideFormat != RenderTargetFormat.GetValue())
+				{
+					TargetTextureArray->OverrideFormat = RenderTargetFormat.GetValue();
+					TargetTextureArray->ClearColor = FLinearColor(0.5, 0, 0, 0);
+					TargetTextureArray->InitAutoFormat(NumCells.X, NumCells.Y, NumAttributes);
+					TargetTextureArray->UpdateResourceImmediate(true);
+					return true;
+				}
 			}
 		}
 		else if (UTextureRenderTarget2D* TargetTexture2D = Cast<UTextureRenderTarget2D>(TargetTexture))
@@ -2497,15 +2502,18 @@ bool FGrid2DCollectionRWInstanceData_GameThread::UpdateTargetTexture(ENiagaraGpu
 			{
 				const FNiagaraGrid2DLegacyTiled2DInfo Tiled2DInfo(NumCells, NumAttributes);
 
-				const ETextureRenderTargetFormat RenderTargetFormat  = FNiagaraUtilities::BufferFormatToRenderTargetFormat(BufferFormat);
-				if (TargetTexture2D->SizeX != Tiled2DInfo.Size.X || TargetTexture2D->SizeY != Tiled2DInfo.Size.Y || TargetTexture2D->RenderTargetFormat != RenderTargetFormat)
+				const TOptional<ETextureRenderTargetFormat> RenderTargetFormat  = FNiagaraUtilities::BufferFormatToRenderTargetFormat(BufferFormat, GNDIGrid2DFomatCaps);
+				if ( RenderTargetFormat.IsSet() )
 				{
-					TargetTexture2D->RenderTargetFormat = RenderTargetFormat;
-					TargetTexture2D->ClearColor = FLinearColor(0.5, 0, 0, 0);
-					TargetTexture2D->bAutoGenerateMips = false;
-					TargetTexture2D->InitAutoFormat(Tiled2DInfo.Size.X, Tiled2DInfo.Size.Y);
-					TargetTexture2D->UpdateResourceImmediate(true);
-					return true;
+					if (TargetTexture2D->SizeX != Tiled2DInfo.Size.X || TargetTexture2D->SizeY != Tiled2DInfo.Size.Y || TargetTexture2D->RenderTargetFormat != RenderTargetFormat.GetValue())
+					{
+						TargetTexture2D->RenderTargetFormat = RenderTargetFormat.GetValue();
+						TargetTexture2D->ClearColor = FLinearColor(0.5, 0, 0, 0);
+						TargetTexture2D->bAutoGenerateMips = false;
+						TargetTexture2D->InitAutoFormat(Tiled2DInfo.Size.X, Tiled2DInfo.Size.Y);
+						TargetTexture2D->UpdateResourceImmediate(true);
+						return true;
+					}
 				}
 			}
 		}
@@ -2526,12 +2534,12 @@ void FGrid2DCollectionRWInstanceData_RenderThread::BeginSimulate(FRDGBuilder& Gr
 		}
 	}
 
-	if (DestinationData == nullptr)
+	if (DestinationData == nullptr && PixelFormat.IsSet())
 	{
 		DestinationData = new FGrid2DBuffer();
 		Buffers.Emplace(DestinationData);
 
-		const FRDGTextureDesc TextureDesc = FRDGTextureDesc::Create2DArray(NumCells, PixelFormat, FClearValueBinding::Black, ETextureCreateFlags::ShaderResource | ETextureCreateFlags::UAV, uint16(NumAttributes));
+		const FRDGTextureDesc TextureDesc = FRDGTextureDesc::Create2DArray(NumCells, PixelFormat.GetValue(), FClearValueBinding::Black, ETextureCreateFlags::ShaderResource | ETextureCreateFlags::UAV, uint16(NumAttributes));
 
 		const TCHAR* GridTextureName = TEXT("Grid2D::GridTexture");
 	#if 0
@@ -2589,7 +2597,7 @@ void FNiagaraDataInterfaceProxyGrid2DCollectionProxy::PreStage(const FNDIGpuComp
 		const FNiagaraSimStageData& SimStageData = Context.GetSimStageData();
 		const FIntVector3& ElementCount = SimStageData.DispatchArgs.ElementCount;
 
-		if (ProxyData->ClearBeforeNonIterationStage && (ElementCount.X != ProxyData->NumCells.X || ElementCount.Y != ProxyData->NumCells.Y))
+		if (ProxyData->ClearBeforeNonIterationStage && ProxyData->DestinationData && (ElementCount.X != ProxyData->NumCells.X || ElementCount.Y != ProxyData->NumCells.Y))
 		{
 			AddClearUAVPass(GraphBuilder, ProxyData->DestinationData->GetOrCreateUAV(GraphBuilder), FVector4f(ForceInitToZero));
 		}
diff --git a/Engine/Plugins/FX/Niagara/Source/Niagara/Private/NiagaraDataInterfaceGrid3DCollection.cpp b/Engine/Plugins/FX/Niagara/Source/Niagara/Private/NiagaraDataInterfaceGrid3DCollection.cpp
--- a/Engine/Plugins/FX/Niagara/Source/Niagara/Private/NiagaraDataInterfaceGrid3DCollection.cpp
+++ b/Engine/Plugins/FX/Niagara/Source/Niagara/Private/NiagaraDataInterfaceGrid3DCollection.cpp
@@ -63,6 +63,8 @@ static FAutoConsoleVariableRef CVarNiagaraGrid3DOverrideFormat(
 	ECVF_Default
 );
 
+static constexpr EPixelFormatCapabilities GNDIGrid3DFomatCaps = EPixelFormatCapabilities::TypedUAVLoad | EPixelFormatCapabilities::TypedUAVStore | EPixelFormatCapabilities::Texture3D;
+
 const FString UNiagaraDataInterfaceGrid3DCollection::NumTilesName(TEXT("_NumTiles"));
 const FString UNiagaraDataInterfaceGrid3DCollection::OneOverNumTilesName(TEXT("_OneOverNumTiles"));
 const FString UNiagaraDataInterfaceGrid3DCollection::UnitClampMinName(TEXT("_UnitClampMin"));
@@ -3394,46 +3396,9 @@ bool UNiagaraDataInterfaceGrid3DCollection::InitPerInstanceData(void* PerInstanc
 		InstanceData->NumTiles.X = 1;
 		InstanceData->NumTiles.Y = 1;
 		InstanceData->NumTiles.Z = 1;
+		
+		InstanceData->PixelFormat = FNiagaraUtilities::BufferFormatToPixelFormat(BufferFormat, GNDIGrid3DFomatCaps, NumAttribChannelsFound);
 
-		switch (BufferFormat)
-		{
-		case ENiagaraGpuBufferFormat::Float:
-		{
-			switch (NumAttribChannelsFound)
-			{
-			case 1: InstanceData->PixelFormat = EPixelFormat::PF_R32_FLOAT; break;
-			case 2:	InstanceData->PixelFormat = EPixelFormat::PF_G32R32F; break;
-			case 3: InstanceData->PixelFormat = EPixelFormat::PF_A32B32G32R32F; break;
-			case 4:	InstanceData->PixelFormat = EPixelFormat::PF_A32B32G32R32F; break;
-			}
-			break;
-		}
-		/** 16-bit per channel floating point, range [-65504, 65504] */
-		case ENiagaraGpuBufferFormat::HalfFloat:
-		{
-			switch (NumAttribChannelsFound)
-			{
-			case 1: InstanceData->PixelFormat = EPixelFormat::PF_R16F; break;
-			case 2:	InstanceData->PixelFormat = EPixelFormat::PF_G16R16F; break;
-			case 3: InstanceData->PixelFormat = EPixelFormat::PF_FloatRGBA; break; // #todo(dmp): PF_FloatRGB translates to r11g11b10
-			case 4:	InstanceData->PixelFormat = EPixelFormat::PF_FloatRGBA; break;
-			}
-			break;
-		}
-		/** 8-bit per channel fixed point, range [0, 1]. */
-		case ENiagaraGpuBufferFormat::UnsignedNormalizedByte:
-		{
-			switch (NumAttribChannelsFound)
-			{
-			case 1: InstanceData->PixelFormat = EPixelFormat::PF_R8; break;
-			case 2:	InstanceData->PixelFormat = EPixelFormat::PF_R8G8; break;
-			case 3: InstanceData->PixelFormat = EPixelFormat::PF_R8G8B8A8; break; // #todo(dmp): need a 3 channel variant
-			case 4:	InstanceData->PixelFormat = EPixelFormat::PF_R8G8B8A8; break;
-			}
-
-			break;
-		}
-		}
 		InstanceData->UseRGBATexture = true;
 	}
 	else
@@ -3473,7 +3438,7 @@ bool UNiagaraDataInterfaceGrid3DCollection::InitPerInstanceData(void* PerInstanc
 		check(InstanceData->NumTiles.Z > 0);
 
 		InstanceData->UseRGBATexture = false;
-		InstanceData->PixelFormat = FNiagaraUtilities::BufferFormatToPixelFormat(BufferFormat);
+		InstanceData->PixelFormat = FNiagaraUtilities::BufferFormatToPixelFormat(BufferFormat, GNDIGrid3DFomatCaps);
 	}
 
 	InstanceData->TargetTexture = nullptr;
@@ -4061,15 +4026,18 @@ bool FGrid3DCollectionRWInstanceData_GameThread::UpdateTargetTexture(ENiagaraGpu
 	if (TargetTexture != nullptr)
 	{
 		const FIntVector RTSize(NumCells.X * NumTiles.X, NumCells.Y * NumTiles.Y, NumCells.Z * NumTiles.Z);
-		const EPixelFormat RenderTargetFormat = FNiagaraUtilities::BufferFormatToPixelFormat(BufferFormat);
-		if (TargetTexture->SizeX != RTSize.X || TargetTexture->SizeY != RTSize.Y || TargetTexture->SizeZ != RTSize.Z || TargetTexture->OverrideFormat != RenderTargetFormat)
+		const TOptional<EPixelFormat> RenderTargetFormat = FNiagaraUtilities::BufferFormatToPixelFormat(BufferFormat, GNDIGrid3DFomatCaps);
+		if (RenderTargetFormat.IsSet())
 		{
-			TargetTexture->OverrideFormat = RenderTargetFormat;
-			TargetTexture->ClearColor = FLinearColor(0, 0, 0, 0);
-			TargetTexture->InitAutoFormat(RTSize.X, RTSize.Y, RTSize.Z);
-			TargetTexture->UpdateResourceImmediate(true);
+			if (TargetTexture->SizeX != RTSize.X || TargetTexture->SizeY != RTSize.Y || TargetTexture->SizeZ != RTSize.Z || TargetTexture->OverrideFormat != RenderTargetFormat.GetValue())
+			{
+				TargetTexture->OverrideFormat = RenderTargetFormat.GetValue();
+				TargetTexture->ClearColor = FLinearColor(0, 0, 0, 0);
+				TargetTexture->InitAutoFormat(RTSize.X, RTSize.Y, RTSize.Z);
+				TargetTexture->UpdateResourceImmediate(true);
 
-			return true;
+				return true;
+			}
 		}
 	}
 
@@ -4088,13 +4056,13 @@ void FGrid3DCollectionRWInstanceData_RenderThread::BeginSimulate(FRDGBuilder& Gr
 		}
 	}
 
-	if (DestinationData == nullptr)
+	if (DestinationData == nullptr && PixelFormat.IsSet())
 	{
 		DestinationData = new FGrid3DBuffer();
 		Buffers.Emplace(DestinationData);
 
 		const FIntVector TextureSize(NumCells.X * NumTiles.X, NumCells.Y * NumTiles.Y, NumCells.Z * NumTiles.Z);
-		const FRDGTextureDesc TextureDesc = FRDGTextureDesc::Create3D(TextureSize, PixelFormat, FClearValueBinding::Black, ETextureCreateFlags::ShaderResource | ETextureCreateFlags::UAV);
+		const FRDGTextureDesc TextureDesc = FRDGTextureDesc::Create3D(TextureSize, PixelFormat.GetValue(), FClearValueBinding::Black, ETextureCreateFlags::ShaderResource | ETextureCreateFlags::UAV);
 
 		const TCHAR* GridTextureName = TEXT("Grid3D::GridTexture");
 	#if 0
@@ -4152,6 +4120,7 @@ void FNiagaraDataInterfaceProxyGrid3DCollectionProxy::PreStage(const FNDIGpuComp
 		const FIntVector3& ElementCount = SimStageData.DispatchArgs.ElementCount;
 
 		if (ProxyData->ClearBeforeNonIterationStage && 
+			ProxyData->DestinationData &&
 			(ElementCount.X != ProxyData->NumCells.X || ElementCount.Y != ProxyData->NumCells.Y || ElementCount.Z != ProxyData->NumCells.Z))
 		{
 			AddClearUAVPass(GraphBuilder, ProxyData->DestinationData->GetOrCreateUAV(GraphBuilder), FVector4f(ForceInitToZero));
diff --git a/Engine/Plugins/FX/Niagara/Source/Niagara/Public/NiagaraCommon.h b/Engine/Plugins/FX/Niagara/Source/Niagara/Public/NiagaraCommon.h
--- a/Engine/Plugins/FX/Niagara/Source/Niagara/Public/NiagaraCommon.h
+++ b/Engine/Plugins/FX/Niagara/Source/Niagara/Public/NiagaraCommon.h
@@ -39,6 +39,7 @@ class UNiagaraParameterDefinitionsBase;
 struct FNiagaraParameterStore;
 enum class EPSCPoolMethod : uint8;
 enum ETextureRenderTargetFormat : int;
+enum class EPixelFormatCapabilities : uint32;
 
 namespace ERHIFeatureLevel { enum Type : int; }
 enum EShaderPlatform : uint16;
@@ -1518,15 +1519,26 @@ namespace FNiagaraUtilities
 
 	NIAGARA_API FString SystemInstanceIDToString(FNiagaraSystemInstanceID ID);
 
-	/** Converts a Niagara format into a EPixelFormat */
-	NIAGARA_API EPixelFormat BufferFormatToPixelFormat(ENiagaraGpuBufferFormat NiagaraFormat);
+	/**
+	Converts a Niagara format into a EPixelFormat, will look for fallback options if the requested format can not be satisfied based on the RequiredCapabilities.
+	If the return value optional is invalid that means no format was found.
+	*/
+	NIAGARA_API TOptional<EPixelFormat> BufferFormatToPixelFormat(ENiagaraGpuBufferFormat NiagaraFormat, EPixelFormatCapabilities RequiredCapabilities, int NumberOfChannels = 1);
 
-	/** Converts a Niagara format into a ETextureRenderTargetFormat */
-	NIAGARA_API ETextureRenderTargetFormat BufferFormatToRenderTargetFormat(ENiagaraGpuBufferFormat NiagaraFormat);
+	/**
+	Converts a Niagara format into a ETextureRenderTargetFormat, will look for fallback options if the requested format can not be satisfied based on the RequiredCapabilities.
+	If the return value optional is invalid that means no format was found.
+	*/
+	NIAGARA_API TOptional<ETextureRenderTargetFormat> BufferFormatToRenderTargetFormat(ENiagaraGpuBufferFormat NiagaraFormat, EPixelFormatCapabilities RequiredCapabilities);
 
 	NIAGARA_API FString SanitizeNameForObjectsAndPackages(const FString& InName);
 
 	NIAGARA_API FNiagaraVariable ResolveAliases(const FNiagaraVariable& InVar, const FNiagaraAliasContext& InContext);
+
+	UE_DEPRECATED(5.3, "Please update your code to include EPixelFormatCapabilities.")
+	NIAGARA_API EPixelFormat BufferFormatToPixelFormat(ENiagaraGpuBufferFormat NiagaraFormat);
+	UE_DEPRECATED(5.3, "Please update your code to include EPixelFormatCapabilities.")
+	NIAGARA_API ETextureRenderTargetFormat BufferFormatToRenderTargetFormat(ENiagaraGpuBufferFormat NiagaraFormat);
 };
 
 USTRUCT()
