Remove IsInGameThread check on Niagara Oscilloscope

diff --git a/Engine/Plugins/FX/Niagara/Source/Niagara/Private/NiagaraDataInterfaceAudioOscilloscope.cpp b/Engine/Plugins/FX/Niagara/Source/Niagara/Private/NiagaraDataInterfaceAudioOscilloscope.cpp
--- a/Engine/Plugins/FX/Niagara/Source/Niagara/Private/NiagaraDataInterfaceAudioOscilloscope.cpp
+++ b/Engine/Plugins/FX/Niagara/Source/Niagara/Private/NiagaraDataInterfaceAudioOscilloscope.cpp
@@ -32,7 +32,6 @@ FNiagaraDataInterfaceProxyOscilloscope::FNiagaraDataInterfaceProxyOscilloscope(i
 	, ScopeInMilliseconds(InScopeInMillseconds)
 	, NumChannelsInDownsampledBuffer(0)
 {
-	check(IsInGameThread());
 	DeviceCreatedHandle = FAudioDeviceManagerDelegates::OnAudioDeviceCreated.AddRaw(this, &FNiagaraDataInterfaceProxyOscilloscope::OnNewDeviceCreated);
 	DeviceDestroyedHandle = FAudioDeviceManagerDelegates::OnAudioDeviceDestroyed.AddRaw(this, &FNiagaraDataInterfaceProxyOscilloscope::OnDeviceDestroyed);
 
@@ -403,7 +402,6 @@ void UNiagaraDataInterfaceAudioOscilloscope::PostInitProperties()
 
 void FNiagaraDataInterfaceProxyOscilloscope::OnBeginDestroy()
 {
-	check(IsInGameThread());
 	FAudioDeviceManagerDelegates::OnAudioDeviceCreated.Remove(DeviceCreatedHandle);
 	FAudioDeviceManagerDelegates::OnAudioDeviceDestroyed.Remove(DeviceDestroyedHandle);
 	DeviceCreatedHandle.Reset();
