- Correctly support negative UV Scale with SubUV

diff --git a/Engine/Plugins/FX/Niagara/Shaders/Private/NiagaraSpriteVertexFactory.ush b/Engine/Plugins/FX/Niagara/Shaders/Private/NiagaraSpriteVertexFactory.ush
--- a/Engine/Plugins/FX/Niagara/Shaders/Private/NiagaraSpriteVertexFactory.ush
+++ b/Engine/Plugins/FX/Niagara/Shaders/Private/NiagaraSpriteVertexFactory.ush
@@ -604,11 +604,8 @@ float3x3 CalcTangentBasis(FVertexFactoryIntermediates Intermediates)
 	return Result;
 }
 
-void ComputeBillboardUVs(FVertexFactoryInput Input, float2 ParticleSize, float SubImageIndex, out float2 UVForPosition, out float2 UVForTexturing, out float2 UVForTexturingUnflipped)
+void ComputeBillboardUVs(FVertexFactoryInput Input, float2 ParticleSize, float SubImageIndex, float2 UVScale, out float2 UVForPosition, out float2 UVForTexturing, out float2 UVForTexturingUnflipped)
 {
-	// Encoding the UV flip in the sign of the size data.
-	float2 UVFlip = sign(ParticleSize);
-
 #if FEATURE_LEVEL >= FEATURE_LEVEL_SM4
 	BRANCH
 	if (NiagaraSpriteVFLooseParameters.NumCutoutVerticesPerFrame > 0)
@@ -616,32 +613,37 @@ void ComputeBillboardUVs(FVertexFactoryInput Input, float2 ParticleSize, float S
 		// Avoid uint divide which is extremely slow on GCN
 		const uint VertexIndexMask = NiagaraSpriteVFLooseParameters.NumCutoutVerticesPerFrame == 8 ? 7 : 3; // BVC_FourVertices || BVC_EightVertices
 		uint CutoutVertexIndex = Input.VertexId & VertexIndexMask;
-		float NumFrames = NiagaraSpriteVF.SubImageSize.x * NiagaraSpriteVF.SubImageSize.y;
-		uint SubImageIndexInt = (uint)fmod(SubImageIndex, NumFrames);
+		const float NumFrames = NiagaraSpriteVF.SubImageSize.x * NiagaraSpriteVF.SubImageSize.y;
+		const uint SubImageIndexInt = (uint)fmod(SubImageIndex, NumFrames);
 
+		const float2 UVFlip = sign(ParticleSize) * sign(UVScale);
+		
+		// Reverse the winding order of the polygon when only flipping in one direction to counteract UV flipping
 		FLATTEN
-		if (UVFlip.x * UVFlip.y < 0)
+		if (UVFlip.x * UVFlip.y < 0.0f)
 		{
-			// Reverse the winding order of the polygon when only flipping in one direction to counteract UV flipping
 			CutoutVertexIndex = NiagaraSpriteVFLooseParameters.NumCutoutVerticesPerFrame - 1 - CutoutVertexIndex;
 		}
 
 		// Fetch the billboard space positions from the appropriate frame of the preprocessed cutout geometry
-		UVForTexturing = UVForTexturingUnflipped = NiagaraSpriteVFLooseParameters.CutoutGeometry[SubImageIndexInt * NiagaraSpriteVFLooseParameters.NumCutoutVerticesPerFrame + CutoutVertexIndex];
-
+		UVForPosition = NiagaraSpriteVFLooseParameters.CutoutGeometry[SubImageIndexInt * NiagaraSpriteVFLooseParameters.NumCutoutVerticesPerFrame + CutoutVertexIndex];
+		UVForTexturingUnflipped = UVForPosition;
+		UVForTexturing = UVForPosition;
+		
 		// Invert positions on the billboard so that the cutout geometry still contains the visible portion
 		// This changes the winding order when only one direction is flipped
-		UVForPosition.x = UVFlip.x < 0.0 ? 1.0 - UVForTexturing.x : UVForTexturing.x;
-		UVForPosition.y = UVFlip.y < 0.0 ? 1.0 - UVForTexturing.y : UVForTexturing.y;
+		UVForPosition.x = UVFlip.x < 0.0f ? 1.0f - UVForTexturing.x : UVForTexturing.x;
+		UVForPosition.y = UVFlip.y < 0.0f ? 1.0f - UVForTexturing.y : UVForTexturing.y;
+
+		UVForTexturing = float2(0.5f, 0.5f) + ((UVForPosition - float2(0.5f, 0.5f)) * UVScale * sign(ParticleSize));
 	}
 	else
 #endif
 	{
-		UVForTexturing.x = UVFlip.x < 0.0 ? 1.0 - Input.TexCoord.x : Input.TexCoord.x;
-		UVForTexturing.y = UVFlip.y < 0.0 ? 1.0 - Input.TexCoord.y : Input.TexCoord.y;
-
 		// Note: not inverting positions, as that would change the winding order
-		UVForPosition = UVForTexturingUnflipped = Input.TexCoord.xy;
+		UVForPosition = Input.TexCoord.xy;
+		UVForTexturingUnflipped = UVForPosition;
+		UVForTexturing = float2(0.5f, 0.5f) + ((UVForTexturingUnflipped - float2(0.5f, 0.5f)) * UVScale * sign(ParticleSize));
 	}
 }
 
@@ -878,9 +880,8 @@ FVertexFactoryIntermediates GetVertexFactoryIntermediates(FVertexFactoryInput In
 	float2 UVForPosition;
 	float2 UVForTexturing;
 	float2 UVForTexturingUnflipped;
-	ComputeBillboardUVs(Input, ParticleSize, SubImageIndex, UVForPosition, UVForTexturing, UVForTexturingUnflipped);
+	ComputeBillboardUVs(Input, ParticleSize, SubImageIndex, UVScale, UVForPosition, UVForTexturing, UVForTexturingUnflipped);
 
-	UVForTexturing.xy = float2(0.5f, 0.5f) + (UVForTexturing.xy - float2(0.5f, 0.5f) ) *  UVScale.xy; 
 	const float2 Size = abs(Intermediates.Size.xy);
 
 	// When enabled try to main a pixel sized sprite and optionally fade particles based on coverage
@@ -927,7 +928,7 @@ FVertexFactoryIntermediates GetVertexFactoryIntermediates(FVertexFactoryInput In
 		float2 PrevUVForPosition;
 		float2 PrevUVForTexturing;
 		float2 PrevUVForTexturingUnflipped;
-		ComputeBillboardUVs(Input, PrevParticleSize, SubImageIndex, PrevUVForPosition, PrevUVForTexturing, PrevUVForTexturingUnflipped);
+		ComputeBillboardUVs(Input, PrevParticleSize, SubImageIndex, UVScale, PrevUVForPosition, PrevUVForTexturing, PrevUVForTexturingUnflipped);
 			
 		const float2x3 PrevTangents = float2x3(PrevRight, PrevUp);
 		const float3 PrevVertexOffset = PrevCameraOffset + mul(abs(PrevParticleSize) * (PrevUVForPosition - PivotOffset), PrevTangents);
