- Debug draw for Position Array Data Interface

diff --git a/Engine/Plugins/FX/Niagara/Source/Niagara/Classes/NiagaraDataInterfaceArrayFloat.h b/Engine/Plugins/FX/Niagara/Source/Niagara/Classes/NiagaraDataInterfaceArrayFloat.h
--- a/Engine/Plugins/FX/Niagara/Source/Niagara/Classes/NiagaraDataInterfaceArrayFloat.h
+++ b/Engine/Plugins/FX/Niagara/Source/Niagara/Classes/NiagaraDataInterfaceArrayFloat.h
@@ -57,6 +57,10 @@ class UNiagaraDataInterfaceArrayPosition : public UNiagaraDataInterfaceArray
 public:
 	GENERATED_BODY()
 
+#if WITH_NIAGARA_DEBUGGER
+	virtual void DrawDebugHud(FNDIDrawDebugHudContext& DebugHudContext) const;
+#endif
+
 	UPROPERTY(EditAnywhere, BlueprintReadWrite, Category = "Array")
 	TArray<FNiagaraPosition> PositionData;
 
diff --git a/Engine/Plugins/FX/Niagara/Source/Niagara/Private/NiagaraDataInterfaceArrayFloat.cpp b/Engine/Plugins/FX/Niagara/Source/Niagara/Private/NiagaraDataInterfaceArrayFloat.cpp
--- a/Engine/Plugins/FX/Niagara/Source/Niagara/Private/NiagaraDataInterfaceArrayFloat.cpp
+++ b/Engine/Plugins/FX/Niagara/Source/Niagara/Private/NiagaraDataInterfaceArrayFloat.cpp
@@ -1,5 +1,43 @@
 // Copyright Epic Games, Inc. All Rights Reserved.
 
 #include "NiagaraDataInterfaceArrayFloat.h"
+#include "NiagaraSystemInstance.h"
+
+#include "Engine/Canvas.h"
+#include "DrawDebugHelpers.h"
 
 #include UE_INLINE_GENERATED_CPP_BY_NAME(NiagaraDataInterfaceArrayFloat)
+
+#if WITH_NIAGARA_DEBUGGER
+
+static float GNDIArrayPositionDebugRadius = 32.0f;
+static FAutoConsoleVariableRef CVarNDIArrayPositionDebugRadius(
+	TEXT("fx.Niagara.Array.PositionDebugRadius"),
+	GNDIArrayPositionDebugRadius,
+	TEXT("When using the Niagara Debugger the radius of position array debug display"),
+	ECVF_Default
+);
+
+void UNiagaraDataInterfaceArrayPosition::DrawDebugHud(FNDIDrawDebugHudContext& DebugHudContext) const
+{
+	Super::DrawDebugHud(DebugHudContext);
+
+	if (DebugHudContext.IsVerbose())
+	{
+		UCanvas* Canvas = DebugHudContext.GetCanvas();
+		const FNiagaraLWCConverter LWCConverter = DebugHudContext.GetSystemInstance()->GetLWCConverter();
+		const float Radius = GNDIArrayPositionDebugRadius;
+		for (const FNiagaraPosition& Position : PositionData)
+		{
+			const FVector WorldPosition = LWCConverter.ConvertSimulationPositionToWorld(Position);
+			const FVector ScreenPosition = Canvas->Project(WorldPosition, false);
+			if (ScreenPosition.Z <= 0.0f)
+			{
+				continue;
+			}
+
+			DrawDebugCanvasWireSphere(Canvas, WorldPosition, FColor::Red, Radius, 16);
+		}
+	}
+}
+#endif
diff --git a/Engine/Plugins/FX/Niagara/Source/Niagara/Private/NiagaraDebugHud.cpp b/Engine/Plugins/FX/Niagara/Source/Niagara/Private/NiagaraDebugHud.cpp
--- a/Engine/Plugins/FX/Niagara/Source/Niagara/Private/NiagaraDebugHud.cpp
+++ b/Engine/Plugins/FX/Niagara/Source/Niagara/Private/NiagaraDebugHud.cpp
@@ -153,6 +153,7 @@ namespace NiagaraDebugLocal
 		MakeTuple(TEXT("SystemShowActiveOnlyInWorld="), TEXT("When enabled only active systems are shown in world"), [](FString Arg) {Settings.bSystemShowActiveOnlyInWorld = FCString::Atoi(*Arg) != 0; }),
 		MakeTuple(TEXT("SystemDebugVerbosity="), TEXT("Set the in world system debug verbosity"), [](FString Arg) {Settings.SystemDebugVerbosity = FMath::Clamp(ENiagaraDebugHudVerbosity(FCString::Atoi(*Arg)), ENiagaraDebugHudVerbosity::None, ENiagaraDebugHudVerbosity::Verbose); }),
 		MakeTuple(TEXT("SystemEmitterVerbosity="), TEXT("Set the in world system emitter debug verbosity"), [](FString Arg) {Settings.SystemEmitterVerbosity = FMath::Clamp(ENiagaraDebugHudVerbosity(FCString::Atoi(*Arg)), ENiagaraDebugHudVerbosity::None, ENiagaraDebugHudVerbosity::Verbose); }),
+		MakeTuple(TEXT("DataInterfaceVerbosity="), TEXT("Set the in world system data interface debug verbosity"), [](FString Arg) {Settings.DataInterfaceVerbosity = FMath::Clamp(ENiagaraDebugHudVerbosity(FCString::Atoi(*Arg)), ENiagaraDebugHudVerbosity::None, ENiagaraDebugHudVerbosity::Verbose); }),
 		MakeTuple(TEXT("SystemVariables="), TEXT("Set the system variables to display"), [](FString Arg) {Settings.SystemVariables.Empty(); FNiagaraDebugHUDVariable::InitFromString(Arg, Settings.SystemVariables); GCachedSystemVariables.Empty(); }),
 		MakeTuple(TEXT("ShowSystemVariables="), TEXT("Set system variables visibility"), [](FString Arg) {Settings.bShowSystemVariables = FCString::Atoi(*Arg) != 0; GCachedSystemVariables.Empty(); }),
 
