Allow wildcard pins to accept drops from resolved numeric pins

diff --git a/Engine/Plugins/FX/Niagara/Source/NiagaraEditor/Private/EdGraphSchema_Niagara.cpp b/Engine/Plugins/FX/Niagara/Source/NiagaraEditor/Private/EdGraphSchema_Niagara.cpp
--- a/Engine/Plugins/FX/Niagara/Source/NiagaraEditor/Private/EdGraphSchema_Niagara.cpp
+++ b/Engine/Plugins/FX/Niagara/Source/NiagaraEditor/Private/EdGraphSchema_Niagara.cpp
@@ -1370,6 +1370,22 @@ void UEdGraphSchema_Niagara::TrySetDefaultValue(UEdGraphPin& Pin, const FString&
 	}
 }
 
+FNiagaraTypeDefinition ResolveNumericPinType(const UEdGraphPin* Pin)
+{
+	FNiagaraTypeDefinition PinType = UEdGraphSchema_Niagara::PinToTypeDefinition(Pin);
+	if (PinType == FNiagaraTypeDefinition::GetGenericNumericDef())
+	{
+		if (UNiagaraNode* Node = Cast<UNiagaraNode>(Pin->GetOwningNode()))
+		{
+			if (UNiagaraGraph* NiagaraGraph = Node->GetNiagaraGraph())
+			{
+				return NiagaraGraph->GetCachedNumericConversion(Pin);
+			}
+		}
+	}
+	return PinType;
+}
+
 bool UEdGraphSchema_Niagara::TryCreateConnection(UEdGraphPin* PinA, UEdGraphPin* PinB) const
 {
 	const FScopedTransaction Transaction(NSLOCTEXT("UnrealEd", "NiagaraEditorCreateConnection", "Niagara Editor: Create Connection"));
@@ -1377,8 +1393,8 @@ bool UEdGraphSchema_Niagara::TryCreateConnection(UEdGraphPin* PinA, UEdGraphPin*
 	const FPinConnectionResponse Response = CanCreateConnection(PinA, PinB);
 	bool bModified = false;
 
-	FNiagaraTypeDefinition TypeA = PinToTypeDefinition(PinA);
-	FNiagaraTypeDefinition TypeB = PinToTypeDefinition(PinB);
+	FNiagaraTypeDefinition TypeA = ResolveNumericPinType(PinA);
+	FNiagaraTypeDefinition TypeB = ResolveNumericPinType(PinB);
 	
 	switch (Response.Response)
 	{
@@ -1827,6 +1843,16 @@ FPinConnectionResponse UEdGraphSchema_Niagara::GetWildcardConnectionResponse(con
 		PinBType = TmpType;
 		bPinsSwapped = true;
 	}
+
+	// we can't link to numeric types, but maybe it was already resolved to a usable type
+	if (PinBType == FNiagaraTypeDefinition::GetGenericNumericDef())
+	{
+		if (UNiagaraGraph* NiagaraGraph = const_cast<UNiagaraGraph*>(NodeB->GetNiagaraGraph()))
+		{
+			FNiagaraTypeDefinition ResolvedType = NiagaraGraph->GetCachedNumericConversion(PinB);
+			PinBType = ResolvedType;
+		}
+	}
 	
 	FString Message;
 	ECanCreateConnectionResponse Response = CONNECT_RESPONSE_DISALLOW;
