- Volume renderer fix for MIC being lost on cooked builds

#rnx

diff --git a/Engine/Plugins/FX/Niagara/Source/Niagara/Private/NiagaraVolumeRendererProperties.cpp b/Engine/Plugins/FX/Niagara/Source/Niagara/Private/NiagaraVolumeRendererProperties.cpp
--- a/Engine/Plugins/FX/Niagara/Source/Niagara/Private/NiagaraVolumeRendererProperties.cpp
+++ b/Engine/Plugins/FX/Niagara/Source/Niagara/Private/NiagaraVolumeRendererProperties.cpp
@@ -123,13 +123,13 @@ void UNiagaraVolumeRendererProperties::Serialize(FArchive& Ar)
 	// MIC will replace the main material during serialize
 	// Be careful if adding code that looks at the material to make sure you get the correct one
 	{
-//#if WITH_EDITORONLY_DATA
-//		TOptional<TGuardValue<TObjectPtr<UMaterialInterface>>> MICGuard;
-//		if (Ar.IsSaving() && Ar.IsCooking() && MICMaterial)
-//		{
-//			MICGuard.Emplace(Material, MICMaterial);
-//		}
-//#endif
+#if WITH_EDITORONLY_DATA
+		TOptional<TGuardValue<TObjectPtr<UMaterialInterface>>> MICGuard;
+		if (Ar.IsSaving() && Ar.IsCooking() && MICMaterial)
+		{
+			MICGuard.Emplace(Material, MICMaterial);
+		}
+#endif
 
 		Super::Serialize(Ar);
 	}
