- Add fx.Niagara.SystemSimulation.TickTaskAllowFrameOverlap to control if we allow frame overlap
- Currently disabled to see if a low frequency crash is related

diff --git a/Engine/Plugins/FX/Niagara/Source/Niagara/Private/NiagaraSystemSimulation.cpp b/Engine/Plugins/FX/Niagara/Source/Niagara/Private/NiagaraSystemSimulation.cpp
--- a/Engine/Plugins/FX/Niagara/Source/Niagara/Private/NiagaraSystemSimulation.cpp
+++ b/Engine/Plugins/FX/Niagara/Source/Niagara/Private/NiagaraSystemSimulation.cpp
@@ -121,14 +121,6 @@ static FAutoConsoleVariableRef CVarNiagaraSkipTickDeltaSeconds(
 	ECVF_Default
 );
 
-static int32 GNiagaraSystemSimulationTickTaskShouldWait = 0;
-static FAutoConsoleVariableRef CVarNiagaraSystemSimulationTickTaskShouldWait(
-	TEXT("fx.Niagara.SystemSimulation.TickTaskShouldWait"),
-	GNiagaraSystemSimulationTickTaskShouldWait,
-	TEXT("When enabled the tick task will wait for concurrent work to complete, when disabled the task is complete once the GT tick is complete."),
-	ECVF_Default
-);
-
 static int32 GNiagaraSystemSimulationMaxTickSubsteps = 100;
 static FAutoConsoleVariableRef CVarNiagaraSystemSimulationMaxTickSubsteps(
 	TEXT("fx.Niagara.SystemSimulation.MaxTickSubsteps"),
@@ -139,6 +131,22 @@ static FAutoConsoleVariableRef CVarNiagaraSystemSimulationMaxTickSubsteps(
 
 namespace NiagaraSystemSimulationLocal
 {
+	static bool GTickTaskShouldWait = false;
+	static FAutoConsoleVariableRef CVarGTickTaskShouldWait(
+		TEXT("fx.Niagara.SystemSimulation.TickTaskShouldWait"),
+		GTickTaskShouldWait,
+		TEXT("When enabled the tick task will wait for concurrent work to complete, when disabled the task is complete once the GT tick is complete."),
+		ECVF_Default
+	);
+
+	static int32 GTickTaskAllowFrameOverlap = false;
+	static FAutoConsoleVariableRef CVarTickTaskAllowFrameOverlap(
+		TEXT("fx.Niagara.SystemSimulation.TickTaskAllowFrameOverlap"),
+		GTickTaskAllowFrameOverlap,
+		TEXT("When enabled we allow ticks to overlap beyond PostActorTick until either EOF updates or the next tick."),
+		ECVF_Default
+	);
+
 #if NIAGARA_SYSTEMSIMULATION_DEBUGGING
 	static int GDebugKillAllOnSpawn = 0;
 	static FAutoConsoleVariableRef CVarDebugKillAllOnSpawn(
@@ -1257,13 +1265,13 @@ void FNiagaraSystemSimulation::Tick_GameThread_Internal(float DeltaSeconds, cons
 			Instance->ConcurrentTickGraphEvent = ConcurrentTickGraphEvent;
 		}
 		ConcurrentTickTask->Unlock();
-		if (bIsSolo || GNiagaraSystemSimulationTickTaskShouldWait || !System->AsyncWorkCanOverlapTickGroups())
+		if (bIsSolo || NiagaraSystemSimulationLocal::GTickTaskShouldWait || !System->AsyncWorkCanOverlapTickGroups())
 		{
 			MyCompletionGraphEvent->DontCompleteUntil(AllWorkCompleteGraphEvent);
 		}
 		else
 		{
-			if ( System->AllDIsPostSimulateCanOverlapFrames() )
+			if ( System->AllDIsPostSimulateCanOverlapFrames() && NiagaraSystemSimulationLocal::GTickTaskAllowFrameOverlap )
 			{
 				WorldManager->MarkSimulationsForEndOfFrameWait(this);
 			}
@@ -1445,13 +1453,13 @@ void FNiagaraSystemSimulation::Spawn_GameThread(float DeltaSeconds, bool bPostAc
 			ConcurrentTickTask->Unlock();
 
 			// Some simulations require us to complete inside PostActorTick and can not overlap to EOF updates / next frame, ensure we wait for them in the right location
-			if ( System->AllDIsPostSimulateCanOverlapFrames() && !GNiagaraSystemSimulationTickTaskShouldWait)
+			if ( System->AllDIsPostSimulateCanOverlapFrames() && NiagaraSystemSimulationLocal::GTickTaskAllowFrameOverlap )
 			{
-				WorldManager->MarkSimulationForPostActorWork(this);
+				WorldManager->MarkSimulationsForEndOfFrameWait(this);
 			}
 			else
 			{
-				WorldManager->MarkSimulationsForEndOfFrameWait(this);
+				WorldManager->MarkSimulationForPostActorWork(this);
 			}
 		}
 		else
