Experimental Vector VM improvements: faster execution speed, faster optimization speed, uses less runtime memory, added debugging features

diff --git a/Engine/Plugins/FX/Niagara/Source/Niagara/Private/NiagaraAsyncCompile.cpp b/Engine/Plugins/FX/Niagara/Source/Niagara/Private/NiagaraAsyncCompile.cpp
--- a/Engine/Plugins/FX/Niagara/Source/Niagara/Private/NiagaraAsyncCompile.cpp
+++ b/Engine/Plugins/FX/Niagara/Source/Niagara/Private/NiagaraAsyncCompile.cpp
@@ -367,7 +367,8 @@ void FNiagaraAsyncCompileTask::OptimizeByteCode()
 
 		FVectorVMOptimizeContext OptimizeContext;
 		FMemory::Memzero(&OptimizeContext, sizeof(OptimizeContext));
-		OptimizeVectorVMScript(ByteCode, ExeData->ByteCode.GetLength(), ExtFnTable.GetData(), ExtFnTable.Num(), &OptimizeContext, VVMFlag_OptOmitStats);
+		uint64 AssetPathHash = CityHash64((char *)AssetPath.GetCharArray().GetData(), AssetPath.GetCharArray().Num());
+		OptimizeVectorVMScript(ByteCode, ExeData->ByteCode.GetLength(), ExtFnTable.GetData(), ExtFnTable.Num(), &OptimizeContext, AssetPathHash, VVMFlag_OptOmitStats);
 
 		// freeze the OptimizeContext
 		FreezeVectorVMOptimizeContext(OptimizeContext, ExeData->ExperimentalContextData);
diff --git a/Engine/Plugins/FX/Niagara/Source/Niagara/Private/NiagaraCompilationTypes.cpp b/Engine/Plugins/FX/Niagara/Source/Niagara/Private/NiagaraCompilationTypes.cpp
--- a/Engine/Plugins/FX/Niagara/Source/Niagara/Private/NiagaraCompilationTypes.cpp
+++ b/Engine/Plugins/FX/Niagara/Source/Niagara/Private/NiagaraCompilationTypes.cpp
@@ -163,11 +163,6 @@ bool operator==(const FNiagaraDataSetProperties& Lhs, const FNiagaraDataSetPrope
 		&& Lhs.Variables == Rhs.Variables;
 }
 
-bool operator==(const FVectorVMExternalFunctionContextProxy& Lhs, const FVectorVMExternalFunctionContextProxy& Rhs)
-{
-	return true;
-}
-
 bool operator==(const FNiagaraCompilerTag& Lhs, const FNiagaraCompilerTag& Rhs)
 {
 	return Lhs.Variable == Rhs.Variable
diff --git a/Engine/Plugins/FX/Niagara/Source/Niagara/Private/NiagaraModule.cpp b/Engine/Plugins/FX/Niagara/Source/Niagara/Private/NiagaraModule.cpp
--- a/Engine/Plugins/FX/Niagara/Source/Niagara/Private/NiagaraModule.cpp
+++ b/Engine/Plugins/FX/Niagara/Source/Niagara/Private/NiagaraModule.cpp
@@ -272,9 +272,6 @@ void INiagaraModule::StartupModule()
 	VectorVM::Init();
 	FNiagaraTypeHelper::InitStaticTypes();
 
-#if VECTORVM_SUPPORTS_EXPERIMENTAL
-	InitVectorVM();
-#endif
 	LLM_SCOPE(ELLMTag::Niagara);
 	FNiagaraTypeDefinition::Init();
 
diff --git a/Engine/Plugins/FX/Niagara/Source/Niagara/Private/NiagaraScriptExecutionContext.cpp b/Engine/Plugins/FX/Niagara/Source/Niagara/Private/NiagaraScriptExecutionContext.cpp
--- a/Engine/Plugins/FX/Niagara/Source/Niagara/Private/NiagaraScriptExecutionContext.cpp
+++ b/Engine/Plugins/FX/Niagara/Source/Niagara/Private/NiagaraScriptExecutionContext.cpp
@@ -283,14 +283,14 @@ bool FNiagaraScriptExecutionContextBase::ExecuteInternal_Experimental(uint32 Num
 {
 	TRACE_CPUPROFILER_EVENT_SCOPE(VectorVM_Experimental);
 	FVectorVMExecContext ExecCtx;
-	ExecCtx.VVMState            = VectorVMState;
-	ExecCtx.DataSets            = DataSetMetaTable;
-	ExecCtx.ExtFunctionTable    = FunctionTable;
-	ExecCtx.UserPtrTable        = UserPtrTable;
-	ExecCtx.NumInstances        = NumInstances;
-	ExecCtx.ConstantTableData   = ConstantBufferTable.Buffers.GetData();
-	ExecCtx.ConstantTableSizes  = ConstantBufferTable.BufferSizes.GetData();
-	ExecCtx.ConstantTableCount  = ConstantBufferTable.Buffers.Num();
+	ExecCtx.VVMState               = VectorVMState;
+	ExecCtx.DataSets               = DataSetMetaTable;
+	ExecCtx.ExtFunctionTable       = FunctionTable;
+	ExecCtx.UserPtrTable           = UserPtrTable;
+	ExecCtx.NumInstances           = NumInstances;
+	ExecCtx.ConstantTableData      = ConstantBufferTable.Buffers.GetData();
+	ExecCtx.ConstantTableNumBytes  = ConstantBufferTable.BufferSizes.GetData();
+	ExecCtx.ConstantTableCount     = ConstantBufferTable.Buffers.Num();
 
 	if (VectorVMState)
 	{
@@ -342,24 +342,23 @@ bool FNiagaraScriptExecutionContextBase::ExecuteInternal_Legacy(uint32 NumInstan
 
 	ExecArgs.bAllowParallel = bAllowParallel;
 	FVectorVMSerializeState UESerializeState = { };
-#ifdef VVM_INCLUDE_SERIALIZATION
+#if VECTORVM_SUPPORTS_SERIALIZATION
 	static const IConsoleVariable* CVarInstancesPerChunk = IConsoleManager::Get().FindConsoleVariable(TEXT("vm.InstancesPerChunk"));
 	int32 NumParallelInstancesPerChunk = CVarInstancesPerChunk ? CVarInstancesPerChunk->GetInt() : 128;
 
-	UESerializeState.ReallocFn        = VVMRealloc;
-	UESerializeState.FreeFn           = VVMFree;
-	UESerializeState.NumInstances     = NumInstances;
-	UESerializeState.NumTempRegisters = ExecData.NumTempRegisters;
-	UESerializeState.NumTempRegFlags  = UESerializeState.NumTempRegisters;
-		
-	UESerializeState.TempRegFlags     = (unsigned char *)VVMRealloc(NULL, UESerializeState.NumTempRegisters, __FILE__, __LINE__);
-	FMemory::Memset(UESerializeState.TempRegFlags, 0, UESerializeState.NumTempRegisters);
-	UESerializeState.Bytecode         = (unsigned char *)ExecData.ByteCode.GetData().GetData();
-	UESerializeState.NumBytecodeBytes = ExecData.ByteCode.GetData().Num();
+	UESerializeState.ReallocFn          = VVMRealloc;
+	UESerializeState.FreeFn             = VVMFree;
+	UESerializeState.NumInstances       = NumInstances;
+	UESerializeState.NumTempRegisters   = ExecData.NumTempRegisters;
+			
+	UESerializeState.RegisterTableFlags = (unsigned char *)VVMRealloc(NULL, UESerializeState.NumTempRegisters, __FILE__, __LINE__);
+	FMemory::Memset(UESerializeState.RegisterTableFlags, 0, UESerializeState.NumTempRegisters);
+	UESerializeState.Bytecode           = (unsigned char *)ExecData.ByteCode.GetData().GetData();
+	UESerializeState.NumBytecodeBytes   = ExecData.ByteCode.GetData().Num();
 	SerializeVectorVMInputDataSets(&UESerializeState, DataSetMetaTable,  ConstantBufferTable.Buffers.GetData(), ConstantBufferTable.BufferSizes.GetData(), ConstantBufferTable.Buffers.Num());
 
-	UESerializeState.NumChunks        = (NumInstances + NumParallelInstancesPerChunk - 1) / NumParallelInstancesPerChunk;
-	UESerializeState.Chunks           = (FVectorVMSerializeChunk *)VVMRealloc(NULL, sizeof(FVectorVMSerializeChunk) * UESerializeState.NumChunks, __FILE__, __LINE__);
+	UESerializeState.NumChunks          = (NumInstances + NumParallelInstancesPerChunk - 1) / NumParallelInstancesPerChunk;
+	UESerializeState.Chunks             = (FVectorVMSerializeChunk *)VVMRealloc(NULL, sizeof(FVectorVMSerializeChunk) * UESerializeState.NumChunks, __FILE__, __LINE__);
 
 		
 	int NumExternalFunctions = FunctionTable.Num();
@@ -391,12 +390,12 @@ bool FNiagaraScriptExecutionContextBase::ExecuteInternal_Legacy(uint32 NumInstan
 			}
 		}
 	}
-#endif //VVM_INCLUDE_SERIALIZATION
+#endif //VECTORVM_SUPPORTS_SERIALIZATION
 	VectorVM::Exec(ExecArgs, &UESerializeState);
 
-#ifdef VVM_INCLUDE_SERIALIZATION
+#if VECTORVM_SUPPORTS_SERIALIZATION
 		SerializeVectorVMOutputDataSets(&UESerializeState, DataSetMetaTable,  ConstantBufferTable.Buffers.GetData(), ConstantBufferTable.BufferSizes.GetData(), ConstantBufferTable.Buffers.Num());
-#endif //VVM_INCLUDE_SERIALIZATION
+#endif //VECTORVM_SUPPORTS_SERIALIZATION
 
 	for (int32 Idx = 0; Idx < DataSetInfo.Num(); Idx++)
 	{
@@ -405,14 +404,14 @@ bool FNiagaraScriptExecutionContextBase::ExecuteInternal_Legacy(uint32 NumInstan
 		Info.DataSet->NumSpawnedIDs = Info.DataSet->GetSpawnedIDsTable().Num();
 	}
 
-#		ifdef VVM_INCLUDE_SERIALIZATION
+#		if VECTORVM_SUPPORTS_SERIALIZATION
 		{
 			//we don't own the bytecode memory so we can't free it
 			UESerializeState.Bytecode = NULL;
 			UESerializeState.NumBytecodeBytes = 0;
 			FreeVectorVMSerializeState(&UESerializeState);
 		}
-#		endif //VVM_INCLUDE_SERIALIZATION
+#		endif //VECTORVM_SUPPORTS_SERIALIZATION
 
 	return true;
 }
@@ -940,4 +939,3 @@ bool FNiagaraSystemScriptExecutionContext::GeneratePerInstanceDIFunctionTable(FN
 	}
 	return true;
 };
-
diff --git a/Engine/Plugins/FX/Niagara/Source/NiagaraEditor/Private/NiagaraCompilationTasks.cpp b/Engine/Plugins/FX/Niagara/Source/NiagaraEditor/Private/NiagaraCompilationTasks.cpp
--- a/Engine/Plugins/FX/Niagara/Source/NiagaraEditor/Private/NiagaraCompilationTasks.cpp
+++ b/Engine/Plugins/FX/Niagara/Source/NiagaraEditor/Private/NiagaraCompilationTasks.cpp
@@ -700,9 +700,9 @@ void FNiagaraSystemCompilationTask::FCompileTaskInfo::ProcessCompilation(FNiagar
 			ExtFnTable[i].NumOutputs = ExeData->CalledVMExternalFunctions[i].GetNumOutputs();
 		}
 
-		FVectorVMOptimizeContext OptimizeContext;
-		FMemory::Memzero(&OptimizeContext, sizeof(OptimizeContext));
-		OptimizeVectorVMScript(ByteCode, ExeData->ByteCode.GetLength(), ExtFnTable.GetData(), ExtFnTable.Num(), &OptimizeContext, VVMFlag_OptOmitStats);
+		FVectorVMOptimizeContext OptimizeContext = { };
+		uint64 AssetPathHash = CityHash64((char *)AssetPath.GetCharArray().GetData(), AssetPath.GetCharArray().Num());
+		OptimizeVectorVMScript(ByteCode, ExeData->ByteCode.GetLength(), ExtFnTable.GetData(), ExtFnTable.Num(), &OptimizeContext, AssetPathHash, VVMFlag_OptOmitStats);
 
 		// freeze the OptimizeContext
 		FreezeVectorVMOptimizeContext(OptimizeContext, ExeData->ExperimentalContextData);
diff --git a/Engine/Source/Runtime/Core/Public/Math/UnrealMathNeon.h b/Engine/Source/Runtime/Core/Public/Math/UnrealMathNeon.h
--- a/Engine/Source/Runtime/Core/Public/Math/UnrealMathNeon.h
+++ b/Engine/Source/Runtime/Core/Public/Math/UnrealMathNeon.h
@@ -2828,6 +2828,7 @@ FORCEINLINE VectorRegister4Int VectorFloatToInt(const VectorRegister4Double& A)
 * @param Ptr	Memory pointer
 */
 #define VectorIntStore( Vec, Ptr )			vst1q_s32( (int32*)(Ptr), Vec )
+#define VectorIntStore_16( Vec, Ptr )		vst1q_s16( (int16*)(Ptr), Vec )
 
 /**
 * Loads 4 int32s from unaligned memory.
@@ -2859,7 +2860,8 @@ FORCEINLINE VectorRegister4Int VectorFloatToInt(const VectorRegister4Double& A)
 * @param Ptr	Unaligned memory pointer to the 4 int32s
 * @return		VectorRegister4Int(*Ptr, *Ptr, *Ptr, *Ptr)
 */
-#define VectorIntLoad1( Ptr )	vld1q_dup_s32((int32*)(Ptr))
+#define VectorIntLoad1( Ptr )	                    vld1q_dup_s32((int32*)(Ptr))
+#define VectorIntLoad1_16(Ptr)                      vld1q_dup_s16((int16*)(Ptr))
 
 #define VectorIntSet1(F)                            vdupq_n_s32(F)
 #define VectorSetZero()                             vdupq_n_s32(0)
@@ -2876,7 +2878,7 @@ FORCEINLINE VectorRegister4Int VectorRoundToIntHalfToEven(const VectorRegister4F
 	return vcvtnq_s32_f32(Vec);
 }
 
-inline VectorRegister4i VectorIntExpandLow16To32(VectorRegister4i V) {
+FORCEINLINE VectorRegister4i VectorIntExpandLow16To32(VectorRegister4i V) {
 	int16x4x2_t res = vzip_s16(vget_low_u16(V), vdup_n_u16(0));
 	return vcombine_s16(res.val[0], res.val[1]);
 }
diff --git a/Engine/Source/Runtime/Core/Public/Math/UnrealMathSSE.h b/Engine/Source/Runtime/Core/Public/Math/UnrealMathSSE.h
--- a/Engine/Source/Runtime/Core/Public/Math/UnrealMathSSE.h
+++ b/Engine/Source/Runtime/Core/Public/Math/UnrealMathSSE.h
@@ -3445,6 +3445,7 @@ FORCEINLINE VectorRegister4Int VectorFloatToInt(const VectorRegister4Double& A)
 * @param Ptr	Memory pointer
 */
 #define VectorIntStore( Vec, Ptr )			_mm_storeu_si128( (VectorRegister4Int*)(Ptr), Vec )
+#define VectorIntStore_16( Vec, Ptr )       _mm_storeu_si64( (VectorRegister4Int*)(Ptr), Vec )
 
 /**
 * Loads 4 int32s from unaligned memory.
@@ -3476,7 +3477,8 @@ FORCEINLINE VectorRegister4Int VectorFloatToInt(const VectorRegister4Double& A)
 * @param Ptr	Unaligned memory pointer to the 4 int32s
 * @return		VectorRegister4Int(*Ptr, *Ptr, *Ptr, *Ptr)
 */
-#define VectorIntLoad1( Ptr )	_mm_set1_epi32(*(Ptr))
+#define VectorIntLoad1(Ptr)                         _mm_set1_epi32(*(Ptr))
+#define VectorIntLoad1_16(Ptr)                      _mm_set1_epi16(*(Ptr))
 #define VectorSetZero()								_mm_setzero_si128()
 #define VectorSet1(F)								_mm_set1_ps(F)
 #define VectorIntSet1(F)							_mm_set1_epi32(F)
diff --git a/Engine/Source/Runtime/VectorVM/Private/VectorVM.cpp b/Engine/Source/Runtime/VectorVM/Private/VectorVM.cpp
--- a/Engine/Source/Runtime/VectorVM/Private/VectorVM.cpp
+++ b/Engine/Source/Runtime/VectorVM/Private/VectorVM.cpp
@@ -1,6 +1,7 @@
 // Copyright Epic Games, Inc. All Rights Reserved.
 
 #include "VectorVM.h"
+#include "VectorVMSerialization.h"
 #include "Modules/ModuleManager.h"
 #include "UObject/Class.h"
 #include "UObject/Package.h"
@@ -211,14 +212,6 @@ static FAutoConsoleVariableRef CVarVVMChunkSizeInBytes(
 	ECVF_Default
 );
 
-static int32 GVVMPageSizeInKB = 64;
-static FAutoConsoleVariableRef CVarVVMPageSizeInKB(
-	TEXT("vm.PageSizeInKB"),
-	GVVMPageSizeInKB,
-	TEXT("Minimum allocation per VM instance.  There are 64 of these, so multiply GVVMPageSizeInKB * 64 * 1024 to get total number of bytes used by the VVM\n"),
-	ECVF_ReadOnly 
-);
-
 static int32 GVVMMaxThreadsPerScript = 8;
 static FAutoConsoleVariableRef CVarVVMMaxThreadsPerScript(
 	TEXT("vm.MaxThreadsPerScript"),
@@ -283,8 +276,6 @@ static FAutoConsoleVariableRef CVarbBatchPackVMOutput(
 	ECVF_Default
 );
 
-#include "VectorVMExperimental.inl"
-
 uint8 VectorVM::GetNumOpCodes()
 {
 	return (uint8)EVectorVMOp::NumOpcodes;
@@ -2875,7 +2866,7 @@ void VectorVM::Exec(FVectorVMExecArgs& Args, FVectorVMSerializeState *SerializeS
 	}
 #endif
 
-#ifdef VVM_INCLUDE_SERIALIZATION
+#if VECTORVM_SUPPORTS_SERIALIZATION
 	uint64 EndTime = FPlatformTime::Cycles64();
 	if (SerializeState) {
 		SerializeState->ExecDt = EndTime - StartTime; //NOTE: doesn't work if ParallelFor splits the work into multiple threads
