Fixed emitter isolation behavior

diff --git a/Engine/Plugins/FX/Niagara/Source/NiagaraEditor/Private/ViewModels/NiagaraSystemViewModel.cpp b/Engine/Plugins/FX/Niagara/Source/NiagaraEditor/Private/ViewModels/NiagaraSystemViewModel.cpp
--- a/Engine/Plugins/FX/Niagara/Source/NiagaraEditor/Private/ViewModels/NiagaraSystemViewModel.cpp
+++ b/Engine/Plugins/FX/Niagara/Source/NiagaraEditor/Private/ViewModels/NiagaraSystemViewModel.cpp
@@ -2272,8 +2272,12 @@ void FNiagaraSystemViewModel::IsolateSelectedEmitters()
 void FNiagaraSystemViewModel::ToggleIsolateEmitterAndSelectedEmitters(FGuid OriginEmitter)
 {
 	TArray<FGuid> SelectedEmitterHandles = GetSelectionViewModel()->GetSelectedEmitterHandleIds();
-	// we treat the origin emitter as if he was selected
-	SelectedEmitterHandles.AddUnique(OriginEmitter);
+
+	if(SelectedEmitterHandles.Contains(OriginEmitter) == false)
+	{
+		SelectedEmitterHandles.Empty();
+		SelectedEmitterHandles.Add(OriginEmitter);
+	}
 
 	bool bAnyIsolationMismatch = false;
 
