Improved tooltips in the hierarchy editor

diff --git a/Engine/Plugins/FX/Niagara/Source/NiagaraEditor/Private/ViewModels/HierarchyEditor/NiagaraHierarchyViewModelBase.cpp b/Engine/Plugins/FX/Niagara/Source/NiagaraEditor/Private/ViewModels/HierarchyEditor/NiagaraHierarchyViewModelBase.cpp
--- a/Engine/Plugins/FX/Niagara/Source/NiagaraEditor/Private/ViewModels/HierarchyEditor/NiagaraHierarchyViewModelBase.cpp
+++ b/Engine/Plugins/FX/Niagara/Source/NiagaraEditor/Private/ViewModels/HierarchyEditor/NiagaraHierarchyViewModelBase.cpp
@@ -1063,7 +1063,7 @@ FNiagaraHierarchyItemViewModelBase::FCanPerformActionResults FNiagaraHierarchyRo
 		}
 		else
 		{
-			FText Message = LOCTEXT("CantDropHierarchyItemOnRootDragMessage", "Can not add {0} to the hierarchy root directly. Custom sections only accept categories at the root level!");
+			FText Message = LOCTEXT("CantDropHierarchyItemOnRootDragMessage", "Can not add {0} here. Please add it to a category!");
 			Message = FText::FormatOrdered(Message, FText::FromString(DraggedItem->ToString()));
 			Results.CanPerformMessage = Message;
 		}
@@ -1250,7 +1250,14 @@ FNiagaraHierarchyItemViewModelBase::FCanPerformActionResults FNiagaraHierarchySe
 				Results.CanPerformMessage = Message;
 			}
 		}
-	}	 
+	}
+	else if(UNiagaraHierarchyItem* Item = Cast<UNiagaraHierarchyItem>(DraggedItem->GetDataMutable()))
+	{
+		FText Message = LOCTEXT("CantDropItemOnSectionDragMessage", "Can't drop items onto sections. Please drag a category onto section {0}");
+		Message = FText::FormatOrdered(Message, FText::FromString(ToString()));
+		Results.bCanPerform = false;
+		Results.CanPerformMessage = Message;
+	}
 
 	return Results;
 }
diff --git a/Engine/Plugins/FX/Niagara/Source/NiagaraEditor/Private/Widgets/SNiagaraHierarchy.cpp b/Engine/Plugins/FX/Niagara/Source/NiagaraEditor/Private/Widgets/SNiagaraHierarchy.cpp
--- a/Engine/Plugins/FX/Niagara/Source/NiagaraEditor/Private/Widgets/SNiagaraHierarchy.cpp
+++ b/Engine/Plugins/FX/Niagara/Source/NiagaraEditor/Private/Widgets/SNiagaraHierarchy.cpp
@@ -127,7 +127,7 @@ bool SNiagaraHierarchyCategory::OnVerifyCategoryRename(const FText& NewName, FTe
 
 		if(CategoryNames.Contains(NewName.ToString()))
 		{
-			OutTooltip = LOCTEXT("HierarchyCategoryCantRename_DuplicateOnLayer", "Another category of the same name already exists on the same hierarchy level!");
+			OutTooltip = LOCTEXT("HierarchyCategoryCantRename_DuplicateOnLayer", "Another category of the same name already exists!");
 			return false;
 		}
 	}
