Calling stack refresh after renaming a user parameter to ensure the UI is getting updated and isn't stuck with the old name of a linked parameter

diff --git a/Engine/Plugins/FX/Niagara/Source/NiagaraEditor/Private/ViewModels/NiagaraSystemViewModel.cpp b/Engine/Plugins/FX/Niagara/Source/NiagaraEditor/Private/ViewModels/NiagaraSystemViewModel.cpp
--- a/Engine/Plugins/FX/Niagara/Source/NiagaraEditor/Private/ViewModels/NiagaraSystemViewModel.cpp
+++ b/Engine/Plugins/FX/Niagara/Source/NiagaraEditor/Private/ViewModels/NiagaraSystemViewModel.cpp
@@ -1362,6 +1362,7 @@ bool FNiagaraSystemViewModel::RenameParameter(const FNiagaraVariable TargetParam
 			System->HandleVariableRenamed(TargetParameter, FNiagaraVariableBase(TargetParameter.GetType(), NewName), true);
 		}
 
+		GetSystemStackViewModel()->RequestRefreshDeferred();
 		return true;
 	}
 
