Added delete button to parameters in a Niagara parameter collection. Not displayed in instances.
Also displayed in input & output parameter script view models.

Also fixed a typo & made some minor API improvements, and deleted some unneeded headers

diff --git a/Engine/Plugins/FX/Niagara/Source/NiagaraEditor/Private/Customizations/NiagaraParameterCollectionCustomNodeBuilder.h b/Engine/Plugins/FX/Niagara/Source/NiagaraEditor/Private/Customizations/NiagaraParameterCollectionCustomNodeBuilder.h
--- a/Engine/Plugins/FX/Niagara/Source/NiagaraEditor/Private/Customizations/NiagaraParameterCollectionCustomNodeBuilder.h
+++ b/Engine/Plugins/FX/Niagara/Source/NiagaraEditor/Private/Customizations/NiagaraParameterCollectionCustomNodeBuilder.h
@@ -1,16 +1,14 @@
 // Copyright Epic Games, Inc. All Rights Reserved.
 
 #pragma once
-#include "DetailLayoutBuilder.h"
+
+#include "IDetailCustomNodeBuilder.h"
 #include "DetailCategoryBuilder.h"
 #include "IDetailPropertyRow.h"
 #include "DetailWidgetRow.h"
 #include "PropertyCustomizationHelpers.h"
 #include "Widgets/Text/SInlineEditableTextBlock.h"
 #include "ViewModels/NiagaraParameterCollectionViewModel.h"
-#include "ViewModels/NiagaraScriptInputCollectionViewModel.h"
-#include "ViewModels/NiagaraScriptOutputCollectionViewModel.h"
-#include "ViewModels/NiagaraScriptViewModel.h"
 #include "ViewModels/NiagaraParameterViewModel.h"
 #include "NiagaraEditorStyle.h"
 #include "IDetailChildrenBuilder.h"
@@ -19,9 +17,10 @@
 #include "Widgets/Images/SImage.h"
 #include "NiagaraEditorModule.h"
 #include "Modules/ModuleManager.h"
-#include "INiagaraEditorTypeUtilities.h"
 #include "Widgets/Layout/SBox.h"
-#include "NiagaraScript.h"
+#include "Widgets/Input/SButton.h"
+
+#define LOCTEXT_NAMESPACE "FNiagaraParameterCollectionCustomNodeBuilder"
 
 class FNiagaraParameterCollectionCustomNodeBuilder : public IDetailCustomNodeBuilder
 {
@@ -51,7 +50,7 @@ public:
 		static const FName NiagaraCustomNodeBuilder("NiagaraCustomNodeBuilder");
 		return NiagaraCustomNodeBuilder;
 	}
-
+	
 	virtual void GenerateChildContent(IDetailChildrenBuilder& ChildrenBuilder) override
 	{
 		const TArray<TSharedRef<INiagaraParameterViewModel>>& Parameters = ViewModel->GetParameters();
@@ -69,7 +68,7 @@ public:
 					.Style(FNiagaraEditorStyle::Get(), "NiagaraEditor.ParameterInlineEditableText")
 					.Text(Parameter, &INiagaraParameterViewModel::GetNameText)
 					.OnVerifyTextChanged(Parameter, &INiagaraParameterViewModel::VerifyNodeNameTextChanged)
-					.OnTextCommitted(Parameter, &INiagaraParameterViewModel::NameTextComitted);
+					.OnTextCommitted(Parameter, &INiagaraParameterViewModel::NameTextCommitted);
 			}
 			else
 			{
@@ -138,7 +137,7 @@ public:
 						.ValueContent()
 						[
 							SNew(SObjectPropertyEntryBox)
-								.AllowedClass(Parameter->GetType()->GetClass())
+								.AllowedClass(Parameter->GetType().GetClass())
 								.ObjectPath_Lambda(
 									[Parameter]() -> FString
 									{
@@ -179,6 +178,7 @@ public:
 					TSharedPtr<SWidget> DefaultNameWidget;
 					TSharedPtr<SWidget> DefaultValueWidget;
 					FDetailWidgetRow& CustomWidget = Row->CustomWidget(true);
+					
 					Row->GetDefaultWidgets(DefaultNameWidget, DefaultValueWidget, CustomWidget);
 					CustomWidget
 						.NameContent()
@@ -198,12 +198,28 @@ public:
 							];
 					}
 
+					CustomWidget.ExtensionContent()
+						[
+							SNew(SButton)
+							.OnClicked_Raw(this, &FNiagaraParameterCollectionCustomNodeBuilder::DeleteParameter, Parameter->GetVariable())
+							.Content()
+							[
+								SNew(SImage)
+								.Image(FAppStyle::GetBrush("Icons.Delete"))
+							]
+						];
 				}
 			}
 		}
 	}
 
 private:
+	FReply DeleteParameter(FNiagaraVariable Parameter)
+	{
+		ViewModel->DeleteParameters({Parameter});
+		return FReply::Handled();
+	}
+	
 	void OnCollectionViewModelChanged()
 	{
 		OnRebuildChildren.ExecuteIfBound();
@@ -214,3 +230,5 @@ private:
 	FSimpleDelegate OnRebuildChildren;
 	bool bAllowMetaData;
 };
+
+#undef LOCTEXT_NAMESPACE
diff --git a/Engine/Plugins/FX/Niagara/Source/NiagaraEditor/Private/ViewModels/NiagaraCollectionParameterViewModel.cpp b/Engine/Plugins/FX/Niagara/Source/NiagaraEditor/Private/ViewModels/NiagaraCollectionParameterViewModel.cpp
--- a/Engine/Plugins/FX/Niagara/Source/NiagaraEditor/Private/ViewModels/NiagaraCollectionParameterViewModel.cpp
+++ b/Engine/Plugins/FX/Niagara/Source/NiagaraEditor/Private/ViewModels/NiagaraCollectionParameterViewModel.cpp
@@ -38,6 +38,11 @@ void FNiagaraCollectionParameterViewModel::Reset()
 	OnDefaultValueChanged().RemoveAll(this);
 }
 
+FNiagaraVariable FNiagaraCollectionParameterViewModel::GetVariable() const
+{
+	return Parameter;
+}
+
 FName FNiagaraCollectionParameterViewModel::GetName() const
 {
 	if (CollectionInst && CollectionInst->GetParent())
@@ -55,7 +60,7 @@ FText FNiagaraCollectionParameterViewModel::GetTypeDisplayName() const
 	return FText::Format(LOCTEXT("TypeTextFormat", "Type: {0}"), Parameter.GetType().GetNameText());
 }
 
-void FNiagaraCollectionParameterViewModel::NameTextComitted(const FText& Name, ETextCommit::Type CommitInfo)
+void FNiagaraCollectionParameterViewModel::NameTextCommitted(const FText& Name, ETextCommit::Type CommitInfo)
 {
 	FName NewName = *CollectionInst->GetParent()->ParameterNameFromFriendlyName(*Name.ToString());
 
@@ -91,9 +96,9 @@ bool FNiagaraCollectionParameterViewModel::VerifyNodeNameTextChanged(const FText
 	return true;
 }
 
-TSharedPtr<FNiagaraTypeDefinition> FNiagaraCollectionParameterViewModel::GetType() const
+FNiagaraTypeDefinition FNiagaraCollectionParameterViewModel::GetType() const
 {
-	return MakeShareable(new FNiagaraTypeDefinition(Parameter.GetType()));
+	return Parameter.GetType();
 }
 
 bool FNiagaraCollectionParameterViewModel::CanChangeParameterType()const
diff --git a/Engine/Plugins/FX/Niagara/Source/NiagaraEditor/Private/ViewModels/NiagaraCollectionParameterViewModel.h b/Engine/Plugins/FX/Niagara/Source/NiagaraEditor/Private/ViewModels/NiagaraCollectionParameterViewModel.h
--- a/Engine/Plugins/FX/Niagara/Source/NiagaraEditor/Private/ViewModels/NiagaraCollectionParameterViewModel.h
+++ b/Engine/Plugins/FX/Niagara/Source/NiagaraEditor/Private/ViewModels/NiagaraCollectionParameterViewModel.h
@@ -18,11 +18,12 @@ public:
 	void Reset();
 
 	//~ INiagaraParameterViewmodel interface.
+	virtual FNiagaraVariable GetVariable() const override;
 	virtual FName GetName() const override;
-	virtual void NameTextComitted(const FText& Name, ETextCommit::Type CommitInfo) override;
+	virtual void NameTextCommitted(const FText& Name, ETextCommit::Type CommitInfo) override;
 	virtual bool VerifyNodeNameTextChanged(const FText& NewText, FText& OutErrorMessage) override;
 	virtual FText GetTypeDisplayName() const override;
-	virtual TSharedPtr<FNiagaraTypeDefinition> GetType() const override;
+	virtual FNiagaraTypeDefinition GetType() const override;
 	virtual void SelectedTypeChanged(TSharedPtr<FNiagaraTypeDefinition> Item, ESelectInfo::Type SelectionType) override;
 	virtual EDefaultValueType GetDefaultValueType() override;
 	virtual TSharedRef<FStructOnScope> GetDefaultValueStruct() override;
diff --git a/Engine/Plugins/FX/Niagara/Source/NiagaraEditor/Private/ViewModels/NiagaraParameterCollectionAssetViewModel.cpp b/Engine/Plugins/FX/Niagara/Source/NiagaraEditor/Private/ViewModels/NiagaraParameterCollectionAssetViewModel.cpp
--- a/Engine/Plugins/FX/Niagara/Source/NiagaraEditor/Private/ViewModels/NiagaraParameterCollectionAssetViewModel.cpp
+++ b/Engine/Plugins/FX/Niagara/Source/NiagaraEditor/Private/ViewModels/NiagaraParameterCollectionAssetViewModel.cpp
@@ -213,20 +213,33 @@ void FNiagaraParameterCollectionAssetViewModel::DeleteSelectedParameters()
 		TSet<FNiagaraVariable> VarsToDelete;
 		for (TSharedRef<INiagaraParameterViewModel> Parameter : GetSelection().GetSelectedObjects())
 		{
-			FNiagaraVariable ToDelete(*Parameter->GetType().Get(), *Collection->ParameterNameFromFriendlyName(Parameter->GetName().ToString()));
-			VarsToDelete.Add(ToDelete);
+			VarsToDelete.Add(Parameter->GetVariable());
 		}
 		GetSelection().ClearSelectedObjects();
 
-		FScopedTransaction ScopedTransaction(LOCTEXT("DeleteNPCParameter", "Delete Parameter"));
-		for (FNiagaraVariable ParamToDelete : VarsToDelete)
-		{
-			Collection->RemoveParameter(ParamToDelete);
-		}
+		DeleteParameters(VarsToDelete.Array());
+	}
+}
 
-		CollectionChanged(true);
-		RefreshParameterViewModels();
+void FNiagaraParameterCollectionAssetViewModel::DeleteParameters(TArray<FNiagaraVariable> ParametersToDelete)
+{
+	check(Collection && Instance);
+
+	TSet<FNiagaraVariableBase> ResolvedParametersToDelete;
+	for(const FNiagaraVariableBase& Parameter : ParametersToDelete)
+	{
+		FNiagaraVariableBase ToDelete(Parameter.GetType(), *Collection->ParameterNameFromFriendlyName(Parameter.GetName().ToString()));
+		ResolvedParametersToDelete.Add(ToDelete);
+	}
+
+	FScopedTransaction ScopedTransaction(LOCTEXT("DeleteNPCParameter", "Delete Parameter"));
+	for (FNiagaraVariable ParamToDelete : ResolvedParametersToDelete)
+	{
+		Collection->RemoveParameter(ParamToDelete);
 	}
+
+	CollectionChanged(true);
+	RefreshParameterViewModels();
 }
 
 const TArray<TSharedRef<INiagaraParameterViewModel>>& FNiagaraParameterCollectionAssetViewModel::GetParameters()
@@ -385,7 +398,7 @@ void FNiagaraParameterCollectionAssetViewModel::OnParameterTypeChanged(FNiagaraV
 
 	Collection->GetDefaultInstance()->RemoveParameter(ParameterVariable);
 
-	FNiagaraTypeDefinition Type = *ParameterViewModels[Index]->GetType();
+	FNiagaraTypeDefinition Type = ParameterViewModels[Index]->GetType();
 	Collection->GetParameters()[Index].SetType(Type);
 
 	Collection->GetDefaultInstance()->AddParameter(Collection->GetParameters()[Index]);
diff --git a/Engine/Plugins/FX/Niagara/Source/NiagaraEditor/Private/ViewModels/NiagaraParameterCollectionAssetViewModel.h b/Engine/Plugins/FX/Niagara/Source/NiagaraEditor/Private/ViewModels/NiagaraParameterCollectionAssetViewModel.h
--- a/Engine/Plugins/FX/Niagara/Source/NiagaraEditor/Private/ViewModels/NiagaraParameterCollectionAssetViewModel.h
+++ b/Engine/Plugins/FX/Niagara/Source/NiagaraEditor/Private/ViewModels/NiagaraParameterCollectionAssetViewModel.h
@@ -31,6 +31,7 @@ public:
 	virtual void AddParameter(TSharedPtr<FNiagaraTypeDefinition> ParameterType) override;
 	virtual void RemoveParameter(FNiagaraVariable& Parameter);
 	virtual void DeleteSelectedParameters() override;
+	virtual void DeleteParameters(TArray<FNiagaraVariable> ParametersToDelete) override;
 	virtual const TArray<TSharedRef<INiagaraParameterViewModel>>& GetParameters() override;
 	virtual EVisibility GetAddButtonVisibility()const override;
 
diff --git a/Engine/Plugins/FX/Niagara/Source/NiagaraEditor/Private/ViewModels/NiagaraParameterCollectionViewModel.h b/Engine/Plugins/FX/Niagara/Source/NiagaraEditor/Private/ViewModels/NiagaraParameterCollectionViewModel.h
--- a/Engine/Plugins/FX/Niagara/Source/NiagaraEditor/Private/ViewModels/NiagaraParameterCollectionViewModel.h
+++ b/Engine/Plugins/FX/Niagara/Source/NiagaraEditor/Private/ViewModels/NiagaraParameterCollectionViewModel.h
@@ -49,6 +49,8 @@ public:
 	/** Deletes the currently selected parameters. */
 	virtual void DeleteSelectedParameters() = 0;
 
+	virtual void DeleteParameters(TArray<FNiagaraVariable> ParametersToDelete) = 0;
+
 	/** Gets the parameter view models. */
 	virtual const TArray<TSharedRef<INiagaraParameterViewModel>>& GetParameters() = 0;
 
diff --git a/Engine/Plugins/FX/Niagara/Source/NiagaraEditor/Private/ViewModels/NiagaraParameterViewModel.h b/Engine/Plugins/FX/Niagara/Source/NiagaraEditor/Private/ViewModels/NiagaraParameterViewModel.h
--- a/Engine/Plugins/FX/Niagara/Source/NiagaraEditor/Private/ViewModels/NiagaraParameterViewModel.h
+++ b/Engine/Plugins/FX/Niagara/Source/NiagaraEditor/Private/ViewModels/NiagaraParameterViewModel.h
@@ -33,6 +33,8 @@ public:
 
 public:
 	virtual ~INiagaraParameterViewModel() {}
+
+	virtual FNiagaraVariable GetVariable() const = 0;
 	
 	/** Gets the name of the parameter. */
 	virtual FName GetName() const = 0;
@@ -44,7 +46,7 @@ public:
 	virtual FText GetNameText() const = 0;
 
 	/** Handles a hame text change being comitter from the UI. */
-	virtual void NameTextComitted(const FText& Name, ETextCommit::Type CommitInfo) = 0;
+	virtual void NameTextCommitted(const FText& Name, ETextCommit::Type CommitInfo) = 0;
 
 	/** Handles verification of an in-progress variable name change in the UI.*/
 	virtual bool VerifyNodeNameTextChanged(const FText& NewText, FText& OutErrorMessage) = 0;
@@ -56,7 +58,7 @@ public:
 	virtual bool CanChangeParameterType() const = 0;
 
 	/** Gets the type of the paramter. */
-	virtual TSharedPtr<FNiagaraTypeDefinition> GetType() const = 0;
+	virtual FNiagaraTypeDefinition GetType() const = 0;
 
 	/** Handles the paramter type being changed from the UI. */
 	virtual void SelectedTypeChanged(TSharedPtr<FNiagaraTypeDefinition> Item, ESelectInfo::Type SelectionType) = 0;
diff --git a/Engine/Plugins/FX/Niagara/Source/NiagaraEditor/Private/ViewModels/NiagaraScriptInputCollectionViewModel.cpp b/Engine/Plugins/FX/Niagara/Source/NiagaraEditor/Private/ViewModels/NiagaraScriptInputCollectionViewModel.cpp
--- a/Engine/Plugins/FX/Niagara/Source/NiagaraEditor/Private/ViewModels/NiagaraScriptInputCollectionViewModel.cpp
+++ b/Engine/Plugins/FX/Niagara/Source/NiagaraEditor/Private/ViewModels/NiagaraScriptInputCollectionViewModel.cpp
@@ -2,13 +2,11 @@
 
 #include "NiagaraScriptInputCollectionViewModel.h"
 #include "NiagaraScript.h"
-#include "NiagaraEmitter.h"
 #include "NiagaraScriptSource.h"
 #include "NiagaraTypes.h"
 #include "NiagaraScriptParameterViewModel.h"
 #include "NiagaraGraph.h"
 #include "NiagaraNodeInput.h"
-#include "INiagaraEditorTypeUtilities.h"
 #include "NiagaraEditorUtilities.h"
 #include "NiagaraDataInterface.h"
 
@@ -174,27 +172,31 @@ void FNiagaraScriptInputCollectionViewModel::DeleteSelectedParameters()
 {
 	if (GetSelection().GetSelectedObjects().Num() > 0)
 	{
-		TSet<FName> InputNamesToDelete;
+		TSet<FNiagaraVariable> VariablesToDelete;
 		for (TSharedRef<INiagaraParameterViewModel> InputParameter : GetSelection().GetSelectedObjects())
 		{
-			InputNamesToDelete.Add(InputParameter->GetName());
+			VariablesToDelete.Add(InputParameter->GetVariable());
 		}
 		GetSelection().ClearSelectedObjects();
+		DeleteParameters(VariablesToDelete.Array());		
+	}
+}
 
-		if (Graph.IsValid())
-		{
-			FScopedTransaction ScopedTransaction(NSLOCTEXT("NiagaraEmitterInputEditor", "DeletedSelectedNodes", "Delete selected nodes"));
-			Graph->Modify();
+void FNiagaraScriptInputCollectionViewModel::DeleteParameters(TArray<FNiagaraVariable> ParametersToDelete)
+{
+	if (Graph.IsValid())
+	{
+		FScopedTransaction ScopedTransaction(NSLOCTEXT("NiagaraEmitterInputEditor", "DeletedSelectedNodes", "Delete selected nodes"));
+		Graph->Modify();
 
-			TArray<UNiagaraNodeInput*> InputNodes;
-			Graph->GetNodesOfClass(InputNodes);
-			for (UNiagaraNodeInput* InputNode : InputNodes)
+		TArray<UNiagaraNodeInput*> InputNodes;
+		Graph->GetNodesOfClass(InputNodes);
+		for (UNiagaraNodeInput* InputNode : InputNodes)
+		{
+			if (ParametersToDelete.Contains(InputNode->Input))
 			{
-				if (InputNamesToDelete.Contains(InputNode->Input.GetName()))
-				{
-					InputNode->Modify();
-					InputNode->DestroyNode();
-				}
+				InputNode->Modify();
+				InputNode->DestroyNode();
 			}
 		}
 	}
@@ -470,7 +472,7 @@ void FNiagaraScriptInputCollectionViewModel::OnParameterValueChangedInternal(TSh
 			// and the case where we're updating a compiled variable.
 			if (InputNode->Usage == ENiagaraInputNodeUsage::Parameter && InputNode->Input.GetName() == ChangedParameter->GetName())
 			{
-				if (ensureMsgf(InputNode->Input.GetType() == *ChangedParameter->GetType().Get(), TEXT("Can not propagate variable values when the types don't match.")))
+				if (ensureMsgf(InputNode->Input.GetType() == ChangedParameter->GetType(), TEXT("Can not propagate variable values when the types don't match.")))
 				{
 					InputNodesToUpdate.Add(InputNode);
 				}
diff --git a/Engine/Plugins/FX/Niagara/Source/NiagaraEditor/Private/ViewModels/NiagaraScriptInputCollectionViewModel.h b/Engine/Plugins/FX/Niagara/Source/NiagaraEditor/Private/ViewModels/NiagaraScriptInputCollectionViewModel.h
--- a/Engine/Plugins/FX/Niagara/Source/NiagaraEditor/Private/ViewModels/NiagaraScriptInputCollectionViewModel.h
+++ b/Engine/Plugins/FX/Niagara/Source/NiagaraEditor/Private/ViewModels/NiagaraScriptInputCollectionViewModel.h
@@ -27,6 +27,7 @@ public:
 	virtual FText GetDisplayName() const override;
 	virtual void AddParameter(TSharedPtr<FNiagaraTypeDefinition> ParameterType) override;
 	virtual void DeleteSelectedParameters() override;
+	virtual void DeleteParameters(TArray<FNiagaraVariable> ParametersToDelete) override;
 	virtual const TArray<TSharedRef<INiagaraParameterViewModel>>& GetParameters() override;
 
 	/** Rebuilds the parameter view models. */
diff --git a/Engine/Plugins/FX/Niagara/Source/NiagaraEditor/Private/ViewModels/NiagaraScriptOutputCollectionViewModel.cpp b/Engine/Plugins/FX/Niagara/Source/NiagaraEditor/Private/ViewModels/NiagaraScriptOutputCollectionViewModel.cpp
--- a/Engine/Plugins/FX/Niagara/Source/NiagaraEditor/Private/ViewModels/NiagaraScriptOutputCollectionViewModel.cpp
+++ b/Engine/Plugins/FX/Niagara/Source/NiagaraEditor/Private/ViewModels/NiagaraScriptOutputCollectionViewModel.cpp
@@ -132,37 +132,42 @@ void FNiagaraScriptOutputCollectionViewModel::DeleteSelectedParameters()
 	{
 		return;
 	}
-
-	TArray<UNiagaraNodeOutput*> OutputNodes;
-	Graph->FindOutputNodes(OutputNodes);
-	if (OutputNodes.Num() == 0)
-	{
-		return;
-	}
 	
 	if(GetSelection().GetSelectedObjects().Num() > 0)
 	{
-		TSet<FName> OutputNamesToDelete;
+		TSet<FNiagaraVariable> VariablesToDelete;
 		for (TSharedRef<INiagaraParameterViewModel> OutputParameter : GetSelection().GetSelectedObjects())
 		{
-			OutputNamesToDelete.Add(OutputParameter->GetName());
+			VariablesToDelete.Add(OutputParameter->GetVariable());
 		}
 		GetSelection().ClearSelectedObjects();
+		
+		DeleteParameters(VariablesToDelete.Array());
+	}
+}
 
-		auto DeletePredicate = [&OutputNamesToDelete] (FNiagaraVariable& OutputVariable)
-		{ 
-			return OutputNamesToDelete.Contains(OutputVariable.GetName());
-		};
+void FNiagaraScriptOutputCollectionViewModel::DeleteParameters(TArray<FNiagaraVariable> ParametersToDelete)
+{
+	TArray<UNiagaraNodeOutput*> OutputNodes;
+	Graph->FindOutputNodes(OutputNodes);
+	if (OutputNodes.Num() == 0)
+	{
+		return;
+	}
 
-		FScopedTransaction ScopedTransaction(LOCTEXT("DeletedSelectedNodes", "Delete selected nodes"));
-		for (UNiagaraNodeOutput* OtherOutputNode : OutputNodes)
-		{
-			OtherOutputNode->Modify();
-			OtherOutputNode->Outputs.RemoveAll(DeletePredicate);
-			OtherOutputNode->NotifyOutputVariablesChanged();
-		}
-		OnOutputParametersChangedDelegate.Broadcast();
+	auto DeletePredicate = [&ParametersToDelete] (FNiagaraVariable& OutputVariable)
+	{ 
+		return ParametersToDelete.Contains(OutputVariable);
+	};
+	
+	FScopedTransaction ScopedTransaction(LOCTEXT("DeletedSelectedNodes", "Delete selected nodes"));
+	for (UNiagaraNodeOutput* OtherOutputNode : OutputNodes)
+	{
+		OtherOutputNode->Modify();
+		OtherOutputNode->Outputs.RemoveAll(DeletePredicate);
+		OtherOutputNode->NotifyOutputVariablesChanged();
 	}
+	OnOutputParametersChangedDelegate.Broadcast();
 }
 
 const TArray<TSharedRef<INiagaraParameterViewModel>>& FNiagaraScriptOutputCollectionViewModel::GetParameters()
diff --git a/Engine/Plugins/FX/Niagara/Source/NiagaraEditor/Private/ViewModels/NiagaraScriptOutputCollectionViewModel.h b/Engine/Plugins/FX/Niagara/Source/NiagaraEditor/Private/ViewModels/NiagaraScriptOutputCollectionViewModel.h
--- a/Engine/Plugins/FX/Niagara/Source/NiagaraEditor/Private/ViewModels/NiagaraScriptOutputCollectionViewModel.h
+++ b/Engine/Plugins/FX/Niagara/Source/NiagaraEditor/Private/ViewModels/NiagaraScriptOutputCollectionViewModel.h
@@ -30,6 +30,7 @@ public:
 	virtual void AddParameter(TSharedPtr<FNiagaraTypeDefinition> ParameterType) override;
 	virtual bool CanDeleteParameters() const override;
 	virtual void DeleteSelectedParameters() override;
+	virtual void DeleteParameters(TArray<FNiagaraVariable> ParametersToDelete) override;
 	virtual const TArray<TSharedRef<INiagaraParameterViewModel>>& GetParameters() override;
 
 	/** Rebuilds the parameter view models. */
diff --git a/Engine/Plugins/FX/Niagara/Source/NiagaraEditor/Private/ViewModels/NiagaraScriptParameterViewModel.cpp b/Engine/Plugins/FX/Niagara/Source/NiagaraEditor/Private/ViewModels/NiagaraScriptParameterViewModel.cpp
--- a/Engine/Plugins/FX/Niagara/Source/NiagaraEditor/Private/ViewModels/NiagaraScriptParameterViewModel.cpp
+++ b/Engine/Plugins/FX/Niagara/Source/NiagaraEditor/Private/ViewModels/NiagaraScriptParameterViewModel.cpp
@@ -50,6 +50,12 @@ FNiagaraScriptParameterViewModel::~FNiagaraScriptParameterViewModel()
 	//UE_LOG(LogNiagaraEditor, Log, TEXT("Delete %p Var %s"), this, *DebugName);
 }
 
+FNiagaraVariable FNiagaraScriptParameterViewModel::GetVariable() const
+{
+	check(GraphVariable);
+	return *GraphVariable;
+}
+
 void FNiagaraScriptParameterViewModel::Reset()
 {
 	GraphVariableOwner = nullptr;
@@ -81,7 +87,7 @@ FText FNiagaraScriptParameterViewModel::GetTypeDisplayName() const
 	return FText::Format(LOCTEXT("TypeTextFormat", "Type: {0}"), GraphVariable->GetType().GetNameText());
 }
 
-void FNiagaraScriptParameterViewModel::NameTextComitted(const FText& Name, ETextCommit::Type CommitInfo)
+void FNiagaraScriptParameterViewModel::NameTextCommitted(const FText& Name, ETextCommit::Type CommitInfo)
 {
 	check(GraphVariable);
 	FName NewName = *Name.ToString();
@@ -98,10 +104,10 @@ bool FNiagaraScriptParameterViewModel::VerifyNodeNameTextChanged(const FText& Ne
 	return OwningNode == nullptr || FNiagaraEditorUtilities::VerifyNameChangeForInputOrOutputNode(*OwningNode, GraphVariable->GetName(), NewText.ToString(), OutErrorMessage);
 }
 
-TSharedPtr<FNiagaraTypeDefinition> FNiagaraScriptParameterViewModel::GetType() const
+FNiagaraTypeDefinition FNiagaraScriptParameterViewModel::GetType() const
 {
 	check(GraphVariable);
-	return MakeShareable(new FNiagaraTypeDefinition(GraphVariable->GetType()));
+	return GraphVariable->GetType();
 }
 
 bool FNiagaraScriptParameterViewModel::CanChangeSortOrder() const
diff --git a/Engine/Plugins/FX/Niagara/Source/NiagaraEditor/Private/ViewModels/NiagaraScriptParameterViewModel.h b/Engine/Plugins/FX/Niagara/Source/NiagaraEditor/Private/ViewModels/NiagaraScriptParameterViewModel.h
--- a/Engine/Plugins/FX/Niagara/Source/NiagaraEditor/Private/ViewModels/NiagaraScriptParameterViewModel.h
+++ b/Engine/Plugins/FX/Niagara/Source/NiagaraEditor/Private/ViewModels/NiagaraScriptParameterViewModel.h
@@ -36,11 +36,12 @@ public:
 	virtual ~FNiagaraScriptParameterViewModel();
 
 	//~ INiagaraParameterViewmodel interface.
+	virtual FNiagaraVariable GetVariable() const override;
 	virtual FName GetName() const override;
-	virtual void NameTextComitted(const FText& Name, ETextCommit::Type CommitInfo) override;
+	virtual void NameTextCommitted(const FText& Name, ETextCommit::Type CommitInfo) override;
 	virtual bool VerifyNodeNameTextChanged(const FText& NewText, FText& OutErrorMessage) override;
 	virtual FText GetTypeDisplayName() const override;
-	virtual TSharedPtr<FNiagaraTypeDefinition> GetType() const override;
+	virtual FNiagaraTypeDefinition GetType() const override;
 	virtual void SelectedTypeChanged(TSharedPtr<FNiagaraTypeDefinition> Item, ESelectInfo::Type SelectionType) override;
 	virtual EDefaultValueType GetDefaultValueType() override;
 	virtual TSharedRef<FStructOnScope> GetDefaultValueStruct() override;
diff --git a/Engine/Plugins/FX/Niagara/Source/NiagaraEditor/Private/Widgets/SNiagaraParameterCollection.cpp b/Engine/Plugins/FX/Niagara/Source/NiagaraEditor/Private/Widgets/SNiagaraParameterCollection.cpp
--- a/Engine/Plugins/FX/Niagara/Source/NiagaraEditor/Private/Widgets/SNiagaraParameterCollection.cpp
+++ b/Engine/Plugins/FX/Niagara/Source/NiagaraEditor/Private/Widgets/SNiagaraParameterCollection.cpp
@@ -417,6 +417,12 @@ TSharedRef<SWidget> SNiagaraParameterCollection::GetAddMenuContent()
 	return AddMenuBuilder.MakeWidget();
 }
 
+FReply SNiagaraParameterCollection::DeleteParameter(TSharedRef<INiagaraParameterViewModel> Item) const
+{
+	Collection->DeleteParameters({FNiagaraVariable(Item->GetType(), Item->GetName())});
+	return FReply::Handled();
+}
+
 TSharedRef<ITableRow> SNiagaraParameterCollection::OnGenerateRowForParameter(TSharedRef<INiagaraParameterViewModel> Item, const TSharedRef<STableViewBase>& OwnerTable)
 {
 	// Name widget
@@ -445,7 +451,7 @@ TSharedRef<ITableRow> SNiagaraParameterCollection::OnGenerateRowForParameter(TSh
 				.Style(FNiagaraEditorStyle::Get(), "NiagaraEditor.ParameterInlineEditableText")
 				.Text(Item, &INiagaraParameterViewModel::GetNameText)
 				.OnVerifyTextChanged(Item, &INiagaraParameterViewModel::VerifyNodeNameTextChanged)
-				.OnTextCommitted(Item, &INiagaraParameterViewModel::NameTextComitted)
+				.OnTextCommitted(Item, &INiagaraParameterViewModel::NameTextCommitted)
 				.IsSelected(this, &SNiagaraParameterCollection::IsItemSelected, Item)
 				.IsEnabled(TAttribute<bool>(Item, &INiagaraParameterViewModel::IsEditingEnabled))
 			];
@@ -482,7 +488,7 @@ TSharedRef<ITableRow> SNiagaraParameterCollection::OnGenerateRowForParameter(TSh
 	if (Item->GetDefaultValueType() == INiagaraParameterViewModel::EDefaultValueType::Struct)
 	{
 		FNiagaraEditorModule& NiagaraEditorModule = FModuleManager::GetModuleChecked<FNiagaraEditorModule>("NiagaraEditor");
-		FNiagaraTypeDefinition ParameterType = *Item->GetType().Get();
+		FNiagaraTypeDefinition ParameterType = Item->GetType();
 		TSharedPtr<INiagaraEditorTypeUtilities, ESPMode::ThreadSafe> TypeEditorUtilities = NiagaraEditorModule.GetTypeUtilities(ParameterType);
 		TSharedPtr<SNiagaraParameterEditor> ParameterEditor;
 		if (TypeEditorUtilities.IsValid() && TypeEditorUtilities->CanCreateParameterEditor())
@@ -522,7 +528,7 @@ TSharedRef<ITableRow> SNiagaraParameterCollection::OnGenerateRowForParameter(TSh
 	else if (Item->GetDefaultValueType() == INiagaraParameterViewModel::EDefaultValueType::DataInterface)
 	{
 		FNiagaraEditorModule& NiagaraEditorModule = FModuleManager::GetModuleChecked<FNiagaraEditorModule>("NiagaraEditor");
-		TSharedPtr<INiagaraEditorTypeUtilities, ESPMode::ThreadSafe> TypeEditorUtilities = NiagaraEditorModule.GetTypeUtilities(*Item->GetType().Get());
+		TSharedPtr<INiagaraEditorTypeUtilities, ESPMode::ThreadSafe> TypeEditorUtilities = NiagaraEditorModule.GetTypeUtilities(Item->GetType());
 		if (TypeEditorUtilities.IsValid() && TypeEditorUtilities->CanCreateDataInterfaceEditor())
 		{
 			CustomValueEditor = TypeEditorUtilities->CreateDataInterfaceEditor(Item->GetDefaultValueDataInterface(),
@@ -543,7 +549,7 @@ TSharedRef<ITableRow> SNiagaraParameterCollection::OnGenerateRowForParameter(TSh
 	}
 	else if (Item->GetDefaultValueType() == INiagaraParameterViewModel::EDefaultValueType::ObjectAsset)
 	{
-		const FNiagaraTypeDefinition ParameterType = *Item->GetType().Get();
+		const FNiagaraTypeDefinition ParameterType = Item->GetType();
 
 		CustomValueEditor = SNew(SObjectPropertyEntryBox)
 			.AllowedClass(ParameterType.GetClass())
@@ -571,6 +577,15 @@ TSharedRef<ITableRow> SNiagaraParameterCollection::OnGenerateRowForParameter(TSh
 		DetailsWidget = SNullWidget::NullWidget;
 	}
 
+	TSharedRef<SWidget> ParameterDeleteButton = SNew(SButton)
+		.OnClicked_Raw(this, &SNiagaraParameterCollection::DeleteParameter, Item)
+		.Visibility_Raw(this, &SNiagaraParameterCollection::ShowParameterDeleteButton)
+		.Content()
+		[
+			SNew(SImage)
+			.Image(FAppStyle::GetBrush("Icons.Delete"))
+		];
+	
 	if (CustomValueEditor.IsValid())
 	{
 		// Should the value of the parameter be enabled?
@@ -609,6 +624,12 @@ TSharedRef<ITableRow> SNiagaraParameterCollection::OnGenerateRowForParameter(TSh
 					[
 						CustomValueEditor.ToSharedRef()
 					]
+					+ SHorizontalBox::Slot()
+					.HAlign(HAlign_Right)
+					.Padding(4.f)
+					[
+						ParameterDeleteButton
+					]
 #if defined(DEBUG_SORT_ORDER)
 					+SHorizontalBox::Slot()
 					.AutoWidth()
@@ -661,6 +682,12 @@ TSharedRef<ITableRow> SNiagaraParameterCollection::OnGenerateRowForParameter(TSh
 				.OnSlotResized(SSplitter::FOnSlotResized::CreateSP(this, &SNiagaraParameterCollection::ParameterContentColumnWidthChanged))
 				[
 					SNew(SHorizontalBox)
+					+ SHorizontalBox::Slot()
+					.HAlign(HAlign_Right)
+					.Padding(4.f)
+					[
+						ParameterDeleteButton
+					]
 #if defined(DEBUG_SORT_ORDER)
 					+ SHorizontalBox::Slot()
 					.AutoWidth()
@@ -832,7 +859,11 @@ FReply SNiagaraParameterCollection::OnItemAcceptDrop(FDragDropEvent const& DragD
 	return bWasDropHandled ? FReply::Handled() : FReply::Unhandled();
 }
 
-
+EVisibility SNiagaraParameterCollection::ShowParameterDeleteButton() const
+{
+	// if we can add, we can also remove
+	return Collection->GetAddButtonVisibility();
+}
 
 void SNiagaraParameterCollection::ParameterViewModelTypeChanged()
 {
diff --git a/Engine/Plugins/FX/Niagara/Source/NiagaraEditor/Private/Widgets/SNiagaraParameterCollection.h b/Engine/Plugins/FX/Niagara/Source/NiagaraEditor/Private/Widgets/SNiagaraParameterCollection.h
--- a/Engine/Plugins/FX/Niagara/Source/NiagaraEditor/Private/Widgets/SNiagaraParameterCollection.h
+++ b/Engine/Plugins/FX/Niagara/Source/NiagaraEditor/Private/Widgets/SNiagaraParameterCollection.h
@@ -66,6 +66,8 @@ private:
 	/** Gets the content for the add menu. */
 	TSharedRef<SWidget> GetAddMenuContent();
 
+	FReply DeleteParameter(TSharedRef<INiagaraParameterViewModel> Item) const;
+
 	/** Generates a row for the parameter list view. */
 	TSharedRef<ITableRow> OnGenerateRowForParameter(TSharedRef<INiagaraParameterViewModel> Item, const TSharedRef<STableViewBase>& OwnerTable);
 
@@ -109,6 +111,8 @@ private:
 	TOptional<EItemDropZone> OnItemCanAcceptDrop(const FDragDropEvent& DragDropEvent, EItemDropZone DropZone, TSharedRef<INiagaraParameterViewModel> DropItem);
 	FReply OnItemAcceptDrop(FDragDropEvent const& DragDropEvent, EItemDropZone DropZone, TSharedRef<INiagaraParameterViewModel> DropItem);
 
+	EVisibility ShowParameterDeleteButton() const;
+
 private:
 	/** The view model for the parameter collection. */
 	TSharedPtr<INiagaraParameterCollectionViewModel> Collection;
