Fixed compile issue

diff --git a/Engine/Plugins/FX/Niagara/Source/NiagaraEditor/Private/Customizations/NiagaraScriptDetails.cpp b/Engine/Plugins/FX/Niagara/Source/NiagaraEditor/Private/Customizations/NiagaraScriptDetails.cpp
--- a/Engine/Plugins/FX/Niagara/Source/NiagaraEditor/Private/Customizations/NiagaraScriptDetails.cpp
+++ b/Engine/Plugins/FX/Niagara/Source/NiagaraEditor/Private/Customizations/NiagaraScriptDetails.cpp
@@ -12,6 +12,8 @@
 #include "NiagaraMetaDataCustomNodeBuilder.h"
 #include "NiagaraParameterCollectionCustomNodeBuilder.h"
 #include "ViewModels/NiagaraParameterCollectionViewModel.h"
+#include "ViewModels/NiagaraScriptInputCollectionViewModel.h"
+#include "ViewModels/NiagaraScriptOutputCollectionViewModel.h"
 #include "ViewModels/NiagaraParameterViewModel.h"
 #include "NiagaraScript.h"
 #include "ViewModels/NiagaraScriptGraphViewModel.h"
