- Default SubUV blending to enabled on mesh / sprite renderers

diff --git a/Engine/Plugins/FX/Niagara/Source/Niagara/Private/NiagaraMeshRendererProperties.cpp b/Engine/Plugins/FX/Niagara/Source/Niagara/Private/NiagaraMeshRendererProperties.cpp
--- a/Engine/Plugins/FX/Niagara/Source/Niagara/Private/NiagaraMeshRendererProperties.cpp
+++ b/Engine/Plugins/FX/Niagara/Source/Niagara/Private/NiagaraMeshRendererProperties.cpp
@@ -245,7 +245,7 @@ UNiagaraMeshRendererProperties::UNiagaraMeshRendererProperties()
 	: bOverrideMaterials(false)
 	, bUseHeterogeneousVolumes(false)
 	, bSortOnlyWhenTranslucent(true)
-	, bSubImageBlend(false)
+	, bSubImageBlend(true)
 	, bLockedAxisEnable(false)
 {
 	// Initialize the array with a single, defaulted entry
@@ -391,9 +391,16 @@ void UNiagaraMeshRendererProperties::Serialize(FArchive& Ar)
 	Ar.UsingCustomVersion(FNiagaraCustomVersion::GUID);
 	const int32 NiagaraVersion = Ar.CustomVer(FNiagaraCustomVersion::GUID);
 
-	if (Ar.IsLoading() && (NiagaraVersion < FNiagaraCustomVersion::DisableSortingByDefault))
+	if (Ar.IsLoading())
 	{
-		SortMode = ENiagaraSortMode::ViewDistance;
+		if (NiagaraVersion < FNiagaraCustomVersion::DisableSortingByDefault)
+		{
+			SortMode = ENiagaraSortMode::ViewDistance;
+		}
+		if (NiagaraVersion < FNiagaraCustomVersion::SubImageBlendEnabledByDefault)
+		{
+			bSubImageBlend = false;
+		}
 	}
 
 	Super::Serialize(Ar);
diff --git a/Engine/Plugins/FX/Niagara/Source/Niagara/Private/NiagaraSpriteRendererProperties.cpp b/Engine/Plugins/FX/Niagara/Source/Niagara/Private/NiagaraSpriteRendererProperties.cpp
--- a/Engine/Plugins/FX/Niagara/Source/Niagara/Private/NiagaraSpriteRendererProperties.cpp
+++ b/Engine/Plugins/FX/Niagara/Source/Niagara/Private/NiagaraSpriteRendererProperties.cpp
@@ -54,7 +54,7 @@ FCookStatsManager::FAutoRegisterCallback NiagaraCutoutCookStats::RegisterCookSta
 UNiagaraSpriteRendererProperties::UNiagaraSpriteRendererProperties()
 	: Material(nullptr)
 	, MaterialUserParamBinding(FNiagaraTypeDefinition(UMaterialInterface::StaticClass()))
-	, bSubImageBlend(false)
+	, bSubImageBlend(true)
 	, bRemoveHMDRollInVR(false)
 	, bSortOnlyWhenTranslucent(true)
 #if WITH_EDITORONLY_DATA
@@ -213,15 +213,22 @@ void UNiagaraSpriteRendererProperties::Serialize(FStructuredArchive::FRecord Rec
 	const int32 NiagaraVersion = Ar.CustomVer(FNiagaraCustomVersion::GUID);
 	const int32 UE5MainVersion = Ar.CustomVer(FUE5MainStreamObjectVersion::GUID);
 
-	if (Ar.IsLoading() && (NiagaraVersion < FNiagaraCustomVersion::DisableSortingByDefault))
+	// Default Property Changes
+	if (Ar.IsLoading())
 	{
-		SortMode = ENiagaraSortMode::ViewDistance;
-	}
-
-	if (Ar.IsLoading() && (UE5MainVersion < FUE5MainStreamObjectVersion::NiagaraSpriteRendererFacingAlignmentAutoDefault))
-	{
-		Alignment = ENiagaraSpriteAlignment::Unaligned;
-		FacingMode = ENiagaraSpriteFacingMode::FaceCamera;
+		if (NiagaraVersion < FNiagaraCustomVersion::DisableSortingByDefault)
+		{
+			SortMode = ENiagaraSortMode::ViewDistance;
+		}
+		if (UE5MainVersion < FUE5MainStreamObjectVersion::NiagaraSpriteRendererFacingAlignmentAutoDefault)
+		{
+			Alignment = ENiagaraSpriteAlignment::Unaligned;
+			FacingMode = ENiagaraSpriteFacingMode::FaceCamera;
+		}
+		if (NiagaraVersion < FNiagaraCustomVersion::SubImageBlendEnabledByDefault)
+		{
+			bSubImageBlend = false;
+		}
 	}
 
 	// MIC will replace the main material during serialize
diff --git a/Engine/Plugins/FX/Niagara/Source/NiagaraCore/Public/NiagaraCustomVersion.h b/Engine/Plugins/FX/Niagara/Source/NiagaraCore/Public/NiagaraCustomVersion.h
--- a/Engine/Plugins/FX/Niagara/Source/NiagaraCore/Public/NiagaraCustomVersion.h
+++ b/Engine/Plugins/FX/Niagara/Source/NiagaraCore/Public/NiagaraCustomVersion.h
@@ -189,6 +189,8 @@ struct FNiagaraCustomVersion
 
 		RibbonRendererLinkOrderDefaultIsUniqueID,	// Ribbon renderer will default to unique ID rather than normalized age to make more things 'just work'
 
+		SubImageBlendEnabledByDefault,	// Renderer SubImage Blends are enabled by default
+
 		// DO NOT ADD A NEW VERSION UNLESS YOU HAVE TALKED TO THE NIAGARA LEAD. Mismanagement of these versions can lead to data loss if it is adjusted in multiple streams simultaneously.
 		// -----<new versions can be added above this line>  -------------------------------------------------
 		VersionPlusOne,
