- Reset the ticker & completion delegates once FinishCapture has fired

#rnx

diff --git a/Engine/Plugins/FX/Niagara/Source/Niagara/Classes/NiagaraSimCacheCapture.cpp b/Engine/Plugins/FX/Niagara/Source/Niagara/Classes/NiagaraSimCacheCapture.cpp
--- a/Engine/Plugins/FX/Niagara/Source/Niagara/Classes/NiagaraSimCacheCapture.cpp
+++ b/Engine/Plugins/FX/Niagara/Source/Niagara/Classes/NiagaraSimCacheCapture.cpp
@@ -20,6 +20,7 @@ void FNiagaraSimCacheCapture::FinishCapture()
 	if (TickerHandle.IsValid())
 	{
 		FTSTicker::GetCoreTicker().RemoveTicker(TickerHandle);
+		TickerHandle.Reset();
 	}
 
 	UNiagaraSimCache* CaptureSimCache = WeakCaptureSimCache.Get();
@@ -28,6 +29,8 @@ void FNiagaraSimCacheCapture::FinishCapture()
 		CaptureSimCache->EndWrite();
 	}
 	CaptureComplete.Broadcast(CaptureSimCache);
+	CaptureComplete.Clear();
+
 	WeakCaptureComponent = nullptr;
 }
 
