- Fix a crash when changing levels after capturing a sim cache for debugging from the details panel
- The transient package has already been cleared out so the object pointer was stale

diff --git a/Engine/Plugins/FX/Niagara/Source/NiagaraEditor/Private/Customizations/NiagaraComponentDetails.cpp b/Engine/Plugins/FX/Niagara/Source/NiagaraEditor/Private/Customizations/NiagaraComponentDetails.cpp
--- a/Engine/Plugins/FX/Niagara/Source/NiagaraEditor/Private/Customizations/NiagaraComponentDetails.cpp
+++ b/Engine/Plugins/FX/Niagara/Source/NiagaraEditor/Private/Customizations/NiagaraComponentDetails.cpp
@@ -76,9 +76,12 @@ FNiagaraComponentDetails::~FNiagaraComponentDetails()
 		GEngine->OnWorldDestroyed().RemoveAll(this);
 	}
 
-	for(UNiagaraSimCache*& SimCache : CapturedCaches)
+	for (const TWeakObjectPtr<UNiagaraSimCache>& WeakSimCache : CapturedCaches)
 	{
-		SimCache->MarkAsGarbage();
+		if (UNiagaraSimCache* SimCache = WeakSimCache.Get())
+		{
+			SimCache->MarkAsGarbage();
+		}
 	}
 
 	FGameDelegates::Get().GetEndPlayMapDelegate().RemoveAll(this);
@@ -357,7 +360,7 @@ FReply FNiagaraComponentDetails::OnCaptureSelectedSystem()
 		int32 CaptureIndex = ComponentCaptures.AddDefaulted();
 		UNiagaraSimCache* SimCache = NewObject<UNiagaraSimCache>(GetTransientPackage());
 		SimCache->SetFlags(RF_Transient);
-		CapturedCaches.Add(SimCache);
+		CapturedCaches.Emplace(SimCache);
 		FNiagaraSimCacheCapture& SimCacheCapture = ComponentCaptures[CaptureIndex];
 
 		FNiagaraSimCacheCreateParameters CreateParameters;
diff --git a/Engine/Plugins/FX/Niagara/Source/NiagaraEditor/Private/Customizations/NiagaraComponentDetails.h b/Engine/Plugins/FX/Niagara/Source/NiagaraEditor/Private/Customizations/NiagaraComponentDetails.h
--- a/Engine/Plugins/FX/Niagara/Source/NiagaraEditor/Private/Customizations/NiagaraComponentDetails.h
+++ b/Engine/Plugins/FX/Niagara/Source/NiagaraEditor/Private/Customizations/NiagaraComponentDetails.h
@@ -55,7 +55,7 @@ private:
 	IDetailLayoutBuilder* Builder = nullptr;
 
 	TArray<FNiagaraSimCacheCapture> ComponentCaptures;
-	TArray<UNiagaraSimCache*> CapturedCaches;
+	TArray<TWeakObjectPtr<UNiagaraSimCache>> CapturedCaches;
 };
 
 class FNiagaraSystemUserParameterDetails : public IDetailCustomization
