- Add InsideVolume output to RVT sample
- Change user parameter binding to be a UObject rather than exposing URuntimeVirtualTexture to Niagara

diff --git a/Engine/Plugins/FX/Niagara/Shaders/Private/NiagaraDataInterfaceVirtualTextureTemplate.ush b/Engine/Plugins/FX/Niagara/Shaders/Private/NiagaraDataInterfaceVirtualTextureTemplate.ush
--- a/Engine/Plugins/FX/Niagara/Shaders/Private/NiagaraDataInterfaceVirtualTextureTemplate.ush
+++ b/Engine/Plugins/FX/Niagara/Shaders/Private/NiagaraDataInterfaceVirtualTextureTemplate.ush
@@ -76,8 +76,9 @@ float4 SampleRVTLayer_{ParameterName}(float2 SampleUV, Texture2D InTexture, Text
 	return TextureVirtualSample(InTexture, {ParameterName}_SharedSampler, PageTable, 0, VTUniform_Unpack(InTextureUniforms));
 }
 
-void SampleRVT_{ParameterName}(in float3 WorldPosition, out float3 BaseColor, out float Specular, out float Roughness, out float3 Normal, out float WorldHeight, out float Mask)
+void SampleRVT_{ParameterName}(in float3 WorldPosition, out bool bInsideVolume, out float3 BaseColor, out float Specular, out float Roughness, out float3 Normal, out float WorldHeight, out float Mask)
 {
+	bInsideVolume = false;
 	BaseColor = float3(0.0f, 0.0f, 0.0f);
 	Specular = 0.5f;
 	Roughness = 0.5f;
@@ -91,6 +92,9 @@ void SampleRVT_{ParameterName}(in float3 WorldPosition, out float3 BaseColor, ou
 
 	float2 SampleUV = VirtualTextureWorldToUV(LWCWorldPosition, LWCUVOrigin, {ParameterName}_UVUniforms[1].xyz, {ParameterName}_UVUniforms[2].xyz);
 
+	// Test to see if we are inside the volume, but still take the samples as it will clamp to the edge
+	bInsideVolume = all(SampleUV >- 0.0f) && all(SampleUV < 1.0f);
+
 	// Sample Textures
 	float4 LayerSample[3];
 	LayerSample[0] = ({ParameterName}_ValidLayersMask & 0x1) != 0 ? SampleRVTLayer_{ParameterName}(SampleUV, {ParameterName}_VirtualTexture0, {ParameterName}_VirtualTexture0PageTable, {ParameterName}_VirtualTexture0TextureUniforms) : 0;
diff --git a/Engine/Plugins/FX/Niagara/Source/Niagara/Private/DataInterface/NiagaraDataInterfaceVirtualTexture.cpp b/Engine/Plugins/FX/Niagara/Source/Niagara/Private/DataInterface/NiagaraDataInterfaceVirtualTexture.cpp
--- a/Engine/Plugins/FX/Niagara/Source/Niagara/Private/DataInterface/NiagaraDataInterfaceVirtualTexture.cpp
+++ b/Engine/Plugins/FX/Niagara/Source/Niagara/Private/DataInterface/NiagaraDataInterfaceVirtualTexture.cpp
@@ -46,6 +46,18 @@ namespace NDIVirtualTextureLocal
 	const FName		SampleRVTName("SampleRVT");
 	constexpr int32	MaxRVTLayers = 3;
 
+	struct FNDIFunctionVersion
+	{
+		enum Type
+		{
+			InitialVersion = 0,
+			AddInsideVolumeToSampleRVT = 1,
+
+			VersionPlusOne,
+			LatestVersion = VersionPlusOne - 1
+		};
+	};
+
 	struct FNDIInstanceData_GameThread
 	{
 		TWeakObjectPtr<URuntimeVirtualTexture> CurrentTexture = nullptr;
@@ -96,7 +108,7 @@ UNiagaraDataInterfaceVirtualTexture::UNiagaraDataInterfaceVirtualTexture(FObject
 
 	Proxy.Reset(new FNDIProxy());
 
-	FNiagaraTypeDefinition Def(URuntimeVirtualTexture::StaticClass());
+	FNiagaraTypeDefinition Def(UObject::StaticClass());
 	TextureUserParameter.Parameter.SetType(Def);
 }
 
@@ -174,6 +186,7 @@ void UNiagaraDataInterfaceVirtualTexture::GetFunctions(TArray<FNiagaraFunctionSi
 		Sig.bSupportsCPU = false;
 		Sig.bSupportsGPU = true;
 		Sig.bExperimental = true;
+		Sig.SetFunctionVersion(FNDIFunctionVersion::LatestVersion);
 		Sig.Inputs.Emplace(FNiagaraTypeDefinition(GetClass()), TEXT("Texture"));
 		Sig.Outputs.Emplace(FNiagaraTypeDefinition::GetBoolDef(), TEXT("BaseColorValid"));
 		Sig.Outputs.Emplace(FNiagaraTypeDefinition::GetBoolDef(), TEXT("SpecularValid"));
@@ -191,8 +204,10 @@ void UNiagaraDataInterfaceVirtualTexture::GetFunctions(TArray<FNiagaraFunctionSi
 		Sig.bSupportsCPU = false;
 		Sig.bSupportsGPU = true;
 		Sig.bExperimental = true;
+		Sig.SetFunctionVersion(FNDIFunctionVersion::LatestVersion);
 		Sig.Inputs.Emplace(FNiagaraTypeDefinition(GetClass()), TEXT("Texture"));
 		Sig.Inputs.Emplace(FNiagaraTypeDefinition::GetPositionDef(), TEXT("WorldPosition"));
+		Sig.Outputs.Emplace(FNiagaraTypeDefinition::GetBoolDef(), TEXT("InsideVolume"));
 		Sig.Outputs.Emplace(FNiagaraTypeDefinition::GetVec3Def(), TEXT("BaseColor"));
 		Sig.Outputs.Emplace(FNiagaraTypeDefinition::GetFloatDef(), TEXT("Specular"));
 		Sig.Outputs.Emplace(FNiagaraTypeDefinition::GetFloatDef(), TEXT("Roughness"));
@@ -326,6 +341,24 @@ bool UNiagaraDataInterfaceVirtualTexture::GetFunctionHLSL(const FNiagaraDataInte
 
 	return ValidGpuFunctions.Contains(FunctionInfo.DefinitionName);
 }
+
+bool UNiagaraDataInterfaceVirtualTexture::UpgradeFunctionCall(FNiagaraFunctionSignature& Sig)
+{
+	using namespace NDIVirtualTextureLocal;
+
+	bool bModified = false;
+	if (Sig.FunctionVersion < FNDIFunctionVersion::AddInsideVolumeToSampleRVT)
+	{
+		if (Sig.Name == SampleRVTName)
+		{
+			Sig.Outputs.Insert(FNiagaraVariable(FNiagaraTypeDefinition::GetBoolDef(), TEXT("InsideVolume")), 0);
+			bModified = true;
+		}
+	}
+
+	Sig.SetFunctionVersion(FNDIFunctionVersion::LatestVersion);
+	return bModified;
+}
 #endif
 
 void UNiagaraDataInterfaceVirtualTexture::BuildShaderParameters(FNiagaraShaderParametersBuilder& ShaderParametersBuilder) const
diff --git a/Engine/Plugins/FX/Niagara/Source/Niagara/Private/DataInterface/NiagaraDataInterfaceVirtualTexture.h b/Engine/Plugins/FX/Niagara/Source/Niagara/Private/DataInterface/NiagaraDataInterfaceVirtualTexture.h
--- a/Engine/Plugins/FX/Niagara/Source/Niagara/Private/DataInterface/NiagaraDataInterfaceVirtualTexture.h
+++ b/Engine/Plugins/FX/Niagara/Source/Niagara/Private/DataInterface/NiagaraDataInterfaceVirtualTexture.h
@@ -55,6 +55,7 @@ public:
 	NIAGARA_API virtual void GetCommonHLSL(FString& OutHlsl) override;
 	NIAGARA_API virtual void GetParameterDefinitionHLSL(const FNiagaraDataInterfaceGPUParamInfo& ParamInfo, FString& OutHLSL) override;
 	NIAGARA_API virtual bool GetFunctionHLSL(const FNiagaraDataInterfaceGPUParamInfo& ParamInfo, const FNiagaraDataInterfaceGeneratedFunction& FunctionInfo, int FunctionInstanceIndex, FString& OutHLSL) override;
+	NIAGARA_API virtual bool UpgradeFunctionCall(FNiagaraFunctionSignature& FunctionSignature) override;
 #endif
 	NIAGARA_API virtual void BuildShaderParameters(FNiagaraShaderParametersBuilder& ShaderParametersBuilder) const override;
 	NIAGARA_API virtual void SetShaderParameters(const FNiagaraDataInterfaceSetShaderParametersContext& Context) const override;
