- Remove unsued variable

#rnx

diff --git a/Engine/Plugins/FX/Niagara/Source/Niagara/Classes/NiagaraSystem.h b/Engine/Plugins/FX/Niagara/Source/Niagara/Classes/NiagaraSystem.h
--- a/Engine/Plugins/FX/Niagara/Source/Niagara/Classes/NiagaraSystem.h
+++ b/Engine/Plugins/FX/Niagara/Source/Niagara/Classes/NiagaraSystem.h
@@ -545,9 +545,6 @@ protected:
 	UPROPERTY(EditAnywhere, AdvancedDisplay, Category = "Performance", meta = (DisplayName="Disable Debug Switches"))
 	uint32 bDisableDebugSwitchesOnCook : 1;
 
-	UPROPERTY(EditAnywhere, AdvancedDisplay, Category = "System")
-	bool bUseLegacyEmitterDataInterfaceDepencendyResolution = false;
-
 	/* When set the system needs to compile before it can be activated. */
 	uint32 bNeedsRequestCompile : 1;
 
