- Add manual scopes so we can see assets in GDME

diff --git a/Engine/Plugins/FX/Niagara/Source/Niagara/Private/NiagaraComponent.cpp b/Engine/Plugins/FX/Niagara/Source/Niagara/Private/NiagaraComponent.cpp
--- a/Engine/Plugins/FX/Niagara/Source/Niagara/Private/NiagaraComponent.cpp
+++ b/Engine/Plugins/FX/Niagara/Source/Niagara/Private/NiagaraComponent.cpp
@@ -199,6 +199,9 @@ FNiagaraSceneProxy::FNiagaraSceneProxy(UNiagaraComponent* InComponent)
 		bIsHeterogeneousVolume = RenderData->HasAnyHeterogeneousVolumesEnabled();
 
 		SystemStatID = InComponent->GetAsset()->GetStatID(false, false);
+#if NIAGARAPROXY_EVENTS_ENABLED
+		SystemStatString = InComponent->GetAsset()->GetFName().ToString();
+#endif
 	}
 }
 
@@ -440,6 +443,40 @@ void FNiagaraSceneProxy::GetDynamicMeshElements(const TArray<const FSceneView*>&
 
 	FScopeCycleCounter SystemStatCounter(SystemStatID);
 
+#if NIAGARAPROXY_EVENTS_ENABLED
+	bool bEndPlatformEvent = false;
+#if CPUPROFILERTRACE_ENABLED
+	bool bEndTraceEvent = false;
+	if (UE_TRACE_CHANNELEXPR_IS_ENABLED(CpuChannel))
+	{
+		bEndTraceEvent = true;
+		FCpuProfilerTrace::OutputBeginDynamicEvent(GetResourceName());
+	}
+	else
+#endif //CPUPROFILERTRACE_ENABLED
+#if ENABLE_STATNAMEDEVENTS
+	{
+		bEndPlatformEvent = true;
+		FPlatformMisc::BeginNamedEvent(FColor(0), *SystemStatString);
+	}
+#endif //ENABLE_STATNAMEDEVENTS
+	ON_SCOPE_EXIT
+	{
+#if CPUPROFILERTRACE_ENABLED
+		if (bEndTraceEvent)
+		{
+			FCpuProfilerTrace::OutputEndEvent();
+		}
+#endif //CPUPROFILERTRACE_ENABLED
+#if ENABLE_STATNAMEDEVENTS
+		if (bEndPlatformEvent)
+		{
+			FPlatformMisc::EndNamedEvent();
+		}
+#endif //ENABLE_STATNAMEDEVENTS
+	};
+#endif
+
 	if (RenderData)
 	{
 		// Renderers often rely upon executing GPU code, such as sorting / culling, if the FX system is invalid then the count manager, etc,
diff --git a/Engine/Plugins/FX/Niagara/Source/Niagara/Public/NiagaraSceneProxy.h b/Engine/Plugins/FX/Niagara/Source/Niagara/Public/NiagaraSceneProxy.h
--- a/Engine/Plugins/FX/Niagara/Source/Niagara/Public/NiagaraSceneProxy.h
+++ b/Engine/Plugins/FX/Niagara/Source/Niagara/Public/NiagaraSceneProxy.h
@@ -11,6 +11,8 @@ class FNiagaraSystemRenderData;
 class UNiagaraComponent;
 enum class ENiagaraOcclusionQueryMode : uint8;
 
+#define NIAGARAPROXY_EVENTS_ENABLED	(!STATS && (ENABLE_STATNAMEDEVENTS || CPUPROFILERTRACE_ENABLED))
+
 /**
 * Scene proxy for drawing niagara particle simulations.
 */
@@ -96,6 +98,9 @@ private:
 	FMatrix LocalToWorldInverse;
 
 	TStatId SystemStatID;
+#if NIAGARAPROXY_EVENTS_ENABLED
+	FString SystemStatString;
+#endif
 
 	FDynamicData DynamicData;
 
