- Fix array check to use inner for struct / object check

#rnx

diff --git a/Engine/Plugins/FX/Niagara/Source/NiagaraEditor/Private/NiagaraValidationRules.cpp b/Engine/Plugins/FX/Niagara/Source/NiagaraEditor/Private/NiagaraValidationRules.cpp
--- a/Engine/Plugins/FX/Niagara/Source/NiagaraEditor/Private/NiagaraValidationRules.cpp
+++ b/Engine/Plugins/FX/Niagara/Source/NiagaraEditor/Private/NiagaraValidationRules.cpp
@@ -218,6 +218,12 @@ namespace NiagaraValidation
 		for (TFieldIterator<const FProperty> PropertyIt(Struct); PropertyIt; ++PropertyIt)
 		{
 			const FProperty* Property = *PropertyIt;
+			if (const FArrayProperty* ArrayProperty = CastField<const FArrayProperty>(Property))
+			{
+				// If we are an array change the property to be the inner one to check for struct / object
+				Property = ArrayProperty->Inner;
+			}
+
 			if (const FStructProperty* StructProperty = CastField<const FStructProperty>(Property))
 			{
 				if (StructProperty->Struct)
@@ -228,10 +234,6 @@ namespace NiagaraValidation
 					}
 				}
 			}
-			else if (const FArrayProperty* ArrayProperty = CastField<const FArrayProperty>(Property))
-			{
-				Property = ArrayProperty->Inner;
-			}
 			else if (CastField<const FWeakObjectProperty>(Property) || CastField<const FObjectProperty>(Property) || CastField<const FSoftObjectProperty>(Property))
 			{
 				return true;
