Fixed undo/redo for adding a scratch pad, consists of following fixes
- SNiagaraScratchPadEditor now keeps weak ptrs instead of shared ptrs. This would keep the scratch pad view models from being deleted properly
- Added CleanInvalidTabs function that will check tabs for validty and close them if needed
-- Added 'IsPayloadValid' function to the scratch pad tab factory that will test if a given tab's contents are still valid. Testing against all available scratch pad view models to see if a tab should get closed. This will return false if the graph the tab is displaying is no longer found in active scratch pad view models (= it got deleted via undo, even if the UObject still exists somewhere)
- Marking the scratch pad outer as RF_Transactional as the script creation undoing wouldn't get properly recorded and would remain
- clearing the selection of the details panel (used for scratch pad) since it will, upon undo, force a refresh on the deleted object

diff --git a/Engine/Plugins/FX/Niagara/Source/NiagaraEditor/Private/Toolkits/SystemToolkitModes/NiagaraSystemToolkitModeBase.cpp b/Engine/Plugins/FX/Niagara/Source/NiagaraEditor/Private/Toolkits/SystemToolkitModes/NiagaraSystemToolkitModeBase.cpp
--- a/Engine/Plugins/FX/Niagara/Source/NiagaraEditor/Private/Toolkits/SystemToolkitModes/NiagaraSystemToolkitModeBase.cpp
+++ b/Engine/Plugins/FX/Niagara/Source/NiagaraEditor/Private/Toolkits/SystemToolkitModes/NiagaraSystemToolkitModeBase.cpp
@@ -119,6 +119,8 @@ void FNiagaraSystemToolkitModeBase::OnActiveDocumentChanged(TSharedPtr<SDockTab>
 
 	SwitcherIdx = 0;
 	TSharedPtr<SDockTab>  ActiveTab = Toolkit->GetSystemViewModel()->GetDocumentViewModel()->GetActiveDocumentTab().Pin();
+	// active tab might be nullptr due to a scratch pad tab closing. In that case, the details panel should update to clear the selected objects
+	UpdateSelectionForActiveDocument();
 	if (ActiveTab.IsValid())
 	{
 		if (ActiveTab->GetLayoutIdentifier().TabType == SystemOverviewTabID)
@@ -129,7 +131,6 @@ void FNiagaraSystemToolkitModeBase::OnActiveDocumentChanged(TSharedPtr<SDockTab>
 		{
 			SwitcherIdx = 1;
 		}
-		UpdateSelectionForActiveDocument();
 
 		TSharedPtr<FNiagaraScratchPadScriptViewModel> ScratchScriptVM = Toolkit->GetSystemViewModel()->GetDocumentViewModel()->GetActiveScratchPadViewModelIfSet();
 		if (ScratchScriptVM.IsValid())
diff --git a/Engine/Plugins/FX/Niagara/Source/NiagaraEditor/Private/ViewModels/NiagaraScratchPadViewModel.cpp b/Engine/Plugins/FX/Niagara/Source/NiagaraEditor/Private/ViewModels/NiagaraScratchPadViewModel.cpp
--- a/Engine/Plugins/FX/Niagara/Source/NiagaraEditor/Private/ViewModels/NiagaraScratchPadViewModel.cpp
+++ b/Engine/Plugins/FX/Niagara/Source/NiagaraEditor/Private/ViewModels/NiagaraScratchPadViewModel.cpp
@@ -474,6 +474,13 @@ TSharedPtr<FNiagaraScratchPadScriptViewModel> UNiagaraScratchPadViewModel::Creat
 	TArray<TObjectPtr<UNiagaraScript>>* TargetScripts;
 	GetOuterAndTargetScripts(GetSystemViewModel(), ScriptOuter, TargetScripts);
 
+	// we have to make sure we can record the script outer into the transaction buffer. If not, undoing the creation of this script will keep the (deleted/undone) script in the outer
+	if(ScriptOuter->HasAnyFlags(RF_Transactional) == false)
+	{
+		ScriptOuter->SetFlags(RF_Transactional);
+	}
+	ScriptOuter->Modify();
+	
 	UNiagaraScript* NewScript = nullptr;
 	switch (InScriptUsage)
 	{
diff --git a/Engine/Plugins/FX/Niagara/Source/NiagaraEditor/Private/ViewModels/NiagaraSystemEditorDocumentsViewModel.cpp b/Engine/Plugins/FX/Niagara/Source/NiagaraEditor/Private/ViewModels/NiagaraSystemEditorDocumentsViewModel.cpp
--- a/Engine/Plugins/FX/Niagara/Source/NiagaraEditor/Private/ViewModels/NiagaraSystemEditorDocumentsViewModel.cpp
+++ b/Engine/Plugins/FX/Niagara/Source/NiagaraEditor/Private/ViewModels/NiagaraSystemEditorDocumentsViewModel.cpp
@@ -223,6 +223,14 @@ void UNiagaraSystemEditorDocumentsViewModel::SwapEditableScripts(TSharedPtr < cl
 
 }
 
+void UNiagaraSystemEditorDocumentsViewModel::CleanInvalidTabs() const
+{
+	if(DocumentManager.IsValid())
+	{
+		DocumentManager->CleanInvalidTabs();
+	}
+}
+
 
 void UNiagaraSystemEditorDocumentsViewModel::CloseChildScript(UEdGraph* InGraph)
 {
@@ -230,7 +238,7 @@ void UNiagaraSystemEditorDocumentsViewModel::CloseChildScript(UEdGraph* InGraph)
 }
 
 
-TSharedRef<FNiagaraSystemViewModel> UNiagaraSystemEditorDocumentsViewModel::GetSystemViewModel()
+TSharedRef<FNiagaraSystemViewModel> UNiagaraSystemEditorDocumentsViewModel::GetSystemViewModel() const
 {
 	TSharedPtr<FNiagaraSystemViewModel> SystemViewModel = SystemViewModelWeak.Pin();
 	checkf(SystemViewModel.IsValid(), TEXT("SystemViewModel destroyed before system editor document view model."));
diff --git a/Engine/Plugins/FX/Niagara/Source/NiagaraEditor/Private/ViewModels/NiagaraSystemViewModel.cpp b/Engine/Plugins/FX/Niagara/Source/NiagaraEditor/Private/ViewModels/NiagaraSystemViewModel.cpp
--- a/Engine/Plugins/FX/Niagara/Source/NiagaraEditor/Private/ViewModels/NiagaraSystemViewModel.cpp
+++ b/Engine/Plugins/FX/Niagara/Source/NiagaraEditor/Private/ViewModels/NiagaraSystemViewModel.cpp
@@ -1809,6 +1809,8 @@ void FNiagaraSystemViewModel::RefreshAll()
 	CurveSelectionViewModel->Refresh();
 	SystemStackViewModel->InitializeWithViewModels(this->AsShared(), TSharedPtr<FNiagaraEmitterHandleViewModel>(), FNiagaraStackViewModelOptions(true, false));
 	SelectionViewModel->Refresh();
+	// clean invalid tabs has to be called after view models have refreshed as invalid tabs can be determined by view model state
+	EditorDocumentsViewModel->CleanInvalidTabs();
 	bPendingAssetMessagesChanged = true;
 }
 
diff --git a/Engine/Plugins/FX/Niagara/Source/NiagaraEditor/Private/Widgets/NiagaraScratchScriptEditor.cpp b/Engine/Plugins/FX/Niagara/Source/NiagaraEditor/Private/Widgets/NiagaraScratchScriptEditor.cpp
--- a/Engine/Plugins/FX/Niagara/Source/NiagaraEditor/Private/Widgets/NiagaraScratchScriptEditor.cpp
+++ b/Engine/Plugins/FX/Niagara/Source/NiagaraEditor/Private/Widgets/NiagaraScratchScriptEditor.cpp
@@ -93,13 +93,13 @@ void SNiagaraScratchPadScriptEditor::SetViewModel(TSharedPtr<FNiagaraScratchPadS
 
 	ScriptViewModel = InScriptViewModel;
 
-	if (ScriptViewModel)
+	if (ScriptViewModel.IsValid())
 	{
 		if (Graph.IsValid() && Graph->GetViewModel() != InScriptViewModel->GetGraphViewModel())
 		{
 			Graph->UpdateViewModel(InScriptViewModel->GetGraphViewModel());
 		}
-		NodeIDHandle = ScriptViewModel->OnNodeIDFocusRequested().AddLambda(
+		NodeIDHandle = ScriptViewModel.Pin()->OnNodeIDFocusRequested().AddLambda(
 			[this](FNiagaraScriptIDAndGraphFocusInfo* FocusInfo)
 			{
 				if (Graph.IsValid() && FocusInfo != nullptr)
@@ -109,7 +109,7 @@ void SNiagaraScratchPadScriptEditor::SetViewModel(TSharedPtr<FNiagaraScratchPadS
 			}
 		);
 
-		PinIDHandle = ScriptViewModel->OnPinIDFocusRequested().AddLambda(
+		PinIDHandle = ScriptViewModel.Pin()->OnPinIDFocusRequested().AddLambda(
 			[this](FNiagaraScriptIDAndGraphFocusInfo* FocusInfo)
 			{
 				if (Graph.IsValid() && FocusInfo != nullptr)
@@ -235,6 +235,30 @@ TSharedRef<FGenericTabHistory> FNiagaraGraphEditorSummoner::CreateTabHistoryNode
 	return MakeShareable(new FNiagaraGraphTabHistory(SharedThis(this), Payload));
 }
 
+bool FNiagaraGraphEditorSummoner::IsPayloadValid(TSharedRef<FTabPayload> Payload) const
+{
+	bool bBaseResult = FDocumentTabFactoryForObjects<UEdGraph>::IsPayloadValid(Payload);
+
+	if(bBaseResult == false)
+	{
+		return false;
+	}
+
+	UEdGraph* ObjectPayload = FTabPayload_UObject::CastChecked<UEdGraph>(Payload);
+
+	// if our script view models don't contain the tab's graph, we know the tab is outdated and should be closed
+	// this can happen if we undo/redo the creation of a scratch pad, where the data is deleted but the tab is still displaying the graph
+	for(TSharedRef<FNiagaraScratchPadScriptViewModel> ScriptViewModel : EditorPtr.Pin()->GetSystemViewModel()->GetScriptScratchPadViewModel()->GetScriptViewModels())
+	{
+		if(ScriptViewModel->GetEditableGraphs().Contains(ObjectPayload))
+		{
+			return true;
+		}
+	}
+	
+	return false;
+}
+
 TAttribute<FText> FNiagaraGraphEditorSummoner::ConstructTabNameForObject(UEdGraph* DocumentID) const 
 {
 	UNiagaraScript* Script = DocumentID->GetTypedOuter<UNiagaraScript>();
diff --git a/Engine/Plugins/FX/Niagara/Source/NiagaraEditor/Private/Widgets/NiagaraScratchScriptEditor.h b/Engine/Plugins/FX/Niagara/Source/NiagaraEditor/Private/Widgets/NiagaraScratchScriptEditor.h
--- a/Engine/Plugins/FX/Niagara/Source/NiagaraEditor/Private/Widgets/NiagaraScratchScriptEditor.h
+++ b/Engine/Plugins/FX/Niagara/Source/NiagaraEditor/Private/Widgets/NiagaraScratchScriptEditor.h
@@ -36,7 +36,7 @@ class SNiagaraScratchPadScriptEditor : public SCompoundWidget
 
 	TSharedPtr<FNiagaraScratchPadScriptViewModel> GetViewModel()
 	{
-		return ScriptViewModel;
+		return ScriptViewModel.Pin();
 	}
 
 	void SetViewModel(TSharedPtr<FNiagaraScratchPadScriptViewModel> InViewModel);
@@ -46,42 +46,42 @@ private:
 
 	void ClearHandles()
 	{
-		if (ScriptViewModel)
+		if (ScriptViewModel.IsValid())
 		{
-			ScriptViewModel->OnNodeIDFocusRequested().Remove(NodeIDHandle);
-			ScriptViewModel->OnPinIDFocusRequested().Remove(PinIDHandle);
+			ScriptViewModel.Pin()->OnNodeIDFocusRequested().Remove(NodeIDHandle);
+			ScriptViewModel.Pin()->OnPinIDFocusRequested().Remove(PinIDHandle);
 		}
 	}
 
 	FText GetNameText() const
 	{
-		return ScriptViewModel->GetDisplayName();
+		return ScriptViewModel.Pin()->GetDisplayName();
 	}
 
 	FText GetNameToolTipText() const
 	{
-		return ScriptViewModel->GetToolTip();
+		return ScriptViewModel.Pin()->GetToolTip();
 	}
 
 	EVisibility GetUnappliedChangesVisibility() const
 	{
-		return ScriptViewModel->HasUnappliedChanges() ? EVisibility::Visible : EVisibility::Collapsed;
+		return ScriptViewModel.Pin()->HasUnappliedChanges() ? EVisibility::Visible : EVisibility::Collapsed;
 	}
 
 	FReply OnApplyButtonClicked()
 	{
-		ScriptViewModel->ApplyChanges();
+		ScriptViewModel.Pin()->ApplyChanges();
 
 		return FReply::Handled();
 	}
 
 	bool GetApplyButtonIsEnabled() const
 	{
-		return ScriptViewModel->HasUnappliedChanges();
+		return ScriptViewModel.Pin()->HasUnappliedChanges();
 	}
 
 private:
-	TSharedPtr<FNiagaraScratchPadScriptViewModel> ScriptViewModel;
+	TWeakPtr<FNiagaraScratchPadScriptViewModel> ScriptViewModel;
 	TSharedPtr<SNiagaraScriptGraph> Graph;
 
 	FDelegateHandle NodeIDHandle;
@@ -137,8 +137,7 @@ public:
 	virtual void OnTabBackgrounded(TSharedPtr<SDockTab> Tab) const override;
 
 	virtual void OnTabRefreshed(TSharedPtr<SDockTab> Tab) const override;
-
-
+	
 	virtual void SaveState(TSharedPtr<SDockTab> Tab, TSharedPtr<FTabPayload> Payload) const override;
 
 	virtual TSharedRef<SDockTab> OnSpawnTab(const FSpawnTabArgs& SpawnArgs, TWeakPtr<FTabManager> WeakTabManager) const override;
@@ -152,6 +151,8 @@ protected:
 
 	virtual TSharedRef<FGenericTabHistory> CreateTabHistoryNode(TSharedPtr<FTabPayload> Payload) override;
 
+	/** If the payload (a scratch pad script's graph) is invalid, it will be closed whenever the document tracker refreshes or cleans tabs. */
+	virtual bool IsPayloadValid(TSharedRef<FTabPayload> Payload) const override;
 protected:
 	TWeakPtr<FNiagaraSystemToolkit> EditorPtr;
 	FOnCreateGraphEditorWidget OnCreateGraphEditorWidget;
diff --git a/Engine/Plugins/FX/Niagara/Source/NiagaraEditor/Public/ViewModels/NiagaraSystemEditorDocumentsViewModel.h b/Engine/Plugins/FX/Niagara/Source/NiagaraEditor/Public/ViewModels/NiagaraSystemEditorDocumentsViewModel.h
--- a/Engine/Plugins/FX/Niagara/Source/NiagaraEditor/Public/ViewModels/NiagaraSystemEditorDocumentsViewModel.h
+++ b/Engine/Plugins/FX/Niagara/Source/NiagaraEditor/Public/ViewModels/NiagaraSystemEditorDocumentsViewModel.h
@@ -57,6 +57,7 @@ public:
 	void SetPrimaryDocumentID(const FName& TabId) { PrimaryDocumentTabId = TabId; }
 	NIAGARAEDITOR_API void SwapEditableScripts(TSharedPtr < class FNiagaraScratchPadScriptViewModel> OldScriptViewModel, TSharedPtr < class FNiagaraScratchPadScriptViewModel> NewScriptViewModel);
 
+	void CleanInvalidTabs() const;
 protected:
 	NIAGARAEDITOR_API TSharedPtr<SDockTab> OpenDocument(const UObject* DocumentID, FDocumentTracker::EOpenDocumentCause Cause);
 
@@ -71,7 +72,7 @@ protected:
 
 
 private:
-	NIAGARAEDITOR_API TSharedRef<FNiagaraSystemViewModel> GetSystemViewModel();
+	NIAGARAEDITOR_API TSharedRef<FNiagaraSystemViewModel> GetSystemViewModel() const;
 
 	FScriptToolkitsActiveDocumentChanged ActiveDocChangedDelegate;
 
