Added a safety 'IsValid' check captures 1 more crash in a System context (previous fixes were fixing it for emitters)

diff --git a/Engine/Plugins/FX/Niagara/Source/NiagaraEditor/Private/ViewModels/NiagaraScratchPadScriptViewModel.cpp b/Engine/Plugins/FX/Niagara/Source/NiagaraEditor/Private/ViewModels/NiagaraScratchPadScriptViewModel.cpp
--- a/Engine/Plugins/FX/Niagara/Source/NiagaraEditor/Private/ViewModels/NiagaraScratchPadScriptViewModel.cpp
+++ b/Engine/Plugins/FX/Niagara/Source/NiagaraEditor/Private/ViewModels/NiagaraScratchPadScriptViewModel.cpp
@@ -82,6 +82,11 @@ void FNiagaraScratchPadScriptViewModel::Initialize(UNiagaraScript* Script, UNiag
 	ParameterPaneViewModel->Init(UIContext);
 }
 
+bool FNiagaraScratchPadScriptViewModel::IsValid() const
+{
+	return EditScript.Script != nullptr && GraphViewModel->GetScriptSource() != nullptr;
+}
+
 void FNiagaraScratchPadScriptViewModel::OnGraphSubObjectSelectionChanged(const UObject* Obj)
 {
 	OnGraphSelectionChanged().Broadcast(Obj);
diff --git a/Engine/Plugins/FX/Niagara/Source/NiagaraEditor/Private/ViewModels/NiagaraSystemEditorDocumentsViewModel.cpp b/Engine/Plugins/FX/Niagara/Source/NiagaraEditor/Private/ViewModels/NiagaraSystemEditorDocumentsViewModel.cpp
--- a/Engine/Plugins/FX/Niagara/Source/NiagaraEditor/Private/ViewModels/NiagaraSystemEditorDocumentsViewModel.cpp
+++ b/Engine/Plugins/FX/Niagara/Source/NiagaraEditor/Private/ViewModels/NiagaraSystemEditorDocumentsViewModel.cpp
@@ -392,6 +392,11 @@ TSharedPtr<FNiagaraScratchPadScriptViewModel> UNiagaraSystemEditorDocumentsViewM
 		{
 			for (TSharedRef<FNiagaraScratchPadScriptViewModel> ScratchScriptViewModel : ScratchViewModel->GetScriptViewModels())
 			{
+				if(ScratchScriptViewModel->IsValid() == false)
+				{
+					continue;
+				}
+				
 				UNiagaraScript* EditScript = ScratchScriptViewModel->GetGraphViewModel()->GetScriptSource()->GetTypedOuter<UNiagaraScript>();
 				if (EditScript == Script)
 				{
diff --git a/Engine/Plugins/FX/Niagara/Source/NiagaraEditor/Public/ViewModels/NiagaraScratchPadScriptViewModel.h b/Engine/Plugins/FX/Niagara/Source/NiagaraEditor/Public/ViewModels/NiagaraScratchPadScriptViewModel.h
--- a/Engine/Plugins/FX/Niagara/Source/NiagaraEditor/Public/ViewModels/NiagaraScratchPadScriptViewModel.h
+++ b/Engine/Plugins/FX/Niagara/Source/NiagaraEditor/Public/ViewModels/NiagaraScratchPadScriptViewModel.h
@@ -27,6 +27,8 @@ public:
 
 	NIAGARAEDITOR_API void Initialize(UNiagaraScript* Script, UNiagaraScript* InEditScript, TWeakPtr<class FNiagaraSystemViewModel> InSystemViewModel);
 
+	NIAGARAEDITOR_API bool IsValid() const;
+	
 	NIAGARAEDITOR_API void TransferFromOldWhenDoingApply(TSharedPtr< FNiagaraScratchPadScriptViewModel> InOldScriptVM);
 
 	NIAGARAEDITOR_API void Finalize();
