- System properties view cleanup
- Removed bAutoDeactivate as it was never used
- bRequireCurrentFrameData moved to advanced display
- Adjusted warmup to only show it ticks > 0
- Reorder categories and hide Asset Options

diff --git a/Engine/Plugins/FX/Niagara/Source/Niagara/Classes/NiagaraSystem.h b/Engine/Plugins/FX/Niagara/Source/Niagara/Classes/NiagaraSystem.h
--- a/Engine/Plugins/FX/Niagara/Source/Niagara/Classes/NiagaraSystem.h
+++ b/Engine/Plugins/FX/Niagara/Source/Niagara/Classes/NiagaraSystem.h
@@ -207,7 +207,7 @@ struct FNiagaraSystemScalabilityOverrides
 
 /** A Niagara System contains multiple Niagara Emitters to create various effects.
  * Niagara Systems can be placed in the world, unlike Emitters, and expose User Parameters to configure an effect at runtime.*/
-UCLASS(BlueprintType, meta= (LoadBehavior = "LazyOnDemand"), MinimalAPI)
+UCLASS(BlueprintType, MinimalAPI, meta = (LoadBehavior = "LazyOnDemand"))
 class UNiagaraSystem : public UFXSystemAsset, public INiagaraParameterDefinitionsSubscriber
 {
 	GENERATED_UCLASS_BODY()
@@ -371,6 +371,7 @@ public:
 	FORCEINLINE float GetWarmupTime()const { return WarmupTime; }
 	FORCEINLINE int32 GetWarmupTickCount()const { return WarmupTickCount; }
 	FORCEINLINE float GetWarmupTickDelta()const { return WarmupTickDelta; }
+
 	FORCEINLINE bool HasFixedTickDelta() const { return bFixedTickDelta; }
 	FORCEINLINE float GetFixedTickDeltaTime()const { return FixedTickDeltaTime; }
 	NIAGARA_API virtual void GetAssetRegistryTags(TArray<FAssetRegistryTag>& OutTags)  const override;
@@ -378,8 +379,8 @@ public:
 	FORCEINLINE bool NeedsDeterminism() const { return bDeterminism; }
 	FORCEINLINE int32 GetRandomSeed() const { return RandomSeed; }
 
-	FORCEINLINE void SetWarmupTime(float InWarmupTime) { WarmupTime = InWarmupTime; ResolveWarmupTickCount(); }
-	FORCEINLINE void SetWarmupTickDelta(float InWarmupTickDelta) { WarmupTickDelta = InWarmupTickDelta; ResolveWarmupTickCount(); }
+	NIAGARA_API void SetWarmupTime(float InWarmupTime);
+	NIAGARA_API void SetWarmupTickDelta(float InWarmupTickDelta);
 	NIAGARA_API void ResolveWarmupTickCount();
 
 #if STATS
@@ -478,7 +479,6 @@ public:
 	NIAGARA_API void UpdateSystemAfterLoad();
 	NIAGARA_API void EnsureFullyLoaded() const;
 
-	bool ShouldAutoDeactivate() const { return bAutoDeactivate; }
 	NIAGARA_API bool IsLooping() const;
 
 	const TArray<TSharedRef<const FNiagaraEmitterCompiledData>>& GetEmitterCompiledData() const { return EmitterCompiledData; };
@@ -689,10 +689,9 @@ public:
 
 	/** When enabled, we follow the settings on the UNiagaraComponent for tick order. When this option is disabled, we ignore any dependencies from data interfaces or other variables and instead fire off the simulation as early in the frame as possible. This greatly
 	reduces overhead and allows the game thread to run faster, but comes at a tradeoff if the dependencies might leave gaps or other visual artifacts.*/
-	UPROPERTY(EditAnywhere, Category = "Performance")
+	UPROPERTY(EditAnywhere, Category = "Performance", AdvancedDisplay)
 	bool bRequireCurrentFrameData = true;
 
-	NIAGARA_API bool HasSystemScriptDIsWithPerInstanceData() const;
 	FORCEINLINE bool HasDIsWithPostSimulateTick() const { return bHasDIsWithPostSimulateTick; }
 	FORCEINLINE bool AllDIsPostSimulateCanOverlapFrames() const { return bAllDIsPostSimulateCanOverlapFrames; }
 	FORCEINLINE bool AsyncWorkCanOverlapTickGroups() const { return bAllDIsPostStageCanOverlapTickGroups; }
@@ -885,31 +884,28 @@ protected:
 	UPROPERTY(EditAnywhere, Category = "System", meta = (SkipSystemResetOnChange = "true", EditCondition = "bFixedBounds"))
 	FBox FixedBounds;
 
-	UPROPERTY(EditAnywhere, Category = Performance, meta = (ToolTip = "Auto-deactivate system if all emitters are determined to not spawn particles again, regardless of lifetime."))
-	bool bAutoDeactivate;
-
 	/**
 	When disabled we will generate a RandomSeed per instance on reset which is not deterministic.
 	When enabled we will always use the RandomSeed from the system plus the components RandomSeedOffset, this allows for determinism but variance between components.
 	*/
-	UPROPERTY(EditAnywhere, Category = "Random")
+	UPROPERTY(EditAnywhere, Category = "System")
 	bool bDeterminism = false;
 
 	/** Seed used for system script random number generator. */
-	UPROPERTY(EditAnywhere, Category = "Random", meta = (EditCondition = "bDeterminism"))
+	UPROPERTY(EditAnywhere, Category = "System", meta = (EditCondition = "bDeterminism", EditConditionHides))
 	int32 RandomSeed = 0;
 
 	/** Warm up time in seconds. Used to calculate WarmupTickCount. Rounds down to the nearest multiple of WarmupTickDelta. */
-	UPROPERTY(EditAnywhere, Category = Warmup, meta = (ForceUnits=s))
-	float WarmupTime;
+	UPROPERTY(EditAnywhere, Category = "System", meta = (ForceUnits=s))
+	float WarmupTime = 0.0f;
 
 	/** Number of ticks to process for warmup. You can set by this or by time via WarmupTime. */
-	UPROPERTY(EditAnywhere, Category = Warmup)
-	int32 WarmupTickCount;
+	UPROPERTY(EditAnywhere, Category = "System", meta = (EditCondition = "WarmupTime > 0.0", EditConditionHides))
+	int32 WarmupTickCount = 0;
 
 	/** Delta time to use for warmup ticks. */
-	UPROPERTY(EditAnywhere, Category = Warmup, meta = (ForceUnits=s))
-	float WarmupTickDelta;
+	UPROPERTY(EditAnywhere, Category = "System", meta = (ForceUnits=s, EditCondition = "WarmupTime > 0.0", EditConditionHides))
+	float WarmupTickDelta = 1.0f / 15.0f;
 
 	UPROPERTY(EditAnywhere, Category = "System", meta = (InlineEditConditionToggle))
 	bool bFixedTickDelta = false;
diff --git a/Engine/Plugins/FX/Niagara/Source/Niagara/Private/NiagaraSystem.cpp b/Engine/Plugins/FX/Niagara/Source/Niagara/Private/NiagaraSystem.cpp
--- a/Engine/Plugins/FX/Niagara/Source/Niagara/Private/NiagaraSystem.cpp
+++ b/Engine/Plugins/FX/Niagara/Source/Niagara/Private/NiagaraSystem.cpp
@@ -167,10 +167,6 @@ UNiagaraSystem::UNiagaraSystem(const FObjectInitializer& ObjectInitializer)
 , bIsolateEnabled(false)
 #endif
 , FixedBounds(FBox(FVector(-100), FVector(100)))
-, bAutoDeactivate(true)
-, WarmupTime(0.0f)
-, WarmupTickCount(0)
-, WarmupTickDelta(1.0f / 15.0f)
 , bNeedsGPUContextInitForDataInterfaces(false)
 , bNeedsAsyncOptimize(true)
 , CurrentScalabilitySettings(*new FNiagaraSystemScalabilitySettings())
@@ -824,6 +820,18 @@ void UNiagaraSystem::Serialize(FArchive& Ar)
 	}
 }
 
+void UNiagaraSystem::SetWarmupTime(float InWarmupTime)
+{
+	WarmupTime = InWarmupTime;
+	ResolveWarmupTickCount();
+}
+
+void UNiagaraSystem::SetWarmupTickDelta(float InWarmupTickDelta)
+{
+	WarmupTickDelta = InWarmupTickDelta;
+	ResolveWarmupTickCount();
+}
+
 void UNiagaraSystem::ResolveWarmupTickCount()
 {
 	//Set the WarmupTickCount to feed back to the user.
diff --git a/Engine/Plugins/FX/Niagara/Source/NiagaraEditor/Private/NiagaraSystemDetailsCustomization.cpp b/Engine/Plugins/FX/Niagara/Source/NiagaraEditor/Private/NiagaraSystemDetailsCustomization.cpp
--- a/Engine/Plugins/FX/Niagara/Source/NiagaraEditor/Private/NiagaraSystemDetailsCustomization.cpp
+++ b/Engine/Plugins/FX/Niagara/Source/NiagaraEditor/Private/NiagaraSystemDetailsCustomization.cpp
@@ -16,40 +16,39 @@ void FNiagaraSystemDetails::CustomizeDetails(IDetailLayoutBuilder& InDetailLayou
 {
 	// we display the scalability category within scalability mode, which is why we hide it here
 	InDetailLayout.HideCategory("Scalability");
+	InDetailLayout.HideCategory("Asset Options");
 
-	InDetailLayout.SortCategories([](const TMap<FName, IDetailCategoryBuilder*>& CategoryMap)
-	{
-		for (const TPair<FName, IDetailCategoryBuilder*>& Pair : CategoryMap)
+	InDetailLayout.SortCategories(
+		[](const TMap<FName, IDetailCategoryBuilder*>& CategoryMap)
 		{
-			int32 SortOrder = Pair.Value->GetSortOrder();
-			const FName& CategoryName = Pair.Key;
-
-			if (CategoryName == "Random")
-			{
-				SortOrder = 1;
-			}
-			else if (CategoryName == "Rendering")
-			{
-				SortOrder = 2;
-			}
-			else if (CategoryName == "System")
-			{
-				SortOrder = 3;
-			}
-			else if (CategoryName == "Warmup")
+			for (const TPair<FName, IDetailCategoryBuilder*>& Pair : CategoryMap)
 			{
-				SortOrder = 4;
-			}
-			else if (CategoryName == "Performance")
-			{
-				SortOrder = 5;
-			}
-			else
-			{
-				SortOrder += 6;
-			}
+				int32 SortOrder = Pair.Value->GetSortOrder();
+				const FName& CategoryName = Pair.Key;
+
+				if (CategoryName == "System")
+				{
+					SortOrder = 1;
+				}
+				else if (CategoryName == "Rendering")
+				{
+					SortOrder = 2;
+				}
+				else if (CategoryName == "Debug")
+				{
+					SortOrder = 3;
+				}
+				else if (CategoryName == "Performance")
+				{
+					SortOrder = 4;
+				}
+				else
+				{
+					SortOrder += 5;
+				}
 
-			Pair.Value->SetSortOrder(SortOrder);
+				Pair.Value->SetSortOrder(SortOrder);
+			}
 		}
-	});
+	);
 }
