- Fix issue with emitter self binding mode
- Make sure we report the failed emitter name

diff --git a/Engine/Plugins/FX/Niagara/Source/Niagara/Private/NiagaraDataInterfaceEmitterBinding.cpp b/Engine/Plugins/FX/Niagara/Source/Niagara/Private/NiagaraDataInterfaceEmitterBinding.cpp
--- a/Engine/Plugins/FX/Niagara/Source/Niagara/Private/NiagaraDataInterfaceEmitterBinding.cpp
+++ b/Engine/Plugins/FX/Niagara/Source/Niagara/Private/NiagaraDataInterfaceEmitterBinding.cpp
@@ -1,6 +1,7 @@
 // Copyright Epic Games, Inc. All Rights Reserved.
 
 #include "NiagaraDataInterfaceEmitterBinding.h"
+#include "NiagaraConstants.h"
 #include "NiagaraEmitterInstance.h"
 #include "NiagaraSystem.h"
 #include "NiagaraSystemInstance.h"
@@ -115,6 +116,37 @@ const FNiagaraEmitterHandle* FNiagaraDataInterfaceEmitterBinding::ResolveHandle(
 				}
 			}
 		}
+		else
+		{
+			for (const FNiagaraEmitterHandle& EmitterHandle : OwnerSystem->GetEmitterHandles())
+			{
+				FVersionedNiagaraEmitterData* EmitterData = EmitterHandle.GetEmitterData();
+				if (EmitterHandle.GetIsEnabled() && EmitterData)
+				{
+					bool bFoundDI = false;
+					bool bIsEmitterNamespace = false;
+					EmitterData->ForEachScript(
+						[DataInterface, &bFoundDI, &bIsEmitterNamespace](const UNiagaraScript* Script)
+						{
+							for (const FNiagaraScriptResolvedDataInterfaceInfo& ResolvedDI : Script->GetResolvedDataInterfaces())
+							{
+								if (ResolvedDI.ResolvedDataInterface == DataInterface)
+								{
+									bFoundDI = true;
+									bIsEmitterNamespace = FNiagaraVariableBase::IsInNameSpace(FNiagaraConstants::EmitterNamespaceString, ResolvedDI.CompileName);
+									return;
+								}
+							}
+						}
+					);
+
+					if (bFoundDI)
+					{
+						return bIsEmitterNamespace ? &EmitterHandle : nullptr;
+					}
+				}
+			}
+		}
 	}
 	else if (BindingMode == ENiagaraDataInterfaceEmitterBindingMode::Other)
 	{
@@ -148,6 +180,20 @@ UNiagaraEmitter* FNiagaraDataInterfaceEmitterBinding::Resolve(const UNiagaraData
 
 FString FNiagaraDataInterfaceEmitterBinding::ResolveUniqueName(const UNiagaraDataInterface* DataInterface) const
 {
-	UNiagaraEmitter* ResolvedEmitter = Resolve(DataInterface);
-	return ResolvedEmitter ? ResolvedEmitter->GetUniqueEmitterName() : FString();
+	if (UNiagaraEmitter* ResolvedEmitter = Resolve(DataInterface))
+	{
+		return ResolvedEmitter->GetUniqueEmitterName();
+	}
+	else
+	{
+		if (BindingMode == ENiagaraDataInterfaceEmitterBindingMode::Self)
+		{
+			return FString(TEXT("Self"));
+		}
+		else if (BindingMode == ENiagaraDataInterfaceEmitterBindingMode::Other)
+		{
+			return EmitterName.ToString();
+		}
+	}
+	return FString();
 }
