Simple tick group based exec ordering for systems reading from NDCs.
We define a tick group which reads must be ordered after.
Reads before or writes after this TG emit a warning.

diff --git a/Engine/Plugins/FX/Niagara/Source/Niagara/Classes/NiagaraDataInterface.h b/Engine/Plugins/FX/Niagara/Source/Niagara/Classes/NiagaraDataInterface.h
--- a/Engine/Plugins/FX/Niagara/Source/Niagara/Classes/NiagaraDataInterface.h
+++ b/Engine/Plugins/FX/Niagara/Source/Niagara/Classes/NiagaraDataInterface.h
@@ -680,6 +680,8 @@ public:
 
 	virtual bool HasTickGroupPrereqs() const { return false; }
 	virtual ETickingGroup CalculateTickGroup(const void* PerInstanceData) const { return NiagaraFirstTickGroup; }
+	virtual bool HasTickGroupPostreqs() const { return false; }
+	virtual ETickingGroup CalculateFinalTickGroup(const void* PerInstanceData) const { return NiagaraLastTickGroup; }
 
 	/** Determines if this type definition matches to a known data interface type.*/
 	static NIAGARA_API bool IsDataInterfaceType(const FNiagaraTypeDefinition& TypeDef);
diff --git a/Engine/Plugins/FX/Niagara/Source/Niagara/Internal/DataInterface/NiagaraDataInterfaceDataChannelRead.h b/Engine/Plugins/FX/Niagara/Source/Niagara/Internal/DataInterface/NiagaraDataInterfaceDataChannelRead.h
--- a/Engine/Plugins/FX/Niagara/Source/Niagara/Internal/DataInterface/NiagaraDataInterfaceDataChannelRead.h
+++ b/Engine/Plugins/FX/Niagara/Source/Niagara/Internal/DataInterface/NiagaraDataInterfaceDataChannelRead.h
@@ -220,7 +220,7 @@ struct FNDIDataChannelReadInstanceData
 	virtual ~FNDIDataChannelReadInstanceData();
 	FNiagaraDataBuffer* GetReadBufferCPU(bool bPrevFrame);
 	bool Init(UNiagaraDataInterfaceDataChannelRead* Interface, FNiagaraSystemInstance* Instance);
-	bool Tick(UNiagaraDataInterfaceDataChannelRead* Interface, FNiagaraSystemInstance* Instance);
+	bool Tick(UNiagaraDataInterfaceDataChannelRead* Interface, FNiagaraSystemInstance* Instance, bool bIsInit=false);
 	bool PostTick(UNiagaraDataInterfaceDataChannelRead* Interface, FNiagaraSystemInstance* Instance);
 };
 
diff --git a/Engine/Plugins/FX/Niagara/Source/Niagara/Internal/DataInterface/NiagaraDataInterfaceDataChannelWrite.h b/Engine/Plugins/FX/Niagara/Source/Niagara/Internal/DataInterface/NiagaraDataInterfaceDataChannelWrite.h
--- a/Engine/Plugins/FX/Niagara/Source/Niagara/Internal/DataInterface/NiagaraDataInterfaceDataChannelWrite.h
+++ b/Engine/Plugins/FX/Niagara/Source/Niagara/Internal/DataInterface/NiagaraDataInterfaceDataChannelWrite.h
@@ -112,6 +112,9 @@ public:
 	NIAGARA_API virtual void PostStageTick(FNDICpuPostStageContext& Context) override;
 	NIAGARA_API virtual void ProvidePerInstanceDataForRenderThread(void* DataForRenderThread, void* PerInstanceData, const FNiagaraSystemInstanceID& SystemInstance) override;
 
+	virtual bool HasTickGroupPostreqs() const override;
+	virtual ETickingGroup CalculateFinalTickGroup(const void* PerInstanceData) const override;
+
 	//We cannot overlap frames as we must correctly sync up with the data channel manager on Begin/End frame etc.
 	virtual bool PostSimulateCanOverlapFrames() const { return false; }
 	//We cannot have post stage overlap tick groups so that the write DI can publish it's contents to the data channel at the correct time to allow same frame reads.
diff --git a/Engine/Plugins/FX/Niagara/Source/Niagara/Private/DataInterface/NiagaraDataInterfaceDataChannelRead.cpp b/Engine/Plugins/FX/Niagara/Source/Niagara/Private/DataInterface/NiagaraDataInterfaceDataChannelRead.cpp
--- a/Engine/Plugins/FX/Niagara/Source/Niagara/Private/DataInterface/NiagaraDataInterfaceDataChannelRead.cpp
+++ b/Engine/Plugins/FX/Niagara/Source/Niagara/Private/DataInterface/NiagaraDataInterfaceDataChannelRead.cpp
@@ -147,12 +147,12 @@ bool FNDIDataChannelReadInstanceData::Init(UNiagaraDataInterfaceDataChannelRead*
 {
 	EmitterInstance = Interface->EmitterBinding.Resolve(Instance, Interface);
 
-	bool bSuccess = Tick(Interface, Instance);
+	bool bSuccess = Tick(Interface, Instance, true);
 	bSuccess &= PostTick(Interface, Instance);
 	return bSuccess;
 }
 
-bool FNDIDataChannelReadInstanceData::Tick(UNiagaraDataInterfaceDataChannelRead* Interface, FNiagaraSystemInstance* Instance)
+bool FNDIDataChannelReadInstanceData::Tick(UNiagaraDataInterfaceDataChannelRead* Interface, FNiagaraSystemInstance* Instance, bool bIsInit)
 {
 	ConsumeIndex = 0;
 	ConditionalSpawns.Reset();
@@ -256,6 +256,22 @@ bool FNDIDataChannelReadInstanceData::Tick(UNiagaraDataInterfaceDataChannelRead*
 				FNiagaraDataChannelSearchParameters SearchParams;
 				SearchParams.OwningComponent = Instance->GetAttachComponent();
 				DataChannelData = DataChannelPtr->FindData(SearchParams, ENiagaraResourceAccess::ReadOnly);//TODO: Maybe should have two paths, one for system instances and another for SceneComponents...
+			}	
+
+			if(const UNiagaraDataChannel* ChannelPtr = DataChannelPtr->GetDataChannel())
+			{
+				if(!bIsInit && ChannelPtr->ShouldEnforceTickGroupReadWriteOrder() && Interface->bReadCurrentFrame)
+				{
+					ETickingGroup CurrTG = DataChannelPtr->GetCurrentTickGroup();
+					ETickingGroup MinTickGroup = Interface->CalculateTickGroup(nullptr);//We don't use the per instance data...
+					if(CurrTG < MinTickGroup)
+					{
+						static UEnum* TGEnum = StaticEnum<ETickingGroup>();
+						UE_LOG(LogNiagara, Warning, TEXT("NDC Read DI is required to tick on or after %s but is reading in %s. This may cause us to have incorrectly ordered reads and writes to this NDC and thereform miss data.")
+						, *TGEnum->GetDisplayNameTextByValue((int32)MinTickGroup).ToString()
+						, *TGEnum->GetDisplayNameTextByValue((int32)CurrTG).ToString());
+					}
+				}
 			}
 
 			const FNiagaraDataSetCompiledData& CPUSourceDataCompiledData = DataChannelPtr->GetDataChannel()->GetCompiledData(ENiagaraSimTarget::CPUSim);
@@ -556,6 +572,10 @@ bool UNiagaraDataInterfaceDataChannelRead::HasTickGroupPrereqs() const
 	{
 		return true;
 	}
+	else if (Channel && Channel->Get())
+	{
+		return Channel->Get()->ShouldEnforceTickGroupReadWriteOrder();
+	}
 	return false;
 }
 
@@ -565,6 +585,11 @@ ETickingGroup UNiagaraDataInterfaceDataChannelRead::CalculateTickGroup(const voi
 	{
 		return (ETickingGroup)GNDCReadForceTG;
 	}
+	else if(Channel && Channel->Get() && Channel->Get()->ShouldEnforceTickGroupReadWriteOrder())
+	{
+		return (ETickingGroup)((int32)Channel->Get()->GetFinalWriteTickGroup() + 1);
+	}
+
 	return NiagaraFirstTickGroup; 
 }
 
diff --git a/Engine/Plugins/FX/Niagara/Source/Niagara/Private/DataInterface/NiagaraDataInterfaceDataChannelWrite.cpp b/Engine/Plugins/FX/Niagara/Source/Niagara/Private/DataInterface/NiagaraDataInterfaceDataChannelWrite.cpp
--- a/Engine/Plugins/FX/Niagara/Source/Niagara/Private/DataInterface/NiagaraDataInterfaceDataChannelWrite.cpp
+++ b/Engine/Plugins/FX/Niagara/Source/Niagara/Private/DataInterface/NiagaraDataInterfaceDataChannelWrite.cpp
@@ -422,6 +422,24 @@ void UNiagaraDataInterfaceDataChannelWrite::ProvidePerInstanceDataForRenderThrea
 	//FNDIDataChannelProxy::ProvidePerInstanceDataForRenderThread(DataForRenderThread, PerInstanceData, SystemInstance);
 }
 
+bool UNiagaraDataInterfaceDataChannelWrite::HasTickGroupPostreqs() const
+{
+	if (Channel && Channel->Get())
+	{
+		return Channel->Get()->ShouldEnforceTickGroupReadWriteOrder();
+	}
+	return false;
+}
+
+ETickingGroup UNiagaraDataInterfaceDataChannelWrite::CalculateFinalTickGroup(const void* PerInstanceData) const
+{
+	if(Channel && Channel->Get() && Channel->Get()->ShouldEnforceTickGroupReadWriteOrder())
+	{
+		return Channel->Get()->GetFinalWriteTickGroup();
+	}
+	return NiagaraLastTickGroup;
+}
+
 #if WITH_EDITORONLY_DATA
 
 void UNiagaraDataInterfaceDataChannelWrite::PostCompile()
diff --git a/Engine/Plugins/FX/Niagara/Source/Niagara/Private/NiagaraDataChannel.cpp b/Engine/Plugins/FX/Niagara/Source/Niagara/Private/NiagaraDataChannel.cpp
--- a/Engine/Plugins/FX/Niagara/Source/Niagara/Private/NiagaraDataChannel.cpp
+++ b/Engine/Plugins/FX/Niagara/Source/Niagara/Private/NiagaraDataChannel.cpp
@@ -17,6 +17,12 @@
 
 #define LOCTEXT_NAMESPACE "NiagaraDataChannels"
 
+namespace NDCCVars
+{
+	bool bEmitWarningsOnLateNDCWrites = !UE_BUILD_SHIPPING;
+	static FAutoConsoleVariableRef CVarEmitWarningsOnLateNDCWrites(TEXT("fx.Niagara.DataChannels.DebugDumpWriterDI"), bEmitWarningsOnLateNDCWrites, TEXT("If true, late writes to NDCs will generate warnings. Late meaning after their final allowed tick group."), ECVF_Default);
+};
+
 //////////////////////////////////////////////////////////////////////////
 
 void FNiagaraDataChannelGameDataLayout::Init(const TArray<FNiagaraVariable>& Variables)
@@ -580,10 +586,27 @@ void FNiagaraDataChannelData::ConsumePublishRequests(UNiagaraDataChannelHandler*
 	//There should be no access on other threads at this point anyway but lock just to be safe.
 	FScopeLock Lock(&PublishCritSec);
 
-	if(PublishRequests.Num() == 0)
+	const UNiagaraDataChannel* DataChannel = Owner->GetDataChannel();
+	if(PublishRequests.Num() == 0 || DataChannel == nullptr)
 	{
 		return;
 	}
+	
+	if(NDCCVars::bEmitWarningsOnLateNDCWrites && DataChannel->ShouldEnforceTickGroupReadWriteOrder())
+	{
+		ETickingGroup PublishSourceTG = (ETickingGroup)(FMath::Clamp((int32)Owner->GetCurrentTickGroup() - 1, 0, (int32)ETickingGroup::TG_MAX-1));//We're consuming from the previous TG.
+		ETickingGroup FinalWriteTG = DataChannel->GetFinalWriteTickGroup();
+
+		//TODO: Possibly allow late writes to be deferred to the next frame?
+		if(PublishSourceTG > FinalWriteTG)
+		{
+			static UEnum* TGEnum = StaticEnum<ETickingGroup>();
+			UE_LOG(LogNiagara, Warning, TEXT("Data Channel %s is being written to in Tick Group %s which is after it's final write tick group %s. This may cause incorrect read / write ordering and missed data.")
+			, *DataChannel->GetAsset()->GetName()
+			, *TGEnum->GetDisplayNameTextByValue((int32)PublishSourceTG).ToString()
+			, *TGEnum->GetDisplayNameTextByValue((int32)FinalWriteTG).ToString());
+		}
+	}
 
 	UWorld* World = Owner->GetWorld();
 	check(IsValid(World));
diff --git a/Engine/Plugins/FX/Niagara/Source/Niagara/Private/NiagaraDataChannelHandler.cpp b/Engine/Plugins/FX/Niagara/Source/Niagara/Private/NiagaraDataChannelHandler.cpp
--- a/Engine/Plugins/FX/Niagara/Source/Niagara/Private/NiagaraDataChannelHandler.cpp
+++ b/Engine/Plugins/FX/Niagara/Source/Niagara/Private/NiagaraDataChannelHandler.cpp
@@ -15,7 +15,7 @@ void UNiagaraDataChannelHandler::Init(const UNiagaraDataChannel* InChannel)
 
 void UNiagaraDataChannelHandler::BeginFrame(float DeltaTime, FNiagaraWorldManager* OwningWorld)
 {
-
+	CurrentTG = TG_PrePhysics;
 }
 
 void UNiagaraDataChannelHandler::EndFrame(float DeltaTime, FNiagaraWorldManager* OwningWorld)
@@ -25,6 +25,7 @@ void UNiagaraDataChannelHandler::EndFrame(float DeltaTime, FNiagaraWorldManager*
 
 void UNiagaraDataChannelHandler::Tick(float DeltaTime, ETickingGroup TickGroup, FNiagaraWorldManager* OwningWorld)
 {
+	CurrentTG = TickGroup;
 }
 
 UNiagaraDataChannelWriter* UNiagaraDataChannelHandler::GetDataChannelWriter()
diff --git a/Engine/Plugins/FX/Niagara/Source/Niagara/Private/NiagaraSystemInstance.cpp b/Engine/Plugins/FX/Niagara/Source/Niagara/Private/NiagaraSystemInstance.cpp
--- a/Engine/Plugins/FX/Niagara/Source/Niagara/Private/NiagaraSystemInstance.cpp
+++ b/Engine/Plugins/FX/Niagara/Source/Niagara/Private/NiagaraSystemInstance.cpp
@@ -125,7 +125,7 @@ FNiagaraSystemInstance::FNiagaraSystemInstance(UWorld& InWorld, UNiagaraSystem&
 	  , ParametersValid(false)
 	  , bSolo(false)
 	  , bForceSolo(false)
-	  , bNotifyOnCompletion(false), bHasGPUEmitters(false), bDataInterfacesHaveTickPrereqs(false)
+	  , bNotifyOnCompletion(false), bHasGPUEmitters(false), bDataInterfacesHaveTickPrereqs(false), bDataInterfacesHaveTickPostreqs(false)
 	  , bDataInterfacesInitialized(false)
 	  , bAlreadyBound(false)
 	  , bLODDistanceIsValid(false)
@@ -319,14 +319,15 @@ void FNiagaraSystemInstance::DumpTickInfo(FOutputDevice& Ar)
 		PrereqInfo.Appendf(TEXT(" PreReq(%s = %s)"), *PrereqComponent->GetFullName(), *TickingGroupEnum->GetNameStringByIndex(PrereqTG));
 	}
 
-	if (bDataInterfacesHaveTickPrereqs)
+	if (bDataInterfacesHaveTickPrereqs || bDataInterfacesHaveTickPostreqs)
 	{
 		for (TPair<TWeakObjectPtr<UNiagaraDataInterface>, int32>& Pair : DataInterfaceInstanceDataOffsets)
 		{
 			if (UNiagaraDataInterface* Interface = Pair.Key.Get())
 			{
 				ETickingGroup PrereqTG = Interface->CalculateTickGroup(&DataInterfaceInstanceData[Pair.Value]);
-				PrereqInfo.Appendf(TEXT(" DataInterface(%s = %s)"), *Interface->GetFullName(), *TickingGroupEnum->GetNameStringByIndex(PrereqTG));
+				ETickingGroup PostreqTG = Interface->CalculateFinalTickGroup(&DataInterfaceInstanceData[Pair.Value]);
+				PrereqInfo.Appendf(TEXT(" DataInterface(%s = %s - %s)"), *Interface->GetFullName(), *TickingGroupEnum->GetNameStringByIndex(PrereqTG), *TickingGroupEnum->GetNameStringByIndex(PostreqTG));
 			}
 		}
 	}
@@ -1399,6 +1400,7 @@ FNDIStageTickHandler* FNiagaraSystemInstance::GetSystemDIStageTickHandler(ENiaga
 void FNiagaraSystemInstance::InitDataInterfaces()
 {
 	bDataInterfacesHaveTickPrereqs = false;
+	bDataInterfacesHaveTickPostreqs = false;
 
 	// If the System is invalid, it is possible that our cached data interfaces are now bogus and could point to invalid memory.
 	// Only the UNiagaraComponent or UNiagaraSystem can hold onto GC references to the DataInterfaces.
@@ -1477,6 +1479,11 @@ void FNiagaraSystemInstance::InitDataInterfaces()
 					{
 						bDataInterfacesHaveTickPrereqs = Interface->HasTickGroupPrereqs();
 					}
+					
+					if( bDataInterfacesHaveTickPostreqs == false)
+					{
+						bDataInterfacesHaveTickPostreqs = Interface->HasTickGroupPostreqs();
+					}
 
 					if (bIsGPUScript)
 					{
@@ -1784,7 +1791,8 @@ float FNiagaraSystemInstance::GetLODDistance()
 
 ETickingGroup FNiagaraSystemInstance::CalculateTickGroup() const
 {
-	ETickingGroup NewTickGroup = (ETickingGroup)0;
+	ETickingGroup NewMinTickGroup = (ETickingGroup)0;
+	ETickingGroup NewMaxTickGroup = ETickingGroup::TG_MAX;
 
 	// Debugging feature to force last tick group
 	if (GNiagaraForceLastTickGroup)
@@ -1810,7 +1818,7 @@ ETickingGroup FNiagaraSystemInstance::CalculateTickGroup() const
 					PrereqTG = PrereqSMC->bBlendPhysics ? FMath::Max(PrereqTG, ETickingGroup(TG_EndPhysics + 1)) : PrereqTG;
 				}
 
-				NewTickGroup = FMath::Max(NewTickGroup, PrereqTG);
+				NewMinTickGroup = FMath::Max(NewMinTickGroup, PrereqTG);
 			}
 
 			// Handle data interfaces that have tick dependencies
@@ -1821,39 +1829,59 @@ ETickingGroup FNiagaraSystemInstance::CalculateTickGroup() const
 					if (UNiagaraDataInterface* Interface = Pair.Key.Get())
 					{
 						ETickingGroup PrereqTG = Interface->CalculateTickGroup(&DataInterfaceInstanceData[Pair.Value]);
-						NewTickGroup = FMath::Max(NewTickGroup, PrereqTG);
+						NewMinTickGroup = FMath::Max(NewMinTickGroup, PrereqTG);
 					}
 				}
 			}
+			if ( bDataInterfacesHaveTickPostreqs )
+			{
+				for (const TPair<TWeakObjectPtr<UNiagaraDataInterface>, int32>& Pair : DataInterfaceInstanceDataOffsets)
+				{
+					if (UNiagaraDataInterface* Interface = Pair.Key.Get())
+					{
+						ETickingGroup PostReq = Interface->CalculateFinalTickGroup(&DataInterfaceInstanceData[Pair.Value]);
+						NewMaxTickGroup = FMath::Min(NewMaxTickGroup, PostReq);
+					}
+				}
+			}
+			
+			if(NewMinTickGroup > NewMaxTickGroup)
+			{
+				static UEnum* TGEnum = StaticEnum<ETickingGroup>();
+				UE_LOG(LogNiagara, Warning, TEXT("Niagara Component has DIs with conflicting Tick Group Dependencies. This may lead to some incorrect behavior.\nSystem:%s\nMinTickGroup:%s\nMaxTickGroup:%s")
+				, *System->GetName()
+				, *TGEnum->GetDisplayNameTextByValue((int32)NewMinTickGroup).ToString()
+				, *TGEnum->GetDisplayNameTextByValue((int32)NewMaxTickGroup).ToString());
+			}
 
 			// Clamp tick group to our range
-			NewTickGroup = FMath::Clamp(NewTickGroup, NiagaraFirstTickGroup, NiagaraLastTickGroup);
+			NewMinTickGroup = FMath::Clamp(NewMinTickGroup, NiagaraFirstTickGroup, NiagaraLastTickGroup);
 			break;
 
 		case ENiagaraTickBehavior::UseComponentTickGroup:
 			if (USceneComponent* Component = AttachComponent.Get())
 			{
-				NewTickGroup = FMath::Clamp((ETickingGroup)Component->PrimaryComponentTick.TickGroup, NiagaraFirstTickGroup, NiagaraLastTickGroup);
+				NewMinTickGroup = FMath::Clamp((ETickingGroup)Component->PrimaryComponentTick.TickGroup, NiagaraFirstTickGroup, NiagaraLastTickGroup);
 			}
 			else
 			{
-				NewTickGroup = NiagaraFirstTickGroup;
+				NewMinTickGroup = NiagaraFirstTickGroup;
 			}
 			break;
 
 		case ENiagaraTickBehavior::ForceTickFirst:
-			NewTickGroup = NiagaraFirstTickGroup;
+			NewMinTickGroup = NiagaraFirstTickGroup;
 			break;
 
 		case ENiagaraTickBehavior::ForceTickLast:
-			NewTickGroup = NiagaraLastTickGroup;
+			NewMinTickGroup = NiagaraLastTickGroup;
 			break;
 	}
 
 
 	//UE_LOG(LogNiagara, Log, TEXT("TickGroup: %s %d %d"), *Component->GetPathName(), (int32)TickBehavior, (int32)NewTickGroup);
 
-	return NewTickGroup;
+	return NewMinTickGroup;
 }
 
 void FNiagaraSystemInstance::SetTickBehavior(ENiagaraTickBehavior NewTickBehavior)
diff --git a/Engine/Plugins/FX/Niagara/Source/Niagara/Public/NiagaraDataChannel.h b/Engine/Plugins/FX/Niagara/Source/Niagara/Public/NiagaraDataChannel.h
--- a/Engine/Plugins/FX/Niagara/Source/Niagara/Public/NiagaraDataChannel.h
+++ b/Engine/Plugins/FX/Niagara/Source/Niagara/Public/NiagaraDataChannel.h
@@ -123,6 +123,11 @@ public:
 	template<typename TFunc>
 	static void ForEachDataChannel(TFunc Func);
 
+	bool ShouldEnforceTickGroupReadWriteOrder()const {return bEnforceTickGroupReadWriteOrder;}
+	
+	/** If we are enforcing tick group read/write ordering the this returns the final tick group that this NDC can be written to. All reads must happen in Tick groups after this or next frame. */
+	ETickingGroup GetFinalWriteTickGroup()const { return FinalWriteTickGroup; }
+
 private:
 
 	/** The variables that define the data contained in this Data Channel. */
@@ -132,7 +137,15 @@ private:
 	/** If true, we keep our previous frame's data. This comes at a memory and performance cost but allows users to avoid tick order dependency by reading last frame's data. Some users will prefer a frame of latency to tick order dependency. */
 	UPROPERTY(EditAnywhere, Category = "Data Channel")
 	bool bKeepPreviousFrameData = true;
-		
+
+	/** If true we ensure that all writes happen in or before the Tick Group specified in EndWriteTickGroup and that all reads happen in tick groups after this. */
+	UPROPERTY(EditAnywhere, Category = "Data Channel")
+	bool bEnforceTickGroupReadWriteOrder = false;
+
+	/** The final tick group that this data channel can be written to. */
+	UPROPERTY(EditAnywhere, Category = "Data Channel", meta=(EditCondition="bEnforceTickGroupReadWriteOrder"))
+	TEnumAsByte<ETickingGroup> FinalWriteTickGroup = ETickingGroup::TG_EndPhysics;
+
 	/**
 	Data layout for payloads in Niagara datasets.
 	*/
diff --git a/Engine/Plugins/FX/Niagara/Source/Niagara/Public/NiagaraDataChannelHandler.h b/Engine/Plugins/FX/Niagara/Source/Niagara/Public/NiagaraDataChannelHandler.h
--- a/Engine/Plugins/FX/Niagara/Source/Niagara/Public/NiagaraDataChannelHandler.h
+++ b/Engine/Plugins/FX/Niagara/Source/Niagara/Public/NiagaraDataChannelHandler.h
@@ -55,6 +55,9 @@ public:
 
 	NIAGARA_API FNiagaraDataChannelDataPtr CreateData();
 
+	/** Returns the tick group we're currently being processed or most recently processed. */
+	ETickingGroup GetCurrentTickGroup() { return CurrentTG; }
+
 protected:
 
 	UPROPERTY()
@@ -67,4 +70,6 @@ protected:
 	/** Helper object allowing BP to read data in this channel. */
 	UPROPERTY()
 	TObjectPtr<UNiagaraDataChannelReader> Reader;
-};
+
+	ETickingGroup CurrentTG = ETickingGroup::TG_PrePhysics;
+};
\ No newline at end of file
diff --git a/Engine/Plugins/FX/Niagara/Source/Niagara/Public/NiagaraSystemInstance.h b/Engine/Plugins/FX/Niagara/Source/Niagara/Public/NiagaraSystemInstance.h
--- a/Engine/Plugins/FX/Niagara/Source/Niagara/Public/NiagaraSystemInstance.h
+++ b/Engine/Plugins/FX/Niagara/Source/Niagara/Public/NiagaraSystemInstance.h
@@ -570,6 +570,8 @@ private:
 	uint32 bHasGPUEmitters : 1;
 	/** The system contains data interfaces that can have tick group prerequisites. */
 	uint32 bDataInterfacesHaveTickPrereqs : 1;
+	/** The system contains data interfaces that can have tick group post requisites. */
+	uint32 bDataInterfacesHaveTickPostreqs : 1;
 
 	uint32 bDataInterfacesInitialized : 1;
 
