Fixed dupe CVar

#rnx

diff --git a/Engine/Plugins/FX/Niagara/Source/Niagara/Private/NiagaraDataChannel.cpp b/Engine/Plugins/FX/Niagara/Source/Niagara/Private/NiagaraDataChannel.cpp
--- a/Engine/Plugins/FX/Niagara/Source/Niagara/Private/NiagaraDataChannel.cpp
+++ b/Engine/Plugins/FX/Niagara/Source/Niagara/Private/NiagaraDataChannel.cpp
@@ -20,7 +20,7 @@
 namespace NDCCVars
 {
 	bool bEmitWarningsOnLateNDCWrites = !UE_BUILD_SHIPPING;
-	static FAutoConsoleVariableRef CVarEmitWarningsOnLateNDCWrites(TEXT("fx.Niagara.DataChannels.DebugDumpWriterDI"), bEmitWarningsOnLateNDCWrites, TEXT("If true, late writes to NDCs will generate warnings. Late meaning after their final allowed tick group."), ECVF_Default);
+	static FAutoConsoleVariableRef CVarEmitWarningsOnLateNDCWrites(TEXT("fx.Niagara.DataChannels.WarnOnLateWrites"), bEmitWarningsOnLateNDCWrites, TEXT("If true, late writes to NDCs will generate warnings. Late meaning after their final allowed tick group."), ECVF_Default);
 };
 
 //////////////////////////////////////////////////////////////////////////
