Use explicit pin types on op nodes for type casting of connected inputs

diff --git a/Engine/Plugins/FX/Niagara/Source/NiagaraEditor/Private/NiagaraHlslTranslator.cpp b/Engine/Plugins/FX/Niagara/Source/NiagaraEditor/Private/NiagaraHlslTranslator.cpp
--- a/Engine/Plugins/FX/Niagara/Source/NiagaraEditor/Private/NiagaraHlslTranslator.cpp
+++ b/Engine/Plugins/FX/Niagara/Source/NiagaraEditor/Private/NiagaraHlslTranslator.cpp
@@ -7448,6 +7448,19 @@ void TNiagaraHlslTranslator<GraphBridge>::Operation(const FOpNode* Operation, TA
 				OutputHlsl = IOInfo.HlslSnippet;
 			}
 			check(!OutputHlsl.IsEmpty());
+
+			// add type casts if necessary
+			for (int i = 0; i < InputPins.Num(); i++)
+			{
+				const FOutputPin* LinkedPin = GraphBridge::GetLinkedOutputPin(InputPins[i]);
+				if (LinkedPin && InputTypes[i] != FNiagaraTypeDefinition::GetGenericNumericDef() && InputTypes[i] != GraphBridge::GetPinType(LinkedPin, ENiagaraStructConversion::UserFacing))
+				{
+					FString NumberedArg(TEXT("{") + FString::FromInt(i) + TEXT("}"));
+					FString TypeName = GetStructHlslTypeName(InputTypes[i]);
+					FString CastArg(TEXT("((") + TypeName + TEXT(")") + NumberedArg + TEXT(")"));
+					OutputHlsl.ReplaceInline(*NumberedArg, *CastArg);
+				}
+			}
 			Outputs.Add(AddBodyChunk(GetUniqueSymbolName(IOInfo.Name), OutputHlsl, OutputType, Inputs));
 		}
 	}
