Filter out attributes (for linked inputs) that are not allowed in the current usage context.

diff --git a/Engine/Plugins/FX/Niagara/Source/NiagaraEditor/Private/ViewModels/Stack/NiagaraStackFunctionInput.cpp b/Engine/Plugins/FX/Niagara/Source/NiagaraEditor/Private/ViewModels/Stack/NiagaraStackFunctionInput.cpp
--- a/Engine/Plugins/FX/Niagara/Source/NiagaraEditor/Private/ViewModels/Stack/NiagaraStackFunctionInput.cpp
+++ b/Engine/Plugins/FX/Niagara/Source/NiagaraEditor/Private/ViewModels/Stack/NiagaraStackFunctionInput.cpp
@@ -31,6 +31,7 @@
 #include "NiagaraSettings.h"
 #include "ViewModels/NiagaraSystemScriptViewModel.h"
 #include "ScopedTransaction.h"
+#include "Algo/RemoveIf.h"
 #include "EdGraph/EdGraphPin.h"
 #include "Framework/Notifications/NotificationManager.h"
 #include "Materials/Material.h"
@@ -1547,6 +1548,20 @@ FName GetNamespaceForUsage(ENiagaraScriptUsage Usage)
 	}
 }
 
+bool IsNamespaceAllowedInUsage(FName Namespace, ENiagaraScriptUsage Usage)
+{
+	FName UsageName = GetNamespaceForUsage(Usage);
+	if (UsageName == FNiagaraConstants::SystemNamespace && (Namespace == FNiagaraConstants::EmitterNamespace || Namespace == FNiagaraConstants::ParticleAttributeNamespace))
+	{
+		return false;
+	}
+	if (UsageName == FNiagaraConstants::EmitterNamespace && Namespace == FNiagaraConstants::ParticleAttributeNamespace)
+	{
+		return false;
+	}
+	return true;
+}
+
 UNiagaraScript* UNiagaraStackFunctionInput::FindConversionScript(const FNiagaraTypeDefinition& FromType, TMap<FNiagaraTypeDefinition, UNiagaraScript*>& ConversionScriptCache, bool bIncludeConversionScripts) const
 {
 	if (bIncludeConversionScripts == false)
@@ -1740,6 +1755,24 @@ void UNiagaraStackFunctionInput::GetAvailableParameterHandles(TArray<FNiagaraPar
 			UE_LOG(LogNiagaraEditor, Warning, TEXT("Failed to load NiagaraParameterCollection '%s'"), *CollectionAsset.GetObjectPathString());
 		}
 	}
+
+	// Filter gathered attributes by allowed usage
+	ENiagaraScriptUsage CurrentUsage = CurrentOutputNode->GetUsage();
+	AvailableParameterHandles.SetNum(Algo::RemoveIf(AvailableParameterHandles, [CurrentUsage](FNiagaraParameterHandle& Handle)
+	{
+		return !IsNamespaceAllowedInUsage(Handle.GetNamespace(), CurrentUsage);
+	}));
+	TArray<FNiagaraVariable> ParameterNames;
+	AvailableConversionHandles.GetKeys(ParameterNames);
+	for (int i = 0; i < ParameterNames.Num(); i++)
+	{
+		FNiagaraVariable Parameter = ParameterNames[i];
+		FNiagaraParameterHandle Handle(Parameter.GetName());
+		if (!IsNamespaceAllowedInUsage(Handle.GetNamespace(), CurrentUsage))
+		{
+			AvailableConversionHandles.Remove(Parameter);
+		}
+	}
 }
 
 UNiagaraNodeFunctionCall* UNiagaraStackFunctionInput::GetDefaultFunctionNode() const
