- Niagara tick dependency validation cleanup
- Use resolved DIs rather than cached defaults
- Checks the system properties for bRequireCurrentFrameData
- Messaging changed to make it clear what the issue is
- Added a fix button which will disable bRequireCurrentFrameData at the system level

diff --git a/Engine/Plugins/FX/Niagara/Source/NiagaraEditor/Private/NiagaraValidationRules.cpp b/Engine/Plugins/FX/Niagara/Source/NiagaraEditor/Private/NiagaraValidationRules.cpp
--- a/Engine/Plugins/FX/Niagara/Source/NiagaraEditor/Private/NiagaraValidationRules.cpp
+++ b/Engine/Plugins/FX/Niagara/Source/NiagaraEditor/Private/NiagaraValidationRules.cpp
@@ -925,63 +925,98 @@ void UNiagaraValidationRule_SimulationStageBudget::CheckValidity(const FNiagaraV
 
 void UNiagaraValidationRule_TickDependencyCheck::CheckValidity(const FNiagaraValidationContext& Context, TArray<FNiagaraValidationResult>& OutResults) const
 {
-	UNiagaraSystem* NiagaraSystem = &Context.ViewModel->GetSystem();
-
 	if (!bCheckActorComponentInterface && !bCheckCameraDataInterface && !bCheckSkeletalMeshInterface)
 	{
 		return;
 	}
 
+	UNiagaraSystem* NiagaraSystem = &Context.ViewModel->GetSystem();
+	if (!NiagaraSystem->bRequireCurrentFrameData)
+	{
+		return;
+	}
+
+	TSet<UNiagaraDataInterface*> VisitedDIs;
 	NiagaraSystem->ForEachScript(
 		[&](UNiagaraScript* NiagaraScript)
 		{
-			const TArray<FNiagaraScriptDataInterfaceInfo>& CachedDefaultDIs = NiagaraScript->GetCachedDefaultDataInterfaces();
-			for ( const FNiagaraScriptDataInterfaceInfo& CachedDefaultDI : CachedDefaultDIs )
+			for ( const FNiagaraScriptResolvedDataInterfaceInfo& ResolvedDI : NiagaraScript->GetResolvedDataInterfaces() )
 			{
+				// Have we already encounted this DI?
+				UNiagaraDataInterface* RuntimeDI = ResolvedDI.ResolvedDataInterface;
+				if ( VisitedDIs.Contains(RuntimeDI) )
+				{
+					continue;
+				}
+				VisitedDIs.Add(RuntimeDI);
+
+				// Should we generate issues for this DI?
 				bool bWarnTickDependency = false;
-				if (UNiagaraDataInterfaceCamera* CameraDataInterface = Cast<UNiagaraDataInterfaceCamera>(CachedDefaultDI.DataInterface))
+				if (UNiagaraDataInterfaceCamera* CameraDataInterface = Cast<UNiagaraDataInterfaceCamera>(RuntimeDI))
 				{
 					bWarnTickDependency = bCheckCameraDataInterface && CameraDataInterface->bRequireCurrentFrameData;
 				}
-				else if (UNiagaraDataInterfaceSkeletalMesh* SkeletalMeshDataInterface = Cast<UNiagaraDataInterfaceSkeletalMesh>(CachedDefaultDI.DataInterface))
+				else if (UNiagaraDataInterfaceSkeletalMesh* SkeletalMeshDataInterface = Cast<UNiagaraDataInterfaceSkeletalMesh>(RuntimeDI))
 				{
 					bWarnTickDependency = bCheckSkeletalMeshInterface && SkeletalMeshDataInterface->bRequireCurrentFrameData;
 				}
-				else if (UNiagaraDataInterfaceActorComponent* ActorComponentDataInterface = Cast<UNiagaraDataInterfaceActorComponent>(CachedDefaultDI.DataInterface))
+				else if (UNiagaraDataInterfaceActorComponent* ActorComponentDataInterface = Cast<UNiagaraDataInterfaceActorComponent>(RuntimeDI))
 				{
 					bWarnTickDependency = bCheckActorComponentInterface && ActorComponentDataInterface->bRequireCurrentFrameData;
 				}
-
-				if (bWarnTickDependency)
+				if (bWarnTickDependency == false)
 				{
-					UNiagaraEmitter* NiagaraEmitter = NiagaraScript->GetTypedOuter<UNiagaraEmitter>();
+					continue;
+				}
 
-					UObject* StackObject = nullptr;
-					if ( NiagaraEmitter )
+				// Generate issue
+				UObject* StackObject = nullptr;
+				if (ResolvedDI.ResolvedSourceEmitterName.Len() > 0)
+				{
+					for (const TSharedPtr<FNiagaraEmitterHandleViewModel>& EmitterViewModel : Context.ViewModel->GetEmitterHandleViewModels())
 					{
-						// If we are attatched to an emitter than is not enabled bail
-						const TSharedPtr<FNiagaraEmitterHandleViewModel> EmitterViewModel = NiagaraValidation::GetEmitterViewModel(Context, NiagaraEmitter);
-						if (EmitterViewModel == nullptr || EmitterViewModel->GetIsEnabled() == false)
+						if (EmitterViewModel->GetName() == FName(ResolvedDI.ResolvedSourceEmitterName))
 						{
-							return;
+							StackObject = NiagaraValidation::GetStackEntry<UNiagaraStackEmitterPropertiesItem>(EmitterViewModel->GetEmitterStackViewModel());
+							break;
 						}
-						StackObject = NiagaraValidation::GetStackEntry<UNiagaraStackEmitterPropertiesItem>(EmitterViewModel->GetEmitterStackViewModel());
-					}
-					else
-					{
-						StackObject = NiagaraValidation::GetStackEntry<UNiagaraStackSystemPropertiesItem>(Context.ViewModel->GetSystemStackViewModel());
 					}
+				}
+				else
+				{
+					StackObject = NiagaraValidation::GetStackEntry<UNiagaraStackSystemPropertiesItem>(Context.ViewModel->GetSystemStackViewModel());
+				}
 
-					const FText OwnerStackText = FText::FromString(NiagaraEmitter ? NiagaraEmitter->GetName() : FString(TEXT("System")));
-					const FText DIClassText = FText::FromName(CachedDefaultDI.DataInterface->GetClass()->GetFName());
-					const FText VariableText = FText::FromName(CachedDefaultDI.Name);
-					OutResults.Emplace(
-						Severity,
-						FText::Format(LOCTEXT("TickDependencyCheckFormat", "'{0}' has a tick dependency that can potentially be removed"), OwnerStackText),
-						FText::Format(LOCTEXT("TickDependencyCheckDetailedFormat", "'{0}' has a tick dependency from data interace '{1}' variable '{2}' that can removed by unchecking 'RequireCurrentFrameData'"), OwnerStackText, DIClassText, VariableText),
-						StackObject
-					);
+				if (StackObject == nullptr)
+				{
+					continue;
 				}
+
+				const FText DIClassText = FText::FromName(RuntimeDI->GetClass()->GetFName());
+				const FText DIVariableText = FText::FromName(ResolvedDI.Name);
+				FNiagaraValidationResult& ValiationResult = OutResults.Emplace_GetRef(
+					Severity,
+					LOCTEXT("TickDependencyCheckFormat", "Performance issue due to late ticking which may cause waits on the game thread."),
+					FText::Format(LOCTEXT("TickDependencyCheckDetailedFormat", "'{0}' has a tick dependency that can removed by unchecking 'RequireCurrentFrameData' on the data interface.  This could introduce a frame of latency but will allow the system to execute immediatly in the frame.  Parameter Name '{1}'."), DIClassText, DIVariableText),
+					StackObject
+				);
+
+				ValiationResult.Fixes.Emplace(
+					FNiagaraValidationFix(
+						LOCTEXT("TickDependencyCheckFix", "Disable RequireCurrentFrameData in System Properties"),
+						FNiagaraValidationFixDelegate::CreateLambda(
+							[WeakNiagaraSystem=MakeWeakObjectPtr(NiagaraSystem)]()
+							{
+								if (UNiagaraSystem* Sys = WeakNiagaraSystem.Get())
+								{
+									const FScopedTransaction Transaction(LOCTEXT("FixtSystemRequireCurrentFrameData", "System Require Current Frame Data Disabled"));
+									Sys->Modify();
+									Sys->bRequireCurrentFrameData = false;
+								}
+							}
+						)
+					)
+				);
 			}
 		}
 	);
