- Fix for linux build

#rnx

diff --git a/Engine/Plugins/FX/Niagara/Source/NiagaraEditor/Private/NiagaraValidationRules.cpp b/Engine/Plugins/FX/Niagara/Source/NiagaraEditor/Private/NiagaraValidationRules.cpp
--- a/Engine/Plugins/FX/Niagara/Source/NiagaraEditor/Private/NiagaraValidationRules.cpp
+++ b/Engine/Plugins/FX/Niagara/Source/NiagaraEditor/Private/NiagaraValidationRules.cpp
@@ -973,7 +973,7 @@ void UNiagaraValidationRule_TickDependencyCheck::CheckValidity(const FNiagaraVal
 				UObject* StackObject = nullptr;
 				if (ResolvedDI.ResolvedSourceEmitterName.Len() > 0)
 				{
-					for (const TSharedPtr<FNiagaraEmitterHandleViewModel>& EmitterViewModel : Context.ViewModel->GetEmitterHandleViewModels())
+					for (const TSharedRef<FNiagaraEmitterHandleViewModel>& EmitterViewModel : Context.ViewModel->GetEmitterHandleViewModels())
 					{
 						if (EmitterViewModel->GetName() == FName(ResolvedDI.ResolvedSourceEmitterName))
 						{
