- Add cast shadows option to sprite / mesh / ribbon renderers

diff --git a/Engine/Plugins/FX/Niagara/Source/Niagara/Private/NiagaraRendererMeshes.cpp b/Engine/Plugins/FX/Niagara/Source/Niagara/Private/NiagaraRendererMeshes.cpp
--- a/Engine/Plugins/FX/Niagara/Source/Niagara/Private/NiagaraRendererMeshes.cpp
+++ b/Engine/Plugins/FX/Niagara/Source/Niagara/Private/NiagaraRendererMeshes.cpp
@@ -112,6 +112,7 @@ FNiagaraRendererMeshes::FNiagaraRendererMeshes(ERHIFeatureLevel::Type FeatureLev
 	RendererVisibility = Properties->RendererVisibility;
 	bAccurateMotionVectors = Properties->NeedsPreciseMotionVectors();
 	bIsHeterogeneousVolume = Properties->UseHeterogeneousVolumes();
+	bCastShadows = Properties->bCastShadows;
 
 	if (Properties->UseHeterogeneousVolumes())
 	{
@@ -1108,9 +1109,9 @@ void FNiagaraRendererMeshes::CreateMeshBatchForSection(
 	MeshBatch.VertexFactory = &VertexFactory;
 	MeshBatch.LCI = NULL;
 	MeshBatch.ReverseCulling = SceneProxy.IsLocalToWorldDeterminantNegative();
-	MeshBatch.CastShadow = SceneProxy.CastsDynamicShadow();
+	MeshBatch.CastShadow = SceneProxy.CastsDynamicShadow() && bCastShadows;
 #if RHI_RAYTRACING
-	MeshBatch.CastRayTracedShadow = SceneProxy.CastsDynamicShadow();
+	MeshBatch.CastRayTracedShadow = SceneProxy.CastsDynamicShadow() && bCastShadows;
 #endif
 	MeshBatch.DepthPriorityGroup = (ESceneDepthPriorityGroup)SceneProxy.GetDepthPriorityGroup(&View);
 
diff --git a/Engine/Plugins/FX/Niagara/Source/Niagara/Private/NiagaraRendererRibbons.cpp b/Engine/Plugins/FX/Niagara/Source/Niagara/Private/NiagaraRendererRibbons.cpp
--- a/Engine/Plugins/FX/Niagara/Source/Niagara/Private/NiagaraRendererRibbons.cpp
+++ b/Engine/Plugins/FX/Niagara/Source/Niagara/Private/NiagaraRendererRibbons.cpp
@@ -709,6 +709,7 @@ FNiagaraRendererRibbons::FNiagaraRendererRibbons(ERHIFeatureLevel::Type FeatureL
 	FacingMode = Properties->FacingMode;
 	DrawDirection = Properties->DrawDirection;
 	RendererLayout = &Properties->RendererLayout;
+	bCastShadows = Properties->bCastShadows;
 	bGpuRibbonLinkIsFloat = Properties->bGpuRibbonLinkIsFloat;
 	GpuRibbonLinkOrderOffset = Properties->GpuRibbonLinkOrderOffset;
 
@@ -2077,11 +2078,10 @@ inline void FNiagaraRendererRibbons::SetupMeshBatchAndCollectorResourceForView(F
 	RenderingViewResources->VertexFactory.InitResource(RHICmdList);
 	RenderingViewResources->VertexFactory.SetRibbonUniformBuffer(RenderingViewResources->UniformBuffer);
 
-
 	OutMeshBatch.VertexFactory = &RenderingViewResources->VertexFactory;
-	OutMeshBatch.CastShadow = SceneProxy->CastsDynamicShadow();
+	OutMeshBatch.CastShadow = SceneProxy->CastsDynamicShadow() && bCastShadows;
 #if RHI_RAYTRACING
-	OutMeshBatch.CastRayTracedShadow = SceneProxy->CastsDynamicShadow();
+	OutMeshBatch.CastRayTracedShadow = SceneProxy->CastsDynamicShadow() && bCastShadows;
 #endif
 	OutMeshBatch.bUseAsOccluder = false;
 	OutMeshBatch.ReverseCulling = SceneProxy->IsLocalToWorldDeterminantNegative();
diff --git a/Engine/Plugins/FX/Niagara/Source/Niagara/Private/NiagaraRendererSprites.cpp b/Engine/Plugins/FX/Niagara/Source/Niagara/Private/NiagaraRendererSprites.cpp
--- a/Engine/Plugins/FX/Niagara/Source/Niagara/Private/NiagaraRendererSprites.cpp
+++ b/Engine/Plugins/FX/Niagara/Source/Niagara/Private/NiagaraRendererSprites.cpp
@@ -102,6 +102,7 @@ FNiagaraRendererSprites::FNiagaraRendererSprites(ERHIFeatureLevel::Type FeatureL
 	MaxFacingCameraBlendDistance = Properties->MaxFacingCameraBlendDistance;
 	RendererVisibility = Properties->RendererVisibility;
 	bAccurateMotionVectors = Properties->NeedsPreciseMotionVectors();
+	bCastShadows = Properties->bCastShadows;
 
 	PixelCoverageMode = Properties->PixelCoverageMode;
 	if (PixelCoverageMode == ENiagaraRendererPixelCoverageMode::Automatic)
@@ -901,9 +902,9 @@ void FNiagaraRendererSprites::CreateMeshBatchForView(
 	VertexFactory.LooseParameterUniformBuffer = FNiagaraSpriteVFLooseParametersRef::CreateUniformBufferImmediate(VFLooseParams, UniformBuffer_SingleFrame);
 
 	MeshBatch.VertexFactory = &VertexFactory;
-	MeshBatch.CastShadow = SceneProxy.CastsDynamicShadow();
+	MeshBatch.CastShadow = SceneProxy.CastsDynamicShadow() && bCastShadows;
 #if RHI_RAYTRACING
-	MeshBatch.CastRayTracedShadow = SceneProxy.CastsDynamicShadow();
+	MeshBatch.CastRayTracedShadow = SceneProxy.CastsDynamicShadow() && bCastShadows;
 #endif
 	MeshBatch.bUseAsOccluder = false;
 	MeshBatch.ReverseCulling = SceneProxy.IsLocalToWorldDeterminantNegative();
diff --git a/Engine/Plugins/FX/Niagara/Source/Niagara/Public/NiagaraMeshRendererProperties.h b/Engine/Plugins/FX/Niagara/Source/Niagara/Public/NiagaraMeshRendererProperties.h
--- a/Engine/Plugins/FX/Niagara/Source/Niagara/Public/NiagaraMeshRendererProperties.h
+++ b/Engine/Plugins/FX/Niagara/Source/Niagara/Public/NiagaraMeshRendererProperties.h
@@ -293,6 +293,13 @@ public:
 	UPROPERTY(EditAnywhere, Category = "Mesh Rendering")
 	uint32 bLockedAxisEnable : 1;
 
+	/**
+	When disabled the renderer will not cast shadows.
+	The component controls if shadows are enabled, this flag allows you to disable the renderer casting shadows.
+	*/
+	UPROPERTY(EditAnywhere, Category = "Rendering")
+	uint8 bCastShadows : 1 = 1;
+
 	/** The materials to be used instead of the StaticMesh's materials. Note that each material must have the Niagara Mesh Particles flag checked. If the ParticleMesh
 	requires more materials than exist in this array or any entry in this array is set to None, we will use the ParticleMesh's existing Material instead.*/
 	UPROPERTY(EditAnywhere, Category = "Mesh Rendering", meta = (EditCondition = "bOverrideMaterials", EditConditionHides))
diff --git a/Engine/Plugins/FX/Niagara/Source/Niagara/Public/NiagaraRendererMeshes.h b/Engine/Plugins/FX/Niagara/Source/Niagara/Public/NiagaraRendererMeshes.h
--- a/Engine/Plugins/FX/Niagara/Source/Niagara/Public/NiagaraRendererMeshes.h
+++ b/Engine/Plugins/FX/Niagara/Source/Niagara/Public/NiagaraRendererMeshes.h
@@ -173,6 +173,7 @@ private:
 	uint32 bEnableFrustumCulling : 1;
 	uint32 bAccurateMotionVectors : 1;
 	uint32 bIsHeterogeneousVolume : 1;
+	uint32 bCastShadows : 1;
 
 	uint32 bSubImageBlend : 1;
 	FVector2f SubImageSize = FVector2f::ZeroVector;
diff --git a/Engine/Plugins/FX/Niagara/Source/Niagara/Public/NiagaraRendererRibbons.h b/Engine/Plugins/FX/Niagara/Source/Niagara/Public/NiagaraRendererRibbons.h
--- a/Engine/Plugins/FX/Niagara/Source/Niagara/Public/NiagaraRendererRibbons.h
+++ b/Engine/Plugins/FX/Niagara/Source/Niagara/Public/NiagaraRendererRibbons.h
@@ -391,6 +391,7 @@ protected:
 	ENiagaraRibbonDrawDirection DrawDirection;
 	
 	const FNiagaraRendererLayout* RendererLayout;
+	bool bCastShadows = true;
 	bool bGpuRibbonLinkIsFloat = false;
 	uint32 GpuRibbonLinkOrderOffset = INDEX_NONE;
 
diff --git a/Engine/Plugins/FX/Niagara/Source/Niagara/Public/NiagaraRendererSprites.h b/Engine/Plugins/FX/Niagara/Source/Niagara/Public/NiagaraRendererSprites.h
--- a/Engine/Plugins/FX/Niagara/Source/Niagara/Public/NiagaraRendererSprites.h
+++ b/Engine/Plugins/FX/Niagara/Source/Niagara/Public/NiagaraRendererSprites.h
@@ -108,6 +108,7 @@ private:
 	uint32 bEnableCulling : 1;
 	uint32 bEnableDistanceCulling : 1;
 	uint32 bAccurateMotionVectors : 1;
+	uint32 bCastShadows : 1;
 	uint32 bSetAnyBoundVars : 1;
 	uint32 bVisTagInParamStore : 1;
 
diff --git a/Engine/Plugins/FX/Niagara/Source/Niagara/Public/NiagaraRibbonRendererProperties.h b/Engine/Plugins/FX/Niagara/Source/Niagara/Public/NiagaraRibbonRendererProperties.h
--- a/Engine/Plugins/FX/Niagara/Source/Niagara/Public/NiagaraRibbonRendererProperties.h
+++ b/Engine/Plugins/FX/Niagara/Source/Niagara/Public/NiagaraRibbonRendererProperties.h
@@ -328,6 +328,13 @@ public:
 	UPROPERTY()
 	uint8 bLinkOrderUseUniqueID : 1;
 
+	/**
+	When disabled the renderer will not cast shadows.
+	The component controls if shadows are enabled, this flag allows you to disable the renderer casting shadows.
+	*/
+	UPROPERTY(EditAnywhere, Category = "Rendering")
+	uint8 bCastShadows : 1 = 1;
+
 	/** Tessellation factor to apply to the width of the ribbon.
 	* Ranges from 1 to 16. Greater values increase amount of tessellation.
 	*/
diff --git a/Engine/Plugins/FX/Niagara/Source/Niagara/Public/NiagaraSpriteRendererProperties.h b/Engine/Plugins/FX/Niagara/Source/Niagara/Public/NiagaraSpriteRendererProperties.h
--- a/Engine/Plugins/FX/Niagara/Source/Niagara/Public/NiagaraSpriteRendererProperties.h
+++ b/Engine/Plugins/FX/Niagara/Source/Niagara/Public/NiagaraSpriteRendererProperties.h
@@ -213,6 +213,13 @@ public:
 	UPROPERTY(EditAnywhere, Category = "Visibility")
 	uint8 bEnableCameraDistanceCulling : 1;
 
+	/**
+	When disabled the renderer will not cast shadows.
+	The component controls if shadows are enabled, this flag allows you to disable the renderer casting shadows.
+	*/
+	UPROPERTY(EditAnywhere, Category = "Rendering")
+	uint8 bCastShadows : 1 = 1;
+
 	/** Sort precision to use when sorting is active. */
 	UPROPERTY(EditAnywhere, Category = "Sorting", meta = (EditCondition = "SortMode != ENiagaraSortMode::None", EditConditionHides))
 	ENiagaraRendererSortPrecision SortPrecision = ENiagaraRendererSortPrecision::Default;
