- Test execution state before performing the LWC tile boundary test
- This fixes an issue where crossing tile boundaries while deactivating could bring the instance back to active

diff --git a/Engine/Plugins/FX/Niagara/Source/Niagara/Private/NiagaraSystemSimulation.cpp b/Engine/Plugins/FX/Niagara/Source/Niagara/Private/NiagaraSystemSimulation.cpp
--- a/Engine/Plugins/FX/Niagara/Source/Niagara/Private/NiagaraSystemSimulation.cpp
+++ b/Engine/Plugins/FX/Niagara/Source/Niagara/Private/NiagaraSystemSimulation.cpp
@@ -1176,11 +1176,14 @@ void FNiagaraSystemSimulation::Tick_GameThread_Internal(float DeltaSeconds, cons
 			{
 				if (USceneComponent* SceneComponent = Instance->GetAttachComponent())
 				{
-					if (UFXSystemComponent::RequiresLWCTileRecache(Instance->GetLWCTile(), SceneComponent->GetComponentLocation()))
+					if ( Instance->GetRequestedExecutionState() == ENiagaraExecutionState::Active && Instance->GetActualExecutionState() == ENiagaraExecutionState::Active)
 					{
-						//-OPT: For safety we reset everything, but if everything is local space we may not need to, or we could rebase.
-						UE_LOG(LogNiagara, Warning, TEXT("NiagaraComponent(%s - %s) required LWC tile recache and was reset."), *GetFullNameSafe(SceneComponent), *GetFullNameSafe(System));
-						Instance->Reset(FNiagaraSystemInstance::EResetMode::ResetAll);
+						if (UFXSystemComponent::RequiresLWCTileRecache(Instance->GetLWCTile(), SceneComponent->GetComponentLocation()))
+						{
+							//-OPT: For safety we reset everything, but if everything is local space we may not need to, or we could rebase.
+							UE_LOG(LogNiagara, Warning, TEXT("NiagaraComponent(%s - %s) required LWC tile recache and was reset."), *GetFullNameSafe(SceneComponent), *GetFullNameSafe(System));
+							Instance->Reset(FNiagaraSystemInstance::EResetMode::ResetAll);
+						}
 					}
 				}
 			}
