- Add LOD mode to Mesh Renderer
- Allows you to specify a LOD level to use or a LOD bias based on Min LOD
- If the LOD level is not available it will fallback to the next available lower level

diff --git a/Engine/Plugins/FX/Niagara/Source/Niagara/Private/DataInterface/NiagaraDataInterfaceDynamicMesh.cpp b/Engine/Plugins/FX/Niagara/Source/Niagara/Private/DataInterface/NiagaraDataInterfaceDynamicMesh.cpp
--- a/Engine/Plugins/FX/Niagara/Source/Niagara/Private/DataInterface/NiagaraDataInterfaceDynamicMesh.cpp
+++ b/Engine/Plugins/FX/Niagara/Source/Niagara/Private/DataInterface/NiagaraDataInterfaceDynamicMesh.cpp
@@ -606,7 +606,7 @@ namespace NDIDynamicMeshLocal
 			return FBox(FVector(-100.0f), FVector(100.0f));
 		}
 
-		virtual void GetLODModelData(FLODModelData& OutLODModelData) const
+		virtual void GetLODModelData(FLODModelData& OutLODModelData, int32 LODLevel) const
 		{
 			const FNDIInstanceData_RenderThread* InstanceData = OwnerProxy->InstanceData_RT.Find(SystemInstanceID);
 			if (InstanceData && InstanceData->NumTriangles > 0 && InstanceData->NumVertices > 0)
diff --git a/Engine/Plugins/FX/Niagara/Source/Niagara/Private/NiagaraMeshRendererProperties.cpp b/Engine/Plugins/FX/Niagara/Source/Niagara/Private/NiagaraMeshRendererProperties.cpp
--- a/Engine/Plugins/FX/Niagara/Source/Niagara/Private/NiagaraMeshRendererProperties.cpp
+++ b/Engine/Plugins/FX/Niagara/Source/Niagara/Private/NiagaraMeshRendererProperties.cpp
@@ -9,6 +9,7 @@
 #include "NiagaraEmitterInstance.h"
 #include "NiagaraGPUSortInfo.h"
 #include "NiagaraSystem.h"
+#include "NiagaraSystemImpl.h"
 
 #include "MaterialDomain.h"
 #include "Materials/MaterialRenderProxy.h"
@@ -45,18 +46,25 @@ public:
 	{
 		WeakStaticMesh	= StaticMesh;
 		RenderData		= StaticMesh->GetRenderData();
+		MeshMinLOD		= StaticMesh->GetMinLODIdx();
 		MinLOD			= StaticMesh->GetMinLODIdx();
 		LocalBounds		= StaticMesh->GetExtendedBounds().GetBox();
 	}
 
+	void SetMinLODBias(int32 MinLODBias) override
+	{
+		MinLOD = FMath::Max(MeshMinLOD + MinLODBias, 0);
+	}
+
 	FBox GetLocalBounds() const override
 	{
 		return LocalBounds;
 	}
 
-	void GetLODModelData(FLODModelData& OutLODModelData) const override
+	void GetLODModelData(FLODModelData& OutLODModelData, int32 LODLevel) const override
 	{
-		OutLODModelData.LODIndex = RenderData->GetCurrentFirstLODIdx(MinLOD);
+		LODLevel = FMath::Max(MinLOD, LODLevel);
+		OutLODModelData.LODIndex = RenderData->GetCurrentFirstLODIdx(LODLevel);
 		if (!RenderData->LODResources.IsValidIndex(OutLODModelData.LODIndex))
 		{
 			OutLODModelData.LODIndex = INDEX_NONE;
@@ -135,6 +143,7 @@ public:
 
 	TWeakObjectPtr<const UStaticMesh>	WeakStaticMesh;
 	const class FStaticMeshRenderData*	RenderData = nullptr;
+	int32								MeshMinLOD = 0;
 	int32								MinLOD = 0;
 	FBox								LocalBounds = FBox(ForceInitToZero);
 };
@@ -212,6 +221,14 @@ FNiagaraMeshRendererMeshProperties::FNiagaraMeshRendererMeshProperties()
 	MeshParameterBinding.SetUsage(ENiagaraParameterBindingUsage::NotParticle);
 	MeshParameterBinding.SetAllowedInterfaces({UNiagaraRenderableMeshInterface::StaticClass()});
 	MeshParameterBinding.SetAllowedObjects({UStaticMesh::StaticClass()});
+
+	LODLevelBinding.SetUsage(ENiagaraParameterBindingUsage::System | ENiagaraParameterBindingUsage::Emitter | ENiagaraParameterBindingUsage::StaticVariable);
+	LODLevelBinding.SetAllowedTypeDefinitions({ FNiagaraTypeDefinition::GetIntDef().ToStaticDef() });
+	LODLevelBinding.SetDefaultParameter(FNiagaraTypeDefinition::GetIntDef().ToStaticDef(), 0);
+
+	LODBiasBinding.SetUsage(ENiagaraParameterBindingUsage::System | ENiagaraParameterBindingUsage::Emitter | ENiagaraParameterBindingUsage::StaticVariable);
+	LODBiasBinding.SetAllowedTypeDefinitions({ FNiagaraTypeDefinition::GetIntDef().ToStaticDef() });
+	LODBiasBinding.SetDefaultParameter(FNiagaraTypeDefinition::GetIntDef().ToStaticDef(), 0);
 #endif
 }
 
@@ -248,9 +265,6 @@ UNiagaraMeshRendererProperties::UNiagaraMeshRendererProperties()
 	, bSubImageBlend(true)
 	, bLockedAxisEnable(false)
 {
-	// Initialize the array with a single, defaulted entry
-	Meshes.AddDefaulted();
-
 #if WITH_EDITORONLY_DATA
 	FlipbookSuffixFormat = TEXT("_{frame_number}");
 	FlipbookSuffixNumDigits = 1;
@@ -451,6 +465,9 @@ void UNiagaraMeshRendererProperties::InitBindings()
 
 		//Default custom sorting to age
 		CustomSortingBinding = FNiagaraConstants::GetAttributeDefaultBinding(SYS_PARAM_PARTICLES_NORMALIZED_AGE);
+
+		// Initialize the array with a single, defaulted entry
+		Meshes.AddDefaulted();
 	}
 
 	SetPreviousBindings(FVersionedNiagaraEmitter(), SourceMode);
@@ -483,6 +500,8 @@ void UNiagaraMeshRendererProperties::UpdateSourceModeDerivates(ENiagaraRendererS
 		for (FNiagaraMeshRendererMeshProperties& Mesh : Meshes)
 		{
 			Mesh.MeshParameterBinding.OnRenameEmitter(SrcEmitter.Emitter->GetUniqueEmitterName());
+			Mesh.LODLevelBinding.OnRenameEmitter(SrcEmitter.Emitter->GetUniqueEmitterName());
+			Mesh.LODBiasBinding.OnRenameEmitter(SrcEmitter.Emitter->GetUniqueEmitterName());
 		}
 #endif
 	}
@@ -552,6 +571,39 @@ void UNiagaraMeshRendererProperties::CacheFromCompiledData(const FNiagaraDataSet
 	MaterialParamValidMask |= bDynamicParam1Valid ? GetDynamicParameterChannelMask(EmitterData, DynamicMaterial1Binding.GetName(), 0xf) << 4 : 0;
 	MaterialParamValidMask |= bDynamicParam2Valid ? GetDynamicParameterChannelMask(EmitterData, DynamicMaterial2Binding.GetName(), 0xf) << 8 : 0;
 	MaterialParamValidMask |= bDynamicParam3Valid ? GetDynamicParameterChannelMask(EmitterData, DynamicMaterial3Binding.GetName(), 0xf) << 12 : 0;
+
+	// Gather LOD information per mesh
+	UNiagaraSystem* OwnerSystem = GetTypedOuter<UNiagaraSystem>();
+	for (FNiagaraMeshRendererMeshProperties& Mesh : Meshes)
+	{
+		Mesh.LODLevel = Mesh.LODLevelBinding.GetDefaultValue<int32>();
+		Mesh.LODBias = Mesh.LODBiasBinding.GetDefaultValue<int32>();
+
+		if (OwnerSystem && (Mesh.LODLevelBinding.AliasedParameter.IsValid() || Mesh.LODBiasBinding.AliasedParameter.IsValid()))
+		{
+			OwnerSystem->ForEachScript(
+				[&Mesh](const UNiagaraScript* NiagaraScript)
+				{
+					if (Mesh.LODLevelBinding.AliasedParameter.IsValid())
+					{
+						TOptional<int32> VariableValue = NiagaraScript->GetStaticVariableValue<int32>(Mesh.LODLevelBinding.ResolvedParameter);
+						if (VariableValue.IsSet())
+						{
+							Mesh.LODLevel = VariableValue.GetValue();
+						}
+					}
+					if (Mesh.LODBiasBinding.AliasedParameter.IsValid())
+					{
+						TOptional<int32> VariableValue = NiagaraScript->GetStaticVariableValue<int32>(Mesh.LODBiasBinding.ResolvedParameter);
+						if (VariableValue.IsSet())
+						{
+							Mesh.LODBias = VariableValue.GetValue();
+						}
+					}
+				}
+			);
+		}
+	}
 #endif
 }
 
@@ -1141,6 +1193,8 @@ void UNiagaraMeshRendererProperties::RenameVariable(const FNiagaraVariableBase&
 	for (FNiagaraMeshRendererMeshProperties& Mesh : Meshes)
 	{
 		Mesh.MeshParameterBinding.OnRenameVariable(OldVariable, NewVariable, InEmitter.Emitter->GetUniqueEmitterName());
+		Mesh.LODLevelBinding.OnRenameVariable(OldVariable, NewVariable, InEmitter.Emitter->GetUniqueEmitterName());
+		Mesh.LODBiasBinding.OnRenameVariable(OldVariable, NewVariable, InEmitter.Emitter->GetUniqueEmitterName());
 	}
 #endif
 }
@@ -1153,6 +1207,8 @@ void UNiagaraMeshRendererProperties::RemoveVariable(const FNiagaraVariableBase&
 	for (FNiagaraMeshRendererMeshProperties& Mesh : Meshes)
 	{
 		Mesh.MeshParameterBinding.OnRemoveVariable(OldVariable, InEmitter.Emitter->GetUniqueEmitterName());
+		Mesh.LODLevelBinding.OnRemoveVariable(OldVariable, InEmitter.Emitter->GetUniqueEmitterName());
+		Mesh.LODBiasBinding.OnRemoveVariable(OldVariable, InEmitter.Emitter->GetUniqueEmitterName());
 	}
 #endif
 }
diff --git a/Engine/Plugins/FX/Niagara/Source/Niagara/Private/NiagaraRendererMeshes.cpp b/Engine/Plugins/FX/Niagara/Source/Niagara/Private/NiagaraRendererMeshes.cpp
--- a/Engine/Plugins/FX/Niagara/Source/Niagara/Private/NiagaraRendererMeshes.cpp
+++ b/Engine/Plugins/FX/Niagara/Source/Niagara/Private/NiagaraRendererMeshes.cpp
@@ -222,15 +222,31 @@ void FNiagaraRendererMeshes::Initialize(const UNiagaraRendererProperties* InProp
 			continue;
 		}
 
+
 		// We have a valid mesh fill in the details
 		FMeshData& MeshData = Meshes.AddDefaulted_GetRef();
 		MeshData.RenderableMesh = RenderableMesh;
 		MeshData.SourceMeshIndex = SourceMeshIndex;
+		MeshData.LODLevel = 0;
+		MeshData.LODBias = 0;
 		MeshData.PivotOffset = FVector3f(MeshProperties.PivotOffset);
 		MeshData.PivotOffsetSpace = MeshProperties.PivotOffsetSpace;
 		MeshData.Scale = FVector3f(MeshProperties.Scale);
 		MeshData.Rotation = FQuat4f(MeshProperties.Rotation.Quaternion());
 
+		if (MeshProperties.LODMode == ENiagaraMeshLODMode::LODLevel)
+		{
+			MeshData.LODLevel = MeshProperties.LODLevel;
+		}
+		else if (MeshProperties.LODMode == ENiagaraMeshLODMode::LODBias)
+		{
+			RenderableMesh->SetMinLODBias(MeshProperties.LODBias);
+		}
+		else		
+		{
+			MeshData.LODBias = MeshProperties.LODBias;
+		}
+
 		// Get materials and remap them into the base material list
 		TArray<UMaterialInterface*> UsedMaterials;
 		RenderableMesh->GetUsedMaterials(UsedMaterials);
@@ -1291,7 +1307,7 @@ void FNiagaraRendererMeshes::GetDynamicMeshElements(const TArray<const FSceneVie
 				const FMeshData& MeshData = Meshes[MeshIndex];
 
 				INiagaraRenderableMesh::FLODModelData LODModel;
-				MeshData.RenderableMesh->GetLODModelData(LODModel);
+				MeshData.RenderableMesh->GetLODModelData(LODModel, MeshData.LODLevel);
 				if (LODModel.LODIndex == INDEX_NONE)
 				{
 					continue;
@@ -1463,7 +1479,7 @@ void FNiagaraRendererMeshes::GetDynamicRayTracingInstances(FRayTracingMaterialGa
 		const FMeshData& MeshData = Meshes[MeshIndex];
 
 		INiagaraRenderableMesh::FLODModelData LODModel;
-		MeshData.RenderableMesh->GetLODModelData(LODModel);
+		MeshData.RenderableMesh->GetLODModelData(LODModel, MeshData.LODLevel);
 		if (LODModel.LODIndex == INDEX_NONE || LODModel.RayTracingGeometry == nullptr)
 		{
 			continue;
diff --git a/Engine/Plugins/FX/Niagara/Source/Niagara/Public/NiagaraMeshRendererProperties.h b/Engine/Plugins/FX/Niagara/Source/Niagara/Public/NiagaraMeshRendererProperties.h
--- a/Engine/Plugins/FX/Niagara/Source/Niagara/Public/NiagaraMeshRendererProperties.h
+++ b/Engine/Plugins/FX/Niagara/Source/Niagara/Public/NiagaraMeshRendererProperties.h
@@ -54,6 +54,37 @@ enum class ENiagaraMeshLockedAxisSpace : uint8
 	Local
 };
 
+UENUM()
+enum class ENiagaraMeshLODMode : uint8
+{
+	/*
+	* Uses the provided LOD level to render all mesh particles.
+	* If the LOD is not streamed in or available on the platform the next available lower LOD level will be used.
+	* For example, LOD Level is set to 1 but the first available is LOD 3 then LOD 3 will be used.
+	*/
+	LODLevel,
+
+	/**
+	* Takes the highest available LOD for the platform + LOD bias to render all mesh particles
+	* If the LOD is not streamed in or available on the platform the next available lower LOD level will be used.
+	* For example, LOD bias is set to 1, the current platform has Min LOD of 2 then 3 will be the used LOD.
+	*/
+	LODBias,
+
+	/*
+	* Calculates the LOD level for the instance based on screen space size and renders all particles with this LOD.
+	* LOD bias is applied to the calculation.
+	*/
+	//PerInstance,
+
+	/*
+	* Calculates and renders each particle with it's calcualted LOD level.
+	* LOD bias is applied to the calculation.
+	* Note: This is the most expensive option as we need to bucket the particles into each LOD and dispatch draws per bucket.
+	*/
+	//PerParticle,
+};
+
 USTRUCT()
 struct FNiagaraMeshMICOverride
 {
@@ -152,6 +183,25 @@ struct FNiagaraMeshRendererMeshProperties
 	UPROPERTY(EditAnywhere, Category = "Mesh")
 	FNiagaraParameterBinding MeshParameterBinding;
 
+	UPROPERTY(EditAnywhere, Category = "Mesh")
+	ENiagaraMeshLODMode LODMode = ENiagaraMeshLODMode::LODLevel;
+
+#if WITH_EDITORONLY_DATA
+	/** Absolute LOD level to use */
+	UPROPERTY(EditAnywhere, Category = "Mesh", meta = (UIMin = "0", DisplayName="LOD Level", EditCondition = "LODMode == ENiagaraMeshLODMode::LODLevel", EditConditionHides))
+	FNiagaraParameterBindingWithValue LODLevelBinding;
+
+	/* LOD bias to apply to the LOD calculation. */
+	UPROPERTY(EditAnywhere, Category = "Mesh", meta = (UIMin = "0", DisplayName = "LOD Bias", EditCondition = "LODMode != ENiagaraMeshLODMode::LODLevel", EditConditionHides))
+	FNiagaraParameterBindingWithValue LODBiasBinding;
+#endif
+
+	UPROPERTY()
+	int32 LODLevel = 0;
+
+	UPROPERTY()
+	int32 LODBias = 0;
+
 	/** Scale of the mesh */
 	UPROPERTY(EditAnywhere, Category = "Mesh")
 	FVector Scale;
diff --git a/Engine/Plugins/FX/Niagara/Source/Niagara/Public/NiagaraRenderableMeshInterface.h b/Engine/Plugins/FX/Niagara/Source/Niagara/Public/NiagaraRenderableMeshInterface.h
--- a/Engine/Plugins/FX/Niagara/Source/Niagara/Public/NiagaraRenderableMeshInterface.h
+++ b/Engine/Plugins/FX/Niagara/Source/Niagara/Public/NiagaraRenderableMeshInterface.h
@@ -37,10 +37,12 @@ public:
 
 	virtual ~INiagaraRenderableMesh() {}
 
+	// Set the Min LOD Bias
+	virtual void SetMinLODBias(int32 MinLODBias) {}
 	// Get the local bounds for the mesh
 	virtual FBox GetLocalBounds() const = 0;
 	// Gather all the relevant mesh data to render the mesh
-	virtual void GetLODModelData(FLODModelData& OutLODModelData) const = 0;
+	virtual void GetLODModelData(FLODModelData& OutLODModelData, int32 LODLevel) const = 0;
 	// Setup the vertex factory for the mesh
 	virtual void SetupVertexFactory(class FNiagaraMeshVertexFactory& InVertexFactory, const FLODModelData& LODModelData) const = 0;
 	// Gather a list of used materials
diff --git a/Engine/Plugins/FX/Niagara/Source/Niagara/Public/NiagaraRendererMeshes.h b/Engine/Plugins/FX/Niagara/Source/Niagara/Public/NiagaraRendererMeshes.h
--- a/Engine/Plugins/FX/Niagara/Source/Niagara/Public/NiagaraRendererMeshes.h
+++ b/Engine/Plugins/FX/Niagara/Source/Niagara/Public/NiagaraRendererMeshes.h
@@ -83,6 +83,8 @@ protected:
 	{
 		FNiagaraRenderableMeshPtr RenderableMesh;
 		uint32 SourceMeshIndex = INDEX_NONE;
+		int32 LODLevel = 0;
+		int32 LODBias = 0;
 		FVector3f PivotOffset = FVector3f::ZeroVector;
 		ENiagaraMeshPivotOffsetSpace PivotOffsetSpace = ENiagaraMeshPivotOffsetSpace::Mesh;
 		FVector3f Scale = FVector3f(1.0f, 1.0f, 1.0f);
