- Add option for ribbon plane normals to follow the geometry rather than face screen / custom facing vector
- This is now the new default as it matches the other ribbon shape modes
- Existing content will use the previous mode

diff --git a/Engine/Plugins/FX/Niagara/Shaders/Private/NiagaraRibbonVertexFactory.ush b/Engine/Plugins/FX/Niagara/Shaders/Private/NiagaraRibbonVertexFactory.ush
--- a/Engine/Plugins/FX/Niagara/Shaders/Private/NiagaraRibbonVertexFactory.ush
+++ b/Engine/Plugins/FX/Niagara/Shaders/Private/NiagaraRibbonVertexFactory.ush
@@ -588,9 +588,9 @@ FRibbonPositionIntermediates GetRibbonPositionIntermediates(
 	
 	Intermediates.GeoUp = cross(Intermediates.TangentY, Intermediates.TangentX);
 	Intermediates.GeoRight = Intermediates.TangentX;
-	
-	Intermediates.TangentZ = NiagaraRibbonVFLooseParameters.Shape >= 1 ? Intermediates.GeoUp : Facing; // Intermediates.GeoUp;
-	
+
+	Intermediates.TangentZ = NiagaraRibbonVFLooseParameters.UseGeometryNormals ? Intermediates.GeoUp : Facing;
+
 	Intermediates.TangentZ = SafeNormalize(Intermediates.TangentZ * SliceVertexNormals.y + Intermediates.TangentX * SliceVertexNormals.x, float3(0, 0, 1));
 	
 	// Flip the normal only for multiplane shape if the normal is looking away from the camera.
@@ -605,6 +605,7 @@ FRibbonPositionIntermediates GetRibbonPositionIntermediates(
 	Intermediates.TangentX = Intermediates.TangentY;
 	Intermediates.TangentY = t;
 
+
 	// Overrides (FVertexFactoryIntermediates)0; 
 	Intermediates.RibbonWidth = 1.0f;
 
diff --git a/Engine/Plugins/FX/Niagara/Source/Niagara/Private/NiagaraRendererRibbons.cpp b/Engine/Plugins/FX/Niagara/Source/Niagara/Private/NiagaraRendererRibbons.cpp
--- a/Engine/Plugins/FX/Niagara/Source/Niagara/Private/NiagaraRendererRibbons.cpp
+++ b/Engine/Plugins/FX/Niagara/Source/Niagara/Private/NiagaraRendererRibbons.cpp
@@ -710,6 +710,7 @@ FNiagaraRendererRibbons::FNiagaraRendererRibbons(ERHIFeatureLevel::Type FeatureL
 	DrawDirection = Properties->DrawDirection;
 	RendererLayout = &Properties->RendererLayout;
 	bCastShadows = Properties->bCastShadows;
+	bUseGeometryNormals = Properties->bUseGeometryNormals;
 	bGpuRibbonLinkIsFloat = Properties->bGpuRibbonLinkIsFloat;
 	GpuRibbonLinkOrderOffset = Properties->GpuRibbonLinkOrderOffset;
 
@@ -2068,6 +2069,7 @@ inline void FNiagaraRendererRibbons::SetupMeshBatchAndCollectorResourceForView(F
 	VFLooseParams.FacingMode = static_cast<uint32>(FacingMode);
 	VFLooseParams.Shape = static_cast<uint32>(ShapeState.Shape);
 	VFLooseParams.NeedsPreciseMotionVectors = GenerationConfig.NeedsPreciseMotionVectors();
+	VFLooseParams.UseGeometryNormals = (ShapeState.Shape != ENiagaraRibbonShapeMode::Plane || bUseGeometryNormals) ? 1 : 0;
 
 	VFLooseParams.IndirectDrawOutput = bShouldUseGPUInitIndices ? (FRHIShaderResourceView*)RenderingViewResources->IndirectDrawBuffer.SRV : GetDummyUIntBuffer();
 	VFLooseParams.IndirectDrawOutputOffset = bShouldUseGPUInitIndices ? 0 : -1;
diff --git a/Engine/Plugins/FX/Niagara/Source/Niagara/Private/NiagaraRibbonRendererProperties.cpp b/Engine/Plugins/FX/Niagara/Source/Niagara/Private/NiagaraRibbonRendererProperties.cpp
--- a/Engine/Plugins/FX/Niagara/Source/Niagara/Private/NiagaraRibbonRendererProperties.cpp
+++ b/Engine/Plugins/FX/Niagara/Source/Niagara/Private/NiagaraRibbonRendererProperties.cpp
@@ -66,6 +66,7 @@ UNiagaraRibbonRendererProperties::UNiagaraRibbonRendererProperties()
 	, Shape(ENiagaraRibbonShapeMode::Plane)
 	, bEnableAccurateGeometry(false)
 	, bUseMaterialBackfaceCulling(false)
+	, bUseGeometryNormals(true)
 	, bUseGPUInit(false)
 	, bUseConstantFactor(false)
 	, bScreenSpaceTessellation(true)
@@ -159,6 +160,11 @@ void UNiagaraRibbonRendererProperties::PostLoad()
 		bLinkOrderUseUniqueID = false;
 	}
 
+	if (NiagaraVer < FNiagaraCustomVersion::RibbonPlaneUseGeometryNormals)
+	{
+		bUseGeometryNormals = false;
+	}
+
 	ChangeToPositionBinding(PositionBinding);
 #endif
 	
diff --git a/Engine/Plugins/FX/Niagara/Source/Niagara/Public/NiagaraRendererRibbons.h b/Engine/Plugins/FX/Niagara/Source/Niagara/Public/NiagaraRendererRibbons.h
--- a/Engine/Plugins/FX/Niagara/Source/Niagara/Public/NiagaraRendererRibbons.h
+++ b/Engine/Plugins/FX/Niagara/Source/Niagara/Public/NiagaraRendererRibbons.h
@@ -392,6 +392,7 @@ protected:
 	
 	const FNiagaraRendererLayout* RendererLayout;
 	bool bCastShadows = true;
+	bool bUseGeometryNormals = false;
 	bool bGpuRibbonLinkIsFloat = false;
 	uint32 GpuRibbonLinkOrderOffset = INDEX_NONE;
 
diff --git a/Engine/Plugins/FX/Niagara/Source/Niagara/Public/NiagaraRibbonRendererProperties.h b/Engine/Plugins/FX/Niagara/Source/Niagara/Public/NiagaraRibbonRendererProperties.h
--- a/Engine/Plugins/FX/Niagara/Source/Niagara/Public/NiagaraRibbonRendererProperties.h
+++ b/Engine/Plugins/FX/Niagara/Source/Niagara/Public/NiagaraRibbonRendererProperties.h
@@ -309,6 +309,10 @@ public:
 	UPROPERTY(EditAnywhere, Category = "Ribbon Shape")
 	uint8 bUseMaterialBackfaceCulling : 1;
 
+	/** When enabled the ribbons normals will follow the shape of the geometry rather than being aligned to screen / custom facing. */
+	UPROPERTY(EditAnywhere, Category = "Ribbon Shape", meta = (EditCondition = "Shape == ENiagaraRibbonShapeMode::Plane", EditConditionHides))
+	uint8 bUseGeometryNormals : 1;
+
 	/**
 	*	Whether we use the CPU or GPU to generate ribbon geometry for CPU systems.
 	*	GPU systems will always use a fully GPU initialization pipeline,
diff --git a/Engine/Plugins/FX/Niagara/Source/NiagaraCore/Public/NiagaraCustomVersion.h b/Engine/Plugins/FX/Niagara/Source/NiagaraCore/Public/NiagaraCustomVersion.h
--- a/Engine/Plugins/FX/Niagara/Source/NiagaraCore/Public/NiagaraCustomVersion.h
+++ b/Engine/Plugins/FX/Niagara/Source/NiagaraCore/Public/NiagaraCustomVersion.h
@@ -191,6 +191,8 @@ struct FNiagaraCustomVersion
 
 		SubImageBlendEnabledByDefault,	// Renderer SubImage Blends are enabled by default
 
+		RibbonPlaneUseGeometryNormals,	// Ribbon renderer will use geometry normals by default rather than screen / facing aligned normals
+
 		// DO NOT ADD A NEW VERSION UNLESS YOU HAVE TALKED TO THE NIAGARA LEAD. Mismanagement of these versions can lead to data loss if it is adjusted in multiple streams simultaneously.
 		// -----<new versions can be added above this line>  -------------------------------------------------
 		VersionPlusOne,
diff --git a/Engine/Plugins/FX/Niagara/Source/NiagaraVertexFactories/Public/NiagaraRibbonVertexFactory.h b/Engine/Plugins/FX/Niagara/Source/NiagaraVertexFactories/Public/NiagaraRibbonVertexFactory.h
--- a/Engine/Plugins/FX/Niagara/Source/NiagaraVertexFactories/Public/NiagaraRibbonVertexFactory.h
+++ b/Engine/Plugins/FX/Niagara/Source/NiagaraVertexFactories/Public/NiagaraRibbonVertexFactory.h
@@ -74,6 +74,7 @@ BEGIN_GLOBAL_SHADER_PARAMETER_STRUCT(FNiagaraRibbonVFLooseParameters, NIAGARAVER
 	SHADER_PARAMETER(uint32, FacingMode)
 	SHADER_PARAMETER(uint32, Shape)
 	SHADER_PARAMETER(uint32, NeedsPreciseMotionVectors)
+	SHADER_PARAMETER(uint32, UseGeometryNormals)
 	SHADER_PARAMETER_SRV(Buffer<uint>, SortedIndices)
 	SHADER_PARAMETER_SRV(Buffer<float>, TangentsAndDistances)
 	SHADER_PARAMETER_SRV(Buffer<uint>, MultiRibbonIndices)
