Moved Niagara dynamic update to its own render command pipe.

diff --git a/Engine/Plugins/FX/Niagara/Source/Niagara/Private/NiagaraComponent.cpp b/Engine/Plugins/FX/Niagara/Source/Niagara/Private/NiagaraComponent.cpp
--- a/Engine/Plugins/FX/Niagara/Source/Niagara/Private/NiagaraComponent.cpp
+++ b/Engine/Plugins/FX/Niagara/Source/Niagara/Private/NiagaraComponent.cpp
@@ -43,6 +43,8 @@ DECLARE_CYCLE_STAT(TEXT("Get Dynamic Mesh Elements (RT)"), STAT_NiagaraComponent
 
 DEFINE_LOG_CATEGORY(LogNiagara);
 
+DEFINE_RENDER_COMMAND_PIPE(NiagaraDynamicData);
+
 static int GNiagaraSoloTickEarly = 1;
 static FAutoConsoleVariableRef CVarNiagaraSoloTickEarly(
 	TEXT("fx.Niagara.Solo.TickEarly"),
@@ -2139,8 +2141,8 @@ void UNiagaraComponent::SendRenderDynamicData_Concurrent()
 				SystemInstanceController->GenerateSetDynamicDataCommands(SetDataCommands, *RenderData, *NiagaraProxy);
 			}
 
-			ENQUEUE_RENDER_COMMAND(NiagaraSetDynamicData)(
-				[NiagaraProxy, CommandsRT = MoveTemp(SetDataCommands), NewProxyDynamicData](FRHICommandListImmediate& RHICmdList)
+			ENQUEUE_RENDER_COMMAND(NiagaraSetDynamicData)(UE::RenderCommandPipe::NiagaraDynamicData,
+				[NiagaraProxy, CommandsRT = MoveTemp(SetDataCommands), NewProxyDynamicData]
 				{
 					SCOPE_CYCLE_COUNTER(STAT_NiagaraSetDynamicData);
 					PARTICLE_PERF_STAT_CYCLES_WITH_COUNT_RT(NewProxyDynamicData.PerfStatsContext, RenderUpdate, 1);
