- Add platform set to Niagara System to allow systems to be culled for specific platforms / quality levels

diff --git a/Engine/Plugins/FX/Niagara/Source/Niagara/Classes/NiagaraSystem.h b/Engine/Plugins/FX/Niagara/Source/Niagara/Classes/NiagaraSystem.h
--- a/Engine/Plugins/FX/Niagara/Source/Niagara/Classes/NiagaraSystem.h
+++ b/Engine/Plugins/FX/Niagara/Source/Niagara/Classes/NiagaraSystem.h
@@ -708,7 +708,6 @@ public:
 	FORCEINLINE bool GetOverrideScalabilitySettings()const { return bOverrideScalabilitySettings; }
 	FORCEINLINE void SetOverrideScalabilitySettings(bool bOverride) { bOverrideScalabilitySettings = bOverride; }
 
-
 	NIAGARA_API void GatherStaticVariables(TArray<FNiagaraVariable>& OutVars, TArray<FNiagaraVariable>& OutEmitterVars) const;
 #endif
 	NIAGARA_API UNiagaraEffectType* GetEffectType()const;
@@ -716,6 +715,12 @@ public:
 	NIAGARA_API const FNiagaraSystemScalabilityOverride& GetCurrentOverrideSettings() const;
 	FORCEINLINE bool NeedsSortedSignificanceCull()const{ return bNeedsSortedSignificanceCull; }
 	
+	NIAGARA_API FNiagaraPlatformSet& GetScalabilityPlatformSet() { return Platforms; }
+	NIAGARA_API const FNiagaraPlatformSet& GetScalabilityPlatformSet() const { return Platforms; }
+
+	/** Returns true if this emitter's platform filter allows it on this platform and quality level. */
+	NIAGARA_API bool IsAllowedByScalability() const;
+
 	NIAGARA_API void UpdateScalability();
 
 	NIAGARA_API ENiagaraCullProxyMode GetCullProxyMode()const;
@@ -820,6 +825,9 @@ protected:
 	UPROPERTY(EditAnywhere, Category = "Scalability", meta = (EditCondition="bOverrideScalabilitySettings", DisplayInScalabilityContext))
 	FNiagaraSystemScalabilityOverrides SystemScalabilityOverrides;
 
+	UPROPERTY(EditAnywhere, Category = "Scalability", meta = (DisplayInScalabilityContext))
+	FNiagaraPlatformSet Platforms;
+
 	/** Handles to the emitter this System will simulate. */
 	UPROPERTY()
 	TArray<FNiagaraEmitterHandle> EmitterHandles;
diff --git a/Engine/Plugins/FX/Niagara/Source/Niagara/Private/NiagaraComponent.cpp b/Engine/Plugins/FX/Niagara/Source/Niagara/Private/NiagaraComponent.cpp
--- a/Engine/Plugins/FX/Niagara/Source/Niagara/Private/NiagaraComponent.cpp
+++ b/Engine/Plugins/FX/Niagara/Source/Niagara/Private/NiagaraComponent.cpp
@@ -1238,6 +1238,12 @@ void UNiagaraComponent::ActivateInternal(bool bReset /* = false */, bool bIsScal
 		return;
 	}
 
+	if (!Asset->IsAllowedByScalability())
+	{
+		OnSystemComplete(true);
+		return;
+	}
+
 	bIsCulledByScalability = false;
 	
 	//Ensure we're not paused. Non-scalability activate calls should override pause state.
diff --git a/Engine/Plugins/FX/Niagara/Source/Niagara/Private/NiagaraEmitter.cpp b/Engine/Plugins/FX/Niagara/Source/Niagara/Private/NiagaraEmitter.cpp
--- a/Engine/Plugins/FX/Niagara/Source/Niagara/Private/NiagaraEmitter.cpp
+++ b/Engine/Plugins/FX/Niagara/Source/Niagara/Private/NiagaraEmitter.cpp
@@ -1070,7 +1070,8 @@ bool UNiagaraEmitter::NeedsLoadForTargetPlatform(const ITargetPlatform* TargetPl
 {
 	// Don't load disabled emitters.
 	// Awkwardly, this requires us to look for ourselves in the owning system.
-	if (const UNiagaraSystem* OwnerSystem = GetTypedOuter<const UNiagaraSystem>())
+	const UNiagaraSystem* OwnerSystem = GetTypedOuter<const UNiagaraSystem>();
+	if ( OwnerSystem )
 	{
 		for (const FNiagaraEmitterHandle& EmitterHandle : OwnerSystem->GetEmitterHandles())
 		{
@@ -1090,6 +1091,13 @@ bool UNiagaraEmitter::NeedsLoadForTargetPlatform(const ITargetPlatform* TargetPl
 		return true;
 	}
 
+	if (OwnerSystem && !OwnerSystem->GetScalabilityPlatformSet().IsEnabledForPlatform(TargetPlatform->IniPlatformName()))
+	{
+		UE_LOG(LogNiagara, Verbose, TEXT("Pruned emitter %s for platform %s from system scalability"), *GetFullName(), *TargetPlatform->DisplayName().ToString())
+		return false;
+	}
+
+
 	bool bIsEnabled = IsEnabledOnPlatform(TargetPlatform->IniPlatformName());
 	if(!bIsEnabled)
 	{
diff --git a/Engine/Plugins/FX/Niagara/Source/Niagara/Private/NiagaraSystem.cpp b/Engine/Plugins/FX/Niagara/Source/Niagara/Private/NiagaraSystem.cpp
--- a/Engine/Plugins/FX/Niagara/Source/Niagara/Private/NiagaraSystem.cpp
+++ b/Engine/Plugins/FX/Niagara/Source/Niagara/Private/NiagaraSystem.cpp
@@ -1025,9 +1025,9 @@ void UNiagaraSystem::PostLoad()
 #endif // WITH_EDITORONLY_DATA
 
 	//Apply platform set redirectors
-	auto ApplyPlatformSetRedirects = [](FNiagaraPlatformSet& Platforms)
+	auto ApplyPlatformSetRedirects = [](FNiagaraPlatformSet& PlatformSet)
 	{
-		Platforms.ApplyRedirects();
+		PlatformSet.ApplyRedirects();
 	};
 	ForEachPlatformSet(ApplyPlatformSetRedirects);
 
@@ -3565,6 +3565,11 @@ void UNiagaraSystem::ResolveScalabilitySettings()
 	FNiagaraWorldManager::RequestInvalidateCachedSystemScalabilityDataForAllWorlds();
 }
 
+bool UNiagaraSystem::IsAllowedByScalability() const
+{
+	return Platforms.IsActive();
+}
+
 void UNiagaraSystem::UpdateScalability()
 {
 	CacheFromCompiledData();
diff --git a/Engine/Plugins/FX/Niagara/Source/NiagaraEditor/Private/Toolkits/SystemToolkitModes/NiagaraSystemToolkitMode_Scalability.cpp b/Engine/Plugins/FX/Niagara/Source/NiagaraEditor/Private/Toolkits/SystemToolkitModes/NiagaraSystemToolkitMode_Scalability.cpp
--- a/Engine/Plugins/FX/Niagara/Source/NiagaraEditor/Private/Toolkits/SystemToolkitModes/NiagaraSystemToolkitMode_Scalability.cpp
+++ b/Engine/Plugins/FX/Niagara/Source/NiagaraEditor/Private/Toolkits/SystemToolkitModes/NiagaraSystemToolkitMode_Scalability.cpp
@@ -126,6 +126,7 @@ void FNiagaraSystemToolkitMode_Scalability::BindSystemPreviewOverrides(UNiagaraS
 	TSharedPtr<FNiagaraSystemViewModel> SystemViewModel = SystemToolkit.Pin()->GetSystemViewModel();
 	UNiagaraSystemScalabilityViewModel* ScalabilityViewModel = SystemViewModel->GetScalabilityViewModel();
 
+	SystemViewModel->GetSystem().GetScalabilityPlatformSet().OnOverrideQualityLevelDelegate.BindUObject(ScalabilityViewModel, &UNiagaraSystemScalabilityViewModel::GetPreviewQualityLevel);
 	for(FNiagaraSystemScalabilityOverride& SystemScalabilityOverride : SystemViewModel->GetSystem().GetSystemScalabilityOverrides().Overrides)
 	{
 		SystemScalabilityOverride.Platforms.OnOverrideQualityLevelDelegate.BindUObject(ScalabilityViewModel, &UNiagaraSystemScalabilityViewModel::GetPreviewQualityLevel);
@@ -162,6 +163,7 @@ void FNiagaraSystemToolkitMode_Scalability::UnbindSystemPreviewOverrides()
 {
 	TSharedPtr<FNiagaraSystemViewModel> SystemViewModel = SystemToolkit.Pin()->GetSystemViewModel();
 
+	SystemViewModel->GetSystem().GetScalabilityPlatformSet().OnOverrideQualityLevelDelegate.Unbind();
 	for(FNiagaraSystemScalabilityOverride& SystemScalabilityOverride : SystemViewModel->GetSystem().GetSystemScalabilityOverrides().Overrides)
 	{
 		SystemScalabilityOverride.Platforms.OnOverrideQualityLevelDelegate.Unbind();
diff --git a/Engine/Plugins/FX/Niagara/Source/NiagaraEditorWidgets/Private/SNiagaraOverviewStackNode.cpp b/Engine/Plugins/FX/Niagara/Source/NiagaraEditorWidgets/Private/SNiagaraOverviewStackNode.cpp
--- a/Engine/Plugins/FX/Niagara/Source/NiagaraEditorWidgets/Private/SNiagaraOverviewStackNode.cpp
+++ b/Engine/Plugins/FX/Niagara/Source/NiagaraEditorWidgets/Private/SNiagaraOverviewStackNode.cpp
@@ -857,9 +857,24 @@ void SNiagaraOverviewStackNode::OnScalabilityModeChanged(bool bActive)
 EVisibility SNiagaraOverviewStackNode::ShowExcludedOverlay() const
 {
 	// we only want actual results in scalability mode and for nodes representing emitters (not system nodes)
-	if(bScalabilityModeActive && EmitterHandleViewModelWeak.IsValid())
-	{		
-		return EmitterHandleViewModelWeak.Pin()->GetEmitterHandle()->GetEmitterData()->IsAllowedByScalability() ? EVisibility::Hidden : EVisibility::HitTestInvisible;
+	if(bScalabilityModeActive)
+	{
+		if (UNiagaraSystemScalabilityViewModel* ScalabilityVM = ScalabilityViewModel.Get())
+		{
+			const TSharedPtr<FNiagaraSystemViewModel> SystemViewModel = ScalabilityVM->GetSystemViewModel().Pin();
+			if (SystemViewModel && !SystemViewModel->GetSystem().IsAllowedByScalability())
+			{
+				return EVisibility::HitTestInvisible;
+			}
+		}
+
+		if (const TSharedPtr<FNiagaraEmitterHandleViewModel> EmitterHandleViewModel = EmitterHandleViewModelWeak.Pin())
+		{
+			if (!EmitterHandleViewModel->GetEmitterHandle()->GetEmitterData()->IsAllowedByScalability())
+			{
+				return EVisibility::HitTestInvisible;
+			}
+		}
 	}
 	
 	return EVisibility::Hidden; 
