- Code clean up, removing dead members / functions

#rnx

diff --git a/Engine/Plugins/FX/Niagara/Source/Niagara/Private/NiagaraRendererMeshes.cpp b/Engine/Plugins/FX/Niagara/Source/Niagara/Private/NiagaraRendererMeshes.cpp
--- a/Engine/Plugins/FX/Niagara/Source/Niagara/Private/NiagaraRendererMeshes.cpp
+++ b/Engine/Plugins/FX/Niagara/Source/Niagara/Private/NiagaraRendererMeshes.cpp
@@ -1318,13 +1318,10 @@ void FNiagaraRendererMeshes::GetDynamicMeshElements(const TArray<const FSceneVie
 				// Get the next vertex factory to use
 				// TODO: Find a way to safely pool these such that they won't be concurrently accessed by multiple views
 				FNiagaraMeshVertexFactory& VertexFactory = CollectorResources->VertexFactory;
-				VertexFactory.SetParticleFactoryType(NVFT_Mesh);
 				if (ViewFamily.Scene != nullptr)
 				{
 					VertexFactory.SetNiagaraMeshFeatureLevel(ViewFamily.Scene->GetFeatureLevel());
 				}
-				VertexFactory.SetMeshIndex(MeshIndex);
-				VertexFactory.SetLODIndex(LODModel.LODIndex);
 				VertexFactory.EnablePrimitiveIDElement(ParticleMeshRenderData.bUseGPUScene);
 				VertexFactory.InitResource(FRHICommandListImmediate::Get());
 				MeshData.RenderableMesh->SetupVertexFactory(VertexFactory, LODModel);
@@ -1493,13 +1490,10 @@ void FNiagaraRendererMeshes::GetDynamicRayTracingInstances(FRayTracingMaterialGa
 		// Get the next vertex factory to use
 		// TODO: Find a way to safely pool these such that they won't be concurrently accessed by multiple views
 		FNiagaraMeshVertexFactory& VertexFactory = CollectorResources->VertexFactory;
-		VertexFactory.SetParticleFactoryType(NVFT_Mesh);
 		if (View->Family->Scene != nullptr)
 		{
 			VertexFactory.SetNiagaraMeshFeatureLevel(View->Family->Scene->GetFeatureLevel());
 		}
-		VertexFactory.SetMeshIndex(MeshIndex);
-		VertexFactory.SetLODIndex(LODModel.LODIndex);
 		VertexFactory.EnablePrimitiveIDElement(ParticleMeshRenderData.bUseGPUScene);
 		VertexFactory.InitResource(Context.GraphBuilder.RHICmdList);
 		MeshData.RenderableMesh->SetupVertexFactory(VertexFactory, LODModel);
diff --git a/Engine/Plugins/FX/Niagara/Source/Niagara/Private/NiagaraRendererRibbons.cpp b/Engine/Plugins/FX/Niagara/Source/Niagara/Private/NiagaraRendererRibbons.cpp
--- a/Engine/Plugins/FX/Niagara/Source/Niagara/Private/NiagaraRendererRibbons.cpp
+++ b/Engine/Plugins/FX/Niagara/Source/Niagara/Private/NiagaraRendererRibbons.cpp
@@ -2075,7 +2075,6 @@ inline void FNiagaraRendererRibbons::SetupMeshBatchAndCollectorResourceForView(F
 	VFLooseParams.IndirectDrawOutputOffset = bShouldUseGPUInitIndices ? 0 : -1;
 
 	// Collector.AllocateOneFrameResource uses default ctor, initialize the vertex factory
-	RenderingViewResources->VertexFactory.SetParticleFactoryType(NVFT_Ribbon);
 	RenderingViewResources->VertexFactory.LooseParameterUniformBuffer = FNiagaraRibbonVFLooseParametersRef::CreateUniformBufferImmediate(VFLooseParams, UniformBuffer_SingleFrame);
 	RenderingViewResources->VertexFactory.InitResource(RHICmdList);
 	RenderingViewResources->VertexFactory.SetRibbonUniformBuffer(RenderingViewResources->UniformBuffer);
diff --git a/Engine/Plugins/FX/Niagara/Source/Niagara/Private/NiagaraRendererSprites.cpp b/Engine/Plugins/FX/Niagara/Source/Niagara/Private/NiagaraRendererSprites.cpp
--- a/Engine/Plugins/FX/Niagara/Source/Niagara/Private/NiagaraRendererSprites.cpp
+++ b/Engine/Plugins/FX/Niagara/Source/Niagara/Private/NiagaraRendererSprites.cpp
@@ -474,8 +474,6 @@ void FNiagaraRendererSprites::InitializeSortInfo(FParticleSpriteRenderData& Part
 
 void FNiagaraRendererSprites::SetupVertexFactory(FRHICommandListBase& RHICmdList, FParticleSpriteRenderData& ParticleSpriteRenderData, FNiagaraSpriteVertexFactory& VertexFactory) const
 {
-	VertexFactory.SetParticleFactoryType(NVFT_Sprite);
-
 	// Set facing / alignment
 	{
 		ENiagaraSpriteFacingMode ActualFacingMode = FacingMode;
diff --git a/Engine/Plugins/FX/Niagara/Source/NiagaraVertexFactories/Public/NiagaraMeshVertexFactory.h b/Engine/Plugins/FX/Niagara/Source/NiagaraVertexFactories/Public/NiagaraMeshVertexFactory.h
--- a/Engine/Plugins/FX/Niagara/Source/NiagaraVertexFactories/Public/NiagaraMeshVertexFactory.h
+++ b/Engine/Plugins/FX/Niagara/Source/NiagaraVertexFactories/Public/NiagaraMeshVertexFactory.h
@@ -129,15 +129,11 @@ public:
 	/** Default constructor. */
 	FNiagaraMeshVertexFactory(ENiagaraVertexFactoryType InType, ERHIFeatureLevel::Type InFeatureLevel)
 		: FNiagaraVertexFactoryBase(InType, InFeatureLevel)
-		, MeshIndex(-1)
-		, LODIndex(-1)
 		, bAddPrimitiveIDElement(true)
 	{}
 
 	FNiagaraMeshVertexFactory()
 		: FNiagaraVertexFactoryBase(NVFT_MAX, ERHIFeatureLevel::Num)
-		, MeshIndex(-1)
-		, LODIndex(-1)
 		, bAddPrimitiveIDElement(true)
 	{}
 
@@ -208,12 +204,6 @@ public:
 	// FRenderResource interface.
 	NIAGARAVERTEXFACTORIES_API virtual void InitRHI(FRHICommandListBase& RHICmdList) override;
 
-	int32 GetMeshIndex() const { return MeshIndex; }
-	void SetMeshIndex(int32 InMeshIndex) { MeshIndex = InMeshIndex; }
-
-	int32 GetLODIndex() const { return LODIndex; }
-	void SetLODIndex(int32 InLODIndex) { LODIndex = InLODIndex; }
-
 	bool IsPrimitiveIDElementEnabled() const { return bAddPrimitiveIDElement; }
 	void EnablePrimitiveIDElement(bool bEnable) { bAddPrimitiveIDElement = bEnable; }
 	
@@ -226,8 +216,6 @@ protected:
 
 protected:
 	FStaticMeshDataType Data;
-	int32 MeshIndex;	
-	int32 LODIndex;
 	bool bAddPrimitiveIDElement;
 
 	/** Uniform buffer with mesh particle parameters. */
diff --git a/Engine/Plugins/FX/Niagara/Source/NiagaraVertexFactories/Public/NiagaraVertexFactory.h b/Engine/Plugins/FX/Niagara/Source/NiagaraVertexFactories/Public/NiagaraVertexFactory.h
--- a/Engine/Plugins/FX/Niagara/Source/NiagaraVertexFactories/Public/NiagaraVertexFactory.h
+++ b/Engine/Plugins/FX/Niagara/Source/NiagaraVertexFactories/Public/NiagaraVertexFactory.h
@@ -64,12 +64,6 @@ public:
 	/** Default constructor. */
 	explicit FNiagaraVertexFactoryBase(ENiagaraVertexFactoryType Type, ERHIFeatureLevel::Type InFeatureLevel)
 		: FVertexFactory(InFeatureLevel)
-		, LastFrameSetup(MAX_uint32)
-		, LastViewFamily(nullptr)
-		, LastView(nullptr)
-		, LastFrameRealTime(-1.0)
-		, ParticleFactoryType(Type)
-		, bInUse(false)
 	{
 		bNeedsDeclaration = false;
 	}
@@ -80,56 +74,7 @@ public:
 		OutEnvironment.SetDefine(TEXT("NIAGARA_PARTICLE_FACTORY"), TEXT("1"));
 	}
 
-	/** Return the vertex factory type */
-	FORCEINLINE ENiagaraVertexFactoryType GetParticleFactoryType() const
-	{
-		return ParticleFactoryType;
-	}
-
-	inline void SetParticleFactoryType(ENiagaraVertexFactoryType InType)
-	{
-		ParticleFactoryType = InType;
-	}
-
-	/** Specify whether the factory is in use or not. */
-	FORCEINLINE void SetInUse(bool bInInUse)
-	{
-		bInUse = bInInUse;
-	}
-
-	/** Return the vertex factory type */
-	FORCEINLINE bool GetInUse() const
-	{
-		return bInUse;
-	}
-
 	ERHIFeatureLevel::Type GetFeatureLevel() const { check(HasValidFeatureLevel());  return FRenderResource::GetFeatureLevel(); }
-
-	bool CheckAndUpdateLastFrame(const FSceneViewFamily& ViewFamily, const FSceneView *View = nullptr) const
-	{
-		if (LastFrameSetup != MAX_uint32 && (&ViewFamily == LastViewFamily) && (View == LastView) && ViewFamily.FrameNumber == LastFrameSetup && LastFrameRealTime == ViewFamily.Time.GetRealTimeSeconds())
-		{
-			return false;
-		}
-		LastFrameSetup = ViewFamily.FrameNumber;
-		LastFrameRealTime = ViewFamily.Time.GetRealTimeSeconds();
-		LastViewFamily = &ViewFamily;
-		LastView = View;
-		return true;
-	}
-
-private:
-	/** Last state where we set this. We only need to setup these once per frame, so detemine same frame by number, time, and view family. */
-	mutable uint32 LastFrameSetup;
-	mutable const FSceneViewFamily *LastViewFamily;
-	mutable const FSceneView *LastView;
-	mutable double LastFrameRealTime;
-
-	/** The type of the vertex factory. */
-	ENiagaraVertexFactoryType ParticleFactoryType;
-
-	/** Whether the vertex factory is in use. */
-	bool bInUse;
 };
 
 /**
