Fixed children inputs not having right click context menu actions for summary view

diff --git a/Engine/Plugins/FX/Niagara/Source/NiagaraEditor/Private/ViewModels/Stack/NiagaraStackFunctionInput.cpp b/Engine/Plugins/FX/Niagara/Source/NiagaraEditor/Private/ViewModels/Stack/NiagaraStackFunctionInput.cpp
--- a/Engine/Plugins/FX/Niagara/Source/NiagaraEditor/Private/ViewModels/Stack/NiagaraStackFunctionInput.cpp
+++ b/Engine/Plugins/FX/Niagara/Source/NiagaraEditor/Private/ViewModels/Stack/NiagaraStackFunctionInput.cpp
@@ -390,7 +390,7 @@ bool UNiagaraStackFunctionInput::SupportsSummaryView() const
 {
 	if(GetInputMetaData().IsSet())
 	{
-		if(GetInputMetaData()->bInlineEditConditionToggle || !GetInputMetaData()->ParentAttribute.IsNone())
+		if(GetInputMetaData()->bInlineEditConditionToggle)
 		{
 			return false;
 		}
