- Fix for static mesh DI not picking up the right SM when using child blueprints
- Add support for mesh parameter binding to SM DI

diff --git a/Engine/Plugins/FX/Niagara/Source/Niagara/Internal/DataInterface/NiagaraDataInterfaceStaticMesh.h b/Engine/Plugins/FX/Niagara/Source/Niagara/Internal/DataInterface/NiagaraDataInterfaceStaticMesh.h
--- a/Engine/Plugins/FX/Niagara/Source/Niagara/Internal/DataInterface/NiagaraDataInterfaceStaticMesh.h
+++ b/Engine/Plugins/FX/Niagara/Source/Niagara/Internal/DataInterface/NiagaraDataInterfaceStaticMesh.h
@@ -11,27 +11,25 @@ UENUM()
 enum class ENDIStaticMesh_SourceMode : uint8
 {
 	/**
-	Default behavior.
+	Default behavior follows the order of.
 	- Use "Source" when specified (either set explicitly or via blueprint with Set Niagara Static Mesh Component).
-	- When no source is specified, attempt to find a Static Mesh Component on an attached actor or component.
-	- If no source actor/component specified and no attached component found, fall back to the "Default Mesh" specified.
+	- Use Mesh Parameter Binding if valid
+	- Find Static Mesh Component, Attached Actor, Attached Component
+	- Falls back to 'Default Mesh' specified on the data interface
 	*/
 	Default,
 
-	/**
-	Only use "Source" (either set explicitly or via blueprint with Set Niagara Static Mesh Component).
-	*/
+	/**	Only use "Source" (either set explicitly or via blueprint with Set Niagara Static Mesh Component). */
 	Source,
 
-	/**
-	Only use the parent actor or component the system is attached to.
-	*/
+	/**	Only use the parent actor or component the system is attached to. */
 	AttachParent,
 
-	/**
-	Only use the "Default Mesh" specified.
-	*/
+	/** Only use the "Default Mesh" specified. */
 	DefaultMeshOnly,
+
+	/** Only use the mesh parameter binding. */
+	MeshParameterBinding,
 };
 
 USTRUCT()
@@ -60,7 +58,7 @@ public:
 	/** Controls how to retrieve the Static Mesh Component to attach to. */
 	UPROPERTY(EditAnywhere, Category = "Mesh")
 	ENDIStaticMesh_SourceMode SourceMode = ENDIStaticMesh_SourceMode::Default;
-	
+
 #if WITH_EDITORONLY_DATA
 	/** Mesh used to sample from when not overridden by a source actor from the scene. Only available in editor for previewing. This is removed in cooked builds. */
 	UPROPERTY(EditAnywhere, Category = "Mesh")
@@ -114,6 +112,10 @@ public:
 	UPROPERTY(EditAnywhere, Category = "LOD")
 	FNiagaraUserParameterBinding LODIndexUserParameter;
 
+	/** Mesh parameter binding can be either an Actor (in which case we find the component), static mesh component or a static mesh. */
+	UPROPERTY(EditAnywhere, Category = "Mesh")
+	FNiagaraUserParameterBinding MeshParameterBinding;
+
 	/** When attached to an Instanced Static Mesh, which instance should be read from. */
 	UPROPERTY(EditAnywhere, Category = "Mesh")
 	int32 InstanceIndex = INDEX_NONE;
@@ -185,7 +187,7 @@ public:
 #endif
 	//~ UNiagaraDataInterface interface
 
-	NIAGARA_API UStaticMesh* GetStaticMesh(USceneComponent*& OutComponent, class FNiagaraSystemInstance* SystemInstance = nullptr);
+	NIAGARA_API UStaticMesh* GetStaticMesh(USceneComponent*& OutComponent, class FNiagaraSystemInstance* SystemInstance = nullptr, UObject* ParameterBindingValue = nullptr);
 	NIAGARA_API void SetSourceComponentFromBlueprints(UStaticMeshComponent* ComponentToUse);
 	NIAGARA_API void SetDefaultMeshFromBlueprints(UStaticMesh* MeshToUse);
 
diff --git a/Engine/Plugins/FX/Niagara/Source/Niagara/Private/DataInterface/NiagaraDataInterfaceStaticMesh.cpp b/Engine/Plugins/FX/Niagara/Source/Niagara/Private/DataInterface/NiagaraDataInterfaceStaticMesh.cpp
--- a/Engine/Plugins/FX/Niagara/Source/Niagara/Private/DataInterface/NiagaraDataInterfaceStaticMesh.cpp
+++ b/Engine/Plugins/FX/Niagara/Source/Niagara/Private/DataInterface/NiagaraDataInterfaceStaticMesh.cpp
@@ -662,6 +662,8 @@ namespace NDIStaticMeshLocal
 
 		FNiagaraParameterDirectBinding<int32> LODIndexUserBinding;
 
+		FNiagaraParameterDirectBinding<UObject*> MeshParameterBinding;
+
 		/** Cached socket information, if available */
 		TArray<FTransform3f> CachedSockets;
 
@@ -854,9 +856,11 @@ namespace NDIStaticMeshLocal
 			NumFilteredSockets = 0;
 			FilteredAndUnfilteredSockets.Empty();
 
+			MeshParameterBinding.Init(SystemInstance->GetInstanceParameters(), Interface->MeshParameterBinding.Parameter);
+
 			// Get component / mesh we are using
 			USceneComponent* SceneComponent = nullptr;
-			UStaticMesh* StaticMesh = Interface->GetStaticMesh(SceneComponent, SystemInstance);
+			UStaticMesh* StaticMesh = Interface->GetStaticMesh(SceneComponent, SystemInstance, MeshParameterBinding.GetValue());
 			SceneComponentWeakPtr = SceneComponent;
 
 			// Gather attached information
@@ -1141,29 +1145,28 @@ namespace NDIStaticMeshLocal
 				return true;
 			}
 
+			// Check to see if either the mesh has become invalid
 			UStaticMesh* Mesh = StaticMeshWeakPtr.Get();
-			if (bMeshValid)
+			if (bMeshValid && !Mesh)
 			{
-				if (!Mesh)
-				{
-					// The static mesh we were bound to is no longer valid so we have to trigger a reset.
-					return true;
-				}
-				else if (UStaticMeshComponent* StaticMeshComp = Cast<UStaticMeshComponent>(Component))
-				{
-					if (Mesh != StaticMeshComp->GetStaticMesh())
-					{
-						// The mesh changed on the component we're attached to so we have to reset
-						return true;
-					}
-				}
+				return true;
 			}
 
+			// Has an external call invalidated the SM DI
 			if (Interface != nullptr && ChangeId != Interface->ChangeId)
 			{
 				return true;
 			}
 
+			// Check to see if we would resolve to a different mesh / component
+			USceneComponent* ResolvedComponent = nullptr;
+			UStaticMesh* ResolvedMesh = Interface->GetStaticMesh(ResolvedComponent, Instance, MeshParameterBinding.GetValue());
+			if (Component != ResolvedComponent || Mesh != ResolvedMesh)
+			{
+				return true;
+			}
+
+			// Has the LOD changed
 			if (Mesh != nullptr)
 			{
 				//TODO: If the LOD index selection changes then we reset the system
@@ -1669,6 +1672,7 @@ UNiagaraDataInterfaceStaticMesh::UNiagaraDataInterfaceStaticMesh(FObjectInitiali
 	Proxy.Reset(new NDIStaticMeshLocal::FRenderProxy());
 
 	LODIndexUserParameter.Parameter.SetType(FNiagaraTypeDefinition::GetIntDef());
+	MeshParameterBinding.Parameter.SetType(FNiagaraTypeDefinition(UObject::StaticClass()));
 }
 
 void UNiagaraDataInterfaceStaticMesh::PostInitProperties()
@@ -3266,6 +3270,7 @@ bool UNiagaraDataInterfaceStaticMesh::Equals(const UNiagaraDataInterface* Other)
 		OtherTyped->bAllowSamplingFromStreamingLODs == bAllowSamplingFromStreamingLODs &&
 		OtherTyped->LODIndex == LODIndex &&
 		OtherTyped->LODIndexUserParameter == LODIndexUserParameter &&
+		OtherTyped->MeshParameterBinding == MeshParameterBinding &&
 		OtherTyped->InstanceIndex == InstanceIndex &&
 		OtherTyped->FilteredSockets == FilteredSockets;
 }
@@ -3293,6 +3298,7 @@ bool UNiagaraDataInterfaceStaticMesh::CopyToInternal(UNiagaraDataInterface* Dest
 	OtherTyped->bAllowSamplingFromStreamingLODs = bAllowSamplingFromStreamingLODs;
 	OtherTyped->LODIndex = LODIndex;
 	OtherTyped->LODIndexUserParameter = LODIndexUserParameter;
+	OtherTyped->MeshParameterBinding = MeshParameterBinding;
 	OtherTyped->InstanceIndex = InstanceIndex;
 	OtherTyped->BindSourceDelegates();
 	return true;
@@ -3420,7 +3426,7 @@ void UNiagaraDataInterfaceStaticMesh::GetFeedback(UNiagaraSystem* Asset, UNiagar
 }
 #endif //WITH_EDITOR
 
-UStaticMesh* UNiagaraDataInterfaceStaticMesh::GetStaticMesh(USceneComponent*& OutComponent, class FNiagaraSystemInstance* SystemInstance)
+UStaticMesh* UNiagaraDataInterfaceStaticMesh::GetStaticMesh(USceneComponent*& OutComponent, class FNiagaraSystemInstance* SystemInstance, UObject* ParameterBindingValue)
 {
 	// Helper to scour an actor (or its parents) for a valid Static mesh component
 	auto FindActorMeshComponent = [](AActor* Actor, bool bRecurseParents = false) -> UStaticMeshComponent*
@@ -3459,73 +3465,101 @@ UStaticMesh* UNiagaraDataInterfaceStaticMesh::GetStaticMesh(USceneComponent*& Ou
 		return nullptr;
 	};
 
-	UStaticMeshComponent* FoundMeshComponent = nullptr;	
-
 	const bool bTrySource = SourceMode == ENDIStaticMesh_SourceMode::Default || SourceMode == ENDIStaticMesh_SourceMode::Source;
+	const bool bTryMeshParameter = SourceMode == ENDIStaticMesh_SourceMode::Default || SourceMode == ENDIStaticMesh_SourceMode::MeshParameterBinding;
 	const bool bTryAttachParent = SourceMode == ENDIStaticMesh_SourceMode::Default || SourceMode == ENDIStaticMesh_SourceMode::AttachParent;
 	const bool bTryDefaultMesh = SourceMode == ENDIStaticMesh_SourceMode::Default || SourceMode == ENDIStaticMesh_SourceMode::DefaultMeshOnly;
 
+	UStaticMesh* OutMesh = nullptr;
+	UStaticMeshComponent* OutStaticMeshComponent = nullptr;
 	if (bTrySource && ::IsValid(SourceComponent))
 	{
-		FoundMeshComponent = SourceComponent;
+		OutStaticMeshComponent = SourceComponent;
+		OutMesh = OutStaticMeshComponent->GetStaticMesh();
 	}
 	else if (bTrySource && SoftSourceActor.Get())
 	{
-		FoundMeshComponent = FindActorMeshComponent(SoftSourceActor.Get());
+		OutStaticMeshComponent = FindActorMeshComponent(SoftSourceActor.Get());
+		OutMesh = OutStaticMeshComponent->GetStaticMesh();
 	}
-	else if (bTryAttachParent && SystemInstance)
+	else
 	{
-		if (USceneComponent* AttachComponent = SystemInstance->GetAttachComponent())
+		// Try parameter binding, this could either be a component or a static mesh
+		if (bTryMeshParameter && ParameterBindingValue)
 		{
-			// First, try to find the mesh component up the attachment hierarchy
-			for (USceneComponent* Curr = AttachComponent; Curr; Curr = Curr->GetAttachParent())
+			if (AActor* Actor = Cast<AActor>(ParameterBindingValue))
+			{
+				OutStaticMeshComponent = FindActorMeshComponent(Actor, true);
+				OutMesh = OutStaticMeshComponent->GetStaticMesh();
+			}
+			else
 			{
-				UStaticMeshComponent* ParentComp = Cast<UStaticMeshComponent>(Curr);
-				if (::IsValid(ParentComp))
+				OutStaticMeshComponent = Cast<UStaticMeshComponent>(ParameterBindingValue);
+				if (OutStaticMeshComponent)
 				{
-					FoundMeshComponent = ParentComp;
-					break;
+					OutMesh = OutStaticMeshComponent->GetStaticMesh();
+				}
+				else
+				{
+					OutMesh = Cast<UStaticMesh>(ParameterBindingValue);
 				}
 			}
-			
-			if (!FoundMeshComponent)
+		}
+
+		// Try attached parent, only if we did not already find a valid component / mesh
+		if (bTryAttachParent && SystemInstance && !OutStaticMeshComponent && !OutMesh)
+		{
+			if (USceneComponent* AttachComponent = SystemInstance->GetAttachComponent())
 			{
-				// Next, try to find one in our outer chain
-				UStaticMeshComponent* OuterComp = AttachComponent->GetTypedOuter<UStaticMeshComponent>();
-				if (::IsValid(OuterComp))
+				// First, try to find the mesh component up the attachment hierarchy
+				for (USceneComponent* Curr = AttachComponent; Curr; Curr = Curr->GetAttachParent())
+				{
+					UStaticMeshComponent* ParentComp = Cast<UStaticMeshComponent>(Curr);
+					if (::IsValid(ParentComp))
+					{
+						OutStaticMeshComponent = ParentComp;
+						break;
+					}
+				}
+
+				if (!OutStaticMeshComponent)
 				{
-					FoundMeshComponent = OuterComp;
+					// Next, try to find one in our outer chain
+					UStaticMeshComponent* OuterComp = AttachComponent->GetTypedOuter<UStaticMeshComponent>();
+					if (::IsValid(OuterComp))
+					{
+						OutStaticMeshComponent = OuterComp;
+					}
+					else if (AActor* Actor = AttachComponent->GetAttachmentRootActor())
+					{
+						// Final fall-back, look for any mesh component on our root actor or any of its parents
+						OutStaticMeshComponent = FindActorMeshComponent(Actor, true);
+					}
 				}
-				else if (AActor* Actor = AttachComponent->GetAttachmentRootActor())
+
+				if (OutStaticMeshComponent)
 				{
-					// Final fall-back, look for any mesh component on our root actor or any of its parents
-					FoundMeshComponent = FindActorMeshComponent(Actor, true);
+					OutMesh = OutStaticMeshComponent->GetStaticMesh();
 				}
 			}
 		}
 	}
 
-	UStaticMesh* Mesh = nullptr;
-	OutComponent = nullptr;
-	if (FoundMeshComponent)
-	{
-		Mesh = FoundMeshComponent->GetStaticMesh();
-		OutComponent = FoundMeshComponent;
-	}
-	else if (bTryDefaultMesh)
+	OutComponent = OutStaticMeshComponent;
+	if (bTryDefaultMesh && !OutComponent && !OutMesh)
 	{
-		Mesh = DefaultMesh;
+		OutMesh = DefaultMesh;
 	}
 
 #if WITH_EDITORONLY_DATA
-	if (!Mesh && !FoundMeshComponent && (!SystemInstance || !SystemInstance->GetWorld()->IsGameWorld()))
+	if (!OutMesh && !OutComponent && (!SystemInstance || !SystemInstance->GetWorld()->IsGameWorld()))
 	{
 		// NOTE: We don't fall back on the preview mesh if we have a valid static mesh component referenced
-		Mesh = PreviewMesh.LoadSynchronous();		
+		OutMesh = PreviewMesh.LoadSynchronous();		
 	}
 #endif
 
-	return Mesh;
+	return OutMesh;
 }
 
 void UNiagaraDataInterfaceStaticMesh::SetSourceComponentFromBlueprints(UStaticMeshComponent* ComponentToUse)
