Fixing issue with NDC writer and invalid Data.
Adding double type to niagara type helpers.

diff --git a/Engine/Plugins/FX/Niagara/Source/Niagara/Private/NiagaraDataChannelAccessor.cpp b/Engine/Plugins/FX/Niagara/Source/Niagara/Private/NiagaraDataChannelAccessor.cpp
--- a/Engine/Plugins/FX/Niagara/Source/Niagara/Private/NiagaraDataChannelAccessor.cpp
+++ b/Engine/Plugins/FX/Niagara/Source/Niagara/Private/NiagaraDataChannelAccessor.cpp
@@ -281,7 +281,7 @@ void FNiagaraDataChannelGameDataWriterBase::BeginWrite()
 
 void FNiagaraDataChannelGameDataWriterBase::Publish(FNiagaraDataChannelDataPtr& Destination)
 {
-	if (Data->Num() > 0)
+	if (Data && Data->Num() > 0)
 	{
 		FNiagaraDataChannelPublishRequest PublishRequest;
 		PublishRequest.bVisibleToCPUSims = true;
diff --git a/Engine/Plugins/FX/Niagara/Source/Niagara/Private/NiagaraTypes.cpp b/Engine/Plugins/FX/Niagara/Source/Niagara/Private/NiagaraTypes.cpp
--- a/Engine/Plugins/FX/Niagara/Source/Niagara/Private/NiagaraTypes.cpp
+++ b/Engine/Plugins/FX/Niagara/Source/Niagara/Private/NiagaraTypes.cpp
@@ -307,6 +307,7 @@ FNiagaraTypeDefinition FNiagaraTypeHelper::Vector2DDef;
 FNiagaraTypeDefinition FNiagaraTypeHelper::VectorDef;
 FNiagaraTypeDefinition FNiagaraTypeHelper::Vector4Def;
 FNiagaraTypeDefinition FNiagaraTypeHelper::QuatDef;
+FNiagaraTypeDefinition FNiagaraTypeHelper::DoubleDef;
 
 FRWLock FNiagaraTypeHelper::RemapTableLock;
 TMap<TWeakObjectPtr<UScriptStruct>, FNiagaraTypeHelper::FRemapEntry> FNiagaraTypeHelper::RemapTable;
@@ -684,7 +685,7 @@ FNiagaraTypeDefinition FNiagaraTypeHelper::GetLWCType(const FNiagaraTypeDefiniti
 
 	if(InType == FNiagaraTypeDefinition::GetFloatDef())
 	{
-		return FNiagaraTypeDefinition(FNiagaraDouble::StaticStruct(), FNiagaraTypeDefinition::EAllowUnfriendlyStruct::Allow);
+		return DoubleDef;
 	}
 	else if (InType == FNiagaraTypeDefinition::GetVec2Def())
 	{
@@ -725,6 +726,7 @@ void FNiagaraTypeHelper::InitStaticTypes()
 		VectorDef = FNiagaraTypeDefinition(FindObjectChecked<UScriptStruct>(CoreUObjectPkg, TEXT("Vector")), FNiagaraTypeDefinition::EAllowUnfriendlyStruct::Allow);
 		Vector4Def = FNiagaraTypeDefinition(FindObjectChecked<UScriptStruct>(CoreUObjectPkg, TEXT("Vector4")), FNiagaraTypeDefinition::EAllowUnfriendlyStruct::Allow);
 		QuatDef = FNiagaraTypeDefinition(FindObjectChecked<UScriptStruct>(CoreUObjectPkg, TEXT("Quat")), FNiagaraTypeDefinition::EAllowUnfriendlyStruct::Allow);
+		DoubleDef = FNiagaraTypeDefinition(FNiagaraDouble::StaticStruct(), FNiagaraTypeDefinition::EAllowUnfriendlyStruct::Allow);
 	}
 }
 
@@ -748,6 +750,11 @@ FNiagaraTypeDefinition FNiagaraTypeHelper::GetQuatDef()
 	return QuatDef;
 }
 
+FNiagaraTypeDefinition FNiagaraTypeHelper::GetDoubleDef()
+{
+	return DoubleDef;
+}
+
 UScriptStruct* FNiagaraTypeHelper::FindNiagaraFriendlyTopLevelStruct(UScriptStruct* InStruct, ENiagaraStructConversion StructConversion)
 {
 	if (!InStruct)
diff --git a/Engine/Plugins/FX/Niagara/Source/Niagara/Public/NiagaraTypes.h b/Engine/Plugins/FX/Niagara/Source/Niagara/Public/NiagaraTypes.h
--- a/Engine/Plugins/FX/Niagara/Source/Niagara/Public/NiagaraTypes.h
+++ b/Engine/Plugins/FX/Niagara/Source/Niagara/Public/NiagaraTypes.h
@@ -429,6 +429,7 @@ public:
 	static NIAGARA_API FNiagaraTypeDefinition GetVectorDef();
 	static NIAGARA_API FNiagaraTypeDefinition GetVector4Def();
 	static NIAGARA_API FNiagaraTypeDefinition GetQuatDef();
+	static NIAGARA_API FNiagaraTypeDefinition GetDoubleDef();
 
 	static NIAGARA_API void InitStaticTypes();
 private:
@@ -458,6 +459,7 @@ private:
 	static FNiagaraTypeDefinition VectorDef;
 	static FNiagaraTypeDefinition Vector4Def;
 	static FNiagaraTypeDefinition QuatDef;
+	static FNiagaraTypeDefinition DoubleDef;
 };
 
 /** Information about how this type should be laid out in an FNiagaraDataSet */
