- Don't allow precache if it's not for our shader platform

diff --git a/Engine/Plugins/FX/Niagara/Source/NiagaraShader/Private/NiagaraShared.cpp b/Engine/Plugins/FX/Niagara/Source/NiagaraShader/Private/NiagaraShared.cpp
--- a/Engine/Plugins/FX/Niagara/Source/NiagaraShader/Private/NiagaraShared.cpp
+++ b/Engine/Plugins/FX/Niagara/Source/NiagaraShader/Private/NiagaraShared.cpp
@@ -540,12 +540,16 @@ void FNiagaraShaderScript::UpdateCachedData_PostCompile(bool bCalledFromSerializ
 			CachedData.bViewUniformBufferUsed |= NiagaraShader->bNeedsViewUniformBuffer;
 
 			// request precache the compute shader
-			if (IsResourcePSOPrecachingEnabled() || IsComponentPSOPrecachingEnabled())
+			// Note: this function can be called for different shader platforms so only precache if it's for the platform we are running
+			if (GMaxRHIShaderPlatform == GameThreadShaderMap->GetShaderPlatform())
 			{
-				check(NiagaraShader->GetFrequency() == SF_Compute)
-				FRHIShader* RHIShader = GameThreadShaderMap->GetResource()->GetShader(Shader->GetResourceIndex());
-				FRHIComputeShader* RHIComputeShader = static_cast<FRHIComputeShader*>(RHIShader);
-				PipelineStateCache::PrecacheComputePipelineState(RHIComputeShader);
+				if (IsResourcePSOPrecachingEnabled() || IsComponentPSOPrecachingEnabled())
+				{
+					check(NiagaraShader->GetFrequency() == SF_Compute);
+					FRHIShader* RHIShader = GameThreadShaderMap->GetResource()->GetShader(Shader->GetResourceIndex());
+					FRHIComputeShader* RHIComputeShader = static_cast<FRHIComputeShader*>(RHIShader);
+					PipelineStateCache::PrecacheComputePipelineState(RHIComputeShader);
+				}
 			}
 		}
 	}
