- Add ability to restrict catch up ticking in the baker
- For example, if you run an expensive fluid simulation that ticks [at] 60hz but the sim runs [REDACTED_HANDLE] we won't start to play catch up an instead only send a single tick per frame

diff --git a/Engine/Plugins/FX/Niagara/Source/Niagara/Classes/NiagaraBakerSettings.h b/Engine/Plugins/FX/Niagara/Source/Niagara/Classes/NiagaraBakerSettings.h
--- a/Engine/Plugins/FX/Niagara/Source/Niagara/Classes/NiagaraBakerSettings.h
+++ b/Engine/Plugins/FX/Niagara/Source/Niagara/Classes/NiagaraBakerSettings.h
@@ -167,6 +167,10 @@ public:
 	UPROPERTY(EditAnywhere, Category = "Settings", AdvancedDisplay, meta = (ClampMin=1, ClampMax=480))
 	int FramesPerSecond = 60;
 
+	/** Locks the playback to the simulation frame rate, i.e. no multi-tick. */
+	UPROPERTY(EditAnywhere, Category = "Settings")
+	uint8 bLockToSimulationFrameRate : 1;
+
 	/** Should the preview playback as looping or not. */
 	UPROPERTY(EditAnywhere, Category = "Settings")
 	uint8 bPreviewLooping : 1;
diff --git a/Engine/Plugins/FX/Niagara/Source/Niagara/Private/NiagaraBakerSettings.cpp b/Engine/Plugins/FX/Niagara/Source/Niagara/Private/NiagaraBakerSettings.cpp
--- a/Engine/Plugins/FX/Niagara/Source/Niagara/Private/NiagaraBakerSettings.cpp
+++ b/Engine/Plugins/FX/Niagara/Source/Niagara/Private/NiagaraBakerSettings.cpp
@@ -16,6 +16,7 @@
 UNiagaraBakerSettings::UNiagaraBakerSettings(const FObjectInitializer& Init)
 	: Super(Init)
 {
+	bLockToSimulationFrameRate = false;
 	bPreviewLooping = true;
 	bRenderComponentOnly = true;
 	Outputs.Add(Init.CreateDefaultSubobject<UNiagaraBakerOutputTexture2D>(this, "DefaultOutput"));
diff --git a/Engine/Plugins/FX/Niagara/Source/NiagaraEditor/Private/ViewModels/NiagaraBakerViewModel.cpp b/Engine/Plugins/FX/Niagara/Source/NiagaraEditor/Private/ViewModels/NiagaraBakerViewModel.cpp
--- a/Engine/Plugins/FX/Niagara/Source/NiagaraEditor/Private/ViewModels/NiagaraBakerViewModel.cpp
+++ b/Engine/Plugins/FX/Niagara/Source/NiagaraEditor/Private/ViewModels/NiagaraBakerViewModel.cpp
@@ -389,6 +389,22 @@ void FNiagaraBakerViewModel::SetSimTickRate(int TickRate)
 	}
 }
 
+bool FNiagaraBakerViewModel::LockToSimulationFrameRate() const
+{
+	UNiagaraBakerSettings* BakerSettings = GetBakerSettings();
+	return BakerSettings ? BakerSettings->bLockToSimulationFrameRate : true;
+}
+
+void FNiagaraBakerViewModel::ToggleLockToSimulationFrameRate()
+{
+	if (UNiagaraBakerSettings* BakerSettings = GetBakerSettings())
+	{
+		const FScopedTransaction Transaction(LOCTEXT("SetAllowMultiTick", "SetAllowMultiTick"));
+		BakerSettings->Modify();
+		BakerSettings->bLockToSimulationFrameRate = BakerSettings->bLockToSimulationFrameRate ? 0 : 1;
+	}
+}
+
 void FNiagaraBakerViewModel::AddOutput(UClass* Class)
 {
 	if (UNiagaraBakerSettings* BakerSettings = GetBakerSettings())
diff --git a/Engine/Plugins/FX/Niagara/Source/NiagaraEditor/Private/ViewModels/NiagaraBakerViewModel.h b/Engine/Plugins/FX/Niagara/Source/NiagaraEditor/Private/ViewModels/NiagaraBakerViewModel.h
--- a/Engine/Plugins/FX/Niagara/Source/NiagaraEditor/Private/ViewModels/NiagaraBakerViewModel.h
+++ b/Engine/Plugins/FX/Niagara/Source/NiagaraEditor/Private/ViewModels/NiagaraBakerViewModel.h
@@ -107,6 +107,9 @@ public:
 	int GetSimTickRate() const;
 	void SetSimTickRate(int TickRate);
 
+	bool LockToSimulationFrameRate() const;
+	void ToggleLockToSimulationFrameRate();
+
 	void AddOutput(UClass* Class);
 	void RemoveCurrentOutput();
 	bool CanRemoveCurrentOutput() const;
diff --git a/Engine/Plugins/FX/Niagara/Source/NiagaraEditor/Private/Widgets/SNiagaraBakerWidget.cpp b/Engine/Plugins/FX/Niagara/Source/NiagaraEditor/Private/Widgets/SNiagaraBakerWidget.cpp
--- a/Engine/Plugins/FX/Niagara/Source/NiagaraEditor/Private/Widgets/SNiagaraBakerWidget.cpp
+++ b/Engine/Plugins/FX/Niagara/Source/NiagaraEditor/Private/Widgets/SNiagaraBakerWidget.cpp
@@ -864,6 +864,19 @@ void SNiagaraBakerWidget::MakeSimTickRateMenu(FMenuBuilder& MenuBuilder) const
 		LOCTEXT("SimTickRate", "Simulation Tick Rate")
 	);
 
+	MenuBuilder.AddMenuEntry(
+		LOCTEXT("LockToSimulationFrameRate", "Lock To Simulation Frame Rate"),
+		LOCTEXT("LockToSimulationFrameRateTooltip", "When enabled we will update at the simulation frame rate.  I.e. only ever tick once per frame even if the application frame rate is lower, this will result in the sim slowing down."),
+		FSlateIcon(),
+		FUIAction(
+			FExecuteAction::CreateSP(ViewModel, &FNiagaraBakerViewModel::ToggleLockToSimulationFrameRate),
+			FCanExecuteAction(),
+			FIsActionChecked::CreateSP(ViewModel, &FNiagaraBakerViewModel::LockToSimulationFrameRate)
+		),
+		NAME_None,
+		EUserInterfaceActionType::ToggleButton
+	);
+
 	MenuBuilder.BeginSection("FPS", LOCTEXT("FPS", "FPS"));
 	{
 		static const int32 DefaultFPS[] = { 240, 120, 60, 50, 30, 20, 15 };
@@ -976,6 +989,12 @@ void SNiagaraBakerWidget::Tick(const FGeometry& AllottedGeometry, const double _
 			DeltaTime *= ViewModel->GetPlaybackRate();
 		}
 
+		if (BakerSettings->bLockToSimulationFrameRate)
+		{
+			const float MaxTickRate = 1.0f / BakerSettings->FramesPerSecond;
+			DeltaTime = FMath::Min(MaxTickRate, DeltaTime);
+		}
+
 		if (PlaybackMode != EPlaybackMode::Stopped)
 		{
 			PreviewRelativeTime += PlaybackMode == EPlaybackMode::PlayingForward ? DeltaTime : -DeltaTime;
