Fixed circular refreshes between stack refresh <-> details panel refresh.
The details panel has updated to an update with 1 frame delay, which means the value guard triggers before the refresh is done.
By removing the value guard, and instead resetting the bool manually, we keep it's state until the operation is done.

diff --git a/Engine/Plugins/FX/Niagara/Source/NiagaraEditor/Private/ViewModels/Stack/NiagaraStackObject.cpp b/Engine/Plugins/FX/Niagara/Source/NiagaraEditor/Private/ViewModels/Stack/NiagaraStackObject.cpp
--- a/Engine/Plugins/FX/Niagara/Source/NiagaraEditor/Private/ViewModels/Stack/NiagaraStackObject.cpp
+++ b/Engine/Plugins/FX/Niagara/Source/NiagaraEditor/Private/ViewModels/Stack/NiagaraStackObject.cpp
@@ -32,7 +32,7 @@ void UNiagaraStackObject::Initialize(FRequiredEntryData InRequiredEntryData, UOb
 	WeakObject = InObject;
 	bIsTopLevelObject = bInIsTopLevelObject;
 	OwningNiagaraNode = InOwningNiagaraNode;
-	bIsRefresingDataInterfaceErrors = false;
+	bIsRefreshingDataInterfaceErrors = false;
 
 	MessageLogGuid = GetSystemViewModel()->GetMessageLogGuid();
 
@@ -241,7 +241,7 @@ void UNiagaraStackObject::RefreshChildrenInternal(const TArray<UNiagaraStackEntr
 		// First we need to refresh the errors on the data interface so that the rows in the property row generator 
 		// are correct.
 		{
-			TGuardValue<bool> RefreGuard(bIsRefresingDataInterfaceErrors, true);
+			bIsRefreshingDataInterfaceErrors = true;
 			DataInterfaceObject->RefreshErrors();
 		}
 
@@ -362,10 +362,12 @@ void UNiagaraStackObject::PostRefreshChildrenInternal()
 
 void UNiagaraStackObject::PropertyRowsRefreshed()
 {
-	if(bIsRefresingDataInterfaceErrors == false)
+	if(bIsRefreshingDataInterfaceErrors == false)
 	{
 		RefreshChildren();
 	}
+	
+	bIsRefreshingDataInterfaceErrors = false;
 }
 
 void UNiagaraStackObject::OnMessageManagerRefresh(const TArray<TSharedRef<const INiagaraMessage>>& NewMessages)
@@ -386,4 +388,3 @@ void UNiagaraStackObject::OnMessageManagerRefresh(const TArray<TSharedRef<const
 		RefreshChildren();
 	}
 }
-
diff --git a/Engine/Plugins/FX/Niagara/Source/NiagaraEditor/Public/ViewModels/Stack/NiagaraStackObject.h b/Engine/Plugins/FX/Niagara/Source/NiagaraEditor/Public/ViewModels/Stack/NiagaraStackObject.h
--- a/Engine/Plugins/FX/Niagara/Source/NiagaraEditor/Public/ViewModels/Stack/NiagaraStackObject.h
+++ b/Engine/Plugins/FX/Niagara/Source/NiagaraEditor/Public/ViewModels/Stack/NiagaraStackObject.h
@@ -83,7 +83,7 @@ private:
 	TArray<FRegisteredClassCustomization> RegisteredClassCustomizations;
 	TArray<FRegisteredPropertyCustomization> RegisteredPropertyCustomizations;
 	TSharedPtr<IPropertyRowGenerator> PropertyRowGenerator;
-	bool bIsRefresingDataInterfaceErrors;
+	bool bIsRefreshingDataInterfaceErrors;
 
 	/** An optional object guid that can be provided to identify the object this stack entry represents. This can be used for summary view purposes and is also given to the child property rows. */
 	TOptional<FGuid> ObjectGuid;
