- Don't want about missing pool if the niagara component was in the used list
- This is because edits to the system will cause a recompile + clear the component pool

diff --git a/Engine/Plugins/FX/Niagara/Source/Niagara/Private/NiagaraComponentPool.cpp b/Engine/Plugins/FX/Niagara/Source/Niagara/Private/NiagaraComponentPool.cpp
--- a/Engine/Plugins/FX/Niagara/Source/Niagara/Private/NiagaraComponentPool.cpp
+++ b/Engine/Plugins/FX/Niagara/Source/Niagara/Private/NiagaraComponentPool.cpp
@@ -461,16 +461,8 @@ void UNiagaraComponentPool::ReclaimWorldParticleSystem(UNiagaraComponent* Compon
 				Pair.Value.KillUnusedComponents(CurrentTime - GNiagaraSystemPoolKillUnusedTime, Asset);
 			}
 		}
-		
-		FNCPool* NCPool = WorldParticleSystemPools.Find(Asset);
-		if (!NCPool)
-		{
-			UE_LOG(LogNiagara, Warning, TEXT("WorldNC Pool trying to reclaim a system for which it doesn't have a pool! Likely because SetAsset() has been called on this NC. | World: %p | NC: %p | Sys: %s"), Component->GetWorld(), Component, *Component->GetAsset()->GetFullName());
-			//Just add the new pool and reclaim to that one.
-			NCPool = &WorldParticleSystemPools.Add(Asset);
-		}
 
-	#if ENABLE_NC_POOL_DEBUGGING
+#if ENABLE_NC_POOL_DEBUGGING
 		bool bWasInList = false;
 		if (Component->PoolingMethod == ENCPoolMethod::AutoRelease)
 		{
@@ -485,7 +477,24 @@ void UNiagaraComponentPool::ReclaimWorldParticleSystem(UNiagaraComponent* Compon
 		{
 			UE_LOG(LogNiagara, Error, TEXT("World Niagara System Pool is reclaiming a component that is not in it's InUse list!"));
 		}
-	#endif
+#else
+		constexpr bool bWasInList = false;
+#endif
+
+		FNCPool* NCPool = WorldParticleSystemPools.Find(Asset);
+		if (!NCPool)
+		{
+			// We only need to warn if it was in the in use list
+			// When in the editor update context's can clear out the existing pools so this will fire incorrectly.
+			if (!bWasInList)
+			{
+				UE_LOG(LogNiagara, Warning, TEXT("WorldNC Pool trying to reclaim a system for which it doesn't have a pool! Likely because SetAsset() has been called on this NC. | World: %p | NC: %p | Sys: %s"), Component->GetWorld(), Component, *Component->GetAsset()->GetFullName());
+			}
+
+			//Just add the new pool and reclaim to that one.
+			NCPool = &WorldParticleSystemPools.Add(Asset);
+		}
+
 		NCPool->Reclaim(Component, CurrentTime);
 	}
 	else
