Collection of EVVM fixes
-Cleans up the checks of the constant buffers at runtime
-Removes the search for conflicts that was exclusive to exec_index & i2f and makes sure merging exec_index only considers temp registers
-Adds an EVVM assembly string representation in the Generated Code view
-fixes random_addi implementation that was sharing the buffer increment value between the random seed and the addition operand
-fixes a OOB access/stomp when dealing with newly identified unused ext function output operands

diff --git a/Engine/Plugins/FX/Niagara/Source/Niagara/Classes/NiagaraScript.h b/Engine/Plugins/FX/Niagara/Source/Niagara/Classes/NiagaraScript.h
--- a/Engine/Plugins/FX/Niagara/Source/Niagara/Classes/NiagaraScript.h
+++ b/Engine/Plugins/FX/Niagara/Source/Niagara/Classes/NiagaraScript.h
@@ -562,6 +562,11 @@ public:
 	UPROPERTY()
 	TArray<uint8> ExperimentalContextData;
 
+#if WITH_EDITORONLY_DATA
+	UPROPERTY()
+	FString LastExperimentalAssemblyScript;
+#endif
+
 	UPROPERTY()
 	uint32 bReadsSignificanceIndex : 1;
 
diff --git a/Engine/Plugins/FX/Niagara/Source/Niagara/Private/NiagaraAsyncCompile.cpp b/Engine/Plugins/FX/Niagara/Source/Niagara/Private/NiagaraAsyncCompile.cpp
--- a/Engine/Plugins/FX/Niagara/Source/Niagara/Private/NiagaraAsyncCompile.cpp
+++ b/Engine/Plugins/FX/Niagara/Source/Niagara/Private/NiagaraAsyncCompile.cpp
@@ -368,7 +368,10 @@ void FNiagaraAsyncCompileTask::OptimizeByteCode()
 		FVectorVMOptimizeContext OptimizeContext;
 		FMemory::Memzero(&OptimizeContext, sizeof(OptimizeContext));
 		uint64 AssetPathHash = CityHash64((char *)AssetPath.GetCharArray().GetData(), AssetPath.GetCharArray().Num());
-		OptimizeVectorVMScript(ByteCode, ExeData->ByteCode.GetLength(), ExtFnTable.GetData(), ExtFnTable.Num(), &OptimizeContext, AssetPathHash, VVMFlag_OptOmitStats);
+		OptimizeVectorVMScript(ByteCode, ExeData->ByteCode.GetLength(), ExtFnTable.GetData(), ExtFnTable.Num(), &OptimizeContext, AssetPathHash, VVMFlag_OptOmitStats | VVMFlag_OptSaveIntermediateState);
+
+		// extract a human readable version of the script
+		GenerateHumanReadableVectorVMScript(OptimizeContext, ExeData->LastExperimentalAssemblyScript);
 
 		// freeze the OptimizeContext
 		FreezeVectorVMOptimizeContext(OptimizeContext, ExeData->ExperimentalContextData);
diff --git a/Engine/Plugins/FX/Niagara/Source/NiagaraEditor/Private/NiagaraCompilationTasks.cpp b/Engine/Plugins/FX/Niagara/Source/NiagaraEditor/Private/NiagaraCompilationTasks.cpp
--- a/Engine/Plugins/FX/Niagara/Source/NiagaraEditor/Private/NiagaraCompilationTasks.cpp
+++ b/Engine/Plugins/FX/Niagara/Source/NiagaraEditor/Private/NiagaraCompilationTasks.cpp
@@ -702,7 +702,10 @@ void FNiagaraSystemCompilationTask::FCompileTaskInfo::ProcessCompilation(FNiagar
 
 		FVectorVMOptimizeContext OptimizeContext = { };
 		uint64 AssetPathHash = CityHash64((char *)AssetPath.GetCharArray().GetData(), AssetPath.GetCharArray().Num());
-		OptimizeVectorVMScript(ByteCode, ExeData->ByteCode.GetLength(), ExtFnTable.GetData(), ExtFnTable.Num(), &OptimizeContext, AssetPathHash, VVMFlag_OptOmitStats);
+		OptimizeVectorVMScript(ByteCode, ExeData->ByteCode.GetLength(), ExtFnTable.GetData(), ExtFnTable.Num(), &OptimizeContext, AssetPathHash, VVMFlag_OptOmitStats | VVMFlag_OptSaveIntermediateState);
+
+		// extract a human readable version of the script
+		GenerateHumanReadableVectorVMScript(OptimizeContext, ExeData->LastExperimentalAssemblyScript);
 
 		// freeze the OptimizeContext
 		FreezeVectorVMOptimizeContext(OptimizeContext, ExeData->ExperimentalContextData);
diff --git a/Engine/Plugins/FX/Niagara/Source/NiagaraEditor/Private/Widgets/SNiagaraGeneratedCodeView.cpp b/Engine/Plugins/FX/Niagara/Source/NiagaraEditor/Private/Widgets/SNiagaraGeneratedCodeView.cpp
--- a/Engine/Plugins/FX/Niagara/Source/NiagaraEditor/Private/Widgets/SNiagaraGeneratedCodeView.cpp
+++ b/Engine/Plugins/FX/Niagara/Source/NiagaraEditor/Private/Widgets/SNiagaraGeneratedCodeView.cpp
@@ -1,28 +1,29 @@
 // Copyright Epic Games, Inc. All Rights Reserved.
 
 #include "SNiagaraGeneratedCodeView.h"
-#include "Textures/SlateIcon.h"
 #include "Framework/Commands/UIAction.h"
 #include "Framework/MultiBox/MultiBoxBuilder.h"
-#include "Widgets/Input/SSearchBox.h"
-#include "Widgets/Input/SButton.h"
+#include "HAL/PlatformApplicationMisc.h"
 #include "ISequencer.h"
-#include "ViewModels/NiagaraSystemViewModel.h"
-#include "ViewModels/NiagaraEmitterHandleViewModel.h"
-#include "ViewModels/NiagaraSystemSelectionViewModel.h"
-#include "NiagaraEmitterHandle.h"
+#include "NiagaraEditorStyle.h"
+#include "NiagaraEditorUtilities.h"
 #include "NiagaraEmitter.h"
+#include "NiagaraEmitterHandle.h"
 #include "NiagaraScript.h"
-#include "ViewModels/NiagaraSystemScriptViewModel.h"
+#include "NiagaraSettings.h"
+#include "NiagaraSystem.h"
 #include "Styling/AppStyle.h"
-#include "Widgets/Layout/SScrollBox.h"
+#include "Textures/SlateIcon.h"
 #include "UObject/Class.h"
-#include "NiagaraSystem.h"
-#include "NiagaraEditorStyle.h"
-#include "HAL/PlatformApplicationMisc.h"
-#include "NiagaraEditorUtilities.h"
-#include "Widgets/NiagaraHLSLSyntaxHighlighter.h"
+#include "ViewModels/NiagaraEmitterHandleViewModel.h"
+#include "ViewModels/NiagaraSystemScriptViewModel.h"
+#include "ViewModels/NiagaraSystemSelectionViewModel.h"
+#include "ViewModels/NiagaraSystemViewModel.h"
 #include "Widgets/Docking/SDockTab.h"
+#include "Widgets/Input/SButton.h"
+#include "Widgets/Input/SSearchBox.h"
+#include "Widgets/Layout/SScrollBox.h"
+#include "Widgets/NiagaraHLSLSyntaxHighlighter.h"
 
 #define LOCTEXT_NAMESPACE "NiagaraGeneratedCodeView"
 
@@ -352,12 +353,20 @@ void SNiagaraGeneratedCodeView::SystemSelectionChanged()
 
 void SNiagaraGeneratedCodeView::UpdateUI()
 {
+	enum EScriptDisplayType
+	{
+		HLSL_SCRIPT_TYPE,
+		GPU_SCRIPT_TYPE,
+		ASSEMBLY_SCRIPT_TYPE,
+		EVVM_ASSEMBLY_SCRIPT_TYPE,
+	};
+
 	//Mark the UI as needing an update. We now do the very expensive slate updates from Tick() iside UIUpdate_Internal.
 	//This way we can avoid refreshing this UI constantly when this tab isn't even in view.
 	bUIUpdatePending = true;
 
 	TArray<UNiagaraScript*> Scripts;
-	TArray<uint32> ScriptDisplayTypes;
+	TArray<EScriptDisplayType> ScriptDisplayTypes;
 	UNiagaraSystem& System = SystemViewModel->GetSystem();
 	Scripts.Add(System.GetSystemSpawnScript());
 	Scripts.Add(System.GetSystemUpdateScript());
@@ -375,6 +384,8 @@ void SNiagaraGeneratedCodeView::UpdateUI()
 		}
 	}
 
+	const bool bIncludeEVVM = GetDefault<UNiagaraSettings>()->bExperimentalVMEnabled;
+
 	// Mark the scripts with the correct display type and copy references for the non-gpu scripts for the assembly view.
 	int32 OriginalScriptCount = Scripts.Num();
 	ScriptDisplayTypes.AddUninitialized(OriginalScriptCount);
@@ -383,14 +394,20 @@ void SNiagaraGeneratedCodeView::UpdateUI()
 		UNiagaraScript* Script = Scripts[i];
 		if (Script->GetUsage() == ENiagaraScriptUsage::ParticleGPUComputeScript)
 		{
-			ScriptDisplayTypes[i] = 1;
+			ScriptDisplayTypes[i] = GPU_SCRIPT_TYPE;
 		}
 		else
 		{
-			ScriptDisplayTypes[i] = 0;
+			ScriptDisplayTypes[i] = HLSL_SCRIPT_TYPE;
 
 			Scripts.Add(Script);
-			ScriptDisplayTypes.Add(2);
+			ScriptDisplayTypes.Add(ASSEMBLY_SCRIPT_TYPE);
+
+			if (bIncludeEVVM)
+			{
+				Scripts.Add(Script);
+				ScriptDisplayTypes.Add(EVVM_ASSEMBLY_SCRIPT_TYPE);
+			}
 		}
 	}
 		
@@ -405,89 +422,97 @@ void SNiagaraGeneratedCodeView::UpdateUI()
 
 	for (int32 i = 0; i < GeneratedCode.Num(); i++)
 	{
-		TArray<FString> OutputByLines;
 		GeneratedCode[i].Hlsl = FText::GetEmpty();
-		FString SumString;
 
-		bool bIsGPU = ScriptDisplayTypes[i] == 1;
-		bool bIsAssembly = ScriptDisplayTypes[i] == 2;
 		if (Scripts[i] != nullptr)
 		{
-			// GPU combined spawn / update script
-			if (bIsGPU)
+			const FNiagaraVMExecutableData& ExeData = Scripts[i]->GetVMExecutableData();
+
+			GeneratedCode[i].Usage = Scripts[i]->Usage;
+			GeneratedCode[i].UsageId = Scripts[i]->GetUsageId();
+
+			if (ExeData.IsValid())
 			{
-				GeneratedCode[i].Usage = Scripts[i]->Usage;
-				if (Scripts[i]->GetVMExecutableData().IsValid())
+				const FString* SourceText = nullptr;
+				bool bAddLineNumbers = false;
+
+				FText AssemblyIdText = FText::GetEmpty();
+
+				switch (ScriptDisplayTypes[i])
 				{
-					Scripts[i]->GetVMExecutableData().LastHlslTranslationGPU.ParseIntoArrayLines(OutputByLines, false);
+					// GPU combined spawn / update script
+					case GPU_SCRIPT_TYPE:
+						SourceText = &ExeData.LastHlslTranslationGPU;
+						bAddLineNumbers = true;
+						break;
+
+					case ASSEMBLY_SCRIPT_TYPE:
+						SourceText = &ExeData.LastAssemblyTranslation;
+						AssemblyIdText = LOCTEXT("IsAssembly", "Assembly");
+						break;
+
+					case EVVM_ASSEMBLY_SCRIPT_TYPE:
+						SourceText = &ExeData.LastExperimentalAssemblyScript;
+						AssemblyIdText = LOCTEXT("IsEVVMAssembly", "EVVM Assembly");
+						break;
+
+					default:
+						SourceText = &ExeData.LastHlslTranslation;
+						bAddLineNumbers = true;
+						break;
 				}
-			}
-			else if (bIsAssembly)
-			{
-				GeneratedCode[i].Usage = Scripts[i]->Usage;
-				GeneratedCode[i].UsageId = Scripts[i]->GetUsageId();
-				if (Scripts[i]->GetVMExecutableData().IsValid())
+
+				if (Scripts[i]->Usage == ENiagaraScriptUsage::ParticleEventScript)
 				{
-					Scripts[i]->GetVMExecutableData().LastAssemblyTranslation.ParseIntoArrayLines(OutputByLines, false);
+					FText EventName;
+					if (FNiagaraEditorUtilities::TryGetEventDisplayName(Scripts[i]->GetTypedOuter<UNiagaraEmitter>(), Scripts[i]->GetUsageId(), EventName) == false)
+					{
+						EventName = NSLOCTEXT("NiagaraNodeOutput", "UnknownEventName", "Unknown");
+					}
+					GeneratedCode[i].UsageName = FText::Format(LOCTEXT("UsageNameEvent", "{0}-{1}{2}"), ScriptEnum->GetDisplayNameTextByValue((int64)Scripts[i]->Usage), EventName, AssemblyIdText);
 				}
-			}
-			else
-			{
-				GeneratedCode[i].Usage = Scripts[i]->Usage;
-				GeneratedCode[i].UsageId = Scripts[i]->GetUsageId();
-				if (Scripts[i]->GetVMExecutableData().IsValid())
+				// GPU combined spawn / update script
+				else if (ScriptDisplayTypes[i] == GPU_SCRIPT_TYPE && i == GeneratedCode.Num() - 1 && Scripts[i]->IsParticleSpawnScript())
 				{
-					Scripts[i]->GetVMExecutableData().LastHlslTranslation.ParseIntoArrayLines(OutputByLines, false);
+					GeneratedCode[i].UsageName = LOCTEXT("UsageNameGPU", "GPU Spawn/Update");
+				}
+				else
+				{
+					GeneratedCode[i].UsageName = FText::Format(LOCTEXT("UsageName", "{0}{1}"), ScriptEnum->GetDisplayNameTextByValue((int64)Scripts[i]->Usage), AssemblyIdText);
 				}
-			}
-		}
-		else
-		{
-			GeneratedCode[i].Usage = ENiagaraScriptUsage::ParticleSpawnScript;
-		}
 
-		GeneratedCode[i].HlslByLines.SetNum(OutputByLines.Num());
-		if (bIsAssembly)
-		{
-			if (Scripts[i] != nullptr && Scripts[i]->GetVMExecutableData().IsValid())
-			{
-				GeneratedCode[i].Hlsl = FText::FromString(Scripts[i]->GetVMExecutableData().LastAssemblyTranslation);
+				if (SourceText)
+				{
+					if (bAddLineNumbers)
+					{
+						FString SumString;
+						TArray<FString> OutputByLines;
+
+						SourceText->ParseIntoArrayLines(OutputByLines, false);
+
+						SumString.Reserve(SourceText->Len() + OutputByLines.Num() * 10);
+						GeneratedCode[i].HlslByLines.Reset(OutputByLines.Num());
+
+						for (int32 k = 0; k < OutputByLines.Num(); k++)
+						{
+							GeneratedCode[i].HlslByLines.Add(FString::Printf(TEXT("/*%04d*/\t\t%s\r\n"), k, *OutputByLines[k]));
+							SumString.Append(GeneratedCode[i].HlslByLines[k]);
+						}
+						GeneratedCode[i].Hlsl = FText::FromString(SumString);
+					}
+					else
+					{
+						GeneratedCode[i].Hlsl = FText::FromString(*SourceText);
+						SourceText->ParseIntoArrayLines(GeneratedCode[i].HlslByLines, false);
+					}
+				}
 			}
-			GeneratedCode[i].HlslByLines = OutputByLines;
 		}
 		else
 		{
-			for (int32 k = 0; k < OutputByLines.Num(); k++)
-			{
-				GeneratedCode[i].HlslByLines[k] = FString::Printf(TEXT("/*%04d*/\t\t%s\r\n"), k, *OutputByLines[k]);
-				SumString.Append(GeneratedCode[i].HlslByLines[k]);
-			}
-			GeneratedCode[i].Hlsl = FText::FromString(SumString);
-		}
-		FText AssemblyIdText = LOCTEXT("IsAssembly", "Assembly");
-
-		if (Scripts[i] == nullptr)
-		{
+			GeneratedCode[i].Usage = ENiagaraScriptUsage::ParticleSpawnScript;
 			GeneratedCode[i].UsageName = LOCTEXT("UsageNameInvalid", "Invalid");
 		}
-		else if (Scripts[i]->Usage == ENiagaraScriptUsage::ParticleEventScript)
-		{
-			FText EventName;
-			if (FNiagaraEditorUtilities::TryGetEventDisplayName(Scripts[i]->GetTypedOuter<UNiagaraEmitter>(), Scripts[i]->GetUsageId(), EventName) == false)
-			{
-				EventName = NSLOCTEXT("NiagaraNodeOutput", "UnknownEventName", "Unknown");
-			}
-			GeneratedCode[i].UsageName = FText::Format(LOCTEXT("UsageNameEvent", "{0}-{1}{2}"), ScriptEnum->GetDisplayNameTextByValue((int64)Scripts[i]->Usage), EventName, bIsAssembly ? AssemblyIdText : FText::GetEmpty());
-		}
-		// GPU combined spawn / update script
-		else if (bIsGPU && i == GeneratedCode.Num() - 1 && Scripts[i]->IsParticleSpawnScript())
-		{
-			GeneratedCode[i].UsageName = LOCTEXT("UsageNameGPU", "GPU Spawn/Update");
-		}
-		else
-		{
-			GeneratedCode[i].UsageName = FText::Format(LOCTEXT("UsageName", "{0}{1}"), ScriptEnum->GetDisplayNameTextByValue((int64)Scripts[i]->Usage), bIsAssembly ? AssemblyIdText : FText::GetEmpty());
-		}
 
 		//We now do the very expensive slate stuff inside Update UI_Internal
 	}
diff --git a/Engine/Source/Runtime/VectorVM/Private/VectorVMExperimental.cpp b/Engine/Source/Runtime/VectorVM/Private/VectorVMExperimental.cpp
--- a/Engine/Source/Runtime/VectorVM/Private/VectorVMExperimental.cpp
+++ b/Engine/Source/Runtime/VectorVM/Private/VectorVMExperimental.cpp
@@ -680,7 +680,8 @@ FVectorVMState *AllocVectorVMState(FVectorVMOptimizeContext *OptimizeCtx) {
 static void VVMBuildMapTableCaches(FVectorVMExecContext *ExecCtx)
 {
 	//constant buffers
-	check(ExecCtx->VVMState->NumConstBuffers < 0xFF);
+	check(ExecCtx->ConstantTableCount <= 0xFF);
+	check(ExecCtx->VVMState->NumConstBuffers <= 0xFFFF);
 	for (uint32 i = 0; i < ExecCtx->VVMState->NumConstBuffers; ++i)
 	{
 		uint32 RemappedIdx = ExecCtx->VVMState->ConstRemapTable[i];
@@ -1204,9 +1205,10 @@ static VM_FORCEINLINE const uint8 *VVM_random(const bool CT_MultipleLoops, const
 	VVMSer_instruction(0, 1);
 	if (!VVM_serSyncRandom(InsPtr, BatchState, ExecCtx, SS, CmpSS, 1, EVectorVMOp::random))
 	{
-		uint8 *P0   = BatchState->RegPtrTable[((uint16 *)InsPtr)[0]];
-		uint8 *P1   = BatchState->RegPtrTable[((uint16 *)InsPtr)[1]];
-		uint32 Inc0 = (uint32)BatchState->RegIncTable[((uint16 *)InsPtr)[0]];
+		uint16* RegIndices = (uint16*)InsPtr;
+		uint8 *P0   = BatchState->RegPtrTable[RegIndices[0]];
+		uint8 *P1   = BatchState->RegPtrTable[RegIndices[1]];
+		uint32 Inc0 = (uint32)BatchState->RegIncTable[RegIndices[0]];
 		if (CT_MultipleLoops)
 		{
 			uint8 *End = P1 + sizeof(FVecReg) * NumLoops;
@@ -1234,9 +1236,10 @@ static VM_FORCEINLINE const uint8 *VVM_randomi(const bool CT_MultipleLoops, cons
 	VVMSer_instruction(1, 1);
 	if (!VVM_serSyncRandom(InsPtr, BatchState, ExecCtx, SS, CmpSS, 1, EVectorVMOp::randomi))
 	{
-		uint8 *P0   = BatchState->RegPtrTable[((uint16 *)InsPtr)[0]];
-		uint8 *P1   = BatchState->RegPtrTable[((uint16 *)InsPtr)[1]];
-		uint32 Inc0 = (uint32)BatchState->RegIncTable[((uint16 *)InsPtr)[0]];
+		uint16* RegIndices = (uint16*)InsPtr;
+		uint8 *P0   = BatchState->RegPtrTable[RegIndices[0]];
+		uint8 *P1   = BatchState->RegPtrTable[RegIndices[1]];
+		uint32 Inc0 = (uint32)BatchState->RegIncTable[RegIndices[0]];
 		if (CT_MultipleLoops)
 		{
 			uint8 *End = P1 + sizeof(FVecReg) * NumLoops;
@@ -1264,10 +1267,13 @@ static VM_FORCEINLINE const uint8 *VVM_random_add(const bool CT_MultipleLoops, c
 	VVMSer_instruction(0, 2);
 	if (!VVM_serSyncRandom(InsPtr, BatchState, ExecCtx, SS, CmpSS, 2, EVectorVMOp::random_add))
 	{
-		uint8 *P0   = BatchState->RegPtrTable[((uint16 *)InsPtr)[0]];
-		uint8 *P1   = BatchState->RegPtrTable[((uint16 *)InsPtr)[1]];
-		uint8 *P2   = BatchState->RegPtrTable[((uint16 *)InsPtr)[2]];
-		uint32 Inc0 = (uint32)BatchState->RegIncTable[((uint16 *)InsPtr)[0]];
+		uint16* RegIndices = (uint16*)InsPtr;
+		uint8 *P0   = BatchState->RegPtrTable[RegIndices[0]];
+		uint8 *P1   = BatchState->RegPtrTable[RegIndices[1]];
+		uint8 *P2   = BatchState->RegPtrTable[RegIndices[2]];
+		uint32 Inc0 = (uint32)BatchState->RegIncTable[RegIndices[0]];
+		uint32 Inc1 = (uint32)BatchState->RegIncTable[RegIndices[1]];
+
 		if (CT_MultipleLoops)
 		{
 			uint8 *End = P2 + sizeof(FVecReg) * NumLoops;
@@ -1276,7 +1282,7 @@ static VM_FORCEINLINE const uint8 *VVM_random_add(const bool CT_MultipleLoops, c
 				VectorRegister4f R0 = VectorLoad((float *)P0);
 				VectorRegister4f R1 = VectorLoad((float *)P1);
 				P0 += Inc0;
-				P1 += Inc0;
+				P1 += Inc1;
 				VectorRegister4f Res = VectorAdd(VVM_nextRandom(BatchState, R0), R1);
 				VectorStoreAligned(Res, (float *)P2);
 				P2 += sizeof(FVecReg);
@@ -1298,10 +1304,11 @@ static VM_FORCEINLINE const uint8 *VVM_random_2x(const bool CT_MultipleLoops, co
 	VVMSer_instruction(0, 2);
 	if (!VVM_serSyncRandom(InsPtr, BatchState, ExecCtx, SS, CmpSS, 1, EVectorVMOp::random_2x) && !VVM_serSyncRandom(InsPtr, BatchState, ExecCtx, SS, CmpSS, 2, EVectorVMOp::random_2x))
 	{
-		uint8 *P0   = BatchState->RegPtrTable[((uint16 *)InsPtr)[0]];
-		uint8 *P1   = BatchState->RegPtrTable[((uint16 *)InsPtr)[1]];
-		uint8 *P2   = BatchState->RegPtrTable[((uint16 *)InsPtr)[2]];
-		uint32 Inc0 = (uint32)BatchState->RegIncTable[((uint16 *)InsPtr)[0]];
+		uint16* RegIndices = (uint16*)InsPtr;
+		uint8 *P0   = BatchState->RegPtrTable[RegIndices[0]];
+		uint8 *P1   = BatchState->RegPtrTable[RegIndices[1]];
+		uint8 *P2   = BatchState->RegPtrTable[RegIndices[2]];
+		uint32 Inc0 = (uint32)BatchState->RegIncTable[RegIndices[0]];
 		if (CT_MultipleLoops)
 		{
 			uint8 *End = P1 + sizeof(FVecReg) * NumLoops;
@@ -1331,9 +1338,10 @@ static VM_FORCEINLINE const uint8 *VVM_random_2x(const bool CT_MultipleLoops, co
 
 VM_FORCEINLINE const uint8 *VVM_half_to_float(const bool CT_MultipleLoops, const uint8 *InsPtr, FVectorVMBatchState *BatchState, int NumLoops)
 {
-	uint8 *P0   = BatchState->RegPtrTable[((uint16 *)InsPtr)[0]];
-	uint8 *P1   = BatchState->RegPtrTable[((uint16 *)InsPtr)[1]];
-	uint32 Inc0 = (uint32)BatchState->RegIncTable[((uint16 *)InsPtr)[0]] >> 1; //move the input 2 bytes instead of four
+	uint16* RegIndices = (uint16*)InsPtr;
+	uint8 *P0   = BatchState->RegPtrTable[RegIndices[0]];
+	uint8 *P1   = BatchState->RegPtrTable[RegIndices[1]];
+	uint32 Inc0 = (uint32)BatchState->RegIncTable[RegIndices[0]] >> 1; //move the input 2 bytes instead of four
 	uint8 *End  = P1 + sizeof(FVecReg) * NumLoops;
 	do {
 		FPlatformMath::VectorLoadHalf((float *)P1, (uint16 *)P0);
@@ -1345,8 +1353,9 @@ VM_FORCEINLINE const uint8 *VVM_half_to_float(const bool CT_MultipleLoops, const
 
 static const uint8 *VVM_update_id(const bool CT_MultipleLoops, const uint8 *InsPtr, FVectorVMBatchState *BatchState, FVectorVMExecContext *ExecCtx, int NumLoops)
 {
-	int32 *R1             = (int32 *)BatchState->RegPtrTable[((uint16 *)InsPtr)[0]];
-	int32 *R2             = (int32 *)BatchState->RegPtrTable[((uint16 *)InsPtr)[1]];
+	uint16* RegIndices = (uint16*)InsPtr;
+	int32 *R1             = (int32 *)BatchState->RegPtrTable[RegIndices[0]];
+	int32 *R2             = (int32 *)BatchState->RegPtrTable[RegIndices[1]];
 	uint8 DataSetIdx      = InsPtr[4];
 	FDataSetMeta *DataSet = &ExecCtx->DataSets[DataSetIdx];
 	check(DataSetIdx < (uint32)ExecCtx->DataSets.Num());
diff --git a/Engine/Source/Runtime/VectorVM/Private/VectorVMExperimentalOptimizer.cpp b/Engine/Source/Runtime/VectorVM/Private/VectorVMExperimentalOptimizer.cpp
--- a/Engine/Source/Runtime/VectorVM/Private/VectorVMExperimentalOptimizer.cpp
+++ b/Engine/Source/Runtime/VectorVM/Private/VectorVMExperimentalOptimizer.cpp
@@ -1,5 +1,6 @@
 // Copyright Epic Games, Inc. All Rights Reserved.
 #include "CoreTypes.h"
+#include "VectorVM.h"
 #include "VectorVMExperimental.h"
 
 #if VECTORVM_SUPPORTS_EXPERIMENTAL
@@ -187,6 +188,170 @@ void ReinterpretVectorVMOptimizeContextData(TConstArrayView<uint8> ContextData,
 	Context.Error.CallbackFn      = nullptr;
 }
 
+namespace VectorVMScriptStringHelper
+{
+
+static void GenerateConstantTableString(const FVectorVMOptimizeContext& Context, FString& OpsConstantTable)
+{
+	// todo - implement a string of the form:
+	// 0 | Engine_WorldDeltaTime
+	// 1 | Engine_DeltaTime
+	// 2 | Engine_InverseDeltaTime
+	// to put into the human readable script version
+}
+
+// todo - note that this is just an initial implementation that needs to be revisited to confirm that it's actually
+// closely matching the byte that will be run.  In particular the register representations when aliasing output
+// buffers seems confusing.
+static bool GenerateInstructionString(const FVectorVMOptimizeContext& Context, uint32 InstrIndex, FString& InstructionString)
+{
+	const FVectorVMOptimizeInstruction& Instruction = Context.Intermediate.Instructions[InstrIndex];
+
+	TStringBuilder<128> OpName;
+	TStringBuilder<256> InstString;
+
+#if WITH_EDITOR
+	OpName << VectorVM::GetOpName(Instruction.OpCode);
+#else
+	OpName << TEXT("__OP__") << (int32)Instruction.OpCode;
+#endif
+
+	auto WriteRegisterList = [&](int32 RegOffset, int32 Count, FStringBuilderBase& StringBuilder) -> void
+	{
+		for (int32 RegIt = 0; RegIt < Count; ++RegIt)
+		{
+			if (RegIt)
+			{
+				StringBuilder << TEXT(", ");
+			}
+			switch (Context.Intermediate.RegisterUsageType[RegOffset + RegIt])
+			{
+				case VVM_RT_TEMPREG: StringBuilder << TEXT("R"); break;
+				case VVM_RT_CONST: StringBuilder << TEXT("C"); break;
+				case VVM_RT_INPUT: StringBuilder << TEXT("I"); break;
+				case VVM_RT_OUTPUT: StringBuilder << TEXT("O"); break;
+			}
+
+			StringBuilder << TEXT("[");
+			StringBuilder << Context.Intermediate.RegisterUsageBuffer[RegOffset + RegIt];
+			StringBuilder << TEXT("]");
+		}
+	};
+	
+	switch (Instruction.OpCat)
+	{
+		case EVectorVMOpCategory::Input:
+		{
+			check(Instruction.NumOutputRegisters == 1);
+			WriteRegisterList(Instruction.RegPtrOffset + Instruction.NumInputRegisters, 1, InstString);
+			InstString << TEXT(" = ");
+			InstString << OpName;
+			InstString << TEXT("(");
+			InstString << TEXT(")");
+		} break;
+
+		case EVectorVMOpCategory::Output:
+		{
+			InstString << OpName;
+			InstString << TEXT("(");
+			InstString << Instruction.Output.DataSetIdx;
+			InstString << TEXT(", ");
+			InstString << Instruction.Output.DstRegIdx;
+			InstString << TEXT(", ");
+			WriteRegisterList(Instruction.RegPtrOffset, Instruction.NumInputRegisters, InstString);
+			InstString << TEXT(")");
+		} break;
+
+		case EVectorVMOpCategory::Op:
+		{
+			check(Instruction.NumOutputRegisters == 1);
+			WriteRegisterList(Instruction.RegPtrOffset + Instruction.NumInputRegisters, 1, InstString);
+			InstString << TEXT(" = ");
+#if WITH_EDITOR
+			InstString << VectorVM::GetOpName(Instruction.OpCode);
+#else
+			InstString << TEXT("__OP__") << (int32) Instruction.OpCode;
+#endif
+			InstString << TEXT("(");
+			WriteRegisterList(Instruction.RegPtrOffset, Instruction.NumInputRegisters, InstString);
+			InstString << TEXT(")");
+		} break;
+
+		case EVectorVMOpCategory::ExtFnCall:
+		{
+			InstString << OpName;
+			InstString << TEXT("(");
+			WriteRegisterList(Instruction.RegPtrOffset, Instruction.NumInputRegisters, InstString);
+			WriteRegisterList(Instruction.RegPtrOffset + Instruction.NumInputRegisters, Instruction.NumOutputRegisters, InstString);
+			InstString << TEXT(")");
+		} break;
+
+		case EVectorVMOpCategory::IndexGen:
+		{
+			check(Instruction.NumOutputRegisters == 1);
+			WriteRegisterList(Instruction.RegPtrOffset + Instruction.NumInputRegisters, 1, InstString);
+			InstString << TEXT(" = ");
+			InstString << OpName;
+			InstString << TEXT("(");
+			WriteRegisterList(Instruction.RegPtrOffset, Instruction.NumInputRegisters, InstString);
+			InstString << TEXT(")");
+		} break;
+
+		case EVectorVMOpCategory::RWBuffer:
+		case EVectorVMOpCategory::Stat:
+		case EVectorVMOpCategory::Other:
+		{
+			InstString << OpName;
+			InstString << TEXT("(");
+			WriteRegisterList(Instruction.RegPtrOffset, Instruction.NumInputRegisters, InstString);
+			InstString << TEXT(")");
+		} break;
+	}
+
+	InstructionString = InstString.ToString();
+
+	return Instruction.InsMergedIdx == INDEX_NONE;
+}
+
+}; // VectorVMScriptStringHelper
+
+void GenerateHumanReadableVectorVMScript(const FVectorVMOptimizeContext& Context, FString& VMScript)
+{
+	using namespace VectorVMScriptStringHelper;
+
+	FString OpsConstantTable;
+	GenerateConstantTableString(Context, OpsConstantTable);
+
+	VMScript += TEXT("\n-------------------------------\n");
+	VMScript += TEXT("Summary\n");
+	VMScript += TEXT("-------------------------------\n");
+	VMScript += FString::Printf(TEXT("Num Byte Code Ops: %d\n"), Context.Intermediate.NumInstructions);
+	VMScript += FString::Printf(TEXT("Num Constants: %d\n"), Context.NumConstsRemapped);
+
+	//Dump the constant table
+	VMScript += TEXT("\n-------------------------------\n");
+	VMScript += TEXT("Constant Table\n");
+	VMScript += TEXT("-------------------------------\n");
+	VMScript += OpsConstantTable;
+
+	VMScript += TEXT("-------------------------------\n");
+	VMScript += FString::Printf(TEXT("Byte Code (%d Ops)\n"), Context.Intermediate.NumInstructions);
+	VMScript += TEXT("-------------------------------\n");
+
+	//Dump the bytecode		
+	for (uint32 op_idx = 0; op_idx < Context.Intermediate.NumInstructions; ++op_idx)
+	{
+		FString InstructionString;
+		if (GenerateInstructionString(Context, op_idx, InstructionString))
+		{
+			VMScript += FString::Printf(TEXT("%d\t| "), op_idx) + InstructionString;
+			VMScript += TEXT(";\n");
+		}
+	}
+
+	VMScript += TEXT("-------------------------------\n");
+}
+
 void FreeVectorVMOptimizeContext(FVectorVMOptimizeContext *OptContext)
 {
 	//save init data
@@ -1420,9 +1585,6 @@ uint32 OptimizeVectorVMScript(const uint8 *InBytecode, int InBytecodeLen, FVecto
 		for (uint32 i = 0; i < OptContext->Intermediate.NumInstructions; ++i)
 		{
 			FVectorVMOptimizeInstruction *Ins = OptContext->Intermediate.Instructions + i;
-			if (i == 106) {
-				int x =100;
-			}
 			int InsMergeIdx = -1;
 			if (Ins->OpCode == EVectorVMOp::exec_index)
 			{
@@ -1433,6 +1595,11 @@ uint32 OptimizeVectorVMScript(const uint8 *InBytecode, int InBytecodeLen, FVecto
 					FVectorVMOptimizeInstruction *Ins2 = OptContext->Intermediate.Instructions + j;
 					for (int k = 0; k < Ins2->NumInputRegisters; ++k)
 					{
+						// we're only looking for temporary registers
+						if (OptContext->Intermediate.RegisterUsageType[Ins2->RegPtrOffset + k] != VVM_RT_TEMPREG)
+						{
+							continue;
+						}
 						uint16 InputSSAReg = OptContext->Intermediate.SSARegisterUsageBuffer[Ins2->RegPtrOffset + k];
 						if (InputSSAReg == ExecSSARegIdx)
 						{
@@ -1650,18 +1817,6 @@ uint32 OptimizeVectorVMScript(const uint8 *InBytecode, int InBytecodeLen, FVecto
 			{
 				if (Ins0->OpCode == EVectorVMOp::exec_index && Ins1->OpCode == EVectorVMOp::i2f)
 				{
-					for (uint32 j = MergableOps[i].InsIdx0 + 1; j < OptContext->Intermediate.NumInstructions; ++j)
-					{
-						if (j != MergableOps[i].InsIdx1)
-						{
-							FVectorVMOptimizeInstruction *OtherIns = OptContext->Intermediate.Instructions + j;
-							for (int k = 0; k < OtherIns->NumInputRegisters; ++k)
-							{
-								uint16 SSA = OptContext->Intermediate.SSARegisterUsageBuffer[OtherIns->RegPtrOffset + k];
-								check(SSA != OptContext->Intermediate.SSARegisterUsageBuffer[Ins0->RegPtrOffset]);
-							}
-						}
-					}
 					VVMCreateNewRegVars(1);
 					VVMSetRegsFrom1(0, 1);
 					VVMSetMergedIns(EVectorVMOp::exec_indexf, 0, 1);
@@ -2584,9 +2739,9 @@ uint32 OptimizeVectorVMScript(const uint8 *InBytecode, int InBytecodeLen, FVecto
 				uint16 RegIdx = Ins->RegPtrOffset + Ins->NumInputRegisters + j;
 				uint16 SSAReg = OptContext->Intermediate.SSARegisterUsageBuffer[RegIdx];
 				uint8 RegType = OptContext->Intermediate.RegisterUsageType[RegIdx];
-				if (RegType == VVM_RT_TEMPREG && SSAReg != 0xFFFF)
+				if (RegType == VVM_RT_TEMPREG)
 				{
-					if (SSAUseMap2Inv[SSAReg] != 0xFFFF)
+					if (SSAReg != 0xFFFF && SSAUseMap2Inv[SSAReg] != 0xFFFF)
 					{
 						OptContext->Intermediate.RegisterUsageBuffer[RegIdx] = SSAUseMap2Inv[SSAReg];
 						--NumSSARegistersUsed2;
diff --git a/Engine/Source/Runtime/VectorVM/Public/VectorVMExperimental.h b/Engine/Source/Runtime/VectorVM/Public/VectorVMExperimental.h
--- a/Engine/Source/Runtime/VectorVM/Public/VectorVMExperimental.h
+++ b/Engine/Source/Runtime/VectorVM/Public/VectorVMExperimental.h
@@ -444,6 +444,7 @@ VECTORVM_API uint32           OptimizeVectorVMScript                 (const uint
 VECTORVM_API void             FreeVectorVMOptimizeContext            (FVectorVMOptimizeContext *Context);
 VECTORVM_API void             FreezeVectorVMOptimizeContext          (const FVectorVMOptimizeContext& Context, TArray<uint8>& ContextData);
 VECTORVM_API void             ReinterpretVectorVMOptimizeContextData (TConstArrayView<uint8> ContextData, FVectorVMOptimizeContext& Context);
+VECTORVM_API void             GenerateHumanReadableVectorVMScript    (const FVectorVMOptimizeContext& Context, FString& VMScript);
 
 //serialize functions
 VECTORVM_API uint32           SerializeVectorVMInputDataSets         (FVectorVMSerializeState *SerializeState, FVectorVMExecContext *ExecContext);
