- Validation rule to check if sorting is enabled on sprite / mesh renderers

diff --git a/Engine/Plugins/FX/Niagara/Source/NiagaraEditor/Private/NiagaraValidationRules.cpp b/Engine/Plugins/FX/Niagara/Source/NiagaraEditor/Private/NiagaraValidationRules.cpp
--- a/Engine/Plugins/FX/Niagara/Source/NiagaraEditor/Private/NiagaraValidationRules.cpp
+++ b/Engine/Plugins/FX/Niagara/Source/NiagaraEditor/Private/NiagaraValidationRules.cpp
@@ -9,9 +9,11 @@
 #include "NiagaraDataInterfaceUtilities.h"
 #include "NiagaraEditorSettings.h"
 #include "NiagaraNodeFunctionCall.h"
+#include "NiagaraMeshRendererProperties.h"
 #include "NiagaraScriptSource.h"
 #include "NiagaraSettings.h"
 #include "NiagaraSimulationStageBase.h"
+#include "NiagaraSpriteRendererProperties.h"
 #include "NiagaraSystemImpl.h"
 #include "NiagaraSystemEditorData.h"
 #include "DataInterface/NiagaraDataInterfaceActorComponent.h"
@@ -87,6 +89,20 @@ namespace NiagaraValidation
 		return nullptr;
 	}
 
+	// helper function to get renderer stack item
+	UNiagaraStackRendererItem* GetRendererStackItem(UNiagaraStackViewModel* StackViewModel, UNiagaraRendererProperties* RendererProperties)
+	{
+		TArray<UNiagaraStackRendererItem*> RendererItems = GetStackEntries<UNiagaraStackRendererItem>(StackViewModel);
+		for (UNiagaraStackRendererItem* Item : RendererItems)
+		{
+			if (Item->GetRendererProperties() == RendererProperties)
+			{
+				return Item;
+			}
+		}
+		return nullptr;
+	}
+
 	// --------------------------------------------------------------------------------------------------------------------------------------------
 	// Common fixes and links
 
@@ -424,26 +440,21 @@ void UNiagaraValidationRule_BannedRenderers::CheckValidity(const FNiagaraValidat
 				FNiagaraPlatformSet::GatherConflicts(CheckSets, Conflicts);
 				if (Conflicts.Num() > 0)
 				{
-					TArray<UNiagaraStackRendererItem*> RendererItems = NiagaraValidation::GetStackEntries<UNiagaraStackRendererItem>(EmitterHandleModel.Get().GetEmitterStackViewModel());
-					for (UNiagaraStackRendererItem* Item : RendererItems)
+					if ( UNiagaraStackRendererItem* StackItem = NiagaraValidation::GetRendererStackItem(EmitterHandleModel.Get().GetEmitterStackViewModel(), RendererProperties) )
 					{
-						if (Item->GetRendererProperties() != RendererProperties)
-						{
-							continue;
-						}						
-						FNiagaraValidationResult Result = Results.AddDefaulted_GetRef();
+						FNiagaraValidationResult& Result = Results.AddDefaulted_GetRef();
 						
 						Result.Severity = ENiagaraValidationSeverity::Warning;
 						Result.SummaryText = LOCTEXT("BannedRenderSummary", "Banned renderers used.");
 						Result.Description = LOCTEXT("BannedRenderDescription", "Please ensure only allowed renderers are used for each platform according to the validation rules in the System's Effect Type.");
-						Result.SourceObject = Item;
+						Result.SourceObject = StackItem;
 						
 						NiagaraValidation::AddGoToFXTypeLink(Result, System.GetEffectType());
 
 						//Add autofix to disable the module
 						FNiagaraValidationFix& DisableRendererFix = Result.Fixes.AddDefaulted_GetRef();
 						DisableRendererFix.Description = LOCTEXT("DisableBannedRendererFix", "Disable Banned Renderer");
-						TWeakObjectPtr<UNiagaraStackRendererItem> WeakRendererItem = Item;
+						TWeakObjectPtr<UNiagaraStackRendererItem> WeakRendererItem = StackItem;
 						DisableRendererFix.FixDelegate = FNiagaraValidationFixDelegate::CreateLambda(
 							[WeakRendererItem]()
 							{
@@ -451,9 +462,8 @@ void UNiagaraValidationRule_BannedRenderers::CheckValidity(const FNiagaraValidat
 								{
 									RendererItem->SetIsEnabled(false);
 								}
-							});
-
-						Results.Add(Result);
+							}
+						);
 					}
 				}
 			}
@@ -593,6 +603,67 @@ void UNiagaraValidationRule_BannedDataInterfaces::CheckValidity(const FNiagaraVa
 	);
 }
 
+template<typename TRendererType>
+FNiagaraValidationResult* NiagaraRendererCheckSortingEnabled(const TSharedRef<FNiagaraEmitterHandleViewModel>& EmitterHandleModel, UNiagaraRendererProperties* InProperties, TArray<FNiagaraValidationResult>& Results, ENiagaraValidationSeverity Severity)
+{
+	TRendererType* Properties = Cast<TRendererType>(InProperties);
+	if (!Properties || !Properties->GetIsEnabled() || Properties->SortMode == ENiagaraSortMode::None)
+	{
+		return nullptr;
+	}
+
+	UNiagaraStackRendererItem* StackItem = NiagaraValidation::GetRendererStackItem(EmitterHandleModel.Get().GetEmitterStackViewModel(), Properties);
+	if (StackItem == nullptr)
+	{
+		return nullptr;
+	}
+
+	FNiagaraValidationResult& Result = Results.AddDefaulted_GetRef();
+	Result.SummaryText = LOCTEXT("RendererSortingEnabled", "Sorting is enabled on the renderer.");
+	Result.Description = LOCTEXT("RendererSortingEnabledDesc", "Sorting is enabled on the renderer, this costs performance consider if it can be disabled or not.");
+	Result.Severity = Severity;
+	Result.SourceObject = StackItem;
+	Result.Fixes.Emplace(
+		LOCTEXT("DisableSortingFix", "Disable sorting on the renderer"),
+		FNiagaraValidationFixDelegate::CreateLambda(
+			[WeakRenderer=MakeWeakObjectPtr(Properties)]()
+			{
+				if (WeakRenderer.IsValid())
+				{
+					const FScopedTransaction Transaction(LOCTEXT("DisableSorting", "Disable Sorting"));
+					WeakRenderer.Get()->Modify();
+					WeakRenderer.Get()->SortMode = ENiagaraSortMode::None;
+				}
+			}
+		)
+	);
+	return &Result;
+}
+
+void UNiagaraValidationRule_RendererSortingEnabled::CheckValidity(const FNiagaraValidationContext& Context, TArray<FNiagaraValidationResult>& Results)  const
+{
+	UNiagaraSystem& System = Context.ViewModel->GetSystem();
+	TArray<TSharedRef<FNiagaraEmitterHandleViewModel>> EmitterHandleViewModels = Context.ViewModel->GetEmitterHandleViewModels();
+	for (TSharedRef<FNiagaraEmitterHandleViewModel> EmitterHandleModel : EmitterHandleViewModels)
+	{
+		FVersionedNiagaraEmitterData* EmitterData = EmitterHandleModel.Get().GetEmitterHandle()->GetEmitterData();
+		EmitterData->ForEachRenderer(
+			[&Results, EmitterHandleModel, this, &System](UNiagaraRendererProperties* RendererProperties)
+			{
+				if ( NiagaraRendererCheckSortingEnabled<UNiagaraSpriteRendererProperties>(EmitterHandleModel, RendererProperties, Results, Severity) )
+				{
+					return;
+				}
+
+				if ( NiagaraRendererCheckSortingEnabled<UNiagaraMeshRendererProperties>(EmitterHandleModel, RendererProperties, Results, Severity) )
+				{
+					return;
+				}
+			}
+		);
+	}
+}
+
 void UNiagaraValidationRule_GpuUsage::CheckValidity(const FNiagaraValidationContext& Context, TArray<FNiagaraValidationResult>& OutResults) const
 {
 	UNiagaraSystem& System = Context.ViewModel->GetSystem();
diff --git a/Engine/Plugins/FX/Niagara/Source/NiagaraEditor/Public/NiagaraValidationRules.h b/Engine/Plugins/FX/Niagara/Source/NiagaraEditor/Public/NiagaraValidationRules.h
--- a/Engine/Plugins/FX/Niagara/Source/NiagaraEditor/Public/NiagaraValidationRules.h
+++ b/Engine/Plugins/FX/Niagara/Source/NiagaraEditor/Public/NiagaraValidationRules.h
@@ -95,6 +95,19 @@ public:
 	virtual void CheckValidity(const FNiagaraValidationContext& Context, TArray<FNiagaraValidationResult>& OutResults) const override;
 };
 
+/** Checks to see if renderers have sorting enabled on them or not. */
+UCLASS(Category = "Validation", DisplayName = "Renderer Sorting Enabled")
+class UNiagaraValidationRule_RendererSortingEnabled : public UNiagaraValidationRule
+{
+	GENERATED_BODY()
+
+public:
+	UPROPERTY(EditAnywhere, Category = Validation)
+	ENiagaraValidationSeverity Severity = ENiagaraValidationSeverity::Warning;
+
+	virtual void CheckValidity(const FNiagaraValidationContext& Context, TArray<FNiagaraValidationResult>& OutResults) const override;
+};
+
 /** This validation rule that can be used to inform or ban GPU usage on the provided platforms. */
 UCLASS(Category = "Validation", DisplayName = "Gpu Usage")
 class UNiagaraValidationRule_GpuUsage : public UNiagaraValidationRule
