- Add option to select particle mesh LOD based on component bounds
- All particles will render at this LOD

diff --git a/Engine/Plugins/FX/Niagara/Source/Niagara/Private/NiagaraMeshRendererProperties.cpp b/Engine/Plugins/FX/Niagara/Source/Niagara/Private/NiagaraMeshRendererProperties.cpp
--- a/Engine/Plugins/FX/Niagara/Source/Niagara/Private/NiagaraMeshRendererProperties.cpp
+++ b/Engine/Plugins/FX/Niagara/Source/Niagara/Private/NiagaraMeshRendererProperties.cpp
@@ -11,11 +11,12 @@
 #include "NiagaraSystem.h"
 #include "NiagaraSystemImpl.h"
 
-#include "MaterialDomain.h"
-#include "Materials/MaterialRenderProxy.h"
 #include "Engine/StaticMesh.h"
+#include "Materials/MaterialRenderProxy.h"
 #include "Materials/MaterialInstanceConstant.h"
 #include "Modules/ModuleManager.h"
+#include "MaterialDomain.h"
+#include "SceneManagement.h"
 
 #include UE_INLINE_GENERATED_CPP_BY_NAME(NiagaraMeshRendererProperties)
 
@@ -87,6 +88,11 @@ public:
 		}
 	}
 
+	virtual int32 ComputeLOD(const FVector& SphereOrigin, const float SphereRadius, const FSceneView& SceneView, float LODDistanceFactor) override
+	{
+		return ComputeStaticMeshLOD(RenderData, SphereOrigin, SphereRadius, SceneView, MinLOD, LODDistanceFactor);
+	}
+
 	static void InitVertexFactoryComponents(
 		const FStaticMeshVertexBuffers& VertexBuffers,
 		FNiagaraMeshVertexFactory* VertexFactory,
diff --git a/Engine/Plugins/FX/Niagara/Source/Niagara/Private/NiagaraRendererMeshes.cpp b/Engine/Plugins/FX/Niagara/Source/Niagara/Private/NiagaraRendererMeshes.cpp
--- a/Engine/Plugins/FX/Niagara/Source/Niagara/Private/NiagaraRendererMeshes.cpp
+++ b/Engine/Plugins/FX/Niagara/Source/Niagara/Private/NiagaraRendererMeshes.cpp
@@ -227,8 +227,8 @@ void FNiagaraRendererMeshes::Initialize(const UNiagaraRendererProperties* InProp
 		FMeshData& MeshData = Meshes.AddDefaulted_GetRef();
 		MeshData.RenderableMesh = RenderableMesh;
 		MeshData.SourceMeshIndex = SourceMeshIndex;
+		MeshData.LODMode = MeshProperties.LODMode;
 		MeshData.LODLevel = 0;
-		MeshData.LODBias = 0;
 		MeshData.PivotOffset = FVector3f(MeshProperties.PivotOffset);
 		MeshData.PivotOffsetSpace = MeshProperties.PivotOffsetSpace;
 		MeshData.Scale = FVector3f(MeshProperties.Scale);
@@ -244,7 +244,7 @@ void FNiagaraRendererMeshes::Initialize(const UNiagaraRendererProperties* InProp
 		}
 		else		
 		{
-			MeshData.LODBias = MeshProperties.LODBias;
+			MeshData.LODDistanceFactor = MeshProperties.LODDistanceFactor;
 		}
 
 		// Get materials and remap them into the base material list
@@ -1307,7 +1307,17 @@ void FNiagaraRendererMeshes::GetDynamicMeshElements(const TArray<const FSceneVie
 				const FMeshData& MeshData = Meshes[MeshIndex];
 
 				INiagaraRenderableMesh::FLODModelData LODModel;
-				MeshData.RenderableMesh->GetLODModelData(LODModel, MeshData.LODLevel);
+				if (MeshData.LODMode == ENiagaraMeshLODMode::ByComponentBounds)
+				{
+					const FBoxSphereBounds ProxyBounds = SceneProxy->GetBounds();
+					const int32 LODLevel = MeshData.RenderableMesh->ComputeLOD(ProxyBounds.Origin, ProxyBounds.SphereRadius, *View, MeshData.LODDistanceFactor);
+					MeshData.RenderableMesh->GetLODModelData(LODModel, LODLevel);
+				}
+				else
+				{
+					MeshData.RenderableMesh->GetLODModelData(LODModel, MeshData.LODLevel);
+				}
+
 				if (LODModel.LODIndex == INDEX_NONE)
 				{
 					continue;
diff --git a/Engine/Plugins/FX/Niagara/Source/Niagara/Public/NiagaraMeshRendererProperties.h b/Engine/Plugins/FX/Niagara/Source/Niagara/Public/NiagaraMeshRendererProperties.h
--- a/Engine/Plugins/FX/Niagara/Source/Niagara/Public/NiagaraMeshRendererProperties.h
+++ b/Engine/Plugins/FX/Niagara/Source/Niagara/Public/NiagaraMeshRendererProperties.h
@@ -72,15 +72,17 @@ enum class ENiagaraMeshLODMode : uint8
 	LODBias,
 
 	/*
-	* Calculates the LOD level for the instance based on screen space size and renders all particles with this LOD.
-	* LOD bias is applied to the calculation.
+	* The LOD level is calculated based on screen space size of the component bounds.
+	* All particles will be rendered with the same calculated LOD level.
+	* Increasing 'LOD calculation scale' will result in lower quality LODs being used, this is useful as component bounds generally are larger than the particle mesh bounds.
 	*/
-	//PerInstance,
+	ByComponentBounds,
 
 	/*
+	* The LOD level is calcuated per particle using the particle position and mesh sphere bounds.
+	* This involves running a dispatch & draw per LOD level.
 	* Calculates and renders each particle with it's calcualted LOD level.
-	* LOD bias is applied to the calculation.
-	* Note: This is the most expensive option as we need to bucket the particles into each LOD and dispatch draws per bucket.
+	* Increasing 'LOD calculation scale' will result in lower quality LODs being used.
 	*/
 	//PerParticle,
 };
@@ -192,7 +194,7 @@ struct FNiagaraMeshRendererMeshProperties
 	FNiagaraParameterBindingWithValue LODLevelBinding;
 
 	/* LOD bias to apply to the LOD calculation. */
-	UPROPERTY(EditAnywhere, Category = "Mesh", meta = (UIMin = "0", DisplayName = "LOD Bias", EditCondition = "LODMode != ENiagaraMeshLODMode::LODLevel", EditConditionHides))
+	UPROPERTY(EditAnywhere, Category = "Mesh", meta = (UIMin = "0", DisplayName = "LOD Bias", EditCondition = "LODMode == ENiagaraMeshLODMode::LODBias", EditConditionHides))
 	FNiagaraParameterBindingWithValue LODBiasBinding;
 #endif
 
@@ -202,6 +204,10 @@ struct FNiagaraMeshRendererMeshProperties
 	UPROPERTY()
 	int32 LODBias = 0;
 
+	/** Used in LOD calculation to modify the distance, i.e. increasing the value will make lower poly LODs transition closer to the camera. */
+	UPROPERTY(EditAnywhere, Category = "Mesh", meta = (UIMin = "0", DisplayName = "LOD Distance Factor", EditCondition = "LODMode == ENiagaraMeshLODMode::ByComponentBounds || LODMode == ENiagaraMeshLODMode::PerParticle", EditConditionHides))
+	float LODDistanceFactor = 1.0f;
+
 	/** Scale of the mesh */
 	UPROPERTY(EditAnywhere, Category = "Mesh")
 	FVector Scale;
diff --git a/Engine/Plugins/FX/Niagara/Source/Niagara/Public/NiagaraRenderableMeshInterface.h b/Engine/Plugins/FX/Niagara/Source/Niagara/Public/NiagaraRenderableMeshInterface.h
--- a/Engine/Plugins/FX/Niagara/Source/Niagara/Public/NiagaraRenderableMeshInterface.h
+++ b/Engine/Plugins/FX/Niagara/Source/Niagara/Public/NiagaraRenderableMeshInterface.h
@@ -14,6 +14,7 @@ class FNiagaraMeshVertexFactory;
 class UMaterialInterface;
 class FRayTracingGeometry;
 class FRHIUniformBuffer;
+class FSceneView;
 struct FStaticMeshSection;
 
 // Abstact class that for a renderable mesh
@@ -47,6 +48,8 @@ public:
 	virtual void SetupVertexFactory(class FNiagaraMeshVertexFactory& InVertexFactory, const FLODModelData& LODModelData) const = 0;
 	// Gather a list of used materials
 	virtual void GetUsedMaterials(TArray<UMaterialInterface*>& OutMaterials) const = 0;
+
+	virtual int32 ComputeLOD(const FVector& SphereOrigin, const float SphereRadius, const FSceneView& SceneView, float LODDistanceFactor) { return 0; }
 };
 
 using FNiagaraRenderableMeshPtr = TSharedPtr<INiagaraRenderableMesh, ESPMode::ThreadSafe>;
diff --git a/Engine/Plugins/FX/Niagara/Source/Niagara/Public/NiagaraRendererMeshes.h b/Engine/Plugins/FX/Niagara/Source/Niagara/Public/NiagaraRendererMeshes.h
--- a/Engine/Plugins/FX/Niagara/Source/Niagara/Public/NiagaraRendererMeshes.h
+++ b/Engine/Plugins/FX/Niagara/Source/Niagara/Public/NiagaraRendererMeshes.h
@@ -83,8 +83,9 @@ protected:
 	{
 		FNiagaraRenderableMeshPtr RenderableMesh;
 		uint32 SourceMeshIndex = INDEX_NONE;
+		ENiagaraMeshLODMode LODMode = ENiagaraMeshLODMode::LODLevel;
 		int32 LODLevel = 0;
-		int32 LODBias = 0;
+		float LODDistanceFactor = 1.0f;
 		FVector3f PivotOffset = FVector3f::ZeroVector;
 		ENiagaraMeshPivotOffsetSpace PivotOffsetSpace = ENiagaraMeshPivotOffsetSpace::Mesh;
 		FVector3f Scale = FVector3f(1.0f, 1.0f, 1.0f);
