- Add ability to capture data interface usage either all or a selection to audit command let

diff --git a/Engine/Plugins/FX/Niagara/Source/NiagaraEditor/Private/Commandlets/NiagaraSystemAuditCommandlet.cpp b/Engine/Plugins/FX/Niagara/Source/NiagaraEditor/Private/Commandlets/NiagaraSystemAuditCommandlet.cpp
--- a/Engine/Plugins/FX/Niagara/Source/NiagaraEditor/Private/Commandlets/NiagaraSystemAuditCommandlet.cpp
+++ b/Engine/Plugins/FX/Niagara/Source/NiagaraEditor/Private/Commandlets/NiagaraSystemAuditCommandlet.cpp
@@ -129,7 +129,28 @@ int32 UNiagaraSystemAuditCommandlet::Main(const FString& Params)
 
 	FParse::Bool(*Params, TEXT("RendererDetailed="), bRendererDetailed);
 
-	FParse::Bool(*Params, TEXT("CaptureDataInterfaceUsage="), bCaptureDataInterfaceUsage);
+	// Disable on specific platforms
+	// Example To Capture Specific: -run=NiagaraSystemAuditCommandlet -CaptureDataInterfaceUsage=MyDataInterface
+	// Example To Capture All: -run=NiagaraSystemAuditCommandlet -CaptureDataInterfaceUsage
+	{
+		FString DINamesArrayString;
+		if (FParse::Value(*Params, TEXT("CaptureDataInterfaceUsage="), DINamesArrayString, false))
+		{
+			bCaptureDataInterfaceUsage = true;
+
+			TArray<FString> DINamesArray;
+			DINamesArrayString.ParseIntoArray(DINamesArray, TEXT(","));
+
+			for (FString DIString : DINamesArray)
+			{
+				NiagaraDataInterfaceUsageToCapture.Add(FName(*DIString));
+			}
+		}
+		else
+		{
+			bCaptureDataInterfaceUsage = FParse::Param(*Params, TEXT("CaptureDataInterfaceUsage"));
+		}
+	}
 
 	// Validation enabled?
 	bool bRunValidation = false;
@@ -234,6 +255,8 @@ bool UNiagaraSystemAuditCommandlet::ProcessNiagaraSystems()
 			continue;
 		}
 
+		NiagaraSystem->WaitForCompilationComplete(true, false);
+
 		// Iterate over all data interfaces used by the system / emitters
 		TSet<FName> SystemDataInterfacesWihPrereqs;
 		TSet<FName> SystemUserDataInterfaces;
@@ -241,9 +264,13 @@ bool UNiagaraSystemAuditCommandlet::ProcessNiagaraSystems()
 		{
 			if (bCaptureDataInterfaceUsage)
 			{
-				FDataInterfaceUsage& DataInterfaceUsage = NiagaraDataInterfaceUsage.FindOrAdd(DataInterface->GetClass()->GetFName());
-				++DataInterfaceUsage.UsageCount;
-				DataInterfaceUsage.Systems.Add(NiagaraSystem->GetFName());
+				const FName DIClassName = DataInterface->GetClass()->GetFName();
+				if ( NiagaraDataInterfaceUsageToCapture.IsEmpty() || NiagaraDataInterfaceUsageToCapture.Contains(DIClassName) )
+				{
+					FDataInterfaceUsage& DataInterfaceUsage = NiagaraDataInterfaceUsage.FindOrAdd(DataInterface->GetClass()->GetFName());
+					++DataInterfaceUsage.UsageCount;
+					DataInterfaceUsage.Systems.Add(NiagaraSystem->GetFName());
+				}
 			}
 
 			if (DataInterface->HasTickGroupPrereqs())
diff --git a/Engine/Plugins/FX/Niagara/Source/NiagaraEditor/Public/Commandlets/NiagaraSystemAuditCommandlet.h b/Engine/Plugins/FX/Niagara/Source/NiagaraEditor/Public/Commandlets/NiagaraSystemAuditCommandlet.h
--- a/Engine/Plugins/FX/Niagara/Source/NiagaraEditor/Public/Commandlets/NiagaraSystemAuditCommandlet.h
+++ b/Engine/Plugins/FX/Niagara/Source/NiagaraEditor/Public/Commandlets/NiagaraSystemAuditCommandlet.h
@@ -39,6 +39,7 @@ class UNiagaraSystemAuditCommandlet : public UCommandlet
 		TSet<FName>		Systems;
 	};
 	bool bCaptureDataInterfaceUsage = false;
+	TSet<FName> NiagaraDataInterfaceUsageToCapture;
 	TMap<FName, FDataInterfaceUsage>  NiagaraDataInterfaceUsage;
 
 	/** Detailed renderer information. */
