Fix for how digesting niagara graphs deals with both duplicate connections (multiple connections going to an input pin) and connections to bogus nodes (nodes that aren't present in the UNiagaraGraph::Nodes list).

diff --git a/Engine/Plugins/FX/Niagara/Source/NiagaraEditor/Private/NiagaraGraphDigest.cpp b/Engine/Plugins/FX/Niagara/Source/NiagaraEditor/Private/NiagaraGraphDigest.cpp
--- a/Engine/Plugins/FX/Niagara/Source/NiagaraEditor/Private/NiagaraGraphDigest.cpp
+++ b/Engine/Plugins/FX/Niagara/Source/NiagaraEditor/Private/NiagaraGraphDigest.cpp
@@ -573,18 +573,33 @@ void FNiagaraCompilationGraph::Create(const UNiagaraGraph* InGraph, const FNiaga
 		for (FNiagaraCompilationInputPin& InputPin : CompilationNode->InputPins)
 		{
 			const UEdGraphPin* SourceInputPin = CompilationNode->SourceNode->Pins[InputPin.SourcePinIndex];
+			const UEdGraphPin* SourceLinkedPin = nullptr;
 
-			check(SourceInputPin->LinkedTo.Num() <= 1);
+			// apparently some content exists where we'll have multiple copies of a connection added.  Find the first non-null LinkedPin
+			for (const UEdGraphPin* CurrentLinkedPin : SourceInputPin->LinkedTo)
+			{
+				if (SourceLinkedPin == nullptr)
+				{
+					SourceLinkedPin = CurrentLinkedPin;
+					break;
+				}
+			}
 
-			if (!SourceInputPin->LinkedTo.IsEmpty())
+			if (SourceLinkedPin)
 			{
-				if (const UEdGraphPin* LinkedPin = TraceThroughIgnoredNodes(SourceInputPin->LinkedTo[0]))
+				if (const UEdGraphPin* LinkedPin = TraceThroughIgnoredNodes(SourceLinkedPin))
 				{
 					const UEdGraphNode* LinkedSourceNode = LinkedPin->GetOwningNode();
 					const int32* LinkedNodeIndexPtr = NodeIndexMap.Find(LinkedSourceNode);
-					if (ensure(LinkedNodeIndexPtr && Nodes.IsValidIndex(*LinkedNodeIndexPtr)))
+
+					// There are situations where a pin is connected to a node that doesn't exist within the UNiagaraGraph::Nodes array.
+					// This corruption seems connected to auto-generated input nodes for function calls.  For now we're going to ignore those
+					// connections
+					const int32 LinkedNodeIndex = (LinkedNodeIndexPtr && Nodes.IsValidIndex(*LinkedNodeIndexPtr)) ? *LinkedNodeIndexPtr : INDEX_NONE;
+
+					if (LinkedNodeIndex != INDEX_NONE)
 					{
-						TUniquePtr<FNiagaraCompilationNode>& LinkedNode = Nodes[*LinkedNodeIndexPtr];
+						TUniquePtr<FNiagaraCompilationNode>& LinkedNode = Nodes[LinkedNodeIndex];
 						const int32 LinkedSourcePinIndex = LinkedNode->SourceNode->Pins.IndexOfByKey(LinkedPin);
 						if (LinkedSourcePinIndex != INDEX_NONE)
 						{
@@ -1645,6 +1660,7 @@ FNiagaraCompilationNode::FNiagaraCompilationNode(ENodeType InNodeType, const UEd
 
 		SourceNode = InNode;
 		OwningGraph = &Context.ParentGraph;
+		const UEdGraph* SourceOwningGraph = OwningGraph->SourceGraph.Get();
 
 		const int32 PinCount = SourceNode->Pins.Num();
 
@@ -1665,6 +1681,22 @@ FNiagaraCompilationNode::FNiagaraCompilationNode(ENodeType InNodeType, const UEd
 				continue;
 			}
 
+			// in the case of the pin being connected to a node that is not in the graph's Node array we just skip the pin
+			bool bConnectedToInvalidNode = false;
+			for (const UEdGraphPin* LinkedPin : Pin->LinkedTo)
+			{
+				if (!SourceOwningGraph->Nodes.Contains(LinkedPin->GetOwningNode()))
+				{
+					bConnectedToInvalidNode = true;
+					break;
+				}
+			}
+
+			if (bConnectedToInvalidNode)
+			{
+				continue;
+			}
+
 			FNiagaraCompilationPin* CompilationPin = nullptr;
 
 			if (Pin->Direction == EEdGraphPinDirection::EGPD_Input)
