- Pixel Coverage for translucent now uses A rather than RGB
- This fixes an issue with particle darkening as pixel coverage kicks in rather than fading out

diff --git a/Engine/Plugins/FX/Niagara/Source/Niagara/Private/NiagaraRendererSprites.cpp b/Engine/Plugins/FX/Niagara/Source/Niagara/Private/NiagaraRendererSprites.cpp
--- a/Engine/Plugins/FX/Niagara/Source/Niagara/Private/NiagaraRendererSprites.cpp
+++ b/Engine/Plugins/FX/Niagara/Source/Niagara/Private/NiagaraRendererSprites.cpp
@@ -597,7 +597,7 @@ FNiagaraSpriteUniformBufferRef FNiagaraRendererSprites::CreateViewUniformBuffer(
 				if (IsTranslucentOnlyBlendMode(ParticleSpriteRenderData.BlendMode))
 				{
 					ParticleSpriteRenderData.bHasTranslucentMaterials = true;
-					PerViewUniformParameters.PixelCoverageColorBlend = FVector4f(PixelCoverageBlend, PixelCoverageBlend, PixelCoverageBlend, 0.0f);
+					PerViewUniformParameters.PixelCoverageColorBlend = FVector4f(0.0f, 0.0f, 0.0f, PixelCoverageBlend);
 				}
 				else if (IsAdditiveBlendMode(ParticleSpriteRenderData.BlendMode))
 				{
