- Niagara per particle mesh LODs
- This is disabled if we can not run the sort / cull process (i.e. OpenGL RHI) as we don't have a CPU LOD path currently

diff --git a/Engine/Plugins/FX/Niagara/Shaders/Private/NiagaraSortKeyGen.usf b/Engine/Plugins/FX/Niagara/Shaders/Private/NiagaraSortKeyGen.usf
--- a/Engine/Plugins/FX/Niagara/Shaders/Private/NiagaraSortKeyGen.usf
+++ b/Engine/Plugins/FX/Niagara/Shaders/Private/NiagaraSortKeyGen.usf
@@ -48,6 +48,7 @@ uint RendererVisTagAttributeOffset;
 int MeshIndex;
 uint MeshIndexAttributeOffset;
 float2 CullDistanceRangeSquared;
+float4 LODScreenSize;
 float4 LocalBoundingSphere;
 float3 CullingWorldSpaceOffset;
 float3 SystemLWCTile;
@@ -171,6 +172,13 @@ bool CullParticle(in uint ParticleIndex)
 		return true;
 	}
 
+	// Cull based on LOD level
+	const float ScreenRadius = saturate(LODScreenSize.z / max(1.0f, CamDistSq * LODScreenSize.w));
+	if ( ScreenRadius < LODScreenSize.x || ScreenRadius >= LODScreenSize.y )
+	{
+		return true;
+	}
+	
 	// Cull based on renderer visibility
 	if (RendererVisTagAttributeOffset != -1)
 	{
diff --git a/Engine/Plugins/FX/Niagara/Shaders/Private/NiagaraUpdateMeshGPUSceneInstances.usf b/Engine/Plugins/FX/Niagara/Shaders/Private/NiagaraUpdateMeshGPUSceneInstances.usf
--- a/Engine/Plugins/FX/Niagara/Shaders/Private/NiagaraUpdateMeshGPUSceneInstances.usf
+++ b/Engine/Plugins/FX/Niagara/Shaders/Private/NiagaraUpdateMeshGPUSceneInstances.usf
@@ -57,6 +57,7 @@ int VisibilityTagDataOffset;
 
 float3 LocalBoundingCenter;
 float2 DistanceCullRangeSquared;
+float4 LODScreenSize;
 
 int bNeedsPrevTransform;
 
@@ -166,18 +167,25 @@ void UpdateMeshInstancesCS(uint3 GroupId : SV_GroupID, uint GroupIndex : SV_Grou
 	Params.PrimitiveInvNonUniformScale	= PrimitiveData.InvNonUniformScale;
 
 	NiagaraMeshParticleTransforms Transforms = NiagaraCalculateMeshParticleTransforms(Params);
-
+	
 	// Transform the bounding center to world-space and distance cull
 	{
-		FLWCVector3 BCenter = LWCMultiply(LocalBoundingCenter, Transforms.LocalToWorld);
-		float3 CamVec = LWCToFloat(LWCSubtract(BCenter, PrimaryView.WorldCameraOrigin));
-		float CamDistSq = dot(CamVec, CamVec);
+		const FLWCVector3 BCenter = LWCMultiply(LocalBoundingCenter, Transforms.LocalToWorld);
+		const float3 CamVec = LWCToFloat(LWCSubtract(BCenter, PrimaryView.WorldCameraOrigin));
+		const float CamDistSq = length2(CamVec);
 		if (CamDistSq < DistanceCullRangeSquared.x  || CamDistSq > DistanceCullRangeSquared.y)
 		{
 			// We are outside the distance ranges required
 			CullInstance(InstanceId);
 			return;
 		}
+		
+		const float ScreenRadius = saturate(LODScreenSize.z / max(1.0f, CamDistSq * LODScreenSize.w));
+		if ( ScreenRadius < LODScreenSize.x || ScreenRadius >= LODScreenSize.y )
+		{
+			CullInstance(InstanceId);
+			return;
+		}
 	}
 
 	const uint PayloadDataFlags = GetPayloadDataFlags();
diff --git a/Engine/Plugins/FX/Niagara/Source/Niagara/Classes/NiagaraGPUSortInfo.h b/Engine/Plugins/FX/Niagara/Source/Niagara/Classes/NiagaraGPUSortInfo.h
--- a/Engine/Plugins/FX/Niagara/Source/Niagara/Classes/NiagaraGPUSortInfo.h
+++ b/Engine/Plugins/FX/Niagara/Source/Niagara/Classes/NiagaraGPUSortInfo.h
@@ -70,6 +70,7 @@ struct FNiagaraGPUSortInfo
 	FVector CullingWorldSpaceOffset = FVector(0, 0, 0);
 	FVector3f SystemLWCTile = FVector3f(0, 0, 0);
 	FVector2f DistanceCullRange { 0.0f, FLT_MAX };
+	FVector4f LODScreenSize = FVector4f(0.0f, 2.0f, 1.0f, 0.0f);
 	TArray<FPlane, TFixedAllocator<MaxCullPlanes>> CullPlanes;
 
 	// The GPUSortManager bindings for this sort task.
diff --git a/Engine/Plugins/FX/Niagara/Source/Niagara/Private/NiagaraGpuComputeDispatch.cpp b/Engine/Plugins/FX/Niagara/Source/Niagara/Private/NiagaraGpuComputeDispatch.cpp
--- a/Engine/Plugins/FX/Niagara/Source/Niagara/Private/NiagaraGpuComputeDispatch.cpp
+++ b/Engine/Plugins/FX/Niagara/Source/Niagara/Private/NiagaraGpuComputeDispatch.cpp
@@ -2287,6 +2287,7 @@ void FNiagaraGpuComputeDispatch::GenerateSortKeys(FRHICommandListImmediate& RHIC
 			Params.MeshIndex = SortInfo.MeshIndex;
 			Params.MeshIndexAttributeOffset = SortInfo.MeshIndexAttributeOffset;
 			Params.CullDistanceRangeSquared = SortInfo.DistanceCullRange * SortInfo.DistanceCullRange;
+			Params.LODScreenSize = SortInfo.LODScreenSize;
 			Params.LocalBoundingSphere = FVector4f((FVector3f)SortInfo.LocalBSphere.Center, SortInfo.LocalBSphere.W);
 			Params.CullingWorldSpaceOffset = (FVector3f)SortInfo.CullingWorldSpaceOffset;
 			Params.SystemLWCTile = SortInfo.SystemLWCTile;
diff --git a/Engine/Plugins/FX/Niagara/Source/Niagara/Private/NiagaraMeshRendererProperties.cpp b/Engine/Plugins/FX/Niagara/Source/Niagara/Private/NiagaraMeshRendererProperties.cpp
--- a/Engine/Plugins/FX/Niagara/Source/Niagara/Private/NiagaraMeshRendererProperties.cpp
+++ b/Engine/Plugins/FX/Niagara/Source/Niagara/Private/NiagaraMeshRendererProperties.cpp
@@ -88,6 +88,22 @@ public:
 		}
 	}
 
+	virtual FIntVector2 GetLODRange() const override
+	{
+		return FIntVector2(MeshMinLOD, RenderData->LODResources.Num());
+	}
+
+	virtual FVector3f GetLODScreenSize(int32 LODLevel) const override
+	{
+		constexpr int32 MaxLODLevel = MAX_STATIC_MESH_LODS - 1;
+		LODLevel = FMath::Clamp(LODLevel, 0, MaxLODLevel);
+		return FVector3f(
+			LODLevel < MaxLODLevel ? RenderData->ScreenSize[LODLevel + 1].GetValue() : 0.0f,
+			RenderData->ScreenSize[LODLevel].GetValue(),
+			RenderData->Bounds.SphereRadius
+		);
+	}
+
 	virtual int32 ComputeLOD(const FVector& SphereOrigin, const float SphereRadius, const FSceneView& SceneView, float LODDistanceFactor) override
 	{
 		return ComputeStaticMeshLOD(RenderData, SphereOrigin, SphereRadius, SceneView, MinLOD, LODDistanceFactor);
@@ -218,6 +234,7 @@ FNiagaraMeshRendererMeshProperties::FNiagaraMeshRendererMeshProperties()
 #if WITH_EDITORONLY_DATA
 	, UserParamBinding_DEPRECATED(FNiagaraTypeDefinition(UStaticMesh::StaticClass()))
 #endif
+	, LODRange(0, MAX_STATIC_MESH_LODS)
 	, Scale(1.0f, 1.0f, 1.0f)
 	, Rotation(FRotator::ZeroRotator)
 	, PivotOffset(ForceInitToZero)
@@ -313,10 +330,17 @@ FNiagaraRenderer* UNiagaraMeshRendererProperties::CreateEmitterRenderer(ERHIFeat
 	{
 		if (MeshProperties.HasValidRenderableMesh())
 		{
-			// There's at least one valid mesh
-			FNiagaraRenderer* NewRenderer = new FNiagaraRendererMeshes(FeatureLevel, this, Emitter);
+			FNiagaraRendererMeshes* NewRenderer = new FNiagaraRendererMeshes(FeatureLevel, this, Emitter);
 			NewRenderer->Initialize(this, Emitter, InController);
-			return NewRenderer;
+			if (NewRenderer->HasValidMeshes())
+			{
+				return NewRenderer;
+			}
+
+			// There are cases where we might end up with no meshes to render due to LODs or features not being enabled on that platform
+			// so we discard the renderer here, the cost to do this work in HasValidRenderableMesh makes it not worthwhile
+			delete NewRenderer;
+			return nullptr;
 		}
 	}
 
@@ -1179,6 +1203,21 @@ void UNiagaraMeshRendererProperties::PostEditChangeProperty(FPropertyChangedEven
 		}
 	}
 
+	if (MemberPropertyName == GET_MEMBER_NAME_CHECKED(UNiagaraMeshRendererProperties, Meshes))
+	{
+		for (FNiagaraMeshRendererMeshProperties& MeshProperties : Meshes)
+		{
+			if (MeshProperties.bUseLODRange)
+			{
+				MeshProperties.LODRange.X = FMath::Clamp(MeshProperties.LODRange.X, 0, MAX_STATIC_MESH_LODS - 1);
+				MeshProperties.LODRange.Y = FMath::Clamp(MeshProperties.LODRange.Y, 1, MAX_STATIC_MESH_LODS);
+
+				MeshProperties.LODRange.X = FMath::Clamp(MeshProperties.LODRange.X, 0, MeshProperties.LODRange.Y - 1);
+				MeshProperties.LODRange.Y = FMath::Clamp(MeshProperties.LODRange.Y, MeshProperties.LODRange.X + 1, MAX_STATIC_MESH_LODS);
+			}
+		}
+	}
+
 	// Update our MICs if we change override material / material bindings / meshes
 	//-OPT: Could narrow down further to only static materials
 	if ((MemberPropertyName == GET_MEMBER_NAME_CHECKED(UNiagaraMeshRendererProperties, OverrideMaterials)) ||
diff --git a/Engine/Plugins/FX/Niagara/Source/Niagara/Private/NiagaraRendererMeshes.cpp b/Engine/Plugins/FX/Niagara/Source/Niagara/Private/NiagaraRendererMeshes.cpp
--- a/Engine/Plugins/FX/Niagara/Source/Niagara/Private/NiagaraRendererMeshes.cpp
+++ b/Engine/Plugins/FX/Niagara/Source/Niagara/Private/NiagaraRendererMeshes.cpp
@@ -106,6 +106,7 @@ FNiagaraRendererMeshes::FNiagaraRendererMeshes(ERHIFeatureLevel::Type FeatureLev
 	SubImageSize = FVector2f(Properties->SubImageSize);	// LWC_TODO: Precision loss
 	bSubImageBlend = Properties->bSubImageBlend;
 	bEnableFrustumCulling = Properties->bEnableFrustumCulling;
+	bEnableLODCulling = false;
 	bEnableCulling = bEnableFrustumCulling;
 	DistanceCullRange = FVector2f(0, FLT_MAX);
 	DistanceCullRangeSquared = FVector2f(0, FLT_MAX);
@@ -128,7 +129,7 @@ FNiagaraRendererMeshes::FNiagaraRendererMeshes(ERHIFeatureLevel::Type FeatureLev
 	{
 		DistanceCullRange = FVector2f(Properties->MinCameraDistance, Properties->MaxCameraDistance);
 		DistanceCullRangeSquared = DistanceCullRange* DistanceCullRange;
-		bEnableCulling = true;
+		bEnableCulling |= true;
 	}
 
 	// Ensure valid value for the locked axis
@@ -161,7 +162,7 @@ FNiagaraRendererMeshes::FNiagaraRendererMeshes(ERHIFeatureLevel::Type FeatureLev
 	}
 	else if (Data.GetVariableComponentOffsets(Properties->MeshIndexBinding.GetDataSetBindableVariable(), FloatOffset, ParticleMeshIndexOffset, HalfOffset))
 	{
-		bEnableCulling = ParticleMeshIndexOffset != INDEX_NONE;
+		bEnableCulling |= ParticleMeshIndexOffset != INDEX_NONE;
 	}
 
 	MaterialParamValidMask = Properties->MaterialParamValidMask;
@@ -222,58 +223,77 @@ void FNiagaraRendererMeshes::Initialize(const UNiagaraRendererProperties* InProp
 			continue;
 		}
 
-
-		// We have a valid mesh fill in the details
-		FMeshData& MeshData = Meshes.AddDefaulted_GetRef();
-		MeshData.RenderableMesh = RenderableMesh;
-		MeshData.SourceMeshIndex = SourceMeshIndex;
-		MeshData.LODMode = MeshProperties.LODMode;
-		MeshData.LODLevel = 0;
-		MeshData.PivotOffset = FVector3f(MeshProperties.PivotOffset);
-		MeshData.PivotOffsetSpace = MeshProperties.PivotOffsetSpace;
-		MeshData.Scale = FVector3f(MeshProperties.Scale);
-		MeshData.Rotation = FQuat4f(MeshProperties.Rotation.Quaternion());
-
+		FIntVector2 LODRange = FIntVector2(0, 1);
 		if (MeshProperties.LODMode == ENiagaraMeshLODMode::LODLevel)
 		{
-			MeshData.LODLevel = MeshProperties.LODLevel;
+			LODRange.X = MeshProperties.LODLevel;
+			LODRange.Y = LODRange.X + 1;
 		}
-		else if (MeshProperties.LODMode == ENiagaraMeshLODMode::LODBias)
+		else if (MeshProperties.LODMode == ENiagaraMeshLODMode::PerParticle)
 		{
-			RenderableMesh->SetMinLODBias(MeshProperties.LODBias);
+			LODRange = RenderableMesh->GetLODRange();
+			if (MeshProperties.bUseLODRange)
+			{
+				LODRange.X = FMath::Max(MeshProperties.LODRange.X, LODRange.X);
+				LODRange.Y = FMath::Max(MeshProperties.LODRange.Y, LODRange.Y);
+			}
+			bEnableLODCulling |= LODRange.Y - LODRange.X > 1;
 		}
-		else		
+
+		for (int32 LOD=LODRange.X; LOD < LODRange.Y; ++LOD)
 		{
-			MeshData.LODDistanceFactor = MeshProperties.LODDistanceFactor;
-		}
+			// We have a valid mesh fill in the details
+			FMeshData& MeshData = Meshes.AddDefaulted_GetRef();
+			MeshData.RenderableMesh = RenderableMesh;
+			MeshData.SourceMeshIndex = SourceMeshIndex;
+			MeshData.LODMode = MeshProperties.LODMode;
+			MeshData.LODLevel = LOD;
+			MeshData.PivotOffset = FVector3f(MeshProperties.PivotOffset);
+			MeshData.PivotOffsetSpace = MeshProperties.PivotOffsetSpace;
+			MeshData.Scale = FVector3f(MeshProperties.Scale);
+			MeshData.Rotation = FQuat4f(MeshProperties.Rotation.Quaternion());
+
+			if (MeshProperties.LODMode == ENiagaraMeshLODMode::LODBias)
+			{
+				RenderableMesh->SetMinLODBias(MeshProperties.LODBias);
+			}
+			else if (MeshProperties.LODMode == ENiagaraMeshLODMode::ByComponentBounds || MeshProperties.LODMode == ENiagaraMeshLODMode::PerParticle)
+			{
+				MeshData.LODDistanceFactor = MeshProperties.LODDistanceFactor;
+				if (MeshProperties.LODMode == ENiagaraMeshLODMode::PerParticle)
+				{
+					MeshData.LODScreenSize = RenderableMesh->GetLODScreenSize(LOD);
+					MeshData.LODScreenSize.X = LOD < LODRange.Y - 1 ? MeshData.LODScreenSize.X : 0.0f;
+					MeshData.LODScreenSize.Y = LOD > LODRange.X ? MeshData.LODScreenSize.Y : 2.0f;
+				}
+			}
 
-		// Get materials and remap them into the base material list
-		TArray<UMaterialInterface*> UsedMaterials;
-		RenderableMesh->GetUsedMaterials(UsedMaterials);
-		Properties->ApplyMaterialOverrides(Emitter, UsedMaterials);
+			// Get materials and remap them into the base material list
+			TArray<UMaterialInterface*> UsedMaterials;
+			RenderableMesh->GetUsedMaterials(UsedMaterials);
+			Properties->ApplyMaterialOverrides(Emitter, UsedMaterials);
 
-		for (UMaterialInterface* UsedMaterial : UsedMaterials)
-		{
-			MeshData.MaterialRemapTable.Add(
-				BaseMaterials_GT.IndexOfByPredicate(
-					[&](UMaterialInterface* LookMat)
-					{
-						if (LookMat == UsedMaterial)
-						{
-							return true;
-						}
-						if (UMaterialInstanceDynamic* MID = Cast<UMaterialInstanceDynamic>(LookMat))
+			for (UMaterialInterface* UsedMaterial : UsedMaterials)
+			{
+				MeshData.MaterialRemapTable.Add(
+					BaseMaterials_GT.IndexOfByPredicate(
+						[&](UMaterialInterface* LookMat)
 						{
-							return UsedMaterial == MID->Parent;
+							if (LookMat == UsedMaterial)
+							{
+								return true;
+							}
+							if (UMaterialInstanceDynamic* MID = Cast<UMaterialInstanceDynamic>(LookMat))
+							{
+								return UsedMaterial == MID->Parent;
+							}
+							return false;
 						}
-						return false;
-					}
-				)
-			);
+					)
+				);
+			}
 		}
 	}
-
-	checkf(Meshes.Num() > 0, TEXT("At least one valid mesh is required to instantiate a mesh renderer"));
 }
 
 void FNiagaraRendererMeshes::ReleaseRenderThreadResources()
@@ -384,7 +404,8 @@ void FNiagaraRendererMeshes::PrepareParticleMeshRenderData(FParticleMeshRenderDa
 		}
 
 		// Do we need culling? (When using GPU Scene, we don't cull in the sort passes)
-		ParticleMeshRenderData.bNeedsCull = bEnableCulling && !ParticleMeshRenderData.bUseGPUScene;
+		ParticleMeshRenderData.bNeedsCull = (bEnableCulling || bEnableLODCulling) && !ParticleMeshRenderData.bUseGPUScene;
+		ParticleMeshRenderData.bAllowPerParticleMeshLODs = bEnableLODCulling;
 		ParticleMeshRenderData.bSortCullOnGpu = (ParticleMeshRenderData.bNeedsSort && FNiagaraUtilities::AllowGPUSorting(ShaderPlatform)) || (ParticleMeshRenderData.bNeedsCull && FNiagaraUtilities::AllowGPUCulling(ShaderPlatform));
 
 		// Validate what we setup
@@ -393,6 +414,7 @@ void FNiagaraRendererMeshes::PrepareParticleMeshRenderData(FParticleMeshRenderDa
 			if (!ensureMsgf(!ParticleMeshRenderData.bNeedsCull || ParticleMeshRenderData.bSortCullOnGpu, TEXT("Culling is requested on GPU but we don't support sorting, this will result in incorrect rendering.")))
 			{
 				ParticleMeshRenderData.bNeedsCull = false;
+				ParticleMeshRenderData.bAllowPerParticleMeshLODs = false;
 			}
 			ParticleMeshRenderData.bNeedsSort &= ParticleMeshRenderData.bSortCullOnGpu;
 
@@ -402,6 +424,7 @@ void FNiagaraRendererMeshes::PrepareParticleMeshRenderData(FParticleMeshRenderDa
 				//ensureMsgf(false, TEXT("Culling & sorting is not supported once the culled counts have been acquired, sorting & culling will be disabled for these draws."));
 				ParticleMeshRenderData.bNeedsSort = false;
 				ParticleMeshRenderData.bNeedsCull = false;
+				ParticleMeshRenderData.bAllowPerParticleMeshLODs = false;
 			}
 		}
 		else
@@ -414,7 +437,7 @@ void FNiagaraRendererMeshes::PrepareParticleMeshRenderData(FParticleMeshRenderDa
 			}
 
 			// For CPU sims, decide if we should sort / cull on the GPU or not
-			if ( ParticleMeshRenderData.bSortCullOnGpu )
+			if ( ParticleMeshRenderData.bSortCullOnGpu && !ParticleMeshRenderData.bAllowPerParticleMeshLODs)
 			{
 				const int32 NumInstances = ParticleMeshRenderData.SourceParticleData->GetNumInstances();
 
@@ -426,6 +449,9 @@ void FNiagaraRendererMeshes::PrepareParticleMeshRenderData(FParticleMeshRenderDa
 
 				ParticleMeshRenderData.bSortCullOnGpu = bSortMoveToGpu || bCullMoveToGpu;
 			}
+
+			// Update if we support per particle LODs, we either need to be able to run the sort & cull process or use GPU scene
+			ParticleMeshRenderData.bAllowPerParticleMeshLODs &= ParticleMeshRenderData.bSortCullOnGpu || ParticleMeshRenderData.bUseGPUScene;
 		}
 
 		// Update layout as it could have changed
@@ -665,7 +691,7 @@ void FNiagaraRendererMeshes::PreparePerMeshData(FParticleMeshRenderData& Particl
 	}
 }
 
-uint32 FNiagaraRendererMeshes::PerformSortAndCull(FRHICommandListBase& RHICmdList, FParticleMeshRenderData& ParticleMeshRenderData, FGlobalDynamicReadBuffer& ReadBuffer, FNiagaraGPUSortInfo& SortInfo, FNiagaraGpuComputeDispatchInterface* ComputeDispatchInterface, int32 MeshIndex) const
+uint32 FNiagaraRendererMeshes::PerformSortAndCull(FRHICommandListBase& RHICmdList, FParticleMeshRenderData& ParticleMeshRenderData, FGlobalDynamicReadBuffer& ReadBuffer, FNiagaraGPUSortInfo& SortInfo, FNiagaraGpuComputeDispatchInterface* ComputeDispatchInterface, const FSceneView& View, const FMeshData& MeshData) const
 {
 	// Emitter mode culls earlier on
 	if (SourceMode == ENiagaraRendererSourceDataMode::Emitter)
@@ -680,7 +706,9 @@ uint32 FNiagaraRendererMeshes::PerformSortAndCull(FRHICommandListBase& RHICmdLis
 	{
 		SortInfo.LocalBSphere = ParticleMeshRenderData.CullingSphere;
 		SortInfo.CullingWorldSpaceOffset = ParticleMeshRenderData.WorldSpacePivotOffset;
-		SortInfo.MeshIndex = MeshIndex;
+		SortInfo.MeshIndex = MeshData.SourceMeshIndex;
+		SortInfo.LODScreenSize = FNiagaraRendererMeshes::GetShaderLODScreenSize(View, MeshData);
+
 		if (ParticleMeshRenderData.bSortCullOnGpu)
 		{
 			SortInfo.CulledGPUParticleCountOffset = ParticleMeshRenderData.bNeedsCull ? ComputeDispatchInterface->GetGPUInstanceCounterManager().AcquireCulledEntry() : INDEX_NONE;
@@ -976,6 +1004,20 @@ FNiagaraMeshUniformBufferRef FNiagaraRendererMeshes::CreateVFUniformBuffer(const
 	return FNiagaraMeshUniformBufferRef::CreateUniformBufferImmediate(Params, UniformBuffer_SingleFrame);
 }
 
+FVector4f FNiagaraRendererMeshes::GetShaderLODScreenSize(const FSceneView& View, const FMeshData& MeshData)
+{
+	const FSceneView& LODView = GetLODView(View);
+	const FMatrix& ProjMatrix = LODView.ViewMatrices.GetProjectionMatrix();
+	const float LODSize = MeshData.LODScreenSize.Z * FMath::Max(ProjMatrix.M[0][0], ProjMatrix.M[1][1]) * 0.5f;
+
+	return FVector4f(
+		FMath::Square(MeshData.LODScreenSize.X * 0.5f),
+		FMath::Square(MeshData.LODScreenSize.Y * 0.5f),
+		FMath::Square(LODSize),
+		ProjMatrix.M[2][3] * LODView.LODDistanceFactor * MeshData.LODDistanceFactor
+	);
+}
+
 void FNiagaraRendererMeshes::SetupElementForGPUScene(
 	const FParticleMeshRenderData& ParticleMeshRenderData,
 	const FNiagaraMeshCommonParameters& CommonParameters,
@@ -1016,6 +1058,7 @@ void FNiagaraRendererMeshes::SetupElementForGPUScene(
 		GPUSceneRes.GPUWriteParams.VisibilityTagDataOffset 	= ParticleMeshRenderData.RendererVisTagOffset;
 		GPUSceneRes.GPUWriteParams.LocalBoundingCenter		= (FVector3f)LocalBounds.GetCenter();
 		GPUSceneRes.GPUWriteParams.DistanceCullRangeSquared = DistanceCullRangeSquared;
+		GPUSceneRes.GPUWriteParams.LODScreenSize			= GetShaderLODScreenSize(View, MeshData);
 		GPUSceneRes.GPUWriteParams.bNeedsPrevTransform		= bNeedsPrevTransform ? 1 : 0;
 
 		// We need to set this flag to force the system to always cull individual instances, because we may need to discard instances that are:
@@ -1286,14 +1329,13 @@ void FNiagaraRendererMeshes::GetDynamicMeshElements(const TArray<const FSceneVie
 				InitializeSortInfo(ParticleMeshRenderData, *SceneProxy, *View, ViewIndex, bIsInstancedStereo, SortInfo);
 			}
 
+			const bool bHasMeshIndexValidBinding = ParticleMeshRenderData.MeshIndexOffset != INDEX_NONE || EmitterModeMeshIndex != INDEX_NONE;
+
 			FNiagaraGpuComputeDispatchInterface* ComputeDispatchInterface = SceneProxy->GetComputeDispatchInterface();
 			for (int32 MeshIndex = 0; MeshIndex < Meshes.Num(); ++MeshIndex)
 			{
-				// No binding for mesh index or we don't allow culling we will only render the first mesh for all particles
-				const bool bHasValidBinding = ParticleMeshRenderData.MeshIndexOffset != INDEX_NONE
-					|| EmitterModeMeshIndex != INDEX_NONE;
-
-				if (MeshIndex > 0 && (!bHasValidBinding || (!ParticleMeshRenderData.bNeedsCull && !ParticleMeshRenderData.bUseGPUScene)))
+				const FMeshData& MeshData = Meshes[MeshIndex];
+				if (MeshData.SourceMeshIndex > 0 && (!bHasMeshIndexValidBinding || (!ParticleMeshRenderData.bNeedsCull && !ParticleMeshRenderData.bUseGPUScene)))
 				{
 					break;
 				}
@@ -1304,7 +1346,14 @@ void FNiagaraRendererMeshes::GetDynamicMeshElements(const TArray<const FSceneVie
 					continue;
 				}
 
-				const FMeshData& MeshData = Meshes[MeshIndex];
+				// If we can no perform sorting / culling on the GPU we need to skip meshes which are LODs
+				if (MeshData.LODMode == ENiagaraMeshLODMode::PerParticle && !ParticleMeshRenderData.bAllowPerParticleMeshLODs)
+				{
+					if (MeshIndex != 0 && Meshes[MeshIndex - 1].SourceMeshIndex == Meshes[MeshIndex].SourceMeshIndex)
+					{
+						continue;
+					}
+				}
 
 				INiagaraRenderableMesh::FLODModelData LODModel;
 				if (MeshData.LODMode == ENiagaraMeshLODMode::ByComponentBounds)
@@ -1339,7 +1388,7 @@ void FNiagaraRendererMeshes::GetDynamicMeshElements(const TArray<const FSceneVie
 				PreparePerMeshData(ParticleMeshRenderData, VertexFactory, *SceneProxy, MeshData);
 
 				// Sort/Cull particles if needed.
-				const uint32 NumInstances = PerformSortAndCull(RHICmdList, ParticleMeshRenderData, Collector.GetDynamicReadBuffer(), SortInfo, ComputeDispatchInterface, MeshData.SourceMeshIndex);
+				const uint32 NumInstances = PerformSortAndCull(RHICmdList, ParticleMeshRenderData, Collector.GetDynamicReadBuffer(), SortInfo, ComputeDispatchInterface, *View, MeshData);
 				if ( NumInstances > 0 )
 				{
 					// Increment stats
@@ -1512,7 +1561,7 @@ void FNiagaraRendererMeshes::GetDynamicRayTracingInstances(FRayTracingMaterialGa
 
 		// Sort/Cull particles if needed.
 		FNiagaraGpuComputeDispatchInterface* ComputeDispatchInterface = SceneProxy->GetComputeDispatchInterface();
-		const uint32 NumInstances = PerformSortAndCull(RHICmdList, ParticleMeshRenderData, Context.RayTracingMeshResourceCollector.GetDynamicReadBuffer(), SortInfo, ComputeDispatchInterface, MeshData.SourceMeshIndex);
+		const uint32 NumInstances = PerformSortAndCull(RHICmdList, ParticleMeshRenderData, Context.RayTracingMeshResourceCollector.GetDynamicReadBuffer(), SortInfo, ComputeDispatchInterface, *View, MeshData);
 		if ( NumInstances == 0 )
 		{
 			continue;
diff --git a/Engine/Plugins/FX/Niagara/Source/Niagara/Public/NiagaraMeshRendererProperties.h b/Engine/Plugins/FX/Niagara/Source/Niagara/Public/NiagaraMeshRendererProperties.h
--- a/Engine/Plugins/FX/Niagara/Source/Niagara/Public/NiagaraMeshRendererProperties.h
+++ b/Engine/Plugins/FX/Niagara/Source/Niagara/Public/NiagaraMeshRendererProperties.h
@@ -84,7 +84,7 @@ enum class ENiagaraMeshLODMode : uint8
 	* Calculates and renders each particle with it's calcualted LOD level.
 	* Increasing 'LOD calculation scale' will result in lower quality LODs being used.
 	*/
-	//PerParticle,
+	PerParticle,
 };
 
 USTRUCT()
@@ -208,6 +208,17 @@ struct FNiagaraMeshRendererMeshProperties
 	UPROPERTY(EditAnywhere, Category = "Mesh", meta = (UIMin = "0", DisplayName = "LOD Distance Factor", EditCondition = "LODMode == ENiagaraMeshLODMode::ByComponentBounds || LODMode == ENiagaraMeshLODMode::PerParticle", EditConditionHides))
 	float LODDistanceFactor = 1.0f;
 
+	/**
+	When enabled you can restrict the LOD range we consider for LOD calculation.
+	This can be useful to reduce the performance impact, as it reduces the number of draw calls required.
+	*/
+	UPROPERTY(EditAnywhere, Category = "Mesh", meta = (DisplayName = "Use LOD Range", EditCondition = "LODMode == ENiagaraMeshLODMode::PerParticle", EditConditionHides))
+	bool bUseLODRange = false;
+
+	/** Used to restrict the range of LODs we include when dynamically calculating the LOD level. */
+	UPROPERTY(EditAnywhere, Category = "Mesh", meta = (DisplayName = "LOD Range", EditCondition = "bUseLODRange && LODMode == ENiagaraMeshLODMode::PerParticle", EditConditionHides))
+	FIntVector2 LODRange;
+
 	/** Scale of the mesh */
 	UPROPERTY(EditAnywhere, Category = "Mesh")
 	FVector Scale;
diff --git a/Engine/Plugins/FX/Niagara/Source/Niagara/Public/NiagaraRenderableMeshInterface.h b/Engine/Plugins/FX/Niagara/Source/Niagara/Public/NiagaraRenderableMeshInterface.h
--- a/Engine/Plugins/FX/Niagara/Source/Niagara/Public/NiagaraRenderableMeshInterface.h
+++ b/Engine/Plugins/FX/Niagara/Source/Niagara/Public/NiagaraRenderableMeshInterface.h
@@ -49,6 +49,11 @@ public:
 	// Gather a list of used materials
 	virtual void GetUsedMaterials(TArray<UMaterialInterface*>& OutMaterials) const = 0;
 
+	// Get the available LOD range
+	virtual FIntVector2 GetLODRange() const { return FIntVector2(0, 0); }
+	// Return parameters used for LOD calculation, screen size min, screen size max & sphere size
+	virtual FVector3f GetLODScreenSize(int32 LODLevel) const { return FVector3f(0.0f, 1.0f, 1.0f); }
+	// Compute the LOD level for the sphere location
 	virtual int32 ComputeLOD(const FVector& SphereOrigin, const float SphereRadius, const FSceneView& SceneView, float LODDistanceFactor) { return 0; }
 };
 
diff --git a/Engine/Plugins/FX/Niagara/Source/Niagara/Public/NiagaraRendererMeshes.h b/Engine/Plugins/FX/Niagara/Source/Niagara/Public/NiagaraRendererMeshes.h
--- a/Engine/Plugins/FX/Niagara/Source/Niagara/Public/NiagaraRendererMeshes.h
+++ b/Engine/Plugins/FX/Niagara/Source/Niagara/Public/NiagaraRendererMeshes.h
@@ -38,6 +38,8 @@ public:
 #endif
 	//FNiagaraRenderer Interface END
 
+	bool HasValidMeshes() const { return Meshes.Num() > 0; }
+
 protected:
 	struct FParticleMeshRenderData
 	{
@@ -50,6 +52,7 @@ protected:
 		bool							bSortCullOnGpu = false;
 		bool							bNeedsSort = false;
 		bool							bNeedsCull = false;
+		bool							bAllowPerParticleMeshLODs = false;
 		bool							bIsGpuLowLatencyTranslucency = false;
 
 		const FNiagaraRendererLayout*	RendererLayout = nullptr;
@@ -86,6 +89,7 @@ protected:
 		ENiagaraMeshLODMode LODMode = ENiagaraMeshLODMode::LODLevel;
 		int32 LODLevel = 0;
 		float LODDistanceFactor = 1.0f;
+		FVector3f LODScreenSize = FVector3f(0.0f, 2.0f, 1.0f);
 		FVector3f PivotOffset = FVector3f::ZeroVector;
 		ENiagaraMeshPivotOffsetSpace PivotOffsetSpace = ENiagaraMeshPivotOffsetSpace::Mesh;
 		FVector3f Scale = FVector3f(1.0f, 1.0f, 1.0f);
@@ -126,10 +130,12 @@ protected:
 	NIAGARA_API void PrepareParticleRenderBuffers(FRHICommandListBase& RHICmdList, FParticleMeshRenderData& ParticleMeshRenderData, FGlobalDynamicReadBuffer& DynamicReadBuffer) const;
 	NIAGARA_API void InitializeSortInfo(const FParticleMeshRenderData& ParticleMeshRenderData, const FNiagaraSceneProxy& SceneProxy, const FSceneView& View, int32 ViewIndex, bool bIsInstancedStereo, FNiagaraGPUSortInfo& OutSortInfo) const;
 	NIAGARA_API void PreparePerMeshData(FParticleMeshRenderData& ParticleMeshRenderData, const FNiagaraMeshVertexFactory& VertexFactory, const FNiagaraSceneProxy& SceneProxy, const FMeshData& MeshData) const;
-	NIAGARA_API uint32 PerformSortAndCull(FRHICommandListBase& RHICmdList, FParticleMeshRenderData& ParticleMeshRenderData, FGlobalDynamicReadBuffer& ReadBuffer, FNiagaraGPUSortInfo& SortInfo, class FNiagaraGpuComputeDispatchInterface* ComputeDispatchInterface, int32 MeshIndex) const;
+	NIAGARA_API uint32 PerformSortAndCull(FRHICommandListBase& RHICmdList, FParticleMeshRenderData& ParticleMeshRenderData, FGlobalDynamicReadBuffer& ReadBuffer, FNiagaraGPUSortInfo& SortInfo, class FNiagaraGpuComputeDispatchInterface* ComputeDispatchInterface, const FSceneView& View, const FMeshData& MeshData) const;
 	NIAGARA_API FNiagaraMeshCommonParameters CreateCommonShaderParams(const FParticleMeshRenderData& ParticleMeshRenderData, const FSceneView& View, const FMeshData& MeshData, const FNiagaraSceneProxy& SceneProxy) const;
 	NIAGARA_API FNiagaraMeshUniformBufferRef CreateVFUniformBuffer(const FParticleMeshRenderData& ParticleMeshRenderData, const FNiagaraMeshCommonParameters& CommonParams) const;
 
+	static FVector4f GetShaderLODScreenSize(const FSceneView& View, const FMeshData& MeshData);
+
 	NIAGARA_API void SetupElementForGPUScene(
 		const FParticleMeshRenderData& ParticleMeshRenderData,
 		const FNiagaraMeshCommonParameters& CommonParameters,
@@ -174,6 +180,7 @@ private:
 	uint32 bLockedAxisEnable : 1;
 	uint32 bEnableCulling : 1;
 	uint32 bEnableFrustumCulling : 1;
+	uint32 bEnableLODCulling : 1;
 	uint32 bAccurateMotionVectors : 1;
 	uint32 bIsHeterogeneousVolume : 1;
 	uint32 bCastShadows : 1;
diff --git a/Engine/Plugins/FX/Niagara/Source/NiagaraShader/Public/NiagaraGPUSceneUtils.h b/Engine/Plugins/FX/Niagara/Source/NiagaraShader/Public/NiagaraGPUSceneUtils.h
--- a/Engine/Plugins/FX/Niagara/Source/NiagaraShader/Public/NiagaraGPUSceneUtils.h
+++ b/Engine/Plugins/FX/Niagara/Source/NiagaraShader/Public/NiagaraGPUSceneUtils.h
@@ -31,6 +31,7 @@ public:
 
 		SHADER_PARAMETER(FVector3f, LocalBoundingCenter)
 		SHADER_PARAMETER(FVector2f, DistanceCullRangeSquared)
+		SHADER_PARAMETER(FVector4f, LODScreenSize)
 
 		SHADER_PARAMETER(int, bNeedsPrevTransform)
 	END_SHADER_PARAMETER_STRUCT()
diff --git a/Engine/Plugins/FX/Niagara/Source/NiagaraVertexFactories/Public/NiagaraSortingGPU.h b/Engine/Plugins/FX/Niagara/Source/NiagaraVertexFactories/Public/NiagaraSortingGPU.h
--- a/Engine/Plugins/FX/Niagara/Source/NiagaraVertexFactories/Public/NiagaraSortingGPU.h
+++ b/Engine/Plugins/FX/Niagara/Source/NiagaraVertexFactories/Public/NiagaraSortingGPU.h
@@ -66,6 +66,7 @@ public:
 		SHADER_PARAMETER(int32, MeshIndex)
 		SHADER_PARAMETER(uint32, MeshIndexAttributeOffset)
 		SHADER_PARAMETER(FVector2f, CullDistanceRangeSquared)
+		SHADER_PARAMETER(FVector4f, LODScreenSize)
 		SHADER_PARAMETER(FVector4f, LocalBoundingSphere)
 		SHADER_PARAMETER(FVector3f, CullingWorldSpaceOffset)
 		SHADER_PARAMETER(FVector3f, SystemLWCTile)
