Add Orbit Mode option to Sim Cache Editor Viewport, and defaulting to true to match parity with the Niagara Editor Viewport

diff --git a/Engine/Plugins/FX/Niagara/Source/NiagaraEditor/Private/Widgets/SNiagaraSimCacheViewport.cpp b/Engine/Plugins/FX/Niagara/Source/NiagaraEditor/Private/Widgets/SNiagaraSimCacheViewport.cpp
--- a/Engine/Plugins/FX/Niagara/Source/NiagaraEditor/Private/Widgets/SNiagaraSimCacheViewport.cpp
+++ b/Engine/Plugins/FX/Niagara/Source/NiagaraEditor/Private/Widgets/SNiagaraSimCacheViewport.cpp
@@ -7,6 +7,7 @@
 #include "SEditorViewportToolBarMenu.h"
 #include "SNiagaraSimCacheViewportToolbar.h"
 #include "NiagaraComponent.h"
+#include "NiagaraEditorCommands.h"
 
 //////////////////////////////////////////////////////////////////////////
 // Viewport Client
@@ -79,7 +80,7 @@ void FNiagaraSimCacheViewportClient::Draw(FViewport* InViewport, FCanvas* Canvas
 
 bool FNiagaraSimCacheViewportClient::ShouldOrbitCamera() const
 {
-	return FEditorViewportClient::ShouldOrbitCamera();
+	return bUsingOrbitCamera;
 }
 
 FSceneView* FNiagaraSimCacheViewportClient::CalcSceneView(FSceneViewFamily* ViewFamily, const int32 StereoViewIndex)
@@ -165,6 +166,8 @@ TSharedRef<FEditorViewportClient> SNiagaraSimCacheViewport::MakeEditorViewportCl
 
 	SimCacheViewportClient->SetRealtime(true);
 	SimCacheViewportClient->VisibilityDelegate.BindSP(this, &SNiagaraSimCacheViewport::IsVisible);
+	// Default to orbit camera for parity with the Niagara editor viewport.
+	SimCacheViewportClient->bUsingOrbitCamera = true;
 	
 	return SimCacheViewportClient.ToSharedRef();
 }
@@ -179,9 +182,27 @@ EVisibility SNiagaraSimCacheViewport::OnGetViewportContentVisibility() const
 	return SEditorViewport::OnGetViewportContentVisibility();
 }
 
+void SNiagaraSimCacheViewport::ToggleOrbit()
+{
+	SimCacheViewportClient->ToggleOrbitCamera(!SimCacheViewportClient->bUsingOrbitCamera);
+}
+
+bool SNiagaraSimCacheViewport::IsToggleOrbitChecked()
+{
+	return SimCacheViewportClient->bUsingOrbitCamera;
+}
+
 void SNiagaraSimCacheViewport::BindCommands()
 {
 	SEditorViewport::BindCommands();
+
+	const FNiagaraEditorCommands& Commands = FNiagaraEditorCommands::Get();
+
+	CommandList->MapAction(
+		Commands.ToggleOrbit,
+		FExecuteAction::CreateSP(this, &SNiagaraSimCacheViewport::ToggleOrbit),
+		FCanExecuteAction(),
+		FIsActionChecked::CreateSP(this, &SNiagaraSimCacheViewport::IsToggleOrbitChecked));
 }
 
 void SNiagaraSimCacheViewport::OnFocusViewportToSelection()
diff --git a/Engine/Plugins/FX/Niagara/Source/NiagaraEditor/Private/Widgets/SNiagaraSimCacheViewport.h b/Engine/Plugins/FX/Niagara/Source/NiagaraEditor/Private/Widgets/SNiagaraSimCacheViewport.h
--- a/Engine/Plugins/FX/Niagara/Source/NiagaraEditor/Private/Widgets/SNiagaraSimCacheViewport.h
+++ b/Engine/Plugins/FX/Niagara/Source/NiagaraEditor/Private/Widgets/SNiagaraSimCacheViewport.h
@@ -43,6 +43,8 @@ protected:
 	virtual TSharedRef<FEditorViewportClient> MakeEditorViewportClient() override;
 	virtual TSharedPtr<SWidget> MakeViewportToolbar() override;
 	virtual EVisibility OnGetViewportContentVisibility() const override;
+	void ToggleOrbit();
+	bool IsToggleOrbitChecked();
 	virtual void BindCommands() override;
 	virtual void OnFocusViewportToSelection() override;
 
diff --git a/Engine/Plugins/FX/Niagara/Source/NiagaraEditor/Private/Widgets/SNiagaraSimCacheViewportToolbar.cpp b/Engine/Plugins/FX/Niagara/Source/NiagaraEditor/Private/Widgets/SNiagaraSimCacheViewportToolbar.cpp
--- a/Engine/Plugins/FX/Niagara/Source/NiagaraEditor/Private/Widgets/SNiagaraSimCacheViewportToolbar.cpp
+++ b/Engine/Plugins/FX/Niagara/Source/NiagaraEditor/Private/Widgets/SNiagaraSimCacheViewportToolbar.cpp
@@ -2,6 +2,8 @@
 
 #include "SNiagaraSimCacheViewportToolbar.h"
 
+#include "NiagaraEditorCommands.h"
+
 void SNiagaraSimCacheViewportToolbar::Construct(const FArguments& InArgs, TSharedPtr<SNiagaraSimCacheViewport> InViewport)
 {
 	SCommonEditorViewportToolbarBase::Construct(SCommonEditorViewportToolbarBase::FArguments(), InViewport);
@@ -10,6 +12,8 @@ void SNiagaraSimCacheViewportToolbar::Construct(const FArguments& InArgs, TShare
 void SNiagaraSimCacheViewportToolbar::ExtendOptionsMenu(FMenuBuilder& OptionsMenuBuilder) const
 {
 	SCommonEditorViewportToolbarBase::ExtendOptionsMenu(OptionsMenuBuilder);
+
+	OptionsMenuBuilder.AddMenuEntry(FNiagaraEditorCommands::Get().ToggleOrbit);
 }
 
 void SNiagaraSimCacheViewportToolbar::ExtendLeftAlignedToolbarSlots(TSharedPtr<SHorizontalBox> MainBoxPtr,
