- Remove mutable access to particle data if using a const emitter instance pointer

diff --git a/Engine/Plugins/FX/Niagara/Source/Niagara/Classes/NiagaraDataSet.h b/Engine/Plugins/FX/Niagara/Source/Niagara/Classes/NiagaraDataSet.h
--- a/Engine/Plugins/FX/Niagara/Source/Niagara/Classes/NiagaraDataSet.h
+++ b/Engine/Plugins/FX/Niagara/Source/Niagara/Classes/NiagaraDataSet.h
@@ -281,9 +281,9 @@ public:
 	FORCEINLINE TArray<int32>& GetSpawnedIDsTable() { return SpawnedIDsTable; }
 	FORCEINLINE const TArray<int32>& GetFreeIDTable() const { return FreeIDsTable; }
 	FORCEINLINE const TArray<int32>& GetSpawnedIDsTable() const { return SpawnedIDsTable; }
-	FORCEINLINE int32& GetNumFreeIDs() { return NumFreeIDs; }
-	FORCEINLINE int32& GetMaxUsedID() { return MaxUsedID; }
-	FORCEINLINE int32& GetIDAcquireTag() { return IDAcquireTag; }
+	FORCEINLINE int32* GetNumFreeIDsPtr() { return &NumFreeIDs; }
+	FORCEINLINE int32* GetMaxUsedIDPtr() { return &MaxUsedID; }
+	FORCEINLINE int32 GetIDAcquireTag() const { return IDAcquireTag; }
 	FORCEINLINE void SetIDAcquireTag(int32 InTag) { IDAcquireTag = InTag; }
 	FORCEINLINE FRWBuffer& GetGPUFreeIDs() { return GPUFreeIDs; }
 	FORCEINLINE uint32 GetGPUNumAllocatedIDs() const { return GPUNumAllocatedIDs; }
diff --git a/Engine/Plugins/FX/Niagara/Source/Niagara/Classes/NiagaraEmitterInstance.h b/Engine/Plugins/FX/Niagara/Source/Niagara/Classes/NiagaraEmitterInstance.h
--- a/Engine/Plugins/FX/Niagara/Source/Niagara/Classes/NiagaraEmitterInstance.h
+++ b/Engine/Plugins/FX/Niagara/Source/Niagara/Classes/NiagaraEmitterInstance.h
@@ -85,7 +85,8 @@ public:
 	
 	NIAGARA_API bool GetBoundRendererValue_GT(const FNiagaraVariableBase& InBaseVar, const FNiagaraVariableBase& InSubVar, void* OutValueData) const;
 
-	FNiagaraDataSet& GetData()const { return *ParticleDataSet; }
+	FNiagaraDataSet& GetData() { return *ParticleDataSet; }
+	const FNiagaraDataSet& GetData() const { return *ParticleDataSet; }
 
 	FORCEINLINE bool IsActive()const { return ExecutionState == ENiagaraExecutionState::Active; }
 	FORCEINLINE bool IsDisabled()const { return ExecutionState == ENiagaraExecutionState::Disabled; }
diff --git a/Engine/Plugins/FX/Niagara/Source/Niagara/Private/NiagaraRenderer.cpp b/Engine/Plugins/FX/Niagara/Source/Niagara/Private/NiagaraRenderer.cpp
--- a/Engine/Plugins/FX/Niagara/Source/Niagara/Private/NiagaraRenderer.cpp
+++ b/Engine/Plugins/FX/Niagara/Source/Niagara/Private/NiagaraRenderer.cpp
@@ -270,7 +270,7 @@ FNiagaraDynamicDataBase::FNiagaraDynamicDataBase(const FNiagaraEmitterInstance*
 {
 	check(InEmitter);
 
-	FNiagaraDataSet& DataSet = InEmitter->GetData();
+	const FNiagaraDataSet& DataSet = InEmitter->GetData();
 	SimTarget = DataSet.GetSimTarget();
 	SystemInstanceID = InEmitter->GetParentSystemInstance()->GetId();
 
diff --git a/Engine/Plugins/FX/Niagara/Source/Niagara/Private/NiagaraRendererComponents.cpp b/Engine/Plugins/FX/Niagara/Source/Niagara/Private/NiagaraRendererComponents.cpp
--- a/Engine/Plugins/FX/Niagara/Source/Niagara/Private/NiagaraRendererComponents.cpp
+++ b/Engine/Plugins/FX/Niagara/Source/Niagara/Private/NiagaraRendererComponents.cpp
@@ -37,12 +37,12 @@ static FAutoConsoleVariableRef CVarNiagaraComponentRenderPoolInactiveTimeLimit(
 //////////////////////////////////////////////////////////////////////////
 
 template<typename T>
-void SetValueWithAccessor(FNiagaraVariable& DataVariable, FNiagaraDataSet& Data, int ParticleIndex)
+void SetValueWithAccessor(FNiagaraVariable& DataVariable, const FNiagaraDataSet& Data, int ParticleIndex)
 {
 	DataVariable.SetValue<T>(FNiagaraDataSetAccessor<T>::CreateReader(Data, DataVariable.GetName()).Get(ParticleIndex));
 }
 
-void SetVariableByType(FNiagaraVariable& DataVariable, FNiagaraDataSet& Data, int ParticleIndex)
+void SetVariableByType(FNiagaraVariable& DataVariable, const FNiagaraDataSet& Data, int ParticleIndex)
 {
 	const FNiagaraTypeDefinition& VarType = DataVariable.GetType();
 	if (VarType == FNiagaraTypeDefinition::GetFloatDef()) { SetValueWithAccessor<float>(DataVariable, Data, ParticleIndex); }
@@ -402,8 +402,8 @@ void FNiagaraRendererComponents::PostSystemTick_GameThread(const UNiagaraRendere
 	}
 
 	const double CurrentTime = AttachComponent->GetWorld()->GetRealTimeSeconds();
-	FNiagaraDataSet& Data = Emitter->GetData();
-	FNiagaraDataBuffer& ParticleData = Data.GetCurrentDataChecked();
+	const FNiagaraDataSet& Data = Emitter->GetData();
+	const FNiagaraDataBuffer& ParticleData = Data.GetCurrentDataChecked();
 	FNiagaraDataSetReaderInt32<FNiagaraBool> EnabledAccessor = FNiagaraDataSetAccessor<FNiagaraBool>::CreateReader(Data, Properties->EnabledBinding.GetDataSetBindableVariable().GetName());
 	FNiagaraDataSetReaderInt32<int32> VisTagAccessor = FNiagaraDataSetAccessor<int32>::CreateReader(Data, Properties->RendererVisibilityTagBinding.GetDataSetBindableVariable().GetName());
 	FNiagaraDataSetReaderInt32<int32> UniqueIDAccessor = FNiagaraDataSetAccessor<int32>::CreateReader(Data, FName("UniqueID"));
@@ -685,7 +685,7 @@ void FNiagaraRendererComponents::OnSystemComplete_GameThread(const UNiagaraRende
 void FNiagaraRendererComponents::TickPropertyBindings(
 	const UNiagaraComponentRendererProperties* Properties,
 	USceneComponent* Component,
-	FNiagaraDataSet& Data,
+	const FNiagaraDataSet& Data,
 	int32 ParticleIndex,
 	FComponentPoolEntry& PoolEntry,
 	const FNiagaraLWCConverter& LwcConverter)
diff --git a/Engine/Plugins/FX/Niagara/Source/Niagara/Private/NiagaraRendererLights.cpp b/Engine/Plugins/FX/Niagara/Source/Niagara/Private/NiagaraRendererLights.cpp
--- a/Engine/Plugins/FX/Niagara/Source/Niagara/Private/NiagaraRendererLights.cpp
+++ b/Engine/Plugins/FX/Niagara/Source/Niagara/Private/NiagaraRendererLights.cpp
@@ -69,8 +69,8 @@ FNiagaraDynamicDataBase* FNiagaraRendererLights::GenerateDynamicData(const FNiag
 
 	//Bail if we don't have the required attributes to render this emitter.
 	const UNiagaraLightRendererProperties* Properties = CastChecked<const UNiagaraLightRendererProperties>(InProperties);
-	FNiagaraDataSet& Data = Emitter->GetData();
-	FNiagaraDataBuffer* DataToRender = Data.GetCurrentData();
+	const FNiagaraDataSet& Data = Emitter->GetData();
+	const FNiagaraDataBuffer* DataToRender = Data.GetCurrentData();
 	if (DataToRender == nullptr || Emitter->GetParentSystemInstance() == nullptr)
 	{
 		return nullptr;
diff --git a/Engine/Plugins/FX/Niagara/Source/Niagara/Private/NiagaraScriptExecutionContext.cpp b/Engine/Plugins/FX/Niagara/Source/Niagara/Private/NiagaraScriptExecutionContext.cpp
--- a/Engine/Plugins/FX/Niagara/Source/Niagara/Private/NiagaraScriptExecutionContext.cpp
+++ b/Engine/Plugins/FX/Niagara/Source/Niagara/Private/NiagaraScriptExecutionContext.cpp
@@ -145,7 +145,7 @@ void FNiagaraScriptExecutionContextBase::BindData(int32 Index, FNiagaraDataSet&
 
 	DataSetMetaTable.SetNum(FMath::Max(DataSetMetaTable.Num(), Index + 1));
 	DataSetMetaTable[Index].Init(InputRegisters, OutputRegisters, StartInstance,
-		Output ? &Output->GetIDTable() : nullptr, &DataSet.GetFreeIDTable(), &DataSet.GetNumFreeIDs(), &DataSet.NumSpawnedIDs, &DataSet.GetMaxUsedID(), DataSet.GetIDAcquireTag(), &DataSet.GetSpawnedIDsTable());
+		Output ? &Output->GetIDTable() : nullptr, &DataSet.GetFreeIDTable(), DataSet.GetNumFreeIDsPtr(), &DataSet.NumSpawnedIDs, DataSet.GetMaxUsedIDPtr(), DataSet.GetIDAcquireTag(), &DataSet.GetSpawnedIDsTable());
 
 	if (InputRegisters.Num() > 0)
 	{
@@ -170,7 +170,7 @@ void FNiagaraScriptExecutionContextBase::BindData(int32 Index, FNiagaraDataBuffe
 	TArrayView<uint8 const* RESTRICT const> InputRegisters = Input->GetRegisterTable();
 
 	DataSetMetaTable.SetNum(FMath::Max(DataSetMetaTable.Num(), Index + 1));
-	DataSetMetaTable[Index].Init(InputRegisters, TArrayView<uint8 const* RESTRICT const>(), StartInstance, nullptr, nullptr, &DataSet->GetNumFreeIDs(), &DataSet->NumSpawnedIDs, &DataSet->GetMaxUsedID(), DataSet->GetIDAcquireTag(), &DataSet->GetSpawnedIDsTable());
+	DataSetMetaTable[Index].Init(InputRegisters, TArrayView<uint8 const* RESTRICT const>(), StartInstance, nullptr, nullptr, DataSet->GetNumFreeIDsPtr(), &DataSet->NumSpawnedIDs, DataSet->GetMaxUsedIDPtr(), DataSet->GetIDAcquireTag(), &DataSet->GetSpawnedIDsTable());
 
 	if (InputRegisters.Num() > 0)
 	{
diff --git a/Engine/Plugins/FX/Niagara/Source/Niagara/Public/NiagaraRendererComponents.h b/Engine/Plugins/FX/Niagara/Source/Niagara/Public/NiagaraRendererComponents.h
--- a/Engine/Plugins/FX/Niagara/Source/Niagara/Public/NiagaraRendererComponents.h
+++ b/Engine/Plugins/FX/Niagara/Source/Niagara/Public/NiagaraRendererComponents.h
@@ -65,7 +65,7 @@ private:
 	static NIAGARA_API void TickPropertyBindings(
 		const UNiagaraComponentRendererProperties* Properties,
 		USceneComponent* Component,
-		FNiagaraDataSet& Data,
+		const FNiagaraDataSet& Data,
 		int32 ParticleIndex,
 		FComponentPoolEntry& PoolEntry,
 		const FNiagaraLWCConverter& LwcConverter);
diff --git a/Engine/Plugins/Runtime/GeometryCache/Source/GeometryCache/Private/NiagaraRendererGeometryCache.cpp b/Engine/Plugins/Runtime/GeometryCache/Source/GeometryCache/Private/NiagaraRendererGeometryCache.cpp
--- a/Engine/Plugins/Runtime/GeometryCache/Source/GeometryCache/Private/NiagaraRendererGeometryCache.cpp
+++ b/Engine/Plugins/Runtime/GeometryCache/Source/GeometryCache/Private/NiagaraRendererGeometryCache.cpp
@@ -263,8 +263,8 @@ void FNiagaraRendererGeometryCache::PostSystemTick_GameThread(const UNiagaraRend
 
 	if (Properties->SourceMode == ENiagaraRendererSourceDataMode::Particles)
 	{
-		FNiagaraDataSet& Data = Emitter->GetData();
-		FNiagaraDataBuffer& ParticleData = Data.GetCurrentDataChecked();
+		const FNiagaraDataSet& Data = Emitter->GetData();
+		const FNiagaraDataBuffer& ParticleData = Data.GetCurrentDataChecked();
 		FNiagaraDataSetReaderInt32<FNiagaraBool> EnabledReader = Properties->EnabledAccessor.GetReader(Data);
 		FNiagaraDataSetReaderInt32<int32> VisTagReader = Properties->VisTagAccessor.GetReader(Data);
 		FNiagaraDataSetReaderInt32<int32> ArrayIndexReader = Properties->ArrayIndexAccessor.GetReader(Data);
