Fixed multiple undo/redo related issues around scratch pad renaming.

- Scratch pad script wasn't getting modified due to lack of RF_Transactional, preventing proper Undo/Redo rename
- Affected scripts weren't being modified, so renamed parameters weren't getting recorded, causing loss of input data
- Tab name would still display new name after hitting Undo; changed the tab name to use an attribute so it always reads the correct name

diff --git a/Engine/Plugins/FX/Niagara/Source/NiagaraEditor/Private/ViewModels/NiagaraScratchPadScriptViewModel.cpp b/Engine/Plugins/FX/Niagara/Source/NiagaraEditor/Private/ViewModels/NiagaraScratchPadScriptViewModel.cpp
--- a/Engine/Plugins/FX/Niagara/Source/NiagaraEditor/Private/ViewModels/NiagaraScratchPadScriptViewModel.cpp
+++ b/Engine/Plugins/FX/Niagara/Source/NiagaraEditor/Private/ViewModels/NiagaraScratchPadScriptViewModel.cpp
@@ -162,6 +162,12 @@ void FNiagaraScratchPadScriptViewModel::SetScriptName(FText InScriptName)
 		FScopedTransaction RenameTransaction(LOCTEXT("RenameScriptTransaction", "Rename scratch pad script."));
 
 		FName NewUniqueName = FNiagaraEditorUtilities::GetUniqueObjectName<UNiagaraScript>(OriginalScript->GetOuter(), *NewName);
+		
+		if(OriginalScript->HasAnyFlags(RF_Transactional) == false)
+		{
+			OriginalScript->SetFlags(RF_Transactional);
+		}
+		
 		OriginalScript->Modify();
 		OriginalScript->Rename(*NewUniqueName.ToString(), nullptr, REN_DontCreateRedirectors);
 
diff --git a/Engine/Plugins/FX/Niagara/Source/NiagaraEditor/Private/ViewModels/Stack/NiagaraStackGraphUtilities.cpp b/Engine/Plugins/FX/Niagara/Source/NiagaraEditor/Private/ViewModels/Stack/NiagaraStackGraphUtilities.cpp
--- a/Engine/Plugins/FX/Niagara/Source/NiagaraEditor/Private/ViewModels/Stack/NiagaraStackGraphUtilities.cpp
+++ b/Engine/Plugins/FX/Niagara/Source/NiagaraEditor/Private/ViewModels/Stack/NiagaraStackGraphUtilities.cpp
@@ -3643,6 +3643,7 @@ void FNiagaraStackGraphUtilities::RenameReferencingParameters(UNiagaraSystem* Sy
 			continue;
 		}
 
+		Script->Modify();
 		TArray<FNiagaraVariable> RapidIterationVariables;
 		Script->RapidIterationParameters.GetParameters(RapidIterationVariables);
 
diff --git a/Engine/Plugins/FX/Niagara/Source/NiagaraEditor/Private/Widgets/NiagaraScratchScriptEditor.cpp b/Engine/Plugins/FX/Niagara/Source/NiagaraEditor/Private/Widgets/NiagaraScratchScriptEditor.cpp
--- a/Engine/Plugins/FX/Niagara/Source/NiagaraEditor/Private/Widgets/NiagaraScratchScriptEditor.cpp
+++ b/Engine/Plugins/FX/Niagara/Source/NiagaraEditor/Private/Widgets/NiagaraScratchScriptEditor.cpp
@@ -271,7 +271,7 @@ TAttribute<FText> FNiagaraGraphEditorSummoner::ConstructTabNameForObject(UEdGrap
 			TSharedPtr<FNiagaraScratchPadScriptViewModel> ScratchScriptViewModel = UNiagaraSystemEditorDocumentsViewModel::GetScratchPadViewModelFromGraph(SysViewModel.Get(), DocumentID);
 			if (ScratchScriptViewModel)
 			{
-				return ScratchScriptViewModel->GetDisplayName(); /* FText::FromString(ScratchScriptViewModel->GetOriginalScript()->GetName());*/
+				return TAttribute<FText>::CreateSP(ScratchScriptViewModel.Get(), &FNiagaraScratchPadScriptViewModel::GetDisplayName);
 			}
 		}
 	}
