Added an if guard to catch a case where a module doesn't exist despite the hierarchy module still existing.
This shouldn't happen, and it's not clear why this is, but in the mean time this fixes the crash.

diff --git a/Engine/Plugins/FX/Niagara/Source/NiagaraEditor/Private/ViewModels/Stack/NiagaraStackInputCategory.cpp b/Engine/Plugins/FX/Niagara/Source/NiagaraEditor/Private/ViewModels/Stack/NiagaraStackInputCategory.cpp
--- a/Engine/Plugins/FX/Niagara/Source/NiagaraEditor/Private/ViewModels/Stack/NiagaraStackInputCategory.cpp
+++ b/Engine/Plugins/FX/Niagara/Source/NiagaraEditor/Private/ViewModels/Stack/NiagaraStackInputCategory.cpp
@@ -250,17 +250,23 @@ void UNiagaraStackSummaryCategory::RefreshChildrenInternal(const TArray<UNiagara
 	for(TSharedPtr<FNiagaraHierarchyItemViewModelBase> HierarchyViewModel : AllChildrenViewModels)
 	{
 		UNiagaraHierarchyItemBase* Data = HierarchyViewModel->GetDataMutable();
+		TOptional<FGuid> FunctionCallGuid;
 		if(UNiagaraHierarchyModuleInput* ModuleInput = Cast<UNiagaraHierarchyModuleInput>(Data))
 		{
-			UsedFunctionCallNodes.Add(State.NodeGuidToModuleNodeMap[ModuleInput->GetPersistentIdentity().Guids[0]]);
+			FunctionCallGuid = ModuleInput->GetPersistentIdentity().Guids[0];
 		}
 		else if(UNiagaraHierarchyAssignmentInput* AssignmentInput = Cast<UNiagaraHierarchyAssignmentInput>(Data))
 		{
-			UsedFunctionCallNodes.Add(State.NodeGuidToModuleNodeMap[AssignmentInput->GetPersistentIdentity().Guids[0]]);
+			FunctionCallGuid = AssignmentInput->GetPersistentIdentity().Guids[0];
 		}
 		else if(UNiagaraHierarchyModule* Module = Cast<UNiagaraHierarchyModule>(Data))
 		{
-			UsedFunctionCallNodes.Add(State.NodeGuidToModuleNodeMap[Module->GetPersistentIdentity().Guids[0]]);
+			FunctionCallGuid = Module->GetPersistentIdentity().Guids[0];
+		}
+
+		if(FunctionCallGuid.IsSet() && State.NodeGuidToModuleNodeMap.Contains(FunctionCallGuid.GetValue()))
+		{
+			UsedFunctionCallNodes.Add(State.NodeGuidToModuleNodeMap[FunctionCallGuid.GetValue()]);
 		}
 	}
 	
diff --git a/Engine/Plugins/FX/Niagara/Source/NiagaraEditor/Private/ViewModels/Stack/NiagaraStackSummaryViewInputCollection.cpp b/Engine/Plugins/FX/Niagara/Source/NiagaraEditor/Private/ViewModels/Stack/NiagaraStackSummaryViewInputCollection.cpp
--- a/Engine/Plugins/FX/Niagara/Source/NiagaraEditor/Private/ViewModels/Stack/NiagaraStackSummaryViewInputCollection.cpp
+++ b/Engine/Plugins/FX/Niagara/Source/NiagaraEditor/Private/ViewModels/Stack/NiagaraStackSummaryViewInputCollection.cpp
@@ -105,17 +105,23 @@ void UNiagaraStackSummaryViewCollection::RefreshChildrenInternal(const TArray<UN
 	for(TSharedPtr<FNiagaraHierarchyItemViewModelBase> HierarchyViewModel : AllChildrenViewModels)
 	{
 		UNiagaraHierarchyItemBase* Data = HierarchyViewModel->GetDataMutable();
+		TOptional<FGuid> FunctionCallGuid;
 		if(UNiagaraHierarchyModuleInput* ModuleInput = Cast<UNiagaraHierarchyModuleInput>(Data))
 		{
-			UsedFunctionCallNodes.Add(State.NodeGuidToModuleNodeMap[ModuleInput->GetPersistentIdentity().Guids[0]]);
+			FunctionCallGuid = ModuleInput->GetPersistentIdentity().Guids[0];
 		}
 		else if(UNiagaraHierarchyAssignmentInput* AssignmentInput = Cast<UNiagaraHierarchyAssignmentInput>(Data))
 		{
-			UsedFunctionCallNodes.Add(State.NodeGuidToModuleNodeMap[AssignmentInput->GetPersistentIdentity().Guids[0]]);
+			FunctionCallGuid = AssignmentInput->GetPersistentIdentity().Guids[0];
 		}
 		else if(UNiagaraHierarchyModule* Module = Cast<UNiagaraHierarchyModule>(Data))
 		{
-			UsedFunctionCallNodes.Add(State.NodeGuidToModuleNodeMap[Module->GetPersistentIdentity().Guids[0]]);
+			FunctionCallGuid = Module->GetPersistentIdentity().Guids[0];
+		}
+
+		if(FunctionCallGuid.IsSet() && State.NodeGuidToModuleNodeMap.Contains(FunctionCallGuid.GetValue()))
+		{
+			UsedFunctionCallNodes.Add(State.NodeGuidToModuleNodeMap[FunctionCallGuid.GetValue()]);
 		}
 	}
 	
