- Remove dead code & fix unnecessary ref counting

diff --git a/Engine/Plugins/FX/Niagara/Source/Niagara/Classes/NiagaraEmitterInstance.h b/Engine/Plugins/FX/Niagara/Source/Niagara/Classes/NiagaraEmitterInstance.h
--- a/Engine/Plugins/FX/Niagara/Source/Niagara/Classes/NiagaraEmitterInstance.h
+++ b/Engine/Plugins/FX/Niagara/Source/Niagara/Classes/NiagaraEmitterInstance.h
@@ -141,9 +141,9 @@ public:
 
 	NIAGARA_API bool IsReadyToRun() const;
 
-	void Dump()const;
+	void Dump() const;
 
-	bool WaitForDebugInfo();
+	bool WaitForDebugInfo() const;
 
 	FNiagaraComputeExecutionContext* GetGPUContext()const
 	{
diff --git a/Engine/Plugins/FX/Niagara/Source/Niagara/Private/NiagaraComponent.cpp b/Engine/Plugins/FX/Niagara/Source/Niagara/Private/NiagaraComponent.cpp
--- a/Engine/Plugins/FX/Niagara/Source/Niagara/Private/NiagaraComponent.cpp
+++ b/Engine/Plugins/FX/Niagara/Source/Niagara/Private/NiagaraComponent.cpp
@@ -2690,40 +2690,6 @@ void UNiagaraComponent::SetVariableTextureRenderTarget(FName InVariableName, UTe
 #endif
 }
 
-TArray<FVector> UNiagaraComponent::GetNiagaraParticlePositions_DebugOnly(const FString& InEmitterName)
-{
-	return GetNiagaraParticleValueVec3_DebugOnly(InEmitterName, TEXT("Position"));
-}
-
-TArray<FVector> UNiagaraComponent::GetNiagaraParticleValueVec3_DebugOnly(const FString& InEmitterName, const FString& InValueName)
-{
-	TArray<FVector> Values;
-	FName EmitterName = FName(*InEmitterName);
-	if (SystemInstanceController.IsValid())
-	{
-		if (!SystemInstanceController->GetParticleValueVec3_DebugOnly(Values, EmitterName, FName(*InValueName)))
-		{
-			UE_LOG(LogNiagara, Warning, TEXT("Unable to find variable %s on %s per-particle data. Returning zeroes."), *InValueName, *GetPathName());
-		}
-	}
-	return Values;
-
-}
-
-TArray<float> UNiagaraComponent::GetNiagaraParticleValues_DebugOnly(const FString& InEmitterName, const FString& InValueName)
-{
-	TArray<float> Values;
-	FName EmitterName = FName(*InEmitterName);
-	if (SystemInstanceController.IsValid())
-	{
-		if (!SystemInstanceController->GetParticleValues_DebugOnly(Values, EmitterName, FName(*InValueName)))
-		{
-			UE_LOG(LogNiagara, Warning, TEXT("Unable to find variable %s on %s per-particle data. Returning zeroes."), *InValueName, *GetPathName());
-		}
-	}
-	return Values;
-}
-
 void FixInvalidUserParameters(FNiagaraUserRedirectionParameterStore& ParameterStore)
 {
 	static const FString UserPrefix = FNiagaraConstants::UserNamespace.ToString() + TEXT(".");
diff --git a/Engine/Plugins/FX/Niagara/Source/Niagara/Private/NiagaraEmitterInstance.cpp b/Engine/Plugins/FX/Niagara/Source/Niagara/Private/NiagaraEmitterInstance.cpp
--- a/Engine/Plugins/FX/Niagara/Source/Niagara/Private/NiagaraEmitterInstance.cpp
+++ b/Engine/Plugins/FX/Niagara/Source/Niagara/Private/NiagaraEmitterInstance.cpp
@@ -1268,7 +1268,7 @@ void FNiagaraEmitterInstance::PreTick()
 	ParticleDataSet->SetIDAcquireTag(TickCount);
 }
 
-bool FNiagaraEmitterInstance::WaitForDebugInfo()
+bool FNiagaraEmitterInstance::WaitForDebugInfo() const
 {
 	FNiagaraComputeExecutionContext* DebugContext = GPUExecContext;
 	if (CachedEmitter.GetEmitterData()->SimTarget == ENiagaraSimTarget::GPUComputeSim && DebugContext)
diff --git a/Engine/Plugins/FX/Niagara/Source/Niagara/Private/NiagaraSystemInstance.cpp b/Engine/Plugins/FX/Niagara/Source/Niagara/Private/NiagaraSystemInstance.cpp
--- a/Engine/Plugins/FX/Niagara/Source/Niagara/Private/NiagaraSystemInstance.cpp
+++ b/Engine/Plugins/FX/Niagara/Source/Niagara/Private/NiagaraSystemInstance.cpp
@@ -431,7 +431,7 @@ bool FNiagaraSystemInstance::QueryCaptureResults(const FGuid& RequestId, TArray<
 
 			if (bWaitForGPU)
 			{
-				for (TSharedRef<FNiagaraEmitterInstance, ESPMode::ThreadSafe> CachedEmitter : Emitters)
+				for (const TSharedRef<FNiagaraEmitterInstance, ESPMode::ThreadSafe>& CachedEmitter : Emitters)
 				{
 					CachedEmitter->WaitForDebugInfo();
 				}
@@ -698,7 +698,7 @@ void FNiagaraSystemInstance::Complete(bool bExternalCompletion)
 		SetActualExecutionState(ENiagaraExecutionState::Complete);
 		SetRequestedExecutionState(ENiagaraExecutionState::Complete);
 
-		for (TSharedRef<FNiagaraEmitterInstance, ESPMode::ThreadSafe> Simulation : Emitters)
+		for (const TSharedRef<FNiagaraEmitterInstance, ESPMode::ThreadSafe>& Simulation : Emitters)
 		{
 			Simulation->HandleCompletion(true);
 		}
@@ -962,7 +962,7 @@ void FNiagaraSystemInstance::ResetInternal(bool bResetSimulations)
 		return;
 	}
 
-	for (TSharedRef<FNiagaraEmitterInstance, ESPMode::ThreadSafe> Simulation : Emitters)
+	for (const TSharedRef<FNiagaraEmitterInstance, ESPMode::ThreadSafe>& Simulation : Emitters)
 	{
 		Simulation->ResetSimulation(bResetSimulations);
 	}
@@ -1019,7 +1019,7 @@ bool FNiagaraSystemInstance::IsReadyToRun() const
 
 	// check emitters
 	bool bAllReadyToRun = true;
-	for (TSharedRef<FNiagaraEmitterInstance, ESPMode::ThreadSafe> Simulation : Emitters)
+	for (const TSharedRef<FNiagaraEmitterInstance, ESPMode::ThreadSafe>& Simulation : Emitters)
 	{
 		if (!Simulation->IsReadyToRun())
 		{
@@ -1207,7 +1207,7 @@ void FNiagaraSystemInstance::BindParameters()
 		UE_LOG(LogNiagara, Warning, TEXT("OverrideParameters is null.  Component(%s) System(%s)"), *GetFullNameSafe(AttachComponent.Get()), *GetFullNameSafe(GetSystem()));
 	}
 
-	for (TSharedRef<FNiagaraEmitterInstance, ESPMode::ThreadSafe> Simulation : Emitters)
+	for (const TSharedRef<FNiagaraEmitterInstance, ESPMode::ThreadSafe>& Simulation : Emitters)
 	{
 		Simulation->BindParameters(bAlreadyBound);
 	}
@@ -1233,7 +1233,7 @@ void FNiagaraSystemInstance::UnbindParameters(bool bFromComplete)
 	}
 
 	bAlreadyBound = bFromComplete && bAlreadyBound;
-	for (TSharedRef<FNiagaraEmitterInstance, ESPMode::ThreadSafe> Simulation : Emitters)
+	for (const TSharedRef<FNiagaraEmitterInstance, ESPMode::ThreadSafe>& Simulation : Emitters)
 	{
 		Simulation->UnbindParameters(bFromComplete);
 	}
@@ -1613,7 +1613,7 @@ void FNiagaraSystemInstance::InitDataInterfaces()
 
 	if (GetSystem()->NeedsGPUContextInitForDataInterfaces())
 	{
-		for (TSharedRef<FNiagaraEmitterInstance, ESPMode::ThreadSafe> Simulation : Emitters)
+		for (const TSharedRef<FNiagaraEmitterInstance, ESPMode::ThreadSafe>& Simulation : Emitters)
 		{
 			FNiagaraEmitterInstance& Sim = Simulation.Get();
 			if (Sim.IsDisabled())
@@ -1632,7 +1632,7 @@ void FNiagaraSystemInstance::InitDataInterfaces()
 	//If needed, the system script execution will use these to perform per instance pre and post tick operations on our DIs.
 	SystemSpawnDIStageTickHandler.Init(SystemSimulation->GetSpawnExecutionContext()->Script, this);
 	SystemUpdateDIStageTickHandler.Init(SystemSimulation->GetUpdateExecutionContext()->Script, this);
-	for (TSharedRef<FNiagaraEmitterInstance, ESPMode::ThreadSafe>& Sim : Emitters)
+	for (const TSharedRef<FNiagaraEmitterInstance, ESPMode::ThreadSafe>& Sim : Emitters)
 	{
 		Sim->InitDITickLists();
 	}
@@ -1735,7 +1735,7 @@ void FNiagaraSystemInstance::TickDataInterfaces(float DeltaSeconds, bool bPostSi
 	if (bRebindVMFuncs)
 	{
 		// Dirty data interfaces for emitters
-		for (TSharedRef<FNiagaraEmitterInstance, ESPMode::ThreadSafe> Simulation : Emitters)
+		for (const TSharedRef<FNiagaraEmitterInstance, ESPMode::ThreadSafe>& Simulation : Emitters)
 		{
 			FNiagaraEmitterInstance& Sim = Simulation.Get();
 			if (!Sim.IsDisabled())
@@ -2350,7 +2350,7 @@ void FNiagaraSystemInstance::WaitForConcurrentTickAndFinalize(bool bEnsureComple
 bool FNiagaraSystemInstance::HandleCompletion()
 {
 	bool bEmittersCompleteOrDisabled = true;
-	for (TSharedRef<FNiagaraEmitterInstance, ESPMode::ThreadSafe>&it : Emitters)
+	for (const TSharedRef<FNiagaraEmitterInstance, ESPMode::ThreadSafe>&it : Emitters)
 	{
 		FNiagaraEmitterInstance& Inst = *it;
 		bEmittersCompleteOrDisabled &= Inst.HandleCompletion();
@@ -2742,7 +2742,7 @@ void FNiagaraSystemInstance::DestroyDataInterfaceInstanceData()
 
 TSharedPtr<FNiagaraEmitterInstance, ESPMode::ThreadSafe> FNiagaraSystemInstance::GetSimulationForHandle(const FNiagaraEmitterHandle& EmitterHandle)
 {
-	for (TSharedPtr<FNiagaraEmitterInstance, ESPMode::ThreadSafe> Sim : Emitters)
+	for (const TSharedRef<FNiagaraEmitterInstance, ESPMode::ThreadSafe>& Sim : Emitters)
 	{
 		if(Sim->GetEmitterHandle().GetId() == EmitterHandle.GetId())
 		{
@@ -2783,7 +2783,7 @@ FBox FNiagaraSystemInstance::GetSystemFixedBounds() const
 
 void FNiagaraSystemInstance::SetEmitterFixedBounds(FName EmitterName, const FBox& InLocalBounds)
 {
-	for (TSharedRef<FNiagaraEmitterInstance, ESPMode::ThreadSafe>& Emitter : Emitters)
+	for (const TSharedRef<FNiagaraEmitterInstance, ESPMode::ThreadSafe>& Emitter : Emitters)
 	{
 		if ( Emitter->GetEmitterHandle().GetName() == EmitterName )
 		{
@@ -2811,18 +2811,6 @@ FBox FNiagaraSystemInstance::GetEmitterFixedBounds(FName EmitterName) const
 	return FBox(ForceInit);
 }
 
-FNiagaraEmitterInstance* FNiagaraSystemInstance::GetEmitterByID(FGuid InID)
-{
-	for (TSharedRef<FNiagaraEmitterInstance, ESPMode::ThreadSafe>& Emitter : Emitters)
-	{
-		if (Emitter->GetEmitterHandle().GetId() == InID)
-		{
-			return &Emitter.Get();
-		}
-	}
-	return nullptr;
-}
-
 void FNiagaraSystemInstance::SetForceSolo(bool bInForceSolo)
 {
 	// We may be forced into solo mode so if we match nothing to do here
diff --git a/Engine/Plugins/FX/Niagara/Source/Niagara/Private/NiagaraSystemInstanceController.cpp b/Engine/Plugins/FX/Niagara/Source/Niagara/Private/NiagaraSystemInstanceController.cpp
--- a/Engine/Plugins/FX/Niagara/Source/Niagara/Private/NiagaraSystemInstanceController.cpp
+++ b/Engine/Plugins/FX/Niagara/Source/Niagara/Private/NiagaraSystemInstanceController.cpp
@@ -169,68 +169,6 @@ UMaterialInterface* FNiagaraSystemInstanceController::GetMaterialOverride(const
 	return nullptr;
 }
 
-bool FNiagaraSystemInstanceController::GetParticleValueVec3_DebugOnly(TArray<FVector>& OutValues, FName EmitterName, FName ValueName) const
-{
-	if (SystemInstance.IsValid())
-	{
-		for (TSharedRef<FNiagaraEmitterInstance, ESPMode::ThreadSafe>& Sim : SystemInstance->GetEmitters())
-		{
-			if (Sim->GetEmitterHandle().GetName() == EmitterName)
-			{
-				FNiagaraDataBuffer& ParticleData = Sim->GetData().GetCurrentDataChecked();
-				int32 NumParticles = ParticleData.GetNumInstances();
-				OutValues.SetNum(NumParticles);
-
-				const auto Reader = FNiagaraDataSetAccessor<FVector3f>::CreateReader(Sim->GetData(), ValueName);
-				if (!Reader.IsValid())
-				{
-					return false;
-				}
-
-				for (int32 i = 0; i < NumParticles; ++i)
-				{
-					OutValues[i] = (FVector)Reader.GetSafe(i, FVector3f::ZeroVector);
-				}
-
-				break;
-			}
-		}
-	}
-
-	return true;
-}
-
-bool FNiagaraSystemInstanceController::GetParticleValues_DebugOnly(TArray<float>& OutValues, FName EmitterName, FName ValueName) const
-{
-	if (SystemInstance.IsValid())
-	{
-		for (TSharedRef<FNiagaraEmitterInstance, ESPMode::ThreadSafe>& Sim : SystemInstance->GetEmitters())
-		{
-			if (Sim->GetEmitterHandle().GetName() == EmitterName)
-			{
-				FNiagaraDataBuffer& ParticleData = Sim->GetData().GetCurrentDataChecked();
-				int32 NumParticles = ParticleData.GetNumInstances();
-				OutValues.SetNum(NumParticles);
-
-				const auto Reader = FNiagaraDataSetAccessor<float>::CreateReader(Sim->GetData(), ValueName);
-				if (!Reader.IsValid())
-				{
-					return false;
-				}
-
-				for (int32 i = 0; i < NumParticles; ++i)
-				{
-					OutValues[i] = Reader.GetSafe(i, 0.0f);
-				}
-
-				break;
-			}
-		}
-	}
-
-	return true;
-}
-
 void FNiagaraSystemInstanceController::DebugDump(bool bFullDump)
 {
 	if (SystemInstance.IsValid())
diff --git a/Engine/Plugins/FX/Niagara/Source/Niagara/Public/NiagaraComponent.h b/Engine/Plugins/FX/Niagara/Source/Niagara/Public/NiagaraComponent.h
--- a/Engine/Plugins/FX/Niagara/Source/Niagara/Public/NiagaraComponent.h
+++ b/Engine/Plugins/FX/Niagara/Source/Niagara/Public/NiagaraComponent.h
@@ -560,18 +560,6 @@ public:
 	UFUNCTION(BlueprintCallable, Category = Niagara, meta = (DisplayName = "Set Niagara Variable (TextureRenderTarget)"))
 	NIAGARA_API void SetVariableTextureRenderTarget(FName InVariableName, class UTextureRenderTarget* TextureRenderTarget);
 
-	/** Debug accessors for getting positions in blueprints. */
-	UFUNCTION(BlueprintCallable, Category = Niagara, meta = (DisplayName = "Get Niagara Emitter Positions", DeprecatedFunction, DeprecationMessage = "Get Niagara Emitter Positions is deprecated, use the particle export DI inside your emitter instead."))
-	NIAGARA_API TArray<FVector> GetNiagaraParticlePositions_DebugOnly(const FString& InEmitterName);
-
-	/** Debug accessors for getting a float attribute array in blueprints.  The attribute name should be without namespaces. For example for "Particles.Position", send "Position". */
-	UFUNCTION(BlueprintCallable, Category = Niagara, meta = (DisplayName = "Get Niagara Emitter Float Attrib", DeprecatedFunction, DeprecationMessage = "Get Niagara Emitter Float Attrib is deprecated, use the particle export DI inside your emitter instead."))
-	NIAGARA_API TArray<float> GetNiagaraParticleValues_DebugOnly(const FString& InEmitterName, const FString& InValueName);
-
-	/** Debug accessors for getting a FVector attribute array in blueprints. The attribute name should be without namespaces. For example for "Particles.Position", send "Position". */
-	UFUNCTION(BlueprintCallable, Category = Niagara, meta = (DisplayName = "Get Niagara Emitter Vec3 Attrib", DeprecatedFunction, DeprecationMessage = "Get Niagara Emitter Vec3 Attrib is deprecated, use the particle export DI inside your emitter instead."))
-	NIAGARA_API TArray<FVector> GetNiagaraParticleValueVec3_DebugOnly(const FString& InEmitterName, const FString& InValueName);
-
 	/** Resets the System to it's initial pre-simulated state. */
 	UFUNCTION(BlueprintCallable, Category = Niagara, meta = (DisplayName = "Reset System"))
 	NIAGARA_API void ResetSystem();
diff --git a/Engine/Plugins/FX/Niagara/Source/Niagara/Public/NiagaraSystemInstance.h b/Engine/Plugins/FX/Niagara/Source/Niagara/Public/NiagaraSystemInstance.h
--- a/Engine/Plugins/FX/Niagara/Source/Niagara/Public/NiagaraSystemInstance.h
+++ b/Engine/Plugins/FX/Niagara/Source/Niagara/Public/NiagaraSystemInstance.h
@@ -249,8 +249,6 @@ public:
 	NIAGARA_API void SetEmitterFixedBounds(FName EmitterName, const FBox& InLocalBounds);
 	NIAGARA_API FBox GetEmitterFixedBounds(FName EmitterName) const;
 
-	NIAGARA_API FNiagaraEmitterInstance* GetEmitterByID(FGuid InID);
-
 	NIAGARA_API void SetForceSolo(bool bForceSolo);
 	FORCEINLINE bool IsSolo() const { return bSolo; }
 
diff --git a/Engine/Plugins/FX/Niagara/Source/Niagara/Public/NiagaraSystemInstanceController.h b/Engine/Plugins/FX/Niagara/Source/Niagara/Public/NiagaraSystemInstanceController.h
--- a/Engine/Plugins/FX/Niagara/Source/Niagara/Public/NiagaraSystemInstanceController.h
+++ b/Engine/Plugins/FX/Niagara/Source/Niagara/Public/NiagaraSystemInstanceController.h
@@ -83,9 +83,6 @@ public:
 
 	void GetStreamingMeshInfo(const FBoxSphereBounds& OwnerBounds, FStreamingTextureLevelContext& LevelContext, TArray<FStreamingRenderAssetPrimitiveInfo>& OutStreamingRenderAssets) const;
 
-	bool GetParticleValueVec3_DebugOnly(TArray<FVector>& OutValues, FName EmitterName, FName ValueName) const;
-	bool GetParticleValues_DebugOnly(TArray<float>& OutValues, FName EmitterName, FName ValueName) const;
-
 	/** Dumps system instance state and info to the log (Used by fx.Niagara.DumpComponents commandlet) */
 	void DebugDump(bool bFullDump);
 	SIZE_T GetTotalBytesUsed() const;
