Adding bindless parameter checking to NDI so that bindless resources are counted when writing UAVs to the parameter data.

diff --git a/Engine/Plugins/FX/Niagara/Source/Niagara/Private/NiagaraDataInterface.cpp b/Engine/Plugins/FX/Niagara/Source/Niagara/Private/NiagaraDataInterface.cpp
--- a/Engine/Plugins/FX/Niagara/Source/Niagara/Private/NiagaraDataInterface.cpp
+++ b/Engine/Plugins/FX/Niagara/Source/Niagara/Private/NiagaraDataInterface.cpp
@@ -363,6 +363,17 @@ FVector3f FNiagaraDataInterfaceSetShaderParametersContext::GetSystemLWCTile() co
 bool FNiagaraDataInterfaceSetShaderParametersContext::IsResourceBound(const void* ResourceAddress) const
 {
 	const uint16 ByteOffset = uint16(uintptr_t(ResourceAddress) - uintptr_t(BaseParameters));
+
+#if PLATFORM_SUPPORTS_BINDLESS_RENDERING
+	for (const FShaderParameterBindings::FBindlessResourceParameter& BindlessResourceParameter : ShaderRef->Bindings.BindlessResourceParameters)
+	{
+		if (BindlessResourceParameter.ByteOffset == ByteOffset)
+		{
+			return true;
+		}
+	}
+#endif
+
 	for (const FShaderParameterBindings::FResourceParameter& ResourceParameter : ShaderRef->Bindings.ResourceParameters)
 	{
 		if ( ResourceParameter.ByteOffset == ByteOffset )
@@ -397,6 +408,21 @@ bool FNiagaraDataInterfaceSetShaderParametersContext::IsStructBoundInternal(cons
 	const uint16 ByteStart = uint16(uintptr_t(StructAddress) - uintptr_t(BaseParameters));
 	const uint16 ByteEnd = ByteStart + uint16(StructSize);
 
+#if PLATFORM_SUPPORTS_BINDLESS_RENDERING
+	// Loop over bindless resources
+	for (const FShaderParameterBindings::FBindlessResourceParameter& BindlessResourceParameter : ShaderRef->Bindings.BindlessResourceParameters)
+	{
+		if (BindlessResourceParameter.ByteOffset >= ByteEnd)
+		{
+			break;
+		}
+		if (BindlessResourceParameter.ByteOffset >= ByteStart)
+		{
+			return true;
+		}
+	}
+#endif
+
 	// Loop over resources
 	for (const FShaderParameterBindings::FResourceParameter& ResourceParameter : ShaderRef->Bindings.ResourceParameters)
 	{
