Fixed a script compile error that was triggered when adding a new enum entry to an enum used by a staticswitch.
While the node would update and show the new enum entry, the 'NumOptionsPerVariable' wasn't updated, causing a wrong compile error

diff --git a/Engine/Plugins/FX/Niagara/Source/NiagaraEditor/Private/NiagaraNodeStaticSwitch.cpp b/Engine/Plugins/FX/Niagara/Source/NiagaraEditor/Private/NiagaraNodeStaticSwitch.cpp
--- a/Engine/Plugins/FX/Niagara/Source/NiagaraEditor/Private/NiagaraNodeStaticSwitch.cpp
+++ b/Engine/Plugins/FX/Niagara/Source/NiagaraEditor/Private/NiagaraNodeStaticSwitch.cpp
@@ -475,6 +475,8 @@ void UNiagaraNodeStaticSwitch::AllocateDefaultPins()
 	}
 	
 	TArray<int32> OptionValues = GetOptionValues();
+	NumOptionsPerVariable = OptionValues.Num();
+	
 	for (int32 i = 0; i < OptionValues.Num(); i++)
 	{
 		for (const FNiagaraVariable& Variable : OutputVars)
@@ -694,7 +696,7 @@ UNiagaraScriptVariable* UNiagaraNodeStaticSwitch::GetStaticSwitchScriptVariable(
 	return GetNiagaraGraph()->GetScriptVariable(Variable);
 }
 
-void UNiagaraNodeStaticSwitch::CheckForOutdatedEnum(FTranslator* Translator)
+void UNiagaraNodeStaticSwitch::CheckForOutdatedEnum(FTranslator* Translator) const
 {
 	// we check during compilation if the node is using outdated enum information and log it
 	if(SwitchTypeData.SwitchType == ENiagaraStaticSwitchType::Enum)
diff --git a/Engine/Plugins/FX/Niagara/Source/NiagaraEditor/Private/NiagaraNodeStaticSwitch.h b/Engine/Plugins/FX/Niagara/Source/NiagaraEditor/Private/NiagaraNodeStaticSwitch.h
--- a/Engine/Plugins/FX/Niagara/Source/NiagaraEditor/Private/NiagaraNodeStaticSwitch.h
+++ b/Engine/Plugins/FX/Niagara/Source/NiagaraEditor/Private/NiagaraNodeStaticSwitch.h
@@ -74,7 +74,7 @@ public:
 	UEdGraphPin* GetSelectorPin() const;
 	
 	/** This is a hack used in the translator to check for inconsistencies with old static switches before auto refresh was a thing */
-	void CheckForOutdatedEnum(FTranslator* Translator);
+	void CheckForOutdatedEnum(FTranslator* Translator) const;
 
 	void UpdateCompilerConstantValue(FTranslator* Translator);
 
