Fixes for Niagara AsyncTask compilation mode when it comes to managing data interfaces through the lifetime of the digested graphs.  Also cleans up code related to disabled emitters, GC management of the graphs and breaks up the graph class into a digested version and an instanced version.

diff --git a/Engine/Plugins/FX/Niagara/Source/NiagaraEditor/Private/NiagaraCompilationCopy.cpp b/Engine/Plugins/FX/Niagara/Source/NiagaraEditor/Private/NiagaraCompilationCopy.cpp
--- a/Engine/Plugins/FX/Niagara/Source/NiagaraEditor/Private/NiagaraCompilationCopy.cpp
+++ b/Engine/Plugins/FX/Niagara/Source/NiagaraEditor/Private/NiagaraCompilationCopy.cpp
@@ -479,12 +479,6 @@ void FNiagaraPrecompileData::CollectBakedRapidIterationParameters(const FNiagara
 	}
 }
 
-
-const TMap<FName, UNiagaraDataInterface*>& FNiagaraCompilationCopyData::GetObjectNameMap()
-{
-	return ObjectPtrDecay(InstantiatedGraph->CachedDataInterfaceInstanceDuplicates);
-}
-
 UNiagaraDataInterface* FNiagaraCompilationCopyData::GetDuplicatedDataInterfaceCDOForClass(UClass* Class) const
 {
 	if (UNiagaraDataInterface* ClassDefault = AggregatedDataInterfaceCDODuplicates.FindRef(Class))
@@ -498,13 +492,13 @@ UNiagaraDataInterface* FNiagaraCompilationCopyData::GetDuplicatedDataInterfaceCD
 	return nullptr;
 }
 
-void FNiagaraCompilationCopyData::InstantiateCompilationCopy(const FNiagaraCompilationGraph& SourceGraph, const FNiagaraPrecompileData* PrecompileData, ENiagaraScriptUsage InUsage, FNiagaraFixedConstantResolver ConstantResolver)
+void FNiagaraCompilationCopyData::InstantiateCompilationCopy(const FNiagaraCompilationGraphDigested& SourceGraph, const FNiagaraPrecompileData* PrecompileData, ENiagaraScriptUsage InUsage, FNiagaraFixedConstantResolver ConstantResolver)
 {
 	InstantiatedGraph = SourceGraph.Instantiate(PrecompileData, this, ValidUsages, ConstantResolver);
 
 	if (InstantiatedGraph)
 	{
-		AggregatedDataInterfaceCDODuplicates.Append(ObjectPtrDecay(InstantiatedGraph->CachedDataInterfaceCDODuplicates));
+		AggregatedDataInterfaceCDODuplicates.Append(ObjectPtrDecay(InstantiatedGraph->AggregateDataInterfaceCDODuplicates));
 	}
 }
 
diff --git a/Engine/Plugins/FX/Niagara/Source/NiagaraEditor/Private/NiagaraCompilationPrivate.h b/Engine/Plugins/FX/Niagara/Source/NiagaraEditor/Private/NiagaraCompilationPrivate.h
--- a/Engine/Plugins/FX/Niagara/Source/NiagaraEditor/Private/NiagaraCompilationPrivate.h
+++ b/Engine/Plugins/FX/Niagara/Source/NiagaraEditor/Private/NiagaraCompilationPrivate.h
@@ -170,13 +170,12 @@ public:
 	using FParameterMapHistory = TNiagaraParameterMapHistory<FNiagaraCompilationDigestBridge>;
 	using FParameterMapHistoryWithMetaDataBuilder = TNiagaraParameterMapHistoryWithMetaDataBuilder<FNiagaraCompilationDigestBridge>;
 
-	const TMap<FName, UNiagaraDataInterface*>& GetObjectNameMap();
 	UNiagaraDataInterface* GetDuplicatedDataInterfaceCDOForClass(UClass* Class) const;
 
 	TArray<FParameterMapHistory>& GetPrecomputedHistories() { return PrecompiledHistories; }
 	const TArray<FParameterMapHistory>& GetPrecomputedHistories() const { return PrecompiledHistories; }
 
-	void InstantiateCompilationCopy(const FNiagaraCompilationGraph& SourceGraph, const FNiagaraPrecompileData* PrecompileData, ENiagaraScriptUsage InUsage, FNiagaraFixedConstantResolver ConstantResolver);
+	void InstantiateCompilationCopy(const FNiagaraCompilationGraphDigested& SourceGraph, const FNiagaraPrecompileData* PrecompileData, ENiagaraScriptUsage InUsage, FNiagaraFixedConstantResolver ConstantResolver);
 	void CreateParameterMapHistory(const FNiagaraSystemCompilationTask& CompilationTask, const TArray<FNiagaraVariable>& EncounterableVariables, const TArray<FNiagaraVariable>& InStaticVariables, FNiagaraFixedConstantResolver ConstantResolver, TConstArrayView<FNiagaraSimulationStageInfo> SimStages);
 
 	int32 GetDependentRequestCount() const { return EmitterData.Num(); }
@@ -184,7 +183,7 @@ public:
 
 	TArray<ENiagaraScriptUsage> ValidUsages;
 
-	TSharedPtr<FNiagaraCompilationGraph, ESPMode::ThreadSafe> InstantiatedGraph;
+	TSharedPtr<FNiagaraCompilationGraphInstanced, ESPMode::ThreadSafe> InstantiatedGraph;
 
 	TArray<FParameterMapHistory> PrecompiledHistories;
 
diff --git a/Engine/Plugins/FX/Niagara/Source/NiagaraEditor/Private/NiagaraCompilationTasks.cpp b/Engine/Plugins/FX/Niagara/Source/NiagaraEditor/Private/NiagaraCompilationTasks.cpp
--- a/Engine/Plugins/FX/Niagara/Source/NiagaraEditor/Private/NiagaraCompilationTasks.cpp
+++ b/Engine/Plugins/FX/Niagara/Source/NiagaraEditor/Private/NiagaraCompilationTasks.cpp
@@ -343,8 +343,17 @@ void FNiagaraSystemCompilationTask::Abort()
 	bAborting = true;
 }
 
-FNiagaraSystemCompilationTask::FCompileGroupInfo::FCompileGroupInfo(int32 InEmitterIndex)
-: EmitterIndex(InEmitterIndex)
+
+const FNiagaraSystemCompilationTask::FEmitterInfo* FNiagaraSystemCompilationTask::FSystemInfo::EmitterInfoBySourceEmitter(int32 InSourceEmitterIndex) const
+{
+	return EmitterInfo.FindByPredicate([InSourceEmitterIndex](const FEmitterInfo& Info) -> bool
+	{
+		return Info.SourceEmitterIndex == InSourceEmitterIndex;
+	});
+}
+
+FNiagaraSystemCompilationTask::FCompileGroupInfo::FCompileGroupInfo(int32 InSourceEmitterIndex)
+: SourceEmitterIndex(InSourceEmitterIndex)
 {}
 
 bool FNiagaraSystemCompilationTask::FCompileGroupInfo::HasOutstandingCompileTasks(const FNiagaraSystemCompilationTask& ParentTask) const
@@ -388,20 +397,18 @@ void FNiagaraSystemCompilationTask::FCompileGroupInfo::InstantiateCompileGraph(c
 	const int32 EmitterCount = ParentTask.SystemInfo.EmitterInfo.Num();
 	BasePtr.EmitterData.Reserve(EmitterCount);
 
-	for (int32 EmitterIt = 0; EmitterIt < EmitterCount; ++EmitterIt)
+	for (const FEmitterInfo& EmitterInfo : ParentTask.SystemInfo.EmitterInfo)
 	{
-		const FEmitterInfo& EmitterInfo = ParentTask.SystemInfo.EmitterInfo[EmitterIt];
-
 		TSharedPtr<FNiagaraCompilationCopyData, ESPMode::ThreadSafe> EmitterPtr = MakeShared<FNiagaraCompilationCopyData, ESPMode::ThreadSafe>();
 		EmitterPtr->EmitterUniqueName = EmitterInfo.UniqueEmitterName;
 		EmitterPtr->ValidUsages = BasePtr.ValidUsages;
 
 		// Don't need to copy the graph if we aren't going to use it.
-		if (EmitterInfo.Enabled && ((EmitterIndex == INDEX_NONE) || (EmitterIt == EmitterIndex)))
+		if (EmitterInfo.Enabled && ((SourceEmitterIndex == INDEX_NONE) || (SourceEmitterIndex == EmitterInfo.SourceEmitterIndex)))
 		{
-			const FNiagaraPrecompileData* EmitterRequestData = static_cast<const FNiagaraPrecompileData*>(ParentTask.SystemPrecompileData->GetDependentRequest(EmitterIt).Get());
+			const FNiagaraPrecompileData* EmitterRequestData = static_cast<const FNiagaraPrecompileData*>(ParentTask.SystemPrecompileData->GetDependentRequest(EmitterInfo.DigestedEmitterIndex).Get());
 
-			if (const FNiagaraCompilationGraph* SourceGraph = EmitterInfo.SourceGraph.Get())
+			if (const FNiagaraCompilationGraphDigested* SourceGraph = EmitterInfo.SourceGraph.Get())
 			{
 				EmitterPtr->InstantiateCompilationCopy(*SourceGraph, EmitterRequestData, ENiagaraScriptUsage::EmitterSpawnScript, EmitterInfo.ConstantResolver);
 			}
@@ -418,7 +425,7 @@ void FNiagaraSystemCompilationTask::FCompileGroupInfo::InstantiateCompileGraph(c
 		// skip the deep copy if we're not compiling the system scripts
 		if (BasePtr.ValidUsages.Contains(ENiagaraScriptUsage::SystemSpawnScript))
 		{
-			if (const FNiagaraCompilationGraph* SourceGraph = ParentTask.SystemInfo.SystemSourceGraph.Get())
+			if (const FNiagaraCompilationGraphDigested* SourceGraph = ParentTask.SystemInfo.SystemSourceGraph.Get())
 			{
 				BasePtr.InstantiateCompilationCopy(*SourceGraph, ParentTask.SystemPrecompileData.Get(), ENiagaraScriptUsage::SystemSpawnScript, ParentTask.SystemInfo.ConstantResolver);
 			}
@@ -436,19 +443,17 @@ void FNiagaraSystemCompilationTask::FCompileGroupInfo::InstantiateCompileGraph(c
 	}
 
 	// Now we can finish off the emitters.
-	for (int32 EmitterIt = 0; EmitterIt < EmitterCount; ++EmitterIt)
+	for (const FEmitterInfo& EmitterInfo : ParentTask.SystemInfo.EmitterInfo)
 	{
-		const FEmitterInfo& EmitterInfo = ParentTask.SystemInfo.EmitterInfo[EmitterIt];
-
 		TArray<FNiagaraVariable> EncounterableEmitterVariables;
-		ParentTask.SystemPrecompileData->GetDependentRequest(EmitterIt)->GatherPreCompiledVariables(FString(), EncounterableEmitterVariables);
+		ParentTask.SystemPrecompileData->GetDependentRequest(EmitterInfo.DigestedEmitterIndex)->GatherPreCompiledVariables(FString(), EncounterableEmitterVariables);
 
-		if (EmitterInfo.Enabled && ((EmitterIndex == INDEX_NONE) || (EmitterIt == EmitterIndex)))
+		if (EmitterInfo.Enabled && ((SourceEmitterIndex == INDEX_NONE) || (SourceEmitterIndex == EmitterInfo.SourceEmitterIndex)))
 		{
 			TArray<FNiagaraVariable> StaticVariablesFromEmitter = ParentTask.SystemInfo.StaticVariableResults;
 			StaticVariablesFromEmitter.Append(EmitterInfo.StaticVariableResults);
 
-			BasePtr.EmitterData[EmitterIt]->CreateParameterMapHistory(ParentTask, EncounterableEmitterVariables, StaticVariablesFromEmitter, EmitterInfo.ConstantResolver, EmitterInfo.SimStages);
+			BasePtr.EmitterData[EmitterInfo.DigestedEmitterIndex]->CreateParameterMapHistory(ParentTask, EncounterableEmitterVariables, StaticVariablesFromEmitter, EmitterInfo.ConstantResolver, EmitterInfo.SimStages);
 		}
 	}
 }
@@ -463,7 +468,7 @@ FNiagaraSystemCompilationTask::FCompileGroupInfo* FNiagaraSystemCompilationTask:
 {
 	return CompileGroups.FindByPredicate([EmitterIndex](const FCompileGroupInfo& GroupInfo) -> bool
 	{
-		return GroupInfo.EmitterIndex == EmitterIndex;
+		return GroupInfo.SourceEmitterIndex == EmitterIndex;
 	});
 }
 
@@ -474,27 +479,31 @@ const FNiagaraSystemCompilationTask::FScriptInfo* FNiagaraSystemCompilationTask:
 
 void FNiagaraSystemCompilationTask::FCompileTaskInfo::CollectNamedDataInterfaces(FNiagaraSystemCompilationTask* SystemCompileTask, const FCompileGroupInfo& GroupInfo)
 {
-	const FNiagaraCompilationCopyData* CompilationCopyData = GroupInfo.EmitterIndex == INDEX_NONE
-		? GroupInfo.CompilationCopy.Get()
-		: static_cast<const FNiagaraCompilationCopyData*>(GroupInfo.CompilationCopy->GetDependentRequest(GroupInfo.EmitterIndex).Get());
-	const FScriptInfo& ScriptInfo = SystemCompileTask->DigestedScriptInfo.FindChecked(ScriptKey);
 	FString UniqueEmitterName;
+	const FNiagaraCompilationCopyData* CompilationCopyData = nullptr;
 
-	if (SystemCompileTask->SystemInfo.EmitterInfo.IsValidIndex(GroupInfo.EmitterIndex))
+	if (const FEmitterInfo* EmitterInfo = SystemCompileTask->SystemInfo.EmitterInfoBySourceEmitter(GroupInfo.SourceEmitterIndex))
+	{
+		CompilationCopyData = static_cast<const FNiagaraCompilationCopyData*>(GroupInfo.CompilationCopy->GetDependentRequest(EmitterInfo->DigestedEmitterIndex).Get());
+		UniqueEmitterName = EmitterInfo->UniqueEmitterName;
+	}
+	else
 	{
-		UniqueEmitterName = SystemCompileTask->SystemInfo.EmitterInfo[GroupInfo.EmitterIndex].UniqueEmitterName;
+		CompilationCopyData = GroupInfo.CompilationCopy.Get();
 	}
 
+	const FScriptInfo& ScriptInfo = SystemCompileTask->DigestedScriptInfo.FindChecked(ScriptKey);
+
 	auto AccumulateInputDataInterfaces = [this, &UniqueEmitterName](const FNiagaraCompilationNode& Node) -> bool
 	{
 		if (const FNiagaraCompilationNodeInput* InputNode = Node.AsType<FNiagaraCompilationNodeInput>())
 		{
-			if (InputNode->DataInterfaceName != NAME_None && InputNode->InstancedDataInterface)
+			if (InputNode->DataInterfaceName != NAME_None && InputNode->DuplicatedDataInterface)
 			{
 				constexpr bool bIsParameterMapDataInterface = false;
 				FName DIName = FNiagaraHlslTranslator::GetDataInterfaceName(InputNode->DataInterfaceName, UniqueEmitterName, bIsParameterMapDataInterface);
 
-				NamedDataInterfaces.Add(DIName, InputNode->InstancedDataInterface);
+				NamedDataInterfaces.Add(DIName, InputNode->DuplicatedDataInterface);
 			}
 		}
 
@@ -524,6 +533,8 @@ void FNiagaraSystemCompilationTask::FCompileTaskInfo::CollectNamedDataInterfaces
 
 	if (UNiagaraScript::IsSystemScript(ScriptInfo.Usage))
 	{
+		check(GroupInfo.SourceEmitterIndex == INDEX_NONE);
+
 		const ENiagaraScriptUsage EmitterUsage = ScriptInfo.Usage == ENiagaraScriptUsage::SystemSpawnScript
 			? ENiagaraScriptUsage::EmitterSpawnScript
 			: ENiagaraScriptUsage::EmitterUpdateScript;
@@ -533,7 +544,6 @@ void FNiagaraSystemCompilationTask::FCompileTaskInfo::CollectNamedDataInterfaces
 			return UNiagaraScript::IsEquivalentUsage(OutputNode.Usage, EmitterUsage) && OutputNode.UsageId == EmitterUsageId;
 		};
 
-		check(GroupInfo.EmitterIndex == INDEX_NONE);
 		for (const FNiagaraCompilationCopyData::FSharedCompilationCopy& EmitterCompilationCopy : CompilationCopyData->EmitterData)
 		{
 			if (const FNiagaraCompilationGraph* EmitterGraph = EmitterCompilationCopy->InstantiatedGraph.Get())
@@ -547,13 +557,19 @@ void FNiagaraSystemCompilationTask::FCompileTaskInfo::CollectNamedDataInterfaces
 
 void FNiagaraSystemCompilationTask::FCompileTaskInfo::TranslateAndIssueCompile(FNiagaraSystemCompilationTask* SystemCompileTask, const FCompileGroupInfo& GroupInfo)
 {
-	const FNiagaraPrecompileData* PrecompileData = GroupInfo.EmitterIndex == INDEX_NONE
-		? SystemCompileTask->SystemPrecompileData.Get()
-		: static_cast<const FNiagaraPrecompileData*>(SystemCompileTask->SystemPrecompileData->GetDependentRequest(GroupInfo.EmitterIndex).Get());
+	const FNiagaraPrecompileData* PrecompileData = nullptr;
+	const FNiagaraCompilationCopyData* CompilationCopyData = nullptr;
 
-	const FNiagaraCompilationCopyData* CompilationCopyData = GroupInfo.EmitterIndex == INDEX_NONE
-		? GroupInfo.CompilationCopy.Get()
-		: static_cast<const FNiagaraCompilationCopyData*>(GroupInfo.CompilationCopy->GetDependentRequest(GroupInfo.EmitterIndex).Get());
+	if (const FEmitterInfo* EmitterInfo = SystemCompileTask->SystemInfo.EmitterInfoBySourceEmitter(GroupInfo.SourceEmitterIndex))
+	{
+		PrecompileData = static_cast<const FNiagaraPrecompileData*>(SystemCompileTask->SystemPrecompileData->GetDependentRequest(EmitterInfo->DigestedEmitterIndex).Get());
+		CompilationCopyData = static_cast<const FNiagaraCompilationCopyData*>(GroupInfo.CompilationCopy->GetDependentRequest(EmitterInfo->DigestedEmitterIndex).Get());
+	}
+	else
+	{
+		PrecompileData = SystemCompileTask->SystemPrecompileData.Get();
+		CompilationCopyData = GroupInfo.CompilationCopy.Get();
+	}
 
 	const FScriptInfo& ScriptInfo = SystemCompileTask->DigestedScriptInfo.FindChecked(ScriptKey);
 
@@ -865,12 +881,18 @@ void FNiagaraSystemCompilationTask::DigestSystemInfo()
 	}
 
 	const TArray<FNiagaraEmitterHandle>& EmitterHandles = System_GT->GetEmitterHandles();
-	const int32 EmitterCount = EmitterHandles.Num();
+	const int32 SourceEmitterCount = EmitterHandles.Num();
 
-	SystemInfo.EmitterInfo.Reserve(EmitterCount);
-	for (const FNiagaraEmitterHandle& Handle : EmitterHandles)
+	SystemInfo.EmitterInfo.Reserve(SourceEmitterCount);
+	for (int32 SourceEmitterIndex = 0; SourceEmitterIndex < SourceEmitterCount; ++SourceEmitterIndex)
 	{
-		const int32 EmitterIndex = SystemInfo.EmitterInfo.Num();
+		const FNiagaraEmitterHandle& Handle = EmitterHandles[SourceEmitterIndex];
+		if (!Handle.GetIsEnabled())
+		{
+			continue;
+		}
+
+		const int32 DigestedEmitterIndex = SystemInfo.EmitterInfo.Num();
 
 		FEmitterInfo& EmitterInfo = SystemInfo.EmitterInfo.AddDefaulted_GetRef();
 		const FVersionedNiagaraEmitter& HandleInstance = Handle.GetInstance();
@@ -882,7 +904,8 @@ void FNiagaraSystemCompilationTask::DigestSystemInfo()
 		EmitterInfo.UniqueInstanceName = Handle.GetUniqueInstanceName();
 		EmitterInfo.Enabled = Handle.GetIsEnabled();
 		EmitterInfo.ConstantResolver = FNiagaraFixedConstantResolver(FCompileConstantResolver(HandleInstance, ENiagaraScriptUsage::EmitterSpawnScript));
-		EmitterInfo.EmitterIndex = EmitterIndex;
+		EmitterInfo.SourceEmitterIndex = SourceEmitterIndex;
+		EmitterInfo.DigestedEmitterIndex = DigestedEmitterIndex;
 		EmitterInfo.SourceGraph = DigestDatabase.CreateGraphDigest(EmitterGraph, ChangeIdBuilder);
 
 		{
@@ -1036,7 +1059,7 @@ void FNiagaraSystemCompilationTask::AddScript(int32 EmitterIndex, UNiagaraScript
 
 	Info.Usage = Script->GetUsage();
 	Info.UsageId = Script->GetUsageId();
-	Info.EmitterIndex = EmitterIndex;
+	Info.SourceEmitterIndex = EmitterIndex;
 	for (UNiagaraScript* DependentScript : NiagaraCompilationCopyImpl::FindDependentScripts(System_GT.Get(), EmitterData, Script))
 	{
 		Info.DependentScripts.AddUnique(DependentScript);
@@ -1112,9 +1135,9 @@ bool FNiagaraSystemCompilationTask::Poll(FNiagaraSystemAsyncCompileResults& Resu
 
 					if (const FScriptInfo* ScriptInfo = DigestedScriptInfo.Find(SourceScript))
 					{
-						if (SystemInfo.EmitterInfo.IsValidIndex(ScriptInfo->EmitterIndex))
+						if (const FEmitterInfo* EmitterInfo = SystemInfo.EmitterInfoBySourceEmitter(ScriptInfo->SourceEmitterIndex))
 						{
-							CompileData.UniqueEmitterName = SystemInfo.EmitterInfo[ScriptInfo->EmitterIndex].UniqueEmitterName;
+							CompileData.UniqueEmitterName = EmitterInfo->UniqueEmitterName;
 						}
 
 						// we also need to incorporate the rapid iteration parameters that we encountered
@@ -1313,7 +1336,7 @@ struct FNiagaraSystemCompilationTask::FCollectStaticVariablesTaskBuilder
 	FCollectStaticVariablesTaskBuilder(FNiagaraSystemCompilationTask& CompilationTask)
 	: FoundStaticVariables(CompilationTask.SystemInfo.StaticVariableResults)
 	{
-		if (FNiagaraCompilationGraph* SystemGraph = CompilationTask.SystemInfo.SystemSourceGraph.Get())
+		if (FNiagaraCompilationGraphDigested* SystemGraph = CompilationTask.SystemInfo.SystemSourceGraph.Get())
 		{
 			FStaticVariableBuilderTaskHandle BuilderTask = MakeShared<FStaticVariableBuilderTask, ESPMode::ThreadSafe>();
 			BuilderTask->CompilationTask = &CompilationTask;
@@ -1330,7 +1353,7 @@ struct FNiagaraSystemCompilationTask::FCollectStaticVariablesTaskBuilder
 	FCollectStaticVariablesTaskBuilder(const FNiagaraSystemCompilationTask& CompilationTask, FEmitterInfo& EmitterInfo)
 	: FoundStaticVariables(EmitterInfo.StaticVariableResults)
 	{
-		if (FNiagaraCompilationGraph* EmitterGraph = EmitterInfo.SourceGraph.Get())
+		if (FNiagaraCompilationGraphDigested* EmitterGraph = EmitterInfo.SourceGraph.Get())
 		{
 			TArray<const FNiagaraCompilationNodeOutput*> OutputNodes;
 			EmitterGraph->FindOutputNodes(OutputNodes);
@@ -1506,20 +1529,18 @@ struct FNiagaraSystemCompilationTask::FBuildRapidIterationTaskBuilder
 		FString UniqueEmitterName;
 		FNiagaraFixedConstantResolver ConstantResolver;
 	
-		if (ScriptInfo.EmitterIndex == INDEX_NONE)
+		if (const FEmitterInfo* EmitterInfo = CompilationTask.SystemInfo.EmitterInfoBySourceEmitter(ScriptInfo.SourceEmitterIndex))
 		{
-			Graph = CompilationTask.SystemInfo.SystemSourceGraph.Get();
-			FoundStaticVariables = &CompilationTask.SystemInfo.StaticVariableResults;
-			ConstantResolver = CompilationTask.SystemInfo.ConstantResolver;
+			Graph = EmitterInfo->SourceGraph.Get();
+			FoundStaticVariables = &EmitterInfo->StaticVariableResults;
+			UniqueEmitterName = EmitterInfo->UniqueEmitterName;
+			ConstantResolver = EmitterInfo->ConstantResolver;
 		}
 		else
 		{
-			FEmitterInfo& EmitterInfo = CompilationTask.SystemInfo.EmitterInfo[ScriptInfo.EmitterIndex];
-
-			Graph = EmitterInfo.SourceGraph.Get();
-			FoundStaticVariables = &EmitterInfo.StaticVariableResults;
-			UniqueEmitterName = EmitterInfo.UniqueEmitterName;
-			ConstantResolver = EmitterInfo.ConstantResolver;
+			Graph = CompilationTask.SystemInfo.SystemSourceGraph.Get();
+			FoundStaticVariables = &CompilationTask.SystemInfo.StaticVariableResults;
+			ConstantResolver = CompilationTask.SystemInfo.ConstantResolver;
 		}
 
 		if (Graph)
@@ -1725,10 +1746,8 @@ UE::Tasks::FTask FNiagaraSystemCompilationTask::BuildRapidIterationParametersAsy
 
 	if (SystemCollectStaticVariableTask.IsValid())
 	{
-		const int32 EmitterCount = SystemInfo.EmitterInfo.Num();
-		for (int32 EmitterIndex = 0; EmitterIndex < EmitterCount; ++EmitterIndex)
+		for (FEmitterInfo& EmitterInfo : SystemInfo.EmitterInfo)
 		{
-			FEmitterInfo& EmitterInfo = SystemInfo.EmitterInfo[EmitterIndex];
 			if (!EmitterInfo.Enabled)
 			{
 				continue;
@@ -1739,7 +1758,7 @@ UE::Tasks::FTask FNiagaraSystemCompilationTask::BuildRapidIterationParametersAsy
 				FTask EmitterStaticVariableTask = EmitterTaskBuilder.LaunchCollectedTasks();
 			}, SystemCollectStaticVariableTask);
 
-			CollectStaticVariableTasks.Add(EmitterIndex, EmitterCollectStaticVariableTask);
+			CollectStaticVariableTasks.Add(EmitterInfo.DigestedEmitterIndex, EmitterCollectStaticVariableTask);
 		}
 
 		TArray<FTask> PendingTasks;
@@ -1747,11 +1766,13 @@ UE::Tasks::FTask FNiagaraSystemCompilationTask::BuildRapidIterationParametersAsy
 		{
 			FScriptInfo& ScriptInfo = CurrentIt.Value;
 
+			const FEmitterInfo* EmitterInfo = SystemInfo.EmitterInfoBySourceEmitter(ScriptInfo.SourceEmitterIndex);
+
 			PendingTasks.Add(Launch(UE_SOURCE_LOCATION, [this, &ScriptInfo]
 			{
 				FBuildRapidIterationTaskBuilder TaskBuilder(*this, ScriptInfo);
 				AddNested(TaskBuilder.LaunchCollectedTasks());
-			}, CollectStaticVariableTasks.FindRef(ScriptInfo.EmitterIndex)));
+			}, CollectStaticVariableTasks.FindRef(EmitterInfo ? EmitterInfo->DigestedEmitterIndex : INDEX_NONE)));
 		}
 
 		// lastly we need a single task to copy over parameters between dependent scripts
diff --git a/Engine/Plugins/FX/Niagara/Source/NiagaraEditor/Private/NiagaraCompilationTasks.h b/Engine/Plugins/FX/Niagara/Source/NiagaraEditor/Private/NiagaraCompilationTasks.h
--- a/Engine/Plugins/FX/Niagara/Source/NiagaraEditor/Private/NiagaraCompilationTasks.h
+++ b/Engine/Plugins/FX/Niagara/Source/NiagaraEditor/Private/NiagaraCompilationTasks.h
@@ -68,7 +68,11 @@ struct FNiagaraSystemCompilationTask
 		FNiagaraFixedConstantResolver ConstantResolver;
 		TArray<TObjectKey<UNiagaraScript>> OwnedScriptKeys;
 
-		int32 EmitterIndex;
+		// Index into the EmitterHandles[] of the owning System
+		int32 SourceEmitterIndex;
+
+		// Index of the this in FSystemInfo::EmitterInfo
+		int32 DigestedEmitterIndex;
 		bool Enabled;
 	};
 
@@ -87,17 +91,19 @@ struct FNiagaraSystemCompilationTask
 
 		bool bUseRapidIterationParams;
 		bool bDisableDebugSwitches;
+
+		const FEmitterInfo* EmitterInfoBySourceEmitter(int32 InSourceEmitterIndex) const;
 	};
 
 	struct FCompileGroupInfo
 	{
 		FCompileGroupInfo() = delete;
-		FCompileGroupInfo(int32 InEmitterIndex);
+		FCompileGroupInfo(int32 InSourceEmitterIndex);
 
 		bool HasOutstandingCompileTasks(const FNiagaraSystemCompilationTask& ParentTask) const;
 		void InstantiateCompileGraph(const FNiagaraSystemCompilationTask& ParentTask);
 
-		const int32 EmitterIndex;
+		const int32 SourceEmitterIndex;
 		TArray<int32> CompileTaskIndices;
 		TArray<ENiagaraScriptUsage> ValidUsages;
 		TSharedPtr<FNiagaraCompilationCopyData> CompilationCopy;
@@ -109,7 +115,7 @@ struct FNiagaraSystemCompilationTask
 		FGuid UsageId;
 		FNiagaraParameterStore RapidIterationParameters;
 		TArray<TObjectKey<UNiagaraScript>> DependentScripts;
-		int32 EmitterIndex = INDEX_NONE;
+		int32 SourceEmitterIndex = INDEX_NONE;
 	};
 
 	struct FCompileTaskInfo
@@ -190,7 +196,6 @@ private:
 	TUniquePtr<FDispatchDataCachePutRequests> PutRequestHelper;
 
 	TMap<TObjectKey<UNiagaraScript>, FScriptInfo> DigestedScriptInfo;
-	TSet<FNiagaraDigestedGraphPtr> DigestedGraphs;
 
 	TMap<TObjectKey<UNiagaraParameterCollection>, FNiagaraCompilationNPCHandle> DigestedParameterCollections;
 
diff --git a/Engine/Plugins/FX/Niagara/Source/NiagaraEditor/Private/NiagaraDigestDatabase.cpp b/Engine/Plugins/FX/Niagara/Source/NiagaraEditor/Private/NiagaraDigestDatabase.cpp
--- a/Engine/Plugins/FX/Niagara/Source/NiagaraEditor/Private/NiagaraDigestDatabase.cpp
+++ b/Engine/Plugins/FX/Niagara/Source/NiagaraEditor/Private/NiagaraDigestDatabase.cpp
@@ -56,7 +56,7 @@ void FNiagaraDigestDatabase::ReleaseDatabase()
 {
 	FWriteScopeLock WriteScope(DigestCacheLock);
 
-	CompilationGraphCache.Empty();
+	CompilationGraphCache.Empty(NiagaraGraphDigestDatabaseImpl::GDigestGraphCacheSize);
 	CompilationNPCCache.Empty();
 }
 
@@ -64,11 +64,6 @@ void FNiagaraDigestDatabase::AddReferencedObjects(FReferenceCollector& Collector
 {
 	FReadScopeLock ReadScope(DigestCacheLock);
 
-	for (FCompilationGraphCache::TIterator It(CompilationGraphCache); It; ++It)
-	{
-		It.Value()->AddReferencedObjects(Collector);
-	}
-
 	for (FCompilationNPCCache::TIterator It(CompilationNPCCache); It; ++It)
 	{
 		It.Value()->AddReferencedObjects(Collector);
@@ -107,7 +102,7 @@ FNiagaraDigestedGraphPtr FNiagaraDigestDatabase::CreateGraphDigest(const UNiagar
 
 	if (PendingGraph)
 	{
-		PendingGraph->Create(Graph, Digester);
+		PendingGraph->Digest(Graph, Digester);
 	}
 
 	return PendingGraph;
diff --git a/Engine/Plugins/FX/Niagara/Source/NiagaraEditor/Private/NiagaraDigestDatabase.h b/Engine/Plugins/FX/Niagara/Source/NiagaraEditor/Private/NiagaraDigestDatabase.h
--- a/Engine/Plugins/FX/Niagara/Source/NiagaraEditor/Private/NiagaraDigestDatabase.h
+++ b/Engine/Plugins/FX/Niagara/Source/NiagaraEditor/Private/NiagaraDigestDatabase.h
@@ -13,7 +13,7 @@ class UNiagaraGraph;
 class UNiagaraParameterCollection;
 class UNiagaraScriptSourceBase;
 
-using FNiagaraDigestedGraphPtr = TSharedPtr<FNiagaraCompilationGraph, ESPMode::ThreadSafe>;
+using FNiagaraDigestedGraphPtr = TSharedPtr<FNiagaraCompilationGraphDigested, ESPMode::ThreadSafe>;
 
 class FNiagaraCompilationGraphHandle
 {
@@ -116,7 +116,6 @@ public:
 	void ReleaseDatabase();
 
 	virtual void AddReferencedObjects(FReferenceCollector& Collector) override;
-	
 	virtual FString GetReferencerName() const override;
 
 protected:
diff --git a/Engine/Plugins/FX/Niagara/Source/NiagaraEditor/Private/NiagaraGraphDigest.cpp b/Engine/Plugins/FX/Niagara/Source/NiagaraEditor/Private/NiagaraGraphDigest.cpp
--- a/Engine/Plugins/FX/Niagara/Source/NiagaraEditor/Private/NiagaraGraphDigest.cpp
+++ b/Engine/Plugins/FX/Niagara/Source/NiagaraEditor/Private/NiagaraGraphDigest.cpp
@@ -44,13 +44,13 @@ NIAGARA_GRAPH_DIGEST_NODE_TYPE_LIST;
 
 struct FNiagaraCompilationGraphCreateContext
 {
-	FNiagaraCompilationGraphCreateContext(FNiagaraCompilationGraph& InParentGraph, const FNiagaraGraphChangeIdBuilder& InChangeIdBuilder)
+	FNiagaraCompilationGraphCreateContext(FNiagaraCompilationGraphDigested& InParentGraph, const FNiagaraGraphChangeIdBuilder& InChangeIdBuilder)
 		: ParentGraph(InParentGraph)
 		, ChangeIdBuilder(InChangeIdBuilder)
 	{
 	}
 
-	FNiagaraCompilationGraph& ParentGraph;
+	FNiagaraCompilationGraphDigested& ParentGraph;
 	const FNiagaraGraphChangeIdBuilder& ChangeIdBuilder;
 };
 
@@ -531,7 +531,7 @@ FGuid FNiagaraGraphChangeIdBuilder::RecursiveBuildGraphChangeId(const UNiagaraGr
 	return ChangeId;
 }
 
-void FNiagaraCompilationGraph::Create(const UNiagaraGraph* InGraph, const FNiagaraGraphChangeIdBuilder& ChangeIdBuilder)
+void FNiagaraCompilationGraphDigested::Digest(const UNiagaraGraph* InGraph, const FNiagaraGraphChangeIdBuilder& ChangeIdBuilder)
 {
 	using namespace NiagaraCompilationImpl;
 
@@ -686,34 +686,32 @@ void FNiagaraCompilationGraph::Create(const UNiagaraGraph* InGraph, const FNiaga
 
 	// go through all the data interfaces that have been registered and make sure we populate CachedDataInterfaceCDODuplicates
 	// along with moving things over to transient duplicates
-	if (!CachedDataInterfaceInstanceDuplicates.IsEmpty())
+	for (const FDataInterfaceDuplicateMap::ElementType& DataInterfaceIt : CachedDataInterfaceDuplicates)
 	{
-		for (TMap<FName, TObjectPtr<UNiagaraDataInterface>>::TIterator It(CachedDataInterfaceInstanceDuplicates); It; ++It)
+		if (UNiagaraDataInterface* DataInterface = DataInterfaceIt.Value.Get())
 		{
-			UNiagaraDataInterface* SourceDataInterface = It.Value();
-			UClass* DataInterfaceClass = SourceDataInterface->GetClass();
-
-			RegisterTransientCDO(TransientPackage, DataInterfaceClass);
-
-			It.Value() = DuplicateObject<UNiagaraDataInterface>(SourceDataInterface, TransientPackage);
+			if (UClass* DataInterfaceClass = DataInterface->GetClass())
+			{
+				RegisterTransientCDO(TransientPackage, DataInterfaceClass);
+			}
 		}
 	}
 }
 
-void FNiagaraCompilationGraph::RegisterDataInterface(FName VariableName, UNiagaraDataInterface* SourceDataInterface)
+UNiagaraDataInterface* FNiagaraCompilationGraphDigested::DigestDataInterface(UNiagaraDataInterface* SourceDataInterface)
 {
-	UNiagaraDataInterface* ExistingDataInterface = CachedDataInterfaceInstanceDuplicates.FindRef(VariableName);
-	if (ExistingDataInterface)
-	{
-		check(ExistingDataInterface->Equals(SourceDataInterface));
-	}
-	else
+	check(IsInGameThread());
+	TObjectPtr<UNiagaraDataInterface>& ExistingDuplicate = CachedDataInterfaceDuplicates.FindOrAdd(SourceDataInterface);
+
+	if (!ExistingDuplicate)
 	{
-		CachedDataInterfaceInstanceDuplicates.Add(VariableName, SourceDataInterface);
+		ExistingDuplicate = DuplicateObject<UNiagaraDataInterface>(SourceDataInterface, GetTransientPackage());
 	}
+
+	return ExistingDuplicate;
 }
 
-void FNiagaraCompilationGraph::RegisterObjectAsset(FName VariableName, UObject* SourceObjectAsset)
+void FNiagaraCompilationGraphDigested::RegisterObjectAsset(FName VariableName, UObject* SourceObjectAsset)
 {
 	const UObject* ExistingObjectAsset = CachedNamedObjectAssets.FindRef(VariableName);
 	if (ExistingObjectAsset)
@@ -974,7 +972,7 @@ TArray<const FNiagaraCompilationNode*> FNiagaraCompilationGraph::FindOutputNodes
 	return OutputNodesByUsage;
 }
 
-TSharedPtr<FNiagaraCompilationGraph, ESPMode::ThreadSafe> FNiagaraCompilationGraph::DuplicateSubGraph(
+TSharedPtr<FNiagaraCompilationGraphInstanced, ESPMode::ThreadSafe> FNiagaraCompilationGraphDigested::InstantiateSubGraph(
 	const TArray<ENiagaraScriptUsage>& Usages,
 	const FNiagaraCompilationCopyData* CopyCompilationData,
 	const FNiagaraCompilationBranchMap& Branches,
@@ -982,8 +980,8 @@ TSharedPtr<FNiagaraCompilationGraph, ESPMode::ThreadSafe> FNiagaraCompilationGra
 {
 	using namespace NiagaraCompilationImpl;
 
-	TSharedPtr<FNiagaraCompilationGraph, ESPMode::ThreadSafe> SubGraph = MakeShared<FNiagaraCompilationGraph, ESPMode::ThreadSafe>();
-	SubGraph->InstantiationSourceGraph = this;
+	TSharedPtr<FNiagaraCompilationGraphInstanced, ESPMode::ThreadSafe> SubGraph = MakeShared<FNiagaraCompilationGraphInstanced, ESPMode::ThreadSafe>();
+	SubGraph->InstantiationSourceGraph = AsShared().ToSharedPtr();
 
 	TArray<const FNiagaraCompilationNode*> OutputNodesByUsage = FindOutputNodesByUsage(Usages);
 	const int32 InstantiatedOutputNodeCount = OutputNodesByUsage.Num();
@@ -1009,7 +1007,6 @@ TSharedPtr<FNiagaraCompilationGraph, ESPMode::ThreadSafe> FNiagaraCompilationGra
 
 	// copy over all the properties for the subgraph (this will contain data that could have been culled based on the
 	// output nodes that we're filtering by above
-	SubGraph->SourceGraph = SourceGraph;
 	SubGraph->VariableBinding = VariableBinding;
 	SubGraph->ScriptVariableData = ScriptVariableData;
 	SubGraph->SourceScriptName = SourceScriptName;
@@ -1029,14 +1026,12 @@ TSharedPtr<FNiagaraCompilationGraph, ESPMode::ThreadSafe> FNiagaraCompilationGra
 	}
 
 	// make sure that we've collected the CDO for the registered DI as well
-	SubGraph->CachedDataInterfaceCDODuplicates.Append(CachedDataInterfaceCDODuplicates);
-	SubGraph->CachedDataInterfaceInstanceDuplicates.Append(CachedDataInterfaceInstanceDuplicates);
-	SubGraph->bInstanced = true;
+	SubGraph->AggregateDataInterfaceCDODuplicates.Append(CachedDataInterfaceCDODuplicates);
 
 	return SubGraph;
 }
 
-void FNiagaraCompilationGraph::ValidateRefinement() const
+void FNiagaraCompilationGraphInstanced::ValidateRefinement() const
 {
 	// validate that now that we've refined the graph we have no more generic numerics and also ensure
 	// that there are no more static switches connected
@@ -1070,7 +1065,7 @@ void FNiagaraCompilationGraph::ValidateRefinement() const
 	}
 }
 
-void FNiagaraCompilationGraph::Refine(FNiagaraCompilationGraphInstanceContext& InstantiationContext, const FNiagaraCompilationNodeFunctionCall* CallingNode)
+void FNiagaraCompilationGraphInstanced::Refine(FNiagaraCompilationGraphInstanceContext& InstantiationContext, const FNiagaraCompilationNodeFunctionCall* CallingNode)
 {
 	if (CallingNode)
 	{
@@ -1097,7 +1092,7 @@ void FNiagaraCompilationGraph::Refine(FNiagaraCompilationGraphInstanceContext& I
 	ValidateRefinement();
 }
 
-void FNiagaraCompilationGraph::PatchGenericNumericsFromCaller(FNiagaraCompilationGraphInstanceContext& Context)
+void FNiagaraCompilationGraphInstanced::PatchGenericNumericsFromCaller(FNiagaraCompilationGraphInstanceContext& Context)
 {
 	static const FNiagaraTypeDefinition& GenericTypeDef = FNiagaraTypeDefinition::GetGenericNumericDef();
 	const FNiagaraCompilationNodeFunctionCall* CallingNode = Context.GetCurrentFunctionNode();
@@ -1147,7 +1142,7 @@ void FNiagaraCompilationGraph::PatchGenericNumericsFromCaller(FNiagaraCompilatio
 	}
 }
 
-void FNiagaraCompilationGraph::ResolveNumerics(FNiagaraCompilationGraphInstanceContext& Context)
+void FNiagaraCompilationGraphInstanced::ResolveNumerics(FNiagaraCompilationGraphInstanceContext& Context)
 {
 	const FNiagaraTypeDefinition& GenericTypeDef = FNiagaraTypeDefinition::GetGenericNumericDef();
 	const FNiagaraTypeDefinition& PlaceholderTypeDef = FNiagaraTypeDefinition::GetFloatDef();
@@ -1200,12 +1195,12 @@ void FNiagaraCompilationGraph::ResolveNumerics(FNiagaraCompilationGraphInstanceC
 	}
 }
 
-void FNiagaraCompilationGraph::InheritDebugState(FNiagaraCompilationGraphInstanceContext& Context, FNiagaraCompilationNodeFunctionCall& FunctionCallNode)
+void FNiagaraCompilationGraphInstanced::InheritDebugState(FNiagaraCompilationGraphInstanceContext& Context, FNiagaraCompilationNodeFunctionCall& FunctionCallNode)
 {
 	FunctionCallNode.DebugState = FunctionCallNode.bInheritDebugState ? Context.ConstantResolver.GetDebugState() : ENiagaraFunctionDebugState::NoDebug;
 }
 
-void FNiagaraCompilationGraph::PropagateDefaultValues(FNiagaraCompilationGraphInstanceContext& Context, FNiagaraCompilationNodeFunctionCall& FunctionCallNode)
+void FNiagaraCompilationGraphInstanced::PropagateDefaultValues(FNiagaraCompilationGraphInstanceContext& Context, FNiagaraCompilationNodeFunctionCall& FunctionCallNode)
 {
 	using namespace NiagaraCompilationImpl;
 
@@ -1231,7 +1226,7 @@ void FNiagaraCompilationGraph::PropagateDefaultValues(FNiagaraCompilationGraphIn
 	}
 }
 
-TSharedPtr<FNiagaraCompilationGraph, ESPMode::ThreadSafe> FNiagaraCompilationGraph::Instantiate(const FNiagaraPrecompileData* PrecompileData, const FNiagaraCompilationCopyData* CopyCompilationData, const TArray<ENiagaraScriptUsage>& Usages, const FNiagaraFixedConstantResolver& ConstantResolver) const
+TSharedPtr<FNiagaraCompilationGraphInstanced, ESPMode::ThreadSafe> FNiagaraCompilationGraphDigested::Instantiate(const FNiagaraPrecompileData* PrecompileData, const FNiagaraCompilationCopyData* CopyCompilationData, const TArray<ENiagaraScriptUsage>& Usages, const FNiagaraFixedConstantResolver& ConstantResolver) const
 {
 	// in order to manage the traversal state context and to keep a lit on the potential for crazy recursion we're
 	// going to track the dependent graphs that require instantiation in this quasi tree structure
@@ -1248,7 +1243,7 @@ TSharedPtr<FNiagaraCompilationGraph, ESPMode::ThreadSafe> FNiagaraCompilationGra
 	};
 
 	TArray<FChildrenStackEntry> FunctionsToInstantiate;
-	TSharedPtr<FNiagaraCompilationGraph, ESPMode::ThreadSafe> InstantiatedGraph = DuplicateSubGraph(Usages, CopyCompilationData, FNiagaraCompilationBranchMap(), FunctionsToInstantiate.Emplace_GetRef(INDEX_NONE).Functions);
+	TSharedPtr<FNiagaraCompilationGraphInstanced, ESPMode::ThreadSafe> InstantiatedGraph = InstantiateSubGraph(Usages, CopyCompilationData, FNiagaraCompilationBranchMap(), FunctionsToInstantiate.Emplace_GetRef(INDEX_NONE).Functions);
 
 	// initialize the traversal context with the data that was pulled from the parameter map history done during the
 	// precompile
@@ -1279,15 +1274,23 @@ TSharedPtr<FNiagaraCompilationGraph, ESPMode::ThreadSafe> FNiagaraCompilationGra
 			FNiagaraCompilationBranchMap Branches;
 			EvaluateStaticBranches(InstantiationContext, Branches);
 
-			const auto& OriginalGraph = FunctionToInstantiate->CalledGraph;
-			FunctionToInstantiate->CalledGraph = FunctionToInstantiate->CalledGraph->DuplicateSubGraph({FunctionToInstantiate->CalledScriptUsage}, CopyCompilationData, Branches, ChildStack.Functions);
-			FunctionToInstantiate->CalledGraph->Refine(InstantiationContext, FunctionToInstantiate);
+			FNiagaraCompilationGraphDigested* CalledDigestedGraph = FunctionToInstantiate->CalledGraph->AsDigested();
+			if (ensure(CalledDigestedGraph))
+			{
+				TSharedPtr<FNiagaraCompilationGraphInstanced, ESPMode::ThreadSafe> CalledInstantiatedGraph =
+					CalledDigestedGraph->InstantiateSubGraph({ FunctionToInstantiate->CalledScriptUsage }, CopyCompilationData, Branches, ChildStack.Functions);
+
+				CalledInstantiatedGraph->Refine(InstantiationContext, FunctionToInstantiate);
 
-			InstantiatedGraph->AggregateDataInterfaces(FunctionToInstantiate->CalledGraph.Get());
+				InstantiatedGraph->AggregateDataInterfaces(CalledInstantiatedGraph.Get());
 
-			++TotalGraphCount;
-			TotalNodeCount += FunctionToInstantiate->CalledGraph->Nodes.Num();
-			TotalCulledNodeCount += OriginalGraph->Nodes.Num() - FunctionToInstantiate->CalledGraph->Nodes.Num();
+				++TotalGraphCount;
+				TotalNodeCount += CalledInstantiatedGraph->Nodes.Num();
+				TotalCulledNodeCount += CalledDigestedGraph->Nodes.Num() - CalledInstantiatedGraph->Nodes.Num();
+
+				// now replace the called graph with the instantiated version
+				FunctionToInstantiate->CalledGraph = CalledInstantiatedGraph;
+			}
 
 			CurrentFunctionSetIndex = NextFunctionSetIndex;
 		}
@@ -1304,16 +1307,18 @@ TSharedPtr<FNiagaraCompilationGraph, ESPMode::ThreadSafe> FNiagaraCompilationGra
 					const int32 ChildFunctionIndex = InstantiationContext.FunctionStack.Num() - 1;
 					ensure(InstantiationContext.FunctionStack.IsValidIndex(ChildFunctionIndex));
 
-					FNiagaraCompilationGraph* ParentGraph = InstantiationContext.FunctionStack.IsValidIndex(ParentFunctionIndex)
-						? InstantiationContext.FunctionStack[ParentFunctionIndex]->CalledGraph.Get()
+					FNiagaraCompilationGraphInstanced* ParentGraph = InstantiationContext.FunctionStack.IsValidIndex(ParentFunctionIndex)
+						? InstantiationContext.FunctionStack[ParentFunctionIndex]->CalledGraph->AsInstanced()
 						: InstantiatedGraph.Get();
 
-					FNiagaraCompilationGraph* ChildGraph = InstantiationContext.FunctionStack[ChildFunctionIndex]->CalledGraph.Get();
-
-					ParentGraph->AggregateChildGraph(ChildGraph);
+					FNiagaraCompilationGraphInstanced* ChildGraph = InstantiationContext.FunctionStack[ChildFunctionIndex]->CalledGraph->AsInstanced();
+					
+					if (ensure(ParentGraph && ChildGraph))
+					{
+						ParentGraph->AggregateChildGraph(ChildGraph);
+					}
 				}
 
-
 				FChildrenStackEntry& ParentEntry = FunctionsToInstantiate[CurrentFunctionSetIndex];
 				InstantiationContext.LeaveFunction(ParentEntry.Functions[ParentEntry.CurrentFunctionIndex]);
 
@@ -1411,45 +1416,13 @@ void FNiagaraCompilationGraph::NodeTraversal(
 }
 
 
-void FNiagaraCompilationGraph::AggregateDataInterfaces(const FNiagaraCompilationGraph* ChildGraph)
+void FNiagaraCompilationGraphInstanced::AggregateDataInterfaces(const FNiagaraCompilationGraphInstanced* ChildGraph)
 {
-	for (TMap<FName, TObjectPtr<UNiagaraDataInterface>>::TConstIterator ChildIt(ChildGraph->CachedDataInterfaceInstanceDuplicates); ChildIt; ++ChildIt)
-	{
-		CachedDataInterfaceInstanceDuplicates.Add(ChildIt.Key(), ChildIt.Value());
-
-		// matching behavior in FNiagaraCompileRequestDuplicateData::DuplicateReferencedGraphsRecursive where
-		// child graphs override the values of the parent when it comes to the DIs...this makes me a bit
-		// nervous because there's no concept there of the nodes actually being relevant to the compile
-		// for this version that should at least be taken care of...but weird DI juju for sure
-		#if 0
-		UNiagaraDataInterface* ExistingDI = CachedDataInterfaceInstanceDuplicates.FindRef(ChildIt.Key());
-		if (ExistingDI)
-		{
-			check(ExistingDI->Equals(ChildIt.Value()));
-		}
-		else
-		{
-			CachedDataInterfaceInstanceDuplicates.Add(ChildIt.Key(), ChildIt.Value());
-		}
-		#endif
-	}
-
-	for (TMap<TObjectPtr<UClass>, TObjectPtr<UNiagaraDataInterface>>::TConstIterator ChildIt(ChildGraph->CachedDataInterfaceCDODuplicates); ChildIt; ++ChildIt)
-	{
-		UNiagaraDataInterface* ChildDataInterface = ChildIt.Value();
-		UNiagaraDataInterface* ExistingDI = CachedDataInterfaceCDODuplicates.FindRef(ChildIt.Key());
-		if (ExistingDI)
-		{
-			check(ExistingDI->Equals(ChildDataInterface));
-		}
-		else
-		{
-			CachedDataInterfaceCDODuplicates.Add(ChildIt.Key(), ChildIt.Value());
-		}
-	}
+	//AggregateDataInterfaceCDODuplicates.Append(ObjectPtrDecay(ChildGraph->AggregateDataInterfaceCDODuplicates));
+	AggregateDataInterfaceCDODuplicates.Append(ChildGraph->AggregateDataInterfaceCDODuplicates);
 }
 
-void FNiagaraCompilationGraph::AggregateChildGraph(const FNiagaraCompilationGraph* ChildGraph)
+void FNiagaraCompilationGraphInstanced::AggregateChildGraph(const FNiagaraCompilationGraphInstanced* ChildGraph)
 {
 	if (ChildGraph->bContainsStaticVariables)
 	{
@@ -1473,13 +1446,18 @@ void FNiagaraCompilationGraph::BuildTraversal(const FNiagaraCompilationNode* Roo
 	OrderedNodes = TopologicalSort(CollectConnectedNodes({ RootNode }, Nodes.Num(), Branches), Branches);
 }
 
-void FNiagaraCompilationGraph::AddReferencedObjects(FReferenceCollector& Collector)
+void FNiagaraCompilationGraphDigested::AddReferencedObjects(FReferenceCollector& Collector)
 {
-	Collector.AddReferencedObjects(CachedDataInterfaceInstanceDuplicates);
+	Collector.AddReferencedObjects(CachedDataInterfaceDuplicates);
 	Collector.AddReferencedObjects(CachedDataInterfaceCDODuplicates);
 	Collector.AddReferencedObjects(CachedNamedObjectAssets);
 }
 
+FString FNiagaraCompilationGraphDigested::GetReferencerName() const
+{
+	return TEXT("FNiagaraCompilationGraphDigested");
+}
+
 TArray<const FNiagaraCompilationNode*> FNiagaraCompilationGraph::GetOutputNodes() const
 {
 	TArray<const FNiagaraCompilationNode*> OutputNodes;
@@ -1660,7 +1638,7 @@ FNiagaraCompilationNode::FNiagaraCompilationNode(ENodeType InNodeType, const UEd
 
 		SourceNode = InNode;
 		OwningGraph = &Context.ParentGraph;
-		const UEdGraph* SourceOwningGraph = OwningGraph->SourceGraph.Get();
+		const UEdGraph* SourceOwningGraph = Context.ParentGraph.SourceGraph.Get();
 
 		const int32 PinCount = SourceNode->Pins.Num();
 
@@ -2057,10 +2035,16 @@ FNiagaraCompilationNodeEmitter::FNiagaraCompilationNodeEmitter(const FNiagaraCom
 		return EmitterCopy->EmitterUniqueName == EmitterUniqueName;
 	});
 
-	if (ensure(EmitterCopy))
+	// because toggling the enabled state of an emitter doesn't necessarily change the emitter node in the graph
+	// and so the EmitterCopy may be null in this case.  If that's true, we'll mark the node as being disabled.
+	if (EmitterCopy)
 	{
 		CalledGraph = (*EmitterCopy)->InstantiatedGraph;
 	}
+	else
+	{
+		NodeEnabled = false;
+	}
 }
 
 void FNiagaraCompilationNodeEmitter::BuildParameterMapHistory(FParameterMapHistoryBuilder& Builder, bool bRecursive, bool bFilterForCompilation) const
@@ -3099,7 +3083,7 @@ FNiagaraCompilationNodeInput::FNiagaraCompilationNodeInput(const UNiagaraNodeInp
 		DataInterfaceName = InputVariable.GetName();
 		check(SourceDataInterface);
 		SourceDataInterface->GetEmitterReferencesByName(DataInterfaceEmitterReferences);
-		Context.ParentGraph.RegisterDataInterface(DataInterfaceName, SourceDataInterface);
+		DuplicatedDataInterface = Context.ParentGraph.DigestDataInterface(SourceDataInterface);
 	}
 	else if (InputVariable.IsUObject())
 	{
@@ -3128,17 +3112,10 @@ FNiagaraCompilationNodeInput::FNiagaraCompilationNodeInput(const FNiagaraCompila
 	, bRequired(InNode.bRequired)
 	, bExposed(InNode.bExposed)
 	, bCanAutoBind(InNode.bCanAutoBind)
+	, DuplicatedDataInterface(InNode.DuplicatedDataInterface)
 	, ObjectAssetPath(InNode.ObjectAssetPath)
 {
 	Context.TargetGraph.InputNodeIndices.Add(Context.TargetGraph.Nodes.Num());
-	if (!DataInterfaceName.IsNone())
-	{
-		InstancedDataInterface = Context.SourceGraph.CachedDataInterfaceInstanceDuplicates.FindRef(DataInterfaceName);
-		if (ensure(InstancedDataInterface))
-		{
-			Context.TargetGraph.RegisterDataInterface(DataInterfaceName, InstancedDataInterface);
-		}
-	}
 }
 
 void FNiagaraCompilationNodeInput::BuildParameterMapHistory(FParameterMapHistoryBuilder& Builder, bool bRecursive, bool bFilterForCompilation) const
@@ -3209,9 +3186,9 @@ void FNiagaraCompilationNodeInput::Compile(FTranslator* Translator, TArray<int32
 				{
 					if (!DataInterfaceName.IsNone())
 					{
-						check(InstancedDataInterface);
+						check(DuplicatedDataInterface);
 						check(InputVariable.IsDataInterface());
-						Outputs.Add(Translator->RegisterDataInterface(InputVariable, InstancedDataInterface, false, false));
+						Outputs.Add(Translator->RegisterDataInterface(InputVariable, DuplicatedDataInterface, false, false));
 						return;
 					}
 					else if (!ObjectAssetName.IsNone())
@@ -3235,9 +3212,9 @@ void FNiagaraCompilationNodeInput::Compile(FTranslator* Translator, TArray<int32
 	case ENiagaraInputNodeUsage::Parameter:
 		if (!DataInterfaceName.IsNone())
 		{
-			check(InstancedDataInterface);
+			check(DuplicatedDataInterface);
 			check(InputVariable.IsDataInterface());
-			Outputs.Add(Translator->RegisterDataInterface(InputVariable, InstancedDataInterface, false, false));
+			Outputs.Add(Translator->RegisterDataInterface(InputVariable, DuplicatedDataInterface, false, false));
 			break;
 		}
 		else if (!ObjectAssetName.IsNone())
diff --git a/Engine/Plugins/FX/Niagara/Source/NiagaraEditor/Private/NiagaraGraphDigest.h b/Engine/Plugins/FX/Niagara/Source/NiagaraEditor/Private/NiagaraGraphDigest.h
--- a/Engine/Plugins/FX/Niagara/Source/NiagaraEditor/Private/NiagaraGraphDigest.h
+++ b/Engine/Plugins/FX/Niagara/Source/NiagaraEditor/Private/NiagaraGraphDigest.h
@@ -14,6 +14,8 @@
 class UEdGraphNode;
 class FNiagaraFixedConstantResolver;
 class FNiagaraCompilationGraph;
+class FNiagaraCompilationGraphDigested;
+class FNiagaraCompilationGraphInstanced;
 class FNiagaraCompilationNode;
 class FNiagaraPrecompileData;
 class FNiagaraCompilationCopyData;
@@ -253,30 +255,24 @@ protected:
 class FNiagaraCompilationGraph
 {
 public:
+	virtual ~FNiagaraCompilationGraph() {};
+
 	using FSharedPtr = TSharedPtr<FNiagaraCompilationGraph, ESPMode::ThreadSafe>;
 
-	void Create(const UNiagaraGraph* InGraph, const FNiagaraGraphChangeIdBuilder& Digester);
 	bool IsValid() const { return true; } // todo - when creation is offloaded to a task we'll need this to synchronize completion
 
 	TArray<TUniquePtr<FNiagaraCompilationNode>> Nodes;
-	TWeakObjectPtr<const UNiagaraGraph> SourceGraph;
-	const FNiagaraCompilationGraph* InstantiationSourceGraph = nullptr;
 	TArray<int32> InputNodeIndices;
 	TArray<int32> OutputNodeIndices;
 	FNiagaraScriptVariableBinding VariableBinding;
 	TArray<FNiagaraScriptVariableData> ScriptVariableData;
 	TArray<FNiagaraVariableBase> StaticSwitchInputs;
-	bool bInstanced = false;
+
 	// true if the graph contains a static variable or one of it's recursively referenced graphs contains one
 	bool bContainsStaticVariables = false;
 	FString SourceScriptName;
 	FString SourceScriptFullName;
 
-	TMap<FName, TObjectPtr<UNiagaraDataInterface>> CachedDataInterfaceInstanceDuplicates;
-	TMap<TObjectPtr<UClass>, TObjectPtr<UNiagaraDataInterface>> CachedDataInterfaceCDODuplicates;
-
-	TMap<FName, TObjectPtr<UObject>> CachedNamedObjectAssets;
-
 	void FindOutputNodes(TArray<const FNiagaraCompilationNodeOutput*>& OutputNodes) const;
 	void FindOutputNodes(ENiagaraScriptUsage TargetUsageType, TArray<const FNiagaraCompilationNodeOutput*>& OutputNodes) const;
 	const FNiagaraCompilationNodeOutput* FindOutputNode(ENiagaraScriptUsage Usage, const FGuid& UsageId) const;
@@ -295,15 +291,10 @@ public:
 	TOptional<ENiagaraDefaultMode> GetDefaultMode(const FNiagaraVariableBase&, FNiagaraScriptVariableBinding& Binding) const;
 	TOptional<FNiagaraVariableMetaData> GetMetaData(const FNiagaraVariableBase&) const;
 
-	TSharedPtr<FNiagaraCompilationGraph, ESPMode::ThreadSafe> Instantiate(const FNiagaraPrecompileData* PrecompileData, const FNiagaraCompilationCopyData* CopyCompilationData, const TArray<ENiagaraScriptUsage>& Usages, const FNiagaraFixedConstantResolver& ConstantResolver) const;
 	void CollectReachableNodes(const FNiagaraCompilationNodeOutput* OutputNode, TArray<const FNiagaraCompilationNode*>& ReachableNodes) const;
 	void BuildTraversal(const FNiagaraCompilationNode* RootNode, TArray<const FNiagaraCompilationNode*>& OrderedNodes) const;
 	void BuildTraversal(const FNiagaraCompilationNode* RootNode, const FNiagaraCompilationBranchMap& Branches, TArray<const FNiagaraCompilationNode*>& OrderedNodes) const;
 
-	void RegisterDataInterface(FName VariableName, UNiagaraDataInterface* SourceDataInterface);
-	void RegisterObjectAsset(FName VariableName, UObject* SourceObjectAsset);
-	void AddReferencedObjects(FReferenceCollector& Collector);
-
 	void EvaluateStaticBranches(FNiagaraCompilationGraphInstanceContext& Context, FNiagaraCompilationBranchMap& Branches) const;
 
 	void NodeTraversal(
@@ -312,27 +303,76 @@ public:
 		TFunctionRef<bool(const FNiagaraCompilationNodeOutput&)> RootNodeFilter,
 		TFunctionRef<bool(const FNiagaraCompilationNode&)> NodeOperation) const;
 
+	virtual FNiagaraCompilationGraphDigested* AsDigested() { return nullptr; }
+	virtual FNiagaraCompilationGraphInstanced* AsInstanced() { return nullptr; }
+
 protected:
 	const FNiagaraScriptVariableData* GetScriptVariableData(const FNiagaraVariableBase&) const;
 
-	TSharedPtr<FNiagaraCompilationGraph, ESPMode::ThreadSafe> DuplicateSubGraph(
+	TArray<const FNiagaraCompilationNode*> GetOutputNodes() const;
+};
+
+// Digested version of the compilation graph.  Created on the game thread from a UNiagaraGraph
+// this will act as the source for FNiagaraCompilationGraphInstanced that will be crated in async
+// tasks.  Is also responsible for securing references to UObjects over it's lifetime as an FGCObject
+class FNiagaraCompilationGraphDigested : public TSharedFromThis<FNiagaraCompilationGraphDigested, ESPMode::ThreadSafe>
+										, public FNiagaraCompilationGraph
+										, FGCObject
+{
+public:
+	virtual void AddReferencedObjects(FReferenceCollector& Collector) override;
+	virtual FString GetReferencerName() const override;
+
+	TSharedPtr<FNiagaraCompilationGraphInstanced, ESPMode::ThreadSafe> Instantiate(const FNiagaraPrecompileData* PrecompileData, const FNiagaraCompilationCopyData* CopyCompilationData, const TArray<ENiagaraScriptUsage>& Usages, const FNiagaraFixedConstantResolver& ConstantResolver) const;
+
+	void Digest(const UNiagaraGraph* InGraph, const FNiagaraGraphChangeIdBuilder& Digester);
+
+	UNiagaraDataInterface* DigestDataInterface(UNiagaraDataInterface* SourceDataInterface);
+	void RegisterObjectAsset(FName VariableName, UObject* SourceObjectAsset);
+
+	virtual FNiagaraCompilationGraphDigested* AsDigested() override { return this; }
+
+	using FDataInterfaceDuplicateMap = TMap<TObjectKey<UNiagaraDataInterface>, TObjectPtr<UNiagaraDataInterface>>;
+	FDataInterfaceDuplicateMap CachedDataInterfaceDuplicates;
+	TMap<TObjectPtr<UClass>, TObjectPtr<UNiagaraDataInterface>> CachedDataInterfaceCDODuplicates;
+	TMap<FName, TObjectPtr<UObject>> CachedNamedObjectAssets;
+
+	TWeakObjectPtr<const UNiagaraGraph> SourceGraph;
+
+protected:
+	TSharedPtr<FNiagaraCompilationGraphInstanced, ESPMode::ThreadSafe> InstantiateSubGraph(
 		const TArray<ENiagaraScriptUsage>& Usages,
 		const FNiagaraCompilationCopyData* CopyCompilationData,
 		const FNiagaraCompilationBranchMap& Branches,
 		TArray<FNiagaraCompilationNodeFunctionCall*>& PendingInstantiations) const;
+};
+
+// Instanced version of the compilation graph.  Created from a FNiagaraCompilationGraphDigested
+// this will act as the representation of the UNiagaraGraph that is actually being translated and
+// compiled (i.e. will have nodes stripped out if static branches aren't taken).
+class FNiagaraCompilationGraphInstanced : public FNiagaraCompilationGraph
+{
+public:
+	void AggregateDataInterfaces(const FNiagaraCompilationGraphInstanced* ChildGraph);
+	void AggregateChildGraph(const FNiagaraCompilationGraphInstanced* ChildGraph);
+
+	void ResolveNumerics(FNiagaraCompilationGraphInstanceContext& Context);
 	void Refine(FNiagaraCompilationGraphInstanceContext& Context, const FNiagaraCompilationNodeFunctionCall* CallingNode);
-	void ValidateRefinement() const;
 
+	virtual FNiagaraCompilationGraphInstanced* AsInstanced() { return this; }
+
+	TSharedPtr<const FNiagaraCompilationGraph> InstantiationSourceGraph;
+
+	// the full collection of duplicated CDO data interfaces that are used by this graph (and all
+	// instanced graphs that it references).
+	TMap<TObjectPtr<UClass>, TObjectPtr<UNiagaraDataInterface>> AggregateDataInterfaceCDODuplicates;
+
+protected:
+	void ValidateRefinement() const;
 	void PatchGenericNumericsFromCaller(FNiagaraCompilationGraphInstanceContext& Context);
-	void ResolveNumerics(FNiagaraCompilationGraphInstanceContext& Context);
 	void InheritDebugState(FNiagaraCompilationGraphInstanceContext& Context, FNiagaraCompilationNodeFunctionCall& FunctionCallNode);
 	void PropagateDefaultValues(FNiagaraCompilationGraphInstanceContext& Context, FNiagaraCompilationNodeFunctionCall& FunctionCallNode);
 	void RemoveUnconnectedNodes();
-
-	void AggregateDataInterfaces(const FNiagaraCompilationGraph* ChildGraph);
-	void AggregateChildGraph(const FNiagaraCompilationGraph* ChildGraph);
-
-	TArray<const FNiagaraCompilationNode*> GetOutputNodes() const;
 };
 
 class FNiagaraCompilationScript
@@ -461,7 +501,7 @@ public:
 	bool bExposed;
 	bool bCanAutoBind;
 
-	TObjectPtr<UNiagaraDataInterface> InstancedDataInterface;
+	TObjectPtr<UNiagaraDataInterface> DuplicatedDataInterface;
 	FSoftObjectPath ObjectAssetPath;
 };
 
