Fixed potential nullptr access

diff --git a/Engine/Plugins/FX/Niagara/Source/NiagaraEditor/Private/ViewModels/Stack/NiagaraStackFunctionInput.cpp b/Engine/Plugins/FX/Niagara/Source/NiagaraEditor/Private/ViewModels/Stack/NiagaraStackFunctionInput.cpp
--- a/Engine/Plugins/FX/Niagara/Source/NiagaraEditor/Private/ViewModels/Stack/NiagaraStackFunctionInput.cpp
+++ b/Engine/Plugins/FX/Niagara/Source/NiagaraEditor/Private/ViewModels/Stack/NiagaraStackFunctionInput.cpp
@@ -1757,7 +1757,7 @@ void UNiagaraStackFunctionInput::GetAvailableParameterHandles(TArray<FNiagaraPar
 	}
 
 	// Filter gathered attributes by allowed usage
-	ENiagaraScriptUsage CurrentUsage = CurrentOutputNode->GetUsage();
+	ENiagaraScriptUsage CurrentUsage = CurrentOutputNode ? CurrentOutputNode->GetUsage() : ENiagaraScriptUsage::Module;
 	AvailableParameterHandles.SetNum(Algo::RemoveIf(AvailableParameterHandles, [CurrentUsage](FNiagaraParameterHandle& Handle)
 	{
 		return !IsNamespaceAllowedInUsage(Handle.GetNamespace(), CurrentUsage);
