Improved tooltip for wildcard pins on select node

diff --git a/Engine/Plugins/FX/Niagara/Source/NiagaraEditor/Private/NiagaraNodeSelect.cpp b/Engine/Plugins/FX/Niagara/Source/NiagaraEditor/Private/NiagaraNodeSelect.cpp
--- a/Engine/Plugins/FX/Niagara/Source/NiagaraEditor/Private/NiagaraNodeSelect.cpp
+++ b/Engine/Plugins/FX/Niagara/Source/NiagaraEditor/Private/NiagaraNodeSelect.cpp
@@ -458,7 +458,14 @@ void UNiagaraNodeSelect::GetWildcardPinHoverConnectionTextAddition(const UEdGrap
 {
 	if(ConnectionResponse == ECanCreateConnectionResponse::CONNECT_RESPONSE_DISALLOW)
 	{
-		OutString += TEXT("\nWildcard Pin only allows index types: bool, integer or enums.");
+		if (WildcardPin && WildcardPin->PersistentGuid.IsValid() && WildcardPin->PersistentGuid == SelectorPinGuid)
+		{
+			OutString += TEXT("\nThis wildcard pin only supports bool, integer or enum types.");
+		}
+		else
+		{
+			OutString += TEXT("\nThis wildcard pin only supports simple types (no numerics or data interfaces).");
+		}
 	}
 }
 
