- Fix for component renderer not updating it's coordinate space when changing emitter local space

diff --git a/Engine/Plugins/FX/Niagara/Source/Niagara/Private/NiagaraComponentRendererProperties.cpp b/Engine/Plugins/FX/Niagara/Source/Niagara/Private/NiagaraComponentRendererProperties.cpp
--- a/Engine/Plugins/FX/Niagara/Source/Niagara/Private/NiagaraComponentRendererProperties.cpp
+++ b/Engine/Plugins/FX/Niagara/Source/Niagara/Private/NiagaraComponentRendererProperties.cpp
@@ -350,6 +350,7 @@ bool FindFunctionParameterDefaultValue(const UFunction* Function, const FPropert
 void UNiagaraComponentRendererProperties::CacheFromCompiledData(const FNiagaraDataSetCompiledData* CompiledData)
 {
 	UpdateSourceModeDerivates(ENiagaraRendererSourceDataMode::Particles);
+	UpdateTemplateCoordinateSpace();
 }
 
 #if WITH_EDITORONLY_DATA
@@ -496,7 +497,6 @@ FNiagaraRenderer* UNiagaraComponentRendererProperties::CreateEmitterRenderer(ERH
 	{
 		UpdateSetterFunctions();
 	}
-	EmitterPtr = Emitter->GetCachedEmitter();
 
 	FNiagaraRenderer* NewRenderer = new FNiagaraRendererComponents(FeatureLevel, this, Emitter);
 	NewRenderer->Initialize(this, Emitter, InController);
@@ -509,11 +509,17 @@ void UNiagaraComponentRendererProperties::CreateTemplateComponent()
 	TemplateComponent->SetVisibility(false);
 	TemplateComponent->SetAutoActivate(false);
 	TemplateComponent->SetComponentTickEnabled(false);
+	UpdateTemplateCoordinateSpace();
+}
 
-	// set some defaults on the component
-	FVersionedNiagaraEmitterData* EmitterData = EmitterPtr.GetEmitterData();
-	bool IsWorldSpace = EmitterData ? !EmitterData->bLocalSpace : true;
-	TemplateComponent->SetAbsolute(IsWorldSpace, IsWorldSpace, IsWorldSpace);
+void UNiagaraComponentRendererProperties::UpdateTemplateCoordinateSpace()
+{
+	if (TemplateComponent)
+	{
+		FVersionedNiagaraEmitterData* EmitterData = GetOuterEmitter().GetEmitterData();
+		const bool IsWorldSpace = EmitterData ? !EmitterData->bLocalSpace : true;
+		TemplateComponent->SetAbsolute(IsWorldSpace, IsWorldSpace, IsWorldSpace);
+	}
 }
 
 #if WITH_EDITOR
@@ -607,7 +613,7 @@ void UNiagaraComponentRendererProperties::GetRendererFeedback(const FVersionedNi
 		OutErrors.Add(FNiagaraRendererFeedback(ErrorDescription, ErrorSummary));
 	}
 
-	FVersionedNiagaraEmitterData* EmitterData = EmitterPtr.GetEmitterData();
+	FVersionedNiagaraEmitterData* EmitterData = InEmitter.GetEmitterData();
 	if (InEmitter.Emitter && TemplateComponent && EmitterData)
 	{
 		if (const UNiagaraSettings* Settings = GetDefault<UNiagaraSettings>())
diff --git a/Engine/Plugins/FX/Niagara/Source/Niagara/Public/NiagaraComponentRendererProperties.h b/Engine/Plugins/FX/Niagara/Source/Niagara/Public/NiagaraComponentRendererProperties.h
--- a/Engine/Plugins/FX/Niagara/Source/Niagara/Public/NiagaraComponentRendererProperties.h
+++ b/Engine/Plugins/FX/Niagara/Source/Niagara/Public/NiagaraComponentRendererProperties.h
@@ -166,10 +166,8 @@ protected:
 private:
 	static TArray<TWeakObjectPtr<UNiagaraComponentRendererProperties>> ComponentRendererPropertiesToDeferredInit;
 
-	// this is just used to check for localspace when creating a new template component
-	FVersionedNiagaraEmitter EmitterPtr;
-
 	void CreateTemplateComponent();
+	void UpdateTemplateCoordinateSpace();
 
 	void UpdateSetterFunctions();
 
