- Static mesh DI add socket sampling to safe CPU access functions

diff --git a/Engine/Plugins/FX/Niagara/Source/Niagara/Private/DataInterface/NiagaraDataInterfaceStaticMesh.cpp b/Engine/Plugins/FX/Niagara/Source/Niagara/Private/DataInterface/NiagaraDataInterfaceStaticMesh.cpp
--- a/Engine/Plugins/FX/Niagara/Source/Niagara/Private/DataInterface/NiagaraDataInterfaceStaticMesh.cpp
+++ b/Engine/Plugins/FX/Niagara/Source/Niagara/Private/DataInterface/NiagaraDataInterfaceStaticMesh.cpp
@@ -2516,6 +2516,25 @@ bool UNiagaraDataInterfaceStaticMesh::FunctionNeedsCpuAccess(FName InName)
 		GetWorldVelocityName,
 		//GetInstanceIndexName,
 		SetInstanceIndexName,
+
+		// Socket functions are safe
+		GetSocketCountName,
+		GetFilteredSocketCountName,
+		GetUnfilteredSocketCountName,
+		RandomSocketName,
+		RandomFilteredSocketName,
+		RandomUnfilteredSocketName,
+		GetSocketTransformName,
+		GetSocketTransformWSName,
+		GetSocketTransformWSInterpolatedName,
+		GetFilteredSocketTransformName,
+		GetFilteredSocketTransformWSName,
+		GetFilteredSocketTransformWSInterpolatedName,
+		GetUnfilteredSocketTransformName,
+		GetUnfilteredSocketTransformWSName,
+		GetUnfilteredSocketTransformWSInterpolatedName,
+		GetFilteredSocketName,
+		GetUnfilteredSocketName,
 	};
 
 	return SafeFunctions.Contains(InName) == false;
