- Add Default Mesh to Skeletal Mesh Data Interface
- Unlike the preview mesh the default mesh is not cooked out and can be used in PIE.

diff --git a/Engine/Plugins/FX/Niagara/Source/Niagara/Classes/NiagaraDataInterfaceSkeletalMesh.h b/Engine/Plugins/FX/Niagara/Source/Niagara/Classes/NiagaraDataInterfaceSkeletalMesh.h
--- a/Engine/Plugins/FX/Niagara/Source/Niagara/Classes/NiagaraDataInterfaceSkeletalMesh.h
+++ b/Engine/Plugins/FX/Niagara/Source/Niagara/Classes/NiagaraDataInterfaceSkeletalMesh.h
@@ -314,18 +314,18 @@ enum class ENDISkeletalMesh_SourceMode : uint8
 	Default behavior.
 	- Use "Source" when specified (either set explicitly or via blueprint with Set Niagara Skeletal Mesh Component).
 	- When no source is specified, fall back on attached actor or component.
+	- If no attach actor look at the default mesh
 	*/
 	Default,
 
-	/**
-	Only use "Source" (either set explicitly or via blueprint with Set Niagara Skeletal Mesh Component).
-	*/
+	/**	Only use "Source" (either set explicitly or via blueprint with Set Niagara Skeletal Mesh Component). */
 	Source,
 
-	/**
-	Only use the parent actor or component the system is attached to.
-	*/
-	AttachParent
+	/**	Only use the parent actor or component the system is attached to. */
+	AttachParent,
+
+	/** Only use the "Default Mesh" specified. */
+	DefaultMeshOnly,
 };
 
 UENUM()
@@ -690,6 +690,10 @@ public:
 	TSoftObjectPtr<USkeletalMesh> PreviewMesh;
 #endif
 
+	/** Mesh used to sample from when not overridden by a source actor from the scene. This mesh is NOT removed from cooked builds. */
+	UPROPERTY(EditAnywhere, Category = "Mesh", meta = (DisallowedClasses = "/Script/ApexDestruction.DestructibleMesh"))
+	TObjectPtr<USkeletalMesh> DefaultMesh;
+
 protected:
 	/** The source actor from which to sample. Takes precedence over the direct mesh. Note that this can only be set when used as a user variable on a component in the world.*/
 	UPROPERTY(EditAnywhere, Category = "Mesh", meta = (DisplayName = "Source Actor"))
diff --git a/Engine/Plugins/FX/Niagara/Source/Niagara/Private/NiagaraDataInterfaceSkeletalMesh.cpp b/Engine/Plugins/FX/Niagara/Source/Niagara/Private/NiagaraDataInterfaceSkeletalMesh.cpp
--- a/Engine/Plugins/FX/Niagara/Source/Niagara/Private/NiagaraDataInterfaceSkeletalMesh.cpp
+++ b/Engine/Plugins/FX/Niagara/Source/Niagara/Private/NiagaraDataInterfaceSkeletalMesh.cpp
@@ -1477,6 +1477,7 @@ USkeletalMesh* UNiagaraDataInterfaceSkeletalMesh::GetSkeletalMesh(FNiagaraSystem
 
 	const bool bTrySource = SourceMode == ENDISkeletalMesh_SourceMode::Default || SourceMode == ENDISkeletalMesh_SourceMode::Source;
 	const bool bTryAttachParent = SourceMode == ENDISkeletalMesh_SourceMode::Default || SourceMode == ENDISkeletalMesh_SourceMode::AttachParent;
+	const bool bTryDefaultMesh = SourceMode == ENDISkeletalMesh_SourceMode::Default || SourceMode == ENDISkeletalMesh_SourceMode::DefaultMeshOnly;
 
 	if (MeshUserParameter.Parameter.IsValid() && InstData && SystemInstance != nullptr)
 	{
@@ -1555,6 +1556,10 @@ USkeletalMesh* UNiagaraDataInterfaceSkeletalMesh::GetSkeletalMesh(FNiagaraSystem
 		Mesh = FoundSkelComp->GetSkeletalMeshAsset();
 		SceneComponent = FoundSkelComp;
 	}
+	else if (bTryDefaultMesh && DefaultMesh)
+	{
+		Mesh = DefaultMesh;
+	}
 #if WITH_EDITORONLY_DATA
 	else if (!SystemInstance || !SystemInstance->GetWorld()->IsGameWorld())
 	{
@@ -2387,6 +2392,7 @@ bool UNiagaraDataInterfaceSkeletalMesh::CopyToInternal(UNiagaraDataInterface* De
 #if WITH_EDITORONLY_DATA
 	OtherTyped->PreviewMesh = PreviewMesh;
 #endif
+	OtherTyped->DefaultMesh = DefaultMesh;
 	OtherTyped->BindSourceDelegates();
 
 	return true;
@@ -2403,6 +2409,7 @@ bool UNiagaraDataInterfaceSkeletalMesh::Equals(const UNiagaraDataInterface* Othe
 #if WITH_EDITORONLY_DATA
 		OtherTyped->PreviewMesh == PreviewMesh &&
 #endif
+		OtherTyped->DefaultMesh == DefaultMesh &&
 		OtherTyped->SoftSourceActor == SoftSourceActor &&
 		OtherTyped->MeshUserParameter == MeshUserParameter &&
 		OtherTyped->ComponentTags == ComponentTags &&
