Cleans up RI generation with the AsyncTask compilation mode.  When using rapid iteration parameters (edit mode) we prepare and apply the parameters synchronously (to match existing behavior).

diff --git a/Engine/Plugins/FX/Niagara/Source/Niagara/Classes/NiagaraScript.h b/Engine/Plugins/FX/Niagara/Source/Niagara/Classes/NiagaraScript.h
--- a/Engine/Plugins/FX/Niagara/Source/Niagara/Classes/NiagaraScript.h
+++ b/Engine/Plugins/FX/Niagara/Source/Niagara/Classes/NiagaraScript.h
@@ -1251,7 +1251,7 @@ public:
 
 	/** Updates the RI parameter store based on the provided variables (missing entries will be
 	    added and stale entries will be removed.  Returns true if the parameter store was modified. */
-	bool ApplyRapidIterationParameters(TConstArrayView<FNiagaraVariable> InParameters, bool bAllowRemoval);
+	NIAGARA_API bool ApplyRapidIterationParameters(TConstArrayView<FNiagaraVariable> InParameters, bool bAllowRemoval);
 
 	/** In the event where we "merge" we duplicate the changes of the source script onto the newly cloned copy. This function will synchronize the compiled script 
 		results assuming that the scripts themselves are bound to the same key. This saves looking things up in the DDC. It returns true if successfully synchronized and 
diff --git a/Engine/Plugins/FX/Niagara/Source/NiagaraEditor/Private/NiagaraCompilationTasks.cpp b/Engine/Plugins/FX/Niagara/Source/NiagaraEditor/Private/NiagaraCompilationTasks.cpp
--- a/Engine/Plugins/FX/Niagara/Source/NiagaraEditor/Private/NiagaraCompilationTasks.cpp
+++ b/Engine/Plugins/FX/Niagara/Source/NiagaraEditor/Private/NiagaraCompilationTasks.cpp
@@ -1794,6 +1794,34 @@ UE::Tasks::FTask FNiagaraSystemCompilationTask::BuildRapidIterationParametersAsy
 	return FTask();
 }
 
+void FNiagaraSystemCompilationTask::BuildAndApplyRapidIterationParameters()
+{
+	check(IsInGameThread());
+	BuildRapidIterationParametersAsync().Wait();
+
+	// see FNiagaraActiveCompilationAsyncTask::Apply() for why we don't allow removal
+	constexpr bool bAllowParameterRemoval = false;
+
+	for (const TMap<TObjectKey<UNiagaraScript>, FScriptInfo>::ElementType& CurrentIt : DigestedScriptInfo)
+	{
+		if (UNiagaraScript* Script = CurrentIt.Key.ResolveObjectPtr())
+		{
+			const FScriptInfo& ScriptInfo = CurrentIt.Value;
+
+			TConstArrayView<FNiagaraVariableWithOffset> SrcRapidIterationParameters = ScriptInfo.RapidIterationParameters.ReadParameterVariables();
+			TArray<FNiagaraVariable> DstRapidIterationParameters;
+			DstRapidIterationParameters.Reserve(SrcRapidIterationParameters.Num());
+			for (const FNiagaraVariableWithOffset& ParamWithOffset : SrcRapidIterationParameters)
+			{
+				FNiagaraVariable& Parameter = DstRapidIterationParameters.Add_GetRef(ParamWithOffset);
+				Parameter.SetData(ScriptInfo.RapidIterationParameters.GetParameterData(ParamWithOffset.Offset));
+			}
+
+			Script->ApplyRapidIterationParameters(DstRapidIterationParameters, bAllowParameterRemoval);
+		}
+	}
+}
+
 void FNiagaraSystemCompilationTask::IssueCompilationTasks()
 {
 	using namespace UE::Tasks;
@@ -1882,30 +1910,61 @@ UE::Tasks::FTask FNiagaraSystemCompilationTask::BeginTasks()
 
 	InitialGetRequestHelper.Launch(this);
 
+	TArray<FTask> SystemTaskPrerequisites;
+	SystemTaskPrerequisites.Add(InitialGetRequestHelper.CompletionEvent);
+
+	// if we are baking rapid iteration parameters then we can rely on the results from the DDC to 
+	// supply everything that we need.  If however we're not baking the rapid iteration parameters
+	// then we need to supplement the results from the DDC with our own collection of the RI so
+	// that we can apply them to the script if something has changed.
+	if (SystemInfo.bUseRapidIterationParams)
+	{
+		// for now we need RI generation to be synchronous to match the behavior of the default
+		// compilation mode.  In the future we need to get RI preparation out of the compilation
+		// and instead we are provided with overridden RI values and don't need to feed back the
+		// entire list of RI parameters.
+		constexpr bool bSyncPrepareRapidIterationParameter = true;
+
+		if (bSyncPrepareRapidIterationParameter)
+		{
+			BuildAndApplyRapidIterationParameters();
+		}
+		else
+		{
+			SystemTaskPrerequisites.Add(BuildRapidIterationParametersAsync());
+		}
+	}
+
 	FTask SystemTask = Launch(UE_SOURCE_LOCATION, [this]
 	{
 		if (HasOutstandingCompileTasks())
 		{
-			FTask BuildRIParamTask = BuildRapidIterationParametersAsync();
+			TArray<FTask> CompilationTaskPrerequisites;
 
-			PostRIParameterGetRequestHelper = MakeUnique<FDispatchAndProcessDataCacheGetRequests>();
-
-			Launch(UE_SOURCE_LOCATION, [this]
+			if (!SystemInfo.bUseRapidIterationParams)
 			{
-				PostRIParameterGetRequestHelper->Launch(this);
-			}, BuildRIParamTask);
+				FTask BuildRIParamTask = BuildRapidIterationParametersAsync();
+				PostRIParameterGetRequestHelper = MakeUnique<FDispatchAndProcessDataCacheGetRequests>();
+
+				Launch(UE_SOURCE_LOCATION, [this]
+				{
+					PostRIParameterGetRequestHelper->Launch(this);
+				}, BuildRIParamTask);
+
+				CompilationTaskPrerequisites.Add(PostRIParameterGetRequestHelper->CompletionEvent);
+			}
 
 			FTask InnerTask = Launch(UE_SOURCE_LOCATION, [this]
 			{
 				IssueCompilationTasks();
-			}, PostRIParameterGetRequestHelper->CompletionEvent);
+			}, CompilationTaskPrerequisites);
 		}
 		else
 		{
 			CompilationState = EState::Completed;
 			CompileCompletionEvent.Trigger();
 		}
-	}, InitialGetRequestHelper.CompletionEvent);
+	}, SystemTaskPrerequisites);
 
 	return SystemTask;
 }
diff --git a/Engine/Plugins/FX/Niagara/Source/NiagaraEditor/Private/NiagaraCompilationTasks.h b/Engine/Plugins/FX/Niagara/Source/NiagaraEditor/Private/NiagaraCompilationTasks.h
--- a/Engine/Plugins/FX/Niagara/Source/NiagaraEditor/Private/NiagaraCompilationTasks.h
+++ b/Engine/Plugins/FX/Niagara/Source/NiagaraEditor/Private/NiagaraCompilationTasks.h
@@ -36,6 +36,7 @@ struct FNiagaraSystemCompilationTask
 	UE::Tasks::FTask BeginTasks();
 
 	UE::Tasks::FTask BuildRapidIterationParametersAsync();
+	void BuildAndApplyRapidIterationParameters();
 	void IssueCompilationTasks();
 	bool HasOutstandingCompileTasks() const;
 
