Fix for Niagara landscape sampling getting the sufficient components when overlapping the system bounds on the landscape

diff --git a/Engine/Plugins/FX/Niagara/Source/Niagara/Private/NiagaraDataInterfaceLandscape.cpp b/Engine/Plugins/FX/Niagara/Source/Niagara/Private/NiagaraDataInterfaceLandscape.cpp
--- a/Engine/Plugins/FX/Niagara/Source/Niagara/Private/NiagaraDataInterfaceLandscape.cpp
+++ b/Engine/Plugins/FX/Niagara/Source/Niagara/Private/NiagaraDataInterfaceLandscape.cpp
@@ -1046,8 +1046,8 @@ FNDI_Landscape_SharedResourceHandle FNDI_Landscape_GeneratedData::GetLandscapeDa
 	SystemBoundsInLandscape.Max = SystemBoundsInLandscape.Max.ComponentMin(MaxIntValue);
 
 	const FIntRect SystemRect = FIntRect(
-		FIntPoint(int32(SystemBoundsInLandscape.Min.X), int32(SystemBoundsInLandscape.Min.Y)),
-		FIntPoint(int32(SystemBoundsInLandscape.Max.X), int32(SystemBoundsInLandscape.Max.Y)));
+		FIntPoint(FMath::FloorToInt(SystemBoundsInLandscape.Min.X), FMath::FloorToInt(SystemBoundsInLandscape.Min.Y)),
+		FIntPoint(FMath::CeilToInt(SystemBoundsInLandscape.Max.X), FMath::CeilToInt(SystemBoundsInLandscape.Max.Y)));
 
 	// for obnoxiously large system bounds we need to guard against potential overflow on the number of cells
 	const int32 MaxSystemWidth = FMath::Clamp<int32>(SystemRect.Max.X - SystemRect.Min.X, 0, MaxLandscapeRegionCount);
