- Speculative fix for crash with nullptr data in async task
- Adjust wait code to check for batch task completion at the end of the function to avoid pointer caching
- Ensure change to tasks are visible to the game thread

diff --git a/Engine/Plugins/FX/Niagara/Source/Niagara/Private/NiagaraSystemInstance.cpp b/Engine/Plugins/FX/Niagara/Source/Niagara/Private/NiagaraSystemInstance.cpp
--- a/Engine/Plugins/FX/Niagara/Source/Niagara/Private/NiagaraSystemInstance.cpp
+++ b/Engine/Plugins/FX/Niagara/Source/Niagara/Private/NiagaraSystemInstance.cpp
@@ -2280,51 +2280,55 @@ void FNiagaraSystemInstance::DumpStalledInfo()
 
 void FNiagaraSystemInstance::WaitForConcurrentTickDoNotFinalize(bool bEnsureComplete)
 {
-	CSV_SCOPED_SET_WAIT_STAT(Effects);
 	check(IsInGameThread());
 
 	// Wait for any concurrent ticking for our task
-	const bool bConcurrentTickPending = ConcurrentTickGraphEvent && !ConcurrentTickGraphEvent->IsComplete();
-	const bool bConcurrentBatchTickPending = ConcurrentTickBatchGraphEvent && !ConcurrentTickBatchGraphEvent->IsComplete();
-	if (bConcurrentTickPending || bConcurrentBatchTickPending)
+	const uint64 StartCycles = FPlatformTime::Cycles64();
+	bool bDidWait = false;
+
+	// Wait for system concurrent tick
+	if (ConcurrentTickGraphEvent && !ConcurrentTickGraphEvent->IsComplete())
 	{
-		ensureAlwaysMsgf(!bEnsureComplete, TEXT("FNiagaraSystemInstance::WaitForConcurrentTickDoNotFinalize - Async Work not complete and is expected to be. %s"), *GetSystem()->GetPathName());
+		CSV_SCOPED_SET_WAIT_STAT(Effects);
 		SCOPE_CYCLE_COUNTER(STAT_NiagaraSystemWaitForAsyncTick);
 		PARTICLE_PERF_STAT_CYCLES_GT(FParticlePerfStatsContext(GetWorld(), GetSystem(), Cast<UFXSystemComponent>(GetAttachComponent())), Wait);
+		bDidWait = true;
 
-		const uint64 StartCycles = FPlatformTime::Cycles64();
-		const double WarnSeconds = 5.0;
-		const uint64 WarnCycles = StartCycles + uint64(WarnSeconds / FPlatformTime::GetSecondsPerCycle64());
-
-		// Note must be done in this order as ConcurrentTickBatchGraphEvent is assigned in the async task
-		if (bConcurrentTickPending)
+		extern int32 GNiagaraSystemSimulationTaskStallTimeout;
+		if (GNiagaraSystemSimulationTaskStallTimeout > 0)
 		{
-			extern int32 GNiagaraSystemSimulationTaskStallTimeout;
-			if (GNiagaraSystemSimulationTaskStallTimeout > 0)
-			{
-				const double EndTimeoutSeconds = FPlatformTime::Seconds() + (double(GNiagaraSystemSimulationTaskStallTimeout) / 1000.0);
-				LowLevelTasks::BusyWaitUntil(
-					[this, EndTimeoutSeconds]()
+			const double EndTimeoutSeconds = FPlatformTime::Seconds() + (double(GNiagaraSystemSimulationTaskStallTimeout) / 1000.0);
+			LowLevelTasks::BusyWaitUntil(
+				[this, EndTimeoutSeconds]()
+				{
+					if (FPlatformTime::Seconds() > EndTimeoutSeconds)
 					{
-						if (FPlatformTime::Seconds() > EndTimeoutSeconds)
-						{
-							DumpStalledInfo();
-							return true;
-						}
-						return ConcurrentTickGraphEvent->IsComplete();
+						DumpStalledInfo();
+						return true;
 					}
-				);
-			}
-			else
-			{
-				FTaskGraphInterface::Get().WaitUntilTaskCompletes(ConcurrentTickGraphEvent, ENamedThreads::GameThread_Local);
-			}
+					return ConcurrentTickGraphEvent->IsComplete();
+				}
+			);
 		}
-		FPlatformMisc::MemoryBarrier();
-		if (ConcurrentTickBatchGraphEvent && !ConcurrentTickBatchGraphEvent->IsComplete())
+		else
 		{
-			FTaskGraphInterface::Get().WaitUntilTaskCompletes(ConcurrentTickBatchGraphEvent, ENamedThreads::GameThread_Local);
+			FTaskGraphInterface::Get().WaitUntilTaskCompletes(ConcurrentTickGraphEvent, ENamedThreads::GameThread_Local);
 		}
+	}
+
+	// Wait for instance concurrent tick
+	if (ConcurrentTickBatchGraphEvent && !ConcurrentTickBatchGraphEvent->IsComplete())
+	{
+		CSV_SCOPED_SET_WAIT_STAT(Effects);
+		SCOPE_CYCLE_COUNTER(STAT_NiagaraSystemWaitForAsyncTick);
+		PARTICLE_PERF_STAT_CYCLES_GT(FParticlePerfStatsContext(GetWorld(), GetSystem(), Cast<UFXSystemComponent>(GetAttachComponent())), Wait);
+		bDidWait = true;
+		FTaskGraphInterface::Get().WaitUntilTaskCompletes(ConcurrentTickBatchGraphEvent, ENamedThreads::GameThread_Local);
+	}
+
+	if (bDidWait)
+	{
+		ensureAlwaysMsgf(!bEnsureComplete, TEXT("FNiagaraSystemInstance::WaitForConcurrentTickDoNotFinalize - Async Work not complete and is expected to be. %s"), *GetSystem()->GetPathName());
 
 		const double StallTimeMS = FPlatformTime::ToMilliseconds64(FPlatformTime::Cycles64() - StartCycles);
 		if ((GWaitForAsyncStallWarnThresholdMS > 0.0f) && (StallTimeMS > GWaitForAsyncStallWarnThresholdMS))
@@ -2333,6 +2337,7 @@ void FNiagaraSystemInstance::WaitForConcurrentTickDoNotFinalize(bool bEnsureComp
 			UE_LOG(LogNiagara, Log, TEXT("Niagara Effect stalled GT for %g ms. Component(%s) System(%s)"), StallTimeMS, *GetFullNameSafe(AttachComponent.Get()), *GetFullNameSafe(GetSystem()));
 		}
 	}
+
 	ConcurrentTickGraphEvent = nullptr;
 	ConcurrentTickBatchGraphEvent = nullptr;
 }
diff --git a/Engine/Plugins/FX/Niagara/Source/Niagara/Private/NiagaraSystemSimulation.cpp b/Engine/Plugins/FX/Niagara/Source/Niagara/Private/NiagaraSystemSimulation.cpp
--- a/Engine/Plugins/FX/Niagara/Source/Niagara/Private/NiagaraSystemSimulation.cpp
+++ b/Engine/Plugins/FX/Niagara/Source/Niagara/Private/NiagaraSystemSimulation.cpp
@@ -995,10 +995,17 @@ void FNiagaraSystemSimulation::FlushTickBatch(FNiagaraSystemSimulationTickContex
 
 			for (FNiagaraSystemInstance* Inst : TickBatch)
 			{
-				Inst->ConcurrentTickGraphEvent = nullptr;
 				Inst->ConcurrentTickBatchGraphEvent = InstanceAsyncGraphEvent;
 			}
 
+			// Ensure ConcurrentTickBatchGraphEvent is visible before we clear ConcurrentTickGraphEvent
+			FPlatformMisc::MemoryBarrier();
+
+			for (FNiagaraSystemInstance* Inst : TickBatch)
+			{
+				Inst->ConcurrentTickGraphEvent = nullptr;
+			}
+
 			// Queue finalize task which will run after the instances are complete, track with our all completion event
 			FGraphEventArray FinalizePrereqArray;
 			FinalizePrereqArray.Add(InstanceAsyncGraphEvent);
