- Fix for incorrect texture being used with landscape

diff --git a/Engine/Plugins/FX/Niagara/Source/Niagara/Private/NiagaraDataInterfaceLandscape.cpp b/Engine/Plugins/FX/Niagara/Source/Niagara/Private/NiagaraDataInterfaceLandscape.cpp
--- a/Engine/Plugins/FX/Niagara/Source/Niagara/Private/NiagaraDataInterfaceLandscape.cpp
+++ b/Engine/Plugins/FX/Niagara/Source/Niagara/Private/NiagaraDataInterfaceLandscape.cpp
@@ -331,24 +331,21 @@ struct FNDILandscapeData_RenderThread
 			{
 				if (IAllocatedVirtualTexture* AllocatedTexture = VirtualTexture->GetAllocatedVirtualTexture())
 				{
-					if (FRHITexture* PageTable = AllocatedTexture->GetPageTableTexture(PageTableIndex))
+					PageTableRef = AllocatedTexture->GetPageTableTexture(PageTableIndex);
+					if (PageTableRef.IsValid())
 					{
-						if (FRHITextureReference* TextureReference = PageTable->GetTextureReference())
-						{
-							PageTableRef = TextureReference->GetReferencedTexture();
-							AllocatedTexture->GetPackedPageTableUniform(PageTableUniforms);
+						AllocatedTexture->GetPackedPageTableUniform(PageTableUniforms);
 
-							if (bIncludeWorldToUv)
-							{
-								WorldToUvParameters[0] = VirtualTexture->GetUniformParameter(ERuntimeVirtualTextureShaderUniform_WorldToUVTransform0);
-								WorldToUvParameters[1] = VirtualTexture->GetUniformParameter(ERuntimeVirtualTextureShaderUniform_WorldToUVTransform1);
-								WorldToUvParameters[2] = VirtualTexture->GetUniformParameter(ERuntimeVirtualTextureShaderUniform_WorldToUVTransform2);
-							}
+						if (bIncludeWorldToUv)
+						{
+							WorldToUvParameters[0] = VirtualTexture->GetUniformParameter(ERuntimeVirtualTextureShaderUniform_WorldToUVTransform0);
+							WorldToUvParameters[1] = VirtualTexture->GetUniformParameter(ERuntimeVirtualTextureShaderUniform_WorldToUVTransform1);
+							WorldToUvParameters[2] = VirtualTexture->GetUniformParameter(ERuntimeVirtualTextureShaderUniform_WorldToUVTransform2);
+						}
 
-							if (bIncludeHeightUnpack)
-							{
-								WorldToUvParameters[3] = VirtualTexture->GetUniformParameter(ERuntimeVirtualTextureShaderUniform_WorldHeightUnpack);
-							}
+						if (bIncludeHeightUnpack)
+						{
+							WorldToUvParameters[3] = VirtualTexture->GetUniformParameter(ERuntimeVirtualTextureShaderUniform_WorldHeightUnpack);
 						}
 					}
 				}
