Bake SVTs directly from the Niagara Baker
Render preview with the HV component

diff --git a/Engine/Plugins/FX/Niagara/Source/Niagara/Classes/NiagaraBakerOutputSparseVolumeTexture.h b/Engine/Plugins/FX/Niagara/Source/Niagara/Classes/NiagaraBakerOutputSparseVolumeTexture.h
new file mode 100644
--- /dev/null
+++ b/Engine/Plugins/FX/Niagara/Source/Niagara/Classes/NiagaraBakerOutputSparseVolumeTexture.h
@@ -0,0 +1,36 @@
+// Copyright Epic Games, Inc. All Rights Reserved.
+
+#pragma once
+
+#include "NiagaraBakerOutput.h"
+#include "NiagaraBakerOutputSparseVolumeTexture.generated.h"
+
+UCLASS(meta = (DisplayName = "SparseVolume Texture Output"), MinimalAPI)
+class UNiagaraBakerOutputSparseVolumeTexture : public UNiagaraBakerOutput
+{
+	GENERATED_BODY()
+
+public:
+	UNiagaraBakerOutputSparseVolumeTexture(const FObjectInitializer& Init)	
+	{
+	}
+
+	UPROPERTY(EditAnywhere, Category = "Settings")
+	FNiagaraBakerTextureSource SourceBinding;
+
+	/**
+	When enabled a volume atlas is created, the atlas is along X & Y not Z based on baker settings.
+	*/
+	UPROPERTY(EditAnywhere, Category = "Settings")
+	FString SparseVolumeTextureAssetPathFormat = TEXT("{AssetFolder}/{AssetName}_BakedSVT_{OutputName}");
+
+	NIAGARA_API virtual bool Equals(const UNiagaraBakerOutput& Other) const override;
+
+#if WITH_EDITOR
+	NIAGARA_API FString MakeOutputName() const override;
+#endif
+
+#if WITH_EDITORONLY_DATA
+	NIAGARA_API virtual void PostEditChangeProperty(struct FPropertyChangedEvent& PropertyChangedEvent) override;
+#endif
+};
diff --git a/Engine/Plugins/FX/Niagara/Source/Niagara/Private/NiagaraBakerOutputSparseVolumeTexture.cpp b/Engine/Plugins/FX/Niagara/Source/Niagara/Private/NiagaraBakerOutputSparseVolumeTexture.cpp
new file mode 100644
--- /dev/null
+++ b/Engine/Plugins/FX/Niagara/Source/Niagara/Private/NiagaraBakerOutputSparseVolumeTexture.cpp
@@ -0,0 +1,29 @@
+// Copyright Epic Games, Inc. All Rights Reserved.
+
+#include "NiagaraBakerOutputSparseVolumeTexture.h"
+
+#include UE_INLINE_GENERATED_CPP_BY_NAME(NiagaraBakerOutputSparseVolumeTexture)
+
+bool UNiagaraBakerOutputSparseVolumeTexture::Equals(const UNiagaraBakerOutput& OtherBase) const
+{
+	const UNiagaraBakerOutputSparseVolumeTexture& Other = *CastChecked<UNiagaraBakerOutputSparseVolumeTexture>(&OtherBase);
+	return
+		Super::Equals(Other) &&
+		SourceBinding.SourceName == Other.SourceBinding.SourceName;
+}
+
+#if WITH_EDITOR
+FString UNiagaraBakerOutputSparseVolumeTexture::MakeOutputName() const
+{
+	return FString::Printf(TEXT("SparseVolume_%d"), GetFName().GetNumber());
+}
+#endif
+
+#if WITH_EDITORONLY_DATA
+void UNiagaraBakerOutputSparseVolumeTexture::PostEditChangeProperty(struct FPropertyChangedEvent& PropertyChangedEvent)
+{
+	Super::PostEditChangeProperty(PropertyChangedEvent);
+
+}
+#endif
+
diff --git a/Engine/Plugins/FX/Niagara/Source/NiagaraEditor/NiagaraEditor.Build.cs b/Engine/Plugins/FX/Niagara/Source/NiagaraEditor/NiagaraEditor.Build.cs
--- a/Engine/Plugins/FX/Niagara/Source/NiagaraEditor/NiagaraEditor.Build.cs
+++ b/Engine/Plugins/FX/Niagara/Source/NiagaraEditor/NiagaraEditor.Build.cs
@@ -49,6 +49,7 @@ public class NiagaraEditor : ModuleRules
 				"ToolWidgets",
 				"AssetTools",
 				"LevelSequence",
+				"SparseVolumeTexture",
 			}
 		);
 
diff --git a/Engine/Plugins/FX/Niagara/Source/NiagaraEditor/Private/Customizations/NiagaraBakerOutputCustomization.cpp b/Engine/Plugins/FX/Niagara/Source/NiagaraEditor/Private/Customizations/NiagaraBakerOutputCustomization.cpp
--- a/Engine/Plugins/FX/Niagara/Source/NiagaraEditor/Private/Customizations/NiagaraBakerOutputCustomization.cpp
+++ b/Engine/Plugins/FX/Niagara/Source/NiagaraEditor/Private/Customizations/NiagaraBakerOutputCustomization.cpp
@@ -4,6 +4,7 @@
 #include "NiagaraBakerOutputSimCache.h"
 #include "NiagaraBakerOutputTexture2D.h"
 #include "NiagaraBakerOutputVolumeTexture.h"
+#include "NiagaraBakerOutputSparseVolumeTexture.h"
 
 #include "Editor/EditorEngine.h"
 #include "Widgets/Images/SImage.h"
@@ -353,4 +354,62 @@ void FNiagaraBakerOutputVolumeTextureDetails::CustomizeDetails(IDetailLayoutBuil
 	}
 }
 
+//////////////////////////////////////////////////////////////////////////
+
+TSharedRef<IDetailCustomization> FNiagaraBakerOutputSparseVolumeTextureDetails::MakeInstance()
+{
+	return MakeShared<FNiagaraBakerOutputSparseVolumeTextureDetails>();
+}
+
+FText FNiagaraBakerOutputSparseVolumeTextureDetails::BrowseAtlasToolTipText(TWeakObjectPtr<UNiagaraBakerOutputSparseVolumeTexture> WeakOutput)
+{
+	UNiagaraBakerOutputSparseVolumeTexture* Output = WeakOutput.Get();
+	return FText::Format(LOCTEXT("BrowseAtlasToolTipFormat", "Browse to atlas asset '{0}'"), FText::FromString(Output ? Output->GetAssetPath(Output->SparseVolumeTextureAssetPathFormat, 0) : FString()));
+}
+
+FReply FNiagaraBakerOutputSparseVolumeTextureDetails::BrowseToAtlas(TWeakObjectPtr<UNiagaraBakerOutputSparseVolumeTexture> WeakOutput)
+{
+	if (UNiagaraBakerOutputSparseVolumeTexture* Output = WeakOutput.Get())
+	{
+		const FString AssetPath = Output->GetAssetPath(Output->SparseVolumeTextureAssetPathFormat, 0);
+		FocusContentBrowserToAsset(AssetPath);
+	}
+	return FReply::Handled();
+}
+
+
+void FNiagaraBakerOutputSparseVolumeTextureDetails::CustomizeDetails(IDetailLayoutBuilder& DetailBuilder)
+{
+	// We only support customization on 1 object
+	TArray<TWeakObjectPtr<UObject>> ObjectsCustomized;
+	DetailBuilder.GetObjectsBeingCustomized(ObjectsCustomized);
+	if (ObjectsCustomized.Num() != 1 || !ObjectsCustomized[0]->IsA<UNiagaraBakerOutputSparseVolumeTexture>())
+	{
+		return;
+	}
+	TWeakObjectPtr<UNiagaraBakerOutputSparseVolumeTexture> WeakOutput = CastChecked<UNiagaraBakerOutputSparseVolumeTexture>(ObjectsCustomized[0]);
+
+	static FName NAME_Settings("Settings");
+	IDetailCategoryBuilder& DetailCategory = DetailBuilder.EditCategory(NAME_Settings);
+
+	TArray<TSharedRef<IPropertyHandle>> CategoryProperties;
+	DetailCategory.GetDefaultProperties(CategoryProperties);
+
+	for (TSharedRef<IPropertyHandle>& PropertyHandle : CategoryProperties)
+	{
+		const FName PropertyName = PropertyHandle->GetProperty() ? PropertyHandle->GetProperty()->GetFName() : NAME_None;
+		if (PropertyName == GET_MEMBER_NAME_CHECKED(UNiagaraBakerOutputSparseVolumeTexture, SparseVolumeTextureAssetPathFormat))
+		{
+			BuildAssetPathWidget(
+				DetailCategory, PropertyHandle,
+				TAttribute<FText>::FGetter::CreateStatic(&FNiagaraBakerOutputSparseVolumeTextureDetails::BrowseAtlasToolTipText, WeakOutput),
+				FOnClicked::CreateStatic(&FNiagaraBakerOutputSparseVolumeTextureDetails::BrowseToAtlas, WeakOutput)
+			);
+		}
+		else
+		{
+			DetailCategory.AddProperty(PropertyHandle);
+		}
+	}
+}
 #undef LOCTEXT_NAMESPACE
diff --git a/Engine/Plugins/FX/Niagara/Source/NiagaraEditor/Private/Customizations/NiagaraBakerOutputCustomization.h b/Engine/Plugins/FX/Niagara/Source/NiagaraEditor/Private/Customizations/NiagaraBakerOutputCustomization.h
--- a/Engine/Plugins/FX/Niagara/Source/NiagaraEditor/Private/Customizations/NiagaraBakerOutputCustomization.h
+++ b/Engine/Plugins/FX/Niagara/Source/NiagaraEditor/Private/Customizations/NiagaraBakerOutputCustomization.h
@@ -12,6 +12,7 @@ class IPropertyHandle;
 class UNiagaraBakerOutputSimCache;
 class UNiagaraBakerOutputTexture2D;
 class UNiagaraBakerOutputVolumeTexture;
+class UNiagaraBakerOutputSparseVolumeTexture;
 
 struct FNiagaraBakerOutputDetails : public IDetailCustomization
 {
@@ -61,3 +62,13 @@ struct FNiagaraBakerOutputVolumeTextureDetails : public FNiagaraBakerOutputDetai
 
 	virtual void CustomizeDetails(IDetailLayoutBuilder& DetailBuilder) override;
 };
+
+struct FNiagaraBakerOutputSparseVolumeTextureDetails : public FNiagaraBakerOutputDetails
+{
+	static TSharedRef<IDetailCustomization> MakeInstance();
+
+	static FText BrowseAtlasToolTipText(TWeakObjectPtr<UNiagaraBakerOutputSparseVolumeTexture> WeakOutput);
+	static FReply BrowseToAtlas(TWeakObjectPtr<UNiagaraBakerOutputSparseVolumeTexture> WeakOutput);
+
+	virtual void CustomizeDetails(IDetailLayoutBuilder& DetailBuilder) override;
+};
\ No newline at end of file
diff --git a/Engine/Plugins/FX/Niagara/Source/NiagaraEditor/Private/NiagaraBakerOutputRegistry.cpp b/Engine/Plugins/FX/Niagara/Source/NiagaraEditor/Private/NiagaraBakerOutputRegistry.cpp
--- a/Engine/Plugins/FX/Niagara/Source/NiagaraEditor/Private/NiagaraBakerOutputRegistry.cpp
+++ b/Engine/Plugins/FX/Niagara/Source/NiagaraEditor/Private/NiagaraBakerOutputRegistry.cpp
@@ -5,9 +5,11 @@
 #include "NiagaraBakerOutputSimCache.h"
 #include "NiagaraBakerOutputTexture2D.h"
 #include "NiagaraBakerOutputVolumeTexture.h"
+#include "NiagaraBakerOutputSparseVolumeTexture.h"
 #include "NiagaraBakerRendererOutputSimCache.h"
 #include "NiagaraBakerRendererOutputTexture2D.h"
 #include "NiagaraBakerRendererOutputVolumeTexture.h"
+#include "NiagaraBakerRendererOutputSparseVolumeTexture.h"
 #include "Customizations/NiagaraBakerOutputCustomization.h"
 
 FNiagaraBakerOutputRegistry& FNiagaraBakerOutputRegistry::Get()
@@ -38,6 +40,13 @@ FNiagaraBakerOutputRegistry& FNiagaraBakerOutputRegistry::Get()
 			Entry.CreateRenderer		= []() -> FNiagaraBakerOutputRenderer* { return new FNiagaraBakerRendererOutputVolumeTexture(); };
 			Entry.CreateCustomization	= FOnGetDetailCustomizationInstance::CreateStatic(&FNiagaraBakerOutputVolumeTextureDetails::MakeInstance);
 		}
+		// Sparse Volume Texture
+		{
+			FRegistryEntry& Entry = Singleton->Registry.AddDefaulted_GetRef();
+			Entry.WeakClass = UNiagaraBakerOutputSparseVolumeTexture::StaticClass();
+			Entry.CreateRenderer = []() -> FNiagaraBakerOutputRenderer* { return new FNiagaraBakerRendererOutputSparseVolumeTexture(); };
+			Entry.CreateCustomization = FOnGetDetailCustomizationInstance::CreateStatic(&FNiagaraBakerOutputSparseVolumeTextureDetails::MakeInstance);
+		}
 	}
 	return *Singleton.Get();
 }
diff --git a/Engine/Plugins/FX/Niagara/Source/NiagaraEditor/Private/NiagaraBakerRenderer.cpp b/Engine/Plugins/FX/Niagara/Source/NiagaraEditor/Private/NiagaraBakerRenderer.cpp
--- a/Engine/Plugins/FX/Niagara/Source/NiagaraEditor/Private/NiagaraBakerRenderer.cpp
+++ b/Engine/Plugins/FX/Niagara/Source/NiagaraEditor/Private/NiagaraBakerRenderer.cpp
@@ -34,8 +34,16 @@
 #include "TextureResource.h"
 #include "SceneInterface.h"
 
+#include "NiagaraDataInterfaceGrid3DCollection.h"
+#include "NiagaraDataInterfaceRenderTargetVolume.h"
+
+#include "SparseVolumeTexture/SparseVolumeTexture.h"
+#include "Components/HeterogeneousVolumeComponent.h"
+
 //////////////////////////////////////////////////////////////////////////
 
+DEFINE_LOG_CATEGORY(LogNiagaraBaker);
+
 const FString FNiagaraBakerOutputBindingHelper::STRING_SceneCaptureSource("SceneCaptureSource");
 const FString FNiagaraBakerOutputBindingHelper::STRING_BufferVisualization("BufferVisualization");
 const FString FNiagaraBakerOutputBindingHelper::STRING_EmitterDI("EmitterDI");
@@ -285,6 +293,18 @@ FNiagaraBakerRenderer::~FNiagaraBakerRenderer()
 {
 	DestroyPreviewScene(PreviewComponent, AdvancedPreviewScene);
 	DestroyPreviewScene(SimCachePreviewComponent, SimCacheAdvancedPreviewScene);
+
+	// Clean up SVT preview scene
+	if (SVTPreviewScene && SVTPreviewComponent)
+	{
+		SVTPreviewScene->RemoveComponent(SVTPreviewComponent);
+		SVTPreviewScene = nullptr;
+	}
+	if (SVTPreviewComponent)
+	{
+		SVTPreviewComponent->DestroyComponent();
+		SVTPreviewComponent = nullptr;
+	}
 }
 
 void FNiagaraBakerRenderer::SetAbsoluteTime(float AbsoluteTime, bool bShouldTickComponent)
@@ -324,16 +344,16 @@ void FNiagaraBakerRenderer::RenderSceneCapture(UTextureRenderTarget2D* RenderTar
 	RenderSceneCapture(RenderTarget, PreviewComponent, CaptureSource);
 }
 
-void FNiagaraBakerRenderer::RenderSceneCapture(UTextureRenderTarget2D* RenderTarget, UNiagaraComponent* NiagaraComponent, ESceneCaptureSource CaptureSource) const
+void FNiagaraBakerRenderer::RenderSceneCapture(UTextureRenderTarget2D* RenderTarget, UPrimitiveComponent* BakedDataComponent, ESceneCaptureSource CaptureSource) const
 {
 	UNiagaraBakerSettings* BakerSettings = GetBakerSettings();
-	if (!NiagaraComponent || !RenderTarget || !BakerSettings)
+	if (!BakedDataComponent || !RenderTarget || !BakerSettings)
 	{
 		return;
 	}
 
 	const float WorldTime = GetWorldTime();
-	UWorld* World = NiagaraComponent->GetWorld();
+	UWorld* World = BakedDataComponent->GetWorld();
 
 	FCanvas Canvas(RenderTarget->GameThread_GetRenderTargetResource(), nullptr, FGameTime::CreateUndilated(WorldTime, FApp::GetDeltaTime()), GetFeatureLevel());
 	Canvas.Clear(FLinearColor::Black);
@@ -366,7 +386,7 @@ void FNiagaraBakerRenderer::RenderSceneCapture(UTextureRenderTarget2D* RenderTar
 	{
 		SceneCaptureComponent->PrimitiveRenderMode = ESceneCapturePrimitiveRenderMode::PRM_UseShowOnlyList;
 		SceneCaptureComponent->ShowOnlyComponents.Empty(1);
-		SceneCaptureComponent->ShowOnlyComponents.Add(NiagaraComponent);
+		SceneCaptureComponent->ShowOnlyComponents.Add(BakedDataComponent);
 	}
 	else
 	{
@@ -623,6 +643,70 @@ void FNiagaraBakerRenderer::RenderSimCache(UTextureRenderTarget2D* RenderTarget,
 	SimCachePreviewComponent->SetSimCache(nullptr);
 }
 
+void FNiagaraBakerRenderer::RenderSparseVolumeTexture(UTextureRenderTarget2D* RenderTarget, const FNiagaraBakerOutputFrameIndices Indices, UAnimatedSparseVolumeTexture* SVT) const
+{
+	UNiagaraBakerSettings* BakerSettings = GetBakerSettings();
+	if (!SVT)
+	{
+		return;
+	}
+
+	bool CreatePreviewComponent = SVTPreviewComponent == nullptr;
+	if (CreatePreviewComponent)
+	{
+		SVTPreviewComponent = NewObject<UHeterogeneousVolumeComponent>(GetTransientPackage(), NAME_None, RF_Transient);
+				
+		// create HV component and wire all the things
+		UMaterialInterface* MaterialInterface = LoadObject<UMaterialInterface>(NULL, TEXT("/Engine/EngineMaterials/SparseVolumeMaterial"), NULL, LOAD_None, NULL);
+		
+		UMaterial *Mat = MaterialInterface->GetMaterial();
+
+		// #todo(dmp): we had to duplicate the material itself because we cannot make a mid and send that to HV
+		// HV internally makes a MID from whatever is bound, and the MID of a MID workflow appears broken	
+		UMaterial *DuplicateMat = DuplicateObject<UMaterial>(Mat, SVTPreviewComponent);
+
+		FGuid ExprGuid;
+		DuplicateMat->SetStaticComponentMaskParameterValueEditorOnly("Temperature Mask", false, true, false, false, ExprGuid);
+
+		FGuid SwitchGuid;
+		DuplicateMat->SetStaticSwitchParameterValueEditorOnly("Temperature (Attributes B)", false, SwitchGuid);
+
+		DuplicateMat->SetSparseVolumeTextureParameterValueEditorOnly("SparseVolumeTexture", SVT);
+
+		SVTPreviewComponent->OverrideMaterials.Add(DuplicateMat);
+
+		SVTPreviewComponent->bIssueBlockingRequests = true;
+		SVTPreviewComponent->PostLoad();
+
+		SVTPreviewScene = MakeShareable(new FAdvancedPreviewScene(FPreviewScene::ConstructionValues()));
+		SVTPreviewScene->SetFloorVisibility(false);
+		SVTPreviewScene->AddComponent(SVTPreviewComponent, SVTPreviewComponent->GetRelativeTransform());
+	}
+	
+	TArray<FMaterialParameterInfo> ParameterInfo;
+	TArray<FGuid> ParameterIds;
+	SVTPreviewComponent->OverrideMaterials[0]->GetAllSparseVolumeTextureParameterInfo(ParameterInfo, ParameterIds);
+	USparseVolumeTexture* OldSVT;
+	SVTPreviewComponent->OverrideMaterials[0]->GetSparseVolumeTextureParameterValue(ParameterInfo[0], OldSVT);
+	
+	if (OldSVT != SVT)
+	{
+		SVTPreviewComponent->OverrideMaterials[0]->GetMaterial()->SetSparseVolumeTextureParameterValueEditorOnly("SparseVolumeTexture", SVT);
+	}
+	SVTPreviewComponent->SetFrame(Indices.FrameIndexA);
+	
+	// #todo(dmp): apply world scale and pivot to HV actor
+
+	const float SeekDelta = BakerSettings->GetSeekDelta();
+	SVTPreviewComponent->TickComponent(SeekDelta, ELevelTick::LEVELTICK_All, nullptr);
+
+	SVTPreviewComponent->MarkRenderDynamicDataDirty();
+	UWorld* World = SVTPreviewComponent->GetWorld();
+	World->SendAllEndOfFrameUpdates();
+
+	RenderSceneCapture(RenderTarget, SVTPreviewComponent, ESceneCaptureSource::SCS_SceneColorHDR);
+}
+
 UWorld* FNiagaraBakerRenderer::GetWorld() const
 {
 	return PreviewComponent->GetWorld();
@@ -649,6 +733,7 @@ void FNiagaraBakerRenderer::AddReferencedObjects(FReferenceCollector& Collector)
 	Collector.AddReferencedObject(PreviewComponent);
 	Collector.AddReferencedObject(SceneCaptureComponent);
 	Collector.AddReferencedObject(SimCachePreviewComponent);
+	Collector.AddReferencedObject(SVTPreviewComponent);
 }
 
 FNiagaraBakerOutputRenderer* FNiagaraBakerRenderer::GetOutputRenderer(UClass* Class)
@@ -752,3 +837,75 @@ bool FNiagaraBakerRenderer::ExportVolume(FStringView FilePath, FIntVector ImageS
 		return ExportImage(FilePath, FIntPoint(ImageSize.X, ImageSize.Y * ImageSize.Z), ImageData);
 	}
 }
+
+
+
+bool FVolumeDataInterfaceHelper::Initialize(const TArray<FString>& InputDataInterfacePath, UNiagaraComponent* InNiagaraComponent)
+{
+	NiagaraComponent = InNiagaraComponent;
+
+	DataInterfacePath = InputDataInterfacePath;
+	
+	if (DataInterfacePath.Num() < 2)
+	{
+		return false;
+	}
+
+	const FName DataInterfaceName(DataInterfacePath[0] + "." + DataInterfacePath[1]);
+	UNiagaraDataInterface* DataInterface = FNiagaraBakerOutputBindingHelper::GetDataInterface(NiagaraComponent, DataInterfaceName);
+	if (DataInterface == nullptr)
+	{
+		return false;
+	}
+
+	// Guaranteed since we got a data interface
+	SystemInstance = NiagaraComponent->GetSystemInstanceController()->GetSoloSystemInstance();
+
+	// Render Target Volume
+	if (DataInterface->IsA<UNiagaraDataInterfaceRenderTargetVolume>())
+	{
+		VolumeRenderTargetDataInterface = CastChecked<UNiagaraDataInterfaceRenderTargetVolume>(DataInterface);
+		VolumeRenderTargetProxy = static_cast<FNiagaraDataInterfaceProxyRenderTargetVolumeProxy*>(VolumeRenderTargetDataInterface->GetProxy());
+		VolumeRenderTargetInstanceData_GameThread = reinterpret_cast<FRenderTargetVolumeRWInstanceData_GameThread*>(SystemInstance->FindDataInterfaceInstanceData(VolumeRenderTargetDataInterface));
+		if (VolumeRenderTargetInstanceData_GameThread == nullptr)
+		{
+			return false;
+		}
+	}
+	// Grid 3D
+	else if (DataInterface->IsA<UNiagaraDataInterfaceGrid3DCollection>())
+	{
+		Grid3DDataInterface = CastChecked<UNiagaraDataInterfaceGrid3DCollection>(DataInterface);
+		Grid3DProxy = static_cast<FNiagaraDataInterfaceProxyGrid3DCollectionProxy*>(Grid3DDataInterface->GetProxy());
+		Grid3DInstanceData_GameThread = reinterpret_cast<FGrid3DCollectionRWInstanceData_GameThread*>(SystemInstance->FindDataInterfaceInstanceData(Grid3DDataInterface));
+		if (Grid3DInstanceData_GameThread == nullptr)
+		{
+			return false;
+		}
+
+		if (DataInterfacePath.Num() != 3)
+		{
+			// Perhaps a path to pull all attributes, i.e. whole texture?
+			return false;
+		}
+
+		Grid3DAttributeName = FName(DataInterfacePath[2]);
+		Grid3DVariableIndex = Grid3DInstanceData_GameThread->Vars.IndexOfByPredicate([&](const FNiagaraVariableBase& VariableBase) { return VariableBase.GetName() == Grid3DAttributeName; });
+		if (Grid3DVariableIndex == INDEX_NONE)
+		{
+			return false;
+		}
+		Grid3DAttributeStart = Grid3DInstanceData_GameThread->Offsets[Grid3DVariableIndex];
+		Grid3DAttributeChannels = Grid3DInstanceData_GameThread->Vars[Grid3DVariableIndex].GetType().GetSize() / sizeof(float);
+		Grid3DTextureSize.X = Grid3DInstanceData_GameThread->NumCells.X * Grid3DInstanceData_GameThread->NumTiles.X;
+		Grid3DTextureSize.Y = Grid3DInstanceData_GameThread->NumCells.Y * Grid3DInstanceData_GameThread->NumTiles.Y;
+		Grid3DTextureSize.Z = Grid3DInstanceData_GameThread->NumCells.Z * Grid3DInstanceData_GameThread->NumTiles.Z;
+	}
+	// Unsupported type
+	else
+	{
+		return false;
+	}
+
+	return true;
+}
\ No newline at end of file
diff --git a/Engine/Plugins/FX/Niagara/Source/NiagaraEditor/Private/NiagaraBakerRenderer.h b/Engine/Plugins/FX/Niagara/Source/NiagaraEditor/Private/NiagaraBakerRenderer.h
--- a/Engine/Plugins/FX/Niagara/Source/NiagaraEditor/Private/NiagaraBakerRenderer.h
+++ b/Engine/Plugins/FX/Niagara/Source/NiagaraEditor/Private/NiagaraBakerRenderer.h
@@ -5,9 +5,12 @@
 #include "CoreMinimal.h"
 #include "Engine/EngineTypes.h"
 #include "NiagaraBakerOutput.h"
+#include "NiagaraBakerSettings.h"
 #include "NiagaraSystem.h"
 #include "UObject/GCObject.h"
 
+DECLARE_LOG_CATEGORY_EXTERN(LogNiagaraBaker, Log, All);
+
 class FNiagaraBakerRenderer;
 class UNiagaraComponent;
 class UNiagaraBakerSettings;
@@ -16,6 +19,8 @@ class UTextureRenderTarget2D;
 class USceneCaptureComponent2D;
 class FAdvancedPreviewScene;
 class FCanvas;
+class UAnimatedSparseVolumeTexture;
+class UHeterogeneousVolumeComponent;
 
 struct FNiagaraBakerFeedbackContext
 {
@@ -101,11 +106,12 @@ public:
 	void SetAbsoluteTime(float AbsoluteTime, bool bShouldTickComponent = true);
 
 	void RenderSceneCapture(UTextureRenderTarget2D* RenderTarget, ESceneCaptureSource CaptureSource) const;
-	void RenderSceneCapture(UTextureRenderTarget2D* RenderTarget, UNiagaraComponent* NiagaraComponent, ESceneCaptureSource CaptureSource) const;
+	void RenderSceneCapture(UTextureRenderTarget2D* RenderTarget, UPrimitiveComponent* BakedDataComponent, ESceneCaptureSource CaptureSource) const;
 	void RenderBufferVisualization(UTextureRenderTarget2D* RenderTarget, FName BufferVisualizationMode = NAME_None) const;
 	void RenderDataInterface(UTextureRenderTarget2D* RenderTarget, FName BindingName) const;
 	void RenderParticleAttribute(UTextureRenderTarget2D* RenderTarget, FName BindingName) const;
 	void RenderSimCache(UTextureRenderTarget2D* RenderTarget, UNiagaraSimCache* SimCache) const;
+	void RenderSparseVolumeTexture(UTextureRenderTarget2D* RenderTarget, const FNiagaraBakerOutputFrameIndices Indices, UAnimatedSparseVolumeTexture *SVT) const;
 
 	UWorld* GetWorld() const;
 	float GetWorldTime() const;
@@ -138,5 +144,44 @@ private:
 	TObjectPtr<USceneCaptureComponent2D> SceneCaptureComponent = nullptr;
 
 	mutable TObjectPtr<UNiagaraComponent> SimCachePreviewComponent = nullptr;
+	mutable TObjectPtr<UHeterogeneousVolumeComponent> SVTPreviewComponent = nullptr;
+
 	mutable TSharedPtr<FAdvancedPreviewScene> SimCacheAdvancedPreviewScene;
+	mutable TSharedPtr<FAdvancedPreviewScene> SVTPreviewScene;
 };
+
+class UNiagaraComponent;
+class FNiagaraSystemInstance;
+class UNiagaraDataInterfaceGrid3DCollection;
+struct FNiagaraDataInterfaceProxyGrid3DCollectionProxy;
+struct FGrid3DCollectionRWInstanceData_GameThread;
+class UNiagaraDataInterfaceRenderTargetVolume;
+struct FNiagaraDataInterfaceProxyRenderTargetVolumeProxy;
+struct FRenderTargetVolumeRWInstanceData_GameThread;
+
+class FVolumeDataInterfaceHelper
+{
+public:
+
+	FVolumeDataInterfaceHelper() {};
+
+	UNiagaraComponent* NiagaraComponent = nullptr;
+	FNiagaraSystemInstance* SystemInstance = nullptr;
+	TArray<FString>											DataInterfacePath;
+
+	UNiagaraDataInterfaceGrid3DCollection* Grid3DDataInterface = nullptr;
+	FNiagaraDataInterfaceProxyGrid3DCollectionProxy* Grid3DProxy = nullptr;
+	FGrid3DCollectionRWInstanceData_GameThread* Grid3DInstanceData_GameThread = nullptr;
+	FName													Grid3DAttributeName;
+	int32													Grid3DVariableIndex = INDEX_NONE;
+	int32													Grid3DAttributeStart = INDEX_NONE;
+	int32													Grid3DAttributeChannels = 0;
+	FIntVector												Grid3DTextureSize = FIntVector::ZeroValue;
+
+	UNiagaraDataInterfaceRenderTargetVolume* VolumeRenderTargetDataInterface = nullptr;
+	FNiagaraDataInterfaceProxyRenderTargetVolumeProxy* VolumeRenderTargetProxy = nullptr;
+	FRenderTargetVolumeRWInstanceData_GameThread* VolumeRenderTargetInstanceData_GameThread = nullptr;
+
+	bool Initialize(const TArray<FString>& InputDataInterfacePath, UNiagaraComponent* InNiagaraComponent);
+};
+
diff --git a/Engine/Plugins/FX/Niagara/Source/NiagaraEditor/Private/NiagaraBakerRendererOutputSparseVolumeTexture.cpp b/Engine/Plugins/FX/Niagara/Source/NiagaraEditor/Private/NiagaraBakerRendererOutputSparseVolumeTexture.cpp
new file mode 100644
--- /dev/null
+++ b/Engine/Plugins/FX/Niagara/Source/NiagaraEditor/Private/NiagaraBakerRendererOutputSparseVolumeTexture.cpp
@@ -0,0 +1,324 @@
+// Copyright Epic Games, Inc. All Rights Reserved.
+
+#include "NiagaraBakerRendererOutputSparseVolumeTexture.h"
+#include "NiagaraBakerOutputSparseVolumeTexture.h"
+
+#include "NiagaraBakerSettings.h"
+#include "NiagaraComponent.h"
+#include "NiagaraSystemInstance.h"
+#include "NiagaraSystemInstanceController.h"
+#include "NiagaraComputeExecutionContext.h"
+#include "NiagaraDataInterfaceGrid3DCollection.h"
+#include "NiagaraDataInterfaceRenderTargetVolume.h"
+#include "NiagaraBatchedElements.h"
+
+#include "Engine/Canvas.h"
+#include "Engine/TextureRenderTarget2D.h"
+#include "Engine/VolumeTexture.h"
+#include "HAL/PlatformFileManager.h"
+#include "Misc/PathViews.h"
+#include "Modules/ModuleManager.h"
+#include "Factories/VolumeTextureFactory.h"
+#include "TextureResource.h"
+#include "UObject/UObjectGlobals.h"
+
+#include "TextureRenderTargetVolumeResource.h"
+#include "Engine/TextureRenderTargetVolume.h"
+
+#include "Editor/SparseVolumeTexture/Public/SparseVolumeTextureFactory.h"
+#include "SparseVolumeTexture/SparseVolumeTexture.h"
+#include "SparseVolumeTexture/SparseVolumeTextureData.h"
+#include "SparseVolumeTexture/ISparseVolumeTextureStreamingManager.h"
+
+#if WITH_EDITOR
+#include "AssetRegistry/AssetData.h"
+#include "AssetRegistry/IAssetRegistry.h"
+#include "Factories/Factory.h"
+#include "IAssetTools.h"
+#include "Misc/MessageDialog.h"
+#include "AssetToolsModule.h"
+#endif //WITH_EDITOR
+
+TArray<FNiagaraBakerOutputBinding> FNiagaraBakerRendererOutputSparseVolumeTexture::GetRendererBindings(UNiagaraBakerOutput* InBakerOutput) const
+{
+	TArray<FNiagaraBakerOutputBinding> OutBindings;
+	if (UNiagaraSystem* NiagaraSystem = InBakerOutput->GetTypedOuter<UNiagaraSystem>())
+	{
+		FNiagaraBakerOutputBindingHelper::ForEachEmitterDataInterface(
+			NiagaraSystem,
+			[&](const FString& EmitterName, const FString& VariableName, UNiagaraDataInterface* DataInterface)
+			{
+				if (UNiagaraDataInterfaceGrid3DCollection* Grid3D = Cast<UNiagaraDataInterfaceGrid3DCollection>(DataInterface))
+				{
+					TArray<FNiagaraVariableBase> GridVariables;
+					TArray<uint32> GridVariableOffsets;
+					int32 NumAttribChannelsFound;
+					Grid3D->FindAttributes(GridVariables, GridVariableOffsets, NumAttribChannelsFound);
+
+					for (const FNiagaraVariableBase& GridVariable : GridVariables)
+					{
+						const FString GridVariableString = GridVariable.GetName().ToString();
+
+						FNiagaraBakerOutputBinding& NewBinding = OutBindings.AddDefaulted_GetRef();
+						NewBinding.BindingName = FName(EmitterName + "." + VariableName + "." + GridVariableString);
+						NewBinding.MenuCategory = FText::FromString(EmitterName + " " + TEXT("Grid3DCollection"));
+						NewBinding.MenuEntry = FText::FromString(VariableName + "." + GridVariableString);
+					}
+				}
+				else if (UNiagaraDataInterfaceRenderTargetVolume* VolumeRT = Cast<UNiagaraDataInterfaceRenderTargetVolume>(DataInterface))
+				{
+					FNiagaraBakerOutputBinding& NewBinding = OutBindings.AddDefaulted_GetRef();
+					NewBinding.BindingName = FName(EmitterName + "." + VariableName);
+					NewBinding.MenuCategory = FText::FromString(EmitterName + " " + TEXT("VolumeRenderTarget"));
+					NewBinding.MenuEntry = FText::FromString(VariableName);
+				}
+			}
+		);
+	}
+	return OutBindings;
+}
+
+FIntPoint FNiagaraBakerRendererOutputSparseVolumeTexture::GetPreviewSize(UNiagaraBakerOutput* InBakerOutput, FIntPoint InAvailableSize) const
+{
+	return InAvailableSize;
+}
+
+void FNiagaraBakerRendererOutputSparseVolumeTexture::RenderPreview(UNiagaraBakerOutput* InBakerOutput, const FNiagaraBakerRenderer& BakerRenderer, UTextureRenderTarget2D* RenderTarget, TOptional<FString>& OutErrorString) const
+{
+	BakerRenderer.RenderSceneCapture(RenderTarget, ESceneCaptureSource::SCS_SceneColorHDR);
+}
+
+FIntPoint FNiagaraBakerRendererOutputSparseVolumeTexture::GetGeneratedSize(UNiagaraBakerOutput* InBakerOutput, FIntPoint InAvailableSize) const
+{
+	return InAvailableSize;
+}
+
+void FNiagaraBakerRendererOutputSparseVolumeTexture::RenderGenerated(UNiagaraBakerOutput* InBakerOutput, const FNiagaraBakerRenderer& BakerRenderer, UTextureRenderTarget2D* RenderTarget, TOptional<FString>& OutErrorString) const
+{
+	static FString SVTNotFoundError(TEXT("Sparse Volume Texture asset not found.\nPlease bake to see the result."));
+
+	UNiagaraBakerOutputSparseVolumeTexture* BakerOutput = CastChecked<UNiagaraBakerOutputSparseVolumeTexture>(InBakerOutput);
+	UNiagaraBakerSettings* BakerSettings = BakerOutput->GetTypedOuter<UNiagaraBakerSettings>();
+
+	UAnimatedSparseVolumeTexture* SVT = BakerOutput->GetAsset<UAnimatedSparseVolumeTexture>(BakerOutput->SparseVolumeTextureAssetPathFormat, 0);
+	if (SVT == nullptr)
+	{
+		OutErrorString = SVTNotFoundError;
+		return;
+	}
+
+	const float WorldTime = BakerRenderer.GetWorldTime();
+	FCanvas Canvas(RenderTarget->GameThread_GetRenderTargetResource(), nullptr, FGameTime::CreateUndilated(WorldTime, FApp::GetDeltaTime()), BakerRenderer.GetFeatureLevel());
+
+	const FNiagaraBakerOutputFrameIndices FrameIndices = BakerSettings->GetOutputFrameIndices(BakerOutput, WorldTime);
+
+	BakerRenderer.RenderSparseVolumeTexture(RenderTarget, FrameIndices, SVT);
+}
+
+bool FNiagaraBakerRendererOutputSparseVolumeTexture::BeginBake(FNiagaraBakerFeedbackContext& FeedbackContext, UNiagaraBakerOutput* InBakerOutput)
+{
+	UNiagaraBakerOutputSparseVolumeTexture* OutputVolumeTexture = CastChecked<UNiagaraBakerOutputSparseVolumeTexture>(InBakerOutput);
+
+	const FString AssetFullName = OutputVolumeTexture->GetAssetPath(OutputVolumeTexture->SparseVolumeTextureAssetPathFormat, 0);
+	
+	
+	IAssetRegistry* AssetRegistry = IAssetRegistry::Get();
+	TArray<FAssetData> FoundAssets;
+	bool FoundAsset = false;
+	if (AssetRegistry->GetAssetsByPackageName(FName(AssetFullName), FoundAssets))
+	{
+		if (FoundAssets.Num() > 0)
+		{
+			if (UObject* ExistingOject = StaticLoadObject(UAnimatedSparseVolumeTexture::StaticClass(), nullptr, *AssetFullName))
+			{
+				FoundAsset = true;
+			}
+		}
+	}	
+
+	UNiagaraBakerOutputSparseVolumeTexture* BakerOutput = CastChecked<UNiagaraBakerOutputSparseVolumeTexture>(InBakerOutput);
+	UAnimatedSparseVolumeTexture* SVT = BakerOutput->GetAsset<UAnimatedSparseVolumeTexture>(BakerOutput->SparseVolumeTextureAssetPathFormat, 0);
+
+	if (SVT == nullptr)
+	{
+		SVTAsset = UNiagaraBakerOutput::GetOrCreateAsset<UAnimatedSparseVolumeTexture, USparseVolumeTextureFactory>(AssetFullName);
+	}
+	else
+	{				
+		SVTAsset = NewObject<UAnimatedSparseVolumeTexture>(SVT->GetOuter(), UAnimatedSparseVolumeTexture::StaticClass(), *SVT->GetName(), RF_Public | RF_Standalone);
+		SVTAsset->PostEditChange();				
+	}
+	
+	if (!SVTAsset->BeginInitialize(1))
+	{
+		UE_LOG(LogNiagaraBaker, Error, TEXT("something isnt fun"));
+	}
+
+	return true;
+}
+
+void FNiagaraBakerRendererOutputSparseVolumeTexture::BakeFrame(FNiagaraBakerFeedbackContext& FeedbackContext, UNiagaraBakerOutput* InBakerOutput, int FrameIndex, const FNiagaraBakerRenderer& BakerRenderer)
+{
+	UNiagaraBakerOutputSparseVolumeTexture* OutputVolumeTexture = CastChecked<UNiagaraBakerOutputSparseVolumeTexture>(InBakerOutput);
+
+	FVolumeDataInterfaceHelper DataInterface;
+
+	TArray<FString> DataInterfacePath;
+	OutputVolumeTexture->SourceBinding.SourceName.ToString().ParseIntoArray(DataInterfacePath, TEXT("."));
+	if ( DataInterface.Initialize(DataInterfacePath, BakerRenderer.GetPreviewComponent()) == false )
+	{
+		return;
+	}
+	
+	FNiagaraDataInterfaceProxyRenderTargetVolumeProxy* RT_Proxy_RenderTargetVolume = nullptr;
+	FNiagaraDataInterfaceProxyGrid3DCollectionProxy* RT_Proxy_Grid3D = nullptr;
+	if (DataInterface.VolumeRenderTargetDataInterface != nullptr)
+	{					
+		RT_Proxy_RenderTargetVolume = DataInterface.VolumeRenderTargetProxy;		
+	}
+	else if (DataInterface.Grid3DDataInterface != nullptr)
+	{
+		if (!DataInterface.Grid3DInstanceData_GameThread->UseRGBATexture)
+		{
+			UE_LOG(LogNiagaraBaker, Error, TEXT("Cannot bake SVTs from non RGBA Grid3D Collections"));
+		}		
+		RT_Proxy_Grid3D = DataInterface.Grid3DProxy;	
+	}
+	else
+	{
+		UE_LOG(LogNiagaraBaker, Error, TEXT("Cannot bake from data interface"));
+	}
+					
+	//-OPT: Currently we are flushing rendering commands.  Do not remove this until making access to the frame data safe across threads.
+	TArray<uint8> TextureData;
+	FIntVector VolumeResolution = FIntVector(-1, -1, -1);
+	EPixelFormat VolumeFormat = EPixelFormat::PF_A1;
+
+	ENQUEUE_RENDER_COMMAND(NDIRenderTargetVolume_CacheFrame)
+		(
+			[RT_Proxy_RenderTargetVolume, RT_Proxy_Grid3D, RT_InstanceID = DataInterface.SystemInstance->GetId(), 
+			RT_TextureData = &TextureData, RT_VolumeResolution = &VolumeResolution, RT_VolumeFormat = &VolumeFormat](FRHICommandListImmediate& RHICmdList)
+			{
+				FRHIGPUTextureReadback RenderTargetReadback("ReadVolumeTexture");
+				uint32 BlockBytes = -1;				
+
+				if (RT_Proxy_RenderTargetVolume)
+				{
+					if (const FRenderTargetVolumeRWInstanceData_RenderThread* InstanceData_RT = RT_Proxy_RenderTargetVolume->SystemInstancesToProxyData_RT.Find(RT_InstanceID))
+					{
+						*RT_VolumeResolution = InstanceData_RT->Size;
+						RenderTargetReadback.EnqueueCopy(RHICmdList, InstanceData_RT->RenderTarget->GetRHI(), FIntVector(0, 0, 0), 0, *RT_VolumeResolution);
+						BlockBytes = GPixelFormats[InstanceData_RT->RenderTarget->GetRHI()->GetFormat()].BlockBytes;
+						*RT_VolumeFormat = InstanceData_RT->RenderTarget->GetRHI()->GetFormat();
+					}
+					else
+					{
+						UE_LOG(LogNiagaraBaker, Error, TEXT("No valid volume RT DI to do readback from"));
+						return;
+					}
+				}
+				else if (RT_Proxy_Grid3D)
+				{
+					if (const FGrid3DCollectionRWInstanceData_RenderThread* InstanceData_RT = RT_Proxy_Grid3D->SystemInstancesToProxyData_RT.Find(RT_InstanceID))
+					{
+						if (InstanceData_RT->CurrentData)
+						{
+							*RT_VolumeResolution = InstanceData_RT->NumCells;
+							RenderTargetReadback.EnqueueCopy(RHICmdList, InstanceData_RT->CurrentData->GetPooledTexture()->GetRHI(), FIntVector(0, 0, 0), 0, *RT_VolumeResolution);
+							BlockBytes = GPixelFormats[InstanceData_RT->CurrentData->GetPooledTexture()->GetRHI()->GetFormat()].BlockBytes;
+							*RT_VolumeFormat = InstanceData_RT->CurrentData->GetPooledTexture()->GetRHI()->GetFormat();
+						}
+						else
+						{
+							UE_LOG(LogNiagaraBaker, Error, TEXT("No valid grid DI to do readback from"));
+							return;
+						}
+					}
+					else
+					{
+						UE_LOG(LogNiagaraBaker, Error, TEXT("No valid grid DI to do readback from"));
+						return;
+					}
+				}
+				else
+				{
+					UE_LOG(LogNiagaraBaker, Error, TEXT("No valid grid DI to do readback from"));
+					return;
+				}
+
+				// Sync the GPU. Unfortunately we can't use the fences because not all RHIs implement them yet.
+				RHICmdList.BlockUntilGPUIdle();
+				RHICmdList.FlushResources();
+
+				//Lock the readback staging texture
+				int32 RowPitchInPixels;
+				int32 BufferHeight;
+				const uint8* LockedData = (const uint8*)RenderTargetReadback.Lock(RowPitchInPixels, &BufferHeight);
+						
+				int32 Count = RT_VolumeResolution->X * RT_VolumeResolution->Y * RT_VolumeResolution->Z * BlockBytes;
+				RT_TextureData->AddUninitialized(Count);
+
+				const uint8* SliceStart = LockedData;
+				for (int32 Z = 0; Z < RT_VolumeResolution->Z; ++Z)
+				{
+					const uint8* RowStart = SliceStart;
+					for (int32 Y = 0; Y < RT_VolumeResolution->Y; ++Y)
+					{
+						int32 Offset = 0 + Y * RT_VolumeResolution->X + Z * RT_VolumeResolution->X * RT_VolumeResolution->Y;
+						FMemory::Memcpy(RT_TextureData->GetData() + Offset * BlockBytes, RowStart, BlockBytes * RT_VolumeResolution->X);
+
+						RowStart += RowPitchInPixels * BlockBytes;
+					}
+
+					SliceStart += BufferHeight * RowPitchInPixels * BlockBytes;
+				}
+
+				//Unlock the staging texture
+				RenderTargetReadback.Unlock();					
+			}
+	);
+	FlushRenderingCommands();
+
+	if (TextureData.Num() > 0)
+	{
+
+#if WITH_EDITOR			
+		UE::SVT::FTextureDataCreateInfo SVTCreateInfo;
+		SVTCreateInfo.VirtualVolumeAABBMin = FIntVector3::ZeroValue;
+		SVTCreateInfo.VirtualVolumeAABBMax = VolumeResolution;
+		SVTCreateInfo.FallbackValues[0] = FVector4f(0, 0, 0, 0);
+		SVTCreateInfo.FallbackValues[1] = FVector4f(0, 0, 0, 0);
+		SVTCreateInfo.AttributesFormats[0] = VolumeFormat;
+		SVTCreateInfo.AttributesFormats[1] = PF_Unknown;
+
+		UE::SVT::FTextureData SparseTextureData{};
+		bool Success = SparseTextureData.CreateFromDense(SVTCreateInfo, TArrayView<uint8, int64>((uint8*)TextureData.GetData(), (int64)TextureData.Num() * sizeof(TextureData[0])), TArrayView<uint8>());
+
+		if (!Success)
+		{
+			UE_LOG(LogNiagaraBaker, Error, TEXT("Cannot create SVT for data interface"));
+
+			return;
+		}
+
+		if (!SVTAsset->AppendFrame(SparseTextureData))
+		{
+			UE_LOG(LogNiagaraBaker, Error, TEXT("Cannot append frame to SVT"));
+		}
+#endif			
+	}
+
+}
+
+void FNiagaraBakerRendererOutputSparseVolumeTexture::EndBake(FNiagaraBakerFeedbackContext& FeedbackContext, UNiagaraBakerOutput* InBakerOutput)
+{
+	UNiagaraBakerOutputSparseVolumeTexture* OutputVolumeTexture = CastChecked<UNiagaraBakerOutputSparseVolumeTexture>(InBakerOutput);
+		
+	if (!SVTAsset->EndInitialize())
+	{
+		UE_LOG(LogNiagaraBaker, Error, TEXT("Cannot EndInitialize on creating SVT"));
+	}
+
+	SVTAsset->PostLoad();
+}
+
diff --git a/Engine/Plugins/FX/Niagara/Source/NiagaraEditor/Private/NiagaraBakerRendererOutputSparseVolumeTexture.h b/Engine/Plugins/FX/Niagara/Source/NiagaraEditor/Private/NiagaraBakerRendererOutputSparseVolumeTexture.h
new file mode 100644
--- /dev/null
+++ b/Engine/Plugins/FX/Niagara/Source/NiagaraEditor/Private/NiagaraBakerRendererOutputSparseVolumeTexture.h
@@ -0,0 +1,28 @@
+// Copyright Epic Games, Inc. All Rights Reserved.
+
+#pragma once
+
+#include "NiagaraBakerRenderer.h"
+
+class UNiagaraBakerOutputSparseVolumeTexture;
+class UAnimatedSparseVolumeTexture;
+
+class FNiagaraBakerRendererOutputSparseVolumeTexture : public FNiagaraBakerOutputRenderer
+{
+public:
+	virtual TArray<FNiagaraBakerOutputBinding> GetRendererBindings(UNiagaraBakerOutput* InBakerOutput) const override;
+
+	virtual FIntPoint GetPreviewSize(UNiagaraBakerOutput* InBakerOutput, FIntPoint InAvailableSize) const override;
+	virtual void RenderPreview(UNiagaraBakerOutput* InBakerOutput, const FNiagaraBakerRenderer& BakerRenderer, UTextureRenderTarget2D* RenderTarget, TOptional<FString>& OutErrorString) const override;
+
+	virtual FIntPoint GetGeneratedSize(UNiagaraBakerOutput* InBakerOutput, FIntPoint InAvailableSize) const override;
+	virtual void RenderGenerated(UNiagaraBakerOutput* InBakerOutput, const FNiagaraBakerRenderer& BakerRenderer, UTextureRenderTarget2D* RenderTarget, TOptional<FString>& OutErrorString) const override;
+
+	virtual bool BeginBake(FNiagaraBakerFeedbackContext& FeedbackContext, UNiagaraBakerOutput* InBakerOutput) override;
+	virtual void BakeFrame(FNiagaraBakerFeedbackContext& FeedbackContext, UNiagaraBakerOutput* InBakerOutput, int FrameIndex, const FNiagaraBakerRenderer& BakerRenderer) override;
+	virtual void EndBake(FNiagaraBakerFeedbackContext& FeedbackContext, UNiagaraBakerOutput* InBakerOutput) override;
+
+protected:
+
+	TObjectPtr<UAnimatedSparseVolumeTexture> SVTAsset = nullptr;	
+};
diff --git a/Engine/Plugins/FX/Niagara/Source/NiagaraEditor/Private/NiagaraBakerRendererOutputVolumeTexture.cpp b/Engine/Plugins/FX/Niagara/Source/NiagaraEditor/Private/NiagaraBakerRendererOutputVolumeTexture.cpp
--- a/Engine/Plugins/FX/Niagara/Source/NiagaraEditor/Private/NiagaraBakerRendererOutputVolumeTexture.cpp
+++ b/Engine/Plugins/FX/Niagara/Source/NiagaraEditor/Private/NiagaraBakerRendererOutputVolumeTexture.cpp
@@ -8,10 +8,9 @@
 #include "NiagaraSystemInstance.h"
 #include "NiagaraSystemInstanceController.h"
 #include "NiagaraComputeExecutionContext.h"
+#include "NiagaraBatchedElements.h"
 #include "NiagaraDataInterfaceGrid3DCollection.h"
 #include "NiagaraDataInterfaceRenderTargetVolume.h"
-#include "NiagaraBatchedElements.h"
-
 #include "Engine/Canvas.h"
 #include "Engine/TextureRenderTarget2D.h"
 #include "Engine/VolumeTexture.h"
@@ -25,97 +24,6 @@
 #include "VolumeCache.h"
 #include "VolumeCacheFactory.h"
 
-DEFINE_LOG_CATEGORY_STATIC(LogNiagaraBaker, Log, All);
-
-struct FVolumeDataInterfaceHelper
-{
-	UNiagaraComponent*										NiagaraComponent = nullptr;
-	FNiagaraSystemInstance*									SystemInstance = nullptr;
-	TArray<FString>											DataInterfacePath;
-
-	UNiagaraDataInterfaceGrid3DCollection*					Grid3DDataInterface = nullptr;
-	FNiagaraDataInterfaceProxyGrid3DCollectionProxy*		Grid3DProxy = nullptr;
-	FGrid3DCollectionRWInstanceData_GameThread*				Grid3DInstanceData_GameThread = nullptr;
-	FName													Grid3DAttributeName;
-	int32													Grid3DVariableIndex = INDEX_NONE;
-	int32													Grid3DAttributeStart = INDEX_NONE;
-	int32													Grid3DAttributeChannels = 0;
-	FIntVector												Grid3DTextureSize = FIntVector::ZeroValue;
-
-	UNiagaraDataInterfaceRenderTargetVolume*				VolumeRenderTargetDataInterface = nullptr;
-	FNiagaraDataInterfaceProxyRenderTargetVolumeProxy*		VolumeRenderTargetProxy = nullptr;
-	FRenderTargetVolumeRWInstanceData_GameThread*			VolumeRenderTargetInstanceData_GameThread = nullptr;
-
-	bool Initialize(UNiagaraBakerOutputVolumeTexture* OutputVolumeTexture, UNiagaraComponent* InNiagaraComponent)
-	{
-		NiagaraComponent = InNiagaraComponent;
-
-		OutputVolumeTexture->SourceBinding.SourceName.ToString().ParseIntoArray(DataInterfacePath, TEXT("."));
-		if (DataInterfacePath.Num() < 2)
-		{
-			return false;
-		}
-
-		const FName DataInterfaceName(DataInterfacePath[0] + "." + DataInterfacePath[1]);
-		UNiagaraDataInterface* DataInterface = FNiagaraBakerOutputBindingHelper::GetDataInterface(NiagaraComponent, DataInterfaceName);
-		if (DataInterface == nullptr)
-		{
-			return false;
-		}
-
-		// Guaranteed since we got a data interface
-		SystemInstance = NiagaraComponent->GetSystemInstanceController()->GetSoloSystemInstance();
-
-		// Render Target Volume
-		if (DataInterface->IsA<UNiagaraDataInterfaceRenderTargetVolume>())
-		{
-			VolumeRenderTargetDataInterface = CastChecked<UNiagaraDataInterfaceRenderTargetVolume>(DataInterface);
-			VolumeRenderTargetProxy = static_cast<FNiagaraDataInterfaceProxyRenderTargetVolumeProxy*>(VolumeRenderTargetDataInterface->GetProxy());
-			VolumeRenderTargetInstanceData_GameThread = reinterpret_cast<FRenderTargetVolumeRWInstanceData_GameThread*>(SystemInstance->FindDataInterfaceInstanceData(VolumeRenderTargetDataInterface));
-			if (VolumeRenderTargetInstanceData_GameThread == nullptr)
-			{
-				return false;
-			}
-		}
-		// Grid 3D
-		else if ( DataInterface->IsA<UNiagaraDataInterfaceGrid3DCollection>() )
-		{
-			Grid3DDataInterface = CastChecked<UNiagaraDataInterfaceGrid3DCollection>(DataInterface);
-			Grid3DProxy = static_cast<FNiagaraDataInterfaceProxyGrid3DCollectionProxy*>(Grid3DDataInterface->GetProxy());
-			Grid3DInstanceData_GameThread = reinterpret_cast<FGrid3DCollectionRWInstanceData_GameThread*>(SystemInstance->FindDataInterfaceInstanceData(Grid3DDataInterface));
-			if (Grid3DInstanceData_GameThread == nullptr)
-			{
-				return false;
-			}
-
-			if (DataInterfacePath.Num() != 3)
-			{
-				// Perhaps a path to pull all attributes, i.e. whole texture?
-				return false;
-			}
-
-			Grid3DAttributeName = FName(DataInterfacePath[2]);
-			Grid3DVariableIndex = Grid3DInstanceData_GameThread->Vars.IndexOfByPredicate([&](const FNiagaraVariableBase& VariableBase) { return VariableBase.GetName() == Grid3DAttributeName; });
-			if (Grid3DVariableIndex == INDEX_NONE)
-			{
-				return false;
-			}
-			Grid3DAttributeStart	= Grid3DInstanceData_GameThread->Offsets[Grid3DVariableIndex];
-			Grid3DAttributeChannels	= Grid3DInstanceData_GameThread->Vars[Grid3DVariableIndex].GetType().GetSize() / sizeof(float);
-			Grid3DTextureSize.X		= Grid3DInstanceData_GameThread->NumCells.X * Grid3DInstanceData_GameThread->NumTiles.X;
-			Grid3DTextureSize.Y		= Grid3DInstanceData_GameThread->NumCells.Y * Grid3DInstanceData_GameThread->NumTiles.Y;
-			Grid3DTextureSize.Z		= Grid3DInstanceData_GameThread->NumCells.Z * Grid3DInstanceData_GameThread->NumTiles.Z;
-		}
-		// Unsupported type
-		else
-		{
-			return false;
-		}
-
-		return true;
-	}
-};
-
 TArray<FNiagaraBakerOutputBinding> FNiagaraBakerRendererOutputVolumeTexture::GetRendererBindings(UNiagaraBakerOutput* InBakerOutput) const
 {
 	TArray<FNiagaraBakerOutputBinding> OutBindings;
@@ -170,7 +78,10 @@ void FNiagaraBakerRendererOutputVolumeTexture::RenderPreview(UNiagaraBakerOutput
 	ON_SCOPE_EXIT{ Canvas.Flush_GameThread(); };
 
 	FVolumeDataInterfaceHelper DataInterface;
-	if (DataInterface.Initialize(OutputVolumeTexture, BakerRenderer.GetPreviewComponent()) == false)
+
+	TArray<FString> DataInterfacePath;
+	OutputVolumeTexture->SourceBinding.SourceName.ToString().ParseIntoArray(DataInterfacePath, TEXT("."));
+	if (DataInterface.Initialize(DataInterfacePath, BakerRenderer.GetPreviewComponent()) == false)
 	{
 		OutErrorString = TEXT("Could not find data to preview from.\nPlease ensure binding is set to a valid source.");
 		return;
@@ -406,7 +317,10 @@ void FNiagaraBakerRendererOutputVolumeTexture::BakeFrame(FNiagaraBakerFeedbackCo
 	UNiagaraBakerOutputVolumeTexture* OutputVolumeTexture = CastChecked<UNiagaraBakerOutputVolumeTexture>(InBakerOutput);
 
 	FVolumeDataInterfaceHelper DataInterface;
-	if ( DataInterface.Initialize(OutputVolumeTexture, BakerRenderer.GetPreviewComponent()) == false )
+
+	TArray<FString> DataInterfacePath;
+	OutputVolumeTexture->SourceBinding.SourceName.ToString().ParseIntoArray(DataInterfacePath, TEXT("."));
+	if ( DataInterface.Initialize(DataInterfacePath, BakerRenderer.GetPreviewComponent()) == false )
 	{
 		return;
 	}
