Changes the Niagara AsyncTask compilation queue to keep track of only the tasks that are active.  This will move tasks that have been completed (but not yet collected by the calling code) to a new array.  This removes the risk of deadlocking some code paths where we're blocking on an individual task being complete.

diff --git a/Engine/Plugins/FX/Niagara/Source/NiagaraEditor/Private/NiagaraCompilationTasks.cpp b/Engine/Plugins/FX/Niagara/Source/NiagaraEditor/Private/NiagaraCompilationTasks.cpp
--- a/Engine/Plugins/FX/Niagara/Source/NiagaraEditor/Private/NiagaraCompilationTasks.cpp
+++ b/Engine/Plugins/FX/Niagara/Source/NiagaraEditor/Private/NiagaraCompilationTasks.cpp
@@ -1169,6 +1169,11 @@ bool FNiagaraSystemCompilationTask::CanRemove() const
 		|| (CompilationState == EState::Aborted);
 }
 
+bool FNiagaraSystemCompilationTask::AreResultsPending() const
+{
+	return CompilationState == EState::Completed && !ResultsRetrieved;
+}
+
 void FNiagaraSystemCompilationTask::WaitTillCompileCompletion()
 {
 	const FTimespan WaitTimeout = FTimespan::FromMilliseconds(50.0);
diff --git a/Engine/Plugins/FX/Niagara/Source/NiagaraEditor/Private/NiagaraCompilationTasks.h b/Engine/Plugins/FX/Niagara/Source/NiagaraEditor/Private/NiagaraCompilationTasks.h
--- a/Engine/Plugins/FX/Niagara/Source/NiagaraEditor/Private/NiagaraCompilationTasks.h
+++ b/Engine/Plugins/FX/Niagara/Source/NiagaraEditor/Private/NiagaraCompilationTasks.h
@@ -172,6 +172,7 @@ struct FNiagaraSystemCompilationTask
 	void Tick();
 	bool Poll(FNiagaraSystemAsyncCompileResults& Results) const;
 	bool CanRemove() const;
+	bool AreResultsPending() const;
 
 	void WaitTillCompileCompletion();
 	void WaitTillCachePutCompletion();
diff --git a/Engine/Plugins/FX/Niagara/Source/NiagaraEditor/Private/NiagaraSystemCompilingManager.cpp b/Engine/Plugins/FX/Niagara/Source/NiagaraEditor/Private/NiagaraSystemCompilingManager.cpp
--- a/Engine/Plugins/FX/Niagara/Source/NiagaraEditor/Private/NiagaraSystemCompilingManager.cpp
+++ b/Engine/Plugins/FX/Niagara/Source/NiagaraEditor/Private/NiagaraSystemCompilingManager.cpp
@@ -11,7 +11,7 @@
 
 #define LOCTEXT_NAMESPACE "NiagaraCompilationManager"
 
-static int GNiagaraCompilationMaxActiveTaskCount = 4;
+static int GNiagaraCompilationMaxActiveTaskCount = 16;
 static FAutoConsoleVariableRef CVarNiagaraCompilationMaxActiveTaskCount(
 	TEXT("fx.Niagara.Compilation.MaxActiveTaskCount"),
 	GNiagaraCompilationMaxActiveTaskCount,
@@ -46,7 +46,7 @@ int32 FNiagaraSystemCompilingManager::GetNumRemainingAssets() const
 
 	{
 		FReadScopeLock Read(QueueLock);
-		RemainingAssetCount = QueuedRequests.Num() + ActiveTasks.Num();
+		RemainingAssetCount = QueuedRequests.Num() + ActiveTasks.Num() + RequestsAwaitingRetrieval.Num();
 	}
 
 	return RemainingAssetCount;
@@ -88,7 +88,36 @@ void FNiagaraSystemCompilingManager::ProcessAsyncTasks(bool bLimitExecutionTime)
 
 		// find the list of tasks that we can remove
 		TArray<FNiagaraCompilationTaskHandle> TasksToRemove;
+		TArray<FNiagaraCompilationTaskHandle> TasksToRetrieve;
 		for (TArray<FNiagaraCompilationTaskHandle>::TIterator TaskIt(ActiveTasks); TaskIt; ++TaskIt)
+		{
+			FTaskPtr TaskPtr = SystemRequestMap.FindRef(*TaskIt);
+
+			bool bRemoveCurrent = true;
+			if (TaskPtr.IsValid())
+			{
+				if (TaskPtr->AreResultsPending())
+				{
+					TasksToRetrieve.Add(*TaskIt);
+				}
+				else if (TaskPtr->CanRemove())
+				{
+					TasksToRemove.Add(*TaskIt);
+				}
+				else
+				{
+					bRemoveCurrent = false;
+				}
+			}
+
+			if (bRemoveCurrent)
+			{
+				TaskIt.RemoveCurrent();
+			}
+		}
+
+		// go through the entries that are awaiting retrieval and clean up any that have been retrieved
+		for (TArray<FNiagaraCompilationTaskHandle>::TIterator TaskIt(RequestsAwaitingRetrieval); TaskIt; ++TaskIt)
 		{
 			FTaskPtr TaskPtr = SystemRequestMap.FindRef(*TaskIt);
 			if (!TaskPtr.IsValid() || TaskPtr->CanRemove())
@@ -98,14 +127,20 @@ void FNiagaraSystemCompilingManager::ProcessAsyncTasks(bool bLimitExecutionTime)
 			}
 		}
 
-		// remove our found tasks
+		// remove tasks that can be erased
 		for (FNiagaraCompilationTaskHandle TaskToRemove : TasksToRemove)
 		{
-			QueuedRequests.RemoveSingle(TaskToRemove);
+			ensure(!QueuedRequests.Contains(TaskToRemove));
 			SystemRequestMap.Remove(TaskToRemove);
 		}
-	}
 
+		// finally populate RequestsAwaitingRetrieval with any new entries
+		for (FNiagaraCompilationTaskHandle TaskToRetrieve : TasksToRetrieve)
+		{
+			ensure(!QueuedRequests.Contains(TaskToRetrieve));
+			RequestsAwaitingRetrieval.Add(TaskToRetrieve);
+		}
+	}
 
 	// queue up as many tasks as we can
 	while (ConditionalLaunchTask())
diff --git a/Engine/Plugins/FX/Niagara/Source/NiagaraEditor/Private/NiagaraSystemCompilingManager.h b/Engine/Plugins/FX/Niagara/Source/NiagaraEditor/Private/NiagaraSystemCompilingManager.h
--- a/Engine/Plugins/FX/Niagara/Source/NiagaraEditor/Private/NiagaraSystemCompilingManager.h
+++ b/Engine/Plugins/FX/Niagara/Source/NiagaraEditor/Private/NiagaraSystemCompilingManager.h
@@ -42,10 +42,10 @@ protected:
 	using FTaskPtr = TSharedPtr<FNiagaraSystemCompilationTask, ESPMode::ThreadSafe>;
 
 	mutable FRWLock QueueLock;
-	TArray<FNiagaraCompilationTaskHandle> QueuedRequests;
 	TMap<FNiagaraCompilationTaskHandle, FTaskPtr> SystemRequestMap;
+	TArray<FNiagaraCompilationTaskHandle> ActiveTasks;
+	TArray<FNiagaraCompilationTaskHandle> QueuedRequests;
+	TArray<FNiagaraCompilationTaskHandle> RequestsAwaitingRetrieval;
 
 	std::atomic<FNiagaraCompilationTaskHandle> NextTaskHandle = { 0 };
-
-	TArray<FNiagaraCompilationTaskHandle> ActiveTasks;
 };
