- Deprecate TransformVector & fix incorrect usage
- Replaced with TransformUnitVector to avoid making the same mistaken calls.
- The engine regular TransformVector will not normalize the result which leads to confusion.

diff --git a/Engine/Plugins/FX/Niagara/Source/Niagara/Classes/NiagaraDataInterface.h b/Engine/Plugins/FX/Niagara/Source/Niagara/Classes/NiagaraDataInterface.h
--- a/Engine/Plugins/FX/Niagara/Source/Niagara/Classes/NiagaraDataInterface.h
+++ b/Engine/Plugins/FX/Niagara/Source/Niagara/Classes/NiagaraDataInterface.h
@@ -30,24 +30,34 @@ struct FNDITransformHandlerNoop
 {
 	FORCEINLINE void TransformPosition(FVector3f& V, const FMatrix44f& M) const { }
 	FORCEINLINE void TransformPosition(FVector3d& V, const FMatrix44d& M) const { }
-	FORCEINLINE void TransformVector(FVector3f& V, const FMatrix44f& M) const { }
-	FORCEINLINE void TransformVector(FVector3d& V, const FMatrix44d& M) const { }
+	FORCEINLINE void TransformUnitVector(FVector3f& V, const FMatrix44f& M) const { }
+	FORCEINLINE void TransformUnitVector(FVector3d& V, const FMatrix44d& M) const { }
 	FORCEINLINE void TransformNotUnitVector(FVector3f& V, const FMatrix44f& M) const { }
 	FORCEINLINE void TransformNotUnitVector(FVector3d& V, const FMatrix44d& M) const { }
 	FORCEINLINE void TransformRotation(FQuat4f& Q1, const FQuat4f& Q2) const { }
 	FORCEINLINE void TransformRotation(FQuat4d& Q1, const FQuat4d& Q2) const { }
+
+	UE_DEPRECATED(5.4, "Please update your code to use TransformUnitVector")
+	FORCEINLINE void TransformVector(FVector3f& V, const FMatrix44f& M) const { }
+	UE_DEPRECATED(5.4, "Please update your code to use TransformUnitVector")
+	FORCEINLINE void TransformVector(FVector3d& V, const FMatrix44d& M) const { }
 };
 
 struct FNDITransformHandler
 {
 	FORCEINLINE void TransformPosition(FVector3f& P, const FMatrix44f& M) const { P = M.TransformPosition(P); }
 	FORCEINLINE void TransformPosition(FVector3d& P, const FMatrix44d& M) const { P = M.TransformPosition(P); }
-	FORCEINLINE void TransformVector(FVector3f& V, const FMatrix44f& M) const { V = M.TransformVector(V).GetUnsafeNormal3(); }
-	FORCEINLINE void TransformVector(FVector3d& V, const FMatrix44d& M) const { V = M.TransformVector(V).GetUnsafeNormal3(); }
+	FORCEINLINE void TransformUnitVector(FVector3f& V, const FMatrix44f& M) const { V = M.TransformVector(V).GetUnsafeNormal3(); }
+	FORCEINLINE void TransformUnitVector(FVector3d& V, const FMatrix44d& M) const { V = M.TransformVector(V).GetUnsafeNormal3(); }
 	FORCEINLINE void TransformNotUnitVector(FVector3f& V, const FMatrix44f& M) const { V = M.TransformVector(V); }
 	FORCEINLINE void TransformNotUnitVector(FVector3d& V, const FMatrix44d& M) const { V = M.TransformVector(V); }
 	FORCEINLINE void TransformRotation(FQuat4f& Q1, const FQuat4f& Q2) const { Q1 = Q2 * Q1; }
 	FORCEINLINE void TransformRotation(FQuat4d& Q1, const FQuat4d& Q2) const { Q1 = Q2 * Q1; }
+
+	UE_DEPRECATED(5.4, "Please update your code to use TransformUnitVector")
+	FORCEINLINE void TransformVector(FVector3f& V, const FMatrix44f& M) const { V = M.TransformVector(V).GetUnsafeNormal3(); }
+	UE_DEPRECATED(5.4, "Please update your code to use TransformUnitVector")
+	FORCEINLINE void TransformVector(FVector3d& V, const FMatrix44d& M) const { V = M.TransformVector(V).GetUnsafeNormal3(); }
 };
 
 // FNiagaraDataInterfaceProxy should always outlive any ticks, etc, that are on the rendering thread.
diff --git a/Engine/Plugins/FX/Niagara/Source/Niagara/Private/DataInterface/NiagaraDataInterfaceStaticMesh.cpp b/Engine/Plugins/FX/Niagara/Source/Niagara/Private/DataInterface/NiagaraDataInterfaceStaticMesh.cpp
--- a/Engine/Plugins/FX/Niagara/Source/Niagara/Private/DataInterface/NiagaraDataInterfaceStaticMesh.cpp
+++ b/Engine/Plugins/FX/Niagara/Source/Niagara/Private/DataInterface/NiagaraDataInterfaceStaticMesh.cpp
@@ -284,9 +284,9 @@ namespace NDIStaticMeshLocal
 				TangentY = FVector3f(0.0f, 1.0f, 0.0f);
 				TangentZ = FVector3f(0.0f, 0.0f, 1.0f);
 			}
-			TransformHandler.TransformVector(TangentX, Matrix);
-			TransformHandler.TransformVector(TangentY, Matrix);
-			TransformHandler.TransformVector(TangentZ, Matrix);
+			TransformHandler.TransformUnitVector(TangentX, Matrix);
+			TransformHandler.TransformUnitVector(TangentY, Matrix);
+			TransformHandler.TransformUnitVector(TangentZ, Matrix);
 		}
 	};
 
@@ -1310,9 +1310,15 @@ namespace NDIStaticMeshLocal
 			return Position;
 		}
 
-		FORCEINLINE FVector3f TransformVector(FVector3f Vector) const
+		FORCEINLINE FVector3f TransformUnitVector(FVector3f Vector) const
 		{
-			TransformHandler.TransformVector(Vector, FMatrix44f(InstanceData->TransformInverseTransposed));		// LWC_TODO: Precision loss
+			TransformHandler.TransformUnitVector(Vector, FMatrix44f(InstanceData->TransformInverseTransposed));		// LWC_TODO: Precision loss
+			return Vector;
+		}
+
+		FORCEINLINE FVector3f TransformNotUnitVector(FVector3f Vector) const
+		{
+			TransformHandler.TransformNotUnitVector(Vector, FMatrix44f(InstanceData->TransformInverseTransposed));		// LWC_TODO: Precision loss
 			return Vector;
 		}
 
@@ -1328,9 +1334,15 @@ namespace NDIStaticMeshLocal
 			return Position;
 		}
 
-		FORCEINLINE FVector3f PreviousTransformVector(FVector3f Vector) const
+		FORCEINLINE FVector3f PreviousTransformUnitVector(FVector3f Vector) const
+		{
+			TransformHandler.TransformUnitVector(Vector, FMatrix44f(InstanceData->PrevTransformInverseTransposed));	// LWC_TODO: Precision loss
+			return Vector;
+		}
+
+		FORCEINLINE FVector3f PreviousTransformNotUnitVector(FVector3f Vector) const
 		{
-			TransformHandler.TransformVector(Vector, FMatrix44f(InstanceData->PrevTransformInverseTransposed));	// LWC_TODO: Precision loss
+			TransformHandler.TransformNotUnitVector(Vector, FMatrix44f(InstanceData->PrevTransformInverseTransposed));	// LWC_TODO: Precision loss
 			return Vector;
 		}
 
@@ -1396,7 +1408,7 @@ namespace NDIStaticMeshLocal
 			Tangent  = LODResource->VertexBuffers.StaticMeshVertexBuffer.VertexTangentX(Index0) * BaryCoord.X;
 			Tangent += LODResource->VertexBuffers.StaticMeshVertexBuffer.VertexTangentX(Index1) * BaryCoord.Y;
 			Tangent += LODResource->VertexBuffers.StaticMeshVertexBuffer.VertexTangentX(Index2) * BaryCoord.Z;
-			TransformHandler.TransformVector(Tangent, FMatrix44f(InstanceData->Transform));						// LWC_TODO: Precision loss?
+			TransformHandler.TransformUnitVector(Tangent, FMatrix44f(InstanceData->Transform));						// LWC_TODO: Precision loss?
 			return Tangent;
 		}
 
@@ -1406,7 +1418,7 @@ namespace NDIStaticMeshLocal
 			Tangent  = LODResource->VertexBuffers.StaticMeshVertexBuffer.VertexTangentY(Index0) * BaryCoord.X;
 			Tangent += LODResource->VertexBuffers.StaticMeshVertexBuffer.VertexTangentY(Index1) * BaryCoord.Y;
 			Tangent += LODResource->VertexBuffers.StaticMeshVertexBuffer.VertexTangentY(Index2) * BaryCoord.Z;
-			TransformHandler.TransformVector(Tangent, FMatrix44f(InstanceData->Transform));						// LWC_TODO: Precision loss?
+			TransformHandler.TransformUnitVector(Tangent, FMatrix44f(InstanceData->Transform));						// LWC_TODO: Precision loss?
 			return Tangent;
 		}
 
@@ -1416,7 +1428,7 @@ namespace NDIStaticMeshLocal
 			Tangent  = LODResource->VertexBuffers.StaticMeshVertexBuffer.VertexTangentZ(Index0) * BaryCoord.X;
 			Tangent += LODResource->VertexBuffers.StaticMeshVertexBuffer.VertexTangentZ(Index1) * BaryCoord.Y;
 			Tangent += LODResource->VertexBuffers.StaticMeshVertexBuffer.VertexTangentZ(Index2) * BaryCoord.Z;
-			TransformHandler.TransformVector(Tangent, FMatrix44f(InstanceData->Transform));						// LWC_TODO: Precision loss?
+			TransformHandler.TransformUnitVector(Tangent, FMatrix44f(InstanceData->Transform));						// LWC_TODO: Precision loss?
 			return Tangent;
 		}
 
@@ -1428,8 +1440,8 @@ namespace NDIStaticMeshLocal
 			Tangent += LODResource->VertexBuffers.StaticMeshVertexBuffer.VertexTangentX(Index2) * BaryCoord.Z;
 			FVector3f CurrTangent = Tangent;
 			FVector3f PrevTangent = Tangent;
-			TransformHandler.TransformVector(PrevTangent, FMatrix44f(InstanceData->PrevTransform));						// LWC_TODO: Precision loss?
-			TransformHandler.TransformVector(CurrTangent, FMatrix44f(InstanceData->Transform));						// LWC_TODO: Precision loss?
+			TransformHandler.TransformUnitVector(PrevTangent, FMatrix44f(InstanceData->PrevTransform));						// LWC_TODO: Precision loss?
+			TransformHandler.TransformUnitVector(CurrTangent, FMatrix44f(InstanceData->Transform));						// LWC_TODO: Precision loss?
 			return FMath::Lerp(PrevTangent, CurrTangent, Interp);
 		}
 
@@ -1441,8 +1453,8 @@ namespace NDIStaticMeshLocal
 			Tangent += LODResource->VertexBuffers.StaticMeshVertexBuffer.VertexTangentY(Index2) * BaryCoord.Z;
 			FVector3f CurrTangent = Tangent;
 			FVector3f PrevTangent = Tangent;
-			TransformHandler.TransformVector(PrevTangent, FMatrix44f(InstanceData->PrevTransform));						// LWC_TODO: Precision loss?
-			TransformHandler.TransformVector(CurrTangent, FMatrix44f(InstanceData->Transform));						// LWC_TODO: Precision loss?
+			TransformHandler.TransformUnitVector(PrevTangent, FMatrix44f(InstanceData->PrevTransform));						// LWC_TODO: Precision loss?
+			TransformHandler.TransformUnitVector(CurrTangent, FMatrix44f(InstanceData->Transform));						// LWC_TODO: Precision loss?
 			return FMath::Lerp(PrevTangent, CurrTangent, Interp);
 		}
 
@@ -1454,8 +1466,8 @@ namespace NDIStaticMeshLocal
 			Tangent += LODResource->VertexBuffers.StaticMeshVertexBuffer.VertexTangentZ(Index2) * BaryCoord.Z;
 			FVector3f CurrTangent = Tangent;
 			FVector3f PrevTangent = Tangent;
-			TransformHandler.TransformVector(PrevTangent, FMatrix44f(InstanceData->PrevTransform));						// LWC_TODO: Precision loss?
-			TransformHandler.TransformVector(CurrTangent, FMatrix44f(InstanceData->Transform));						// LWC_TODO: Precision loss?
+			TransformHandler.TransformUnitVector(PrevTangent, FMatrix44f(InstanceData->PrevTransform));						// LWC_TODO: Precision loss?
+			TransformHandler.TransformUnitVector(CurrTangent, FMatrix44f(InstanceData->Transform));						// LWC_TODO: Precision loss?
 			return FMath::Lerp(PrevTangent, CurrTangent, Interp);
 		}
 
@@ -1648,8 +1660,8 @@ namespace NDIStaticMeshLocal
 			const FVector3f CurrPosition = TransformPosition(SocketTransform.GetLocation());
 			const FQuat4f PrevRotation = PreviousTransformRotation(SocketTransform.GetRotation());
 			const FQuat4f CurrRotation = TransformRotation(SocketTransform.GetRotation());
-			const FVector3f PrevScale = PreviousTransformVector(SocketTransform.GetScale3D());
-			const FVector3f CurrScale = TransformVector(SocketTransform.GetScale3D());
+			const FVector3f PrevScale = PreviousTransformNotUnitVector(SocketTransform.GetScale3D());
+			const FVector3f CurrScale = TransformNotUnitVector(SocketTransform.GetScale3D());
 
 			OutPosition = FMath::Lerp(PrevPosition, CurrPosition, Interp);
 			OutRotation = FQuat4f::Slerp(PrevRotation, CurrRotation, Interp);
@@ -3795,9 +3807,9 @@ void UNiagaraDataInterfaceStaticMesh::VMGetVertex(FVectorVMExternalFunctionConte
 	{
 		const FVector3f Position = StaticMeshHelper.TransformPosition(FVector3f::ZeroVector);
 		const FVector3f Velocity = FVector3f::ZeroVector;
-		const FVector3f TangentX = StaticMeshHelper.TransformVector(FVector3f(1.0f, 0.0f, 0.0f));
-		const FVector3f TangentY = StaticMeshHelper.TransformVector(FVector3f(0.0f, 1.0f, 0.0f));
-		const FVector3f TangentZ = StaticMeshHelper.TransformVector(FVector3f(0.0f, 0.0f, 1.0f));
+		const FVector3f TangentX = StaticMeshHelper.TransformUnitVector(FVector3f(1.0f, 0.0f, 0.0f));
+		const FVector3f TangentY = StaticMeshHelper.TransformUnitVector(FVector3f(0.0f, 1.0f, 0.0f));
+		const FVector3f TangentZ = StaticMeshHelper.TransformUnitVector(FVector3f(0.0f, 0.0f, 1.0f));
 		for ( int32 i=0; i < Context.GetNumInstances(); ++i )
 		{
 			OutPosition.SetAndAdvance(Position);
@@ -3846,9 +3858,9 @@ void UNiagaraDataInterfaceStaticMesh::VMGetVertexInterpolated(FVectorVMExternalF
 	{
 		const FVector3f Position = StaticMeshHelper.TransformPosition(FVector3f::ZeroVector);
 		const FVector3f Velocity = FVector3f::ZeroVector;
-		const FVector3f TangentX = StaticMeshHelper.TransformVector(FVector3f(1.0f, 0.0f, 0.0f));
-		const FVector3f TangentY = StaticMeshHelper.TransformVector(FVector3f(0.0f, 1.0f, 0.0f));
-		const FVector3f TangentZ = StaticMeshHelper.TransformVector(FVector3f(0.0f, 0.0f, 1.0f));
+		const FVector3f TangentX = StaticMeshHelper.TransformUnitVector(FVector3f(1.0f, 0.0f, 0.0f));
+		const FVector3f TangentY = StaticMeshHelper.TransformUnitVector(FVector3f(0.0f, 1.0f, 0.0f));
+		const FVector3f TangentZ = StaticMeshHelper.TransformUnitVector(FVector3f(0.0f, 0.0f, 1.0f));
 		for (int32 i = 0; i < Context.GetNumInstances(); ++i)
 		{
 			OutPosition.SetAndAdvance(Position);
@@ -4195,9 +4207,9 @@ void UNiagaraDataInterfaceStaticMesh::VMGetTriangle(FVectorVMExternalFunctionCon
 		const FVector3f Position = StaticMeshHelper.TransformPosition(FVector3f::ZeroVector);
 		const FVector3f PreviousPosition = StaticMeshHelper.PreviousTransformPosition(FVector3f::ZeroVector);
 		const FVector3f Velocity = FVector3f::ZeroVector;
-		const FVector3f TangentX = StaticMeshHelper.TransformVector(FVector3f(1.0f, 0.0f, 0.0f));
-		const FVector3f TangentY = StaticMeshHelper.TransformVector(FVector3f(0.0f, 1.0f, 0.0f));
-		const FVector3f TangentZ = StaticMeshHelper.TransformVector(FVector3f(0.0f, 0.0f, 1.0f));
+		const FVector3f TangentX = StaticMeshHelper.TransformUnitVector(FVector3f(1.0f, 0.0f, 0.0f));
+		const FVector3f TangentY = StaticMeshHelper.TransformUnitVector(FVector3f(0.0f, 1.0f, 0.0f));
+		const FVector3f TangentZ = StaticMeshHelper.TransformUnitVector(FVector3f(0.0f, 0.0f, 1.0f));
 		for (int32 i = 0; i < Context.GetNumInstances(); ++i)
 		{
 			OutPositionParam.SetAndAdvance(Position);
@@ -4250,9 +4262,9 @@ void UNiagaraDataInterfaceStaticMesh::VMGetTriangleInterpolated(FVectorVMExterna
 		const FVector3f Position = StaticMeshHelper.TransformPosition(FVector3f::ZeroVector);
 		const FVector3f PreviousPosition = StaticMeshHelper.PreviousTransformPosition(FVector3f::ZeroVector);
 		const FVector3f Velocity = FVector3f::ZeroVector;
-		const FVector3f TangentX = StaticMeshHelper.TransformVector(FVector3f(1.0f, 0.0f, 0.0f));
-		const FVector3f TangentY = StaticMeshHelper.TransformVector(FVector3f(0.0f, 1.0f, 0.0f));
-		const FVector3f TangentZ = StaticMeshHelper.TransformVector(FVector3f(0.0f, 0.0f, 1.0f));
+		const FVector3f TangentX = StaticMeshHelper.TransformUnitVector(FVector3f(1.0f, 0.0f, 0.0f));
+		const FVector3f TangentY = StaticMeshHelper.TransformUnitVector(FVector3f(0.0f, 1.0f, 0.0f));
+		const FVector3f TangentZ = StaticMeshHelper.TransformUnitVector(FVector3f(0.0f, 0.0f, 1.0f));
 		for (int32 i = 0; i < Context.GetNumInstances(); ++i)
 		{
 			OutPositionParam.SetAndAdvance(Position);
@@ -4476,7 +4488,7 @@ void UNiagaraDataInterfaceStaticMesh::VMGetSocketTransform(FVectorVMExternalFunc
 			const FVector3f Velocity = (Position - PreviousPosition) * StaticMeshHelper.GetInvDeltaSeconds();
 			OutTranslateParam.SetAndAdvance(StaticMeshHelper.TransformPosition(SocketTransform.GetTranslation()));
 			OutRotateParam.SetAndAdvance(StaticMeshHelper.TransformRotation(SocketTransform.GetRotation()));
-			OutScaleParam.SetAndAdvance(StaticMeshHelper.TransformVector(SocketTransform.GetScale3D()));
+			OutScaleParam.SetAndAdvance(StaticMeshHelper.TransformNotUnitVector(SocketTransform.GetScale3D()));
 			OutVelocityParam.SetAndAdvance(Velocity);
 		}
 	}
@@ -4484,7 +4496,7 @@ void UNiagaraDataInterfaceStaticMesh::VMGetSocketTransform(FVectorVMExternalFunc
 	{
 		const FVector3f DefaultTranslate = StaticMeshHelper.TransformPosition(FVector3f::ZeroVector);
 		const FQuat4f DefaultRotation = StaticMeshHelper.TransformRotation(FQuat4f::Identity);
-		const FVector3f DefaultScale = StaticMeshHelper.TransformVector(FVector3f::OneVector);
+		const FVector3f DefaultScale = StaticMeshHelper.TransformNotUnitVector(FVector3f::OneVector);
 		const FVector3f DefaultVelocity = FVector3f::ZeroVector;
 		for (int32 i = 0; i < Context.GetNumInstances(); ++i)
 		{
@@ -4526,7 +4538,7 @@ void UNiagaraDataInterfaceStaticMesh::VMGetSocketTransformInterpolated(FVectorVM
 	{
 		const FVector3f DefaultTranslate = StaticMeshHelper.TransformPosition(FVector3f::ZeroVector);
 		const FQuat4f DefaultRotation = StaticMeshHelper.TransformRotation(FQuat4f::Identity);
-		const FVector3f DefaultScale = StaticMeshHelper.TransformVector(FVector3f::OneVector);
+		const FVector3f DefaultScale = StaticMeshHelper.TransformNotUnitVector(FVector3f::OneVector);
 		const FVector3f DefaultVelocity = FVector3f::ZeroVector;
 		for (int32 i = 0; i < Context.GetNumInstances(); ++i)
 		{
@@ -4561,7 +4573,7 @@ void UNiagaraDataInterfaceStaticMesh::VMGetFilteredSocketTransform(FVectorVMExte
 			const FVector3f Velocity = (Position - PreviousPosition) * StaticMeshHelper.GetInvDeltaSeconds();
 			OutTranslateParam.SetAndAdvance(StaticMeshHelper.TransformPosition(SocketTransform.GetTranslation()));
 			OutRotateParam.SetAndAdvance(StaticMeshHelper.TransformRotation(SocketTransform.GetRotation()));
-			OutScaleParam.SetAndAdvance(StaticMeshHelper.TransformVector(SocketTransform.GetScale3D()));
+			OutScaleParam.SetAndAdvance(StaticMeshHelper.TransformNotUnitVector(SocketTransform.GetScale3D()));
 			OutVelocityParam.SetAndAdvance(Velocity);
 		}
 	}
@@ -4569,7 +4581,7 @@ void UNiagaraDataInterfaceStaticMesh::VMGetFilteredSocketTransform(FVectorVMExte
 	{
 		const FVector3f DefaultTranslate = StaticMeshHelper.TransformPosition(FVector3f::ZeroVector);
 		const FQuat4f DefaultRotation = StaticMeshHelper.TransformRotation(FQuat4f::Identity);
-		const FVector3f DefaultScale = StaticMeshHelper.TransformVector(FVector3f::OneVector);
+		const FVector3f DefaultScale = StaticMeshHelper.TransformNotUnitVector(FVector3f::OneVector);
 		const FVector3f DefaultVelocity = FVector3f::ZeroVector;
 		for (int32 i = 0; i < Context.GetNumInstances(); ++i)
 		{
@@ -4612,7 +4624,7 @@ void UNiagaraDataInterfaceStaticMesh::VMGetFilteredSocketTransformInterpolated(F
 	{
 		const FVector3f DefaultTranslate = StaticMeshHelper.TransformPosition(FVector3f::ZeroVector);
 		const FQuat4f DefaultRotation = StaticMeshHelper.TransformRotation(FQuat4f::Identity);
-		const FVector3f DefaultScale = StaticMeshHelper.TransformVector(FVector3f::OneVector);
+		const FVector3f DefaultScale = StaticMeshHelper.TransformNotUnitVector(FVector3f::OneVector);
 		const FVector3f DefaultVelocity = FVector3f::ZeroVector;
 		for (int32 i = 0; i < Context.GetNumInstances(); ++i)
 		{
@@ -4648,7 +4660,7 @@ void UNiagaraDataInterfaceStaticMesh::VMGetUnfilteredSocketTransform(FVectorVMEx
 			const FVector3f Velocity = (Position - PreviousPosition) * StaticMeshHelper.GetInvDeltaSeconds();
 			OutTranslateParam.SetAndAdvance(StaticMeshHelper.TransformPosition(SocketTransform.GetTranslation()));
 			OutRotateParam.SetAndAdvance(StaticMeshHelper.TransformRotation(SocketTransform.GetRotation()));
-			OutScaleParam.SetAndAdvance(StaticMeshHelper.TransformVector(SocketTransform.GetScale3D()));
+			OutScaleParam.SetAndAdvance(StaticMeshHelper.TransformNotUnitVector(SocketTransform.GetScale3D()));
 			OutVelocityParam.SetAndAdvance(Velocity);
 		}
 	}
@@ -4656,7 +4668,7 @@ void UNiagaraDataInterfaceStaticMesh::VMGetUnfilteredSocketTransform(FVectorVMEx
 	{
 		const FVector3f DefaultTranslate = StaticMeshHelper.TransformPosition(FVector3f::ZeroVector);
 		const FQuat4f DefaultRotation = StaticMeshHelper.TransformRotation(FQuat4f::Identity);
-		const FVector3f DefaultScale = StaticMeshHelper.TransformVector(FVector3f::OneVector);
+		const FVector3f DefaultScale = StaticMeshHelper.TransformNotUnitVector(FVector3f::OneVector);
 		const FVector3f DefaultVelocity = FVector3f::ZeroVector;
 		for (int32 i = 0; i < Context.GetNumInstances(); ++i)
 		{
@@ -4700,7 +4712,7 @@ void UNiagaraDataInterfaceStaticMesh::VMGetUnfilteredSocketTransformInterpolated
 	{
 		const FVector3f DefaultTranslate = StaticMeshHelper.TransformPosition(FVector3f::ZeroVector);
 		const FQuat4f DefaultRotation = StaticMeshHelper.TransformRotation(FQuat4f::Identity);
-		const FVector3f DefaultScale = StaticMeshHelper.TransformVector(FVector3f::OneVector);
+		const FVector3f DefaultScale = StaticMeshHelper.TransformNotUnitVector(FVector3f::OneVector);
 		const FVector3f DefaultVelocity = FVector3f::ZeroVector;
 		for (int32 i = 0; i < Context.GetNumInstances(); ++i)
 		{
@@ -5360,9 +5372,9 @@ void UNiagaraDataInterfaceStaticMesh::VMGetTriangleTangentBasis_Deprecated(FVect
 	}
 	else
 	{
-		const FVector3f TangentX = StaticMeshHelper.TransformVector(FVector3f(1.0f, 0.0f, 0.0f));
-		const FVector3f TangentY = StaticMeshHelper.TransformVector(FVector3f(0.0f, 1.0f, 0.0f));
-		const FVector3f TangentZ = StaticMeshHelper.TransformVector(FVector3f(0.0f, 0.0f, 1.0f));
+		const FVector3f TangentX = StaticMeshHelper.TransformUnitVector(FVector3f(1.0f, 0.0f, 0.0f));
+		const FVector3f TangentY = StaticMeshHelper.TransformUnitVector(FVector3f(0.0f, 1.0f, 0.0f));
+		const FVector3f TangentZ = StaticMeshHelper.TransformUnitVector(FVector3f(0.0f, 0.0f, 1.0f));
 		for (int32 i = 0; i < Context.GetNumInstances(); ++i)
 		{
 			OutNormal.SetAndAdvance(TangentZ);
@@ -5394,7 +5406,7 @@ void UNiagaraDataInterfaceStaticMesh::VMGetTriangleNormal_Deprecated(FVectorVMEx
 	}
 	else
 	{
-		const FVector3f TangentZ = StaticMeshHelper.TransformVector(FVector3f(0.0f, 0.0f, 1.0f));
+		const FVector3f TangentZ = StaticMeshHelper.TransformUnitVector(FVector3f(0.0f, 0.0f, 1.0f));
 		for (int32 i = 0; i < Context.GetNumInstances(); ++i)
 		{
 			OutNormal.SetAndAdvance(TangentZ);
diff --git a/Engine/Plugins/FX/Niagara/Source/Niagara/Private/NDISkeletalMesh_BoneSampling.cpp b/Engine/Plugins/FX/Niagara/Source/Niagara/Private/NDISkeletalMesh_BoneSampling.cpp
--- a/Engine/Plugins/FX/Niagara/Source/Niagara/Private/NDISkeletalMesh_BoneSampling.cpp
+++ b/Engine/Plugins/FX/Niagara/Source/Niagara/Private/NDISkeletalMesh_BoneSampling.cpp
@@ -762,11 +762,11 @@ void UNiagaraDataInterfaceSkeletalMesh::GetSkinnedBoneData(FVectorVMExternalFunc
 				if (Output.bNeedsScale)
 				{
 					FVector3f Scale = CurrSocketTransform.GetScale3D();
-					TransformHandler.TransformVector(Scale, InstanceTransform);
+					TransformHandler.TransformNotUnitVector(Scale, InstanceTransform);
 					if (bInterpolated::Value)
 					{
 						FVector3f PrevScale = PrevSocketTransform.GetScale3D();
-						TransformHandler.TransformVector(PrevScale, PrevInstanceTransform);
+						TransformHandler.TransformNotUnitVector(PrevScale, PrevInstanceTransform);
 						Scale = FMath::Lerp(PrevScale, Scale, Interp);
 					}
 					Output.Scale.SetAndAdvance(Scale);
diff --git a/Engine/Plugins/FX/Niagara/Source/Niagara/Private/NDISkeletalMesh_TriangleSampling.cpp b/Engine/Plugins/FX/Niagara/Source/Niagara/Private/NDISkeletalMesh_TriangleSampling.cpp
--- a/Engine/Plugins/FX/Niagara/Source/Niagara/Private/NDISkeletalMesh_TriangleSampling.cpp
+++ b/Engine/Plugins/FX/Niagara/Source/Niagara/Private/NDISkeletalMesh_TriangleSampling.cpp
@@ -1452,19 +1452,19 @@ void UNiagaraDataInterfaceSkeletalMesh::GetTriCoordSkinnedData(FVectorVMExternal
 
 			if (Output.bNeedsNorm)
 			{
-				TransformHandler.TransformVector(TangentZ, Transform);
+				TransformHandler.TransformUnitVector(TangentZ, Transform);
 				Output.Normal.SetAndAdvance(TangentZ.GetSafeNormal());
 			}
 
 			if (Output.bNeedsBinorm)
 			{
-				TransformHandler.TransformVector(TangentY, Transform);
+				TransformHandler.TransformUnitVector(TangentY, Transform);
 				Output.Binormal.SetAndAdvance(TangentY.GetSafeNormal());
 			}
 
 			if (Output.bNeedsTangent)
 			{
-				TransformHandler.TransformVector(TangentX, Transform);
+				TransformHandler.TransformUnitVector(TangentX, Transform);
 				Output.Tangent.SetAndAdvance(TangentX.GetSafeNormal());
 			}
 		}
diff --git a/Engine/Plugins/FX/Niagara/Source/Niagara/Private/NDISkeletalMesh_VertexSampling.cpp b/Engine/Plugins/FX/Niagara/Source/Niagara/Private/NDISkeletalMesh_VertexSampling.cpp
--- a/Engine/Plugins/FX/Niagara/Source/Niagara/Private/NDISkeletalMesh_VertexSampling.cpp
+++ b/Engine/Plugins/FX/Niagara/Source/Niagara/Private/NDISkeletalMesh_VertexSampling.cpp
@@ -806,19 +806,19 @@ void UNiagaraDataInterfaceSkeletalMesh::GetVertexSkinnedData(FVectorVMExternalFu
 
 					if (Output.bNeedsTangentX)
 					{
-						TransformHandler.TransformVector(TangentX, Transform);
+						TransformHandler.TransformUnitVector(TangentX, Transform);
 						Output.TangentX.SetAndAdvance(TangentX);
 					}
 
 					if (Output.bNeedsTangentY)
 					{
-						TransformHandler.TransformVector(TangentY, Transform);
+						TransformHandler.TransformUnitVector(TangentY, Transform);
 						Output.TangentY.SetAndAdvance(TangentY);
 					}
 
 					if (Output.bNeedsTangentZ)
 					{
-						TransformHandler.TransformVector(TangentZ, Transform);
+						TransformHandler.TransformUnitVector(TangentZ, Transform);
 						Output.TangentZ.SetAndAdvance(TangentZ);
 					}
 				}
diff --git a/Engine/Plugins/FX/Niagara/Source/Niagara/Private/NiagaraDataInterfaceSpline.cpp b/Engine/Plugins/FX/Niagara/Source/Niagara/Private/NiagaraDataInterfaceSpline.cpp
--- a/Engine/Plugins/FX/Niagara/Source/Niagara/Private/NiagaraDataInterfaceSpline.cpp
+++ b/Engine/Plugins/FX/Niagara/Source/Niagara/Private/NiagaraDataInterfaceSpline.cpp
@@ -1198,12 +1198,12 @@ void UNiagaraDataInterfaceSpline::SampleSplineUpVectorByUnitDistance(FVectorVMEx
 		{
 			float DistanceUnitDistance = SplineSampleParam.Get();
 
-			FVector Pos = InstData->GetUpVectorAtDistanceAlongSpline<UseLUT>(DistanceUnitDistance * SplineLength, ESplineCoordinateSpace::Local);
-			TransformHandler.TransformVector(Pos, InstData->Transform);
+			FVector Vec = InstData->GetUpVectorAtDistanceAlongSpline<UseLUT>(DistanceUnitDistance * SplineLength, ESplineCoordinateSpace::Local);
+			TransformHandler.TransformUnitVector(Vec, InstData->Transform);
 
-			*OutPosX.GetDest() = float(Pos.X);		// LWC Precision Loss
-			*OutPosY.GetDest() = float(Pos.Y);
-			*OutPosZ.GetDest() = float(Pos.Z);
+			*OutPosX.GetDest() = float(Vec.X);		// LWC Precision Loss
+			*OutPosY.GetDest() = float(Vec.Y);
+			*OutPosZ.GetDest() = float(Vec.Z);
 			SplineSampleParam.Advance();
 			OutPosX.Advance();
 			OutPosY.Advance();
@@ -1214,12 +1214,12 @@ void UNiagaraDataInterfaceSpline::SampleSplineUpVectorByUnitDistance(FVectorVMEx
 	{
 		for (int32 i = 0; i < Context.GetNumInstances(); ++i)
 		{
-			FVector Pos = FVector(0.0f, 0.0f, 1.0f); 
-			TransformHandler.TransformVector(Pos, InstData->Transform);
+			FVector Vec = FVector(0.0f, 0.0f, 1.0f); 
+			TransformHandler.TransformUnitVector(Vec, InstData->Transform);
 
-			*OutPosX.GetDest() = float(Pos.X);		// LWC Precision Loss
-			*OutPosY.GetDest() = float(Pos.Y);
-			*OutPosZ.GetDest() = float(Pos.Z);
+			*OutPosX.GetDest() = float(Vec.X);		// LWC Precision Loss
+			*OutPosY.GetDest() = float(Vec.Y);
+			*OutPosZ.GetDest() = float(Vec.Z);
 			SplineSampleParam.Advance();
 			OutPosX.Advance();
 			OutPosY.Advance();
@@ -1246,12 +1246,12 @@ void UNiagaraDataInterfaceSpline::SampleSplineRightVectorByUnitDistance(FVectorV
 		{
 			float DistanceUnitDistance = SplineSampleParam.Get();
 
-			FVector Pos = InstData->GetRightVectorAtDistanceAlongSpline<UseLUT>(DistanceUnitDistance * SplineLength, ESplineCoordinateSpace::Local);
-			TransformHandler.TransformVector(Pos, InstData->Transform);
+			FVector Vec = InstData->GetRightVectorAtDistanceAlongSpline<UseLUT>(DistanceUnitDistance * SplineLength, ESplineCoordinateSpace::Local);
+			TransformHandler.TransformUnitVector(Vec, InstData->Transform);
 
-			*OutPosX.GetDest() = float(Pos.X);	// LWC Precision Loss
-			*OutPosY.GetDest() = float(Pos.Y);
-			*OutPosZ.GetDest() = float(Pos.Z);
+			*OutPosX.GetDest() = float(Vec.X);	// LWC Precision Loss
+			*OutPosY.GetDest() = float(Vec.Y);
+			*OutPosZ.GetDest() = float(Vec.Z);
 			SplineSampleParam.Advance();
 			OutPosX.Advance();
 			OutPosY.Advance();
@@ -1262,12 +1262,12 @@ void UNiagaraDataInterfaceSpline::SampleSplineRightVectorByUnitDistance(FVectorV
 	{
 		for (int32 i = 0; i < Context.GetNumInstances(); ++i)
 		{
-			FVector Pos = FVector(-1.0f, 0.0f, 0.0f); 
-			TransformHandler.TransformVector(Pos, InstData->Transform);
+			FVector Vec = FVector(-1.0f, 0.0f, 0.0f); 
+			TransformHandler.TransformUnitVector(Vec, InstData->Transform);
 
-			*OutPosX.GetDest() = float(Pos.X);	// LWC Precision Loss
-			*OutPosY.GetDest() = float(Pos.Y);
-			*OutPosZ.GetDest() = float(Pos.Z);
+			*OutPosX.GetDest() = float(Vec.X);	// LWC Precision Loss
+			*OutPosY.GetDest() = float(Vec.Y);
+			*OutPosZ.GetDest() = float(Vec.Z);
 			SplineSampleParam.Advance();
 			OutPosX.Advance();
 			OutPosY.Advance();
@@ -1293,12 +1293,12 @@ void UNiagaraDataInterfaceSpline::SampleSplineTangentByUnitDistance(FVectorVMExt
 		{
 			float DistanceUnitDistance = SplineSampleParam.Get();
 
-			FVector Pos = InstData->GetTangentAtDistanceAlongSpline<UseLUT>(DistanceUnitDistance * SplineLength, ESplineCoordinateSpace::Local);
-			TransformHandler.TransformVector(Pos, InstData->Transform);
+			FVector Vec = InstData->GetTangentAtDistanceAlongSpline<UseLUT>(DistanceUnitDistance * SplineLength, ESplineCoordinateSpace::Local);
+			TransformHandler.TransformUnitVector(Vec, InstData->Transform);
 
-			*OutPosX.GetDest() = float(Pos.X);	// LWC Precision Loss
-			*OutPosY.GetDest() = float(Pos.Y);
-			*OutPosZ.GetDest() = float(Pos.Z);
+			*OutPosX.GetDest() = float(Vec.X);	// LWC Precision Loss
+			*OutPosY.GetDest() = float(Vec.Y);
+			*OutPosZ.GetDest() = float(Vec.Z);
 			SplineSampleParam.Advance();
 			OutPosX.Advance();
 			OutPosY.Advance();
@@ -1307,14 +1307,13 @@ void UNiagaraDataInterfaceSpline::SampleSplineTangentByUnitDistance(FVectorVMExt
 	}
 	else
 	{
+		// Note if GetTangentAtDistanceAlongSpline fails it will return FVector::ZeroVector
+		const FVector3f Vec = FVector3f::ZeroVector;
 		for (int32 i = 0; i < Context.GetNumInstances(); ++i)
 		{
-			FVector Pos = FVector(EForceInit::ForceInitToZero); 
-			TransformHandler.TransformVector(Pos, InstData->Transform);
-
-			*OutPosX.GetDest() = float(Pos.X);	// LWC Precision Loss
-			*OutPosY.GetDest() = float(Pos.Y);
-			*OutPosZ.GetDest() = float(Pos.Z);
+			*OutPosX.GetDest() = Vec.X;
+			*OutPosY.GetDest() = Vec.Y;
+			*OutPosZ.GetDest() = Vec.Z;
 			SplineSampleParam.Advance();
 			OutPosX.Advance();
 			OutPosY.Advance();
@@ -1340,12 +1339,12 @@ void UNiagaraDataInterfaceSpline::SampleSplineDirectionByUnitDistance(FVectorVME
 		{
 			float DistanceUnitDistance = SplineSampleParam.Get();
 
-			FVector Pos = InstData->GetDirectionAtDistanceAlongSpline<UseLUT>(DistanceUnitDistance * SplineLength, ESplineCoordinateSpace::Local);
-			TransformHandler.TransformVector(Pos, InstData->Transform);
+			FVector Vec = InstData->GetDirectionAtDistanceAlongSpline<UseLUT>(DistanceUnitDistance * SplineLength, ESplineCoordinateSpace::Local);
+			TransformHandler.TransformUnitVector(Vec, InstData->Transform);
 
-			*OutPosX.GetDest() = float(Pos.X);		// LWC Precision Loss
-			*OutPosY.GetDest() = float(Pos.Y);
-			*OutPosZ.GetDest() = float(Pos.Z);
+			*OutPosX.GetDest() = float(Vec.X);		// LWC Precision Loss
+			*OutPosY.GetDest() = float(Vec.Y);
+			*OutPosZ.GetDest() = float(Vec.Z);
 			SplineSampleParam.Advance();
 			OutPosX.Advance();
 			OutPosY.Advance();
@@ -1356,12 +1355,12 @@ void UNiagaraDataInterfaceSpline::SampleSplineDirectionByUnitDistance(FVectorVME
 	{
 		for (int32 i = 0; i < Context.GetNumInstances(); ++i)
 		{
-			FVector Pos = FVector(0.0f, 1.0f, 0.0f); 
-			TransformHandler.TransformVector(Pos, InstData->Transform);
+			FVector Vec = FVector(0.0f, 1.0f, 0.0f);
+			TransformHandler.TransformUnitVector(Vec, InstData->Transform);
 
-			*OutPosX.GetDest() = float(Pos.X);	// LWC Precision Loss
-			*OutPosY.GetDest() = float(Pos.Y);
-			*OutPosZ.GetDest() = float(Pos.Z);
+			*OutPosX.GetDest() = float(Vec.X);	// LWC Precision Loss
+			*OutPosY.GetDest() = float(Vec.Y);
+			*OutPosZ.GetDest() = float(Vec.Z);
 			SplineSampleParam.Advance();
 			OutPosX.Advance();
 			OutPosY.Advance();
