fix error reporting

diff --git a/Engine/Plugins/FX/Niagara/Source/NiagaraEditor/Private/NiagaraBakerRendererOutputSparseVolumeTexture.cpp b/Engine/Plugins/FX/Niagara/Source/NiagaraEditor/Private/NiagaraBakerRendererOutputSparseVolumeTexture.cpp
--- a/Engine/Plugins/FX/Niagara/Source/NiagaraEditor/Private/NiagaraBakerRendererOutputSparseVolumeTexture.cpp
+++ b/Engine/Plugins/FX/Niagara/Source/NiagaraEditor/Private/NiagaraBakerRendererOutputSparseVolumeTexture.cpp
@@ -151,7 +151,8 @@ bool FNiagaraBakerRendererOutputSparseVolumeTexture::BeginBake(FNiagaraBakerFeed
 	
 	if (!SVTAsset->BeginInitialize(1))
 	{
-		UE_LOG(LogNiagaraBaker, Error, TEXT("something isnt fun"));
+		UE_LOG(LogNiagaraBaker, Error, TEXT("Cannot initialize SVT for baking"));
+		return false;
 	}
 
 	return true;
