Fix for Niagara's AsyncTask compilation mode, pushing processing DDC into ExeData to the gamethread to prevent conflicts with GC/async loading

diff --git a/Engine/Plugins/FX/Niagara/Source/NiagaraEditor/Private/NiagaraCompilationTasks.cpp b/Engine/Plugins/FX/Niagara/Source/NiagaraEditor/Private/NiagaraCompilationTasks.cpp
--- a/Engine/Plugins/FX/Niagara/Source/NiagaraEditor/Private/NiagaraCompilationTasks.cpp
+++ b/Engine/Plugins/FX/Niagara/Source/NiagaraEditor/Private/NiagaraCompilationTasks.cpp
@@ -13,6 +13,7 @@
 #include "NiagaraShader.h"
 #include "NiagaraSimulationStageBase.h"
 #include "NiagaraSystem.h"
+#include "NiagaraSystemCompilingManager.h"
 #include "Serialization/ObjectAndNameAsStringProxyArchive.h"
 
 // needed for AsType...pretty ugly
@@ -774,27 +775,41 @@ void FNiagaraSystemCompilationTask::FDispatchAndProcessDataCacheGetRequests::Lau
 				FNiagaraSystemCompilationTask::FCompileTaskInfo& CompileTask = SystemCompileTask->CompileTasks[Response.UserData];
 				if (Response.Status == EStatus::Ok)
 				{
-					FSharedBuffer DDCData = Response.Value.GetData().Decompress();
-					FNiagaraVMExecutableData ExeData;
-					if (NiagaraCompilationCopyImpl::BinaryToExecData(CompileTask.SourceScript.Get(), DDCData, ExeData))
-					{
-						CompileTask.ExeData = MakeShared<FNiagaraVMExecutableData>(MoveTemp(ExeData));
-					}
-				}
-
-				if (!CompileTask.ExeData.IsValid())
-				{
-					CompileTask.DataCachePutKeys.Add(CompileTask.DataCacheGetKey);
-				}
-				else
-				{
-					CompileTask.bFromDerivedDataCache = true;
+					CompileTask.PendingDDCData = Response.Value.GetData().Decompress().MoveToUnique();
 				}
 			}
 
 			if (PendingGetRequestCount.fetch_sub(1, std::memory_order_relaxed) == 1)
 			{
-				CompletionEvent.Trigger();
+				// in order to make sure that processing the binary data from the DDC can properly serialize
+				// objects (based on their path name) we must run BinaryToExecData() on the game thread (to
+				// avoid conflicts with GC or async loading).
+				FNiagaraSystemCompilingManager::Get().QueueGameThreadFunction([CompletionEvent = this->CompletionEvent, SystemCompileTask]() mutable
+				{
+					check(IsInGameThread());
+					for (FNiagaraSystemCompilationTask::FCompileTaskInfo& CompileTask : SystemCompileTask->CompileTasks)
+					{
+						if (!CompileTask.PendingDDCData.IsNull())
+						{
+							FNiagaraVMExecutableData ExeData;
+							if (NiagaraCompilationCopyImpl::BinaryToExecData(CompileTask.SourceScript.Get(), CompileTask.PendingDDCData.MoveToShared(), ExeData))
+							{
+								CompileTask.ExeData = MakeShared<FNiagaraVMExecutableData>(MoveTemp(ExeData));
+							}
+
+							if (!CompileTask.ExeData.IsValid())
+							{
+								CompileTask.DataCachePutKeys.Add(CompileTask.DataCacheGetKey);
+							}
+							else
+							{
+								CompileTask.bFromDerivedDataCache = true;
+							}
+						}
+					}
+
+					CompletionEvent.Trigger();
+				});
 			}
 		});
 	}
diff --git a/Engine/Plugins/FX/Niagara/Source/NiagaraEditor/Private/NiagaraCompilationTasks.h b/Engine/Plugins/FX/Niagara/Source/NiagaraEditor/Private/NiagaraCompilationTasks.h
--- a/Engine/Plugins/FX/Niagara/Source/NiagaraEditor/Private/NiagaraCompilationTasks.h
+++ b/Engine/Plugins/FX/Niagara/Source/NiagaraEditor/Private/NiagaraCompilationTasks.h
@@ -126,6 +126,7 @@ struct FNiagaraSystemCompilationTask
 		FNiagaraTranslatorOutput TranslateOutput;
 		FString TranslatedHlsl;
 		TSharedPtr<FHlslNiagaraCompiler> Compiler;
+		FUniqueBuffer PendingDDCData;
 		TSharedPtr<FNiagaraVMExecutableData> ExeData;
 		TWeakObjectPtr<UNiagaraScript> SourceScript;
 		TObjectKey<UNiagaraScript> ScriptKey;
diff --git a/Engine/Plugins/FX/Niagara/Source/NiagaraEditor/Private/NiagaraSystemCompilingManager.cpp b/Engine/Plugins/FX/Niagara/Source/NiagaraEditor/Private/NiagaraSystemCompilingManager.cpp
--- a/Engine/Plugins/FX/Niagara/Source/NiagaraEditor/Private/NiagaraSystemCompilingManager.cpp
+++ b/Engine/Plugins/FX/Niagara/Source/NiagaraEditor/Private/NiagaraSystemCompilingManager.cpp
@@ -66,6 +66,21 @@ void FNiagaraSystemCompilingManager::Shutdown()
 
 void FNiagaraSystemCompilingManager::ProcessAsyncTasks(bool bLimitExecutionTime)
 {
+	// process any pending GameThreadTasks
+	{
+		TArray<FGameThreadFunction> PendingFunctions;
+
+		{
+			FWriteScopeLock Write(GameThreadFunctionLock);
+			PendingFunctions = MoveTemp(GameThreadFunctions);
+		}
+
+		for (FGameThreadFunction& PendingFunction : PendingFunctions)
+		{
+			PendingFunction();
+		}
+	}
+
 	{
 		FReadScopeLock ReadScope(QueueLock);
 		if (ActiveTasks.IsEmpty())
@@ -295,6 +310,13 @@ void FNiagaraSystemCompilingManager::AbortSystemCompile(FNiagaraCompilationTaskH
 	}
 }
 
+void FNiagaraSystemCompilingManager::QueueGameThreadFunction(FGameThreadFunction GameThreadTask)
+{
+	FWriteScopeLock Write(GameThreadFunctionLock);
+
+	GameThreadFunctions.Add(GameThreadTask);
+}
+
 bool FNiagaraSystemCompilingManager::ConditionalLaunchTask()
 {
 	{
diff --git a/Engine/Plugins/FX/Niagara/Source/NiagaraEditor/Private/NiagaraSystemCompilingManager.h b/Engine/Plugins/FX/Niagara/Source/NiagaraEditor/Private/NiagaraSystemCompilingManager.h
--- a/Engine/Plugins/FX/Niagara/Source/NiagaraEditor/Private/NiagaraSystemCompilingManager.h
+++ b/Engine/Plugins/FX/Niagara/Source/NiagaraEditor/Private/NiagaraSystemCompilingManager.h
@@ -25,6 +25,9 @@ public:
 	bool PollSystemCompile(FNiagaraCompilationTaskHandle TaskHandle, bool bPeek, bool bWait, FNiagaraSystemAsyncCompileResults& Results);
 	void AbortSystemCompile(FNiagaraCompilationTaskHandle TaskHandle);
 
+	using FGameThreadFunction = TFunction<void()>;
+	void QueueGameThreadFunction(FGameThreadFunction GameThreadTask);
+
 protected:
 	// Begin - IAssetCompilingManager
 	virtual FName GetAssetTypeName() const override;
@@ -48,4 +51,7 @@ protected:
 	TArray<FNiagaraCompilationTaskHandle> RequestsAwaitingRetrieval;
 
 	std::atomic<FNiagaraCompilationTaskHandle> NextTaskHandle = { 0 };
+
+	mutable FRWLock GameThreadFunctionLock;
+	TArray<FGameThreadFunction> GameThreadFunctions;
 };
