- Fix for clearing counts as int vs uint

diff --git a/Engine/Plugins/FX/Niagara/Source/Niagara/Private/DataInterface/NiagaraDataInterfaceSimpleCounter.cpp b/Engine/Plugins/FX/Niagara/Source/Niagara/Private/DataInterface/NiagaraDataInterfaceSimpleCounter.cpp
--- a/Engine/Plugins/FX/Niagara/Source/Niagara/Private/DataInterface/NiagaraDataInterfaceSimpleCounter.cpp
+++ b/Engine/Plugins/FX/Niagara/Source/Niagara/Private/DataInterface/NiagaraDataInterfaceSimpleCounter.cpp
@@ -79,9 +79,9 @@ struct FNDISimpleCounterProxy : public FNiagaraDataInterfaceProxyRW
 					RDG_EVENT_NAME("NiagaraSimpleCounter::PreStage"),
 					[CountBufferUAV=CountBuffer.UAV, CountOffset=InstanceData->CountOffset, CountValue=InstanceData->CountValue.GetValue()](FRHICommandListImmediate& RHICmdList)
 					{
-						const TPair<uint32, int32> DataToClear(CountOffset, CountValue);
+						const TPair<uint32, uint32> DataToClear(CountOffset, reinterpret_cast<const uint32&>(CountValue));
 						RHICmdList.Transition(FRHITransitionInfo(CountBufferUAV, ERHIAccess::UAVCompute, ERHIAccess::UAVCompute));
-						NiagaraClearCounts::ClearCountsInt(RHICmdList, CountBufferUAV, MakeArrayView(&DataToClear, 1) );
+						NiagaraClearCounts::ClearCountsUInt(RHICmdList, CountBufferUAV, MakeArrayView(&DataToClear, 1) );
 						RHICmdList.Transition(FRHITransitionInfo(CountBufferUAV, ERHIAccess::UAVCompute, ERHIAccess::UAVCompute));
 					}
 				);
