Fixed mentions of 'behavior examples' -> 'learning examples'

diff --git a/Engine/Plugins/FX/Niagara/Source/NiagaraEditor/Private/SNewEmitterDialog.cpp b/Engine/Plugins/FX/Niagara/Source/NiagaraEditor/Private/SNewEmitterDialog.cpp
--- a/Engine/Plugins/FX/Niagara/Source/NiagaraEditor/Private/SNewEmitterDialog.cpp
+++ b/Engine/Plugins/FX/Niagara/Source/NiagaraEditor/Private/SNewEmitterDialog.cpp
@@ -44,7 +44,7 @@ void SNewEmitterDialog::Construct(const FArguments& InArgs)
 		{
 			SNiagaraNewAssetDialog::FNiagaraNewAssetDialogOption(
 				LOCTEXT("CreateFromEmitterLabel", "New emitter"),
-				LOCTEXT("CreateFromEmitterDescription", "Create a new emitter from a template or behavior emitter (no inheritance) or from a parent (inheritance)"),
+				LOCTEXT("CreateFromEmitterDescription", "Create a new emitter from a template or learning emitter (no inheritance) or from a parent (inheritance)"),
 				LOCTEXT("EmitterPickerHeader", "Select an Emitter"),
 				SNiagaraNewAssetDialog::FOnGetSelectedAssetsFromPicker::CreateSP(this, &SNewEmitterDialog::GetSelectedEmitterAssets_NewAssets),
 				SNiagaraNewAssetDialog::FOnSelectionConfirmed::CreateSP(this, &SNewEmitterDialog::CheckUseInheritance),
diff --git a/Engine/Plugins/FX/Niagara/Source/NiagaraEditor/Private/SNewSystemDialog.cpp b/Engine/Plugins/FX/Niagara/Source/NiagaraEditor/Private/SNewSystemDialog.cpp
--- a/Engine/Plugins/FX/Niagara/Source/NiagaraEditor/Private/SNewSystemDialog.cpp
+++ b/Engine/Plugins/FX/Niagara/Source/NiagaraEditor/Private/SNewSystemDialog.cpp
@@ -63,7 +63,7 @@ void SNewSystemDialog::Construct(const FArguments& InArgs)
 		{
 			SNiagaraNewAssetDialog::FNiagaraNewAssetDialogOption(
 				LOCTEXT("CreateFromSelectedEmittersLabel", "New system from selected emitter(s)"),
-				LOCTEXT("CreateFromSelectedEmittersDescription", "Choose a mix of emitters (inherited) and emitter templates/behavior examples (no inheritance)"),
+				LOCTEXT("CreateFromSelectedEmittersDescription", "Choose a mix of emitters (inherited) and emitter templates/learning examples (no inheritance)"),
 				LOCTEXT("ProjectEmittersLabel", "Select Emitters to Add"),
 				SNiagaraNewAssetDialog::FOnGetSelectedAssetsFromPicker::CreateSP(this, &SNewSystemDialog::GetSelectedProjectEmiterAssets),
 				SNiagaraNewAssetDialog::FOnSelectionConfirmed(),
@@ -120,8 +120,8 @@ void SNewSystemDialog::Construct(const FArguments& InArgs)
 				EmitterAssetPicker->GetSearchBox()
 				),
 			SNiagaraNewAssetDialog::FNiagaraNewAssetDialogOption(
-				LOCTEXT("CreateFromTemplateLabel", "New system from a template or behavior example"),
-				LOCTEXT("CreateFromTemplateDescription", "The new system will be derived from a system template or behavior example"),
+				LOCTEXT("CreateFromTemplateLabel", "New system from a template or learning example"),
+				LOCTEXT("CreateFromTemplateDescription", "The new system will be derived from a system template or learning example"),
 				LOCTEXT("TemplateLabel", "Select a System Template"),
 				SNiagaraNewAssetDialog::FOnGetSelectedAssetsFromPicker::CreateSP(this, &SNewSystemDialog::GetSelectedSystemTemplateAssets),
 				SNiagaraNewAssetDialog::FOnSelectionConfirmed(),
diff --git a/Engine/Plugins/FX/Niagara/Source/NiagaraEditor/Private/Widgets/SNiagaraFilterBox.cpp b/Engine/Plugins/FX/Niagara/Source/NiagaraEditor/Private/Widgets/SNiagaraFilterBox.cpp
--- a/Engine/Plugins/FX/Niagara/Source/NiagaraEditor/Private/Widgets/SNiagaraFilterBox.cpp
+++ b/Engine/Plugins/FX/Niagara/Source/NiagaraEditor/Private/Widgets/SNiagaraFilterBox.cpp
@@ -444,9 +444,9 @@ FText SNiagaraTemplateTabBox::DetermineControlTooltip(ENiagaraScriptTemplateSpec
 				return LOCTEXT("EmitterTemplateTabTooltip", "Templates are intended as starting points for building functional emitters of different types,\n"
 			"Emitter templates are copied into a system as a unique emitter with no inheritance");
 			case ENiagaraScriptTemplateSpecification::Behavior:
-				return LOCTEXT("EmitterBehaviorExampleTabTooltip", "Behavior Examples are intended to serve as a guide to how Niagara works at a feature level.\n"
+				return LOCTEXT("EmitterBehaviorExampleTabTooltip", "Learning Examples are intended to serve as a guide to how Niagara works at a feature level.\n"
 			"Each example shows a simplified setup used to achieve specific outcomes and are intended as starting points, building blocks, or simply as reference.\n"
-			"Behavior Example emitters are copied into a system as a unique emitter with no inheritance");
+			"Learning Example emitters are copied into a system as a unique emitter with no inheritance");
 			default:
 				return FText::GetEmpty();
 			}
@@ -460,7 +460,7 @@ FText SNiagaraTemplateTabBox::DetermineControlTooltip(ENiagaraScriptTemplateSpec
 			case ENiagaraScriptTemplateSpecification::Template:
 				return LOCTEXT("SystemTemplateTabTooltip", "Template systems ");
 			case ENiagaraScriptTemplateSpecification::Behavior:
-				return LOCTEXT("SystemBehaviorExampleTabTooltip", "Behavior Example systems serve to demonstrate different Niagara features that rely on multiple emitters interacting in some way.");
+				return LOCTEXT("SystemBehaviorExampleTabTooltip", "Learning Example systems serve to demonstrate different Niagara features that rely on multiple emitters interacting in some way.");
 			default:
 				return FText::GetEmpty();				
 			}
