- Update StatID when renaming system/emitter

diff --git a/Engine/Plugins/FX/Niagara/Source/Niagara/Classes/NiagaraEmitter.h b/Engine/Plugins/FX/Niagara/Source/Niagara/Classes/NiagaraEmitter.h
--- a/Engine/Plugins/FX/Niagara/Source/Niagara/Classes/NiagaraEmitter.h
+++ b/Engine/Plugins/FX/Niagara/Source/Niagara/Classes/NiagaraEmitter.h
@@ -978,7 +978,11 @@ private:
 	FOnEventHandlersChanged OnEventHandlersChangedDelegate;
 #endif
 
-	void GenerateStatID()const;
+
+public:
+	void UpdateStatID() const;
+private:
+	void GenerateStatID() const;
 #if STATS
 	mutable TStatId StatID_GT;
 	mutable TStatId StatID_GT_CNC;
diff --git a/Engine/Plugins/FX/Niagara/Source/Niagara/Classes/NiagaraSystem.h b/Engine/Plugins/FX/Niagara/Source/Niagara/Classes/NiagaraSystem.h
--- a/Engine/Plugins/FX/Niagara/Source/Niagara/Classes/NiagaraSystem.h
+++ b/Engine/Plugins/FX/Niagara/Source/Niagara/Classes/NiagaraSystem.h
@@ -241,6 +241,7 @@ public:
 	NIAGARA_API PRAGMA_ENABLE_DEPRECATION_WARNINGS
 	virtual void PreSave(FObjectPreSaveContext ObjectSaveContext) override;
 #if WITH_EDITOR
+	NIAGARA_API virtual void PostRename(UObject* OldOuter, const FName OldName) override;
 	NIAGARA_API virtual void PreEditChange(FProperty* PropertyThatWillChange)override;
 	NIAGARA_API virtual void PostEditChangeProperty(struct FPropertyChangedEvent& PropertyChangedEvent) override; 
 	NIAGARA_API virtual void BeginCacheForCookedPlatformData(const ITargetPlatform *TargetPlatform) override;
@@ -977,7 +978,8 @@ public:
 	const FNiagaraSystemStaticBuffers* GetStaticBuffers() const { return StaticBuffers.Get(); }
 
 protected:
-	NIAGARA_API void GenerateStatID()const;
+	void UpdateStatID() const;
+	void GenerateStatID() const;
 #if STATS
 	mutable TStatId StatID_GT;
 	mutable TStatId StatID_GT_CNC;
diff --git a/Engine/Plugins/FX/Niagara/Source/Niagara/Private/NiagaraEmitter.cpp b/Engine/Plugins/FX/Niagara/Source/Niagara/Private/NiagaraEmitter.cpp
--- a/Engine/Plugins/FX/Niagara/Source/Niagara/Private/NiagaraEmitter.cpp
+++ b/Engine/Plugins/FX/Niagara/Source/Niagara/Private/NiagaraEmitter.cpp
@@ -1181,6 +1181,7 @@ void UNiagaraEmitter::PostRename(UObject* OldOuter, const FName OldName)
 	{
 		SetUniqueEmitterName(GetFName().GetPlainNameString());
 	}
+	UpdateStatID();
 }
 
 void UNiagaraEmitter::PostEditChangeProperty(FPropertyChangedEvent& PropertyChangedEvent)
@@ -3408,7 +3409,17 @@ void FMemoryRuntimeEstimation::Init()
 	EstimationCriticalSection = MakeShared<FCriticalSection>();
 }
 
-void UNiagaraEmitter::GenerateStatID()const
+void UNiagaraEmitter::UpdateStatID() const
+{
+#if STATS
+	if (StatID_GT.IsValidStat())
+	{
+		GenerateStatID();
+	}
+#endif
+}
+
+void UNiagaraEmitter::GenerateStatID() const
 {
 #if STATS
 	FString Name = GetOuter() ? GetOuter()->GetFName().ToString() : TEXT("");
diff --git a/Engine/Plugins/FX/Niagara/Source/Niagara/Private/NiagaraSystem.cpp b/Engine/Plugins/FX/Niagara/Source/Niagara/Private/NiagaraSystem.cpp
--- a/Engine/Plugins/FX/Niagara/Source/Niagara/Private/NiagaraSystem.cpp
+++ b/Engine/Plugins/FX/Niagara/Source/Niagara/Private/NiagaraSystem.cpp
@@ -848,6 +848,24 @@ void UNiagaraSystem::ResolveWarmupTickCount()
 }
 
 #if WITH_EDITOR
+void UNiagaraSystem::PostRename(UObject* OldOuter, const FName OldName)
+{
+	Super::PostRename(OldOuter, OldName);
+
+#if STATS
+	UpdateStatID();
+	for (const FNiagaraEmitterHandle& Handle : EmitterHandles)
+	{
+		if (UNiagaraEmitter* Emitter = Handle.GetInstance().Emitter)
+		{
+			Emitter->UpdateStatID();
+		}
+	}
+
+	// Recreate any scene proxies
+	FNiagaraSystemUpdateContext UpdateCtx(this, true);
+#endif
+}
 
 void UNiagaraSystem::PreEditChange(FProperty* PropertyThatWillChange)
 {
@@ -3378,7 +3396,17 @@ FGraphEventRef UNiagaraSystem::GetScriptOptimizationCompletionEvent()
 	return nullptr;
 }
 
-void UNiagaraSystem::GenerateStatID()const
+void UNiagaraSystem::UpdateStatID() const
+{
+#if STATS
+	if (StatID_GT.IsValidStat())
+	{
+		GenerateStatID();
+	}
+#endif
+}
+
+void UNiagaraSystem::GenerateStatID() const
 {
 #if STATS
 	StatID_GT = FDynamicStats::CreateStatId<FStatGroup_STATGROUP_NiagaraSystems>(GetPathName() + TEXT(" [GT]"));
