- Change display name for num threads override to make it more clear they are per group

diff --git a/Engine/Plugins/FX/Niagara/Source/Niagara/Public/NiagaraSimulationStageBase.h b/Engine/Plugins/FX/Niagara/Source/Niagara/Public/NiagaraSimulationStageBase.h
--- a/Engine/Plugins/FX/Niagara/Source/Niagara/Public/NiagaraSimulationStageBase.h
+++ b/Engine/Plugins/FX/Niagara/Source/Niagara/Public/NiagaraSimulationStageBase.h
@@ -122,19 +122,19 @@ public:
 	uint32 bGpuDispatchForceLinear : 1;
 
 	/** When enabled we use a custom number of threads for the dispatch. */
-	UPROPERTY(EditAnywhere, Category = "Simulation Stage")
+	UPROPERTY(EditAnywhere, Category = "Simulation Stage", meta = (DisplayName = "Override Gpu Thread Group Num Threads"))
 	uint32 bOverrideGpuDispatchNumThreads : 1;
 
 	/** Parameter binding / constant value for Num Threads X */
-	UPROPERTY(EditAnywhere, Category = "Simulation Stage")
+	UPROPERTY(EditAnywhere, Category = "Simulation Stage", meta = (DisplayName = "Gpu Thread Group Num Threads X"))
 	FNiagaraParameterBindingWithValue OverrideGpuDispatchNumThreadsX;
 
 	/** Parameter binding / constant value for Num Threads Y */
-	UPROPERTY(EditAnywhere, Category = "Simulation Stage", meta = (EditCondition = "DirectDispatchType == ENiagaraGpuDispatchType::TwoD || DirectDispatchType == ENiagaraGpuDispatchType::ThreeD"))
+	UPROPERTY(EditAnywhere, Category = "Simulation Stage", meta = (DisplayName = "Gpu Thread Group Num Threads Y", EditCondition = "DirectDispatchType == ENiagaraGpuDispatchType::TwoD || DirectDispatchType == ENiagaraGpuDispatchType::ThreeD"))
 	FNiagaraParameterBindingWithValue OverrideGpuDispatchNumThreadsY;
 
 	/** Parameter binding / constant value for Num Threads Z */
-	UPROPERTY(EditAnywhere, Category = "Simulation Stage", meta = (EditCondition = "DirectDispatchType == ENiagaraGpuDispatchType::ThreeD"))
+	UPROPERTY(EditAnywhere, Category = "Simulation Stage", meta = (DisplayName = "Gpu Thread Group Num Threads Z", EditCondition = "DirectDispatchType == ENiagaraGpuDispatchType::ThreeD"))
 	FNiagaraParameterBindingWithValue OverrideGpuDispatchNumThreadsZ;
 
 	/** Dimensions to use for dispatch. */
