Substrate - fix for a niagara shader include that should  have not been renamed (uniform buffer name has not been changed yet).

diff --git a/Engine/Plugins/FX/Niagara/Shaders/Private/NiagaraEmitterInstanceShader.usf b/Engine/Plugins/FX/Niagara/Shaders/Private/NiagaraEmitterInstanceShader.usf
--- a/Engine/Plugins/FX/Niagara/Shaders/Private/NiagaraEmitterInstanceShader.usf
+++ b/Engine/Plugins/FX/Niagara/Shaders/Private/NiagaraEmitterInstanceShader.usf
@@ -18,7 +18,7 @@ NiagaraSimulationShader.usf:
 #if GPU_SIMULATION
 	#include "/Engine/Private/Common.ush"
 	#include "/Engine/Private/DistanceField/GlobalDistanceFieldShared.ush"
-	#include "/Engine/Public/SubstratePublic.ush"
+	#include "/Engine/Generated/UniformBuffers/StrataPublic.ush"
 #else
 	const static float PI = 3.1415926535897932f;
 #endif
