Refactored note functionality.
Its now a general system for stack entries, not for specifically module items, and works via stack editor data & stack entry key
This now supports editing of existing nodes via context menu & notes can be added to all items & function inputs

- Removed 'CustomNote' from the stack issue system. Notes are now a separate system.
- Inline note support added. A 'full view' note can be converted into a small inline note that appears on the row itself.
- UI access via right click menu or direct access widgets.

diff --git a/Engine/Plugins/FX/CascadeToNiagaraConverter/Source/CascadeToNiagaraConverter/Private/NiagaraStackGraphUtilitiesAdapterLibrary.cpp b/Engine/Plugins/FX/CascadeToNiagaraConverter/Source/CascadeToNiagaraConverter/Private/NiagaraStackGraphUtilitiesAdapterLibrary.cpp
--- a/Engine/Plugins/FX/CascadeToNiagaraConverter/Source/CascadeToNiagaraConverter/Private/NiagaraStackGraphUtilitiesAdapterLibrary.cpp
+++ b/Engine/Plugins/FX/CascadeToNiagaraConverter/Source/CascadeToNiagaraConverter/Private/NiagaraStackGraphUtilitiesAdapterLibrary.cpp
@@ -1994,7 +1994,7 @@ void UNiagaraEmitterConversionContext::InternalFinalizeStackEntryAddActions()
 			UNiagaraRendererProperties* NewRendererProperties = It.Value();
 
 			UNiagaraClipboardContent* ClipboardContent = UNiagaraClipboardContent::Create();
-			ClipboardContent->Renderers.Add(NewRendererProperties);
+			ClipboardContent->Renderers.Add(UNiagaraClipboardRenderer::CreateRenderer(ClipboardContent, NewRendererProperties));
 
 			FText PasteWarning = FText();
 			UNiagaraStackItemGroup* RendererStackItemGroup = *RendererStackItemGroupPtr;
diff --git a/Engine/Plugins/FX/Niagara/Source/NiagaraEditor/Private/NiagaraClipboard.cpp b/Engine/Plugins/FX/Niagara/Source/NiagaraEditor/Private/NiagaraClipboard.cpp
--- a/Engine/Plugins/FX/Niagara/Source/NiagaraEditor/Private/NiagaraClipboard.cpp
+++ b/Engine/Plugins/FX/Niagara/Source/NiagaraEditor/Private/NiagaraClipboard.cpp
@@ -126,24 +126,48 @@ bool UNiagaraClipboardFunctionInput::CopyValuesFrom(const UNiagaraClipboardFunct
 	return true;
 }
 
-UNiagaraClipboardFunction* UNiagaraClipboardFunction::CreateScriptFunction(UObject* InOuter, FString InFunctionName, UNiagaraScript* InScript, const FGuid& InScriptVersion, const TArray<FNiagaraStackMessage> InMessages)
+const UNiagaraClipboardRenderer* UNiagaraClipboardRenderer::CreateRenderer(UObject* InOuter, UNiagaraRendererProperties* Renderer, TOptional<FNiagaraStackNoteData> StackNoteData)
+{
+	UNiagaraClipboardRenderer* NewRenderer = Cast<UNiagaraClipboardRenderer>(NewObject<UNiagaraClipboardRenderer>(InOuter));
+	NewRenderer->RendererProperties = CastChecked<UNiagaraRendererProperties>(StaticDuplicateObject(Renderer, InOuter));
+
+	if(StackNoteData.IsSet())
+	{
+		NewRenderer->StackNoteData = StackNoteData.GetValue();
+	}
+
+	return NewRenderer;
+}
+
+UNiagaraClipboardFunction* UNiagaraClipboardFunction::CreateScriptFunction(UObject* InOuter, FString InFunctionName, UNiagaraScript* InScript, const FGuid& InScriptVersion, const TOptional<FNiagaraStackNoteData>& InStackNote)
 {
 	UNiagaraClipboardFunction* NewFunction = Cast<UNiagaraClipboardFunction>(NewObject<UNiagaraClipboardFunction>(InOuter));
 	NewFunction->ScriptMode = ENiagaraClipboardFunctionScriptMode::ScriptAsset;
 	NewFunction->FunctionName = InFunctionName;
 	NewFunction->Script = InScript;
 	NewFunction->ScriptVersion = InScriptVersion;
-	NewFunction->Messages = InMessages;
+
+	if(InStackNote.IsSet())
+	{
+		NewFunction->StackNoteData = InStackNote.GetValue();
+	}
+	
 	return NewFunction;
 }
 
-UNiagaraClipboardFunction* UNiagaraClipboardFunction::CreateAssignmentFunction(UObject* InOuter, FString InFunctionName, const TArray<FNiagaraVariable>& InAssignmentTargets, const TArray<FString>& InAssignmentDefaults)
+UNiagaraClipboardFunction* UNiagaraClipboardFunction::CreateAssignmentFunction(UObject* InOuter, FString InFunctionName, const TArray<FNiagaraVariable>& InAssignmentTargets, const TArray<FString>& InAssignmentDefaults, TOptional<FNiagaraStackNoteData> InStackNoteData)
 {
 	UNiagaraClipboardFunction* NewFunction = Cast<UNiagaraClipboardFunction>(NewObject<UNiagaraClipboardFunction>(InOuter));
 	NewFunction->ScriptMode = ENiagaraClipboardFunctionScriptMode::Assignment;
 	NewFunction->FunctionName = InFunctionName;
 	NewFunction->AssignmentTargets = InAssignmentTargets;
 	NewFunction->AssignmentDefaults = InAssignmentDefaults;
+
+	if(InStackNoteData.IsSet())
+	{
+		NewFunction->StackNoteData = InStackNoteData.GetValue();
+	}
+	
 	return NewFunction;
 }
 
diff --git a/Engine/Plugins/FX/Niagara/Source/NiagaraEditor/Private/NiagaraEditorStyle.cpp b/Engine/Plugins/FX/Niagara/Source/NiagaraEditor/Private/NiagaraEditorStyle.cpp
--- a/Engine/Plugins/FX/Niagara/Source/NiagaraEditor/Private/NiagaraEditorStyle.cpp
+++ b/Engine/Plugins/FX/Niagara/Source/NiagaraEditor/Private/NiagaraEditorStyle.cpp
@@ -500,6 +500,22 @@ void FNiagaraEditorStyle::InitStackWidgets()
 	
 	Set("NiagaraEditor.Stack.NumericDropdownInput", EditableTextBoxStyle);
 	Set("NiagaraEditor.Stack.IntegerAsEnum", NormalText.SetFont(FAppStyle::Get().GetFontStyle("PropertyWindow.NormalFont")));
+
+	// Notes
+	FInlineEditableTextBlockStyle DefaultInlineEditableTextBlockStyle = FAppStyle::GetWidgetStyle<FInlineEditableTextBlockStyle>("InlineEditableTextBlockStyle");
+	
+	FTextBlockStyle NoteHeaderStyle;
+	
+	FFontOutlineSettings NoteHeaderOutlineSettings;
+	NoteHeaderOutlineSettings.OutlineSize = 1.f;
+	NoteHeaderStyle.SetFont(DEFAULT_FONT("Regular", 12.f, NoteHeaderOutlineSettings));
+	NoteHeaderStyle.SetColorAndOpacity(FSlateColor(FStyleColors::White));
+	FInlineEditableTextBlockStyle NoteHeaderInlineEditableStyle;
+	NoteHeaderInlineEditableStyle.SetTextStyle(NoteHeaderStyle);
+	NoteHeaderInlineEditableStyle.SetEditableTextBoxStyle(DefaultInlineEditableTextBlockStyle.EditableTextBoxStyle);
+
+	Set("NiagaraEditor.Stack.Note.HeaderText", NoteHeaderStyle);
+	Set("NiagaraEditor.Stack.Note.HeaderEditableText", NoteHeaderInlineEditableStyle);
 }
 
 void FNiagaraEditorStyle::InitNiagaraSequence()
diff --git a/Engine/Plugins/FX/Niagara/Source/NiagaraEditor/Private/NiagaraEditorUtilities.cpp b/Engine/Plugins/FX/Niagara/Source/NiagaraEditor/Private/NiagaraEditorUtilities.cpp
--- a/Engine/Plugins/FX/Niagara/Source/NiagaraEditor/Private/NiagaraEditorUtilities.cpp
+++ b/Engine/Plugins/FX/Niagara/Source/NiagaraEditor/Private/NiagaraEditorUtilities.cpp
@@ -44,6 +44,7 @@
 #include "ViewModels/NiagaraScratchPadUtilities.h"
 #include "ViewModels/NiagaraScriptViewModel.h"
 #include "ViewModels/Stack/NiagaraParameterHandle.h"
+#include "ViewModels/Stack/NiagaraStackNote.h"
 #include "Widgets/SNiagaraParameterName.h"
 
 #include "EdGraph/EdGraphPin.h"
@@ -75,6 +76,7 @@
 #include "ViewModels/NiagaraEmitterHandleViewModel.h"
 #include "ViewModels/NiagaraEmitterViewModel.h"
 #include "ViewModels/NiagaraParameterPanelViewModel.h"
+#include "Widgets/SToolTip.h"
 
 #define LOCTEXT_NAMESPACE "FNiagaraEditorUtilities"
 
@@ -4573,6 +4575,66 @@ UNiagaraDataInterface* FNiagaraEditorUtilities::GetResolvedRuntimeInstanceForEdi
 	return nullptr;
 }
 
+TSharedRef<SToolTip> FNiagaraEditorUtilities::Tooltips::CreateStackNoteTooltip(UNiagaraStackNote& StackNote)
+{
+	if(StackNote.GetTargetStackNoteData().IsSet() == false)
+	{
+		return SNew(SToolTip);
+	}
+
+	TSharedRef<SVerticalBox> TooltipContent = SNew(SVerticalBox)
+	+ SVerticalBox::Slot()
+	.AutoHeight()
+	.Padding(3.f)
+	[
+		SNew(SHorizontalBox)
+		+ SHorizontalBox::Slot()
+		.AutoWidth()
+		.Padding(2.f)
+		[
+			SNew(SBox)
+			.HeightOverride(16.f)
+			.WidthOverride(16.f)
+			.HAlign(HAlign_Center)
+			.VAlign(VAlign_Center)
+			[
+				SNew(SImage)
+				.Image(FNiagaraEditorStyle::Get().GetBrush("NiagaraEditor.Message.CustomNote"))
+			]
+		]
+		+ SHorizontalBox::Slot()
+		.AutoWidth()
+		.Padding(2.f)
+		[
+			SNew(STextBlock)
+			.Text(StackNote.GetTargetStackNoteData().GetValue().MessageHeader)
+			.TextStyle(&FNiagaraEditorStyle::Get().GetWidgetStyle<FTextBlockStyle>("NiagaraEditor.Stack.Note.HeaderText"))
+		]
+	];
+
+	if(StackNote.GetTargetStackNoteData().GetValue().Message.IsEmpty() == false)
+	{
+		TooltipContent->AddSlot()
+		.AutoHeight()
+		.Padding(5.f)
+		[
+			SNew(STextBlock)
+			.Text(StackNote.GetTargetStackNoteData().GetValue().Message)
+			.AutoWrapText(true)
+		];
+	}
+	
+	return SNew(SToolTip)
+	.Content()
+	[
+		SNew(SBox)
+		.MaxDesiredWidth(750.f)
+		[
+			TooltipContent
+		]
+	];
+}
+
 TMap<FGuid, TArray<FNiagaraVariableBase>> FNiagaraEditorUtilities::Scripts::Validation::ValidateScriptVariableIds(UNiagaraScript* Script, FGuid VersionGuid)
 {
 	bool bReturnResults = false;
diff --git a/Engine/Plugins/FX/Niagara/Source/NiagaraEditor/Private/NiagaraEmitterEditorData.cpp b/Engine/Plugins/FX/Niagara/Source/NiagaraEditor/Private/NiagaraEmitterEditorData.cpp
--- a/Engine/Plugins/FX/Niagara/Source/NiagaraEditor/Private/NiagaraEmitterEditorData.cpp
+++ b/Engine/Plugins/FX/Niagara/Source/NiagaraEditor/Private/NiagaraEmitterEditorData.cpp
@@ -367,6 +367,38 @@ void UNiagaraEmitterEditorData::PostLoad_TransferSummaryDataToNewFormat()
 	}
 }
 
+void UNiagaraEmitterEditorData::PostLoad_TransferModuleStackNotesToNewFormat(UObject* Owner)
+{
+	// since we lack graph context during post load, we do this workaround to find the emitter graph
+	UNiagaraEmitter* Emitter = Cast<UNiagaraEmitter>(Owner);
+	TArray<FNiagaraAssetVersion> AssetVersions = Emitter->GetAllAvailableVersions();
+	UNiagaraGraph* EmitterGraph = nullptr;
+	for(FNiagaraAssetVersion& AssetVersion : AssetVersions)
+	{
+		FVersionedNiagaraEmitterData* EmitterData = Emitter->GetEmitterData(AssetVersion.VersionGuid);
+		if(EmitterData->GetEditorData() == this)
+		{
+			EmitterGraph = Cast<UNiagaraScriptSource>(EmitterData->GraphSource)->NodeGraph;
+		}
+	}
+
+	if(EmitterGraph == nullptr)
+	{
+		return;
+	}
+
+	TArray<UNiagaraNodeFunctionCall*> FunctionCallNodes;
+	EmitterGraph->GetNodesOfClass(FunctionCallNodes);
+
+	for(UNiagaraNodeFunctionCall* FunctionCallNode : FunctionCallNodes)
+	{
+		if(FunctionCallNode->GetDeprecatedCustomNotes().Num() > 0)
+		{
+			StackEditorData->TransferDeprecatedStackNotes(*FunctionCallNode);
+		}
+	}
+}
+
 void UNiagaraEmitterEditorData::PostLoad()
 {
 	Super::PostLoad();
@@ -386,6 +418,7 @@ void UNiagaraEmitterEditorData::PostLoad()
 void UNiagaraEmitterEditorData::PostLoadFromOwner(UObject* InOwner)
 {
 	PostLoad_TransferSummaryDataToNewFormat();
+	PostLoad_TransferModuleStackNotesToNewFormat(InOwner);
 }
 
 #if WITH_EDITORONLY_DATA
diff --git a/Engine/Plugins/FX/Niagara/Source/NiagaraEditor/Private/NiagaraMessageUtilities.cpp b/Engine/Plugins/FX/Niagara/Source/NiagaraEditor/Private/NiagaraMessageUtilities.cpp
--- a/Engine/Plugins/FX/Niagara/Source/NiagaraEditor/Private/NiagaraMessageUtilities.cpp
+++ b/Engine/Plugins/FX/Niagara/Source/NiagaraEditor/Private/NiagaraMessageUtilities.cpp
@@ -151,9 +151,6 @@ UNiagaraStackEntry::FStackIssue FNiagaraMessageUtilities::StackMessageToStackIss
 	case ENiagaraMessageSeverity::Info:
 		StackIssueSeverity = EStackIssueSeverity::Info;
 		break;
-	case ENiagaraMessageSeverity::CustomNote:
-		StackIssueSeverity = EStackIssueSeverity::CustomNote;
-		break;
 	default:
 		StackIssueSeverity = EStackIssueSeverity::Info;
 		break;
@@ -186,10 +183,6 @@ UNiagaraStackEntry::FStackIssue FNiagaraMessageUtilities::StackMessageToStackIss
 		{
 			ShortDescription = LOCTEXT("EmptyShortDescriptionFromStackMessage_Error", "Error");
 		}
-		else if(StackIssueSeverity == EStackIssueSeverity::CustomNote)
-		{
-			ShortDescription = LOCTEXT("EmptyShortDescriptionFromStackMessage_CustomNote", "Custom Note");
-		}
 	}
 	
 	return UNiagaraStackEntry::FStackIssue(
diff --git a/Engine/Plugins/FX/Niagara/Source/NiagaraEditor/Private/NiagaraNodeFunctionCall.cpp b/Engine/Plugins/FX/Niagara/Source/NiagaraEditor/Private/NiagaraNodeFunctionCall.cpp
--- a/Engine/Plugins/FX/Niagara/Source/NiagaraEditor/Private/NiagaraNodeFunctionCall.cpp
+++ b/Engine/Plugins/FX/Niagara/Source/NiagaraEditor/Private/NiagaraNodeFunctionCall.cpp
@@ -226,40 +226,6 @@ TSharedPtr<SGraphNode> UNiagaraNodeFunctionCall::CreateVisualWidget()
 	}
 }
 
-void UNiagaraNodeFunctionCall::AddCustomNote(const FNiagaraStackMessage& StackMessage)
-{
-	StackMessages.Add(StackMessage);
-    OnCustomNotesChangedDelegate.ExecuteIfBound();
-    GetNiagaraGraph()->NotifyGraphChanged();
-}
-
-void UNiagaraNodeFunctionCall::RemoveCustomNote(const FGuid& MessageKey)
-{
-	StackMessages.RemoveAll([&](const FNiagaraStackMessage& Message)
-	{
-		return Message.Guid == MessageKey;
-	});
-		
-	OnCustomNotesChangedDelegate.ExecuteIfBound();
-	GetNiagaraGraph()->NotifyGraphChanged();
-}
-
-void UNiagaraNodeFunctionCall::RemoveCustomNoteViaDelegate(const FGuid MessageKey)
-{
-	const bool bContainsMessage = StackMessages.ContainsByPredicate([&](const FNiagaraStackMessage& Message)
-	{
-		return Message.Guid == MessageKey;
-	});
-
-	if(bContainsMessage)
-	{
-		FScopedTransaction Transaction(LOCTEXT("NoteRemoved", "Note Removed"));
-		Modify();
-			
-		RemoveCustomNote(MessageKey);
-	}
-}
-
 void UNiagaraNodeFunctionCall::AddOrphanedStaticSwitchPinForDataRetention(FNiagaraVariableBase StaticSwitchVariable, const FString& StaticSwitchPinDefault)
 {
 	UEdGraphPin* NewPin = AddStaticSwitchInputPin(StaticSwitchVariable);
diff --git a/Engine/Plugins/FX/Niagara/Source/NiagaraEditor/Private/NiagaraScriptMergeManager.cpp b/Engine/Plugins/FX/Niagara/Source/NiagaraEditor/Private/NiagaraScriptMergeManager.cpp
--- a/Engine/Plugins/FX/Niagara/Source/NiagaraEditor/Private/NiagaraScriptMergeManager.cpp
+++ b/Engine/Plugins/FX/Niagara/Source/NiagaraEditor/Private/NiagaraScriptMergeManager.cpp
@@ -430,18 +430,13 @@ void FNiagaraStackFunctionMergeAdapter::GatherFunctionCallNodes(TArray<UNiagaraN
 	}
 }
 
-const TArray<FNiagaraStackMessage>& FNiagaraStackFunctionMergeAdapter::GetMessages() const
-{
-	return FunctionCallNode->GetCustomNotes();
-}
-
 FNiagaraScriptStackMergeAdapter::FNiagaraScriptStackMergeAdapter(const FVersionedNiagaraEmitter& InOwningEmitter, UNiagaraNodeOutput& InOutputNode, UNiagaraScript& InScript)
 {
 	OutputNode = &InOutputNode;
 	InputNode.Reset();
 	Script = &InScript;
 	UniqueEmitterName = InOwningEmitter.Emitter->GetUniqueEmitterName();
-
+	
 	TArray<FNiagaraStackGraphUtilities::FStackNodeGroup> StackGroups;
 	GetStackNodeGroups(*OutputNode, StackGroups);
 
@@ -1124,8 +1119,6 @@ bool FNiagaraScriptStackDiffResults::IsEmpty() const
 		ChangedVersionOtherModules.Num() == 0 &&
 		EnabledChangedBaseModules.Num() == 0 &&
 		EnabledChangedOtherModules.Num() == 0 &&
-		AddedOtherMessages.Num() == 0 &&
-		RemovedBaseMessagesInOther.Num() == 0 &&
 		RemovedBaseInputOverrides.Num() == 0 &&
 		AddedOtherInputOverrides.Num() == 0 &&
 		ModifiedOtherInputOverrides.Num() == 0 &&
@@ -1204,6 +1197,7 @@ bool FNiagaraEmitterDiffResults::IsEmpty() const
 		AddedSummaryEntriesInOther.Num() == 0 &&
 		AddedSummarySectionsInOther.Num() == 0 &&
 		ModifiedStackEntryDisplayNames.Num() == 0 &&
+		AddedOrModifiedStackNotesInOther.Num() == 0 &&
 		bScratchPadModified == false &&
 		NewShouldShowSummaryViewValue.IsSet() == false;
 }
@@ -1737,6 +1731,14 @@ INiagaraMergeManager::FMergeEmitterResults FNiagaraScriptMergeManager::MergeEmit
 		MergeResults.bModifiedGraph |= StackEntryDisplayNameDiffs.bModifiedGraph;
 		MergeResults.ErrorMessages.Append(StackEntryDisplayNameDiffs.ErrorMessages);
 
+		FApplyDiffResults StackNotesDiffs = ApplyStackNoteDiffs(VersionedMergedInstance, DiffResults);
+		if (StackNotesDiffs.bSucceeded == false)
+		{
+			MergeResults.MergeResult = EMergeEmitterResult::FailedToMerge;
+		}
+		MergeResults.bModifiedGraph |= StackNotesDiffs.bModifiedGraph;
+		MergeResults.ErrorMessages.Append(StackNotesDiffs.ErrorMessages);
+
 #if 0
 		UE_LOG(LogNiagaraEditor, Log, TEXT("A"));
 		//for (FNiagaraEmitterHandle& EmitterHandle : EmitterHandles)
@@ -2236,7 +2238,8 @@ FNiagaraEmitterDiffResults FNiagaraScriptMergeManager::DiffEmitters(const FVersi
 	DiffEmitterSummary(BaseEmitterAdapter->GetEditorData(), OtherEmitterAdapter->GetEditorData(), EmitterDiffResults);
 	DiffEditableProperties(BaseEmitterData, OtherEmitterData, *FVersionedNiagaraEmitterData::StaticStruct(), EmitterDiffResults.DifferentEmitterProperties);
 	DiffStackEntryDisplayNames(BaseEmitterAdapter->GetEditorData(), OtherEmitterAdapter->GetEditorData(), EmitterDiffResults.ModifiedStackEntryDisplayNames);
-
+	DiffStackNotes(BaseEmitterAdapter->GetEditorData(), OtherEmitterAdapter->GetEditorData(), EmitterDiffResults);
+	
 	TArray<UNiagaraScript*> BaseScratchPadScripts;
 	BaseScratchPadScripts.Append(BaseEmitterData->ParentScratchPads->Scripts);
 	BaseScratchPadScripts.Append(BaseEmitterData->ScratchPads->Scripts);
@@ -2543,31 +2546,6 @@ void FNiagaraScriptMergeManager::DiffScriptStacks(TSharedRef<FNiagaraScriptStack
 			DiffResults.EnabledChangedOtherModules.Add(CommonValuePair.OtherValue);
 		}
 
-		const TArray<FNiagaraStackMessage>& BaseMessages = CommonValuePair.BaseValue->GetMessages();
-		const TArray<FNiagaraStackMessage>& OtherMessages = CommonValuePair.OtherValue->GetMessages();
-
-		for(const FNiagaraStackMessage& OtherMessage : OtherMessages)
-		{
-			if(!BaseMessages.ContainsByPredicate([&](const FNiagaraStackMessage& Message)
-			{
-				return Message.Guid == OtherMessage.Guid;
-			}))
-			{
-				DiffResults.AddedOtherMessages.Add({CommonValuePair.OtherValue, OtherMessage});
-			}
-		}
-
-		for(const FNiagaraStackMessage& BaseMessage : BaseMessages)
-		{
-			if(!OtherMessages.ContainsByPredicate([&](const FNiagaraStackMessage& Message)
-			{
-				return Message.Guid == BaseMessage.Guid;
-			}))
-			{
-				DiffResults.RemovedBaseMessagesInOther.Add({CommonValuePair.BaseValue, BaseMessage});
-			}
-		}
-
 		if (CommonValuePair.BaseValue->GetFunctionCallNode()->SelectedScriptVersion != CommonValuePair.OtherValue->GetFunctionCallNode()->SelectedScriptVersion)
 		{
 			DiffResults.ChangedVersionBaseModules.Add(CommonValuePair.BaseValue);
@@ -2667,6 +2645,24 @@ void FNiagaraScriptMergeManager::DiffStackEntryDisplayNames(const UNiagaraEmitte
 	}
 }
 
+void FNiagaraScriptMergeManager::DiffStackNotes(const UNiagaraEmitterEditorData* BaseEditorData, const UNiagaraEmitterEditorData* OtherEditorData, FNiagaraEmitterDiffResults& DiffResults) const
+{
+	if (BaseEditorData != nullptr && OtherEditorData != nullptr)
+	{
+		const TMap<FString, FNiagaraStackNoteData>& OtherStackNotes = OtherEditorData->GetStackEditorData().GetAllStackNotes();
+		for(const auto& Pair : OtherStackNotes)
+		{
+			// if the other stack notes differ in any way, we will overwrite base with other
+			TOptional<FNiagaraStackNoteData> BaseStackNote = BaseEditorData->GetStackEditorData().GetStackNote(Pair.Key);
+			
+			if(BaseStackNote.Get(FNiagaraStackNoteData()) != Pair.Value)
+			{
+				DiffResults.AddedOrModifiedStackNotesInOther.Add(Pair.Key, Pair.Value);
+			}
+		}		
+	}
+}
+
 void FNiagaraScriptMergeManager::DiffScratchPadScripts(const TArray<UNiagaraScript*>& BaseScratchPadScripts, const TArray<UNiagaraScript*>& OtherEmitterScratchPadScripts, FNiagaraEmitterDiffResults& DiffResults) const
 {
 	if (BaseScratchPadScripts.Num() != OtherEmitterScratchPadScripts.Num())
@@ -3334,26 +3330,6 @@ FNiagaraScriptMergeManager::FApplyDiffResults FNiagaraScriptMergeManager::ApplyS
 			}
 		}
 	}
-
-	for(const FNiagaraStackFunctionMessageData& AddedMessage : DiffResults.AddedOtherMessages)
-	{
-		TSharedPtr<FNiagaraStackFunctionMergeAdapter> MatchingModuleAdapter = BaseScriptStackAdapter->GetModuleFunctionById(AddedMessage.Function->GetFunctionCallNode()->NodeGuid);
-
-		if(MatchingModuleAdapter.IsValid())
-		{
-			MatchingModuleAdapter->GetFunctionCallNode()->AddCustomNote(AddedMessage.StackMessage);
-		}
-	}
-
-	for(const FNiagaraStackFunctionMessageData& RemovedMessage : DiffResults.RemovedBaseMessagesInOther)
-	{
-		TSharedPtr<FNiagaraStackFunctionMergeAdapter> MatchingModuleAdapter = BaseScriptStackAdapter->GetModuleFunctionById(RemovedMessage.Function->GetFunctionCallNode()->NodeGuid);
-
-		if(MatchingModuleAdapter.IsValid())
-		{
-			MatchingModuleAdapter->GetFunctionCallNode()->RemoveCustomNote(RemovedMessage.StackMessage.Guid);
-		}
-	}
 	
 	// Update the usage if different
 	if (DiffResults.ChangedOtherUsage.IsSet())
@@ -3805,6 +3781,25 @@ FNiagaraScriptMergeManager::FApplyDiffResults FNiagaraScriptMergeManager::ApplyS
 	return Results;
 }
 
+FNiagaraScriptMergeManager::FApplyDiffResults FNiagaraScriptMergeManager::ApplyStackNoteDiffs(FVersionedNiagaraEmitter BaseEmitter, const FNiagaraEmitterDiffResults& DiffResults) const
+{
+	if(DiffResults.AddedOrModifiedStackNotesInOther.Num() > 0)
+	{
+		if(UNiagaraEmitterEditorData* EmitterEditorData = Cast<UNiagaraEmitterEditorData>(BaseEmitter.GetEmitterData()->GetEditorData()))
+		{
+			for(const auto& Pair : DiffResults.AddedOrModifiedStackNotesInOther)
+			{
+				EmitterEditorData->GetStackEditorData().AddOrReplaceStackNote(Pair.Key, Pair.Value, false);
+			}
+		}
+	}
+
+	FApplyDiffResults Results;
+	Results.bSucceeded = true;
+	Results.bModifiedGraph = false;
+	return Results;
+}
+
 FNiagaraScriptMergeManager::FCachedMergeAdapter* FNiagaraScriptMergeManager::FindOrAddMergeAdapterCacheForEmitter(const FVersionedNiagaraEmitter& VersionedEmitter)
 {
 	FCachedMergeAdapter* CachedMergeAdapter = CachedMergeAdapters.Find(VersionedEmitter);
diff --git a/Engine/Plugins/FX/Niagara/Source/NiagaraEditor/Private/NiagaraScriptMergeManager.h b/Engine/Plugins/FX/Niagara/Source/NiagaraEditor/Private/NiagaraScriptMergeManager.h
--- a/Engine/Plugins/FX/Niagara/Source/NiagaraEditor/Private/NiagaraScriptMergeManager.h
+++ b/Engine/Plugins/FX/Niagara/Source/NiagaraEditor/Private/NiagaraScriptMergeManager.h
@@ -117,8 +117,6 @@ public:
 
 	void GatherFunctionCallNodes(TArray<UNiagaraNodeFunctionCall*>& OutFunctionCallNodes) const;
 
-	const TArray<FNiagaraStackMessage>& GetMessages() const;
-
 private:
 	TWeakObjectPtr<UNiagaraScript> OwningScript;
 	TWeakObjectPtr<UNiagaraNodeFunctionCall> FunctionCallNode;
@@ -338,9 +336,6 @@ struct FNiagaraScriptStackDiffResults
 
 	TArray<TSharedRef<FNiagaraStackFunctionMergeAdapter>> EnabledChangedBaseModules;
 	TArray<TSharedRef<FNiagaraStackFunctionMergeAdapter>> EnabledChangedOtherModules;
-
-	TArray<FNiagaraStackFunctionMessageData> AddedOtherMessages;
-	TArray<FNiagaraStackFunctionMessageData> RemovedBaseMessagesInOther;
 	
 	TArray<TSharedRef<FNiagaraStackFunctionInputOverrideMergeAdapter>> RemovedBaseInputOverrides;
 	TArray<TSharedRef<FNiagaraStackFunctionInputOverrideMergeAdapter>> AddedOtherInputOverrides;
@@ -418,6 +413,8 @@ struct FNiagaraEmitterDiffResults
 
 	TMap<FString, FText> ModifiedStackEntryDisplayNames;
 
+	TMap<FString, FNiagaraStackNoteData> AddedOrModifiedStackNotesInOther;
+	
 private:
 	bool bIsValid;
 	TArray<FText> ErrorMessages;
@@ -500,6 +497,8 @@ public:
 
 	void DiffStackEntryDisplayNames(const UNiagaraEmitterEditorData* BaseEditorData, const UNiagaraEmitterEditorData* OtherEditorData, TMap<FString, FText>& OutModifiedStackEntryDisplayNames) const;
 
+	void DiffStackNotes(const UNiagaraEmitterEditorData* BaseEditorData, const UNiagaraEmitterEditorData* OtherEditorData, FNiagaraEmitterDiffResults& DiffResults) const;
+	
 	void DiffScratchPadScripts(const TArray<UNiagaraScript*>& BaseScratchPadScripts, const TArray<UNiagaraScript*>& OtherEmitterScratchPadScripts, FNiagaraEmitterDiffResults& DiffResults) const;
 
 	virtual void CopyPropertiesToBase(void* BaseDataAddress, const void* OtherDataAddress, TArray<FProperty*> PropertiesToCopy) const override;
@@ -550,6 +549,8 @@ private:
 	
 	FApplyDiffResults ApplyStackEntryDisplayNameDiffs(FVersionedNiagaraEmitter BaseEmitter, const FNiagaraEmitterDiffResults& DiffResults) const;
 
+	FApplyDiffResults ApplyStackNoteDiffs(FVersionedNiagaraEmitter BaseEmitter, const FNiagaraEmitterDiffResults& DiffResults) const;
+	
 	FApplyDiffResults AddModule(
 		TSharedRef<FNiagaraEmitterMergeAdapter> BaseEmitterAdapter,
 		TSharedRef<FNiagaraScratchPadMergeAdapter> ScratchPadAdapter,
diff --git a/Engine/Plugins/FX/Niagara/Source/NiagaraEditor/Private/NiagaraStackEditorData.cpp b/Engine/Plugins/FX/Niagara/Source/NiagaraEditor/Private/NiagaraStackEditorData.cpp
--- a/Engine/Plugins/FX/Niagara/Source/NiagaraEditor/Private/NiagaraStackEditorData.cpp
+++ b/Engine/Plugins/FX/Niagara/Source/NiagaraEditor/Private/NiagaraStackEditorData.cpp
@@ -2,6 +2,9 @@
 
 #include "NiagaraStackEditorData.h"
 
+#include "NiagaraNodeFunctionCall.h"
+#include "ViewModels/Stack/NiagaraStackModuleItem.h"
+
 #include UE_INLINE_GENERATED_CPP_BY_NAME(NiagaraStackEditorData)
 
 bool UNiagaraStackEditorData::GetStackEntryIsRenamePending(const FString& StackEntryKey) const
@@ -73,6 +76,60 @@ void UNiagaraStackEditorData::SetStackEntryWasExpandedPreSearch(const FString& S
 	}
 }
 
+const TMap<FString, FNiagaraStackNoteData>& UNiagaraStackEditorData::GetAllStackNotes()
+{
+	return StackNotes;
+}
+
+TOptional<FNiagaraStackNoteData> UNiagaraStackEditorData::GetStackNote(const FString& StackEntryKey)
+{
+	if(StackNotes.Contains(StackEntryKey))
+	{
+		return StackNotes[StackEntryKey];
+	}
+
+	return {};
+}
+
+bool UNiagaraStackEditorData::HasStackNote(const FString& StackEntryKey)
+{
+	if(StackNotes.Contains(StackEntryKey))
+	{
+		return true;
+	}
+
+	return false;
+}
+
+void UNiagaraStackEditorData::AddOrReplaceStackNote(const FString& StackEntryKey, FNiagaraStackNoteData NewStackMessage, bool bBroadcast)
+{
+	StackNotes.Add(StackEntryKey, NewStackMessage);
+
+	if(bBroadcast)
+	{
+		OnPersistentDataChanged().Broadcast();
+	}
+}
+
+void UNiagaraStackEditorData::DeleteStackNote(const FString& StackEntryKey)
+{
+	if(StackNotes.Contains(StackEntryKey))
+	{
+		StackNotes.Remove(StackEntryKey);
+	}
+}
+
+void UNiagaraStackEditorData::TransferDeprecatedStackNotes(const UNiagaraNodeFunctionCall& FunctionCallNode)
+{
+	for(const FNiagaraStackMessage& StackMessage : FunctionCallNode.GetDeprecatedCustomNotes())
+	{
+		FNiagaraStackNoteData StackNoteData;
+		StackNoteData.MessageHeader = StackMessage.ShortDescription;
+		StackNoteData.Message = StackMessage.MessageText;
+		AddOrReplaceStackNote(FNiagaraStackGraphUtilities::StackKeys::GenerateStackModuleEditorDataKey(FunctionCallNode), StackNoteData, false);
+	}
+}
+
 ENiagaraStackEntryInlineDisplayMode UNiagaraStackEditorData::GetStackEntryInlineDisplayMode(const FString& StackEntryKey) const
 {
 	const ENiagaraStackEntryInlineDisplayMode* InlineDisplayModePtr = StackEntryKeyToInlineDisplayModeMap.Find(StackEntryKey);
diff --git a/Engine/Plugins/FX/Niagara/Source/NiagaraEditor/Private/NiagaraSystemEditorData.cpp b/Engine/Plugins/FX/Niagara/Source/NiagaraEditor/Private/NiagaraSystemEditorData.cpp
--- a/Engine/Plugins/FX/Niagara/Source/NiagaraEditor/Private/NiagaraSystemEditorData.cpp
+++ b/Engine/Plugins/FX/Niagara/Source/NiagaraEditor/Private/NiagaraSystemEditorData.cpp
@@ -11,6 +11,7 @@
 #include "NiagaraNode.h"
 #include "EdGraphSchema_NiagaraSystemOverview.h"
 #include "NiagaraConstants.h"
+#include "NiagaraScriptSource.h"
 #include "NiagaraScriptVariable.h"
 #include "NiagaraSettings.h"
 #include "EdGraph/EdGraph.h"
@@ -137,6 +138,8 @@ void UNiagaraSystemEditorData::PostLoadFromOwner(UObject* InOwner)
 		SynchronizeOverviewGraphWithSystem(*OwningSystem);
 	}
 
+	PostLoad_TransferModuleStackNotesToNewFormat(OwningSystem);
+	
 	// Remove any niagara nodes which may have been pasted into the overview graph in error.  This is no longer possible.
 	TArray<UNiagaraNode*> NiagaraNodes;
 	SystemOverviewGraph->GetNodesOfClass<UNiagaraNode>(NiagaraNodes);
@@ -226,6 +229,33 @@ void UNiagaraSystemEditorData::UpdatePlaybackRangeFromEmitters(UNiagaraSystem& O
 	}
 }
 
+void UNiagaraSystemEditorData::PostLoad_TransferModuleStackNotesToNewFormat(UObject* InOwner)
+{
+	// since we lack graph context during post load, we do this workaround to find the emitter graph
+	UNiagaraSystem* System = Cast<UNiagaraSystem>(InOwner);
+
+	for(UNiagaraScriptSourceBase* SourceBase : System->GetAllSourceScripts())
+	{
+		UNiagaraScriptSource* Source = CastChecked<UNiagaraScriptSource>(SourceBase);
+
+		if(Source->NodeGraph == nullptr)
+		{
+			continue;
+		}
+
+		TArray<UNiagaraNodeFunctionCall*> FunctionCallNodes;
+		Source->NodeGraph->GetNodesOfClass(FunctionCallNodes);
+
+		for(UNiagaraNodeFunctionCall* FunctionCallNode : FunctionCallNodes)
+		{
+			if(FunctionCallNode->GetDeprecatedCustomNotes().Num() > 0)
+			{
+				StackEditorData->TransferDeprecatedStackNotes(*FunctionCallNode);
+			}
+		}
+	}	
+}
+
 void FindOuterNodes(const TArray<UNiagaraOverviewNode*> OverviewNodes, UNiagaraOverviewNode*& OutLeft, UNiagaraOverviewNode*& OutRight)
 {
 	if (OverviewNodes.Num() > 0)
diff --git a/Engine/Plugins/FX/Niagara/Source/NiagaraEditor/Private/ViewModels/NiagaraCurveSelectionViewModel.cpp b/Engine/Plugins/FX/Niagara/Source/NiagaraEditor/Private/ViewModels/NiagaraCurveSelectionViewModel.cpp
--- a/Engine/Plugins/FX/Niagara/Source/NiagaraEditor/Private/ViewModels/NiagaraCurveSelectionViewModel.cpp
+++ b/Engine/Plugins/FX/Niagara/Source/NiagaraEditor/Private/ViewModels/NiagaraCurveSelectionViewModel.cpp
@@ -593,7 +593,7 @@ TSharedPtr<FNiagaraCurveSelectionTreeNode> UNiagaraCurveSelectionViewModel::Crea
 		}
 		else
 		{
-			const FText* DisplayName = StackEditorData.GetStackEntryDisplayName(FNiagaraStackGraphUtilities::GenerateStackModuleEditorDataKey(FunctionCallNode));
+			const FText* DisplayName = StackEditorData.GetStackEntryDisplayName(FNiagaraStackGraphUtilities::StackKeys::GenerateStackModuleEditorDataKey(FunctionCallNode));
 			if(DisplayName != nullptr)
 			{
 				NewFunctionNode->SetDisplayName(*DisplayName);
diff --git a/Engine/Plugins/FX/Niagara/Source/NiagaraEditor/Private/ViewModels/Stack/NiagaraStackClipboardUtilities.cpp b/Engine/Plugins/FX/Niagara/Source/NiagaraEditor/Private/ViewModels/Stack/NiagaraStackClipboardUtilities.cpp
--- a/Engine/Plugins/FX/Niagara/Source/NiagaraEditor/Private/ViewModels/Stack/NiagaraStackClipboardUtilities.cpp
+++ b/Engine/Plugins/FX/Niagara/Source/NiagaraEditor/Private/ViewModels/Stack/NiagaraStackClipboardUtilities.cpp
@@ -255,6 +255,17 @@ bool FNiagaraStackClipboardUtilities::TestCanDeleteSelectionWithMessage(const TA
 	}
 }
 
+void FNiagaraStackClipboardUtilities::CopyNote(const UNiagaraStackEntry* StackEntry)
+{
+	if(StackEntry->HasStackNoteData())
+	{
+		UNiagaraClipboardContent* NewClipboardContent = UNiagaraClipboardContent::Create();
+		NewClipboardContent->StackNote = StackEntry->GetStackNoteData();
+
+		FNiagaraEditorModule::Get().GetClipboard().SetClipboardContent(NewClipboardContent);
+	}
+}
+
 void FNiagaraStackClipboardUtilities::DeleteSelection(const TArray<UNiagaraStackEntry*>& SelectedEntries)
 {
 	FText TransactionMessage;
diff --git a/Engine/Plugins/FX/Niagara/Source/NiagaraEditor/Private/ViewModels/Stack/NiagaraStackEntry.cpp b/Engine/Plugins/FX/Niagara/Source/NiagaraEditor/Private/ViewModels/Stack/NiagaraStackEntry.cpp
--- a/Engine/Plugins/FX/Niagara/Source/NiagaraEditor/Private/ViewModels/Stack/NiagaraStackEntry.cpp
+++ b/Engine/Plugins/FX/Niagara/Source/NiagaraEditor/Private/ViewModels/Stack/NiagaraStackEntry.cpp
@@ -5,6 +5,7 @@
 #include "NiagaraEditorUtilities.h"
 #include "NiagaraEmitter.h"
 #include "ViewModels/Stack/NiagaraStackErrorItem.h"
+#include "ViewModels/Stack/NiagaraStackNote.h"
 #include "ViewModels/NiagaraSystemViewModel.h"
 #include "ViewModels/NiagaraEmitterViewModel.h"
 #include "NiagaraStackEditorData.h"
@@ -29,6 +30,8 @@ const FName UNiagaraStackEntry::FExecutionSubcategoryNames::Event = TEXT("Event"
 const FName UNiagaraStackEntry::FExecutionSubcategoryNames::SimulationStage = TEXT("Simulation Stage");
 const FName UNiagaraStackEntry::FExecutionSubcategoryNames::Render = TEXT("Render");
 
+#define LOCTEXT_NAMESPACE "UNiagaraStackEntry"
+
 UNiagaraStackEntry::FStackIssueFix::FStackIssueFix() : Style(EStackIssueFixStyle::Fix)
 {
 }
@@ -206,6 +209,12 @@ void UNiagaraStackEntry::Finalize()
 	}
 	Children.Empty();
 
+	if(StackNote)
+	{
+		StackNote->Finalize();
+		StackNote = nullptr;
+	}
+
 	for (UNiagaraStackEntry* ErrorChild : ErrorChildren)
 	{
 		ErrorChild->Finalize();
@@ -381,6 +390,11 @@ void UNiagaraStackEntry::GetFilteredChildren(TArray<UNiagaraStackEntry*>& OutFil
 
 		FilteredChildren.Empty();
 		FilteredChildren.Append(ErrorChildren);
+		if(StackNote != nullptr && StackNote->GetShouldShowInStack())
+		{
+			FilteredChildren.Add(StackNote);
+		}
+		
 		for (UNiagaraStackEntry* Child : Children)
 		{
 			bool bPassesFilter = true;
@@ -454,6 +468,11 @@ TSharedPtr<FNiagaraEmitterViewModel> UNiagaraStackEntry::GetEmitterViewModel() c
 	return EmitterViewModel.Pin();
 }
 
+UNiagaraStackNote* UNiagaraStackEntry::GetStackNote()
+{
+	return StackNote;
+}
+
 UNiagaraStackEntry::FOnExpansionChanged& UNiagaraStackEntry::OnExpansionChanged()
 {
 	return ExpansionChangedDelegate;
@@ -597,11 +616,6 @@ void UNiagaraStackEntry::GetRecursiveUsages(bool& bRead, bool& bWrite) const
 	bWrite = GetCollectedUsageData().bHasReferencedParameterWrite;
 }
 
-int32 UNiagaraStackEntry::GetTotalNumberOfCustomNotes() const
-{
-	return GetCollectedIssueData().TotalNumberOfCustomNotes;
-}
-
 int32 UNiagaraStackEntry::GetTotalNumberOfInfoIssues() const
 {
 	return GetCollectedIssueData().TotalNumberOfInfoIssues;
@@ -752,9 +766,9 @@ void UNiagaraStackEntry::RefreshChildren()
 		}
 	}
 
-	Children.Empty();
+	Children.Empty();	
 	Children.Append(NewChildren);
-
+	
 	for (UNiagaraStackEntry* Child : Children)
 	{
 		UNiagaraStackEntry* OuterOwner = Cast<UNiagaraStackEntry>(Child->GetOuter());
@@ -773,6 +787,27 @@ void UNiagaraStackEntry::RefreshChildren()
 			Child->SetOnRequestDrop(FOnRequestDrop::CreateUObject(this, &UNiagaraStackEntry::ChildRequestDrop));
 		}
 	}
+
+	if(HasStackNoteData())
+	{		
+		if(StackNote == nullptr)
+		{
+			StackNote = NewObject<UNiagaraStackNote>(this);
+			StackNote->Initialize(CreateDefaultChildRequiredData(), GetStackEditorDataKey());
+			StackNote->OnNoteChanged().BindUObject(this, &UNiagaraStackEntry::SetStackNoteData);
+		}
+		
+		StackNote->RefreshChildren();
+	}
+	else
+	{
+		if(StackNote != nullptr)
+		{
+			StackNote->OnNoteChanged().Unbind();
+			StackNote->Finalize();
+			StackNote = nullptr;
+		}
+	}
 	
 	// Stack issues refresh
 	NewStackIssues.RemoveAll([this](const FStackIssue& Issue) { return Issue.GetCanBeDismissed() && GetStackEditorData().GetDismissedStackIssueIds().Contains(Issue.GetUniqueIdentifier()); }); 
@@ -906,7 +941,6 @@ const UNiagaraStackEntry::FCollectedIssueData& UNiagaraStackEntry::GetCollectedI
 			CachedCollectedIssueData->TotalNumberOfInfoIssues += ChildStackEntry->GetTotalNumberOfInfoIssues();
 			CachedCollectedIssueData->TotalNumberOfWarningIssues += ChildStackEntry->GetTotalNumberOfWarningIssues();
 			CachedCollectedIssueData->TotalNumberOfErrorIssues += ChildStackEntry->GetTotalNumberOfErrorIssues();
-			CachedCollectedIssueData->TotalNumberOfCustomNotes += ChildStackEntry->GetTotalNumberOfCustomNotes();
 			
 			if (ChildStackEntry->GetIssues().Num() > 0)
 			{
@@ -930,10 +964,6 @@ const UNiagaraStackEntry::FCollectedIssueData& UNiagaraStackEntry::GetCollectedI
 			{
 				CachedCollectedIssueData->TotalNumberOfErrorIssues++;
 			}
-			else if (Issue.GetSeverity() == EStackIssueSeverity::CustomNote)
-			{
-				CachedCollectedIssueData->TotalNumberOfCustomNotes++;
-			}
 		}
 	}
 
@@ -1100,6 +1130,42 @@ void UNiagaraStackEntry::OnRenamed(FText NewName)
 	}
 }
 
+bool UNiagaraStackEntry::HasStackNoteData() const
+{
+	return GetStackEditorData().HasStackNote(GetStackEditorDataKey());
+}
+
+FNiagaraStackNoteData UNiagaraStackEntry::GetStackNoteData() const
+{
+	return GetStackEditorData().GetStackNote(GetStackEditorDataKey()).Get(FNiagaraStackNoteData());
+}
+
+void UNiagaraStackEntry::SetStackNoteData(FNiagaraStackNoteData InStackNoteData)
+{
+	if(InStackNoteData.IsValid())
+	{
+		GetStackEditorData().AddOrReplaceStackNote(GetStackEditorDataKey(), InStackNoteData);
+	}
+	else
+	{
+		GetStackEditorData().DeleteStackNote(GetStackEditorDataKey());	
+	}
+	
+	RefreshChildren();
+}
+
+void UNiagaraStackEntry::DeleteStackNoteData()
+{
+	if(HasStackNoteData())
+	{
+		FScopedTransaction Transaction(LOCTEXT("DeleteStackNoteTransaction", "Deleted Stack Note"));
+		GetStackEditorData().Modify();
+		
+		GetStackEditorData().DeleteStackNote(GetStackEditorDataKey());
+		RefreshChildren();
+	}
+}
+
 FNiagaraHierarchyIdentity UNiagaraStackEntry::DetermineSummaryIdentity() const
 {
 	return FNiagaraHierarchyIdentity();
@@ -1192,3 +1258,4 @@ void UNiagaraStackSpacer::Initialize(FRequiredEntryData InRequiredEntryData, flo
 	ShouldShowInStack = InShouldShowInStack;
 }
 
+#undef LOCTEXT_NAMESPACE
diff --git a/Engine/Plugins/FX/Niagara/Source/NiagaraEditor/Private/ViewModels/Stack/NiagaraStackErrorItem.cpp b/Engine/Plugins/FX/Niagara/Source/NiagaraEditor/Private/ViewModels/Stack/NiagaraStackErrorItem.cpp
--- a/Engine/Plugins/FX/Niagara/Source/NiagaraEditor/Private/ViewModels/Stack/NiagaraStackErrorItem.cpp
+++ b/Engine/Plugins/FX/Niagara/Source/NiagaraEditor/Private/ViewModels/Stack/NiagaraStackErrorItem.cpp
@@ -43,7 +43,7 @@ void UNiagaraStackErrorItem::SetIsExpandedByDefault(bool bIsExpanded)
 
 UNiagaraStackEntry::EStackRowStyle UNiagaraStackErrorItem::GetStackRowStyle() const
 {
-	return StackIssue.GetSeverity() == EStackIssueSeverity::CustomNote ? EStackRowStyle::ItemContentNote : EStackRowStyle::StackIssue;
+	return EStackRowStyle::StackIssue;
 }
 
 UNiagaraStackErrorItem::FOnIssueNotify& UNiagaraStackErrorItem::OnIssueModified()
@@ -146,7 +146,7 @@ EStackIssueSeverity UNiagaraStackErrorItemLongDescription::GetIssueSeverity() co
 
 UNiagaraStackEntry::EStackRowStyle UNiagaraStackErrorItemLongDescription::GetStackRowStyle() const
 {
-	return StackIssue.GetSeverity() == EStackIssueSeverity::CustomNote ? EStackRowStyle::ItemContentNote : EStackRowStyle::StackIssue;
+	return EStackRowStyle::StackIssue;
 }
 
 //UNiagaraStackErrorItemFix
@@ -180,7 +180,7 @@ EStackIssueSeverity UNiagaraStackErrorItemFix::GetIssueSeverity() const
 
 UNiagaraStackEntry::EStackRowStyle UNiagaraStackErrorItemFix::GetStackRowStyle() const
 {
-	return StackIssue.GetSeverity() == EStackIssueSeverity::CustomNote ? EStackRowStyle::ItemContentNote : EStackRowStyle::StackIssue;
+	return EStackRowStyle::StackIssue;
 }
 
 FText UNiagaraStackErrorItemFix::GetFixButtonText() const
@@ -206,7 +206,7 @@ void UNiagaraStackErrorItemDismiss::Initialize(FRequiredEntryData InRequiredEntr
 	UNiagaraStackEntry::Initialize(InRequiredEntryData, ErrorStackEditorDataKey);
 	StackIssue = InStackIssue;
 	IssueFix = FStackIssueFix(
-		StackIssue.GetSeverity() == EStackIssueSeverity::Info || StackIssue.GetSeverity() == EStackIssueSeverity::CustomNote ?
+		StackIssue.GetSeverity() == EStackIssueSeverity::Info ?
 			LOCTEXT("DismissNote", "Dismiss note") :
 			LOCTEXT("DismissError", "Dismiss the issue without fixing (I know what I'm doing)"),
 		FStackIssueFixDelegate::CreateUObject(this, &UNiagaraStackErrorItemDismiss::DismissIssue));
@@ -227,12 +227,12 @@ void UNiagaraStackErrorItemDismiss::DismissIssue()
 
 UNiagaraStackEntry::EStackRowStyle UNiagaraStackErrorItemDismiss::GetStackRowStyle() const
 {
-	return StackIssue.GetSeverity() == EStackIssueSeverity::CustomNote ? EStackRowStyle::ItemContentNote : EStackRowStyle::StackIssue;
+	return EStackRowStyle::StackIssue;
 }
 
 FText UNiagaraStackErrorItemDismiss::GetFixButtonText() const
 {
-	return StackIssue.GetSeverity() == EStackIssueSeverity::Info || StackIssue.GetSeverity() == EStackIssueSeverity::CustomNote ? LOCTEXT("DismissNoteButtonText", "Dismiss") : LOCTEXT("DismissIssueButtonText", "Dismiss issue");
+	return StackIssue.GetSeverity() == EStackIssueSeverity::Info ? LOCTEXT("DismissNoteButtonText", "Dismiss") : LOCTEXT("DismissIssueButtonText", "Dismiss issue");
 }
 
 #undef LOCTEXT_NAMESPACE
diff --git a/Engine/Plugins/FX/Niagara/Source/NiagaraEditor/Private/ViewModels/Stack/NiagaraStackFunctionInput.cpp b/Engine/Plugins/FX/Niagara/Source/NiagaraEditor/Private/ViewModels/Stack/NiagaraStackFunctionInput.cpp
--- a/Engine/Plugins/FX/Niagara/Source/NiagaraEditor/Private/ViewModels/Stack/NiagaraStackFunctionInput.cpp
+++ b/Engine/Plugins/FX/Niagara/Source/NiagaraEditor/Private/ViewModels/Stack/NiagaraStackFunctionInput.cpp
@@ -163,7 +163,7 @@ void UNiagaraStackFunctionInput::Initialize(
 	DisplayName = FText::FromName(InputParameterHandle.GetName());
 
 	InputType = InInputType;
-	StackEditorDataKey = FNiagaraStackGraphUtilities::GenerateStackFunctionInputEditorDataKey(*OwningFunctionCallNode.Get(), InputParameterHandle);
+	StackEditorDataKey = FNiagaraStackGraphUtilities::StackKeys::GenerateStackFunctionInputEditorDataKey(*OwningFunctionCallNode.Get(), InputParameterHandle);
 
 	TArray<UNiagaraScript*> AffectedScriptsNotWeak;
 	for (TWeakObjectPtr<UNiagaraScript> AffectedScript : AffectedScripts)
diff --git a/Engine/Plugins/FX/Niagara/Source/NiagaraEditor/Private/ViewModels/Stack/NiagaraStackGraphUtilities.cpp b/Engine/Plugins/FX/Niagara/Source/NiagaraEditor/Private/ViewModels/Stack/NiagaraStackGraphUtilities.cpp
--- a/Engine/Plugins/FX/Niagara/Source/NiagaraEditor/Private/ViewModels/Stack/NiagaraStackGraphUtilities.cpp
+++ b/Engine/Plugins/FX/Niagara/Source/NiagaraEditor/Private/ViewModels/Stack/NiagaraStackGraphUtilities.cpp
@@ -956,16 +956,21 @@ void FNiagaraStackGraphUtilities::InitializeStackFunctionInput(TSharedRef<FNiaga
 	InitializeStackFunctionInputsInternal(SystemViewModel, EmitterViewModel, StackEditorData, ModuleNode, InputFunctionCallNode, InputSelector);
 }
 
-FString FNiagaraStackGraphUtilities::GenerateStackFunctionInputEditorDataKey(UNiagaraNodeFunctionCall& FunctionCallNode, FNiagaraParameterHandle InputParameterHandle)
+FString FNiagaraStackGraphUtilities::StackKeys::GenerateStackFunctionInputEditorDataKey(const UNiagaraNodeFunctionCall& FunctionCallNode, FNiagaraParameterHandle InputParameterHandle)
 {
 	return FunctionCallNode.GetFunctionName() + InputParameterHandle.GetParameterHandleString().ToString();
 }
 
-FString FNiagaraStackGraphUtilities::GenerateStackModuleEditorDataKey(UNiagaraNodeFunctionCall& ModuleNode)
+FString FNiagaraStackGraphUtilities::StackKeys::GenerateStackModuleEditorDataKey(const UNiagaraNodeFunctionCall& ModuleNode)
 {
 	return ModuleNode.NodeGuid.ToString(EGuidFormats::DigitsWithHyphens);
 }
 
+FString FNiagaraStackGraphUtilities::StackKeys::GenerateStackRendererEditorDataKey(const UNiagaraRendererProperties& Renderer)
+{
+	return FString::Printf(TEXT("Renderer-%s"), *Renderer.GetName());
+}
+
 void ExtractInputPinsFromHistory(FNiagaraParameterMapHistory& History, UEdGraph* FunctionGraph, FNiagaraStackGraphUtilities::ENiagaraGetStackFunctionInputPinsOptions Options, TArray<const UEdGraphPin*>& OutPins)
 {
 	const int32 ModuleNamespaceLength = FNiagaraConstants::ModuleNamespaceString.Len();
@@ -3795,7 +3800,7 @@ void FNiagaraStackGraphUtilities::RenameAssignmentTarget(
 	}
 
 	bool bIsCurrentlyExpanded = StackEditorData != nullptr 
-		? StackEditorData->GetStackEntryIsExpanded(GenerateStackModuleEditorDataKey(OwningAssignmentNode), false) 
+		? StackEditorData->GetStackEntryIsExpanded(StackKeys::GenerateStackModuleEditorDataKey(OwningAssignmentNode), false) 
 		: false;
 
 	if (ensureMsgf(OwningAssignmentNode.RenameAssignmentTarget(CurrentAssignmentTarget.GetName(), NewAssignmentTargetName), TEXT("Failed to rename assignment node input.")))
@@ -3829,7 +3834,7 @@ void FNiagaraStackGraphUtilities::RenameAssignmentTarget(
 		if (StackEditorData != nullptr)
 		{
 			// Restore the expanded state with the new editor data key.
-			FString NewStackEditorDataKey = GenerateStackFunctionInputEditorDataKey(OwningAssignmentNode, NewInputParameterHandle);
+			FString NewStackEditorDataKey = StackKeys::GenerateStackFunctionInputEditorDataKey(OwningAssignmentNode, NewInputParameterHandle);
 			StackEditorData->SetStackEntryIsExpanded(NewStackEditorDataKey, bIsCurrentlyExpanded);
 		}
 
diff --git a/Engine/Plugins/FX/Niagara/Source/NiagaraEditor/Private/ViewModels/Stack/NiagaraStackItemGroup.cpp b/Engine/Plugins/FX/Niagara/Source/NiagaraEditor/Private/ViewModels/Stack/NiagaraStackItemGroup.cpp
--- a/Engine/Plugins/FX/Niagara/Source/NiagaraEditor/Private/ViewModels/Stack/NiagaraStackItemGroup.cpp
+++ b/Engine/Plugins/FX/Niagara/Source/NiagaraEditor/Private/ViewModels/Stack/NiagaraStackItemGroup.cpp
@@ -56,7 +56,7 @@ uint32 UNiagaraStackItemGroup::GetRecursiveStackIssuesCount() const
 		TArray<UNiagaraStackErrorItem*> RecursiveIssues;
 		FNiagaraStackGraphUtilities::GetStackIssuesRecursively(this, RecursiveIssues);
 		RecursiveStackIssuesCount = RecursiveIssues.Num();
-		EStackIssueSeverity MinSeverity = EStackIssueSeverity::CustomNote;
+		EStackIssueSeverity MinSeverity = EStackIssueSeverity::None;
 		for (auto Issue : RecursiveIssues)
 		{
 			if (Issue->GetStackIssue().GetSeverity() < MinSeverity)
diff --git a/Engine/Plugins/FX/Niagara/Source/NiagaraEditor/Private/ViewModels/Stack/NiagaraStackModuleItem.cpp b/Engine/Plugins/FX/Niagara/Source/NiagaraEditor/Private/ViewModels/Stack/NiagaraStackModuleItem.cpp
--- a/Engine/Plugins/FX/Niagara/Source/NiagaraEditor/Private/ViewModels/Stack/NiagaraStackModuleItem.cpp
+++ b/Engine/Plugins/FX/Niagara/Source/NiagaraEditor/Private/ViewModels/Stack/NiagaraStackModuleItem.cpp
@@ -39,6 +39,7 @@
 #include "Toolkits/NiagaraSystemToolkit.h"
 
 // TODO: Remove these
+#include "ViewModels/Stack/NiagaraStackErrorItem.h"
 #include "ViewModels/Stack/NiagaraStackViewModel.h"
 #include "Widgets/Notifications/SNotificationList.h"
 
@@ -70,8 +71,7 @@ UNiagaraNodeFunctionCall& UNiagaraStackModuleItem::GetModuleNode() const
 void UNiagaraStackModuleItem::Initialize(FRequiredEntryData InRequiredEntryData, INiagaraStackItemGroupAddUtilities* InGroupAddUtilities, UNiagaraNodeFunctionCall& InFunctionCallNode)
 {
 	checkf(FunctionCallNode == nullptr, TEXT("Can not set the node more than once."));
-	FString ModuleStackEditorDataKey = InFunctionCallNode.NodeGuid.ToString(EGuidFormats::DigitsWithHyphens);
-	Super::Initialize(InRequiredEntryData, ModuleStackEditorDataKey);
+	Super::Initialize(InRequiredEntryData, FNiagaraStackGraphUtilities::StackKeys::GenerateStackModuleEditorDataKey(InFunctionCallNode));
 	GroupAddUtilities = InGroupAddUtilities;
 	FunctionCallNode = &InFunctionCallNode;
 	OutputNode = FNiagaraStackGraphUtilities::GetEmitterOutputNodeForStackNode(*FunctionCallNode);
@@ -91,15 +91,8 @@ void UNiagaraStackModuleItem::Initialize(FRequiredEntryData InRequiredEntryData,
 		, FObjectKey(FunctionCallNode)
 		, MessageManagerRegistrationKey
 	).BindUObject(this, &UNiagaraStackModuleItem::OnMessageManagerRefresh);
-
-	FunctionCallNode->OnCustomNotesChanged().BindLambda([this]()
-	{
-		RefreshChildren();
-	});
 }
 
-
-
 FText UNiagaraStackModuleItem::GetDisplayName() const
 {
 	if (DisplayNameCache.IsSet() == false)
@@ -137,8 +130,6 @@ void UNiagaraStackModuleItem::FinalizeInternal()
 	{
 		FNiagaraMessageManager::Get()->Unsubscribe(FText::FromString("StackModuleItem"), MessageLogGuid, MessageManagerRegistrationKey);
 	}
-
-	FunctionCallNode->OnCustomNotesChanged().Unbind();
 	
 	Super::FinalizeInternal();
 }
@@ -157,7 +148,7 @@ void UNiagaraStackModuleItem::RefreshChildrenInternal(const TArray<UNiagaraStack
 		{
 			bCanRefresh = true;
 		}
-
+		
 		if (InputCollection == nullptr)
 		{
 			TArray<FString> InputParameterHandlePath;
@@ -1064,18 +1055,6 @@ void UNiagaraStackModuleItem::RefreshIssues(TArray<FStackIssue>& NewIssues)
 	}
 
 	NewIssues.Append(MessageManagerIssues);
-	for(auto& Message : FunctionCallNode->GetCustomNotes())
-	{
-		TArray<FLinkNameAndDelegate> Links;
-		const FText LinkText = LOCTEXT("DeleteNoteLinkLabel", "Delete note");
-
-		// we delete the message rather than dismissing it
-		FSimpleDelegate MessageDelegate = FSimpleDelegate::CreateUObject(FunctionCallNode, &UNiagaraNodeFunctionCall::RemoveCustomNoteViaDelegate, Message.Guid);
-		const FLinkNameAndDelegate Link = FLinkNameAndDelegate(LinkText, MessageDelegate);
-		Links.Add(Link);
-
-		NewIssues.Add(FNiagaraMessageUtilities::StackMessageToStackIssue(Message, GetStackEditorDataKey(), Links));
-	}
 
 	if (FunctionCallNode->FunctionScript == nullptr && FunctionCallNode->GetClass() == UNiagaraNodeFunctionCall::StaticClass())
 	{
@@ -1600,12 +1579,12 @@ void UNiagaraStackModuleItem::Copy(UNiagaraClipboardContent* ClipboardContent) c
 	UNiagaraNodeAssignment* AssignmentNode = Cast<UNiagaraNodeAssignment>(FunctionCallNode);
 	if (AssignmentNode != nullptr)
 	{
-		ClipboardFunction = UNiagaraClipboardFunction::CreateAssignmentFunction(ClipboardContent, AssignmentNode->GetFunctionName(), AssignmentNode->GetAssignmentTargets(), AssignmentNode->GetAssignmentDefaults());
+		ClipboardFunction = UNiagaraClipboardFunction::CreateAssignmentFunction(ClipboardContent, AssignmentNode->GetFunctionName(), AssignmentNode->GetAssignmentTargets(), AssignmentNode->GetAssignmentDefaults(), GetStackNoteData());
 	}
 	else
 	{
 		checkf(FunctionCallNode->FunctionScript != nullptr, TEXT("Can't copy this module because it's script is invalid.  Call TestCanCopyWithMessage to check this."));
-		ClipboardFunction = UNiagaraClipboardFunction::CreateScriptFunction(ClipboardContent, FunctionCallNode->GetFunctionName(), FunctionCallNode->FunctionScript, FunctionCallNode->SelectedScriptVersion, FunctionCallNode->GetCustomNotes());
+		ClipboardFunction = UNiagaraClipboardFunction::CreateScriptFunction(ClipboardContent, FunctionCallNode->GetFunctionName(), FunctionCallNode->FunctionScript, FunctionCallNode->SelectedScriptVersion, GetStackNoteData());
 	}
 
 	ClipboardFunction->DisplayName = GetAlternateDisplayName().Get(FText::GetEmpty());
diff --git a/Engine/Plugins/FX/Niagara/Source/NiagaraEditor/Private/ViewModels/Stack/NiagaraStackNote.cpp b/Engine/Plugins/FX/Niagara/Source/NiagaraEditor/Private/ViewModels/Stack/NiagaraStackNote.cpp
new file mode 100644
--- /dev/null
+++ b/Engine/Plugins/FX/Niagara/Source/NiagaraEditor/Private/ViewModels/Stack/NiagaraStackNote.cpp
@@ -0,0 +1,104 @@
+﻿// Copyright Epic Games, Inc. All Rights Reserved.
+
+#include "ViewModels/Stack/NiagaraStackNote.h"
+#include "NiagaraMessages.h"
+#include "ScopedTransaction.h"
+
+#define LOCTEXT_NAMESPACE "UNiagaraStackNote"
+
+void UNiagaraStackNote::Initialize(FRequiredEntryData InRequiredEntryData, FString InTargetStackEntryKey)
+{
+	TargetStackEntryKey = InTargetStackEntryKey;
+	FString NoteStackEntryKey = GetTargetStackEntryKey();
+	NoteStackEntryKey.Append("-Note");
+	Super::Initialize(InRequiredEntryData, NoteStackEntryKey);
+}
+
+FString UNiagaraStackNote::GetTargetStackEntryKey() const
+{
+	return TargetStackEntryKey;
+}
+
+TOptional<FNiagaraStackNoteData> UNiagaraStackNote::GetTargetStackNoteData() const
+{
+	return GetStackEditorData().GetStackNote(GetTargetStackEntryKey());
+}
+
+void UNiagaraStackNote::ToggleInlineDisplay()
+{
+	FScopedTransaction Transaction(LOCTEXT("ToggleNoteDisplayTransaction", "Toggled Note Display"));
+	GetStackEditorData().Modify();
+
+	FNiagaraStackNoteData UpdatedMessage = GetTargetStackNoteData().GetValue();
+	UpdatedMessage.bInlineNote = !UpdatedMessage.bInlineNote;
+
+	OnNoteChangedDelegate.ExecuteIfBound(UpdatedMessage);
+}
+
+void UNiagaraStackNote::DeleteTargetStackNote()
+{
+	FScopedTransaction Transaction(LOCTEXT("DeleteNoteTransaction", "Deleted Note"));
+	GetStackEditorData().Modify();
+
+	GetStackEditorData().DeleteStackNote(GetTargetStackEntryKey());
+	
+	OnNoteChangedDelegate.ExecuteIfBound(FNiagaraStackNoteData());
+}
+
+bool UNiagaraStackNote::GetShouldShowInStack() const
+{
+	if (GetTargetStackNoteData().IsSet())
+	{
+		//  notes are not displayed in the stack in the sense that the stack note entries themselves won't be displayed.
+		return GetTargetStackNoteData().GetValue().bInlineNote == false;
+	}
+
+	return false;
+}
+
+UNiagaraStackEntry::EStackRowStyle UNiagaraStackNote::GetStackRowStyle() const
+{
+	return EStackRowStyle::ItemContentNote;
+}
+
+void UNiagaraStackNote::Copy(UNiagaraClipboardContent* ClipboardContent) const
+{
+	ClipboardContent->StackNote = GetTargetStackNoteData().GetValue();
+}
+
+void UNiagaraStackNote::Paste(const UNiagaraClipboardContent* ClipboardContent, FText& OutPasteWarning)
+{
+	if (ClipboardContent->StackNote.IsValid())
+	{
+		GetStackEditorData().Modify();
+
+		FNiagaraStackNoteData NoteData = ClipboardContent->StackNote;
+		GetStackEditorData().AddOrReplaceStackNote(GetTargetStackEntryKey(), NoteData);
+	}
+}
+
+bool UNiagaraStackNote::TestCanCopyWithMessage(FText& OutMessage) const
+{
+	if(GetTargetStackNoteData()->IsValid())
+	{
+		OutMessage = LOCTEXT("CanCopyNoteTest", "Copy the contents of this note");
+		return true;
+	}
+
+	OutMessage = LOCTEXT("CantCopyNoteTest", "Can not copy this note due to invalid content.");
+	return false;
+}
+
+bool UNiagaraStackNote::TestCanPasteWithMessage(const UNiagaraClipboardContent* ClipboardContent, FText& OutMessage) const
+{
+	if(ClipboardContent->StackNote.IsValid())
+	{
+		OutMessage = LOCTEXT("CanPasteNoteTest", "Paste the contents of a previously copied note");
+		return true;
+	}
+
+	OutMessage = LOCTEXT("CantPasteNoteTest", "Can not paste into this note. No valid note in clipboard.");
+	return false;
+}
+
+#undef LOCTEXT_NAMESPACE
diff --git a/Engine/Plugins/FX/Niagara/Source/NiagaraEditor/Private/ViewModels/Stack/NiagaraStackRenderItemGroup.cpp b/Engine/Plugins/FX/Niagara/Source/NiagaraEditor/Private/ViewModels/Stack/NiagaraStackRenderItemGroup.cpp
--- a/Engine/Plugins/FX/Niagara/Source/NiagaraEditor/Private/ViewModels/Stack/NiagaraStackRenderItemGroup.cpp
+++ b/Engine/Plugins/FX/Niagara/Source/NiagaraEditor/Private/ViewModels/Stack/NiagaraStackRenderItemGroup.cpp
@@ -158,11 +158,18 @@ void UNiagaraStackRenderItemGroup::Paste(const UNiagaraClipboardContent* Clipboa
 	if (EmitterWeak.IsValid())
 	{
 		UNiagaraEmitter* Emitter = EmitterWeak.Emitter.Get();
-		for (const UNiagaraRendererProperties* ClipboardRenderer : ClipboardContent->Renderers)
+		for (const UNiagaraClipboardRenderer* ClipboardRenderer : ClipboardContent->Renderers)
 		{
-			if (ClipboardRenderer != nullptr)
+			if (ClipboardRenderer != nullptr && ClipboardRenderer->RendererProperties != nullptr)
 			{
-				EmitterWeak.Emitter->AddRenderer(ClipboardRenderer->StaticDuplicateWithNewMergeId(Emitter), EmitterWeak.Version);
+				UNiagaraRendererProperties* NewRenderer = ClipboardRenderer->RendererProperties->StaticDuplicateWithNewMergeId(Emitter);
+				
+				if(ClipboardRenderer->StackNoteData.IsValid())
+				{
+					GetStackEditorData().AddOrReplaceStackNote(FNiagaraStackGraphUtilities::StackKeys::GenerateStackRendererEditorDataKey(*NewRenderer), ClipboardRenderer->StackNoteData);
+				}
+				
+				EmitterWeak.Emitter->AddRenderer(NewRenderer, EmitterWeak.Version);
 			}
 		}
 	}
diff --git a/Engine/Plugins/FX/Niagara/Source/NiagaraEditor/Private/ViewModels/Stack/NiagaraStackRendererItem.cpp b/Engine/Plugins/FX/Niagara/Source/NiagaraEditor/Private/ViewModels/Stack/NiagaraStackRendererItem.cpp
--- a/Engine/Plugins/FX/Niagara/Source/NiagaraEditor/Private/ViewModels/Stack/NiagaraStackRendererItem.cpp
+++ b/Engine/Plugins/FX/Niagara/Source/NiagaraEditor/Private/ViewModels/Stack/NiagaraStackRendererItem.cpp
@@ -39,8 +39,7 @@ UNiagaraStackRendererItem::UNiagaraStackRendererItem()
 void UNiagaraStackRendererItem::Initialize(FRequiredEntryData InRequiredEntryData, UNiagaraRendererProperties* InRendererProperties)
 {
 	checkf(RendererProperties.IsValid() == false, TEXT("Can not initialize more than once."));
-	FString RendererStackEditorDataKey = FString::Printf(TEXT("Renderer-%s"), *InRendererProperties->GetName());
-	Super::Initialize(InRequiredEntryData, RendererStackEditorDataKey);
+	Super::Initialize(InRequiredEntryData, FNiagaraStackGraphUtilities::StackKeys::GenerateStackRendererEditorDataKey(*InRendererProperties));
 	RendererProperties = InRendererProperties;
 	RendererProperties->OnChanged().AddUObject(this, &UNiagaraStackRendererItem::RendererChanged);
 }
@@ -318,7 +317,7 @@ bool UNiagaraStackRendererItem::TestCanCopyWithMessage(FText& OutMessage) const
 
 void UNiagaraStackRendererItem::Copy(UNiagaraClipboardContent* ClipboardContent) const
 {
-	ClipboardContent->Renderers.Add(CastChecked<UNiagaraRendererProperties>(StaticDuplicateObject(RendererProperties.Get(), ClipboardContent)));
+	ClipboardContent->Renderers.Add(UNiagaraClipboardRenderer::CreateRenderer(ClipboardContent, RendererProperties.Get(), GetStackNoteData()));
 }
 
 bool UNiagaraStackRendererItem::TestCanPasteWithMessage(const UNiagaraClipboardContent* ClipboardContent, FText& OutMessage) const
diff --git a/Engine/Plugins/FX/Niagara/Source/NiagaraEditor/Private/ViewModels/Stack/NiagaraStackScriptItemGroup.cpp b/Engine/Plugins/FX/Niagara/Source/NiagaraEditor/Private/ViewModels/Stack/NiagaraStackScriptItemGroup.cpp
--- a/Engine/Plugins/FX/Niagara/Source/NiagaraEditor/Private/ViewModels/Stack/NiagaraStackScriptItemGroup.cpp
+++ b/Engine/Plugins/FX/Niagara/Source/NiagaraEditor/Private/ViewModels/Stack/NiagaraStackScriptItemGroup.cpp
@@ -357,7 +357,7 @@ private:
 		FNiagaraStackGraphUtilities::GetStackFunctionInputs(*NewAssignmentModule, InputVariables, FNiagaraStackGraphUtilities::ENiagaraGetStackFunctionInputPinsOptions::AllInputs);
 		if (InputVariables.Num() == 1)
 		{
-			FString FunctionInputEditorDataKey = FNiagaraStackGraphUtilities::GenerateStackFunctionInputEditorDataKey(*NewAssignmentModule, InputVariables[0].GetName());
+			FString FunctionInputEditorDataKey = FNiagaraStackGraphUtilities::StackKeys::GenerateStackFunctionInputEditorDataKey(*NewAssignmentModule, InputVariables[0].GetName());
 			if (bRenameParameterOnAdd)
 			{
 				StackEditorData.SetStackEntryIsRenamePending(FunctionInputEditorDataKey, true);
@@ -1406,9 +1406,10 @@ void UNiagaraStackScriptItemGroup::PasteModules(const UNiagaraClipboardContent*
 
 			if (NewFunctionCallNode != nullptr)
 			{
-				for(const auto& Message : ClipboardFunction->Messages)
+				if(ClipboardFunction->StackNoteData.IsValid())
 				{
-					NewFunctionCallNode->AddCustomNote(Message);	
+					FString NewModuleStackEditorDataKey = FNiagaraStackGraphUtilities::StackKeys::GenerateStackModuleEditorDataKey(*NewFunctionCallNode);
+					GetStackEditorData().AddOrReplaceStackNote(NewModuleStackEditorDataKey, ClipboardFunction->StackNoteData);	
 				}
 				
 				ClipboardFunction->OnPastedFunctionCallNodeDelegate.ExecuteIfBound(NewFunctionCallNode);
diff --git a/Engine/Plugins/FX/Niagara/Source/NiagaraEditor/Public/NiagaraClipboard.h b/Engine/Plugins/FX/Niagara/Source/NiagaraEditor/Public/NiagaraClipboard.h
--- a/Engine/Plugins/FX/Niagara/Source/NiagaraEditor/Public/NiagaraClipboard.h
+++ b/Engine/Plugins/FX/Niagara/Source/NiagaraEditor/Public/NiagaraClipboard.h
@@ -89,6 +89,21 @@ public:
 
 };
 
+UCLASS()
+class UNiagaraClipboardRenderer : public UObject
+{
+	GENERATED_BODY()
+
+public:
+	static NIAGARAEDITOR_API const UNiagaraClipboardRenderer* CreateRenderer(UObject* InOuter, UNiagaraRendererProperties* Renderer, TOptional<FNiagaraStackNoteData> StackNoteData = TOptional<FNiagaraStackNoteData>());
+
+	UPROPERTY()
+	TObjectPtr<UNiagaraRendererProperties> RendererProperties;
+
+	UPROPERTY()
+	FNiagaraStackNoteData StackNoteData;
+};
+
 UENUM()
 enum class ENiagaraClipboardFunctionScriptMode
 {
@@ -104,9 +119,9 @@ class UNiagaraClipboardFunction : public UObject
 public:
 	DECLARE_DYNAMIC_DELEGATE_OneParam(FOnPastedFunctionCallNode, UNiagaraNodeFunctionCall*, PastedFunctionCall);
 
-	static NIAGARAEDITOR_API UNiagaraClipboardFunction* CreateScriptFunction(UObject* InOuter, FString InFunctionName, UNiagaraScript* InScript, const FGuid& InScriptVersion = FGuid(), const TArray<FNiagaraStackMessage> InStackMessages = TArray<FNiagaraStackMessage>());
+	static NIAGARAEDITOR_API UNiagaraClipboardFunction* CreateScriptFunction(UObject* InOuter, FString InFunctionName, UNiagaraScript* InScript, const FGuid& InScriptVersion = FGuid(), const TOptional<FNiagaraStackNoteData>& InStackNote = TOptional<FNiagaraStackNoteData>());
 
-	static NIAGARAEDITOR_API UNiagaraClipboardFunction* CreateAssignmentFunction(UObject* InOuter, FString InFunctionName, const TArray<FNiagaraVariable>& InAssignmentTargets, const TArray<FString>& InAssignmentDefaults);
+	static NIAGARAEDITOR_API UNiagaraClipboardFunction* CreateAssignmentFunction(UObject* InOuter, FString InFunctionName, const TArray<FNiagaraVariable>& InAssignmentTargets, const TArray<FString>& InAssignmentDefaults, TOptional<FNiagaraStackNoteData> InStackNoteData = TOptional<FNiagaraStackNoteData>());
 
 	UPROPERTY()
 	FString FunctionName;
@@ -134,9 +149,9 @@ public:
 
 	UPROPERTY()
 	FGuid ScriptVersion;
-
+	
 	UPROPERTY()
-	TArray<FNiagaraStackMessage> Messages;
+	FNiagaraStackNoteData StackNoteData;
 };
 
 USTRUCT()
@@ -179,7 +194,7 @@ public:
 	TArray<TObjectPtr<const UNiagaraClipboardFunctionInput>> FunctionInputs;
 
 	UPROPERTY()
-	TArray<TObjectPtr<const UNiagaraRendererProperties>> Renderers;
+	TArray<TObjectPtr<const UNiagaraClipboardRenderer>> Renderers;
 
 	UPROPERTY()
 	TArray<TObjectPtr<const UNiagaraScript>> Scripts;
@@ -194,6 +209,9 @@ public:
 	/** Markup MetaData to specify that if scripts are pasted from the clipboard to automatically fixup their order in the stack to satisfy dependencies. */
 	UPROPERTY()
 	mutable bool bFixupPasteIndexForScriptDependenciesInStack = false;
+
+	UPROPERTY()
+	FNiagaraStackNoteData StackNote;
 };
 
 class FNiagaraClipboard
diff --git a/Engine/Plugins/FX/Niagara/Source/NiagaraEditor/Public/NiagaraEditorUtilities.h b/Engine/Plugins/FX/Niagara/Source/NiagaraEditor/Public/NiagaraEditorUtilities.h
--- a/Engine/Plugins/FX/Niagara/Source/NiagaraEditor/Public/NiagaraEditorUtilities.h
+++ b/Engine/Plugins/FX/Niagara/Source/NiagaraEditor/Public/NiagaraEditorUtilities.h
@@ -48,6 +48,8 @@ struct FNiagaraScriptVersionUpgradeContext;
 struct FNiagaraScriptConversionContext;
 class UUpgradeNiagaraEmitterContext;
 struct FNiagaraMessageStore;
+class UNiagaraStackNote;
+class SToolTip;
 
 enum class ENiagaraFunctionDebugState : uint8;
 
@@ -414,6 +416,11 @@ namespace FNiagaraEditorUtilities
 
 	NIAGARAEDITOR_API UNiagaraDataInterface* GetResolvedRuntimeInstanceForEditorDataInterfaceInstance(const UNiagaraSystem& OwningSystem, UNiagaraDataInterface& EditorDataInterfaceInstance);
 
+	namespace Tooltips
+	{
+		NIAGARAEDITOR_API TSharedRef<SToolTip> CreateStackNoteTooltip(UNiagaraStackNote& StackNote);
+	}
+	
 	namespace Scripts
 	{
 		namespace Validation
diff --git a/Engine/Plugins/FX/Niagara/Source/NiagaraEditor/Public/NiagaraEmitterEditorData.h b/Engine/Plugins/FX/Niagara/Source/NiagaraEditor/Public/NiagaraEmitterEditorData.h
--- a/Engine/Plugins/FX/Niagara/Source/NiagaraEditor/Public/NiagaraEmitterEditorData.h
+++ b/Engine/Plugins/FX/Niagara/Source/NiagaraEditor/Public/NiagaraEmitterEditorData.h
@@ -74,6 +74,7 @@ private:
 	void StackEditorDataChanged();
 
 	void PostLoad_TransferSummaryDataToNewFormat();
+	void PostLoad_TransferModuleStackNotesToNewFormat(UObject* Owner);
 };
 
 
diff --git a/Engine/Plugins/FX/Niagara/Source/NiagaraEditor/Public/NiagaraMessages.h b/Engine/Plugins/FX/Niagara/Source/NiagaraEditor/Public/NiagaraMessages.h
--- a/Engine/Plugins/FX/Niagara/Source/NiagaraEditor/Public/NiagaraMessages.h
+++ b/Engine/Plugins/FX/Niagara/Source/NiagaraEditor/Public/NiagaraMessages.h
@@ -25,8 +25,7 @@ enum class ENiagaraMessageSeverity : uint8
 	Error = 1,
 	PerformanceWarning = 2,
 	Warning = 3,
-	Info = 4,
-	CustomNote = 5 // Should be last
+	Info = 4
 };
 
 //Struct for passing around script asset info from compile event message job to message types
@@ -253,7 +252,7 @@ struct FNiagaraStackMessage
 	FText ShortDescription;
 
 	UPROPERTY()
-	ENiagaraMessageSeverity MessageSeverity = ENiagaraMessageSeverity::CustomNote;
+	ENiagaraMessageSeverity MessageSeverity = ENiagaraMessageSeverity::Info;
 
 	UPROPERTY()
 	bool bAllowDismissal = true;
@@ -262,6 +261,32 @@ struct FNiagaraStackMessage
 	FGuid Guid;
 };
 
+USTRUCT()
+struct FNiagaraStackNoteData
+{
+	GENERATED_BODY()
+
+	UPROPERTY()
+	FText MessageHeader;
+	
+	UPROPERTY()
+	FText Message;
+
+	/** If true, this note should be used inline for a stack entry, instead of adding additional rows. */
+	UPROPERTY()
+	bool bInlineNote = false;
+
+	bool operator==(const FNiagaraStackNoteData& Other) const
+	{
+		return MessageHeader.EqualTo(Other.MessageHeader) && Message.EqualTo(Other.Message); 
+	}
+
+	bool IsValid() const
+	{
+		return MessageHeader.IsEmpty() == false || Message.IsEmpty() == false;
+	}
+};
+
 struct FGenerateNiagaraMessageInfo
 {
 public:
diff --git a/Engine/Plugins/FX/Niagara/Source/NiagaraEditor/Public/NiagaraNodeFunctionCall.h b/Engine/Plugins/FX/Niagara/Source/NiagaraEditor/Public/NiagaraNodeFunctionCall.h
--- a/Engine/Plugins/FX/Niagara/Source/NiagaraEditor/Public/NiagaraNodeFunctionCall.h
+++ b/Engine/Plugins/FX/Niagara/Source/NiagaraEditor/Public/NiagaraNodeFunctionCall.h
@@ -104,7 +104,7 @@ public:
 	virtual void PostLoad()override;
 	virtual void PostEditChangeProperty(struct FPropertyChangedEvent& PropertyChangedEvent) override;
 	//End UObject interface
-
+	
 	//~ Begin UNiagaraNode Interface
 	virtual void Compile(FTranslator* Translator, TArray<int32>& Outputs) const override;
 	virtual UObject* GetReferencedAsset() const override;
@@ -180,12 +180,9 @@ public:
 	// Messages API
 	FNiagaraMessageStore& GetMessageStore() { return MessageStore; }
 
-	// Custom Notes API
-	const TArray<FNiagaraStackMessage>& GetCustomNotes() const { return StackMessages; };
-	NIAGARAEDITOR_API void AddCustomNote(const FNiagaraStackMessage& StackMessage);
-	NIAGARAEDITOR_API void RemoveCustomNote(const FGuid& MessageKey);
-	FSimpleDelegate& OnCustomNotesChanged() { return OnCustomNotesChangedDelegate; }
-	void RemoveCustomNoteViaDelegate(const FGuid MessageKey);
+	// Custom Notes API - Deprecated. Used for transferring to stack editor data
+	const TArray<FNiagaraStackMessage>& GetDeprecatedCustomNotes() const { return StackMessages; }
+	
 	TArray<FGuid> GetBoundPinGuidsByName(FName InputName) const;
 
 	/** Adds a static switch pin to this function call node by variable id and sets it's default value using the supplied data and marks it as
@@ -254,7 +251,5 @@ protected:
 	TMap<FGuid, FName> BoundPinNames;
 
 	FOnInputsChanged OnInputsChangedDelegate;
-
-	FSimpleDelegate OnCustomNotesChangedDelegate;
 };
 
diff --git a/Engine/Plugins/FX/Niagara/Source/NiagaraEditor/Public/NiagaraStackEditorData.h b/Engine/Plugins/FX/Niagara/Source/NiagaraEditor/Public/NiagaraStackEditorData.h
--- a/Engine/Plugins/FX/Niagara/Source/NiagaraEditor/Public/NiagaraStackEditorData.h
+++ b/Engine/Plugins/FX/Niagara/Source/NiagaraEditor/Public/NiagaraStackEditorData.h
@@ -3,6 +3,8 @@
 #pragma once
 
 #include "NiagaraEditorDataBase.h"
+#include "NiagaraMessages.h"
+#include "NiagaraNodeFunctionCall.h"
 
 #include "NiagaraStackEditorData.generated.h"
 struct FStackIssue;
@@ -80,6 +82,14 @@ public:
 	 */
 	void SetStackEntryWasExpandedPreSearch(const FString& StackEntryKey, bool bWasExpandedPreSearch);
 
+	const TMap<FString, FNiagaraStackNoteData>& GetAllStackNotes();
+	NIAGARAEDITOR_API TOptional<FNiagaraStackNoteData> GetStackNote(const FString& StackEntryKey);
+	NIAGARAEDITOR_API bool HasStackNote(const FString& StackEntryKey);
+	NIAGARAEDITOR_API void AddOrReplaceStackNote(const FString& StackEntryKey, FNiagaraStackNoteData NewStackMessage, bool bBroadcast = true);
+	NIAGARAEDITOR_API void DeleteStackNote(const FString& StackEntryKey);
+	
+	void TransferDeprecatedStackNotes(const UNiagaraNodeFunctionCall& FunctionCallNode);
+
 	/*
 	 * Gets whether or not this entry should display all of its content inline. 
 	 */
@@ -198,6 +208,9 @@ private:
 
 	TMap<FString, FText> StackEntryKeyToActiveSectionMap;
 
+	UPROPERTY()
+	TMap<FString, FNiagaraStackNoteData> StackNotes;
+
 	/* Marking those FTexts explicitly as editoronly_data will make localization not pick these up.
 	 * This is a workaround. EditorDataBase in system & emitter is already flagged as editor only, but it doesn't propagate properly */
 #if WITH_EDITORONLY_DATA
diff --git a/Engine/Plugins/FX/Niagara/Source/NiagaraEditor/Public/NiagaraSystemEditorData.h b/Engine/Plugins/FX/Niagara/Source/NiagaraEditor/Public/NiagaraSystemEditorData.h
--- a/Engine/Plugins/FX/Niagara/Source/NiagaraEditor/Public/NiagaraSystemEditorData.h
+++ b/Engine/Plugins/FX/Niagara/Source/NiagaraEditor/Public/NiagaraSystemEditorData.h
@@ -129,7 +129,7 @@ public:
 	
 private:
 	NIAGARAEDITOR_API void UpdatePlaybackRangeFromEmitters(UNiagaraSystem& OwnerSystem);
-
+	void PostLoad_TransferModuleStackNotesToNewFormat(UObject* InOwner);
 private:
 	UPROPERTY(Instanced)
 	TObjectPtr<UNiagaraSystemEditorFolder> RootFolder;
diff --git a/Engine/Plugins/FX/Niagara/Source/NiagaraEditor/Public/ViewModels/Stack/NiagaraStackClipboardUtilities.h b/Engine/Plugins/FX/Niagara/Source/NiagaraEditor/Public/ViewModels/Stack/NiagaraStackClipboardUtilities.h
--- a/Engine/Plugins/FX/Niagara/Source/NiagaraEditor/Public/ViewModels/Stack/NiagaraStackClipboardUtilities.h
+++ b/Engine/Plugins/FX/Niagara/Source/NiagaraEditor/Public/ViewModels/Stack/NiagaraStackClipboardUtilities.h
@@ -22,5 +22,7 @@ namespace FNiagaraStackClipboardUtilities
 
 	NIAGARAEDITOR_API bool TestCanDeleteSelectionWithMessage(const TArray<UNiagaraStackEntry*>& SelectedEntries, FText& OutMessage);
 
+	NIAGARAEDITOR_API void CopyNote(const UNiagaraStackEntry* StackEntry);
+	
 	NIAGARAEDITOR_API void DeleteSelection(const TArray<UNiagaraStackEntry*>& SelectedEntries);
 }
\ No newline at end of file
diff --git a/Engine/Plugins/FX/Niagara/Source/NiagaraEditor/Public/ViewModels/Stack/NiagaraStackEntry.h b/Engine/Plugins/FX/Niagara/Source/NiagaraEditor/Public/ViewModels/Stack/NiagaraStackEntry.h
--- a/Engine/Plugins/FX/Niagara/Source/NiagaraEditor/Public/ViewModels/Stack/NiagaraStackEntry.h
+++ b/Engine/Plugins/FX/Niagara/Source/NiagaraEditor/Public/ViewModels/Stack/NiagaraStackEntry.h
@@ -13,8 +13,10 @@ class FNiagaraEmitterViewModel;
 class FNiagaraScriptViewModel;
 class UNiagaraStackEditorData;
 class UNiagaraStackErrorItem;
+class UNiagaraStackNote;
 class UNiagaraClipboardContent;
 class UNiagaraStackEntry;
+struct FNiagaraStackNoteData;
 
 UENUM()
 enum class EStackIssueSeverity : uint8
@@ -22,7 +24,6 @@ enum class EStackIssueSeverity : uint8
 	Error = 0,
 	Warning, 
 	Info,
-	CustomNote,
 	None
 };
 
@@ -393,6 +394,8 @@ public:
 		}
 	}
 
+	NIAGARAEDITOR_API UNiagaraStackNote* GetStackNote();
+	
 	NIAGARAEDITOR_API FOnExpansionChanged& OnExpansionChanged();
 
 	NIAGARAEDITOR_API FOnExpansionChanged& OnExpansionInOverviewChanged();
@@ -460,8 +463,6 @@ public:
 	NIAGARAEDITOR_API bool HasIssuesOrAnyChildHasIssues() const;
 	NIAGARAEDITOR_API bool HasUsagesOrAnyChildHasUsages() const;
 	NIAGARAEDITOR_API void GetRecursiveUsages(bool& bRead, bool& bWrite) const;
-
-	NIAGARAEDITOR_API int32 GetTotalNumberOfCustomNotes() const;
 	
 	NIAGARAEDITOR_API int32 GetTotalNumberOfInfoIssues() const;
 
@@ -522,6 +523,13 @@ public:
 	/** Handler for when a rename is committed for this stack entry. */
 	NIAGARAEDITOR_API virtual void OnRenamed(FText NewName);
 
+	/** Generally, all stack items with a valid stack key would support the note feature. We only allow them on specific types of of entries. */
+	NIAGARAEDITOR_API virtual bool SupportsStackNotes() { return false; }
+	NIAGARAEDITOR_API bool HasStackNoteData() const;
+	NIAGARAEDITOR_API FNiagaraStackNoteData GetStackNoteData() const;
+	NIAGARAEDITOR_API void SetStackNoteData(FNiagaraStackNoteData InStackNoteData);
+	NIAGARAEDITOR_API void DeleteStackNoteData();
+	
 	virtual bool SupportsSummaryView() const { return false; }
 	NIAGARAEDITOR_API virtual struct FNiagaraHierarchyIdentity DetermineSummaryIdentity() const;
 	NIAGARAEDITOR_API bool IsInSummaryView() const;
@@ -610,16 +618,14 @@ private:
 			: TotalNumberOfInfoIssues(0)
 			, TotalNumberOfWarningIssues(0)
 			, TotalNumberOfErrorIssues(0)
-			, TotalNumberOfCustomNotes(0)
 		{
 		}
 
-		bool HasAnyIssues() const { return TotalNumberOfInfoIssues > 0 || TotalNumberOfWarningIssues > 0 || TotalNumberOfErrorIssues > 0 || TotalNumberOfCustomNotes > 0; }
+		bool HasAnyIssues() const { return TotalNumberOfInfoIssues > 0 || TotalNumberOfWarningIssues > 0 || TotalNumberOfErrorIssues > 0; }
 
 		int32 TotalNumberOfInfoIssues;
 		int32 TotalNumberOfWarningIssues;
 		int32 TotalNumberOfErrorIssues;
-		int32 TotalNumberOfCustomNotes;
 		TArray<UNiagaraStackEntry*> ChildrenWithIssues;
 	};
 
@@ -657,6 +663,9 @@ private:
 
 	mutable TArray<UNiagaraStackEntry*> FilteredChildren;
 
+	UPROPERTY()
+	TObjectPtr<UNiagaraStackNote> StackNote;
+	
 	UPROPERTY()
 	TArray<TObjectPtr<UNiagaraStackErrorItem>> ErrorChildren;
 
diff --git a/Engine/Plugins/FX/Niagara/Source/NiagaraEditor/Public/ViewModels/Stack/NiagaraStackFunctionInput.h b/Engine/Plugins/FX/Niagara/Source/NiagaraEditor/Public/ViewModels/Stack/NiagaraStackFunctionInput.h
--- a/Engine/Plugins/FX/Niagara/Source/NiagaraEditor/Public/ViewModels/Stack/NiagaraStackFunctionInput.h
+++ b/Engine/Plugins/FX/Niagara/Source/NiagaraEditor/Public/ViewModels/Stack/NiagaraStackFunctionInput.h
@@ -105,6 +105,7 @@ public:
 	NIAGARAEDITOR_API virtual FText GetTooltipText() const override;
 	NIAGARAEDITOR_API virtual bool GetIsEnabled() const override;
 	NIAGARAEDITOR_API virtual UObject* GetExternalAsset() const override;
+	virtual bool SupportsStackNotes() override { return true; }
 	virtual bool SupportsCut() const override { return true; }
 	NIAGARAEDITOR_API virtual bool TestCanCutWithMessage(FText& OutMessage) const override;
 	NIAGARAEDITOR_API virtual FText GetCutTransactionText() const override;
diff --git a/Engine/Plugins/FX/Niagara/Source/NiagaraEditor/Public/ViewModels/Stack/NiagaraStackGraphUtilities.h b/Engine/Plugins/FX/Niagara/Source/NiagaraEditor/Public/ViewModels/Stack/NiagaraStackGraphUtilities.h
--- a/Engine/Plugins/FX/Niagara/Source/NiagaraEditor/Public/ViewModels/Stack/NiagaraStackGraphUtilities.h
+++ b/Engine/Plugins/FX/Niagara/Source/NiagaraEditor/Public/ViewModels/Stack/NiagaraStackGraphUtilities.h
@@ -35,6 +35,7 @@ struct FNiagaraEventScriptProperties;
 struct FNiagaraStackModuleData;
 struct FNiagaraModuleDependency;
 struct FNiagaraEmitterHandle;
+class UNiagaraRendererProperties;
 
 namespace FNiagaraStackGraphUtilities
 {
@@ -88,9 +89,13 @@ namespace FNiagaraStackGraphUtilities
 
 	void InitializeStackFunctionInput(TSharedRef<FNiagaraSystemViewModel> SystemViewModel, TSharedPtr<FNiagaraEmitterViewModel> EmitterViewModel, UNiagaraStackEditorData& StackEditorData, UNiagaraNodeFunctionCall& ModuleNode, UNiagaraNodeFunctionCall& InputFunctionCallNode, FName InputName);
 
-	FString GenerateStackFunctionInputEditorDataKey(UNiagaraNodeFunctionCall& FunctionCallNode, FNiagaraParameterHandle InputParameterHandle);
+	namespace StackKeys
+	{
+		FString GenerateStackFunctionInputEditorDataKey(const UNiagaraNodeFunctionCall& FunctionCallNode, FNiagaraParameterHandle InputParameterHandle);
+		FString GenerateStackModuleEditorDataKey(const UNiagaraNodeFunctionCall& ModuleNode);
+		FString GenerateStackRendererEditorDataKey(const UNiagaraRendererProperties& Renderer);
+	}
 
-	FString GenerateStackModuleEditorDataKey(UNiagaraNodeFunctionCall& ModuleNode);
 
 	TArray<FName> StackContextResolution(FVersionedNiagaraEmitter OwningEmitter, UNiagaraNodeOutput* OutputNodeInChain);
 	void BuildParameterMapHistoryWithStackContextResolution(FVersionedNiagaraEmitter OwningEmitter, UNiagaraNodeOutput* OutputNodeInChain, UNiagaraNode* NodeToVisit, TArray<FNiagaraParameterMapHistory>& OutHistories, bool bRecursive = true, bool bFilterForCompilation = true);
diff --git a/Engine/Plugins/FX/Niagara/Source/NiagaraEditor/Public/ViewModels/Stack/NiagaraStackItem.h b/Engine/Plugins/FX/Niagara/Source/NiagaraEditor/Public/ViewModels/Stack/NiagaraStackItem.h
--- a/Engine/Plugins/FX/Niagara/Source/NiagaraEditor/Public/ViewModels/Stack/NiagaraStackItem.h
+++ b/Engine/Plugins/FX/Niagara/Source/NiagaraEditor/Public/ViewModels/Stack/NiagaraStackItem.h
@@ -38,6 +38,8 @@ public:
 	virtual bool TestCanResetToBaseWithMessage(FText& OutCanResetToBaseMessage) const { return false; }
 	virtual void ResetToBase() { }
 
+	virtual bool SupportsStackNotes() override { return true; }
+	
 	virtual bool SupportsEditMode() const { return false; }
 	virtual void OnEditButtonClicked() { }
 	virtual TOptional<FText> GetEditModeButtonText() const { return TOptional<FText>(); }
diff --git a/Engine/Plugins/FX/Niagara/Source/NiagaraEditor/Public/ViewModels/Stack/NiagaraStackModuleItem.h b/Engine/Plugins/FX/Niagara/Source/NiagaraEditor/Public/ViewModels/Stack/NiagaraStackModuleItem.h
--- a/Engine/Plugins/FX/Niagara/Source/NiagaraEditor/Public/ViewModels/Stack/NiagaraStackModuleItem.h
+++ b/Engine/Plugins/FX/Niagara/Source/NiagaraEditor/Public/ViewModels/Stack/NiagaraStackModuleItem.h
@@ -114,22 +114,9 @@ public:
 		SetIsEnabledInternal(bEnabled);
 	}
 
-	void SetNoteMode(bool bEnabled)
-	{
-		bIsNoteModeActive = bEnabled;
-		OnNoteModeSet().ExecuteIfBound(bIsNoteModeActive);
-	}
-
-	bool GetNoteMode() const
-	{
-		return bIsNoteModeActive;
-	}
-
 	NIAGARAEDITOR_API bool IsDebugDrawEnabled() const;
 	NIAGARAEDITOR_API void SetDebugDrawEnabled(bool bInEnabled);
 
-	FOnNoteModeSet& OnNoteModeSet() { return OnNoteModeSetDelegate;}
-
 	NIAGARAEDITOR_API bool OpenSourceAsset() const;
 
 protected:
@@ -174,8 +161,6 @@ private:
 	INiagaraStackItemGroupAddUtilities* GroupAddUtilities;
 
 	mutable TOptional<bool> bIsScratchModuleCache;
-
-	bool bIsNoteModeActive = false;
 	
 	bool bIsModuleScriptReassignmentPending;
 
@@ -185,6 +170,4 @@ private:
 	TArray<FStackIssue> MessageManagerIssues;
 
 	FGuid MessageLogGuid;
-	
-	FOnNoteModeSet OnNoteModeSetDelegate;
 };
diff --git a/Engine/Plugins/FX/Niagara/Source/NiagaraEditor/Public/ViewModels/Stack/NiagaraStackNote.h b/Engine/Plugins/FX/Niagara/Source/NiagaraEditor/Public/ViewModels/Stack/NiagaraStackNote.h
new file mode 100644
--- /dev/null
+++ b/Engine/Plugins/FX/Niagara/Source/NiagaraEditor/Public/ViewModels/Stack/NiagaraStackNote.h
@@ -0,0 +1,52 @@
+﻿// Copyright Epic Games, Inc. All Rights Reserved.
+
+#pragma once
+
+#include "CoreMinimal.h"
+#include "NiagaraClipboard.h"
+#include "NiagaraMessages.h"
+#include "NiagaraStackEditorData.h"
+#include "NiagaraStackEntry.h"
+#include "NiagaraStackNote.generated.h"
+
+UCLASS(MinimalAPI)
+class UNiagaraStackNote : public UNiagaraStackEntry
+{
+	DECLARE_DELEGATE_OneParam(FOnNoteChanged, FNiagaraStackNoteData)
+	
+	GENERATED_BODY()
+public:
+	void Initialize(FRequiredEntryData InRequiredEntryData, FString InTargetStackEntryKey);
+
+	NIAGARAEDITOR_API FString GetTargetStackEntryKey() const;
+	NIAGARAEDITOR_API TOptional<FNiagaraStackNoteData> GetTargetStackNoteData() const;
+
+	void RequestEditHeader() const { OnRequestEditHeaderDelegate.ExecuteIfBound(); }
+	NIAGARAEDITOR_API void ToggleInlineDisplay();
+	NIAGARAEDITOR_API void DeleteTargetStackNote();
+	
+	FSimpleDelegate& OnRequestEditHeader() { return OnRequestEditHeaderDelegate; }
+
+	/** The delegate for the target stack entry to bind to. Forwards note changes which then need to be set via UNiagaraStackEntry::SetStackNoteData. */
+	FOnNoteChanged& OnNoteChanged() { return OnNoteChangedDelegate; }
+
+	virtual bool GetShouldShowInStack() const override;
+protected:
+	virtual bool GetShouldShowInOverview() const override { return false; }
+	virtual EStackRowStyle GetStackRowStyle() const override;
+	
+	virtual bool SupportsRename() const override { return true; }
+	virtual bool SupportsCopy() const override { return true; }
+	virtual bool SupportsPaste() const override { return true; }
+	
+	virtual void Copy(UNiagaraClipboardContent* ClipboardContent) const override;
+	virtual void Paste(const UNiagaraClipboardContent* ClipboardContent, FText& OutPasteWarning) override;
+
+	virtual bool TestCanCopyWithMessage(FText& OutMessage) const override;
+	virtual bool TestCanPasteWithMessage(const UNiagaraClipboardContent* ClipboardContent, FText& OutMessage) const override;
+	
+	FSimpleDelegate OnRequestEditHeaderDelegate;
+	FOnNoteChanged OnNoteChangedDelegate;
+
+	FString TargetStackEntryKey;
+};
diff --git a/Engine/Plugins/FX/Niagara/Source/NiagaraEditor/Public/ViewModels/Stack/NiagaraStackPropertyRow.h b/Engine/Plugins/FX/Niagara/Source/NiagaraEditor/Public/ViewModels/Stack/NiagaraStackPropertyRow.h
--- a/Engine/Plugins/FX/Niagara/Source/NiagaraEditor/Public/ViewModels/Stack/NiagaraStackPropertyRow.h
+++ b/Engine/Plugins/FX/Niagara/Source/NiagaraEditor/Public/ViewModels/Stack/NiagaraStackPropertyRow.h
@@ -26,6 +26,8 @@ public:
 
 	NIAGARAEDITOR_API virtual bool IsExpandedByDefault() const override;
 
+	virtual bool SupportsStackNotes() override { return true; }
+	
 	NIAGARAEDITOR_API virtual bool CanDrag() const override;
 
 	void SetOwnerGuid(TOptional<FGuid> InGuid) { OwnerGuid = InGuid; }
diff --git a/Engine/Plugins/FX/Niagara/Source/NiagaraEditor/Public/ViewModels/Stack/NiagaraStackRendererItem.h b/Engine/Plugins/FX/Niagara/Source/NiagaraEditor/Public/ViewModels/Stack/NiagaraStackRendererItem.h
--- a/Engine/Plugins/FX/Niagara/Source/NiagaraEditor/Public/ViewModels/Stack/NiagaraStackRendererItem.h
+++ b/Engine/Plugins/FX/Niagara/Source/NiagaraEditor/Public/ViewModels/Stack/NiagaraStackRendererItem.h
@@ -56,6 +56,8 @@ public:
 	
 	virtual bool SupportsRename() const override { return true; }
 
+	virtual bool SupportsStackNotes() override { return true; }
+
 	NIAGARAEDITOR_API bool HasBaseRenderer() const;
 
 	virtual bool SupportsChangeEnabled() const override { return true; }
diff --git a/Engine/Plugins/FX/Niagara/Source/NiagaraEditor/Public/ViewModels/Stack/NiagaraStackSimulationStageGroup.h b/Engine/Plugins/FX/Niagara/Source/NiagaraEditor/Public/ViewModels/Stack/NiagaraStackSimulationStageGroup.h
--- a/Engine/Plugins/FX/Niagara/Source/NiagaraEditor/Public/ViewModels/Stack/NiagaraStackSimulationStageGroup.h
+++ b/Engine/Plugins/FX/Niagara/Source/NiagaraEditor/Public/ViewModels/Stack/NiagaraStackSimulationStageGroup.h
@@ -80,6 +80,8 @@ public:
 	NIAGARAEDITOR_API virtual bool GetIsInherited() const override;
 	NIAGARAEDITOR_API virtual FText GetInheritanceMessage() const override;
 
+	virtual bool SupportsStackNotes() override { return true; }
+	
 	virtual bool CanDrag() const override { return true; }
 
 	NIAGARAEDITOR_API virtual bool GetIsEnabled() const override;
diff --git a/Engine/Plugins/FX/Niagara/Source/NiagaraEditorWidgets/Private/NiagaraEditorWidgetsUtilities.cpp b/Engine/Plugins/FX/Niagara/Source/NiagaraEditorWidgets/Private/NiagaraEditorWidgetsUtilities.cpp
--- a/Engine/Plugins/FX/Niagara/Source/NiagaraEditorWidgets/Private/NiagaraEditorWidgetsUtilities.cpp
+++ b/Engine/Plugins/FX/Niagara/Source/NiagaraEditorWidgets/Private/NiagaraEditorWidgetsUtilities.cpp
@@ -353,12 +353,6 @@ void ToggleShouldDebugDraw(TWeakObjectPtr<UNiagaraStackItem> StackItemWeak)
 	}
 }
 
-void EnableNoteMode(TWeakObjectPtr<UNiagaraStackModuleItem> ModuleItem)
-{
-	ModuleItem->GetSystemViewModel()->GetSelectionViewModel()->UpdateSelectedEntries({ModuleItem.Get()}, {}, true);
-	ModuleItem->SetNoteMode(true);
-}
-
 bool FNiagaraStackEditorWidgetsUtilities::AddStackItemContextMenuActions(FMenuBuilder& MenuBuilder, UNiagaraStackItem& StackItem)
 {
 	if (StackItem.SupportsChangeEnabled())
@@ -399,12 +393,6 @@ bool FNiagaraStackEditorWidgetsUtilities::AddStackItemContextMenuActions(FMenuBu
 						NAME_None,
 						EUserInterfaceActionType::Check);
 				}
-
-				MenuBuilder.AddMenuEntry(
-					LOCTEXT("AddNote", "Add Note"),
-					LOCTEXT("AddNoteToolTip", "Add a note to this module item."),
-					FSlateIcon(),
-					FUIAction(FExecuteAction::CreateStatic(&EnableNoteMode, TWeakObjectPtr<UNiagaraStackModuleItem>(ModuleItem))));
 			}
 
 			if (UNiagaraStackRendererItem* RendererItem = Cast<UNiagaraStackRendererItem>(&StackItem))
diff --git a/Engine/Plugins/FX/Niagara/Source/NiagaraEditorWidgets/Private/SNiagaraOverviewStack.cpp b/Engine/Plugins/FX/Niagara/Source/NiagaraEditorWidgets/Private/SNiagaraOverviewStack.cpp
--- a/Engine/Plugins/FX/Niagara/Source/NiagaraEditorWidgets/Private/SNiagaraOverviewStack.cpp
+++ b/Engine/Plugins/FX/Niagara/Source/NiagaraEditorWidgets/Private/SNiagaraOverviewStack.cpp
@@ -52,6 +52,7 @@
 #include "NiagaraEditorStyle.h"
 #include "NiagaraEditorUtilities.h"
 #include "ScopedTransaction.h"
+#include "SNiagaraStackNote.h"
 #include "Styling/StyleColors.h"
 #include "ViewModels/Stack/NiagaraStackRendererItem.h"
 #include "ViewModels/Stack/NiagaraStackSystemSettingsGroup.h"
@@ -320,7 +321,16 @@ private:
 	{
 		return StackEntry->GetIsExpandedInOverview() ? ExpandedImage : CollapsedImage;
 	}
-	
+
+	EVisibility GetNoteVisibility() const
+	{
+		if(StackEntry->GetStackNote() == nullptr)
+		{
+			return EVisibility::Collapsed;
+		}
+		
+		return EVisibility::Visible;
+	}
 
 	uint32 GetOverviewChildrenCount() const
 	{
@@ -362,10 +372,6 @@ private:
 			{
 				IssueHighlightColor = FStyleColors::AccentBlue.GetSpecifiedColor().Desaturate(DesaturateAmount);
 			}
-			else if (StackEntry->GetTotalNumberOfCustomNotes() > 0)
-			{
-				IssueHighlightColor = FStyleColors::AccentWhite.GetSpecifiedColor().Desaturate(DesaturateAmount);
-			}
 			else
 			{
 				IssueHighlightColor = FStyleColors::Transparent;
@@ -977,10 +983,33 @@ TSharedRef<ITableRow> SNiagaraOverviewStack::OnGenerateRowForEntry(UNiagaraStack
 
 		TSharedRef<SHorizontalBox> OptionsBox = SNew(SHorizontalBox);
 		
+		if(StackItem->SupportsStackNotes())
+		{
+			// Stack note
+			OptionsBox->AddSlot()
+			.AutoWidth()
+			.HAlign(HAlign_Center)
+			.VAlign(VAlign_Center)
+			.Padding(2.f)
+			[
+				SNew(SNiagaraStackInlineNote, StackItem)
+				.bInteractable(false)
+				.Visibility_Lambda([StackItem]()
+				{
+					if(StackItem->GetStackNote() == nullptr)
+					{
+						return EVisibility::Collapsed;
+					}
+	
+					return EVisibility::Visible;
+				})
+			];
+		}
+		
 		UNiagaraStackModuleItem* StackModuleItem = Cast<UNiagaraStackModuleItem>(StackItem);
-
+		
 		if (StackModuleItem)
-		{
+		{			
 			// Scratch icon
 			if(StackModuleItem->IsScratchModule())
 			{
diff --git a/Engine/Plugins/FX/Niagara/Source/NiagaraEditorWidgets/Private/SNiagaraStack.cpp b/Engine/Plugins/FX/Niagara/Source/NiagaraEditorWidgets/Private/SNiagaraStack.cpp
--- a/Engine/Plugins/FX/Niagara/Source/NiagaraEditorWidgets/Private/SNiagaraStack.cpp
+++ b/Engine/Plugins/FX/Niagara/Source/NiagaraEditorWidgets/Private/SNiagaraStack.cpp
@@ -51,11 +51,13 @@
 #include "ViewModels/Stack/NiagaraStackRoot.h"
 #include "NiagaraStackCommandContext.h"
 #include "NiagaraEditorUtilities.h"
+#include "SNiagaraStackNote.h"
 #include "Framework/Commands/UICommandList.h"
 #include "ViewModels/NiagaraOverviewGraphViewModel.h"
 #include "Widgets/SNiagaraParameterName.h"
 #include "Styling/StyleColors.h"
 #include "SResetToDefaultPropertyEditor.h"
+#include "ViewModels/Stack/NiagaraStackNote.h"
 
 #define LOCTEXT_NAMESPACE "NiagaraStack"
 
@@ -816,10 +818,7 @@ TSharedRef<SNiagaraStackTableRow> SNiagaraStack::ConstructContainerForItem(UNiag
 				break;
 			case EStackIssueSeverity::Info:
 				IndicatorColor = FStyleColors::Foreground;
-				break;
-			case EStackIssueSeverity::CustomNote:
-				IndicatorColor = FStyleColors::Transparent;
-				break;
+				break;	
 			default:
 				checkf(false, TEXT("Issue severity not set for stack issue."));
 		}
@@ -956,6 +955,14 @@ SNiagaraStack::FRowWidgets SNiagaraStack::ConstructNameAndValueWidgetsForItem(UN
 	{
 		return FRowWidgets(SNew(SNiagaraStackErrorItemFix, CastChecked<UNiagaraStackErrorItemFix>(Item), StackViewModel));
 	}
+	else if(UNiagaraStackNote* StackNote = Cast<UNiagaraStackNote>(Item))
+	{
+		TSharedRef<SNiagaraStackNote> ContentWidget = SNew(SNiagaraStackNote, *StackNote)
+			.bShowEditTextButtons(Container, &SNiagaraStackTableRow::IsHovered);
+		
+		Container->AddFillRowContextMenuHandler(SNiagaraStackTableRow::FOnFillRowContextMenu::CreateSP(ContentWidget, &SNiagaraStackNote::FillRowContextMenu));
+		return FRowWidgets(ContentWidget);
+	}
 	else if (Item->IsA<UNiagaraStackItemFooter>())
 	{
 		UNiagaraStackItemFooter* ItemExpander = CastChecked<UNiagaraStackItemFooter>(Item);
diff --git a/Engine/Plugins/FX/Niagara/Source/NiagaraEditorWidgets/Private/SNiagaraStackNote.cpp b/Engine/Plugins/FX/Niagara/Source/NiagaraEditorWidgets/Private/SNiagaraStackNote.cpp
new file mode 100644
--- /dev/null
+++ b/Engine/Plugins/FX/Niagara/Source/NiagaraEditorWidgets/Private/SNiagaraStackNote.cpp
@@ -0,0 +1,333 @@
+﻿// Copyright Epic Games, Inc. All Rights Reserved.
+
+#include "SNiagaraStackNote.h"
+
+#include "NiagaraEditorStyle.h"
+#include "NiagaraEditorUtilities.h"
+#include "NiagaraEditorWidgetsUtilities.h"
+#include "SlateOptMacros.h"
+#include "Widgets/Input/SButton.h"
+#include "Widgets/Images/SImage.h"
+#include "ScopedTransaction.h"
+#include "ViewModels/Stack/NiagaraStackEntry.h"
+#include "Widgets/SToolTip.h"
+
+#define LOCTEXT_NAMESPACE "SNiagaraStackNote"
+
+BEGIN_SLATE_FUNCTION_BUILD_OPTIMIZATION
+
+void SNiagaraStackInlineNote::Construct(const FArguments& InArgs, UNiagaraStackEntry* InStackEntry)
+{
+	StackEntry = InStackEntry;
+	bInteractable = InArgs._bInteractable;
+
+	TSharedRef<SWidget> ChildWidget = SNew(SImage)
+		.Image(FNiagaraEditorStyle::Get().GetBrush("NiagaraEditor.Message.CustomNote"));
+
+	if(bInteractable)
+	{
+		ChildWidget = SNew(SButton)
+		.ButtonStyle(&FAppStyle::GetWidgetStyle<FButtonStyle>("HoverHintOnly"))
+		.OnClicked(this, &SNiagaraStackInlineNote::OnClicked)
+		.Content()
+		[
+			ChildWidget	
+		];
+	}
+	ChildSlot
+	[
+		ChildWidget
+	];
+
+	UpdateTooltip();
+}
+
+void SNiagaraStackInlineNote::UpdateTooltip()
+{
+	if(StackEntry.IsValid() && StackEntry->GetStackNote() != nullptr)
+	{
+		SetToolTip(FNiagaraEditorUtilities::Tooltips::CreateStackNoteTooltip(*StackEntry->GetStackNote()));
+	}
+}
+
+FReply SNiagaraStackInlineNote::OnClicked() const
+{
+	StackEntry->GetStackNote()->ToggleInlineDisplay();
+	return FReply::Handled();
+}
+
+void SNiagaraStackInlineNote::OnMouseEnter(const FGeometry& MyGeometry, const FPointerEvent& MouseEvent)
+{
+	// this is a workaround. We force a regeneration of the tooltip whenever we enter it.
+	// SetTooltip with a TAttribute bound to a function should work to create the tooltip whenever it's needed, but doesn't work right now.
+	UpdateTooltip();
+	SCompoundWidget::OnMouseEnter(MyGeometry, MouseEvent);
+}
+
+void SNiagaraStackNote::Construct(const FArguments& InArgs, UNiagaraStackNote& InStackNote)
+{
+	StackNote = &InStackNote;
+	ShowEditTextButtons = InArgs._bShowEditTextButtons;
+	
+	StackNote->OnRequestEditHeader().BindSP(this, &SNiagaraStackNote::EditHeaderText);
+
+	Rebuild();
+}
+
+SNiagaraStackNote::~SNiagaraStackNote()
+{
+	if(StackNote.IsValid())
+	{
+		StackNote->OnRequestEditHeader().Unbind();
+	}
+}
+
+void SNiagaraStackNote::Rebuild()
+{	
+	ChildSlot
+	[
+		SNullWidget::NullWidget
+	];
+
+	TOptional<FNiagaraStackNoteData> MatchingStackNote = GetStackNoteData();
+	if(MatchingStackNote.IsSet() == false)
+	{
+		return;
+	}
+	
+	FNiagaraStackNoteData DisplayedStackNote = MatchingStackNote.GetValue();
+	
+	ChildSlot
+	[
+		SAssignNew(ExpandableArea, SExpandableArea)
+		.BorderImage(FAppStyle::GetBrush("NoBorder"))
+		.HeaderContent()
+		[
+			SNew(SHorizontalBox)
+			+ SHorizontalBox::Slot()
+			.AutoWidth()
+			.Padding(2.f)
+			[
+				SNew(SButton)
+				.ButtonStyle(&FAppStyle::GetWidgetStyle<FButtonStyle>("HoverHintOnly"))
+				.OnClicked(this, &SNiagaraStackNote::OnToggleInlineDisplayClicked)
+				.ToolTipText(LOCTEXT("InlineNoteButtonTooltip", "Converts this note to an inline note."))
+				.Content()
+				[
+					SNew(SBox)
+					.HeightOverride(16.f)
+					.WidthOverride(16.f)
+					.HAlign(HAlign_Center)
+					.VAlign(VAlign_Center)
+					[
+						SNew(SImage)
+						.Image(FNiagaraEditorStyle::Get().GetBrush("NiagaraEditor.Message.CustomNote"))
+					]
+				]
+			]
+			+ SHorizontalBox::Slot()
+			[
+				SNew(SHorizontalBox)
+				+ SHorizontalBox::Slot()
+				[
+					SAssignNew(HeaderText, SInlineEditableTextBlock)
+					.Text(this, &SNiagaraStackNote::GetStackNoteHeader)
+					.OnTextCommitted(this, &SNiagaraStackNote::CommitStackNoteHeaderUpdate)
+					.Style(&FNiagaraEditorStyle::Get().GetWidgetStyle<FInlineEditableTextBlockStyle>("NiagaraEditor.Stack.Note.HeaderEditableText"))
+					.IsSelected(FIsSelected::CreateLambda([]() { return false; }))
+				]
+				+ SHorizontalBox::Slot()
+				.AutoWidth()
+				.HAlign(HAlign_Center)
+				.VAlign(VAlign_Bottom)
+				.Padding(10.f, 2.f)
+				[
+					SNew(SButton)
+					.OnClicked(this, &SNiagaraStackNote::OnEditHeaderButtonClicked)
+					.ButtonStyle(&FAppStyle::GetWidgetStyle<FButtonStyle>("HoverHintOnly"))
+					.Visibility(this, &SNiagaraStackNote::GetEditNoteHeaderButtonVisibility)
+					.ToolTipText(LOCTEXT("EditHeaderButtonTooltip", "Edit the header of this note."))
+					[
+						SNew(SImage)
+						.Image(FAppStyle::GetBrush("Icons.Edit"))
+					]
+				]
+			]
+		]
+		.BodyContent()
+		[
+			SNew(SHorizontalBox)
+			+ SHorizontalBox::Slot()
+			[
+				SAssignNew(BodyText, SInlineEditableTextBlock)
+				.Text(this, &SNiagaraStackNote::GetStackNoteBody)
+				.OnTextCommitted(this, &SNiagaraStackNote::CommitStackNoteBodyUpdate)
+				.AutoWrapNonEditText(true)
+				.AutoWrapMultilineEditText(true)
+				.MultiLine(true)
+				.ModiferKeyForNewLine(EModifierKey::Shift)
+				.IsSelected(FIsSelected::CreateLambda([]() { return false; }))
+			]
+			+ SHorizontalBox::Slot()
+			.AutoWidth()
+			.HAlign(HAlign_Center)
+			.VAlign(VAlign_Bottom)
+			.Padding(10.f, 2.f)
+			[
+				SNew(SButton)
+				.OnClicked(this, &SNiagaraStackNote::OnEditBodyButtonClicked)
+				.ButtonStyle(&FAppStyle::GetWidgetStyle<FButtonStyle>("HoverHintOnly"))
+				.Visibility(this, &SNiagaraStackNote::GetEditNoteBodyButtonVisibility)
+				.ToolTipText(LOCTEXT("EditHeaderButtonTooltip", "Edit the message of this note."))
+				[
+					SNew(SImage)
+					.Image(FAppStyle::GetBrush("Icons.Edit"))
+				]
+			]
+		]
+		.Padding(FMargin(20.f, 2.f, 5.f, 2.f))
+	];
+}
+
+void SNiagaraStackNote::FillRowContextMenu(FMenuBuilder& MenuBuilder)
+{
+	MenuBuilder.BeginSection("NoteActions", LOCTEXT("NoteActions", "Note Actions"));
+	{		
+		MenuBuilder.AddMenuEntry(
+			LOCTEXT("EditTitle", "Edit Title"),
+			LOCTEXT("EditTitleTooltip", "Edit the title of this note."),
+			FSlateIcon(),
+			FUIAction(FExecuteAction::CreateSP(this, &SNiagaraStackNote::EditHeaderText)));
+
+		MenuBuilder.AddMenuEntry(
+			LOCTEXT("EditBody", "Edit Message"),
+			LOCTEXT("EditMessageTooltip", "Edit the message of this note."),
+			FSlateIcon(),
+			FUIAction(FExecuteAction::CreateSP(this, &SNiagaraStackNote::EditBodyText)));
+
+		MenuBuilder.AddMenuEntry(
+			LOCTEXT("ToggleInline", "Toggle Inline Display"),
+			LOCTEXT("ToggleInlineTooltip", "Toggle the Inline Display for this note.\nAn inlined note will show up in the row itself, saving on space."),
+			FSlateIcon(),
+			FUIAction(FExecuteAction::CreateSP(this, &SNiagaraStackNote::ToggleInlineDisplay)));
+
+		MenuBuilder.AddMenuEntry(
+			LOCTEXT("DeleteNote", "Delete Note"),
+			LOCTEXT("DeleteNoteTooltip", "Delete this note."),
+			FSlateIcon(),
+			FUIAction(FExecuteAction::CreateSP(this, &SNiagaraStackNote::DeleteStackNote)));
+	}
+	MenuBuilder.EndSection();
+}
+
+void SNiagaraStackNote::Tick(const FGeometry& AllottedGeometry, const double InCurrentTime, const float InDeltaTime)
+{
+	if(StackNote->GetIsRenamePending())
+	{
+		if(HeaderText.IsValid())
+		{
+			EditHeaderText();
+			StackNote->SetIsRenamePending(false);
+		}
+	}
+}
+
+TOptional<FNiagaraStackNoteData> SNiagaraStackNote::GetStackNoteData() const
+{
+	return StackNote->GetTargetStackNoteData();
+}
+
+void SNiagaraStackNote::EditHeaderText() const
+{
+	HeaderText->EnterEditingMode();
+}
+
+void SNiagaraStackNote::EditBodyText() const
+{
+	ExpandableArea->SetExpanded(true);
+	BodyText->EnterEditingMode();
+}
+
+void SNiagaraStackNote::CommitStackNoteHeaderUpdate(const FText& Text, ETextCommit::Type Arg)
+{
+	FScopedTransaction Transaction(LOCTEXT("UpdateNoteHeaderTransaction", "Updated Note Header"));
+	StackNote->GetStackEditorData().Modify();
+	
+	FNiagaraStackNoteData UpdatedMessage = GetStackNoteData().GetValue();
+	UpdatedMessage.MessageHeader = Text;
+
+	StackNote->GetStackEditorData().AddOrReplaceStackNote(StackNote->GetTargetStackEntryKey(), UpdatedMessage);
+	Rebuild();
+}
+
+void SNiagaraStackNote::CommitStackNoteBodyUpdate(const FText& Text, ETextCommit::Type Arg)
+{
+	FScopedTransaction Transaction(LOCTEXT("UpdateNoteMessageTransaction", "Updated Note Message"));
+	StackNote->GetStackEditorData().Modify();
+	
+	FNiagaraStackNoteData UpdatedMessage = GetStackNoteData().GetValue();
+	UpdatedMessage.Message = Text;
+
+	StackNote->GetStackEditorData().AddOrReplaceStackNote(StackNote->GetTargetStackEntryKey(), UpdatedMessage);
+	Rebuild();
+}
+
+void SNiagaraStackNote::ToggleInlineDisplay() const
+{
+	StackNote->ToggleInlineDisplay();
+}
+
+void SNiagaraStackNote::DeleteStackNote() const
+{
+	StackNote->DeleteTargetStackNote();
+}
+
+FText SNiagaraStackNote::GetStackNoteHeader() const
+{
+	return GetStackNoteData().GetValue().MessageHeader;
+}
+
+FText SNiagaraStackNote::GetStackNoteBody() const
+{
+	return GetStackNoteData().GetValue().Message;
+}
+
+FReply SNiagaraStackNote::OnToggleInlineDisplayClicked() const
+{
+	ToggleInlineDisplay();
+	return FReply::Handled();
+}
+
+FReply SNiagaraStackNote::OnEditHeaderButtonClicked()
+{
+	RegisterActiveTimer(0.f, FWidgetActiveTimerDelegate::CreateLambda([this](double, double)
+	{
+		EditHeaderText();
+		return EActiveTimerReturnType::Stop;
+	}));
+	return FReply::Handled();
+}
+
+FReply SNiagaraStackNote::OnEditBodyButtonClicked()
+{
+	RegisterActiveTimer(0.f, FWidgetActiveTimerDelegate::CreateLambda([this](double, double)
+	{
+		EditBodyText();
+		return EActiveTimerReturnType::Stop;
+	}));
+	return FReply::Handled();
+}
+
+EVisibility SNiagaraStackNote::GetEditNoteHeaderButtonVisibility() const
+{
+	return ShowEditTextButtons.Get() ? EVisibility::Visible : EVisibility::Hidden;
+}
+
+EVisibility SNiagaraStackNote::GetEditNoteBodyButtonVisibility() const
+{
+	return ShowEditTextButtons.Get() ? EVisibility::Visible : EVisibility::Hidden;
+}
+
+END_SLATE_FUNCTION_BUILD_OPTIMIZATION
+
+#undef LOCTEXT_NAMESPACE
\ No newline at end of file
diff --git a/Engine/Plugins/FX/Niagara/Source/NiagaraEditorWidgets/Private/SNiagaraStackNote.h b/Engine/Plugins/FX/Niagara/Source/NiagaraEditorWidgets/Private/SNiagaraStackNote.h
new file mode 100644
--- /dev/null
+++ b/Engine/Plugins/FX/Niagara/Source/NiagaraEditorWidgets/Private/SNiagaraStackNote.h
@@ -0,0 +1,85 @@
+﻿// Copyright Epic Games, Inc. All Rights Reserved.
+
+#pragma once
+
+#include "CoreMinimal.h"
+#include "ViewModels/Stack/NiagaraStackNote.h"
+#include "Widgets/SCompoundWidget.h"
+#include "Widgets/Layout/SExpandableArea.h"
+#include "Widgets/Text/SInlineEditableTextBlock.h"
+
+/** A possibly interactable widget showing a note. On hover, summons the note as a tooltip. On clicked, toggles between inline & full display */
+class NIAGARAEDITORWIDGETS_API SNiagaraStackInlineNote : public SCompoundWidget
+{
+public:
+	SLATE_BEGIN_ARGS(SNiagaraStackInlineNote)
+		: _bInteractable(true)
+		{}
+	
+		/** If true, a button will be added. If false, it will be a simple image.  */
+		SLATE_ARGUMENT(bool, bInteractable)	
+	SLATE_END_ARGS()
+
+	void Construct(const FArguments& InArgs, UNiagaraStackEntry* InStackEntry);
+
+	void UpdateTooltip();
+private:
+	FReply OnClicked() const;
+
+	virtual void OnMouseEnter(const FGeometry& MyGeometry, const FPointerEvent& MouseEvent) override;
+private:
+	TWeakObjectPtr<UNiagaraStackEntry> StackEntry;
+	bool bInteractable = true;
+};
+
+/** The full stack note widget displayed in the stack. Features editable title, message & toggle inline display button. */
+class NIAGARAEDITORWIDGETS_API SNiagaraStackNote : public SCompoundWidget
+{
+public:	
+	SLATE_BEGIN_ARGS(SNiagaraStackNote)
+		: _bShowEditTextButtons(true)
+		{}
+		SLATE_ATTRIBUTE(bool, bShowEditTextButtons)
+	SLATE_END_ARGS()
+	
+	void Construct(const FArguments& InArgs, UNiagaraStackNote& StackNote);
+	virtual ~SNiagaraStackNote() override;
+	
+	void Rebuild();
+
+	void FillRowContextMenu(FMenuBuilder& MenuBuilder);
+private:
+	virtual void Tick(const FGeometry& AllottedGeometry, const double InCurrentTime, const float InDeltaTime) override;
+	
+	TOptional<FNiagaraStackNoteData> GetStackNoteData() const;
+	
+	void EditHeaderText() const;
+	void EditBodyText() const;
+	
+	void CommitStackNoteHeaderUpdate(const FText& Text, ETextCommit::Type Arg);
+	void CommitStackNoteBodyUpdate(const FText& Text, ETextCommit::Type Arg);
+	
+	void ToggleInlineDisplay() const;
+	void DeleteStackNote() const;
+	
+	FText GetStackNoteHeader() const;
+	FText GetStackNoteBody() const;
+
+private:
+	FReply OnToggleInlineDisplayClicked() const;
+
+	FReply OnEditHeaderButtonClicked();
+	FReply OnEditBodyButtonClicked();
+	
+	EVisibility GetEditNoteHeaderButtonVisibility() const;
+	EVisibility GetEditNoteBodyButtonVisibility() const;
+
+private:
+	TWeakObjectPtr<UNiagaraStackNote> StackNote;
+private:
+	TSharedPtr<SExpandableArea> ExpandableArea;
+	TSharedPtr<SInlineEditableTextBlock> HeaderText;
+	TSharedPtr<SInlineEditableTextBlock> BodyText;
+
+	TAttribute<bool> ShowEditTextButtons;
+};
diff --git a/Engine/Plugins/FX/Niagara/Source/NiagaraEditorWidgets/Private/Stack/SNiagaraStackErrorItem.cpp b/Engine/Plugins/FX/Niagara/Source/NiagaraEditorWidgets/Private/Stack/SNiagaraStackErrorItem.cpp
--- a/Engine/Plugins/FX/Niagara/Source/NiagaraEditorWidgets/Private/Stack/SNiagaraStackErrorItem.cpp
+++ b/Engine/Plugins/FX/Niagara/Source/NiagaraEditorWidgets/Private/Stack/SNiagaraStackErrorItem.cpp
@@ -31,9 +31,6 @@ void SNiagaraStackErrorItem::Construct(const FArguments& InArgs, UNiagaraStackEr
 	case EStackIssueSeverity::Info:
 		IconBrush = FAppStyle::Get().GetBrush("Icons.InfoWithColor");
 		break;
-	case EStackIssueSeverity::CustomNote:
-		IconBrush = FNiagaraEditorStyle::Get().GetBrush("NiagaraEditor.Message.CustomNote");
-		break;
 	default:
 		IconBrush = FAppStyle::Get().GetBrush("NoBrush");
 		break;
diff --git a/Engine/Plugins/FX/Niagara/Source/NiagaraEditorWidgets/Private/Stack/SNiagaraStackIssueIcon.cpp b/Engine/Plugins/FX/Niagara/Source/NiagaraEditorWidgets/Private/Stack/SNiagaraStackIssueIcon.cpp
--- a/Engine/Plugins/FX/Niagara/Source/NiagaraEditorWidgets/Private/Stack/SNiagaraStackIssueIcon.cpp
+++ b/Engine/Plugins/FX/Niagara/Source/NiagaraEditorWidgets/Private/Stack/SNiagaraStackIssueIcon.cpp
@@ -150,17 +150,6 @@ FText SNiagaraStackIssueIcon::GetIconToolTip() const
 					ToolTipParts.Add(FText::Format(LOCTEXT("InfoFormatMultiple", "{0} infos"), FText::AsNumber(StackEntry->GetTotalNumberOfInfoIssues())));
 				}
 			}
-			if (StackEntry->GetTotalNumberOfCustomNotes() > 0)
-			{
-				if (StackEntry->GetTotalNumberOfCustomNotes() == 1)
-				{
-					ToolTipParts.Add(LOCTEXT("CustomNotesFormatSingle", "1 custom note"));
-				}
-				else
-				{
-					ToolTipParts.Add(FText::Format(LOCTEXT("CustomNotesFormatMultiple", "{0} custom notes"), FText::AsNumber(StackEntry->GetTotalNumberOfCustomNotes())));
-				}
-			}
 
 			if (ToolTipParts.Num() == 4)
 			{
@@ -190,8 +179,6 @@ FText SNiagaraStackIssueIcon::GetIconToolTip() const
 					return LOCTEXT("Warning", "Warning");
 				case EStackIssueSeverity::Info:
 					return LOCTEXT("Issue", "Issue");
-				case EStackIssueSeverity::CustomNote:
-					return LOCTEXT("CustomNote", "Custom Note");
 				default:
 					return FText();
 				}
@@ -225,11 +212,6 @@ FText SNiagaraStackIssueIcon::GetIconToolTip() const
 				{
 					const UNiagaraStackEntry::FStackIssue& Issue = EntryToCheck->GetIssues()[IssueIndex];
 					ToolTipBuilder.AppendLineFormat(IssueLineFormat, GetFullDisplayName(StackViewModel, EntryToCheck), SeverityToText(Issue.GetSeverity()), Issue.GetShortDescription());
-
-					if (Issue.GetSeverity() == EStackIssueSeverity::CustomNote)
-					{
-						ToolTipBuilder.AppendLine(Issue.GetLongDescription());
-					}
 				}
 				TotalIssues += EntryToCheck->GetIssues().Num();
 			}
@@ -280,12 +262,6 @@ void SNiagaraStackIssueIcon::UpdateFromEntry(ENiagaraStructureChangedFlags Flags
 			? FNiagaraEditorWidgetsStyle::Get().GetBrush("Niagara.CompactStackIssue.Info")
 			: FAppStyle::Get().GetBrush("Icons.InfoWithColor");
 	}
-	else if (StackEntry->GetTotalNumberOfCustomNotes() > 0)
-	{
-		IconBrush = IconMode == EIconMode::Compact
-			? FNiagaraEditorWidgetsStyle::Get().GetBrush("Niagara.CompactStackIssue.Message")
-			: FNiagaraEditorStyle::Get().GetBrush("NiagaraEditor.Message.CustomNote");
-	}
 
 	IconToolTipCache.Reset();
 }
diff --git a/Engine/Plugins/FX/Niagara/Source/NiagaraEditorWidgets/Private/Stack/SNiagaraStackModuleItem.cpp b/Engine/Plugins/FX/Niagara/Source/NiagaraEditorWidgets/Private/Stack/SNiagaraStackModuleItem.cpp
--- a/Engine/Plugins/FX/Niagara/Source/NiagaraEditorWidgets/Private/Stack/SNiagaraStackModuleItem.cpp
+++ b/Engine/Plugins/FX/Niagara/Source/NiagaraEditorWidgets/Private/Stack/SNiagaraStackModuleItem.cpp
@@ -40,19 +40,6 @@ bool SNiagaraStackModuleItem::bLibraryOnly = true;
 void SNiagaraStackModuleItem::Construct(const FArguments& InArgs, UNiagaraStackModuleItem& InModuleItem, UNiagaraStackViewModel* InStackViewModel)
 {
 	ModuleItem = &InModuleItem;
-
-	ModuleItem->OnNoteModeSet().BindLambda([this](bool bEnabled)
-	{
-		if(bEnabled)
-		{
-			FSlateApplication::Get().SetKeyboardFocus(ShortDescriptionTextBox);
-		}
-		else
-		{
-			ShortDescriptionTextBox->SetText(FText::GetEmpty());
-			DescriptionTextBox->SetText(FText::GetEmpty());
-		}
-	});
 	SNiagaraStackItem::Construct(SNiagaraStackItem::FArguments(), InModuleItem, InStackViewModel);
 }
 
@@ -188,142 +175,7 @@ FSlateColor SNiagaraStackModuleItem::GetVersionSelectorColor() const
 }
 
 TSharedRef<SWidget> SNiagaraStackModuleItem::AddContainerForRowWidgets(TSharedRef<SWidget> RowWidgets)
-{
-	ShortDescriptionTextBox = SNew(SMultiLineEditableTextBox).AutoWrapText(true);
-	DescriptionTextBox = SNew(SMultiLineEditableTextBox).AutoWrapText(true);
-
-	auto OkPressed = [&]() -> FReply
-	{
-		FText Text = DescriptionTextBox->GetText();
-		if(!Text.IsEmpty())
-		{
-			FScopedTransaction Transaction(LOCTEXT("NoteAdded", "Note Added"));
-			ModuleItem->GetModuleNode().Modify();
-			
-			FNiagaraStackMessage StackMessage(DescriptionTextBox->GetText(), ShortDescriptionTextBox->GetText(), ENiagaraMessageSeverity::CustomNote, false);
-			ModuleItem->GetModuleNode().AddCustomNote(StackMessage);
-		}
-		
-		ModuleItem->SetNoteMode(false);
-		return FReply::Handled();
-	};
-
-	auto CancelPressed = [&]() -> FReply
-	{
-		ModuleItem->SetNoteMode(false);
-		return FReply::Handled();
-	};
-	
-	TSharedPtr<SVerticalBox> AddNoteWidget = SNew(SVerticalBox).Visibility_Lambda([&]()
-	{
-		return ModuleItem->GetNoteMode() ? EVisibility::Visible : EVisibility::Collapsed;
-	})
-	+ SVerticalBox::Slot()
-	.AutoHeight()
-	.HAlign(HAlign_Fill)
-	.VAlign(VAlign_Center)
-	.Padding(5.f)
-	[
-		SNew(SVerticalBox)
-		+ SVerticalBox::Slot()
-		.AutoHeight()
-		.HAlign(HAlign_Left)
-		.VAlign(VAlign_Center)
-		.Padding(5.f)
-		[
-			SNew(SHorizontalBox)
-			+ SHorizontalBox::Slot()
-			.AutoWidth()
-			.VAlign(VAlign_Center)
-			.Padding(0.f, 0.f, 3.f, 0.f)
-			[
-				SNew(SImage)
-				.Image(FNiagaraEditorStyle::Get().GetBrush("NiagaraEditor.Message.CustomNote"))
-			]
-			+ SHorizontalBox::Slot()
-			.AutoWidth()
-			.VAlign(VAlign_Center)
-			[
-				SNew(STextBlock)
-				.Text(LOCTEXT("MessageTitleLabel", "Title"))
-			]
-		]
-		+ SVerticalBox::Slot()
-		.AutoHeight()
-		.HAlign(HAlign_Fill)
-		.VAlign(VAlign_Center)
-		[
-			SNew(SBox)
-			.WidthOverride(300.f)
-			[
-				ShortDescriptionTextBox.ToSharedRef()
-			]
-		]
-	]
-	+ SVerticalBox::Slot()
-	.AutoHeight()
-	.HAlign(HAlign_Fill)
-	.VAlign(VAlign_Center)
-	.Padding(5.f)
-	[
-		SNew(SVerticalBox)
-		+ SVerticalBox::Slot()
-		.AutoHeight()
-		.HAlign(HAlign_Left)
-		.VAlign(VAlign_Center)
-		.Padding(5.f)
-		[
-			SNew(STextBlock)
-			.Text(LOCTEXT("MessageTextLabel", "Message"))
-		]
-		+ SVerticalBox::Slot()
-		.AutoHeight()
-		.HAlign(HAlign_Fill)
-		.VAlign(VAlign_Center)
-		[
-			SNew(SBox)
-			.WidthOverride(300.f)
-			.HeightOverride(125.f)
-			[
-				DescriptionTextBox.ToSharedRef()
-			]
-		]
-	]
-	+ SVerticalBox::Slot()
-	.AutoHeight()
-	.HAlign(HAlign_Right)
-	.VAlign(VAlign_Bottom)
-	.Padding(3.f, 5.f)
-	[
-		SNew(SHorizontalBox)
-		+ SHorizontalBox::Slot()
-		.AutoWidth()
-		.HAlign(HAlign_Center)
-		.VAlign(VAlign_Bottom)
-		.Padding(5.f)
-		[
-			SNew(SButton)
-			.Text(LOCTEXT("CreateNote_Yes", "Add Note"))
-			.ButtonStyle(FAppStyle::Get(), "FlatButton.Success")
-			.IsEnabled_Lambda([&]()
-			{
-				return !DescriptionTextBox->GetText().IsEmpty();
-			})
-			.OnClicked_Lambda(OkPressed)
-		]
-		+ SHorizontalBox::Slot()
-		.AutoWidth()
-		.HAlign(HAlign_Center)
-		.VAlign(VAlign_Bottom)
-		.Padding(5.f)
-		[
-			SNew(SButton)
-			.ButtonStyle(FAppStyle::Get(), "FlatButton.Dark")
-			.Text(LOCTEXT("CreateNote_No", "Cancel"))
-			.OnClicked_Lambda(CancelPressed)
-		]
-	];
-
+{	
 	return SNew(SDropTarget)
 	.OnAllowDrop(this, &SNiagaraStackModuleItem::OnModuleItemAllowDrop)
 	.OnDropped(this, &SNiagaraStackModuleItem::OnModuleItemDrop)
@@ -331,22 +183,7 @@ TSharedRef<SWidget> SNiagaraStackModuleItem::AddContainerForRowWidgets(TSharedRe
 	.VerticalImage(FNiagaraEditorWidgetsStyle::Get().GetBrush("NiagaraEditor.Stack.DropTarget.BorderVertical"))
 	.Content()
 	[
-		SNew(SVerticalBox)
-		+ SVerticalBox::Slot()
-		.AutoHeight()
-		[
-			RowWidgets
-		]
-		+ SVerticalBox::Slot()
-		.AutoHeight()
-		[
-			SNew(SBorder)
-			.BorderImage(FAppStyle::GetBrush("WhiteBrush"))
-			.BorderBackgroundColor(FNiagaraEditorWidgetsStyle::Get().GetColor("NiagaraEditor.Stack.Item.CustomNoteBackgroundColor"))
-			[
-				AddNoteWidget.ToSharedRef()
-			]
-		]
+		RowWidgets
 	];
 }
 
diff --git a/Engine/Plugins/FX/Niagara/Source/NiagaraEditorWidgets/Private/Stack/SNiagaraStackModuleItem.h b/Engine/Plugins/FX/Niagara/Source/NiagaraEditorWidgets/Private/Stack/SNiagaraStackModuleItem.h
--- a/Engine/Plugins/FX/Niagara/Source/NiagaraEditorWidgets/Private/Stack/SNiagaraStackModuleItem.h
+++ b/Engine/Plugins/FX/Niagara/Source/NiagaraEditorWidgets/Private/Stack/SNiagaraStackModuleItem.h
@@ -77,9 +77,6 @@ private:
 
 	TSharedPtr<SComboButton> AddButton;
 
-	TSharedPtr<SMultiLineEditableTextBox> ShortDescriptionTextBox;
-	TSharedPtr<SMultiLineEditableTextBox> DescriptionTextBox;
-
 	static bool bLibraryOnly;
 
 };
\ No newline at end of file
diff --git a/Engine/Plugins/FX/Niagara/Source/NiagaraEditorWidgets/Private/Stack/SNiagaraStackTableRow.cpp b/Engine/Plugins/FX/Niagara/Source/NiagaraEditorWidgets/Private/Stack/SNiagaraStackTableRow.cpp
--- a/Engine/Plugins/FX/Niagara/Source/NiagaraEditorWidgets/Private/Stack/SNiagaraStackTableRow.cpp
+++ b/Engine/Plugins/FX/Niagara/Source/NiagaraEditorWidgets/Private/Stack/SNiagaraStackTableRow.cpp
@@ -1,6 +1,8 @@
 // Copyright Epic Games, Inc. All Rights Reserved.
 
 #include "Stack/SNiagaraStackTableRow.h"
+
+#include "NiagaraEditorStyle.h"
 #include "NiagaraEditorWidgetsStyle.h"
 #include "Styling/AppStyle.h"
 #include "ViewModels/Stack/NiagaraStackViewModel.h"
@@ -23,7 +25,10 @@
 #include "Styling/StyleColors.h"
 #include "ScopedTransaction.h"
 #include "NiagaraEmitterEditorData.h"
+#include "SNiagaraStackNote.h"
 #include "ViewModels/HierarchyEditor/NiagaraSummaryViewViewModel.h"
+#include "ViewModels/Stack/NiagaraStackClipboardUtilities.h"
+#include "ViewModels/Stack/NiagaraStackNote.h"
 #include "ViewModels/Stack/NiagaraStackPropertyRow.h"
 #include "ViewModels/Stack/NiagaraStackRendererItem.h"
 #include "ViewModels/Stack/NiagaraStackSimulationStageGroup.h"
@@ -187,7 +192,7 @@ void SNiagaraStackTableRow::SetNameAndValueContent(TSharedRef<SWidget> InNameWid
 	[
 		InNameWidget
 	];
-
+	
 	TSharedPtr<SWidget> ChildContent;
 	if (InValueWidget.IsValid())
 	{
@@ -200,12 +205,25 @@ void SNiagaraStackTableRow::SetNameAndValueContent(TSharedRef<SWidget> InNameWid
 		.Value(NameColumnWidth)
 		.OnSlotResized(SSplitter::FOnSlotResized::CreateSP(this, &SNiagaraStackTableRow::OnNameColumnWidthChanged))
 		[
-			SNew(SBox)
-			.Padding(FMargin(ContentPadding.Left, 0, 5, 0))
-			.MinDesiredWidth(NameMinWidth.IsSet() ? NameMinWidth.GetValue() : FOptionalSize())
-			.MaxDesiredWidth(NameMaxWidth.IsSet() ? NameMaxWidth.GetValue() : FOptionalSize())
+			SNew(SHorizontalBox)
+			+ SHorizontalBox::Slot()
 			[
-				NameContent
+				SNew(SBox)
+				.Padding(FMargin(ContentPadding.Left, 0, 5, 0))
+				.MinDesiredWidth(NameMinWidth.IsSet() ? NameMinWidth.GetValue() : FOptionalSize())
+				.MaxDesiredWidth(NameMaxWidth.IsSet() ? NameMaxWidth.GetValue() : FOptionalSize())
+				[
+					NameContent
+				]
+			]
+			+ SHorizontalBox::Slot()
+			.AutoWidth()
+			.Padding(2.f)
+			.HAlign(HAlign_Center)
+			.VAlign(VAlign_Center)
+			[
+				SNew(SNiagaraStackInlineNote, StackEntry)
+				.Visibility(this, &SNiagaraStackTableRow::GetInlineNoteVisibility)
 			]
 		]
 
@@ -227,12 +245,25 @@ void SNiagaraStackTableRow::SetNameAndValueContent(TSharedRef<SWidget> InNameWid
 	}
 	else
 	{
-		ChildContent = SNew(SBox)
-		.Padding(FMargin(ContentPadding.Left, 0, ContentPadding.Right, 0))
-		.MinDesiredWidth(NameMinWidth.IsSet() ? NameMinWidth.GetValue() : FOptionalSize())
-		.MaxDesiredWidth(NameMaxWidth.IsSet() ? NameMaxWidth.GetValue() : FOptionalSize())
+		ChildContent = SNew(SHorizontalBox)
+		+ SHorizontalBox::Slot()
 		[
-			NameContent
+			SNew(SBox)
+			.Padding(FMargin(ContentPadding.Left, 0, ContentPadding.Right, 0))
+			.MinDesiredWidth(NameMinWidth.IsSet() ? NameMinWidth.GetValue() : FOptionalSize())
+			.MaxDesiredWidth(NameMaxWidth.IsSet() ? NameMaxWidth.GetValue() : FOptionalSize())
+			[
+				NameContent
+			]
+		]
+		+ SHorizontalBox::Slot()
+		.AutoWidth()
+		.Padding(2.f)
+		.HAlign(HAlign_Center)
+		.VAlign(VAlign_Center)
+		[
+			SNew(SNiagaraStackInlineNote, StackEntry)
+			.Visibility(this, &SNiagaraStackTableRow::GetInlineNoteVisibility)
 		];
 	}
 
@@ -389,6 +420,58 @@ FReply SNiagaraStackTableRow::OnMouseButtonUp(const FGeometry& MyGeometry, const
 			LOCTEXT("CollapseAllItemsToolTip", "Collapse all items under this header."),
 			FSlateIcon(),
 			FUIAction(FExecuteAction::CreateSP(this, &SNiagaraStackTableRow::CollapseChildren)));
+
+		if(StackEntry->SupportsStackNotes())
+		{
+			if(StackEntry->HasStackNoteData() == false)
+			{
+				MenuBuilder.AddMenuEntry(
+				LOCTEXT("AddNote", "Add Note"),
+				LOCTEXT("AddNoteTooltip", "Add a note to this row"),
+				FSlateIcon(FNiagaraEditorStyle::Get().GetStyleSetName(), "NiagaraEditor.Message.CustomNote"),
+				FUIAction(FExecuteAction::CreateSP(this, &SNiagaraStackTableRow::AddStackNote))
+				);
+			}
+			else
+			{				
+				MenuBuilder.AddMenuEntry(
+				LOCTEXT("ToggleInlineNote", "Toggle Inline Note Display"),
+				LOCTEXT("ToggleInlineNoteTooltip", "Toggle the inline display for this row's note.\nAn inlined note will show up in the row itself, saving on space"),
+				FSlateIcon(FNiagaraEditorStyle::Get().GetStyleSetName(), "NiagaraEditor.Message.CustomNote"),
+				FUIAction(FExecuteAction::CreateSP(this, &SNiagaraStackTableRow::ToggleStackNoteInlineDisplay))
+				);
+
+				MenuBuilder.AddMenuEntry(
+				LOCTEXT("CopyNote", "Copy Note"),
+			LOCTEXT("CopyNoteTooltip", "Copy the note of this row"),
+				FSlateIcon(FAppStyle::GetAppStyleSetName(), "GenericCommands.Copy"),
+				FUIAction(FExecuteAction::CreateSP(this, &SNiagaraStackTableRow::CopyStackNote))
+				);
+			}
+			
+			if(const UNiagaraClipboardContent* ClipboardContent = FNiagaraEditorModule::Get().GetClipboard().GetClipboardContent())
+			{
+				if(ClipboardContent->StackNote.IsValid())
+				{
+					MenuBuilder.AddMenuEntry(
+				LOCTEXT("PasteNote", "Paste Note"),
+				LOCTEXT("TPasteNoteTooltip", "Paste the note copied in the clipboard"),
+					FSlateIcon(FAppStyle::GetAppStyleSetName(), "GenericCommands.Paste"),
+					FUIAction(FExecuteAction::CreateSP(this, &SNiagaraStackTableRow::PasteStackNote))
+					);
+				}
+			}
+
+			if(StackEntry->HasStackNoteData())
+			{
+				MenuBuilder.AddMenuEntry(
+				LOCTEXT("DeleteNote", "Delete Note"),
+				LOCTEXT("DeleteNoteTooltip", "Delete this row's note"),
+				FSlateIcon(FAppStyle::GetAppStyleSetName(), "GenericCommands.Delete"),
+				FUIAction(FExecuteAction::CreateSP(this, &SNiagaraStackTableRow::DeleteStackNote))
+				);
+			}
+		}
 		MenuBuilder.EndSection();
 
 		if(StackEntry->GetEmitterViewModel().IsValid() && StackEntry->GetEmitterViewModel())
@@ -465,6 +548,68 @@ void SNiagaraStackTableRow::ExpandChildren()
 	SetExpansionStateRecursive(StackEntry, bIsExpanded);
 }
 
+void SNiagaraStackTableRow::AddStackNote() const
+{
+	if(StackEntry->HasStackNoteData() == false)
+	{
+		FScopedTransaction Transaction(LOCTEXT("NewStackNoteAdded", "Added Stack Note"));
+		StackEntry->GetStackEditorData().Modify();
+		
+		FNiagaraStackNoteData NewStackNote;
+		NewStackNote.MessageHeader = LOCTEXT("DefaultStackNoteHeader", "My Note");
+		NewStackNote.Message = LOCTEXT("DefaultStackNoteMessage", "My Message");
+
+		StackEntry->GetStackEditorData().AddOrReplaceStackNote(StackEntry->GetStackEditorDataKey(), NewStackNote);
+		StackEntry->RefreshChildren();
+		
+		if(UNiagaraStackNote* Note = StackEntry->GetStackNote())
+		{
+			Note->SetIsRenamePending(true);
+		}
+	}
+}
+
+void SNiagaraStackTableRow::DeleteStackNote() const
+{
+	if(StackEntry->HasStackNoteData())
+	{
+		StackEntry->DeleteStackNoteData();
+	}
+}
+
+void SNiagaraStackTableRow::ToggleStackNoteInlineDisplay() const
+{
+	if(StackEntry->GetStackNote() != nullptr)
+	{
+		StackEntry->GetStackNote()->ToggleInlineDisplay();
+	}
+}
+
+void SNiagaraStackTableRow::CopyStackNote() const
+{
+	FNiagaraStackClipboardUtilities::CopyNote(StackEntry);
+}
+
+void SNiagaraStackTableRow::PasteStackNote() const
+{
+	if(const UNiagaraClipboardContent* ClipboardContent = FNiagaraEditorModule::Get().GetClipboard().GetClipboardContent())
+	{
+		if(ClipboardContent->StackNote.IsValid())
+		{
+			FScopedTransaction Transaction(LOCTEXT("PasteNoteTransaction", "Note pasted"));
+			StackEntry->GetStackEditorData().Modify();
+			
+			StackEntry->SetStackNoteData(ClipboardContent->StackNote);
+
+			// if the note we pasted is not inlined, we expand the stack entry in order to show the note
+			if(StackEntry->GetStackNoteData().bInlineNote == false)
+			{
+				StackEntry->SetIsExpanded(true);
+			}
+		}
+	}
+}
+
 EVisibility SNiagaraStackTableRow::GetRowVisibility() const
 {
 	return StackEntry->GetShouldShowInStack()
@@ -495,6 +640,16 @@ EVisibility SNiagaraStackTableRow::GetExpanderVisibility() const
 	}
 }
 
+EVisibility SNiagaraStackTableRow::GetInlineNoteVisibility() const
+{
+	if(StackEntry->GetStackNote() == nullptr || StackEntry->GetStackNote()->GetShouldShowInStack())
+	{
+		return EVisibility::Collapsed;
+	}
+		
+	return StackEntry->GetStackNote()->GetTargetStackNoteData().GetValue().bInlineNote == true ? EVisibility::Visible : EVisibility::Collapsed;
+}
+
 FReply SNiagaraStackTableRow::ExpandButtonClicked()
 {
 	const bool bWillBeExpanded = !StackEntry->GetIsExpanded();
diff --git a/Engine/Plugins/FX/Niagara/Source/NiagaraEditorWidgets/Private/Stack/SNiagaraStackTableRow.h b/Engine/Plugins/FX/Niagara/Source/NiagaraEditorWidgets/Private/Stack/SNiagaraStackTableRow.h
--- a/Engine/Plugins/FX/Niagara/Source/NiagaraEditorWidgets/Private/Stack/SNiagaraStackTableRow.h
+++ b/Engine/Plugins/FX/Niagara/Source/NiagaraEditorWidgets/Private/Stack/SNiagaraStackTableRow.h
@@ -58,7 +58,7 @@ public:
 	FMargin GetContentPadding() const;
 
 	void SetContentPadding(FMargin InContentPadding);
-
+	
 	void SetNameAndValueContent(TSharedRef<SWidget> InNameWidget, TSharedPtr<SWidget> InValueWidget);
 
 	void AddFillRowContextMenuHandler(FOnFillRowContextMenu FillRowContextMenuHandler);
@@ -72,12 +72,20 @@ private:
 
 	void ExpandChildren();
 
+	void AddStackNote() const;
+	void DeleteStackNote() const;
+	void ToggleStackNoteInlineDisplay() const;
+	void CopyStackNote() const;
+	void PasteStackNote() const;
+
 	EVisibility GetRowVisibility() const;
 
 	EVisibility GetExecutionCategoryIconVisibility() const;
 
 	EVisibility GetExpanderVisibility() const;
 
+	EVisibility GetInlineNoteVisibility() const;
+
 	FReply ExpandButtonClicked();
 
 	const FSlateBrush* GetExpandButtonImage() const;
