Avoiding a call to 'ensurefullyloaded' in post load of systems to transfer note data

diff --git a/Engine/Plugins/FX/Niagara/Source/NiagaraEditor/Private/NiagaraSystemEditorData.cpp b/Engine/Plugins/FX/Niagara/Source/NiagaraEditor/Private/NiagaraSystemEditorData.cpp
--- a/Engine/Plugins/FX/Niagara/Source/NiagaraEditor/Private/NiagaraSystemEditorData.cpp
+++ b/Engine/Plugins/FX/Niagara/Source/NiagaraEditor/Private/NiagaraSystemEditorData.cpp
@@ -234,7 +234,21 @@ void UNiagaraSystemEditorData::PostLoad_TransferModuleStackNotesToNewFormat(UObj
 	// since we lack graph context during post load, we do this workaround to find the emitter graph
 	UNiagaraSystem* System = Cast<UNiagaraSystem>(InOwner);
 
-	for(UNiagaraScriptSourceBase* SourceBase : System->GetAllSourceScripts())
+	TArray<UNiagaraScriptSourceBase*> SourceBases;
+	
+	if(UNiagaraScript* SystemSpawnScript = System->GetSystemSpawnScript())
+	{
+		SystemSpawnScript->ConditionalPostLoad();
+		SourceBases.Add(SystemSpawnScript->GetLatestSource());
+	}
+	
+	if(UNiagaraScript* SystemUpdateScript = System->GetSystemUpdateScript())
+	{
+		SystemUpdateScript->ConditionalPostLoad();
+		SourceBases.Add(SystemUpdateScript->GetLatestSource());
+	}
+	
+	for(UNiagaraScriptSourceBase* SourceBase : SourceBases)
 	{
 		UNiagaraScriptSource* Source = CastChecked<UNiagaraScriptSource>(SourceBase);
 
