Add redirect to handle missing CustomNote enum entries

diff --git a/Engine/Plugins/FX/Niagara/Config/BaseNiagara.ini b/Engine/Plugins/FX/Niagara/Config/BaseNiagara.ini
--- a/Engine/Plugins/FX/Niagara/Config/BaseNiagara.ini
+++ b/Engine/Plugins/FX/Niagara/Config/BaseNiagara.ini
@@ -34,6 +34,9 @@
 
 +PropertyRedirects=(OldName="NiagaraParameterBinding.Parameter", NewName="NiagaraParameterBinding.ResolvedParameter")
 
++EnumRedirects=(OldName="ENiagaraMessageSeverity",ValueChanges=(("ENiagaraMessageSeverity::CustomNote", "ENiagaraMessageSeverity::Info")) )
++EnumRedirects=(OldName="EStackIssueSeverity",ValueChanges=(("EStackIssueSeverity::CustomNote", "EStackIssueSeverity::Info")) )
+
 ; Flipbook to Baker
 +ClassRedirects=(OldName="NiagaraFlipbookSettings", NewName="/Script/Niagara.NiagaraBakerSettings")
 +PropertyRedirects=(OldName="NiagaraSystem.FlipbookSettings", NewName="NiagaraSystem.BakerSettings")
