Fix for linking up function inputs when types differ because of SWC/LWC custom structs

diff --git a/Engine/Plugins/FX/Niagara/Source/NiagaraEditor/Private/NiagaraGraphDigest.cpp b/Engine/Plugins/FX/Niagara/Source/NiagaraEditor/Private/NiagaraGraphDigest.cpp
--- a/Engine/Plugins/FX/Niagara/Source/NiagaraEditor/Private/NiagaraGraphDigest.cpp
+++ b/Engine/Plugins/FX/Niagara/Source/NiagaraEditor/Private/NiagaraGraphDigest.cpp
@@ -2629,11 +2629,37 @@ void FNiagaraCompilationNodeFunctionCall::Compile(FTranslator* Translator, TArra
 
 		for (const FNiagaraCompilationNodeInput* FunctionInputNode : FunctionInputNodes)
 		{
+			const FNiagaraTypeDefinition& InputNodeType = FunctionInputNode->InputVariable.GetType();
+
+			TOptional<FNiagaraTypeDefinition> SWCType;
+			if (FNiagaraTypeHelper::IsLWCType(InputNodeType))
+			{
+				SWCType = FNiagaraTypeHelper::GetSWCType(InputNodeType);
+			}
+
 			//Finds the matching Pin in the caller.
-			const FNiagaraCompilationInputPin* CallerPin = InputPins.FindByPredicate([&FunctionInputNode](const FNiagaraCompilationInputPin& InputPin) -> bool
+			auto MatchInputNodePredicate = [&FunctionInputNode, &SWCType](const FNiagaraCompilationInputPin& InputPin) -> bool
 			{
-				return InputPin.Variable.IsEquivalent(FunctionInputNode->InputVariable);
-			});
+				if (InputPin.Variable.GetName() != FunctionInputNode->InputVariable.GetName())
+				{
+					return false;
+				}
+
+				if (InputPin.Variable.IsEquivalent(FunctionInputNode->InputVariable))
+				{
+					return true;
+				}
+
+				// the last thing we need to worry about is when types differ because of LWC vs SWC concerns
+				if (SWCType.IsSet())
+				{
+					return InputPin.Variable.GetType() == *SWCType;
+				}
+
+				return false;
+			};
+
+			const FNiagaraCompilationInputPin* CallerPin = InputPins.FindByPredicate(MatchInputNodePredicate);
 
 			if (!CallerPin)
 			{
@@ -3849,8 +3875,8 @@ void FNiagaraCompilationNodeParameterMapSet::Compile(FTranslator* Translator, TA
 
 	TArray<FTranslator::FCompiledPin, TInlineAllocator<16>> CompileInputs;
 	CompileInputs.Reserve(InputPins.Num());
-
-	// update the translator with the culled function calls before compiling any further
+	// do a first pass over all of the pins so that we can properly cull out input pins and
+	// propagate the disabled pins up the chain
 	for (const FNiagaraCompilationInputPin& InputPin : InputPins)
 	{
 		if (Translator->IsFunctionVariableCulledFromCompilation(InputPin.PinName))
