Fix for Niagara translator error with new compilation mode caused by not doing a prepass over MapSet inputs to handle hidden input pins (so that they can propagate up the chain of nodes and make sure everything appropriate is disabeld).

diff --git a/Engine/Plugins/FX/Niagara/Source/NiagaraEditor/Private/NiagaraGraphDigest.cpp b/Engine/Plugins/FX/Niagara/Source/NiagaraEditor/Private/NiagaraGraphDigest.cpp
--- a/Engine/Plugins/FX/Niagara/Source/NiagaraEditor/Private/NiagaraGraphDigest.cpp
+++ b/Engine/Plugins/FX/Niagara/Source/NiagaraEditor/Private/NiagaraGraphDigest.cpp
@@ -3873,8 +3873,8 @@ void FNiagaraCompilationNodeParameterMapSet::Compile(FTranslator* Translator, TA
 	check(Outputs.Num() == 0);
 	Outputs.Init(INDEX_NONE, OutputPins.Num());
 
-	TArray<FTranslator::FCompiledPin, TInlineAllocator<16>> CompileInputs;
-	CompileInputs.Reserve(InputPins.Num());
+	TArray<const FNiagaraCompilationInputPin*, TInlineAllocator<16>> ActiveInputPins;
+	ActiveInputPins.Reserve(InputPins.Num());
 	// do a first pass over all of the pins so that we can properly cull out input pins and
 	// propagate the disabled pins up the chain
 	for (const FNiagaraCompilationInputPin& InputPin : InputPins)
@@ -3883,18 +3883,30 @@ void FNiagaraCompilationNodeParameterMapSet::Compile(FTranslator* Translator, TA
 		{
 			Translator->CullMapSetInputPin(&InputPin);
 		}
-		else if (NodeEnabled || InputPin.Variable.GetType() == FNiagaraTypeDefinition::GetParameterMapDef())
+		else
 		{
-			int32 CompiledInput = Translator->CompileInputPin(&InputPin);
+			ActiveInputPins.Add(&InputPin);
+		}
+	}
+
+	TArray<FTranslator::FCompiledPin, TInlineAllocator<16>> CompileInputs;
+	CompileInputs.Reserve(ActiveInputPins.Num());
+
+	// update the translator with the culled function calls before compiling any further
+	for (const FNiagaraCompilationInputPin* InputPin : ActiveInputPins)
+	{
+		if (NodeEnabled || InputPin->Variable.GetType() == FNiagaraTypeDefinition::GetParameterMapDef())
+		{
+			int32 CompiledInput = Translator->CompileInputPin(InputPin);
 			if (CompiledInput == INDEX_NONE)
 			{
-				Translator->Error(LOCTEXT("MapSetInputError", "Error compiling input for set node."), this, &InputPin);
+				Translator->Error(LOCTEXT("MapSetInputError", "Error compiling input for set node."), this, InputPin);
 			}
-			CompileInputs.Emplace(CompiledInput, &InputPin);
+			CompileInputs.Emplace(CompiledInput, InputPin);
 		}
 	}
 
-	if (InputPins.Num() && InputPins[0].LinkedTo)
+	if (ActiveInputPins.Num() && ActiveInputPins[0] && ActiveInputPins[0]->LinkedTo)
 	{
 		Translator->ParameterMapSet(this, CompileInputs, Outputs);
 	}
