Added a function to generate unique random numbers picked from a given range

diff --git a/Engine/Plugins/FX/Niagara/Source/Niagara/Private/NiagaraFunctionLibrary.cpp b/Engine/Plugins/FX/Niagara/Source/Niagara/Private/NiagaraFunctionLibrary.cpp
--- a/Engine/Plugins/FX/Niagara/Source/Niagara/Private/NiagaraFunctionLibrary.cpp
+++ b/Engine/Plugins/FX/Niagara/Source/Niagara/Private/NiagaraFunctionLibrary.cpp
@@ -698,6 +698,7 @@ const FName FastPathLibraryName(TEXT("FastPathLibrary"));
 const FName FastPathDot4Name(TEXT("FastPathDot4"));
 const FName FastPathTransformPositionName(TEXT("FastPathTransformPosition"));
 const FName FastMatrixToQuaternionName(TEXT("FastMatrixToQuaternion"));
+const FName FastUniqueRandomName(TEXT("FastUniqueRandom"));
 
 struct FVectorKernelFastDot4
 {
@@ -1074,6 +1075,159 @@ struct FVectorKernelFastMatrixToQuaternion
 	}
 };
 
+struct FVectorKernelFastUniqueRandom
+{
+	static FNiagaraFunctionSignature GetFunctionSignature()
+	{
+		FNiagaraFunctionSignature Sig;
+		Sig.Name = FastUniqueRandomName;
+		Sig.OwnerName = FastPathLibraryName;
+		Sig.bMemberFunction = false;
+		Sig.bRequiresContext = false;
+		Sig.SetDescription(LOCTEXT("FastUniqueRandomDesc", "Fast path for GetUniqueRandom function."));
+		Sig.AddInput(FNiagaraVariable(FNiagaraTypeDefinition::GetIntDef(), TEXT("Seed")));
+		Sig.AddInput(FNiagaraVariable(FNiagaraTypeDefinition::GetIntDef(), TEXT("NumUniqueValues")), LOCTEXT("UniqueValuesTooltip", "The number of unique values to draw from. Example: a value of 20 will generate numbers between [0..19]."));
+		FNiagaraVariable IterationInput(FNiagaraTypeDefinition::GetIntDef(), TEXT("Iterations"));
+		IterationInput.SetValue(3); // default value
+		Sig.AddInput(IterationInput, LOCTEXT("IterationTooltip", "More iterations increase entropy at the cost of performance"));
+		Sig.Outputs.Add(FNiagaraVariable(FNiagaraTypeDefinition::GetIntDef(), TEXT("Value")));
+		return Sig;
+	}
+
+	static FString GetFunctionHLSL()
+	{
+		static FString FunctionHLSL = TEXT(R"(
+			uint CalcLFSRValue(uint Seed, uint NumUniqueValues, int Iterations, uint4 Taps, uint Mask)
+			{
+				uint NextValue = (Seed + 1) & Mask;
+				for (int i = 0; i < Iterations; i++) {
+					do {
+						uint Tap = 0;
+						Tap ^= (Taps.x & NextValue) == 0 ? 0 : 1;
+						Tap ^= (Taps.y & NextValue) == 0 ? 0 : 1;
+						Tap ^= (Taps.z & NextValue) == 0 ? 0 : 1;
+						Tap ^= (Taps.w & NextValue) == 0 ? 0 : 1;
+						NextValue = (((NextValue << 1) | Tap) & Mask);
+					} while (NextValue > NumUniqueValues);
+				}
+				return NextValue - 1;
+			}
+
+			void FastUniqueRandom_FastPathLibrary(int Seed, int NumUniqueValues, int Iterations, out int Result)
+			{
+				if (NumUniqueValues < 2) {
+					NumUniqueValues = 2;
+				}
+				Seed %= NumUniqueValues;
+
+				if (NumUniqueValues < 16) {
+					uint4 Taps = {(1 << 3), (1 << 2), 0, 0};
+					uint Mask = 0xF;
+					Result = CalcLFSRValue(Seed, NumUniqueValues, Iterations, Taps, Mask);
+				}
+				else if (NumUniqueValues < 256) {
+					uint4 Taps = {(1 << 7), (1 << 5), (1 << 4), (1 << 3)};
+					uint Mask = 0xFF;
+					Result = CalcLFSRValue(Seed, NumUniqueValues, Iterations, Taps, Mask);
+				}
+				else if (NumUniqueValues < 4096) {
+					uint4 Taps = {(1 << 11), (1 << 10), (1 << 7), (1 << 5)};
+					uint Mask = 0xFFF;
+					Result = CalcLFSRValue(Seed, NumUniqueValues, Iterations, Taps, Mask);
+				}
+				else if (NumUniqueValues < 65536) {
+					uint4 Taps = {(1 << 15), (1 << 13), (1 << 12), (1 << 10)};
+					uint Mask = 0xFFFF;
+					Result = CalcLFSRValue(Seed, NumUniqueValues, Iterations, Taps, Mask);
+				}
+				else if (NumUniqueValues < 1048576) {
+					uint4 Taps ={(1 << 19), (1 << 16), 0, 0};
+					uint Mask = 0xFFFFF;
+					Result = CalcLFSRValue(Seed, NumUniqueValues, Iterations, Taps, Mask);
+				}
+				else {
+					uint4 Taps = {(1 << 23), (1 << 22), (1 << 20), (1 << 19)};
+					uint Mask = 0xFFFFFF;
+					Result = CalcLFSRValue(Seed, NumUniqueValues, Iterations, Taps, Mask);
+				}
+			}
+		)");
+		return FunctionHLSL;
+	}
+
+	static unsigned int CalcLFSRValue(const unsigned int& Seed, const unsigned int& NumUniqueValues, const int& Iterations, const FIntVector4& Taps, const unsigned int& Mask)
+	{
+		unsigned int NextValue = (Seed + 1) & Mask;
+		for (int i = 0; i < Iterations; i++) {
+			do {
+				unsigned int Tap = 0;
+				Tap ^= (Taps.X & NextValue) == 0 ? 0 : 1;
+				Tap ^= (Taps.Y & NextValue) == 0 ? 0 : 1;
+				Tap ^= (Taps.Z & NextValue) == 0 ? 0 : 1;
+				Tap ^= (Taps.W & NextValue) == 0 ? 0 : 1;
+				NextValue = (((NextValue << 1) | Tap) & Mask);
+			} while (NextValue > NumUniqueValues);
+		}
+		return NextValue - 1;
+	}
+
+	static unsigned int FastUniqueRandom_FastPathLibrary(int Seed, int NumUniqueValues, int Iterations)
+	{
+		if (NumUniqueValues < 2) {
+			NumUniqueValues = 2;
+		}
+		Seed %= NumUniqueValues;
+		
+		// Source for the tap values: Table of Linear Feedback Shift Registers: Roy Ward, Tim Molteno, 2007
+		if (NumUniqueValues < 16) {
+			FIntVector4 Taps = {(1 << 3), (1 << 2), 0, 0};
+			unsigned int Mask = 0xF;
+			return CalcLFSRValue(Seed, NumUniqueValues, Iterations, Taps, Mask);
+		}
+		if (NumUniqueValues < 256) {
+			FIntVector4 Taps = {(1 << 7), (1 << 5), (1 << 4), (1 << 3)};
+			unsigned int Mask = 0xFF;
+			return CalcLFSRValue(Seed, NumUniqueValues, Iterations, Taps, Mask);
+		}
+		if (NumUniqueValues < 4096) {
+			FIntVector4 Taps = {(1 << 11), (1 << 10), (1 << 7), (1 << 5)};
+			unsigned int Mask = 0xFFF;
+			return CalcLFSRValue(Seed, NumUniqueValues, Iterations, Taps, Mask);
+		}
+		if (NumUniqueValues < 65536) {
+			FIntVector4 Taps = {(1 << 15), (1 << 13), (1 << 12), (1 << 10)};
+			unsigned int Mask = 0xFFFF;
+			return CalcLFSRValue(Seed, NumUniqueValues, Iterations, Taps, Mask);
+		}
+		if (NumUniqueValues < 1048576) {
+			FIntVector4 Taps ={(1 << 19), (1 << 16), 0, 0};
+			unsigned int Mask = 0xFFFFF;
+			return CalcLFSRValue(Seed, NumUniqueValues, Iterations, Taps, Mask);
+		}
+		FIntVector4 Taps = {(1 << 23), (1 << 22), (1 << 20), (1 << 19)};
+		unsigned int Mask = 0xFFFFFF;
+		return CalcLFSRValue(Seed, NumUniqueValues, Iterations, Taps, Mask);
+	}
+
+	static void Exec(FVectorVMExternalFunctionContext& Context)
+	{
+		FNDIInputParam<int32> InSeed(Context);
+		FNDIInputParam<int32> InNumUniqueValues(Context);
+		FNDIInputParam<int32> InIterations(Context);
+		FNDIOutputParam<int32> OutValue(Context);
+		
+		for (int32 i = 0; i < Context.GetNumInstances(); ++i)
+		{
+			int32 Seed = InSeed.GetAndAdvance();
+			int32 NumUniqueValues = InNumUniqueValues.GetAndAdvance();
+			int32 Iterations = InIterations.GetAndAdvance();
+
+			int32 Result = FastUniqueRandom_FastPathLibrary(Seed, NumUniqueValues, Iterations);
+			OutValue.SetAndAdvance(Result);
+		}
+	}
+};
+
 bool UNiagaraFunctionLibrary::GetVectorVMFastPathExternalFunction(const FVMExternalFunctionBindingInfo& BindingInfo, FVMExternalFunction &OutFunc)
 {
 	if (BindingInfo.Name == FastPathDot4Name)
@@ -1091,6 +1245,11 @@ bool UNiagaraFunctionLibrary::GetVectorVMFastPathExternalFunction(const FVMExter
 		OutFunc = FVMExternalFunction::CreateStatic(FVectorKernelFastMatrixToQuaternion::Exec);
 		return true;
 	}
+	else if (BindingInfo.Name == FastUniqueRandomName)
+	{
+		OutFunc = FVMExternalFunction::CreateStatic(FVectorKernelFastUniqueRandom::Exec);
+		return true;
+	}
 	return false;
 }
 
@@ -1118,6 +1277,7 @@ void UNiagaraFunctionLibrary::InitVectorVMFastPathOps()
 		FVectorKernelFastDot4::GetFunctionSignature(),
 		FVectorKernelFastTransformPosition::GetFunctionSignature(),
 		FVectorKernelFastMatrixToQuaternion::GetFunctionSignature(),
+		FVectorKernelFastUniqueRandom::GetFunctionSignature(),
 	};
 
 	VectorVMOpsHLSL =
@@ -1125,6 +1285,7 @@ void UNiagaraFunctionLibrary::InitVectorVMFastPathOps()
 		FVectorKernelFastDot4::GetFunctionHLSL(),
 		FVectorKernelFastTransformPosition::GetFunctionHLSL(),
 		FVectorKernelFastMatrixToQuaternion::GetFunctionHLSL(),
+		FVectorKernelFastUniqueRandom::GetFunctionHLSL(),
 	};
 
 	check(VectorVMOps.Num() == VectorVMOpsHLSL.Num());
diff --git a/Engine/Plugins/FX/Niagara/Source/NiagaraEditor/Private/EdGraphSchema_Niagara.cpp b/Engine/Plugins/FX/Niagara/Source/NiagaraEditor/Private/EdGraphSchema_Niagara.cpp
--- a/Engine/Plugins/FX/Niagara/Source/NiagaraEditor/Private/EdGraphSchema_Niagara.cpp
+++ b/Engine/Plugins/FX/Niagara/Source/NiagaraEditor/Private/EdGraphSchema_Niagara.cpp
@@ -14,6 +14,7 @@
 #include "NiagaraEditorModule.h"
 #include "NiagaraEditorUtilities.h"
 #include "NiagaraEmitter.h"
+#include "NiagaraFunctionLibrary.h"
 #include "NiagaraGraph.h"
 #include "NiagaraHlslTranslator.h"
 #include "NiagaraNodeConvert.h"
@@ -547,6 +548,18 @@ PRAGMA_ENABLE_DEPRECATION_WARNINGS
 		{
 			AddScriptFunctionAction({LOCTEXT("Function Menu Title", "Functions").ToString()}, FunctionScriptAsset);
 		}
+
+		// Insert fast path function nodes here (visible with fx.AllowFastPathFunctionLibrary)...
+		const TArray<FNiagaraFunctionSignature>& VMFastPathOps = UNiagaraFunctionLibrary::GetVectorVMFastPathOps();
+		for (const FNiagaraFunctionSignature& Sig : VMFastPathOps)
+		{
+			UNiagaraNodeFunctionCall* FunctionCallNode = NewObject<UNiagaraNodeFunctionCall>(OwnerOfTemporaries);
+			FunctionCallNode->Signature = Sig;
+
+			const FText MenuDesc = FText::FromString(FName::NameToDisplayString(Sig.Name.ToString(), false));
+			TSharedPtr<FNiagaraAction_NewNode> FunctionCallAction = AddNewNodeMenuAction(NewActions, FunctionCallNode, MenuDesc, ENiagaraMenuSections::General, TArray<FString>(), MenuDesc);
+			FunctionCallAction->NodeTemplate = FunctionCallNode;
+		}
 	}
 
 	//Add modules
