- Fix for not unbinding delegates when the source has EndPlay called

diff --git a/Engine/Plugins/FX/Niagara/Source/Niagara/Private/DataInterface/NiagaraDataInterfaceStaticMesh.cpp b/Engine/Plugins/FX/Niagara/Source/Niagara/Private/DataInterface/NiagaraDataInterfaceStaticMesh.cpp
--- a/Engine/Plugins/FX/Niagara/Source/Niagara/Private/DataInterface/NiagaraDataInterfaceStaticMesh.cpp
+++ b/Engine/Plugins/FX/Niagara/Source/Niagara/Private/DataInterface/NiagaraDataInterfaceStaticMesh.cpp
@@ -3641,6 +3641,7 @@ void UNiagaraDataInterfaceStaticMesh::OnSourceEndPlay(AActor* InSource, EEndPlay
 {
 	// Increment change id in case we're able to find a new source component 
 	++ChangeId;
+	UnbindSourceDelegates();
 	SoftSourceActor = nullptr;
 	SourceComponent = nullptr;
 }
@@ -3649,6 +3650,7 @@ void UNiagaraDataInterfaceStaticMesh::SetDefaultMeshFromBlueprints(UStaticMesh*
 {
 	// NOTE: When ChangeId changes the next tick will be skipped and a reset of the per-instance data will be initiated. 
 	++ChangeId;
+	UnbindSourceDelegates();
 	SourceComponent = nullptr;
 	SoftSourceActor = nullptr;
 	DefaultMesh = MeshToUse;
diff --git a/Engine/Plugins/FX/Niagara/Source/Niagara/Private/NiagaraDataInterfaceSkeletalMesh.cpp b/Engine/Plugins/FX/Niagara/Source/Niagara/Private/NiagaraDataInterfaceSkeletalMesh.cpp
--- a/Engine/Plugins/FX/Niagara/Source/Niagara/Private/NiagaraDataInterfaceSkeletalMesh.cpp
+++ b/Engine/Plugins/FX/Niagara/Source/Niagara/Private/NiagaraDataInterfaceSkeletalMesh.cpp
@@ -3384,6 +3384,7 @@ void UNiagaraDataInterfaceSkeletalMesh::OnSourceEndPlay(AActor* InSource, EEndPl
 {
 	// Increment change id in case we're able to find a new source component 
 	++ChangeId;
+	UnbindSourceDelegates();
 	SoftSourceActor = nullptr;
 	SourceComponent = nullptr;
 }
