- Don't allow finalize tasks to run until the main system simulation concurrent tasks is complete
- Move tick batch into context since it's transient to the concurrent tick only

diff --git a/Engine/Plugins/FX/Niagara/Source/Niagara/Private/NiagaraSystemSimulation.cpp b/Engine/Plugins/FX/Niagara/Source/Niagara/Private/NiagaraSystemSimulation.cpp
--- a/Engine/Plugins/FX/Niagara/Source/Niagara/Private/NiagaraSystemSimulation.cpp
+++ b/Engine/Plugins/FX/Niagara/Source/Niagara/Private/NiagaraSystemSimulation.cpp
@@ -976,8 +976,8 @@ void FNiagaraSystemSimulation::SetInstanceState(FNiagaraSystemInstance* Instance
 
 void FNiagaraSystemSimulation::AddSystemToTickBatch(FNiagaraSystemInstance* Instance, FNiagaraSystemSimulationTickContext& Context)
 {
-	TickBatch.Add(Instance);
-	if (TickBatch.Num() == GNiagaraSystemSimulationTickBatchSize)
+	Context.TickBatch.Add(Instance);
+	if (Context.TickBatch.Num() == GNiagaraSystemSimulationTickBatchSize)
 	{
 		FlushTickBatch(Context);
 	}
@@ -985,46 +985,44 @@ void FNiagaraSystemSimulation::AddSystemToTickBatch(FNiagaraSystemInstance* Inst
 
 void FNiagaraSystemSimulation::FlushTickBatch(FNiagaraSystemSimulationTickContext& Context)
 {
-	if (TickBatch.Num() > 0)
+	if (Context.TickBatch.Num() == 0)
 	{
-		// If we are running async create tasks to execute
-		if ( Context.IsRunningAsync() )
-		{
-			// Queue instance concurrent task and track information in the instance
-			FGraphEventRef InstanceAsyncGraphEvent = TGraphTask<FNiagaraSystemInstanceTickConcurrentTask>::CreateTask(&Context.BeforeInstancesTickGraphEvents).ConstructAndDispatchWhenReady(this, TickBatch, Context.World);
+		return;
+	}
 
-			for (FNiagaraSystemInstance* Inst : TickBatch)
-			{
-				Inst->ConcurrentTickBatchGraphEvent = InstanceAsyncGraphEvent;
-			}
+	// If we are running async create tasks to execute
+	if ( Context.IsRunningAsync() )
+	{
+		// Queue instance concurrent task and track information in the instance
+		FGraphEventRef InstanceAsyncGraphEvent = TGraphTask<FNiagaraSystemInstanceTickConcurrentTask>::CreateTask(&Context.BeforeInstancesTickGraphEvents).ConstructAndDispatchWhenReady(this, Context.TickBatch, Context.World);
 
-			// Ensure ConcurrentTickBatchGraphEvent is visible before we clear ConcurrentTickGraphEvent
-			FPlatformMisc::MemoryBarrier();
+		for (FNiagaraSystemInstance* Inst : Context.TickBatch)
+		{
+			Inst->ConcurrentTickBatchGraphEvent = InstanceAsyncGraphEvent;
+		}
 
-			for (FNiagaraSystemInstance* Inst : TickBatch)
-			{
-				Inst->ConcurrentTickGraphEvent = nullptr;
-			}
+		// Ensure ConcurrentTickBatchGraphEvent is visible before we clear ConcurrentTickGraphEvent
+		FPlatformMisc::MemoryBarrier();
 
-			// Queue finalize task which will run after the instances are complete, track with our all completion event
-			FGraphEventArray FinalizePrereqArray;
-			FinalizePrereqArray.Add(InstanceAsyncGraphEvent);
-			FGraphEventRef FinalizeTask = TGraphTask<FNiagaraSystemInstanceFinalizeTask>::CreateTask(&FinalizePrereqArray).ConstructAndDispatchWhenReady(this, TickBatch);
+		// Queue finalize task which will run after the instances are complete, track with our all completion event
+		FGraphEventArray FinalizePrereqArray;
+		FinalizePrereqArray.Add(Context.TickBatch[0]->ConcurrentTickGraphEvent);
+		FinalizePrereqArray.Add(InstanceAsyncGraphEvent);
+		FGraphEventRef FinalizeTask = TGraphTask<FNiagaraSystemInstanceFinalizeTask>::CreateTask(&FinalizePrereqArray).ConstructAndDispatchWhenReady(this, Context.TickBatch);
 
-			check(Context.CompletionEvents != nullptr);
-			Context.CompletionEvents->Add(FinalizeTask);
-		}
-		// Execute immediately
-		else
+		check(Context.CompletionEvents != nullptr);
+		Context.CompletionEvents->Add(FinalizeTask);
+	}
+	// Execute immediately
+	else
+	{
+		for (FNiagaraSystemInstance* Inst : Context.TickBatch)
 		{
-			for (FNiagaraSystemInstance* Inst : TickBatch)
-			{
-				Inst->Tick_Concurrent();
-			}
+			Inst->Tick_Concurrent();
 		}
-
-		TickBatch.Reset();
 	}
+
+	Context.TickBatch.Reset();
 }
 
 /** First phase of system sim tick. Must run on GameThread. */
diff --git a/Engine/Plugins/FX/Niagara/Source/Niagara/Public/NiagaraSystemSimulation.h b/Engine/Plugins/FX/Niagara/Source/Niagara/Public/NiagaraSystemSimulation.h
--- a/Engine/Plugins/FX/Niagara/Source/Niagara/Public/NiagaraSystemSimulation.h
+++ b/Engine/Plugins/FX/Niagara/Source/Niagara/Public/NiagaraSystemSimulation.h
@@ -240,6 +240,8 @@ public:
 	bool								bRunningAsync = false;
 	FGraphEventArray					BeforeInstancesTickGraphEvents;
 	FGraphEventArray*					CompletionEvents = nullptr;
+
+	FNiagaraSystemTickBatch				TickBatch;
 };
 
 /** Simulation performing all system and emitter scripts for a instances of a UNiagaraSystem in a world. */
@@ -417,9 +419,6 @@ protected:
 
 	TOptional<float> MaxDeltaTime;
 
-	/** Current tick batch we're filling ready for processing, potentially in an async task. */
-	FNiagaraSystemTickBatch TickBatch;
-
 	/** Event to track the system simulation async tick is complete. */
 	FGraphEventRef ConcurrentTickGraphEvent;
 	/** Event to track all work is complete, i.e. System Concurrent, Instance Concurrent, Finalize */
