- Clean up to use resolved DIs
- And rename user data interfaces to find to just to find since that's what it was doing

diff --git a/Engine/Plugins/FX/Niagara/Source/NiagaraEditor/Private/Commandlets/NiagaraSystemAuditCommandlet.cpp b/Engine/Plugins/FX/Niagara/Source/NiagaraEditor/Private/Commandlets/NiagaraSystemAuditCommandlet.cpp
--- a/Engine/Plugins/FX/Niagara/Source/NiagaraEditor/Private/Commandlets/NiagaraSystemAuditCommandlet.cpp
+++ b/Engine/Plugins/FX/Niagara/Source/NiagaraEditor/Private/Commandlets/NiagaraSystemAuditCommandlet.cpp
@@ -18,6 +18,7 @@
 #include "ViewModels/NiagaraSystemViewModel.h"
 #include "NiagaraSettings.h"
 #include "NiagaraSystem.h"
+#include "NiagaraSystemImpl.h"
 #include "NiagaraSimulationStageBase.h"
 #include "NiagaraRendererProperties.h"
 #include "NiagaraLightRendererProperties.h"
@@ -87,16 +88,16 @@ int32 UNiagaraSystemAuditCommandlet::Main(const FString& Params)
 	INiagaraModule& NiagaraModule = FModuleManager::LoadModuleChecked<INiagaraModule>("Niagara");
 	// User Data Interfaces to Find
 	{
-		FString UserDataInterfacesToFindString;
-		if (FParse::Value(*Params, TEXT("UserDataInterfacesToFind="), UserDataInterfacesToFindString, false))
+		FString DataInterfacesToFindString;
+		if (FParse::Value(*Params, TEXT("DataInterfacesToFind="), DataInterfacesToFindString, false))
 		{
 			TArray<FString> DataInterfaceNames;
-			UserDataInterfacesToFindString.ParseIntoArray(DataInterfaceNames, TEXT(","));
+			DataInterfacesToFindString.ParseIntoArray(DataInterfaceNames, TEXT(","));
 			for (const FString& DIName : DataInterfaceNames)
 			{
 				if (UClass* FoundClass = UClass::TryFindTypeSlow<UClass>(DIName, EFindFirstObjectOptions::ExactClass))
 				{
-					UserDataInterfacesToFind.Add(FoundClass);
+					DataInterfacesToFind.Add(FoundClass);
 				}
 				else
 				{
@@ -277,7 +278,7 @@ bool UNiagaraSystemAuditCommandlet::ProcessNiagaraSystems()
 			{
 				SystemDataInterfacesWihPrereqs.Add(DataInterface->GetClass()->GetFName());
 			}
-			if (UserDataInterfacesToFind.Contains(DataInterface->GetClass()))
+			if (DataInterfacesToFind.Contains(DataInterface->GetClass()))
 			{
 				SystemUserDataInterfaces.Add(DataInterface->GetClass()->GetFName());
 			}
@@ -460,7 +461,7 @@ bool UNiagaraSystemAuditCommandlet::ProcessNiagaraSystems()
 				}
 				DataInterfaceNames.Append(*it->ToString());
 			}
-			NiagaraSystemsWithUserDataInterface.Add(FString::Printf(TEXT("%s,%s"), *NiagaraSystem->GetPathName(), *DataInterfaceNames));
+			NiagaraSystemsWithDataInterfaceToFind.Add(FString::Printf(TEXT("%s,%s"), *NiagaraSystem->GetPathName(), *DataInterfaceNames));
 		}
 
 		// Run validation last as it will potentially modify the asset
@@ -501,9 +502,9 @@ void UNiagaraSystemAuditCommandlet::DumpResults()
 	DumpSimpleSet(NiagaraSystemsWithEvents, TEXT("NiagaraSystemsWithEvents"), TEXT("Name"));
 	DumpSimpleSet(NiagaraSystemsWithPrerequisites, TEXT("NiagaraSystemsWithPrerequisites"), TEXT("Name,DataInterface"));
 	DumpSimpleSet(NiagaraSystemsWithDynamicBounds, TEXT("NiagaraSystemsWithDynamicBounds"), TEXT("Name,Emitters With Dynamic Bounds"));
-	if (UserDataInterfacesToFind.Num() > 0)
+	if (DataInterfacesToFind.Num() > 0)
 	{
-		DumpSimpleSet(NiagaraSystemsWithUserDataInterface, TEXT("NiagaraSystemsWithUserDataInterface"), TEXT("Name,DataInterface"));
+		DumpSimpleSet(NiagaraSystemsWithDataInterfaceToFind, TEXT("NiagaraSystemsWithDataInterfaceToFind"), TEXT("Name,DataInterface"));
 	}
 	if (NiagaraSystemsWithGPUEmitters.Num() > 0)
 	{
@@ -592,44 +593,26 @@ FArchive* UNiagaraSystemAuditCommandlet::GetOutputFile(const TCHAR* InShortFilen
 TArray<class UNiagaraDataInterface*> UNiagaraSystemAuditCommandlet::GetDataInterfaces(class UNiagaraSystem* NiagaraSystem)
 {
 	TArray<UNiagaraDataInterface*> DataInterfaces;
-	for (UNiagaraDataInterface* ParamDI : NiagaraSystem->GetExposedParameters().GetDataInterfaces())
+	for (UNiagaraDataInterface* DataInterface : NiagaraSystem->GetExposedParameters().GetDataInterfaces())
 	{
-		if (ParamDI != nullptr)
+		if (DataInterface != nullptr)
 		{
-			DataInterfaces.AddUnique(ParamDI);
+			DataInterfaces.AddUnique(DataInterface);
 		}
 	}
 
-	auto GatherScriptDIs =
-		[&](UNiagaraScript* NiagaraScript)
+	NiagaraSystem->ForEachScript(
+		[&DataInterfaces](UNiagaraScript* NiagaraScript)
 		{
-			for (const FNiagaraScriptDataInterfaceInfo& DataInterfaceInfo : NiagaraScript->GetCachedDefaultDataInterfaces())
+			for (const FNiagaraScriptResolvedDataInterfaceInfo& ResolveDataInterfaceInfo : NiagaraScript->GetResolvedDataInterfaces())
 			{
-				if ( UNiagaraDataInterface* ScriptDI = DataInterfaceInfo.DataInterface )
+				if ( UNiagaraDataInterface* DataInterface = ResolveDataInterfaceInfo.ResolvedDataInterface )
 				{
-					DataInterfaces.AddUnique(ScriptDI);
+					DataInterfaces.AddUnique(DataInterface);
 				}
 			}
-		};
-
-	GatherScriptDIs(NiagaraSystem->GetSystemSpawnScript());
-	GatherScriptDIs(NiagaraSystem->GetSystemUpdateScript());
-
-	for (const FNiagaraEmitterHandle& EmitterHandle : NiagaraSystem->GetEmitterHandles())
-	{
-		FVersionedNiagaraEmitterData* EmitterData = EmitterHandle.GetEmitterData();
-		if (EmitterData == nullptr)
-		{
-			continue;
 		}
-
-		TArray<UNiagaraScript*> EmitterScripts;
-		EmitterData->GetScripts(EmitterScripts);
-		for (UNiagaraScript* Script : EmitterScripts)
-		{
-			GatherScriptDIs(Script);
-		}
-	}
+	);
 	return DataInterfaces;
 }
 
diff --git a/Engine/Plugins/FX/Niagara/Source/NiagaraEditor/Public/Commandlets/NiagaraSystemAuditCommandlet.h b/Engine/Plugins/FX/Niagara/Source/NiagaraEditor/Public/Commandlets/NiagaraSystemAuditCommandlet.h
--- a/Engine/Plugins/FX/Niagara/Source/NiagaraEditor/Public/Commandlets/NiagaraSystemAuditCommandlet.h
+++ b/Engine/Plugins/FX/Niagara/Source/NiagaraEditor/Public/Commandlets/NiagaraSystemAuditCommandlet.h
@@ -24,8 +24,8 @@ class UNiagaraSystemAuditCommandlet : public UCommandlet
 	TSet<FString> NiagaraSystemsWithPrerequisites;
 	/** All Niagara systems that contain emitters with dynamic bounds */
 	TSet<FString> NiagaraSystemsWithDynamicBounds;
-	/** All Niagara systems that use a user requested data interface */
-	TSet<FString> NiagaraSystemsWithUserDataInterface;
+	/** All Niagara systems that contain the requested data interface */
+	TSet<FString> NiagaraSystemsWithDataInterfaceToFind;
 	/** All Niagara Systems & Emitters with simulation stages enabled*/
 	TSet<FString> NiagaraSystemsWithSimulationStages;
 
@@ -56,7 +56,7 @@ class UNiagaraSystemAuditCommandlet : public UCommandlet
 	TArray<FName> PackagePaths;
 
 	/** Systems using specific data interfaces to find */
-	TSet<UClass*> UserDataInterfacesToFind;
+	TSet<UClass*> DataInterfacesToFind;
 
 	/** DeviceProfiles we want to disable GPU emitters for. */
 	TArray<class UDeviceProfile*> DeviceProfilesToDisableGpu;
