Do not use Cast as a const_cast
#rnx

diff --git a/Engine/Plugins/Experimental/Mutable/Source/CustomizableObjectEditor/Classes/MuCOE/Nodes/CustomizableObjectNodeTable.h b/Engine/Plugins/Experimental/Mutable/Source/CustomizableObjectEditor/Classes/MuCOE/Nodes/CustomizableObjectNodeTable.h
--- a/Engine/Plugins/Experimental/Mutable/Source/CustomizableObjectEditor/Classes/MuCOE/Nodes/CustomizableObjectNodeTable.h
+++ b/Engine/Plugins/Experimental/Mutable/Source/CustomizableObjectEditor/Classes/MuCOE/Nodes/CustomizableObjectNodeTable.h
@@ -303,10 +303,10 @@ public:
 	void ChangeImagePinMode(UEdGraphPin* Pin, bool bSetDefault = false);
 
 	// Returns true if the pin is in the default mode (same as node)
-	bool IsImagePinDefault(UEdGraphPin* Pin);
+	bool IsImagePinDefault(const UEdGraphPin* Pin) const;
 
 	// Returns true if the pin is a texture array pin
-	bool IsImageArrayPin(UEdGraphPin* Pin);
+	bool IsImageArrayPin(const UEdGraphPin* Pin) const;
 
 	// Returns the image mode of the column
 	ETableTextureType GetColumnImageMode(const FString& ColumnName) const;
diff --git a/Engine/Plugins/Experimental/Mutable/Source/CustomizableObjectEditor/Private/MuCOE/EdGraphSchema_CustomizableObject.cpp b/Engine/Plugins/Experimental/Mutable/Source/CustomizableObjectEditor/Private/MuCOE/EdGraphSchema_CustomizableObject.cpp
--- a/Engine/Plugins/Experimental/Mutable/Source/CustomizableObjectEditor/Private/MuCOE/EdGraphSchema_CustomizableObject.cpp
+++ b/Engine/Plugins/Experimental/Mutable/Source/CustomizableObjectEditor/Private/MuCOE/EdGraphSchema_CustomizableObject.cpp
@@ -766,7 +766,7 @@ void UEdGraphSchema_CustomizableObject::GetContextMenuActionsReconstructAllChild
 				return;
 			}
 
-			UCustomizableObjectNode* Node = Cast<UCustomizableObjectNode>(Context->Node);
+			UCustomizableObjectNode* Node = const_cast<UCustomizableObjectNode*>(Cast<UCustomizableObjectNode>(Context->Node));
 			if (!Node)
 			{
 				return;
@@ -858,7 +858,7 @@ void UEdGraphSchema_CustomizableObject::GetContextMenuActions(UToolMenu* Menu, U
 	}
 	else // On Pin right click
 	{
-		UCustomizableObjectNodeTable* TableNode = Cast<UCustomizableObjectNodeTable>(Context->Node);
+		const UCustomizableObjectNodeTable* TableNode = Cast<UCustomizableObjectNodeTable>(Context->Node);
 		UEdGraphPin* TexturePin = (UEdGraphPin*)Context->Pin;
 
 		if (TableNode && TexturePin && !TableNode->IsImageArrayPin(TexturePin) && !TexturePin->LinkedTo.Num() && (TexturePin->PinType.PinCategory == PC_Image || TexturePin->PinType.PinCategory == PC_PassThroughImage))
@@ -875,7 +875,7 @@ void UEdGraphSchema_CustomizableObject::GetContextMenuActions(UToolMenu* Menu, U
 			Section.AddMenuEntry
 			(
 				"ChangeTexturePinMode", ActionText, ToolTipText, FSlateIcon(),
-				FUIAction(FExecuteAction::CreateLambda([TableNode, TexturePin]()
+				FUIAction(FExecuteAction::CreateLambda([TableNode = const_cast<UCustomizableObjectNodeTable*>(TableNode), TexturePin]()
 					{
 						TableNode->ChangeImagePinMode(TexturePin);
 					}))
@@ -889,7 +889,7 @@ void UEdGraphSchema_CustomizableObject::GetContextMenuActions(UToolMenu* Menu, U
 					LOCTEXT("SetTexturePinModeDefault_Label","Set pin as default."),
 					LOCTEXT("SetTexturePinModeDefault_Tooltip","Set the selected texture pin to use the default node mode."),
 					FSlateIcon(),
-					FUIAction(FExecuteAction::CreateLambda([TableNode, TexturePin]()
+					FUIAction(FExecuteAction::CreateLambda([TableNode = const_cast<UCustomizableObjectNodeTable*>(TableNode), TexturePin]()
 						{
 							TableNode->ChangeImagePinMode(TexturePin, true);
 						}))
diff --git a/Engine/Plugins/Experimental/Mutable/Source/CustomizableObjectEditor/Private/MuCOE/Nodes/CustomizableObjectNodeTable.cpp b/Engine/Plugins/Experimental/Mutable/Source/CustomizableObjectEditor/Private/MuCOE/Nodes/CustomizableObjectNodeTable.cpp
--- a/Engine/Plugins/Experimental/Mutable/Source/CustomizableObjectEditor/Private/MuCOE/Nodes/CustomizableObjectNodeTable.cpp
+++ b/Engine/Plugins/Experimental/Mutable/Source/CustomizableObjectEditor/Private/MuCOE/Nodes/CustomizableObjectNodeTable.cpp
@@ -1150,7 +1150,7 @@ UCustomizableObjectNodeTableRemapPins* UCustomizableObjectNodeTable::CreateRemap
 }
 
 
-bool UCustomizableObjectNodeTable::IsImagePinDefault(UEdGraphPin* Pin)
+bool UCustomizableObjectNodeTable::IsImagePinDefault(const UEdGraphPin* Pin) const
 {
 	if (UCustomizableObjectNodeTableImagePinData* PinData = Cast<UCustomizableObjectNodeTableImagePinData>(GetPinData(*(Pin))))
 	{
@@ -1161,7 +1161,7 @@ bool UCustomizableObjectNodeTable::IsImagePinDefault(UEdGraphPin* Pin)
 }
 
 
-bool UCustomizableObjectNodeTable::IsImageArrayPin(UEdGraphPin* Pin)
+bool UCustomizableObjectNodeTable::IsImageArrayPin(const UEdGraphPin* Pin) const
 {
 	if (UCustomizableObjectNodeTableImagePinData* PinData = Cast<UCustomizableObjectNodeTableImagePinData>(GetPinData(*(Pin))))
 	{
@@ -1174,9 +1174,9 @@ bool UCustomizableObjectNodeTable::IsImageArrayPin(UEdGraphPin* Pin)
 
 ETableTextureType UCustomizableObjectNodeTable::GetColumnImageMode(const FString& ColumnName) const
 {
-	for (UEdGraphPin* Pin : Pins)
+	for (const UEdGraphPin* Pin : Pins)
 	{
-		if (UCustomizableObjectNodeTableImagePinData* PinData = Cast<UCustomizableObjectNodeTableImagePinData>(GetPinData(*(Pin))))
+		if (const UCustomizableObjectNodeTableImagePinData* PinData = Cast<UCustomizableObjectNodeTableImagePinData>(GetPinData(*(Pin))))
 		{
 			if (PinData->ColumnName == ColumnName)
 			{
diff --git a/Engine/Plugins/FX/Niagara/Source/NiagaraEditor/Private/NiagaraActions.cpp b/Engine/Plugins/FX/Niagara/Source/NiagaraEditor/Private/NiagaraActions.cpp
--- a/Engine/Plugins/FX/Niagara/Source/NiagaraEditor/Private/NiagaraActions.cpp
+++ b/Engine/Plugins/FX/Niagara/Source/NiagaraEditor/Private/NiagaraActions.cpp
@@ -747,8 +747,8 @@ void FNiagaraDataInterfaceNodeActionProvider_DataChannelWrite::GetNodeContextMen
 
 	FToolMenuSection& Section = Menu->AddSection("DataChannelWrite", InitForDataChannelHeaderText);
 
-	UNiagaraNodeFunctionCall* FuncNode = CastChecked<UNiagaraNodeFunctionCall>(Context->Node);
-	TWeakObjectPtr<UNiagaraNodeFunctionCall> WeakNode = FuncNode;
+	const UNiagaraNodeFunctionCall* FuncNode = CastChecked<UNiagaraNodeFunctionCall>(Context->Node);
+	TWeakObjectPtr<const UNiagaraNodeFunctionCall> WeakNode = FuncNode;
 
 	auto CreateNodeContextMenu = [WeakNode](UToolMenu* InNewToolMenu)
 	{
@@ -762,7 +762,7 @@ void FNiagaraDataInterfaceNodeActionProvider_DataChannelWrite::GetNodeContextMen
 				auto CreateDataChannelActionEntry = [WeakChannel, WeakNode]()
 				{
 					UNiagaraDataChannel* Channel = WeakChannel.Get();
-					UNiagaraNodeFunctionCall* Node = WeakNode.Get();
+					UNiagaraNodeFunctionCall* Node = const_cast<UNiagaraNodeFunctionCall*>(WeakNode.Get());
 					if(Channel && Node)
 					{
 						Node->RemoveAllDynamicPins();
@@ -806,8 +806,8 @@ void FNiagaraDataInterfaceNodeActionProvider_DataChannelRead::GetNodeContextMenu
 		return;
 	}
 	
-	UNiagaraNodeFunctionCall* FuncNode = CastChecked<UNiagaraNodeFunctionCall>(Context->Node);
-	TWeakObjectPtr<UNiagaraNodeFunctionCall> WeakNode = FuncNode;
+	const UNiagaraNodeFunctionCall* FuncNode = CastChecked<UNiagaraNodeFunctionCall>(Context->Node);
+	TWeakObjectPtr<const UNiagaraNodeFunctionCall> WeakNode = FuncNode;
 
 	// as the lambda gives us the tool menu, it's important to use that instead of the menu we are being passed in from the function
 	auto CreateNodeContextMenu = [WeakNode](UToolMenu* InToolMenu)
@@ -909,8 +909,8 @@ void FNiagaraDataInterfaceNodeActionProvider_DataChannelRead::AddDataChannelInit
 		return;
 	}
 	
-	UNiagaraNodeFunctionCall* FuncNode = CastChecked<UNiagaraNodeFunctionCall>(Context->Node);
-	TWeakObjectPtr<UNiagaraNodeFunctionCall> WeakNode = FuncNode;
+	const UNiagaraNodeFunctionCall* FuncNode = CastChecked<UNiagaraNodeFunctionCall>(Context->Node);
+	TWeakObjectPtr<const UNiagaraNodeFunctionCall> WeakNode = FuncNode;
 
 	FToolMenuSection& MenuSection = ToolMenu->AddSection("DataChannelRead", NiagaraActionsLocal::InitForDataChannelHeaderText);
 	auto InitForDataChannelSection = [&MenuSection, WeakNode](UNiagaraDataChannel* DataChannel)
@@ -921,7 +921,7 @@ void FNiagaraDataInterfaceNodeActionProvider_DataChannelRead::AddDataChannelInit
 			auto CreateDataChannelActionEntry = [WeakChannel, WeakNode]()
 			{
 				UNiagaraDataChannel* Channel = WeakChannel.Get();
-				UNiagaraNodeFunctionCall* Node = WeakNode.Get();
+				UNiagaraNodeFunctionCall* Node = const_cast<UNiagaraNodeFunctionCall*>(WeakNode.Get());
 
 				if (Channel && Node)
 				{
diff --git a/Engine/Source/Editor/AnimGraph/Private/AnimGraphNode_Base.cpp b/Engine/Source/Editor/AnimGraph/Private/AnimGraphNode_Base.cpp
--- a/Engine/Source/Editor/AnimGraph/Private/AnimGraphNode_Base.cpp
+++ b/Engine/Source/Editor/AnimGraph/Private/AnimGraphNode_Base.cpp
@@ -1023,7 +1023,7 @@ void UAnimGraphNode_Base::AutowireNewNode(UEdGraphPin* FromPin)
 
 TSharedRef<SWidget> UAnimGraphNode_Base::MakePropertyBindingWidget(const FAnimPropertyBindingWidgetArgs& InArgs)
 {
-	UAnimGraphNode_Base* FirstAnimGraphNode = InArgs.Nodes[0];
+	const UAnimGraphNode_Base* FirstAnimGraphNode = InArgs.Nodes[0];
 	UClass* BindingClass = FirstAnimGraphNode->Binding ? FirstAnimGraphNode->Binding->GetClass() : nullptr;
 	UAnimGraphNodeBinding* BindingCDO = BindingClass ? CastChecked<UAnimGraphNodeBinding>(BindingClass->GetDefaultObject()) : nullptr;
 
diff --git a/Engine/Source/Editor/AnimGraph/Private/AnimationGraphSchema.cpp b/Engine/Source/Editor/AnimGraph/Private/AnimationGraphSchema.cpp
--- a/Engine/Source/Editor/AnimGraph/Private/AnimationGraphSchema.cpp
+++ b/Engine/Source/Editor/AnimGraph/Private/AnimationGraphSchema.cpp
@@ -813,7 +813,7 @@ void UAnimationGraphSchema::GetContextMenuActions(UToolMenu* Menu, UGraphNodeCon
 {
 	Super::GetContextMenuActions(Menu, Context);
 	
-	if (UAnimGraphNode_Base* AnimGraphNode = Cast<UAnimGraphNode_Base>(Context->Node))
+	if (const UAnimGraphNode_Base* AnimGraphNode = Cast<UAnimGraphNode_Base>(Context->Node))
 	{
 		{
 			// Node contextual actions
@@ -827,7 +827,7 @@ void UAnimationGraphSchema::GetContextMenuActions(UToolMenu* Menu, UGraphNodeCon
 
 		if(Context->Pin && !IsPosePin(Context->Pin->PinType))
 		{
-			TSharedPtr<SWidget> BindingWidget = MakeBindingWidgetForPin({ AnimGraphNode }, Context->Pin->GetFName(), false, true);
+			TSharedPtr<SWidget> BindingWidget = MakeBindingWidgetForPin({ const_cast<UAnimGraphNode_Base*>(AnimGraphNode) }, Context->Pin->GetFName(), false, true);
 			if(BindingWidget.IsValid())
 			{
 				FToolMenuSection& Section = Menu->AddSection("EdGraphSchemaPinActions");
