- Fix a crash where the view model was no longer valid going out of scalability view

diff --git a/Engine/Plugins/FX/Niagara/Source/NiagaraEditor/Private/ViewModels/Stack/NiagaraStackEventScriptItemGroup.cpp b/Engine/Plugins/FX/Niagara/Source/NiagaraEditor/Private/ViewModels/Stack/NiagaraStackEventScriptItemGroup.cpp
--- a/Engine/Plugins/FX/Niagara/Source/NiagaraEditor/Private/ViewModels/Stack/NiagaraStackEventScriptItemGroup.cpp
+++ b/Engine/Plugins/FX/Niagara/Source/NiagaraEditor/Private/ViewModels/Stack/NiagaraStackEventScriptItemGroup.cpp
@@ -372,8 +372,15 @@ bool UNiagaraStackEventScriptItemGroup::HasBaseEventHandler() const
 {
 	if (bHasBaseEventHandlerCache.IsSet() == false)
 	{
-		FVersionedNiagaraEmitter BaseEmitter = GetEmitterViewModel()->GetParentEmitter();
-		bHasBaseEventHandlerCache = BaseEmitter.Emitter != nullptr && FNiagaraScriptMergeManager::Get()->HasBaseEventHandler(BaseEmitter, GetScriptUsageId());
+		if (GetEmitterViewModel().IsValid())
+		{
+			FVersionedNiagaraEmitter BaseEmitter = GetEmitterViewModel()->GetParentEmitter();
+			bHasBaseEventHandlerCache = BaseEmitter.Emitter != nullptr && FNiagaraScriptMergeManager::Get()->HasBaseEventHandler(BaseEmitter, GetScriptUsageId());
+		}
+		else
+		{
+			bHasBaseEventHandlerCache = false;
+		}
 	}
 	return bHasBaseEventHandlerCache.GetValue();
 }
