- Handle rename / remove for object parameter bindings from renderers

diff --git a/Engine/Plugins/FX/Niagara/Source/Niagara/Private/NiagaraRendererProperties.cpp b/Engine/Plugins/FX/Niagara/Source/Niagara/Private/NiagaraRendererProperties.cpp
--- a/Engine/Plugins/FX/Niagara/Source/Niagara/Private/NiagaraRendererProperties.cpp
+++ b/Engine/Plugins/FX/Niagara/Source/Niagara/Private/NiagaraRendererProperties.cpp
@@ -5,6 +5,7 @@
 #include "NiagaraDataSet.h"
 #include "NiagaraConstants.h"
 #include "NiagaraEmitter.h"
+#include "NiagaraParameterBinding.h"
 #include "NiagaraScriptSourceBase.h"
 #include "NiagaraSettings.h"
 #include "NiagaraSystem.h"
@@ -369,6 +370,10 @@ void UNiagaraRendererProperties::RenameEmitter(const FName& InOldName, const UNi
 {
 	const ENiagaraRendererSourceDataMode SourceMode = GetCurrentSourceMode();
 	UpdateSourceModeDerivates(SourceMode);
+	if (InRenamedEmitter)
+	{
+		FNiagaraParameterBinding::ForEachRenameEmitter(this, InRenamedEmitter->GetUniqueEmitterName());
+	}
 }
 
 TArray<FNiagaraVariable> UNiagaraRendererProperties::GetBoundAttributes() const
@@ -683,7 +688,12 @@ void UNiagaraRendererProperties::RenameVariable(const FNiagaraVariableBase& OldV
 		if (Binding)
 			Binding->RenameVariableIfMatching(OldVariable, NewVariable, InEmitter, GetCurrentSourceMode());
 	}
+	if (InEmitter.Emitter)
+	{
+		FNiagaraParameterBinding::ForEachRenameVariable(this, OldVariable, NewVariable, InEmitter.Emitter->GetUniqueEmitterName());
+	}
 }
+
 void UNiagaraRendererProperties::RemoveVariable(const FNiagaraVariableBase& OldVariable,const FVersionedNiagaraEmitter& InEmitter)
 {
 	// Handle the reset to defaults of generic renderer bindings
@@ -708,6 +718,11 @@ void UNiagaraRendererProperties::RemoveVariable(const FNiagaraVariableBase& OldV
 		}
 			
 	}
+
+	if (InEmitter.Emitter)
+	{
+		FNiagaraParameterBinding::ForEachRemoveVariable(this, OldVariable, InEmitter.Emitter->GetUniqueEmitterName());
+	}
 }
 
 #endif
