- Engine.Owner.Velocity will read the owner actors velocity for the initial frame rather than defaulting to zero
- This fixes an issue where reading this attribute in emitter / system spawn scripts or a burst on the first frame would always have zero velocity

diff --git a/Engine/Plugins/FX/Niagara/Source/Niagara/Classes/NiagaraSystem.h b/Engine/Plugins/FX/Niagara/Source/Niagara/Classes/NiagaraSystem.h
--- a/Engine/Plugins/FX/Niagara/Source/Niagara/Classes/NiagaraSystem.h
+++ b/Engine/Plugins/FX/Niagara/Source/Niagara/Classes/NiagaraSystem.h
@@ -491,6 +491,8 @@ public:
 	bool SupportsLargeWorldCoordinates() const { return bSupportLargeWorldCoordinates && bLwcEnabledSettingCached; }
 	FORCEINLINE bool ShouldDisableExperimentalVM() const { return bDisableExperimentalVM; }
 
+	FORCEINLINE bool IsInitialOwnerVelocityFromActor() const { return bInitialOwnerVelocityFromActor; }
+
 #if WITH_EDITORONLY_DATA
 	NIAGARA_API bool UsesEmitter(UNiagaraEmitter* Emitter) const;
 	NIAGARA_API bool UsesEmitter(const FVersionedNiagaraEmitter& VersionedEmitter) const;
@@ -607,6 +609,15 @@ public:
 	UPROPERTY(EditAnywhere, AdvancedDisplay, Category = "Performance", meta = (DisplayName = "Disable Experimental VM"))
 	uint8 bDisableExperimentalVM : 1;
 
+private:
+	/*
+	When enabled we use the owner actor's velocity for the first frame.
+	If we do not have an owner actor, or this is disable then the first frame's velocity will be zero.
+	*/
+	UPROPERTY(EditAnywhere, AdvancedDisplay, Category = "System")
+	uint8 bInitialOwnerVelocityFromActor : 1 = true;
+
+public:
 	/**
 	When enabled this is the default value set on the component.
 	Mask used for stencil buffer writes.
diff --git a/Engine/Plugins/FX/Niagara/Source/Niagara/Private/NiagaraSystem.cpp b/Engine/Plugins/FX/Niagara/Source/Niagara/Private/NiagaraSystem.cpp
--- a/Engine/Plugins/FX/Niagara/Source/Niagara/Private/NiagaraSystem.cpp
+++ b/Engine/Plugins/FX/Niagara/Source/Niagara/Private/NiagaraSystem.cpp
@@ -983,6 +983,11 @@ void UNiagaraSystem::PostLoad()
 	}
 #endif
 
+	if (NiagaraVer < FNiagaraCustomVersion::InitialOwnerVelocityFromActor)
+	{
+		bInitialOwnerVelocityFromActor = false;
+	}
+
 	for (FNiagaraSystemScalabilityOverride& Override : SystemScalabilityOverrides.Overrides)
 	{
 		Override.PostLoad(NiagaraVer);
diff --git a/Engine/Plugins/FX/Niagara/Source/Niagara/Private/NiagaraSystemInstance.cpp b/Engine/Plugins/FX/Niagara/Source/Niagara/Private/NiagaraSystemInstance.cpp
--- a/Engine/Plugins/FX/Niagara/Source/Niagara/Private/NiagaraSystemInstance.cpp
+++ b/Engine/Plugins/FX/Niagara/Source/Niagara/Private/NiagaraSystemInstance.cpp
@@ -739,6 +739,7 @@ void FNiagaraSystemInstance::OnPooledReuse(UWorld& NewWorld)
 
 	FixedBounds_GT.Init();
 	FixedBounds_CNC.Init();
+	PreviousLocation.Reset();
 
 	for (auto&& Emitter : Emitters)
 	{
@@ -1904,19 +1905,30 @@ void FNiagaraSystemInstance::TickInstanceParameters_GameThread(float DeltaSecond
 	{
 		WorldTransform = AttachComponent->GetComponentToWorld();
 		WorldTransform.AddToTranslation(FVector(LWCTile) * -FLargeWorldRenderScalar::GetTileSize());
-	}
-	const bool TransformMatches = GatheredInstanceParameters.ComponentTrans.Equals(WorldTransform);
-	if (TransformMatches)
-	{
-		// we want to update the transforms one more time than the buffer count because even if the transform buffers didn't change,
-		// their derivatives (like velocity) also need to be updated correctly which happens a frame later.
-		GatheredInstanceParameters.TransformMatchCount = FMath::Min(ParameterBufferCount + 1, GatheredInstanceParameters.TransformMatchCount + 1);
+
+		const FVector CurrentLocation = WorldTransform.GetLocation();
+		if (DeltaSeconds > 0.0f)
+		{
+			if (FMath::IsNearlyZero(Age) && !System->IsInitialOwnerVelocityFromActor())
+			{
+				GatheredInstanceParameters.Velocity = FVector::ZeroVector;
+			}
+			else if (PreviousLocation.IsSet())
+			{
+				GatheredInstanceParameters.Velocity = (CurrentLocation - PreviousLocation.GetValue()) / DeltaSeconds;
+			}
+			else if (AActor* OwnerActor = AttachComponent->GetOwner())
+			{
+				GatheredInstanceParameters.Velocity = OwnerActor->GetVelocity();
+			}
+		}
+		PreviousLocation = CurrentLocation;
 	}
 	else
 	{
-		GatheredInstanceParameters.ComponentTrans = WorldTransform;
-		GatheredInstanceParameters.TransformMatchCount = 0;
+		GatheredInstanceParameters.Velocity = FVector::ZeroVector;
 	}
+	GatheredInstanceParameters.ComponentTrans = WorldTransform;
 
 	GatheredInstanceParameters.EmitterCount = Emitters.Num();
 	GatheredInstanceParameters.DeltaSeconds = DeltaSeconds;
@@ -1987,31 +1999,27 @@ void FNiagaraSystemInstance::TickInstanceParameters_Concurrent()
 	FNiagaraSystemParameters& CurrentSystemParameters = SystemParameters[ParameterIndex];
 	FNiagaraOwnerParameters& CurrentOwnerParameters = OwnerParameters[ParameterIndex];
 
-	if (GatheredInstanceParameters.TransformMatchCount <= ParameterBufferCount)
-	{
-		const FMatrix LocalToWorld = GatheredInstanceParameters.ComponentTrans.ToMatrixWithScale();
-		const FMatrix LocalToWorldNoScale = GatheredInstanceParameters.ComponentTrans.ToMatrixNoScale();
-
-		const FVector Location = GatheredInstanceParameters.ComponentTrans.GetLocation();
-		const FVector LastLocation = FMath::IsNearlyZero(CurrentSystemParameters.EngineSystemAge) ? Location : FVector(FVector4(OwnerParameters[GetParameterIndex(true)].EnginePosition));
-		const FQuat LastRotation = GatheredInstanceParameters.ComponentTrans.GetRotation();
-
-		CurrentOwnerParameters.EngineLocalToWorld = FMatrix44f(LocalToWorld);						// LWC_TODO: Precision loss
-		CurrentOwnerParameters.EngineWorldToLocal = FMatrix44f(LocalToWorld.Inverse());
-		CurrentOwnerParameters.EngineLocalToWorldTransposed = FMatrix44f(LocalToWorld.GetTransposed());
-		CurrentOwnerParameters.EngineWorldToLocalTransposed = CurrentOwnerParameters.EngineWorldToLocal.GetTransposed();
-		CurrentOwnerParameters.EngineLocalToWorldNoScale = FMatrix44f(LocalToWorldNoScale);
-		CurrentOwnerParameters.EngineWorldToLocalNoScale = FMatrix44f(LocalToWorldNoScale.Inverse());
-		CurrentOwnerParameters.EngineRotation = FQuat4f((float)LastRotation.X, (float)LastRotation.Y, (float)LastRotation.Z, (float)LastRotation.W);
-		CurrentOwnerParameters.EnginePosition = (FVector3f)Location; // LWC_TODO: precision loss
-		CurrentOwnerParameters.EngineVelocity = GatheredInstanceParameters.DeltaSeconds > 0.0f ? (FVector3f)((Location - LastLocation) / GatheredInstanceParameters.DeltaSeconds) : FVector3f::ZeroVector;
-		CurrentOwnerParameters.EngineXAxis = CurrentOwnerParameters.EngineRotation.GetAxisX();
-		CurrentOwnerParameters.EngineYAxis = CurrentOwnerParameters.EngineRotation.GetAxisY();
-		CurrentOwnerParameters.EngineZAxis = CurrentOwnerParameters.EngineRotation.GetAxisZ();
-		CurrentOwnerParameters.EngineScale = (FVector3f)GatheredInstanceParameters.ComponentTrans.GetScale3D();
-		CurrentOwnerParameters.EngineLWCTile = LWCTile;
-		CurrentOwnerParameters.EngineLWCTile.W = float(FLargeWorldRenderScalar::GetTileSize());
-	}
+	const FMatrix LocalToWorld = GatheredInstanceParameters.ComponentTrans.ToMatrixWithScale();
+	const FMatrix LocalToWorldNoScale = GatheredInstanceParameters.ComponentTrans.ToMatrixNoScale();
+
+	const FVector Location = GatheredInstanceParameters.ComponentTrans.GetLocation();
+	const FQuat Rotation = GatheredInstanceParameters.ComponentTrans.GetRotation();
+
+	CurrentOwnerParameters.EngineLocalToWorld = FMatrix44f(LocalToWorld);						// LWC_TODO: Precision loss
+	CurrentOwnerParameters.EngineWorldToLocal = FMatrix44f(LocalToWorld.Inverse());
+	CurrentOwnerParameters.EngineLocalToWorldTransposed = FMatrix44f(LocalToWorld.GetTransposed());
+	CurrentOwnerParameters.EngineWorldToLocalTransposed = CurrentOwnerParameters.EngineWorldToLocal.GetTransposed();
+	CurrentOwnerParameters.EngineLocalToWorldNoScale = FMatrix44f(LocalToWorldNoScale);
+	CurrentOwnerParameters.EngineWorldToLocalNoScale = FMatrix44f(LocalToWorldNoScale.Inverse());
+	CurrentOwnerParameters.EngineRotation = FQuat4f(Rotation);									// LWC_TODO: precision loss
+	CurrentOwnerParameters.EnginePosition = FVector3f(Location);								// LWC_TODO: precision loss
+	CurrentOwnerParameters.EngineVelocity = FVector3f(GatheredInstanceParameters.Velocity);		// LWC_TODO: precision loss
+	CurrentOwnerParameters.EngineXAxis = CurrentOwnerParameters.EngineRotation.GetAxisX();
+	CurrentOwnerParameters.EngineYAxis = CurrentOwnerParameters.EngineRotation.GetAxisY();
+	CurrentOwnerParameters.EngineZAxis = CurrentOwnerParameters.EngineRotation.GetAxisZ();
+	CurrentOwnerParameters.EngineScale = (FVector3f)GatheredInstanceParameters.ComponentTrans.GetScale3D();
+	CurrentOwnerParameters.EngineLWCTile = LWCTile;
+	CurrentOwnerParameters.EngineLWCTile.W = float(FLargeWorldRenderScalar::GetTileSize());
 
 	CurrentSystemParameters.EngineEmitterCount = GatheredInstanceParameters.EmitterCount;
 	CurrentSystemParameters.EngineAliveEmitterCount = GatheredInstanceParameters.NumAlive;
diff --git a/Engine/Plugins/FX/Niagara/Source/Niagara/Public/NiagaraSystemInstance.h b/Engine/Plugins/FX/Niagara/Source/Niagara/Public/NiagaraSystemInstance.h
--- a/Engine/Plugins/FX/Niagara/Source/Niagara/Public/NiagaraSystemInstance.h
+++ b/Engine/Plugins/FX/Niagara/Source/Niagara/Public/NiagaraSystemInstance.h
@@ -477,6 +477,7 @@ private:
 	TWeakObjectPtr<USceneComponent> AttachComponent;
 
 	FTransform WorldTransform;
+	TOptional<FVector> PreviousLocation;
 
 	ENiagaraTickBehavior TickBehavior;
 
@@ -638,6 +639,7 @@ public:
 	struct FInstanceParameters
 	{
 		FTransform ComponentTrans = FTransform::Identity;
+		FVector Velocity = FVector::ZeroVector;
 
 		float DeltaSeconds = 0.0f;
 		float TimeSeconds = 0.0f;
@@ -648,13 +650,13 @@ public:
 
 		int32 EmitterCount = 0;
 		int32 NumAlive = 0;
-		int32 TransformMatchCount = 0;
 
 		ENiagaraExecutionState RequestedExecutionState = ENiagaraExecutionState::Active;
 
 		void Init(int32 NumEmitters)
 		{
 			ComponentTrans = FTransform::Identity;
+			Velocity = FVector::ZeroVector;
 			DeltaSeconds = 0.0f;
 			TimeSeconds = 0.0f;
 			RealTimeSeconds = 0.0f;
@@ -663,7 +665,6 @@ public:
 
 			EmitterCount = 0;
 			NumAlive = 0;
-			TransformMatchCount = 0;
 
 			RequestedExecutionState = ENiagaraExecutionState::Active;
 		}
diff --git a/Engine/Plugins/FX/Niagara/Source/NiagaraCore/Public/NiagaraCustomVersion.h b/Engine/Plugins/FX/Niagara/Source/NiagaraCore/Public/NiagaraCustomVersion.h
--- a/Engine/Plugins/FX/Niagara/Source/NiagaraCore/Public/NiagaraCustomVersion.h
+++ b/Engine/Plugins/FX/Niagara/Source/NiagaraCore/Public/NiagaraCustomVersion.h
@@ -193,6 +193,8 @@ struct FNiagaraCustomVersion
 
 		RibbonPlaneUseGeometryNormals,	// Ribbon renderer will use geometry normals by default rather than screen / facing aligned normals
 
+		InitialOwnerVelocityFromActor, // Actors velocity is used for the initial velocity before the component has any tracking, old assets use the old zero velocity
+
 		// DO NOT ADD A NEW VERSION UNLESS YOU HAVE TALKED TO THE NIAGARA LEAD. Mismanagement of these versions can lead to data loss if it is adjusted in multiple streams simultaneously.
 		// -----<new versions can be added above this line>  -------------------------------------------------
 		VersionPlusOne,
