Fix for some attribute bindings not getting renamed properly with emitter merging

diff --git a/Engine/Plugins/FX/Niagara/Source/Niagara/Private/NiagaraParameterBinding.cpp b/Engine/Plugins/FX/Niagara/Source/Niagara/Private/NiagaraParameterBinding.cpp
--- a/Engine/Plugins/FX/Niagara/Source/Niagara/Private/NiagaraParameterBinding.cpp
+++ b/Engine/Plugins/FX/Niagara/Source/Niagara/Private/NiagaraParameterBinding.cpp
@@ -175,11 +175,13 @@ void FNiagaraParameterBinding::ForEachRenameEmitter(UObject* InObject, FStringVi
 {
 	for (TFieldIterator<FStructProperty> PropIt(InObject->GetClass()); PropIt; ++PropIt)
 	{
-		FStructProperty* StructProp = *PropIt;
-		if (StructProp->Struct == FNiagaraParameterBinding::StaticStruct())
+		if (FStructProperty* StructProp = *PropIt)
 		{
-			FNiagaraParameterBinding* ParameterBinding = StructProp->ContainerPtrToValuePtr<FNiagaraParameterBinding>(InObject);
-			ParameterBinding->OnRenameEmitter(EmitterName);
+			if (StructProp->Struct && StructProp->Struct->IsChildOf(FNiagaraParameterBinding::StaticStruct()))
+			{
+				FNiagaraParameterBinding* ParameterBinding = StructProp->ContainerPtrToValuePtr<FNiagaraParameterBinding>(InObject);
+				ParameterBinding->OnRenameEmitter(EmitterName);
+			}
 		}
 	}
 }
@@ -188,11 +190,13 @@ void FNiagaraParameterBinding::ForEachRenameVariable(UObject* InObject, const FN
 {
 	for (TFieldIterator<FStructProperty> PropIt(InObject->GetClass()); PropIt; ++PropIt)
 	{
-		FStructProperty* StructProp = *PropIt;
-		if (StructProp->Struct == FNiagaraParameterBinding::StaticStruct())
+		if (FStructProperty* StructProp = *PropIt)
 		{
-			FNiagaraParameterBinding* ParameterBinding = StructProp->ContainerPtrToValuePtr<FNiagaraParameterBinding>(InObject);
-			ParameterBinding->OnRenameVariable(OldVariable, NewVariable, EmitterName);
+			if (StructProp->Struct && StructProp->Struct->IsChildOf(FNiagaraParameterBinding::StaticStruct()))
+			{
+				FNiagaraParameterBinding* ParameterBinding = StructProp->ContainerPtrToValuePtr<FNiagaraParameterBinding>(InObject);
+				ParameterBinding->OnRenameVariable(OldVariable, NewVariable, EmitterName);
+			}
 		}
 	}
 }
@@ -201,11 +205,13 @@ void FNiagaraParameterBinding::ForEachRemoveVariable(UObject* InObject, const FN
 {
 	for (TFieldIterator<FStructProperty> PropIt(InObject->GetClass()); PropIt; ++PropIt)
 	{
-		FStructProperty* StructProp = *PropIt;
-		if (StructProp->Struct == FNiagaraParameterBinding::StaticStruct())
+		if (FStructProperty* StructProp = *PropIt)
 		{
-			FNiagaraParameterBinding* ParameterBinding = StructProp->ContainerPtrToValuePtr<FNiagaraParameterBinding>(InObject);
-			ParameterBinding->OnRemoveVariable(OldVariable, EmitterName);
+			if (StructProp->Struct && StructProp->Struct->IsChildOf(FNiagaraParameterBinding::StaticStruct()))
+			{
+				FNiagaraParameterBinding* ParameterBinding = StructProp->ContainerPtrToValuePtr<FNiagaraParameterBinding>(InObject);
+				ParameterBinding->OnRemoveVariable(OldVariable, EmitterName);
+			}
 		}
 	}
 }
diff --git a/Engine/Plugins/FX/Niagara/Source/Niagara/Private/NiagaraSimulationStageBase.cpp b/Engine/Plugins/FX/Niagara/Source/Niagara/Private/NiagaraSimulationStageBase.cpp
--- a/Engine/Plugins/FX/Niagara/Source/Niagara/Private/NiagaraSimulationStageBase.cpp
+++ b/Engine/Plugins/FX/Niagara/Source/Niagara/Private/NiagaraSimulationStageBase.cpp
@@ -2,6 +2,7 @@
 
 #include "NiagaraSimulationStageBase.h"
 #include "NiagaraCompileHashVisitor.h"
+#include "NiagaraCustomVersion.h"
 #include "NiagaraEmitter.h"
 #include "NiagaraSimulationStageCompileData.h"
 #include "NiagaraSystem.h"
@@ -223,6 +224,7 @@ void UNiagaraSimulationStageGeneric::PostLoad()
 
 #if WITH_EDITORONLY_DATA
 	const int32 UE5Version = GetLinkerCustomVersion(FUE5MainStreamObjectVersion::GUID);
+	const int32 NiagaraVersion = GetLinkerCustomVersion(FNiagaraCustomVersion::GUID);
 
 	// Ensure data wasn't somehow saved incorrectly
 	using namespace NiagaraSimulationStageLocal;
@@ -276,6 +278,23 @@ void UNiagaraSimulationStageGeneric::PostLoad()
 		OverrideGpuDispatchNumThreadsY.SetDefaultValueEditorOnly(OverrideGpuDispatchNumThreads_DEPRECATED.Y);
 		OverrideGpuDispatchNumThreadsZ.SetDefaultValueEditorOnly(OverrideGpuDispatchNumThreads_DEPRECATED.Z);
 	}
+
+	if (NiagaraVersion < FNiagaraCustomVersion::ParameterBindingWithValueRenameFixup)
+	{
+		if (UNiagaraEmitter* Emitter = GetTypedOuter<UNiagaraEmitter>())
+		{
+			const FString& EmitterName = Emitter->GetUniqueEmitterName();
+
+			NumIterations.OnRenameEmitter(EmitterName);
+			NumIterations.OnRenameEmitter(EmitterName);
+			OverrideGpuDispatchNumThreadsX.OnRenameEmitter(EmitterName);
+			OverrideGpuDispatchNumThreadsY.OnRenameEmitter(EmitterName);
+			OverrideGpuDispatchNumThreadsZ.OnRenameEmitter(EmitterName);
+			ElementCountX.OnRenameEmitter(EmitterName);
+			ElementCountY.OnRenameEmitter(EmitterName);
+			ElementCountZ.OnRenameEmitter(EmitterName);
+		}
+	}
 #endif
 }
 
diff --git a/Engine/Plugins/FX/Niagara/Source/NiagaraCore/Public/NiagaraCustomVersion.h b/Engine/Plugins/FX/Niagara/Source/NiagaraCore/Public/NiagaraCustomVersion.h
--- a/Engine/Plugins/FX/Niagara/Source/NiagaraCore/Public/NiagaraCustomVersion.h
+++ b/Engine/Plugins/FX/Niagara/Source/NiagaraCore/Public/NiagaraCustomVersion.h
@@ -195,6 +195,8 @@ struct FNiagaraCustomVersion
 
 		InitialOwnerVelocityFromActor, // Actors velocity is used for the initial velocity before the component has any tracking, old assets use the old zero velocity
 
+		ParameterBindingWithValueRenameFixup, // FNiagaraParameterBindingWithValue wouldn't necessarily have the appropriate ResolvedParameter namespace when it comes to emitter merging
+
 		// DO NOT ADD A NEW VERSION UNLESS YOU HAVE TALKED TO THE NIAGARA LEAD. Mismanagement of these versions can lead to data loss if it is adjusted in multiple streams simultaneously.
 		// -----<new versions can be added above this line>  -------------------------------------------------
 		VersionPlusOne,
diff --git a/Engine/Source/Runtime/Core/Private/UObject/FortniteMainBranchObjectVersion.cpp b/Engine/Source/Runtime/Core/Private/UObject/FortniteMainBranchObjectVersion.cpp
--- a/Engine/Source/Runtime/Core/Private/UObject/FortniteMainBranchObjectVersion.cpp
+++ b/Engine/Source/Runtime/Core/Private/UObject/FortniteMainBranchObjectVersion.cpp
@@ -11,7 +11,7 @@ TMap<FGuid, FGuid> FFortniteMainBranchObjectVersion::GetSystemGuids()
 	SystemGuids.Add(DevGuids.MATERIALSHADERMAP_DERIVEDDATA_VER, FGuid("CE7776E63E9E43FBB8C103EE1AB40B5D"));
 	SystemGuids.Add(DevGuids.NANITE_DERIVEDDATA_VER, FGuid("D28C003F825F4C178F698938FB0E47F7"));
 	SystemGuids.Add(DevGuids.NIAGARASHADERMAP_DERIVEDDATA_VER, FGuid("A6FFCB143BBC4C83A72001468D23D5A4"));
-	SystemGuids.Add(DevGuids.Niagara_LatestScriptCompileVersion, FGuid("DAA0CD97CD254D3996A16E7C8AE85E12"));
+	SystemGuids.Add(DevGuids.Niagara_LatestScriptCompileVersion, FGuid("25FCFB6BB6554C95890E801EC0FED47E"));
 	SystemGuids.Add(DevGuids.SkeletalMeshDerivedDataVersion, FGuid("518D7C51D31D9BE4A9F4B81542AC8119"));
 	SystemGuids.Add(DevGuids.STATICMESH_DERIVEDDATA_VER, FGuid("543843B103794E0E9BA4BE80FB602F79"));
 	return SystemGuids;
