- Handle nullptr for asset during proxy creation

diff --git a/Engine/Plugins/FX/Niagara/Source/Niagara/Private/NiagaraComponent.cpp b/Engine/Plugins/FX/Niagara/Source/Niagara/Private/NiagaraComponent.cpp
--- a/Engine/Plugins/FX/Niagara/Source/Niagara/Private/NiagaraComponent.cpp
+++ b/Engine/Plugins/FX/Niagara/Source/Niagara/Private/NiagaraComponent.cpp
@@ -190,7 +190,9 @@ FNiagaraSceneProxy::FNiagaraSceneProxy(UNiagaraComponent* InComponent)
 	: FPrimitiveSceneProxy(InComponent, InComponent->GetAsset() ? InComponent->GetAsset()->GetFName() : FName())
 	, OcclusionQueryMode(InComponent->GetOcclusionQueryMode())
 {
-	if (FNiagaraSystemInstanceControllerConstPtr SystemInstanceController = InComponent->GetSystemInstanceController())
+	FNiagaraSystemInstanceControllerConstPtr SystemInstanceController = InComponent->GetSystemInstanceController();
+	UNiagaraSystem* NiagaraSystem = InComponent->GetAsset();
+	if (SystemInstanceController && NiagaraSystem)
 	{
 		LLM_SCOPE(ELLMTag::Niagara);
 
@@ -200,9 +202,9 @@ FNiagaraSceneProxy::FNiagaraSceneProxy(UNiagaraComponent* InComponent)
 		bAlwaysHasVelocity = RenderData->HasAnyMotionBlurEnabled();
 		bIsHeterogeneousVolume = RenderData->HasAnyHeterogeneousVolumesEnabled();
 
-		SystemStatID = InComponent->GetAsset()->GetStatID(false, false);
+		SystemStatID = NiagaraSystem->GetStatID(false, false);
 #if NIAGARAPROXY_EVENTS_ENABLED
-		SystemStatString = InComponent->GetAsset()->GetFName().ToString();
+		SystemStatString = NiagaraSystem->GetFName().ToString();
 #endif
 	}
 }
