Support Spawning into any emitter from one shared DI.

- New special case path through the translator to support Engine.Emitter.ID. Into which we hard code the emitter's handle index to pass to the DI's later.
- Also included Engine.ExecIndex here.
- Adding EmitterID to SpawnConditional and including upgrade path.
- Removed unused Spawn function from read DI + supporting code.
- Slight tidy and refactor of Read DI.
- Expose EmitterIndex to translator to be made into the EmiterID

diff --git a/Engine/Plugins/FX/Niagara/Source/Niagara/Classes/NiagaraConstants.h b/Engine/Plugins/FX/Niagara/Source/Niagara/Classes/NiagaraConstants.h
--- a/Engine/Plugins/FX/Niagara/Source/Niagara/Classes/NiagaraConstants.h
+++ b/Engine/Plugins/FX/Niagara/Source/Niagara/Classes/NiagaraConstants.h
@@ -46,6 +46,7 @@
 #define SYS_PARAM_ENGINE_ROTATION                        INiagaraModule::GetVar_Engine_Owner_Rotation()
 #define SYS_PARAM_ENGINE_SCALE                           INiagaraModule::GetVar_Engine_Owner_Scale()
 #define SYS_PARAM_ENGINE_LWC_TILE                        INiagaraModule::GetVar_Engine_Owner_LWC_Tile()
+#define SYS_PARAM_ENGINE_EXEC_INDEX                      INiagaraModule::GetVar_Engine_ExecIndex()
 
 #define SYS_PARAM_ENGINE_LOCAL_TO_WORLD                  INiagaraModule::GetVar_Engine_Owner_SystemLocalToWorld()
 #define SYS_PARAM_ENGINE_WORLD_TO_LOCAL                  INiagaraModule::GetVar_Engine_Owner_SystemWorldToLocal()
@@ -68,6 +69,7 @@
 #define SYS_PARAM_ENGINE_EMITTER_TOTAL_SPAWNED_PARTICLES INiagaraModule::GetVar_Engine_Emitter_TotalSpawnedParticles()
 #define SYS_PARAM_ENGINE_EMITTER_SPAWN_COUNT_SCALE       INiagaraModule::GetVar_Engine_Emitter_SpawnCountScale()
 #define SYS_PARAM_ENGINE_EMITTER_INSTANCE_SEED           INiagaraModule::GetVar_Engine_Emitter_InstanceSeed()
+#define SYS_PARAM_ENGINE_EMITTER_ID		                 INiagaraModule::GetVar_Engine_Emitter_ID()
 #define SYS_PARAM_ENGINE_SYSTEM_NUM_EMITTERS_ALIVE       INiagaraModule::GetVar_Engine_System_NumEmittersAlive()
 #define SYS_PARAM_ENGINE_SYSTEM_SIGNIFICANCE_INDEX		 INiagaraModule::GetVar_Engine_System_SignificanceIndex()
 #define SYS_PARAM_ENGINE_SYSTEM_RANDOM_SEED				INiagaraModule::GetVar_Engine_System_RandomSeed()
diff --git a/Engine/Plugins/FX/Niagara/Source/Niagara/Classes/NiagaraDataInterface.h b/Engine/Plugins/FX/Niagara/Source/Niagara/Classes/NiagaraDataInterface.h
--- a/Engine/Plugins/FX/Niagara/Source/Niagara/Classes/NiagaraDataInterface.h
+++ b/Engine/Plugins/FX/Niagara/Source/Niagara/Classes/NiagaraDataInterface.h
@@ -1343,6 +1343,21 @@ struct FNDIInputParam<FNiagaraSpawnInfo>
 	FORCEINLINE void Reset() { Count.Reset(); InterpStartDt.Reset();  IntervalDt.Reset(); SpawnGroup.Reset();}
 };
 
+template<>
+struct FNDIInputParam<FNiagaraEmitterID>
+{
+	VectorVM::FExternalFuncInputHandler<int32> Index;
+	FORCEINLINE FNDIInputParam(FVectorVMExternalFunctionContext& Context) : Index(Context) {}
+	FORCEINLINE FNDIInputParam() { }
+	FORCEINLINE void Init(FVectorVMExternalFunctionContext& Context) { new(this)FNDIInputParam<FNiagaraEmitterID>(Context); }
+	FORCEINLINE FNiagaraEmitterID GetAndAdvance() { return FNiagaraEmitterID(Index.GetAndAdvance()); }
+	FORCEINLINE FNiagaraEmitterID Get() { return FNiagaraEmitterID(Index.Get()); }
+	FORCEINLINE void Advance(int32 Count = 1) { return Index.Advance(Count); }
+	FORCEINLINE bool IsConstant() const { return Index.IsConstant(); }
+	FORCEINLINE void Reset() { Index.Reset(); }
+};
+
+
 //Helper to deal with types with potentially several output registers.
 template<typename T>
 struct FNDIOutputParam
diff --git a/Engine/Plugins/FX/Niagara/Source/Niagara/Internal/DataInterface/NiagaraDataInterfaceDataChannelRead.h b/Engine/Plugins/FX/Niagara/Source/Niagara/Internal/DataInterface/NiagaraDataInterfaceDataChannelRead.h
--- a/Engine/Plugins/FX/Niagara/Source/Niagara/Internal/DataInterface/NiagaraDataInterfaceDataChannelRead.h
+++ b/Engine/Plugins/FX/Niagara/Source/Niagara/Internal/DataInterface/NiagaraDataInterfaceDataChannelRead.h
@@ -72,10 +72,6 @@ public:
 	UPROPERTY(EditAnywhere, Category="Data Channel")
 	TObjectPtr<UNiagaraDataChannelAsset> Channel;
 	
-	/** A bounds emitter instance when using functions like Spawn. Defaults to Self. */
-	UPROPERTY(EditAnywhere, Category = "Spawning", AdvancedDisplay)
-	FNiagaraDataInterfaceEmitterBinding EmitterBinding;
-
 	/** True if this reader will read the current frame's data. If false, we read the previous frame.
 	* Reading the current frame introduces a tick order dependency but allows for zero latency reads. Any data channel elements that are generated after this reader is used are missed.
 	* Reading the previous frame's data introduces a frame of latency but ensures we never miss any data as we have access to the whole frame.
@@ -97,10 +93,6 @@ public:
 	UPROPERTY(EditAnywhere, Category = "Spawning", AdvancedDisplay)
 	bool bOverrideSpawnGroupToDataChannelIndex = true;
 
-	/** The spawn info variable we'll read from the data channel to control spawning using the SpawnFromSpawnInfo DI function.*/
-	UPROPERTY(EditAnywhere, Category = "Spawning", AdvancedDisplay)
-	FName SpawnInfoName;
-	
 	//UObject Interface
 	NIAGARA_API virtual void PostInitProperties() override;
 	NIAGARA_API virtual void BeginDestroy() override;
@@ -116,6 +108,8 @@ public:
 	NIAGARA_API virtual bool GetFunctionHLSL(FNiagaraDataInterfaceHlslGenerationContext& HlslGenContext, FString& OutHLSL) override;
 	NIAGARA_API virtual void GetParameterDefinitionHLSL(FNiagaraDataInterfaceHlslGenerationContext& HlslGenContext, FString& OutHLSL) override;
 
+	virtual bool UpgradeFunctionCall(FNiagaraFunctionSignature& FunctionSignature)override;
+
 	NIAGARA_API virtual void PostCompile()override;
 #endif	
 
@@ -154,7 +148,6 @@ public:
 	NIAGARA_API void Consume(FVectorVMExternalFunctionContext& Context, int32 FuncIdx);
 
 	//Emitter only functions.
-	NIAGARA_API void SpawnFromSpawnInfo(FVectorVMExternalFunctionContext& Context);
 	NIAGARA_API void SpawnConditional(FVectorVMExternalFunctionContext& Context, int32 FuncIndex);
 
 	FNDIDataChannelCompiledData& GetCompiledData() { return CompiledData; }
@@ -207,15 +200,10 @@ struct FNDIDataChannelReadInstanceData
 	Conditional spawns to inject into the bound emitter based on our DI function calls.
 	Each data channel element we read in a particular execution will have an entry in this array and can spawn(or not) independently.
 	*/
-	TArray<FNiagaraSpawnInfo> ConditionalSpawns;
+	TMap<FNiagaraEmitterInstance*,TArray<FNiagaraSpawnInfo>> PerEmitterConditionalSpawns;
 
-	/*
-	We read a spawn info for future proofing here but currently we only use the spawn count.
-	NOTE: This is 100% dependent on SpawnInfo class not changing. If it does we'll need to update this.
-	*/
-	FNiagaraDataSetAccessor<FNiagaraSpawnInfo> SpawnInfoAccessor;
 	uint32 CachedLayoutHash = INDEX_NONE;
-	FNiagaraEmitterInstance* EmitterInstance = nullptr;
+	FNiagaraSystemInstance* Owner = nullptr;
 
 	virtual ~FNDIDataChannelReadInstanceData();
 	FNiagaraDataBuffer* GetReadBufferCPU(bool bPrevFrame);
diff --git a/Engine/Plugins/FX/Niagara/Source/Niagara/Private/DataInterface/NiagaraDataInterfaceDataChannelRead.cpp b/Engine/Plugins/FX/Niagara/Source/Niagara/Private/DataInterface/NiagaraDataInterfaceDataChannelRead.cpp
--- a/Engine/Plugins/FX/Niagara/Source/Niagara/Private/DataInterface/NiagaraDataInterfaceDataChannelRead.cpp
+++ b/Engine/Plugins/FX/Niagara/Source/Niagara/Private/DataInterface/NiagaraDataInterfaceDataChannelRead.cpp
@@ -52,8 +52,6 @@ static FAutoConsoleVariableRef CVarNDCReadForcePrevFrame(
 	ECVF_Default
 );
 
-
-
 namespace NDIDataChannelReadLocal
 {
 	static const TCHAR* CommonShaderFile = TEXT("/Plugin/FX/Niagara/Private/DataChannel/NiagaraDataInterfaceDataChannelCommon.ush");
@@ -61,18 +59,111 @@ namespace NDIDataChannelReadLocal
 	static const TCHAR* TemplateShaderFile_Read = TEXT("/Plugin/FX/Niagara/Private/DataChannel/NiagaraDataInterfaceDataChannelTemplate_Read.ush");
 	static const TCHAR* TemplateShaderFile_Consume = TEXT("/Plugin/FX/Niagara/Private/DataChannel/NiagaraDataInterfaceDataChannelTemplate_Consume.ush");
 
-	static const FName NumName(TEXT("Num"));
-	static const FName ReadName(TEXT("Read"));
-	static const FName ConsumeName(TEXT("Consume"));
-	static const FName SpawnFromSpawnInfoName(TEXT("SpawnFromSpawnInfo"));
-	static const FName SpawnConditionalName(TEXT("SpawnConditional"));
+	//////////////////////////////////////////////////////////////////////////
+	//Function definitions
+	const FNiagaraFunctionSignature& GetFunctionSig_Num()
+	{
+		static FNiagaraFunctionSignature Sig;
+		if(!Sig.IsValid())
+		{
+			Sig.Name = TEXT("Num");
+#if WITH_EDITORONLY_DATA
+			Sig.Description = LOCTEXT("NumFunctionDescription", "Returns the current number of DataChannel accessible by this interface.");
+#endif
+			Sig.bMemberFunction = true;
+			Sig.bExperimental = true;
+			Sig.AddInput(FNiagaraVariable(FNiagaraTypeDefinition(UNiagaraDataInterfaceDataChannelRead::StaticClass()), TEXT("DataChannel interface")));
+			Sig.AddOutput(FNiagaraVariable(FNiagaraTypeDefinition::GetIntDef(), TEXT("Num")));
+		}		
+		return Sig;
+	}
+
+	const FNiagaraFunctionSignature& GetFunctionSig_Read()
+	{
+		static FNiagaraFunctionSignature Sig;
+		if(!Sig.IsValid())
+		{
+			Sig.Name = TEXT("Read");
+#if WITH_EDITORONLY_DATA
+			Sig.Description = LOCTEXT("ReadFunctionDescription", "Reads DataChannel data at a specific index. Any values we read that are not in the DataChannel data are set to their default values. Returns success if there was a valid DataChannel to read from.");
+#endif
+			Sig.bMemberFunction = true;
+			Sig.bExperimental = true;
+			Sig.AddInput(FNiagaraVariable(FNiagaraTypeDefinition(UNiagaraDataInterfaceDataChannelRead::StaticClass()), TEXT("DataChannel interface")));
+			Sig.AddInput(FNiagaraVariable(FNiagaraTypeDefinition::GetIntDef(), TEXT("Index")), LOCTEXT("ConsumeIndexInputDesc", "The index to read."));
+			Sig.AddOutput(FNiagaraVariable(FNiagaraTypeDefinition::GetBoolDef(), TEXT("Success")), LOCTEXT("ConsumeSuccessOutputDesc", "True if all reads succeeded."));
+			Sig.RequiredOutputs = Sig.Outputs.Num();//The user defines what we read in the graph.			
+		}
+		return Sig;
+	}
+
+	const FNiagaraFunctionSignature& GetFunctionSig_Consume()
+	{
+		static FNiagaraFunctionSignature Sig;
+		if(!Sig.IsValid())
+		{
+			Sig.Name = TEXT("Consume");
+#if WITH_EDITORONLY_DATA
+			Sig.Description = LOCTEXT("ConsumeFunctionDescription", "Consumes an DataChannel from the end of the DataChannel array and reads the specified values. Any values we read that are not in the DataChannel data are set to their default values. Returns success if an DataChannel was available to pop.");
+#endif
+			Sig.bMemberFunction = true;
+			Sig.bExperimental = true;
+			Sig.AddInput(FNiagaraVariable(FNiagaraTypeDefinition(UNiagaraDataInterfaceDataChannelRead::StaticClass()), TEXT("DataChannel interface")));
+			Sig.AddInput(FNiagaraVariable(FNiagaraTypeDefinition::GetBoolDef(), TEXT("Consume")), LOCTEXT("ConsumeInputDesc", "True if this instance (particle/emitter etc) should consume data from the data channel in this call."));
+			Sig.AddOutput(FNiagaraVariable(FNiagaraTypeDefinition::GetBoolDef(), TEXT("Success")), LOCTEXT("ConsumeSuccessOutputDesc", "True if all reads succeeded."));
+			Sig.AddOutput(FNiagaraVariable(FNiagaraTypeDefinition::GetIntDef(), TEXT("Index")), LOCTEXT("ConsumeIndexOutputDesc", "The index we actually read from. If reading failed this can be -1. This allows subsequent reads of the data channel at this index."));
+			Sig.RequiredOutputs = Sig.Outputs.Num();//The user defines what we read in the graph.			
+		}
+		return Sig;
+	}
+
+	enum class FunctionVersion_SpawnConditional : uint32
+	{
+		Initial = 0,
+		EmitterIDParameter = 1,
+	};
 
-	static const FName SpawnInfoName(TEXT("SpawnInfo"));
+	const FNiagaraFunctionSignature& GetFunctionSig_SpawnConditional()
+	{
+		static FNiagaraFunctionSignature Sig;
+		if(!Sig.IsValid())
+		{
+			FNiagaraVariable EnabledVar(FNiagaraTypeDefinition::GetBoolDef(), TEXT("Enable"));
+			EnabledVar.SetValue(FNiagaraBool(true));
+
+			Sig.Name = TEXT("SpawnConditional");
+#if WITH_EDITORONLY_DATA
+			Sig.Description = LOCTEXT("SpawnCustomFunctionDescription", "Will Spawn particles into the bound Emitter between Min and Max counts for every element in the bound Data Channel.\n\
+		Can take optional additional parameters as conditions on spawning that will be compared against the contents of each data channel element.\n\
+		For example you could spawn only for a particular value of an enum.\n\
+		For compound data types that contain multiple component floats or ints, comparissons are done on a per component basis.\n\
+		For example if you add a Vector condition parameter it will be compared against each component of the corresponding Vector in the Data Channel.\n\
+		Result = (Param.X == ChannelValue.X) && (Param.Y == ChannelValue.Y) && (Param.Z == ChannelValue.Z)");
+#endif
+			Sig.bMemberFunction = true;
+			Sig.bExperimental = true;
+			Sig.bRequiresExecPin = true;
+			Sig.ModuleUsageBitmask = ENiagaraScriptUsageMask::Emitter | ENiagaraScriptUsageMask::System;
+			Sig.AddInput(FNiagaraVariable(FNiagaraTypeDefinition(UNiagaraDataInterfaceDataChannelRead::StaticClass()), TEXT("DataChannel interface")));
+			Sig.AddInput(EnabledVar, LOCTEXT("SpawnEnableInputDesc", "Enable or disable this function call. If false, this call with have no effetcs."));
+			Sig.AddInput(FNiagaraVariable(FNiagaraTypeDefinition(FNiagaraEmitterID::StaticStruct()), TEXT("Emitter ID")), LOCTEXT("EmitterIDDesc", "ID of the emitter we'd like to spawn into. This can be obtained from Engine.Emitter.Index."));
+			Sig.AddInput(FNiagaraVariable(StaticEnum<ENDIDataChannelSpawnMode>(), TEXT("Mode")), LOCTEXT("SpawnCondModeInputDesc", "A mode switch that controls how this funciton will behave."));
+			Sig.AddInput(FNiagaraVariable(StaticEnum<ENiagaraConditionalOperator>(), TEXT("Operator")), LOCTEXT("SpawnCondOpInputDesc", "The comparison operator to use when comparing values in the data channel to conditional parameters."));
+			Sig.AddInput(FNiagaraVariable(FNiagaraTypeDefinition::GetIntDef(), TEXT("Min Spawn Count")), LOCTEXT("MinSpawnCountInputDesc", "Minimum number of particles to spawn for each element in the data channel."));
+			Sig.AddInput(FNiagaraVariable(FNiagaraTypeDefinition::GetIntDef(), TEXT("Max Spawn Count")), LOCTEXT("MaxSpawnCountInputDesc", "Maximum number of particles to spawn for each element in the data channel."));
+			Sig.RequiredInputs = Sig.Inputs.Num();
+
+			Sig.FunctionVersion = (uint32)FunctionVersion_SpawnConditional::EmitterIDParameter;
+		}
+		return Sig;
+	}
+	// Function definitions END
+	//////////////////////////////////////////////////////////////////////////
 
 	const TCHAR* GetFunctionTemplate(FName FunctionName)
 	{
-		if (FunctionName == ReadName) return TemplateShaderFile_Read;
-		if (FunctionName == ConsumeName) return TemplateShaderFile_Consume;
+		if (FunctionName == GetFunctionSig_Read().Name) return TemplateShaderFile_Read;
+		if (FunctionName == GetFunctionSig_Consume().Name) return TemplateShaderFile_Consume;
 
 		return nullptr;
 	}
@@ -145,8 +236,6 @@ FNiagaraDataBuffer* FNDIDataChannelReadInstanceData::GetReadBufferCPU(bool bPrev
 
 bool FNDIDataChannelReadInstanceData::Init(UNiagaraDataInterfaceDataChannelRead* Interface, FNiagaraSystemInstance* Instance)
 {
-	EmitterInstance = Interface->EmitterBinding.Resolve(Instance, Interface);
-
 	bool bSuccess = Tick(Interface, Instance, true);
 	bSuccess &= PostTick(Interface, Instance);
 	return bSuccess;
@@ -155,7 +244,10 @@ bool FNDIDataChannelReadInstanceData::Init(UNiagaraDataInterfaceDataChannelRead*
 bool FNDIDataChannelReadInstanceData::Tick(UNiagaraDataInterfaceDataChannelRead* Interface, FNiagaraSystemInstance* Instance, bool bIsInit)
 {
 	ConsumeIndex = 0;
-	ConditionalSpawns.Reset();
+	for(auto& ConditionalSpawnPair : PerEmitterConditionalSpawns)
+	{
+		ConditionalSpawnPair.Value.Reset();
+	}
 
 	const FNDIDataChannelCompiledData& CompiledData = Interface->GetCompiledData();
 
@@ -166,6 +258,8 @@ bool FNDIDataChannelReadInstanceData::Tick(UNiagaraDataInterfaceDataChannelRead*
 		return true;
 	}
 	
+	Owner = Instance;
+
 	//In non test/shipping builds we gather and log and missing parameters that cause us to fail to find correct bindings.
 	TArray<FNiagaraVariableBase> MissingParams;
 
@@ -287,8 +381,6 @@ bool FNDIDataChannelReadInstanceData::Tick(UNiagaraDataInterfaceDataChannelRead*
 			{
 				ChachedDataSetLayoutHash = SourceDataLayoutHash;
 
-				SpawnInfoAccessor = FNiagaraDataSetAccessor<FNiagaraSpawnInfo>(CPUSourceDataCompiledData, Interface->SpawnInfoName);
-
 				//We can likely be more targeted here.
 				//Could probably only update the RT when the GPU data changes and only update the bindings if the function hashes change etc.
 				bUpdateRTData = CompiledData.UsedByGPU();
@@ -360,8 +452,6 @@ UNiagaraDataInterfaceDataChannelRead::UNiagaraDataInterfaceDataChannelRead(FObje
 	: Super(ObjectInitializer)
 {
 	Proxy.Reset(new FNiagaraDataInterfaceProxy_DataChannelRead());
-	EmitterBinding.BindingMode = ENiagaraDataInterfaceEmitterBindingMode::Self;
-	SpawnInfoName = NDIDataChannelReadLocal::SpawnInfoName;
 }
 
 void UNiagaraDataInterfaceDataChannelRead::PostInitProperties()
@@ -470,18 +560,25 @@ void UNiagaraDataInterfaceDataChannelRead::PostStageTick(FNDICpuPostStageContext
 	FNDIDataChannelReadInstanceData* InstanceData = Context.GetPerInstanceData<FNDIDataChannelReadInstanceData>();
 
 	check(InstanceData);
-	for(int32 i=0; i<InstanceData->ConditionalSpawns.Num(); ++i)
+	for (auto& ConditionalSpawnPair : InstanceData->PerEmitterConditionalSpawns)
 	{
-		if(InstanceData->ConditionalSpawns[i].Count > 0)
+		if(FNiagaraEmitterInstance* TargetEmitter = ConditionalSpawnPair.Key)
 		{
-			if(bOverrideSpawnGroupToDataChannelIndex)
+			for (int32 i = 0; i < ConditionalSpawnPair.Value.Num(); ++i)
 			{
-				InstanceData->ConditionalSpawns[i].SpawnGroup = i;
+				FNiagaraSpawnInfo& SpawnInfo = ConditionalSpawnPair.Value[i];
+				if (SpawnInfo.Count > 0)
+				{
+					if (bOverrideSpawnGroupToDataChannelIndex)
+					{
+						SpawnInfo.SpawnGroup = i;
+					}
+					TargetEmitter->GetSpawnInfo().Emplace(SpawnInfo);
+				}
 			}
-			InstanceData->EmitterInstance->GetSpawnInfo().Emplace(InstanceData->ConditionalSpawns[i]);
 		}
+		ConditionalSpawnPair.Value.Reset();
 	}
-	InstanceData->ConditionalSpawns.Reset();
 }
 
 void UNiagaraDataInterfaceDataChannelRead::ProvidePerInstanceDataForRenderThread(void* DataForRenderThread, void* PerInstanceData, const FNiagaraSystemInstanceID& SystemInstance)
@@ -767,9 +864,7 @@ bool UNiagaraDataInterfaceDataChannelRead::Equals(const UNiagaraDataInterface* O
 			Channel == OtherTyped->Channel &&
 			bReadCurrentFrame == OtherTyped->bReadCurrentFrame &&
 			bUpdateSourceDataEveryTick == OtherTyped->bUpdateSourceDataEveryTick &&
-			EmitterBinding == OtherTyped->EmitterBinding &&
-			bOverrideSpawnGroupToDataChannelIndex == OtherTyped->bOverrideSpawnGroupToDataChannelIndex &&
-			SpawnInfoName == OtherTyped->SpawnInfoName)
+			bOverrideSpawnGroupToDataChannelIndex == OtherTyped->bOverrideSpawnGroupToDataChannelIndex)
 		{
 			return true;
 		}
@@ -794,9 +889,7 @@ bool UNiagaraDataInterfaceDataChannelRead::CopyToInternal(UNiagaraDataInterface*
 		DestTyped->CompiledData = CompiledData;
 		DestTyped->bReadCurrentFrame = bReadCurrentFrame;
 		DestTyped->bUpdateSourceDataEveryTick = bUpdateSourceDataEveryTick;
-		DestTyped->EmitterBinding = EmitterBinding;
-		DestTyped->bOverrideSpawnGroupToDataChannelIndex = bOverrideSpawnGroupToDataChannelIndex;
-		DestTyped->SpawnInfoName = SpawnInfoName;
+		DestTyped->bOverrideSpawnGroupToDataChannelIndex = bOverrideSpawnGroupToDataChannelIndex;		
 		return true;
 	}
 
@@ -805,114 +898,30 @@ bool UNiagaraDataInterfaceDataChannelRead::CopyToInternal(UNiagaraDataInterface*
 
 void UNiagaraDataInterfaceDataChannelRead::GetFunctions(TArray<FNiagaraFunctionSignature>& OutFunctions)
 {
-	{
-		FNiagaraFunctionSignature Sig;
-		Sig.Name = NDIDataChannelReadLocal::NumName;
-#if WITH_EDITORONLY_DATA
-		Sig.Description = LOCTEXT("NumFunctionDescription", "Returns the current number of DataChannel accessible by this interface.");
-#endif
-		Sig.bMemberFunction = true;
-		Sig.bExperimental = true;
-		Sig.AddInput(FNiagaraVariable(FNiagaraTypeDefinition(GetClass()), TEXT("DataChannel interface")));
-		Sig.AddOutput(FNiagaraVariable(FNiagaraTypeDefinition::GetIntDef(), TEXT("Num")));
-		OutFunctions.Add(Sig);
-	}
-	{
-		FNiagaraFunctionSignature Sig;
-		Sig.Name = NDIDataChannelReadLocal::ReadName;
-#if WITH_EDITORONLY_DATA
-		Sig.Description = LOCTEXT("ReadFunctionDescription", "Reads DataChannel data at a specific index. Any values we read that are not in the DataChannel data are set to their default values. Returns success if there was a valid DataChannel to read from.");
-#endif
-		Sig.bMemberFunction = true;
-		Sig.bExperimental = true;
-		Sig.AddInput(FNiagaraVariable(FNiagaraTypeDefinition(GetClass()), TEXT("DataChannel interface")));
-		Sig.AddInput(FNiagaraVariable(FNiagaraTypeDefinition::GetIntDef(), TEXT("Index")), LOCTEXT("ConsumeIndexInputDesc", "The index to read."));
-		Sig.AddOutput(FNiagaraVariable(FNiagaraTypeDefinition::GetBoolDef(), TEXT("Success")), LOCTEXT("ConsumeSuccessOutputDesc", "True if all reads succeeded."));
-		Sig.RequiredOutputs = Sig.Outputs.Num();//The user defines what we read in the graph.
-		OutFunctions.Add(Sig);
-	}
-	{
-		FNiagaraFunctionSignature Sig;
-		Sig.Name = NDIDataChannelReadLocal::ConsumeName;
-#if WITH_EDITORONLY_DATA
-		Sig.Description = LOCTEXT("ConsumeFunctionDescription", "Consumes an DataChannel from the end of the DataChannel array and reads the specified values. Any values we read that are not in the DataChannel data are set to their default values. Returns success if an DataChannel was available to pop.");
-#endif
-		Sig.bMemberFunction = true;
-		Sig.bExperimental = true;
-		Sig.AddInput(FNiagaraVariable(FNiagaraTypeDefinition(GetClass()), TEXT("DataChannel interface")));
-		Sig.AddInput(FNiagaraVariable(FNiagaraTypeDefinition::GetBoolDef(), TEXT("Consume")), LOCTEXT("ConsumeInputDesc", "True if this instance (particle/emitter etc) should consume data from the data channel in this call."));
-		Sig.AddOutput(FNiagaraVariable(FNiagaraTypeDefinition::GetBoolDef(), TEXT("Success")), LOCTEXT("ConsumeSuccessOutputDesc", "True if all reads succeeded."));
-		Sig.AddOutput(FNiagaraVariable(FNiagaraTypeDefinition::GetIntDef(), TEXT("Index")), LOCTEXT("ConsumeIndexOutputDesc", "The index we actually read from. If reading failed this can be -1. This allows subsequent reads of the data channel at this index."));
-		Sig.RequiredOutputs = Sig.Outputs.Num();//The user defines what we read in the graph.
-		OutFunctions.Add(Sig);
-	}
-
-	FNiagaraVariable EnabledVar(FNiagaraTypeDefinition::GetBoolDef(), TEXT("Enable"));
-	EnabledVar.SetValue(FNiagaraBool(true));
-
-	{
-		FNiagaraFunctionSignature Sig;
-		Sig.Name = NDIDataChannelReadLocal::SpawnFromSpawnInfoName;
-#if WITH_EDITORONLY_DATA
-		Sig.Description = LOCTEXT("SpawnFromSpawnInfolFunctionDescription", "Will Spawn particles into the bound Emitter based on a spawn info variable read from each Data Channel element. ");
-#endif
-		Sig.bMemberFunction = true;
-		Sig.bExperimental = true;
-		Sig.bRequiresExecPin = true;
-		Sig.ModuleUsageBitmask = ENiagaraScriptUsageMask::Emitter | ENiagaraScriptUsageMask::System;
-		Sig.AddInput(FNiagaraVariable(FNiagaraTypeDefinition(GetClass()), TEXT("DataChannel interface")));
-		Sig.AddInput(EnabledVar, LOCTEXT("SpawnEnableInputDesc", "Enable or disable this function call. If false, this call with have no effetcs."));
-		OutFunctions.Add(Sig);
-	}
-
-	{
-		FNiagaraFunctionSignature Sig;
-		Sig.Name = NDIDataChannelReadLocal::SpawnConditionalName;
-#if WITH_EDITORONLY_DATA
-		Sig.Description = LOCTEXT("SpawnCustomFunctionDescription", "Will Spawn particles into the bound Emitter between Min and Max counts for every element in the bound Data Channel.\n\
-		Can take optional additional parameters as conditions on spawning that will be compared against the contents of each data channel element.\n\
-		For example you could spawn only for a particular value of an enum.\n\
-		For compound data types that contain multiple component floats or ints, comparissons are done on a per component basis.\n\
-		For example if you add a Vector condition parameter it will be compared against each component of the corresponding Vector in the Data Channel.\n\
-		Result = (Param.X == ChannelValue.X) && (Param.Y == ChannelValue.Y) && (Param.Z == ChannelValue.Z)");
-#endif
-		Sig.bMemberFunction = true;
-		Sig.bExperimental = true;
-		Sig.bRequiresExecPin = true;
-		Sig.ModuleUsageBitmask = ENiagaraScriptUsageMask::Emitter | ENiagaraScriptUsageMask::System;
-		Sig.AddInput(FNiagaraVariable(FNiagaraTypeDefinition(GetClass()), TEXT("DataChannel interface")));
-		Sig.AddInput(EnabledVar, LOCTEXT("SpawnEnableInputDesc", "Enable or disable this function call. If false, this call with have no effetcs."));
-		Sig.AddInput(FNiagaraVariable(StaticEnum<ENDIDataChannelSpawnMode>(), TEXT("Mode")), LOCTEXT("SpawnCondModeInputDesc", "A mode switch that controls how this funciton will behave."));
-		Sig.AddInput(FNiagaraVariable(StaticEnum<ENiagaraConditionalOperator>(), TEXT("Operator")), LOCTEXT("SpawnCondOpInputDesc", "The comparison operator to use when comparing values in the data channel to conditional parameters."));
-		Sig.AddInput(FNiagaraVariable(FNiagaraTypeDefinition::GetIntDef(), TEXT("Min Spawn Count")), LOCTEXT("MinSpawnCountInputDesc", "Minimum number of particles to spawn for each element in the data channel."));
-		Sig.AddInput(FNiagaraVariable(FNiagaraTypeDefinition::GetIntDef(), TEXT("Max Spawn Count")), LOCTEXT("MaxSpawnCountInputDesc", "Maximum number of particles to spawn for each element in the data channel."));
-		Sig.RequiredInputs = Sig.Inputs.Num();
-		OutFunctions.Add(Sig);
-	}
+	OutFunctions.Add(NDIDataChannelReadLocal::GetFunctionSig_Num());
+	OutFunctions.Add(NDIDataChannelReadLocal::GetFunctionSig_Read());
+	OutFunctions.Add(NDIDataChannelReadLocal::GetFunctionSig_Consume());
+	OutFunctions.Add(NDIDataChannelReadLocal::GetFunctionSig_SpawnConditional());
 }
 
 void UNiagaraDataInterfaceDataChannelRead::GetVMExternalFunction(const FVMExternalFunctionBindingInfo& BindingInfo, void* InstanceData, FVMExternalFunction& OutFunc)
 {
-	if (BindingInfo.Name == NDIDataChannelReadLocal::NumName)
+	if (BindingInfo.Name == NDIDataChannelReadLocal::GetFunctionSig_Num().Name)
 	{
 		OutFunc = FVMExternalFunction::CreateLambda([this](FVectorVMExternalFunctionContext& Context) { this->Num(Context); });
 	}
 	else
 	{
 		int32 FuncIndex = CompiledData.FindFunctionInfoIndex(BindingInfo.Name, BindingInfo.VariadicInputs, BindingInfo.VariadicOutputs);
-		if (BindingInfo.Name == NDIDataChannelReadLocal::SpawnFromSpawnInfoName)
-		{
-			OutFunc = FVMExternalFunction::CreateLambda([this](FVectorVMExternalFunctionContext& Context) { this->SpawnFromSpawnInfo(Context); });
-		}
-		else if (BindingInfo.Name == NDIDataChannelReadLocal::ReadName)
+		if (BindingInfo.Name == NDIDataChannelReadLocal::GetFunctionSig_Read().Name)
 		{
 			OutFunc = FVMExternalFunction::CreateLambda([this, FuncIndex](FVectorVMExternalFunctionContext& Context) { this->Read(Context, FuncIndex); });
 		}
-		else if (BindingInfo.Name == NDIDataChannelReadLocal::ConsumeName)
+		else if (BindingInfo.Name == NDIDataChannelReadLocal::GetFunctionSig_Consume().Name)
 		{
 			OutFunc = FVMExternalFunction::CreateLambda([this, FuncIndex](FVectorVMExternalFunctionContext& Context) { this->Consume(Context, FuncIndex); });
 		}
-		else if (BindingInfo.Name == NDIDataChannelReadLocal::SpawnConditionalName)
+		else if (BindingInfo.Name == NDIDataChannelReadLocal::GetFunctionSig_SpawnConditional().Name)
 		{
 			OutFunc = FVMExternalFunction::CreateLambda([this, FuncIndex](FVectorVMExternalFunctionContext& Context) { this->SpawnConditional(Context, FuncIndex); });
 		}
@@ -1090,54 +1099,6 @@ void UNiagaraDataInterfaceDataChannelRead::Consume(FVectorVMExternalFunctionCont
 	}
 }
 
-void UNiagaraDataInterfaceDataChannelRead::SpawnFromSpawnInfo(FVectorVMExternalFunctionContext& Context)
-{
-	SCOPE_CYCLE_COUNTER(STAT_NDIDataChannelRead_Spawn);
-
-	//This should only be called from emitter scripts and since it has per instance data then we process them individually.
-	check(Context.GetNumInstances() == 1);
-
-	FNDIRandomHelperFromStream RandHelper(Context);
-
-	VectorVM::FUserPtrHandler<FNDIDataChannelReadInstanceData> InstData(Context);
-		
-	FNDIInputParam<FNiagaraBool> InEnabled(Context);
-
-	FNiagaraEmitterInstance* EmittterInst = InstData->EmitterInstance;
-	FNiagaraDataChannelData* DataChannelData = InstData->DataChannelData.Get();
-
-	bool bReadPrevFrame = bReadCurrentFrame == false || GNDCReadForcePrevFrame;
-	FNiagaraDataBuffer* Data = DataChannelData ? DataChannelData->GetCPUData(bReadPrevFrame) : nullptr;
-
-	if ( INiagaraModule::DataChannelsEnabled() && Data && EmittterInst )
-	{
-		TArray<FNiagaraSpawnInfo>& EmitterSpawnInfos = EmittterInst->GetSpawnInfo();
-
-		const bool bSpawn = InEnabled.GetAndAdvance() && InstData->SpawnInfoAccessor.IsValid();
-		if (bSpawn)
-		{
-			FNiagaraDataSetReaderStruct<FNiagaraSpawnInfo> Reader = InstData->SpawnInfoAccessor.GetReader(Data);
-
-			int32 NumDataChannelInstances = Data->GetNumInstances();
-
-			for (int32 DataChannelIdx = 0; DataChannelIdx < NumDataChannelInstances; ++DataChannelIdx)
-			{
-				FNiagaraSpawnInfo SpawnInfo(0, 0.0f, 0.0f, 0);
-				SpawnInfo = Reader.Get(DataChannelIdx);
-
-				if (SpawnInfo.Count > 0)
-				{
-					if (bOverrideSpawnGroupToDataChannelIndex)
-					{
-						SpawnInfo.SpawnGroup = DataChannelIdx;
-					}
-					EmitterSpawnInfos.Emplace(SpawnInfo);
-				}
-			}
-		}
-	}
-}
-
 void UNiagaraDataInterfaceDataChannelRead::SpawnConditional(FVectorVMExternalFunctionContext& Context, int32 FuncIdx)
 {
 	SCOPE_CYCLE_COUNTER(STAT_NDIDataChannelRead_Spawn);
@@ -1149,10 +1110,15 @@ void UNiagaraDataInterfaceDataChannelRead::SpawnConditional(FVectorVMExternalFun
 
 	VectorVM::FUserPtrHandler<FNDIDataChannelReadInstanceData> InstData(Context);
 
+	FNiagaraSystemInstance* SystemInstance = InstData->Owner;
+	check(SystemInstance);
+
 	//Binding info can be null here as we can be spawning without any conditions, i.e. no variadic parameters to the function.
 	const FNDIDataChannel_FunctionToDataSetBinding* BindingInfo = InstData->FuncToDataSetBindingInfo.IsValidIndex(FuncIdx) ? InstData->FuncToDataSetBindingInfo[FuncIdx].Get() : nullptr;
 
 	FNDIInputParam<FNiagaraBool> InEnabled(Context);
+	FNDIInputParam<FNiagaraEmitterID> InEmitterID(Context);	
+
 	FNDIInputParam<int32> InMode(Context);
 	FNDIInputParam<int32> InOp(Context);
 	FNDIInputParam<int32> InSpawnMin(Context);
@@ -1160,23 +1126,26 @@ void UNiagaraDataInterfaceDataChannelRead::SpawnConditional(FVectorVMExternalFun
 
 	FNDIVariadicInputHandler<16> VariadicInputs(Context, BindingInfo);//TODO: Make static / avoid allocation
 
-	FNiagaraEmitterInstance* EmittterInst = InstData->EmitterInstance;
+	const FNiagaraEmitterID EmitterID = InEmitterID.GetAndAdvance();
+	const int32 NumEmitters = SystemInstance->GetEmitters().Num();
+	FNiagaraEmitterInstance* EmitterInst = SystemInstance->GetEmitterByID(EmitterID);
+
 	FNiagaraDataChannelData* DataChannelData = InstData->DataChannelData.Get();
 
-	bool bReadPrevFrame = bReadCurrentFrame == false || GNDCReadForcePrevFrame;
+	const bool bReadPrevFrame = bReadCurrentFrame == false || GNDCReadForcePrevFrame;
 	FNiagaraDataBuffer* Data = DataChannelData ? DataChannelData->GetCPUData(bReadPrevFrame) : nullptr;
 
-	bool bSpawn = INiagaraModule::DataChannelsEnabled() && Data && Data->GetNumInstances() > 0 && EmittterInst && EmittterInst->IsActive() && InEnabled.GetAndAdvance();
+	const bool bSpawn = INiagaraModule::DataChannelsEnabled() && Data && Data->GetNumInstances() > 0 && EmitterInst && EmitterInst->IsActive() && InEnabled.GetAndAdvance();
 	if(bSpawn)
 	{
 		FNDIRandomHelperFromStream RandHelper(Context);
 
-		TArray<FNiagaraSpawnInfo>& EmitterSpawnInfos = EmittterInst->GetSpawnInfo();
 		ENDIDataChannelSpawnMode Mode = (ENDIDataChannelSpawnMode)InMode.GetAndAdvance();
 		ENiagaraConditionalOperator Op = (ENiagaraConditionalOperator)InOp.GetAndAdvance();
+		int32 NumDataChannelInstances = Data->GetNumInstances();
 
 		//Is mode none or invalid?
-		if ((int32)Mode == (int32)ENDIDataChannelSpawnMode::None || (int32)Mode < 0 || (int32)Mode >= (int32)ENDIDataChannelSpawnMode::Max)
+		if (NumDataChannelInstances == 0 || (int32)Mode == (int32)ENDIDataChannelSpawnMode::None || (int32)Mode < 0 || (int32)Mode >= (int32)ENDIDataChannelSpawnMode::Max)
 		{
 			return;
 		}
@@ -1184,10 +1153,9 @@ void UNiagaraDataInterfaceDataChannelRead::SpawnConditional(FVectorVMExternalFun
 		int32 SpawnMin = InSpawnMin.GetAndAdvance();
 		int32 SpawnMax = InSpawnMax.GetAndAdvance();
 
-		int32 NumDataChannelInstances = Data->GetNumInstances();
-
 		//Each data channel element has an additional spawn entry which accumulates across all spawning calls and can be nulled independently by a suppression call.
-		InstData->ConditionalSpawns.SetNumZeroed(NumDataChannelInstances);
+		TArray<FNiagaraSpawnInfo>& EmitterConditionalSpawns = InstData->PerEmitterConditionalSpawns.FindOrAdd(EmitterInst);
+		EmitterConditionalSpawns.SetNumZeroed(NumDataChannelInstances);
 
 		for (int32 DataChannelIdx = 0; DataChannelIdx < NumDataChannelInstances; ++DataChannelIdx)
 		{
@@ -1215,11 +1183,11 @@ void UNiagaraDataInterfaceDataChannelRead::SpawnConditional(FVectorVMExternalFun
 				int32 Count = RandHelper.RandRange(DataChannelIdx, SpawnMin, SpawnMax);
 				if (Mode == ENDIDataChannelSpawnMode::Accumulate)
 				{
-					InstData->ConditionalSpawns[DataChannelIdx].Count += Count;
+					EmitterConditionalSpawns[DataChannelIdx].Count += Count;
 				}
 				else if (Mode == ENDIDataChannelSpawnMode::Override)
 				{
-					InstData->ConditionalSpawns[DataChannelIdx].Count = Count;
+					EmitterConditionalSpawns[DataChannelIdx].Count = Count;
 				}
 			}
 		}
@@ -1528,6 +1496,37 @@ void UNiagaraDataInterfaceDataChannelRead::GetParameterDefinitionHLSL(FNiagaraDa
 }
 #endif
 
+bool UNiagaraDataInterfaceDataChannelRead::UpgradeFunctionCall(FNiagaraFunctionSignature& FunctionSignature)
+{
+	TArray<FNiagaraFunctionSignature> Funcs;
+	GetFunctions(Funcs);
+
+	for(const FNiagaraFunctionSignature& Func : Funcs)
+	{
+		if (Func.Name == FunctionSignature.Name && Func.FunctionVersion > FunctionSignature.FunctionVersion)
+		{
+			//We need to add back any variadic params from the source signature.
+			TArray<FNiagaraVariableBase> VariadicInputs;
+			FunctionSignature.GetVariadicInputs(VariadicInputs);
+			TArray<FNiagaraVariableBase> VariadicOutputs;
+			FunctionSignature.GetVariadicOutputs(VariadicOutputs);
+
+			FunctionSignature = Func;
+			for(FNiagaraVariableBase& Param : VariadicInputs)
+			{
+				FunctionSignature.AddInput(Param);
+			}
+			for (FNiagaraVariableBase& Param : VariadicOutputs)
+			{
+				FunctionSignature.AddOutput(Param);
+			}
+			return true;
+		}
+	}
+
+	return false;
+}
+
 void UNiagaraDataInterfaceDataChannelRead::BuildShaderParameters(FNiagaraShaderParametersBuilder& ShaderParametersBuilder) const
 {
 	ShaderParametersBuilder.AddNestedStruct<NDIDataChannelReadLocal::FShaderParameters>();
diff --git a/Engine/Plugins/FX/Niagara/Source/Niagara/Private/NiagaraConstants.cpp b/Engine/Plugins/FX/Niagara/Source/Niagara/Private/NiagaraConstants.cpp
--- a/Engine/Plugins/FX/Niagara/Source/Niagara/Private/NiagaraConstants.cpp
+++ b/Engine/Plugins/FX/Niagara/Source/Niagara/Private/NiagaraConstants.cpp
@@ -91,6 +91,8 @@ void FNiagaraConstants::Init()
 		SystemParameters.Add(SYS_PARAM_ENGINE_Y_AXIS);
 		SystemParameters.Add(SYS_PARAM_ENGINE_Z_AXIS);
 		SystemParameters.Add(SYS_PARAM_ENGINE_LWC_TILE);
+		SystemParameters.Add(SYS_PARAM_ENGINE_EXEC_INDEX);
+		SystemParameters.Add(SYS_PARAM_ENGINE_EMITTER_ID);
 
 		SystemParameters.Add(SYS_PARAM_ENGINE_ROTATION);
 
@@ -215,6 +217,9 @@ void FNiagaraConstants::Init()
 		UpdatedSystemParameters.Add(FName(TEXT("Interp Spawn Start Dt")), SYS_PARAM_EMITTER_INTERP_SPAWN_START_DT);
 		UpdatedSystemParameters.Add(FName(TEXT("Spawn Group")), SYS_PARAM_EMITTER_SPAWN_GROUP);
 		UpdatedSystemParameters.Add(FName(TEXT("Inv Delta Time")), SYS_PARAM_ENGINE_INV_DELTA_TIME);
+
+		UpdatedSystemParameters.Add(FName(TEXT("Exec Index")), SYS_PARAM_ENGINE_EXEC_INDEX);
+		UpdatedSystemParameters.Add(FName(TEXT("Emitter ID")), SYS_PARAM_ENGINE_EMITTER_ID);
 	}
 
 	if (SystemStrMap.Num() == 0)
@@ -236,6 +241,8 @@ void FNiagaraConstants::Init()
 		SystemStrMap.Add(SYS_PARAM_ENGINE_Z_AXIS, LOCTEXT("ZAxisDesc", "The Z-axis of the owning component."));
 		SystemStrMap.Add(SYS_PARAM_ENGINE_ROTATION, LOCTEXT("EngineRotationDesc", "The owning component's rotation in world space."));
 		SystemStrMap.Add(SYS_PARAM_ENGINE_LWC_TILE, LOCTEXT("EngineLWCTileDesc", "Due to large world coordinates, the simulation position can be shifted from the actual world position to allow for more accuracy. This is the tile the system is shifted by, so (SimulationPosition + Tile * TileSize) gives the original world position. The x,y,z components of this vector are the tile and the w component is the tile size."));
+		SystemStrMap.Add(SYS_PARAM_ENGINE_EXEC_INDEX, LOCTEXT("EngineExecIndexDesc", "Index of this particle in the current execution. For example, in a spawn script this gives the index of the particle being spawned which can be used to interpolate it's position."));
+		SystemStrMap.Add(SYS_PARAM_ENGINE_EMITTER_ID, LOCTEXT("EngineEmitterIDDesc", "ID of the currently executing Emitter. This can be used to access or otherwise target a specific emitter in certain Data Interface functions. This is only valid for Emitter and Particle scripts."));
 
 		SystemStrMap.Add(SYS_PARAM_ENGINE_LOCAL_TO_WORLD, LOCTEXT("LocalToWorldDesc", "Owning component's local space to world space transform matrix."));
 		SystemStrMap.Add(SYS_PARAM_ENGINE_WORLD_TO_LOCAL, LOCTEXT("WorldToLocalDesc", "Owning component's world space to local space transform matrix."));
diff --git a/Engine/Plugins/FX/Niagara/Source/Niagara/Private/NiagaraModule.cpp b/Engine/Plugins/FX/Niagara/Source/Niagara/Private/NiagaraModule.cpp
--- a/Engine/Plugins/FX/Niagara/Source/Niagara/Private/NiagaraModule.cpp
+++ b/Engine/Plugins/FX/Niagara/Source/Niagara/Private/NiagaraModule.cpp
@@ -172,6 +172,7 @@ FNiagaraVariable INiagaraModule::Engine_Owner_ZAxis;
 FNiagaraVariable INiagaraModule::Engine_Owner_Scale;
 FNiagaraVariable INiagaraModule::Engine_Owner_Rotation;
 FNiagaraVariable INiagaraModule::Engine_Owner_LWC_Tile;
+FNiagaraVariable INiagaraModule::Engine_ExecIndex;
 
 FNiagaraVariable INiagaraModule::Engine_Owner_SystemLocalToWorld;
 FNiagaraVariable INiagaraModule::Engine_Owner_SystemWorldToLocal;
@@ -191,6 +192,7 @@ FNiagaraVariable INiagaraModule::Engine_Emitter_NumParticles;
 FNiagaraVariable INiagaraModule::Engine_Emitter_SimulationPosition;
 FNiagaraVariable INiagaraModule::Engine_Emitter_TotalSpawnedParticles;
 FNiagaraVariable INiagaraModule::Engine_Emitter_SpawnCountScale;
+FNiagaraVariable INiagaraModule::Engine_Emitter_ID;
 FNiagaraVariable INiagaraModule::Engine_System_TickCount;
 FNiagaraVariable INiagaraModule::Engine_System_NumEmittersAlive;
 FNiagaraVariable INiagaraModule::Engine_System_SignificanceIndex;
@@ -324,6 +326,7 @@ void INiagaraModule::StartupModule()
 	Engine_Owner_Scale = FNiagaraVariable(FNiagaraTypeDefinition::GetVec3Def(), TEXT("Engine.Owner.Scale"));
 	Engine_Owner_Rotation = FNiagaraVariable(FNiagaraTypeDefinition::GetQuatDef(), TEXT("Engine.Owner.Rotation"));
 	Engine_Owner_LWC_Tile = FNiagaraVariable(FNiagaraTypeDefinition::GetVec4Def(), TEXT("Engine.Owner.LWCTile"));
+	Engine_ExecIndex = FNiagaraVariable(FNiagaraTypeDefinition::GetIntDef(), TEXT("Engine.ExecIndex"));
 
 	Engine_Owner_SystemLocalToWorld = FNiagaraVariable(FNiagaraTypeDefinition::GetMatrix4Def(), TEXT("Engine.Owner.SystemLocalToWorld"));
 	Engine_Owner_SystemWorldToLocal = FNiagaraVariable(FNiagaraTypeDefinition::GetMatrix4Def(), TEXT("Engine.Owner.SystemWorldToLocal"));
@@ -344,6 +347,7 @@ void INiagaraModule::StartupModule()
 	Engine_Emitter_TotalSpawnedParticles = FNiagaraVariable(FNiagaraTypeDefinition::GetIntDef(), TEXT("Engine.Emitter.TotalSpawnedParticles"));
 	Engine_Emitter_SpawnCountScale = FNiagaraVariable(FNiagaraTypeDefinition::GetFloatDef(), TEXT("Engine.Emitter.SpawnCountScale"));
 	Engine_Emitter_InstanceSeed = FNiagaraVariable(FNiagaraTypeDefinition::GetIntDef(), TEXT("Engine.Emitter.InstanceSeed"));
+	Engine_Emitter_ID = FNiagaraVariable(FNiagaraTypeDefinition(FNiagaraEmitterID::StaticStruct()), TEXT("Engine.Emitter.ID"));
 	Engine_System_TickCount = FNiagaraVariable(FNiagaraTypeDefinition::GetIntDef(), TEXT("Engine.System.TickCount"));
 	Engine_System_NumEmittersAlive = FNiagaraVariable(FNiagaraTypeDefinition::GetIntDef(), TEXT("Engine.System.NumEmittersAlive"));
 	Engine_System_SignificanceIndex = FNiagaraVariable(FNiagaraTypeDefinition::GetIntDef(), TEXT("Engine.System.SignificanceIndex"));
diff --git a/Engine/Plugins/FX/Niagara/Source/Niagara/Private/NiagaraSystemInstance.cpp b/Engine/Plugins/FX/Niagara/Source/Niagara/Private/NiagaraSystemInstance.cpp
--- a/Engine/Plugins/FX/Niagara/Source/Niagara/Private/NiagaraSystemInstance.cpp
+++ b/Engine/Plugins/FX/Niagara/Source/Niagara/Private/NiagaraSystemInstance.cpp
@@ -2824,6 +2824,13 @@ FBox FNiagaraSystemInstance::GetEmitterFixedBounds(FName EmitterName) const
 	return FBox(ForceInit);
 }
 
+FNiagaraEmitterInstance* FNiagaraSystemInstance::GetEmitterByID(FNiagaraEmitterID InID)const
+{
+	//Currently the emitter ID is the direct index of the emitter in the Emitters array, though this may not always be the case.
+	int32 EmitterIndex = InID.ID;
+	return Emitters.IsValidIndex(EmitterIndex) ? &Emitters[EmitterIndex].Get() : nullptr;
+}
+
 void FNiagaraSystemInstance::SetForceSolo(bool bInForceSolo)
 {
 	// We may be forced into solo mode so if we match nothing to do here
diff --git a/Engine/Plugins/FX/Niagara/Source/Niagara/Public/NiagaraCompilationTypes.h b/Engine/Plugins/FX/Niagara/Source/Niagara/Public/NiagaraCompilationTypes.h
--- a/Engine/Plugins/FX/Niagara/Source/Niagara/Public/NiagaraCompilationTypes.h
+++ b/Engine/Plugins/FX/Niagara/Source/Niagara/Public/NiagaraCompilationTypes.h
@@ -40,6 +40,7 @@ struct FNiagaraScriptAsyncCompileData
 	FNiagaraVMExecutableDataId CompileId;
 	TSharedPtr<struct FNiagaraVMExecutableData> ExeData;
 	FString UniqueEmitterName;
+	FNiagaraEmitterID EmitterID = INDEX_NONE;
 	FNiagaraScriptCompileMetrics CompileMetrics;
 
 	bool bFromDerivedDataCache = false;
diff --git a/Engine/Plugins/FX/Niagara/Source/Niagara/Public/NiagaraModule.h b/Engine/Plugins/FX/Niagara/Source/Niagara/Public/NiagaraModule.h
--- a/Engine/Plugins/FX/Niagara/Source/Niagara/Public/NiagaraModule.h
+++ b/Engine/Plugins/FX/Niagara/Source/Niagara/Public/NiagaraModule.h
@@ -163,6 +163,7 @@ public:
 	FORCEINLINE static const FNiagaraVariable&  GetVar_Engine_Owner_Scale() { return Engine_Owner_Scale; }
 	FORCEINLINE static const FNiagaraVariable&  GetVar_Engine_Owner_Rotation() { return Engine_Owner_Rotation; }
 	FORCEINLINE static const FNiagaraVariable&  GetVar_Engine_Owner_LWC_Tile() { return Engine_Owner_LWC_Tile; }
+	FORCEINLINE static const FNiagaraVariable&  GetVar_Engine_ExecIndex() { return Engine_ExecIndex; }
 
 	FORCEINLINE static const FNiagaraVariable&  GetVar_Engine_Owner_SystemLocalToWorld() { return Engine_Owner_SystemLocalToWorld; }
 	FORCEINLINE static const FNiagaraVariable&  GetVar_Engine_Owner_SystemWorldToLocal() { return Engine_Owner_SystemWorldToLocal; }
@@ -183,6 +184,7 @@ public:
 	FORCEINLINE static const FNiagaraVariable&  GetVar_Engine_Emitter_TotalSpawnedParticles() { return Engine_Emitter_TotalSpawnedParticles; }
 	FORCEINLINE static const FNiagaraVariable&  GetVar_Engine_Emitter_SpawnCountScale() { return Engine_Emitter_SpawnCountScale; }
 	FORCEINLINE static const FNiagaraVariable&  GetVar_Engine_Emitter_InstanceSeed() { return Engine_Emitter_InstanceSeed; }
+	FORCEINLINE static const FNiagaraVariable&  GetVar_Engine_Emitter_ID() { return Engine_Emitter_ID; }
 	FORCEINLINE static const FNiagaraVariable&  GetVar_Engine_System_TickCount() { return Engine_System_TickCount; }
 	FORCEINLINE static const FNiagaraVariable&  GetVar_Engine_System_NumEmittersAlive() { return Engine_System_NumEmittersAlive; }
 	FORCEINLINE static const FNiagaraVariable&  GetVar_Engine_System_SignificanceIndex() { return Engine_System_SignificanceIndex; }
@@ -296,6 +298,7 @@ private:
 	static NIAGARA_API FNiagaraVariable Engine_Owner_Scale;
 	static NIAGARA_API FNiagaraVariable Engine_Owner_Rotation;
 	static NIAGARA_API FNiagaraVariable Engine_Owner_LWC_Tile;
+	static NIAGARA_API FNiagaraVariable Engine_ExecIndex;
 
 	static NIAGARA_API FNiagaraVariable Engine_Owner_SystemLocalToWorld;
 	static NIAGARA_API FNiagaraVariable Engine_Owner_SystemWorldToLocal;
@@ -340,6 +343,7 @@ private:
 	static NIAGARA_API FNiagaraVariable Emitter_SimulationTarget;
 	static NIAGARA_API FNiagaraVariable Emitter_RandomSeed;
 	static NIAGARA_API FNiagaraVariable Engine_Emitter_InstanceSeed;
+	static NIAGARA_API FNiagaraVariable Engine_Emitter_ID;
 	static NIAGARA_API FNiagaraVariable Emitter_SpawnRate;
 	static NIAGARA_API FNiagaraVariable Emitter_SpawnInterval;
 	static NIAGARA_API FNiagaraVariable Emitter_InterpSpawnStartDt;
diff --git a/Engine/Plugins/FX/Niagara/Source/Niagara/Public/NiagaraSystemInstance.h b/Engine/Plugins/FX/Niagara/Source/Niagara/Public/NiagaraSystemInstance.h
--- a/Engine/Plugins/FX/Niagara/Source/Niagara/Public/NiagaraSystemInstance.h
+++ b/Engine/Plugins/FX/Niagara/Source/Niagara/Public/NiagaraSystemInstance.h
@@ -249,6 +249,8 @@ public:
 	NIAGARA_API void SetEmitterFixedBounds(FName EmitterName, const FBox& InLocalBounds);
 	NIAGARA_API FBox GetEmitterFixedBounds(FName EmitterName) const;
 
+	NIAGARA_API FNiagaraEmitterInstance* GetEmitterByID(FNiagaraEmitterID ID)const;
+
 	NIAGARA_API void SetForceSolo(bool bForceSolo);
 	FORCEINLINE bool IsSolo() const { return bSolo; }
 
diff --git a/Engine/Plugins/FX/Niagara/Source/Niagara/Public/NiagaraTypes.h b/Engine/Plugins/FX/Niagara/Source/Niagara/Public/NiagaraTypes.h
--- a/Engine/Plugins/FX/Niagara/Source/Niagara/Public/NiagaraTypes.h
+++ b/Engine/Plugins/FX/Niagara/Source/Niagara/Public/NiagaraTypes.h
@@ -191,6 +191,18 @@ struct FNiagaraMatrix
 	FVector4f Row3 = FVector4f(ForceInitToZero);
 };
 
+USTRUCT(meta = (DisplayName = "EmitterID"))
+struct FNiagaraEmitterID
+{
+	GENERATED_BODY()
+	
+	FNiagaraEmitterID() = default;
+	FNiagaraEmitterID(int32 InID):ID(InID) {}
+
+	UPROPERTY(EditAnywhere, Category = Emitter)
+	int32 ID = INDEX_NONE;
+};
+
 USTRUCT()
 struct FNiagaraAssetVersion
 {
diff --git a/Engine/Plugins/FX/Niagara/Source/NiagaraEditor/Private/NiagaraCompilationDigestBridgeImpl.cpp b/Engine/Plugins/FX/Niagara/Source/NiagaraEditor/Private/NiagaraCompilationDigestBridgeImpl.cpp
--- a/Engine/Plugins/FX/Niagara/Source/NiagaraEditor/Private/NiagaraCompilationDigestBridgeImpl.cpp
+++ b/Engine/Plugins/FX/Niagara/Source/NiagaraEditor/Private/NiagaraCompilationDigestBridgeImpl.cpp
@@ -122,6 +122,11 @@ const FNiagaraCompilationGraph* FNiagaraCompilationDigestBridge::GetEmitterGraph
 	return EmitterNode->CalledGraph.Get();
 }
 
+FNiagaraEmitterID FNiagaraCompilationDigestBridge::GetEmitterID(const FEmitterNode* EmitterNode)
+{
+	return EmitterNode->EmitterID;
+}
+
 FString FNiagaraCompilationDigestBridge::GetEmitterUniqueName(const FEmitterNode* EmitterNode)
 {
 	return EmitterNode->EmitterUniqueName;
diff --git a/Engine/Plugins/FX/Niagara/Source/NiagaraEditor/Private/NiagaraCompilationGraphBridgeImpl.cpp b/Engine/Plugins/FX/Niagara/Source/NiagaraEditor/Private/NiagaraCompilationGraphBridgeImpl.cpp
--- a/Engine/Plugins/FX/Niagara/Source/NiagaraEditor/Private/NiagaraCompilationGraphBridgeImpl.cpp
+++ b/Engine/Plugins/FX/Niagara/Source/NiagaraEditor/Private/NiagaraCompilationGraphBridgeImpl.cpp
@@ -218,6 +218,11 @@ const UNiagaraGraph* FNiagaraCompilationGraphBridge::GetEmitterGraph(const FEmit
 	return EmitterNode->GetCalledGraph();
 }
 
+FNiagaraEmitterID FNiagaraCompilationGraphBridge::GetEmitterID(const FEmitterNode* EmitterNode)
+{
+	return EmitterNode->GetEmitterID();
+}
+
 FString FNiagaraCompilationGraphBridge::GetEmitterUniqueName(const FEmitterNode* EmitterNode)
 {
 	return EmitterNode->GetEmitterUniqueName();
diff --git a/Engine/Plugins/FX/Niagara/Source/NiagaraEditor/Private/NiagaraCompilationPrivate.h b/Engine/Plugins/FX/Niagara/Source/NiagaraEditor/Private/NiagaraCompilationPrivate.h
--- a/Engine/Plugins/FX/Niagara/Source/NiagaraEditor/Private/NiagaraCompilationPrivate.h
+++ b/Engine/Plugins/FX/Niagara/Source/NiagaraEditor/Private/NiagaraCompilationPrivate.h
@@ -93,6 +93,8 @@ public:
 	FName ResolveEmitterAlias(FName VariableName) const;
 
 	const FString& GetUniqueEmitterName() const { return EmitterUniqueName; }
+	FNiagaraEmitterID GetEmitterID()const { return EmitterID; }
+
 	void FinishPrecompile(
 		const FNiagaraSystemCompilationTask& CompilationTask,
 		TConstArrayView<FNiagaraVariable> EncounterableVariables,
@@ -138,6 +140,7 @@ public:
 
 	TArray<FNiagaraVariable> EncounteredVariables;
 	FString EmitterUniqueName;
+	FNiagaraEmitterID EmitterID;
 	TArray<FSharedPrecompileData> EmitterData;
 	FNiagaraDigestedGraphPtr DigestedSourceGraph;
 	FString SourceName;
@@ -244,6 +247,7 @@ public:
 
 	TArray<FNiagaraVariable> EncounteredVariables;
 	FString EmitterUniqueName;
+	FNiagaraEmitterID EmitterID;
 	TArray<TSharedPtr<FNiagaraCompileRequestData, ESPMode::ThreadSafe>> EmitterData;
 	TWeakObjectPtr<UNiagaraScriptSource> Source;
 	FString SourceName;
@@ -311,6 +315,7 @@ public:
 
 	TArray<FNiagaraVariable> ChangedFromNumericVars;
 	FString EmitterUniqueName;
+	FNiagaraEmitterID EmitterID;
 	TArray<TSharedPtr<FNiagaraCompileRequestDuplicateData, ESPMode::ThreadSafe>> EmitterData;
 
 	struct FDuplicatedGraphData
diff --git a/Engine/Plugins/FX/Niagara/Source/NiagaraEditor/Private/NiagaraCompilationTasks.cpp b/Engine/Plugins/FX/Niagara/Source/NiagaraEditor/Private/NiagaraCompilationTasks.cpp
--- a/Engine/Plugins/FX/Niagara/Source/NiagaraEditor/Private/NiagaraCompilationTasks.cpp
+++ b/Engine/Plugins/FX/Niagara/Source/NiagaraEditor/Private/NiagaraCompilationTasks.cpp
@@ -1237,6 +1237,7 @@ void FNiagaraSystemCompilationTask::Precompile()
 	{
 		TSharedPtr<FNiagaraPrecompileData, ESPMode::ThreadSafe> EmitterPtr = MakeShared<FNiagaraPrecompileData, ESPMode::ThreadSafe>();
 		EmitterPtr->EmitterUniqueName = EmitterInfo.UniqueEmitterName;
+		EmitterPtr->EmitterID = FNiagaraEmitterID(EmitterInfo.SourceEmitterIndex);
 		EmitterPtr->SourceName = SystemPrecompileData->SourceName;
 		EmitterPtr->DigestedSourceGraph = EmitterInfo.SourceGraph;
 		EmitterPtr->bUseRapidIterationParams = SystemPrecompileData->bUseRapidIterationParams;
diff --git a/Engine/Plugins/FX/Niagara/Source/NiagaraEditor/Private/NiagaraCompiler.cpp b/Engine/Plugins/FX/Niagara/Source/NiagaraEditor/Private/NiagaraCompiler.cpp
--- a/Engine/Plugins/FX/Niagara/Source/NiagaraEditor/Private/NiagaraCompiler.cpp
+++ b/Engine/Plugins/FX/Niagara/Source/NiagaraEditor/Private/NiagaraCompiler.cpp
@@ -250,7 +250,7 @@ void FNiagaraCompileRequestDuplicateData::DuplicateReferencedGraphsRecursive(UNi
 					{
 						EmitterNode->SyncEnabledState(); // Just to be safe, sync here while we likely still have the handle source.
 						EmitterNode->SetOwnerSystem(nullptr);
-						EmitterNode->SetCachedVariablesForCompilation(*Ptr->EmitterUniqueName, Ptr->NodeGraphDeepCopy.Get(), Ptr->SourceDeepCopy.Get());
+						EmitterNode->SetCachedVariablesForCompilation(*Ptr->EmitterUniqueName, Ptr->EmitterID, Ptr->NodeGraphDeepCopy.Get(), Ptr->SourceDeepCopy.Get());
 					}
 				}
 			}
@@ -1028,6 +1028,7 @@ TSharedPtr<FNiagaraCompileRequestDataBase, ESPMode::ThreadSafe> FNiagaraEditorMo
 			const FNiagaraEmitterHandle& Handle = System->GetEmitterHandle(i);
 			TSharedPtr<FNiagaraCompileRequestData, ESPMode::ThreadSafe> EmitterPtr = MakeShared<FNiagaraCompileRequestData, ESPMode::ThreadSafe>();
 			EmitterPtr->EmitterUniqueName = Handle.GetInstance().Emitter->GetUniqueEmitterName();
+			EmitterPtr->EmitterID = FNiagaraEmitterID(i);
 			EmitterPtr->SourceName = BasePtr->SourceName;
 			EmitterPtr->Source = Cast<UNiagaraScriptSource>(Handle.GetEmitterData()->GraphSource);
 			EmitterPtr->bUseRapidIterationParams = BasePtr->bUseRapidIterationParams;
@@ -1353,6 +1354,7 @@ TSharedPtr<FNiagaraCompileRequestDuplicateDataBase, ESPMode::ThreadSafe> FNiagar
 			const FNiagaraEmitterHandle& Handle = OwningSystem->GetEmitterHandle(i);
 			TSharedPtr<FNiagaraCompileRequestDuplicateData, ESPMode::ThreadSafe> EmitterPtr = MakeShared<FNiagaraCompileRequestDuplicateData, ESPMode::ThreadSafe>();
 			EmitterPtr->EmitterUniqueName = Handle.GetInstance().Emitter->GetUniqueEmitterName();
+			EmitterPtr->EmitterID = FNiagaraEmitterID(i);
 			EmitterPtr->ValidUsages = BasePtr->ValidUsages;
 			EmitterPtr->SharedSourceGraphToDuplicatedGraphsMap = BasePtr->SharedSourceGraphToDuplicatedGraphsMap;
 			EmitterPtr->SharedNameToDuplicatedDataInterfaceMap = BasePtr->SharedNameToDuplicatedDataInterfaceMap;
diff --git a/Engine/Plugins/FX/Niagara/Source/NiagaraEditor/Private/NiagaraGraphDigest.cpp b/Engine/Plugins/FX/Niagara/Source/NiagaraEditor/Private/NiagaraGraphDigest.cpp
--- a/Engine/Plugins/FX/Niagara/Source/NiagaraEditor/Private/NiagaraGraphDigest.cpp
+++ b/Engine/Plugins/FX/Niagara/Source/NiagaraEditor/Private/NiagaraGraphDigest.cpp
@@ -2007,6 +2007,7 @@ FNiagaraCompilationNodeAssignment::FNiagaraCompilationNodeAssignment(const FNiag
 FNiagaraCompilationNodeEmitter::FNiagaraCompilationNodeEmitter(const UNiagaraNodeEmitter* InNode, FNiagaraCompilationGraphCreateContext& Context)
 : FNiagaraCompilationNode(ENodeType::Emitter, InNode, Context)
 {
+	EmitterID = InNode->GetEmitterID();
 	EmitterUniqueName = InNode->GetEmitterUniqueName();
 
 	if (const UNiagaraGraph* DependentGraph = InNode->GetCalledGraph())
diff --git a/Engine/Plugins/FX/Niagara/Source/NiagaraEditor/Private/NiagaraGraphDigest.h b/Engine/Plugins/FX/Niagara/Source/NiagaraEditor/Private/NiagaraGraphDigest.h
--- a/Engine/Plugins/FX/Niagara/Source/NiagaraEditor/Private/NiagaraGraphDigest.h
+++ b/Engine/Plugins/FX/Niagara/Source/NiagaraEditor/Private/NiagaraGraphDigest.h
@@ -393,6 +393,7 @@ public:
 	virtual void BuildParameterMapHistory(FParameterMapHistoryBuilder& Builder, bool bRecursive, bool bFilterForCompilation) const override;
 	virtual void Compile(FTranslator* Translator, TArray<int32>& Outputs) const;
 
+	FNiagaraEmitterID EmitterID;
 	FString EmitterUniqueName;
 	FString EmitterName;
 	FString EmitterPathName;
diff --git a/Engine/Plugins/FX/Niagara/Source/NiagaraEditor/Private/NiagaraHlslTranslator.cpp b/Engine/Plugins/FX/Niagara/Source/NiagaraEditor/Private/NiagaraHlslTranslator.cpp
--- a/Engine/Plugins/FX/Niagara/Source/NiagaraEditor/Private/NiagaraHlslTranslator.cpp
+++ b/Engine/Plugins/FX/Niagara/Source/NiagaraEditor/Private/NiagaraHlslTranslator.cpp
@@ -812,6 +812,17 @@ FString TNiagaraHlslTranslator<GraphBridge>::BuildParameterMapHlslDefinitions(TA
 		UniqueVariables.AddUnique(Var);
 	}
 
+	//Add any special case engine provided values in the param maps. These are not uniforms but engine provided things like Engine.ExecIndex and Engine.Emitter.ID.
+ 	for (FNiagaraVariable& Var : ParamMapDefinedEngineVars)
+ 	{
+ 		if (Var.GetType().GetClass() != nullptr)
+ 		{
+ 			continue;
+ 		}
+ 	
+ 		UniqueVariables.AddUnique(Var);
+ 	}
+
 	bool bIsSpawnScript = IsSpawnScript();
 
 	// For now we only care about attributes from the other output parameter map histories.
@@ -1605,6 +1616,9 @@ FNiagaraTranslateResults TNiagaraHlslTranslator<GraphBridge>::Translate(const FN
 			InterpSpawnVariables.Emplace(FNiagaraTypeDefinition::GetIntDef(), TEXT("Interpolation.Emitter_SpawnGroup"));
 		}
 		
+		ParamMapDefinedEngineVars.Emplace(SYS_PARAM_ENGINE_EXEC_INDEX);
+		ParamMapDefinedEngineVars.Emplace(SYS_PARAM_ENGINE_EMITTER_ID);
+		
 		for (int32 HistoryIdx = 0; HistoryIdx < OtherOutputParamMapHistories.Num(); HistoryIdx++)
 		{
 			FParamMapHistory& FoundHistory = OtherOutputParamMapHistories[HistoryIdx];
@@ -3763,6 +3777,8 @@ void FNiagaraHlslTranslator::DefineMain(FString &OutHlslOutput,
 		}
 	}
 
+	FNiagaraEmitterID EmitterID = GetEmitterID();
+
 	// Fill in the defaults for parameters.
 	for (const auto& PerStageChunks : PerStageMainPreSimulateChunks)
 	{
@@ -3785,6 +3801,10 @@ void FNiagaraHlslTranslator::DefineMain(FString &OutHlslOutput,
 			}
 		}
 
+		//Set Engine provided param map values.
+		OutHlslOutput += FString::Printf(TEXT("\tContext.%s.Engine.ExecIndex = ExecIndex();\n"), *TranslationStages[StageIdx].PassNamespace);
+		OutHlslOutput += FString::Printf(TEXT("\tContext.%s.Engine.Emitter.ID.ID = %d;\n"), *TranslationStages[StageIdx].PassNamespace, EmitterID.ID);
+
 		FName ScopeName(TranslationStages[StageIdx].PassNamespace + TEXT("Main"));
 		EnterStatsScope(FNiagaraStatScope(*(CompileOptions.GetName() + TEXT("_") + ScopeName.ToString()), ScopeName), OutHlslOutput);
 		OutHlslOutput += FString::Printf(TEXT("\tSimulate%s(Context);\n"), TranslationStages.Num() > 1 ? *TranslationStages[StageIdx].PassNamespace : TEXT(""));
@@ -5839,6 +5859,8 @@ bool TNiagaraHlslTranslator<GraphBridge>::ParameterMapRegisterExternalConstantNa
 		return false;
 	}
 
+	const bool bIsEngineProvidedValue = ParamMapDefinedEngineVars.Contains(InVariable);//These values are provided by the engine/translator and do not need to be read as a constant
+
 	InVariable = ActiveHistoryForFunctionCalls.ResolveAliases(InVariable);
 
 	if (!InVariable.IsValid())
@@ -5873,118 +5895,121 @@ bool TNiagaraHlslTranslator<GraphBridge>::ParameterMapRegisterExternalConstantNa
 	bool bIsPerInstanceBulkSystemParam = IsBulkSystemScript() && !bIsDataInterface && !bIsUObject && (FNiagaraParameterUtilities::IsUserParameter(InVariable) || FNiagaraParameterUtilities::IsPerInstanceEngineParameter(InVariable, EmitterAlias != nullptr ? *EmitterAlias : TEXT("Emitter")));
 	const bool bIsExternalConstantParameter = FNiagaraParameterUtilities::IsRapidIterationParameter(InVariable) && !InVariable.GetType().IsStatic();
 
-	if (!bIsPerInstanceBulkSystemParam)
+	if(!bIsEngineProvidedValue)//This is not actually a uniform so just add the code chunk to read it directly.
 	{
-		int32 UniformChunk = 0;
-
-		if (false == ParamMapDefinedSystemVars.Contains(InVariable.GetName()))
+		if (!bIsPerInstanceBulkSystemParam)
 		{
-			FString SymbolNameDefined = FlattenedName;
+			int32 UniformChunk = 0;
 
-			if (bIsDataInterface)
+			if (false == ParamMapDefinedSystemVars.Contains(InVariable.GetName()))
 			{
-				UNiagaraDataInterface* DataInterface = nullptr;
-				if (GraphBridge::IsCollectionValid(Collection))
+				FString SymbolNameDefined = FlattenedName;
+
+				if (bIsDataInterface)
 				{
-					DataInterface = GraphBridge::GetCollectionDataInterface(Collection, InVariable);
-					if (DataInterface == nullptr)
+					UNiagaraDataInterface* DataInterface = nullptr;
+					if (GraphBridge::IsCollectionValid(Collection))
 					{
-						Error(FText::Format(LOCTEXT("ParameterCollectionDataInterfaceNotFoundErrorFormat", "Data interface named {0} of type {1} was not found in Parameter Collection {2}"),
-							FText::FromName(InVariable.GetName()), InVariable.GetType().GetNameText(), FText::FromString(GraphBridge::GetCollectionFullName(Collection))), InNodeForErrorReporting, InDefaultPin);
-						return false;
+						DataInterface = GraphBridge::GetCollectionDataInterface(Collection, InVariable);
+						if (DataInterface == nullptr)
+						{
+							Error(FText::Format(LOCTEXT("ParameterCollectionDataInterfaceNotFoundErrorFormat", "Data interface named {0} of type {1} was not found in Parameter Collection {2}"),
+								FText::FromName(InVariable.GetName()), InVariable.GetType().GetNameText(), FText::FromString(GraphBridge::GetCollectionFullName(Collection))), InNodeForErrorReporting, InDefaultPin);
+							return false;
+						}
+					}
+					else
+					{
+						DataInterface = CompileDuplicateData->GetDuplicatedDataInterfaceCDOForClass(const_cast<UClass*>(InVariable.GetType().GetClass()));
+						if (DataInterface == nullptr)
+						{
+							Error(FText::Format(LOCTEXT("GetDuplicatedDataInterfaceCDOForClassFailed", "GetDuplicatedDataInterfaceCDOForClass failed for Variable({0}) Class({1})"), FText::FromName(InVariable.GetName()), InVariable.GetType().GetNameText()), InNodeForErrorReporting, InDefaultPin);
+							return false;
+						}
+					}
+					if (ensure(DataInterface))
+					{
+						Output = RegisterDataInterface(InVariable, DataInterface, true, true);
+						return true;
 					}
 				}
-				else
+				if (bIsUObject)
 				{
-					DataInterface = CompileDuplicateData->GetDuplicatedDataInterfaceCDOForClass(const_cast<UClass*>(InVariable.GetType().GetClass()));
-					if (DataInterface == nullptr)
+					UObject* Object = nullptr;
+					if (GraphBridge::IsCollectionValid(Collection))
 					{
-						Error(FText::Format(LOCTEXT("GetDuplicatedDataInterfaceCDOForClassFailed", "GetDuplicatedDataInterfaceCDOForClass failed for Variable({0}) Class({1})"), FText::FromName(InVariable.GetName()), InVariable.GetType().GetNameText()), InNodeForErrorReporting, InDefaultPin);
-						return false;
+						Object = GraphBridge::GetCollectionUObject(Collection, InVariable);
 					}
+					Output = RegisterUObject(InVariable, Object, true);
+					return true;
 				}
-				if (ensure(DataInterface))
+				if (!InVariable.IsDataAllocated() && !InDefaultPin)
 				{
-					Output = RegisterDataInterface(InVariable, DataInterface, true, true);
-					return true;
+					FNiagaraEditorUtilities::ResetVariableToDefaultValue(InVariable);
 				}
-			}
-			if (bIsUObject)
-			{
-				UObject* Object = nullptr;
-				if (GraphBridge::IsCollectionValid(Collection))
+				else if (!InVariable.IsDataAllocated())
 				{
-					Object = GraphBridge::GetCollectionUObject(Collection, InVariable);
+					FillVariableWithDefaultValue(InVariable, InDefaultPin);
 				}
-				Output = RegisterUObject(InVariable, Object, true);
-				return true;
-			}
-			if (!InVariable.IsDataAllocated() && !InDefaultPin)
-			{
-				FNiagaraEditorUtilities::ResetVariableToDefaultValue(InVariable);
-			}
-			else if (!InVariable.IsDataAllocated())
-			{
-				FillVariableWithDefaultValue(InVariable, InDefaultPin);
-			}
-
-			if (InVariable.GetAllocatedSizeInBytes() != InVariable.GetSizeInBytes())
-			{
-				Error(FText::Format(LOCTEXT("GetParameterUnsetParam", "Variable {0} hasn't had its default value set. Required Bytes: {1} vs Allocated Bytes: {2}"), FText::FromName(InVariable.GetName()), FText::AsNumber(InVariable.GetType().GetSize()), FText::AsNumber(InVariable.GetSizeInBytes())), nullptr, nullptr);
-			}
 
-			if ( IsVariableInUniformBuffer(InVariable) )
-			{
-				// we must ensure that there's a one to one relationship between symbol name and parameter.  The generated VM only
-				// knows about the symbols while the parameter stores knows about the parameters, if these mismatch, then we're going
-				// to be incorrectly addressing the constant table
-				if (!CompilationOutput.ScriptData.Parameters.FindParameter(InVariable))
+				if (InVariable.GetAllocatedSizeInBytes() != InVariable.GetSizeInBytes())
 				{
-					bool AddParameter = true;
+					Error(FText::Format(LOCTEXT("GetParameterUnsetParam", "Variable {0} hasn't had its default value set. Required Bytes: {1} vs Allocated Bytes: {2}"), FText::FromName(InVariable.GetName()), FText::AsNumber(InVariable.GetType().GetSize()), FText::AsNumber(InVariable.GetSizeInBytes())), nullptr, nullptr);
+				}
 
-					// add the parameter, but first evaluate whether any of the symbols for existing parameters would conflict
-					for (const FNiagaraVariable& ExistingParameter : CompilationOutput.ScriptData.Parameters.Parameters)
+				if (IsVariableInUniformBuffer(InVariable))
+				{
+					// we must ensure that there's a one to one relationship between symbol name and parameter.  The generated VM only
+					// knows about the symbols while the parameter stores knows about the parameters, if these mismatch, then we're going
+					// to be incorrectly addressing the constant table
+					if (!CompilationOutput.ScriptData.Parameters.FindParameter(InVariable))
 					{
-						FNameBuilder ExistingParameterName(ExistingParameter.GetName());
-						if (GetSanitizedSymbolName(ExistingParameterName).Equals(SymbolName))
+						bool AddParameter = true;
+
+						// add the parameter, but first evaluate whether any of the symbols for existing parameters would conflict
+						for (const FNiagaraVariable& ExistingParameter : CompilationOutput.ScriptData.Parameters.Parameters)
 						{
-							Error(FText::Format(LOCTEXT("NonUniqueSymbolNames", "Parameters ('{0}' and '{1}') found which resolve to the same HLSL symbol name '{2}'.  These should be disambiguated."),
-								FText::FromName(InVariable.GetName()), FText::FromName(ExistingParameter.GetName()), FText::FromString(SymbolName)), InNodeForErrorReporting, InDefaultPin);
+							FNameBuilder ExistingParameterName(ExistingParameter.GetName());
+							if (GetSanitizedSymbolName(ExistingParameterName).Equals(SymbolName))
+							{
+								Error(FText::Format(LOCTEXT("NonUniqueSymbolNames", "Parameters ('{0}' and '{1}') found which resolve to the same HLSL symbol name '{2}'.  These should be disambiguated."),
+									FText::FromName(InVariable.GetName()), FText::FromName(ExistingParameter.GetName()), FText::FromString(SymbolName)), InNodeForErrorReporting, InDefaultPin);
 
-							AddParameter = false;
+								AddParameter = false;
+							}
 						}
-					}
 
-					if (AddParameter)
-					{
-						CompilationOutput.ScriptData.Parameters.Parameters.Add(InVariable);
+						if (AddParameter)
+						{
+							CompilationOutput.ScriptData.Parameters.Parameters.Add(InVariable);
+						}
 					}
 				}
+
+				UniformChunk = AddUniformChunk(SymbolNameDefined, InVariable, ENiagaraCodeChunkMode::Uniform, UNiagaraScript::IsGPUScript(CompileOptions.TargetUsage));
+			}
+			else
+			{
+				const auto& SystemVar = ParamMapDefinedSystemVars.FindChecked(InVariable.GetName());
+				UniformChunk = SystemVar.ChunkIndex;
 			}
 
-			UniformChunk = AddUniformChunk(SymbolNameDefined, InVariable, ENiagaraCodeChunkMode::Uniform, UNiagaraScript::IsGPUScript(CompileOptions.TargetUsage));
-		}
-		else
-		{
-			const auto& SystemVar = ParamMapDefinedSystemVars.FindChecked(InVariable.GetName());
-			UniformChunk = SystemVar.ChunkIndex;
-		}
-			
-		if (bIsExternalConstantParameter)
-		{
-			Output = UniformChunk;
-			return true;
+			if (bIsExternalConstantParameter)
+			{
+				Output = UniformChunk;
+				return true;
+			}
+			else
+			{
+				//Add this separately as the same uniform can appear in the pre sim chunks more than once in different param maps.
+				PerStageMainPreSimulateChunks[ActiveStageIdx].AddUnique(FString::Printf(TEXT("%s.%s = %s;"), *ParameterMapInstanceName, *GetSanitizedSymbolName(VarName), *GetCodeAsSource(UniformChunk)));
+			}
 		}
-		else
+		else if (bIsPerInstanceBulkSystemParam && !ExternalVariablesForBulkUsage.Contains(InVariable))
 		{
-			//Add this separately as the same uniform can appear in the pre sim chunks more than once in different param maps.
-			PerStageMainPreSimulateChunks[ActiveStageIdx].AddUnique(FString::Printf(TEXT("%s.%s = %s;"), *ParameterMapInstanceName, *GetSanitizedSymbolName(VarName), *GetCodeAsSource(UniformChunk)));
+			ExternalVariablesForBulkUsage.Add(InVariable);
 		}
 	}
-	else if (bIsPerInstanceBulkSystemParam && !ExternalVariablesForBulkUsage.Contains(InVariable))
-	{
-		ExternalVariablesForBulkUsage.Add(InVariable);
-	}
 
 	Output = AddSourceChunk(ParameterMapInstanceName + TEXT(".") + SymbolName, InVariable.GetType());
 	if (CodeChunks.IsValidIndex(Output))
@@ -6409,6 +6434,16 @@ void TNiagaraHlslTranslator<GraphBridge>::Emitter(const FEmitterNode* EmitterNod
 		// Clear out the parameter map writes to emitter module parameters as they should not be shared across emitters.
 		if (ParamMapHistoryIdx != -1 && ParamMapHistoryIdx < ParamMapHistories.Num())
 		{
+			FParamMapHistory& History = ParamMapHistories[ParamMapHistoryIdx];
+
+			//Set some special case values that the engine provides directly
+			FNiagaraVariable EmitterIDVar = ActiveHistoryForFunctionCalls.ResolveAliases(SYS_PARAM_ENGINE_EMITTER_ID);
+			//Add the alias resolved var to the ParamMapDefinedEngineVars so that we can use this array to inject these values into the hlsl struct defs later.
+			ParamMapDefinedEngineVars.Add(EmitterIDVar);
+			History.AddVariable(EmitterIDVar, SYS_PARAM_ENGINE_EMITTER_ID, NAME_None, nullptr);
+			FNiagaraEmitterID EmitterID = GraphBridge::GetEmitterID(EmitterNode);
+			AddBodyChunk(TEXT(""), FString::Printf(TEXT("%s.%s.ID = %d"), *GetParameterMapInstanceName(ParamMapHistoryIdx), *EmitterIDVar.GetName().ToString(), EmitterID.ID), FNiagaraTypeDefinition::GetIntDef(), false);
+
 			for (int32 i = 0; i < ParamMapHistories[ParamMapHistoryIdx].Variables.Num(); i++)
 			{
 				check(ParamMapHistories[ParamMapHistoryIdx].VariablesWithOriginalAliasesIntact.Num() > i);
@@ -10185,6 +10220,12 @@ const FString& TNiagaraHlslTranslator<GraphBridge>::GetEmitterUniqueName() const
 	return CompileData->EmitterUniqueName;
 }
 
+template<typename GraphBridge>
+FNiagaraEmitterID TNiagaraHlslTranslator<GraphBridge>::GetEmitterID() const
+{
+	return CompileData->EmitterID;
+}
+
 template<typename GraphBridge>
 TConstArrayView<FNiagaraVariable> TNiagaraHlslTranslator<GraphBridge>::GetStaticVariables() const
 {
diff --git a/Engine/Plugins/FX/Niagara/Source/NiagaraEditor/Private/NiagaraHlslTranslator.h b/Engine/Plugins/FX/Niagara/Source/NiagaraEditor/Private/NiagaraHlslTranslator.h
--- a/Engine/Plugins/FX/Niagara/Source/NiagaraEditor/Private/NiagaraHlslTranslator.h
+++ b/Engine/Plugins/FX/Niagara/Source/NiagaraEditor/Private/NiagaraHlslTranslator.h
@@ -470,6 +470,7 @@ protected:
 	template<typename T>
 	void BuildConstantBuffer(ENiagaraCodeChunkMode ChunkMode);
 
+	virtual FNiagaraEmitterID GetEmitterID() const = 0;
 	virtual const FString& GetEmitterUniqueName() const = 0;
 	virtual TConstArrayView<FNiagaraVariable> GetStaticVariables() const = 0;
 	virtual UNiagaraDataInterface* GetDataInterfaceCDO(UClass* DIClass) const = 0;
@@ -506,6 +507,8 @@ protected:
 		bool bDefaultExplicit = false; //Whether or not the default value of the variable is explicit, e.g. there is an explicit value on a pin, explicit binding, or explicit custom initialization.
 	};
 
+	TArray<FNiagaraVariable> ParamMapDefinedEngineVars; // Engine provided values that we want to be accessible via the parameter map but not to add as uniforms etc.
+
 	TMap<FName, UniformVariableInfo> ParamMapDefinedSystemVars; // Map from the defined constants to the uniform chunk expressing them (i.e. have we encountered before in this graph?)
 
 	// Synced to the EmitterParameter uniforms encountered for parameter maps thus far.
@@ -726,6 +729,7 @@ protected:
 	void UpdateStaticSwitchConstants(const FPin* Pin);
 
 	virtual const FString& GetEmitterUniqueName() const override;
+	virtual FNiagaraEmitterID GetEmitterID() const override;
 	virtual TConstArrayView<FNiagaraVariable> GetStaticVariables() const override;
 	virtual UNiagaraDataInterface* GetDataInterfaceCDO(UClass* DIClass) const override;
 
diff --git a/Engine/Plugins/FX/Niagara/Source/NiagaraEditor/Private/NiagaraNodeEmitter.cpp b/Engine/Plugins/FX/Niagara/Source/NiagaraEditor/Private/NiagaraNodeEmitter.cpp
--- a/Engine/Plugins/FX/Niagara/Source/NiagaraEditor/Private/NiagaraNodeEmitter.cpp
+++ b/Engine/Plugins/FX/Niagara/Source/NiagaraEditor/Private/NiagaraNodeEmitter.cpp
@@ -159,6 +159,24 @@ void UNiagaraNodeEmitter::NodeConnectionListChanged()
 	//GetGraph()->NotifyGraphChanged();
 }
 
+FNiagaraEmitterID UNiagaraNodeEmitter::GetEmitterID()const
+{
+	if (OwnerSystem != nullptr && EmitterHandleId.IsValid())
+	{
+		const auto& Emitters = OwnerSystem->GetEmitterHandles();
+		for (int32 i = 0; i < Emitters.Num(); ++i)
+		{
+			const FNiagaraEmitterHandle& EmitterHandle = Emitters[i];
+			if (EmitterHandle.GetId() == EmitterHandleId)
+			{
+				return FNiagaraEmitterID(i);
+			}
+		}
+	}
+
+	return CachedEmitterID;
+}
+
 FString UNiagaraNodeEmitter::GetEmitterUniqueName() const
 {
 	if (OwnerSystem != nullptr && EmitterHandleId.IsValid())
@@ -267,9 +285,10 @@ void UNiagaraNodeEmitter::SyncEnabledState()
 	}
 }
 
-void UNiagaraNodeEmitter::SetCachedVariablesForCompilation(const FName& InUniqueName, UNiagaraGraph* InGraph, UNiagaraScriptSourceBase* InSource)
+void UNiagaraNodeEmitter::SetCachedVariablesForCompilation(const FName& InUniqueName, FNiagaraEmitterID InEmitterID, UNiagaraGraph* InGraph, UNiagaraScriptSourceBase* InSource)
 {
 	CachedUniqueName = InUniqueName;
+	CachedEmitterID = InEmitterID;
 	CachedGraphWeakPtr = InGraph;
 	CachedScriptSourceWeakPtr = InSource;
 }
diff --git a/Engine/Plugins/FX/Niagara/Source/NiagaraEditor/Private/NiagaraNodeEmitter.h b/Engine/Plugins/FX/Niagara/Source/NiagaraEditor/Private/NiagaraNodeEmitter.h
--- a/Engine/Plugins/FX/Niagara/Source/NiagaraEditor/Private/NiagaraNodeEmitter.h
+++ b/Engine/Plugins/FX/Niagara/Source/NiagaraEditor/Private/NiagaraNodeEmitter.h
@@ -55,6 +55,7 @@ public:
 	ENiagaraScriptUsage GetUsage() const { return ScriptType; }
 	void SetUsage(ENiagaraScriptUsage InUsage) { ScriptType = InUsage; }
 
+	FNiagaraEmitterID GetEmitterID()const;
 	FString GetEmitterUniqueName() const;
 	UNiagaraScriptSource* GetScriptSource() const;
 	UNiagaraGraph* GetCalledGraph() const;
@@ -62,7 +63,7 @@ public:
 	virtual void Compile(FTranslator* Translator, TArray<int32>& Outputs) const override;
 	virtual void GatherExternalDependencyData(ENiagaraScriptUsage InUsage, const FGuid& InUsageId, TArray<FNiagaraCompileHash>& InReferencedCompileHashes, TArray<FString>& InReferencedObjs) const override;
 
-	void SetCachedVariablesForCompilation(const FName& InUniqueName, UNiagaraGraph* InGraph, UNiagaraScriptSourceBase* InSource);
+	void SetCachedVariablesForCompilation(const FName& InUniqueName, FNiagaraEmitterID InEmitterID, UNiagaraGraph* InGraph, UNiagaraScriptSourceBase* InSource);
 
 protected:
 	UEdGraphPin* PinPendingRename;
@@ -92,6 +93,7 @@ private:
 	ENiagaraScriptUsage ScriptType;
 
 	FName CachedUniqueName;
+	FNiagaraEmitterID CachedEmitterID = INDEX_NONE;
 	TWeakObjectPtr<UNiagaraGraph> CachedGraphWeakPtr;
 	TWeakObjectPtr<UNiagaraScriptSourceBase> CachedScriptSourceWeakPtr;
 };
\ No newline at end of file
diff --git a/Engine/Plugins/FX/Niagara/Source/NiagaraEditor/Public/NiagaraCompilationBridge.h b/Engine/Plugins/FX/Niagara/Source/NiagaraEditor/Public/NiagaraCompilationBridge.h
--- a/Engine/Plugins/FX/Niagara/Source/NiagaraEditor/Public/NiagaraCompilationBridge.h
+++ b/Engine/Plugins/FX/Niagara/Source/NiagaraEditor/Public/NiagaraCompilationBridge.h
@@ -161,6 +161,7 @@ struct FNiagaraCompilationGraphBridge
 	static void BuildTraversal(const FGraph* Graph, const FNode* OutputNode, TArray<const FNode*>& TraversedNodes);
 	static const FGraph* GetEmitterGraph(const FEmitterNode* EmitterNode);
 	static FString GetEmitterUniqueName(const FEmitterNode* EmitterNode);
+	static FNiagaraEmitterID GetEmitterID(const FEmitterNode* EmitterNode);
 	static ENiagaraScriptUsage GetEmitterUsage(const FEmitterNode* EmitterNode);
 	static FString GetEmitterName(const FEmitterNode* EmitterNode);
 	static FString GetEmitterPathName(const FEmitterNode* EmitterNode);
@@ -286,6 +287,7 @@ struct FNiagaraCompilationDigestBridge
 	static void FindOutputNodes(const FGraph* Graph, TArray<const FOutputNode*>& OutputNodes);
 	static void BuildTraversal(const FGraph* Graph, const FNode* OutputNode, TArray<const FNode*>& TraversedNodes);
 	static const FGraph* GetEmitterGraph(const FEmitterNode* EmitterNode);
+	static FNiagaraEmitterID GetEmitterID(const FEmitterNode* EmitterNode);
 	static FString GetEmitterUniqueName(const FEmitterNode* EmitterNode);
 	static ENiagaraScriptUsage GetEmitterUsage(const FEmitterNode* EmitterNode);
 	static FString GetEmitterName(const FEmitterNode* EmitterNode);
diff --git a/Engine/Source/Runtime/Core/Private/UObject/FortniteMainBranchObjectVersion.cpp b/Engine/Source/Runtime/Core/Private/UObject/FortniteMainBranchObjectVersion.cpp
--- a/Engine/Source/Runtime/Core/Private/UObject/FortniteMainBranchObjectVersion.cpp
+++ b/Engine/Source/Runtime/Core/Private/UObject/FortniteMainBranchObjectVersion.cpp
@@ -10,8 +10,8 @@ TMap<FGuid, FGuid> FFortniteMainBranchObjectVersion::GetSystemGuids()
 	SystemGuids.Add(DevGuids.LANDSCAPE_MOBILE_COOK_VERSION, FGuid("32D02EF867C74B71A0D4E0FA41392732"));
 	SystemGuids.Add(DevGuids.MATERIALSHADERMAP_DERIVEDDATA_VER, FGuid("CE7776E63E9E43FBB8C103EE1AB40B5D"));
 	SystemGuids.Add(DevGuids.NANITE_DERIVEDDATA_VER, FGuid("D28C003F825F4C178F698938FB0E47F7"));
-	SystemGuids.Add(DevGuids.NIAGARASHADERMAP_DERIVEDDATA_VER, FGuid("A6FFCB143BBC4C83A72001468D23D5A4"));
-	SystemGuids.Add(DevGuids.Niagara_LatestScriptCompileVersion, FGuid("25FCFB6BB6554C95890E801EC0FED47E"));
+	SystemGuids.Add(DevGuids.NIAGARASHADERMAP_DERIVEDDATA_VER, FGuid("7232E432BA8040DFB25048D86E632B3A"));
+	SystemGuids.Add(DevGuids.Niagara_LatestScriptCompileVersion, FGuid("5761233C13C1441E9BBA7E381D8382FB"));
 	SystemGuids.Add(DevGuids.SkeletalMeshDerivedDataVersion, FGuid("518D7C51D31D9BE4A9F4B81542AC8119"));
 	SystemGuids.Add(DevGuids.STATICMESH_DERIVEDDATA_VER, FGuid("543843B103794E0E9BA4BE80FB602F79"));
 	return SystemGuids;
