non-editor fix

diff --git a/Engine/Plugins/FX/Niagara/Source/Niagara/Private/DataInterface/NiagaraDataInterfaceDataChannelRead.cpp b/Engine/Plugins/FX/Niagara/Source/Niagara/Private/DataInterface/NiagaraDataInterfaceDataChannelRead.cpp
--- a/Engine/Plugins/FX/Niagara/Source/Niagara/Private/DataInterface/NiagaraDataInterfaceDataChannelRead.cpp
+++ b/Engine/Plugins/FX/Niagara/Source/Niagara/Private/DataInterface/NiagaraDataInterfaceDataChannelRead.cpp
@@ -1494,14 +1494,13 @@ void UNiagaraDataInterfaceDataChannelRead::GetParameterDefinitionHLSL(FNiagaraDa
 		}
 	}
 }
-#endif
 
 bool UNiagaraDataInterfaceDataChannelRead::UpgradeFunctionCall(FNiagaraFunctionSignature& FunctionSignature)
 {
 	TArray<FNiagaraFunctionSignature> Funcs;
 	GetFunctions(Funcs);
 
-	for(const FNiagaraFunctionSignature& Func : Funcs)
+	for (const FNiagaraFunctionSignature& Func : Funcs)
 	{
 		if (Func.Name == FunctionSignature.Name && Func.FunctionVersion > FunctionSignature.FunctionVersion)
 		{
@@ -1512,7 +1511,7 @@ bool UNiagaraDataInterfaceDataChannelRead::UpgradeFunctionCall(FNiagaraFunctionS
 			FunctionSignature.GetVariadicOutputs(VariadicOutputs);
 
 			FunctionSignature = Func;
-			for(FNiagaraVariableBase& Param : VariadicInputs)
+			for (FNiagaraVariableBase& Param : VariadicInputs)
 			{
 				FunctionSignature.AddInput(Param);
 			}
@@ -1527,6 +1526,8 @@ bool UNiagaraDataInterfaceDataChannelRead::UpgradeFunctionCall(FNiagaraFunctionS
 	return false;
 }
 
+#endif
+
 void UNiagaraDataInterfaceDataChannelRead::BuildShaderParameters(FNiagaraShaderParametersBuilder& ShaderParametersBuilder) const
 {
 	ShaderParametersBuilder.AddNestedStruct<NDIDataChannelReadLocal::FShaderParameters>();
