Fix for unsanitized symbol name in trasnlator.

diff --git a/Engine/Plugins/FX/Niagara/Source/NiagaraEditor/Private/NiagaraHlslTranslator.cpp b/Engine/Plugins/FX/Niagara/Source/NiagaraEditor/Private/NiagaraHlslTranslator.cpp
--- a/Engine/Plugins/FX/Niagara/Source/NiagaraEditor/Private/NiagaraHlslTranslator.cpp
+++ b/Engine/Plugins/FX/Niagara/Source/NiagaraEditor/Private/NiagaraHlslTranslator.cpp
@@ -6442,7 +6442,7 @@ void TNiagaraHlslTranslator<GraphBridge>::Emitter(const FEmitterNode* EmitterNod
 			ParamMapDefinedEngineVars.Add(EmitterIDVar);
 			History.AddVariable(EmitterIDVar, SYS_PARAM_ENGINE_EMITTER_ID, NAME_None, nullptr);
 			FNiagaraEmitterID EmitterID = GraphBridge::GetEmitterID(EmitterNode);
-			AddBodyChunk(TEXT(""), FString::Printf(TEXT("%s.%s.ID = %d"), *GetParameterMapInstanceName(ParamMapHistoryIdx), *EmitterIDVar.GetName().ToString(), EmitterID.ID), FNiagaraTypeDefinition::GetIntDef(), false);
+			AddBodyChunk(TEXT(""), FString::Printf(TEXT("%s.%s.ID = %d"), *GetParameterMapInstanceName(ParamMapHistoryIdx), *GetSanitizedSymbolName(EmitterIDVar.GetName().ToString()), EmitterID.ID), FNiagaraTypeDefinition::GetIntDef(), false);
 
 			for (int32 i = 0; i < ParamMapHistories[ParamMapHistoryIdx].Variables.Num(); i++)
 			{
