- Fix for not updating the GPU sync mode for packaged builds

diff --git a/Engine/Plugins/FX/Niagara/Source/Niagara/Private/DataInterface/NiagaraDataInterfaceSimpleCounter.cpp b/Engine/Plugins/FX/Niagara/Source/Niagara/Private/DataInterface/NiagaraDataInterfaceSimpleCounter.cpp
--- a/Engine/Plugins/FX/Niagara/Source/Niagara/Private/DataInterface/NiagaraDataInterfaceSimpleCounter.cpp
+++ b/Engine/Plugins/FX/Niagara/Source/Niagara/Private/DataInterface/NiagaraDataInterfaceSimpleCounter.cpp
@@ -165,14 +165,17 @@ void UNiagaraDataInterfaceSimpleCounter::PostInitProperties()
 	}
 }
 
+void UNiagaraDataInterfaceSimpleCounter::PostLoad()
+{
+	Super::PostLoad();
+	UpdateDIProxy();
+}
+
 #if WITH_EDITOR
 void UNiagaraDataInterfaceSimpleCounter::PostEditChangeProperty(struct FPropertyChangedEvent& PropertyChangedEvent)
 {
 	Super::PostEditChangeProperty(PropertyChangedEvent);
-
-	// Ensure proxy properties are up to date
-	FNDISimpleCounterProxy* Proxy_GT = GetProxyAs<FNDISimpleCounterProxy>();
-	Proxy_GT->GpuSyncMode = GpuSyncMode;
+	UpdateDIProxy();
 }
 #endif
 
@@ -201,8 +204,7 @@ bool UNiagaraDataInterfaceSimpleCounter::CopyToInternal(UNiagaraDataInterface* D
 	DestinationTyped->InitialValue = InitialValue;
 
 	// Ensure proxy properties are up to date
-	FNDISimpleCounterProxy* DestinationProxy = DestinationTyped->GetProxyAs<FNDISimpleCounterProxy>();
-	DestinationProxy->GpuSyncMode = GpuSyncMode;
+	DestinationTyped->UpdateDIProxy();
 
 	return true;
 }
@@ -468,6 +470,12 @@ void UNiagaraDataInterfaceSimpleCounter::PushToRenderThreadImpl()
 	}
 }
 
+void UNiagaraDataInterfaceSimpleCounter::UpdateDIProxy()
+{
+	FNDISimpleCounterProxy* Proxy_GT = GetProxyAs<FNDISimpleCounterProxy>();
+	Proxy_GT->GpuSyncMode = GpuSyncMode;
+}
+
 void UNiagaraDataInterfaceSimpleCounter::VMGet(FVectorVMExternalFunctionContext& Context)
 {
 	VectorVM::FUserPtrHandler<FNDISimpleCounterInstanceData_GameThread> InstanceData(Context);
diff --git a/Engine/Plugins/FX/Niagara/Source/Niagara/Private/DataInterface/NiagaraDataInterfaceSimpleCounter.h b/Engine/Plugins/FX/Niagara/Source/Niagara/Private/DataInterface/NiagaraDataInterfaceSimpleCounter.h
--- a/Engine/Plugins/FX/Niagara/Source/Niagara/Private/DataInterface/NiagaraDataInterfaceSimpleCounter.h
+++ b/Engine/Plugins/FX/Niagara/Source/Niagara/Private/DataInterface/NiagaraDataInterfaceSimpleCounter.h
@@ -26,6 +26,7 @@ class UNiagaraDataInterfaceSimpleCounter : public UNiagaraDataInterfaceRWBase
 public:
 	// UObject Interface
 	NIAGARA_API virtual void PostInitProperties() override;
+	NIAGARA_API virtual void PostLoad() override;
 #if WITH_EDITOR
 	NIAGARA_API virtual void PostEditChangeProperty(struct FPropertyChangedEvent& PropertyChangedEvent) override;
 #endif
@@ -56,6 +57,8 @@ public:
 	NIAGARA_API virtual void PushToRenderThreadImpl() override;
 	// UNiagaraDataInterface Interface End
 
+	void UpdateDIProxy();
+
 	// VM functions
 	NIAGARA_API void VMGet(FVectorVMExternalFunctionContext& Context);
 	NIAGARA_API void VMSet(FVectorVMExternalFunctionContext& Context);
