- Add support for SetEmitterEnabled in Niagara
- Similar to Cascade this will halt all spawning, but won't naturally complete the emitter this will follow how emitter state is setup

diff --git a/Engine/Plugins/FX/Niagara/Source/Niagara/Classes/NiagaraEmitterInstance.h b/Engine/Plugins/FX/Niagara/Source/Niagara/Classes/NiagaraEmitterInstance.h
--- a/Engine/Plugins/FX/Niagara/Source/Niagara/Classes/NiagaraEmitterInstance.h
+++ b/Engine/Plugins/FX/Niagara/Source/Niagara/Classes/NiagaraEmitterInstance.h
@@ -158,6 +158,11 @@ public:
 	}
 	FBox GetFixedBounds() const;
 
+	void SetEmitterEnable(bool bNewEnableState)
+	{
+		bAllowSpawning_GT = bNewEnableState;
+	}
+
 	NIAGARA_API UObject* FindBinding(const FNiagaraVariable& InVariable) const;
 	NIAGARA_API UNiagaraDataInterface* FindDataInterface(const FNiagaraVariable& InVariable) const;
 
@@ -258,6 +263,9 @@ private:
 	/** Typical resets must be deferred until the tick as the RT could still be using the current buffer. */
 	uint32 bResetPending : 1;
 
+	bool bAllowSpawning_GT = true;
+	bool bAllowSpawning_CNC = true;
+
 	// This is used to keep track which particles have spawned a component. This is needed when the bOnlyCreateComponentsOnParticleSpawn flag is set in the renderer.
 	// Without this bookkeeping, the particles would lose their components when the render state is recreated or the visibility tag flips them off and on again.
 	mutable TMap<FObjectKey, TSet<int32>> ParticlesWithComponents;
diff --git a/Engine/Plugins/FX/Niagara/Source/Niagara/Private/NiagaraComponent.cpp b/Engine/Plugins/FX/Niagara/Source/Niagara/Private/NiagaraComponent.cpp
--- a/Engine/Plugins/FX/Niagara/Source/Niagara/Private/NiagaraComponent.cpp
+++ b/Engine/Plugins/FX/Niagara/Source/Niagara/Private/NiagaraComponent.cpp
@@ -638,6 +638,9 @@ void UNiagaraComponent::SetLODDistance(float InLODDistance, float InMaxLODDistan
 
 void UNiagaraComponent::SetEmitterEnable(FName EmitterName, bool bNewEnableState)
 {
+	FEmitterOverrideInfo& EmitterOverrideInfo = FindOrAddEmitterOverrideInfo(EmitterName);
+	EmitterOverrideInfo.bEnabled = bNewEnableState;
+
 	if (SystemInstanceController.IsValid() && !SystemInstanceController->IsComplete())
 	{
 		SystemInstanceController->SetEmitterEnable(EmitterName, bNewEnableState);
@@ -669,7 +672,9 @@ void UNiagaraComponent::ClearSystemFixedBounds()
 
 void UNiagaraComponent::SetEmitterFixedBounds(FName EmitterName, FBox LocalBounds)
 {
-	EmitterFixedBounds.FindOrAdd(EmitterName, LocalBounds);
+	FEmitterOverrideInfo& EmitterOverrideInfo = FindOrAddEmitterOverrideInfo(EmitterName);
+	EmitterOverrideInfo.FixedBounds = LocalBounds;
+
 	if (SystemInstanceController.IsValid())
 	{
 		SystemInstanceController->SetEmitterFixedBounds(EmitterName, LocalBounds);
@@ -678,16 +683,20 @@ void UNiagaraComponent::SetEmitterFixedBounds(FName EmitterName, FBox LocalBound
 
 FBox UNiagaraComponent::GetEmitterFixedBounds(FName EmitterName) const
 {
-	if ( const FBox* EmitterBounds = EmitterFixedBounds.Find(EmitterName) )
+	if ( const FEmitterOverrideInfo* EmitterOverrideInfo = FindEmitterOverrideInfo(EmitterName) )
 	{
-		return *EmitterBounds;
+		return EmitterOverrideInfo->FixedBounds;
 	}
 	return FBox(EForceInit::ForceInit);
 }
 
 void UNiagaraComponent::ClearEmitterFixedBounds(FName EmitterName)
 {
-	EmitterFixedBounds.Remove(EmitterName);
+	if (FEmitterOverrideInfo* EmitterOverrideInfo = FindEmitterOverrideInfo(EmitterName))
+	{
+		EmitterOverrideInfo->FixedBounds = FBox(EForceInit::ForceInit);
+	}
+
 	if (SystemInstanceController.IsValid())
 	{
 		SystemInstanceController->SetEmitterFixedBounds(EmitterName, FBox(EForceInit::ForceInit));
@@ -1113,11 +1122,15 @@ bool UNiagaraComponent::InitializeSystem()
 		{
 			SystemInstanceController->SetSystemFixedBounds(SystemFixedBounds);
 		}
-		if (!EmitterFixedBounds.IsEmpty())
+		for ( const FEmitterOverrideInfo& EmitterOverrideInfo : EmitterOverrideInfos)
 		{
-			for ( auto it=EmitterFixedBounds.CreateConstIterator(); it; ++it )
+			if (!EmitterOverrideInfo.bEnabled)
+			{
+				SystemInstanceController->SetEmitterEnable(EmitterOverrideInfo.EmitterName, EmitterOverrideInfo.bEnabled);
+			}
+			if (EmitterOverrideInfo.FixedBounds.IsValid)
 			{
-				SystemInstanceController->SetEmitterFixedBounds(it.Key(), it.Value());
+				SystemInstanceController->SetEmitterFixedBounds(EmitterOverrideInfo.EmitterName, EmitterOverrideInfo.FixedBounds);
 			}
 		}
 
@@ -1801,7 +1814,7 @@ void UNiagaraComponent::OnPooledReuse(UWorld* NewWorld)
 	ForceUpdateTransformTime = 0.0f;
 	CurrLocalBounds.Init();
 	SystemFixedBounds.Init();
-	EmitterFixedBounds.Empty();
+	EmitterOverrideInfos.Empty();
 
 	if (SystemInstanceController != nullptr)
 	{
@@ -2037,6 +2050,42 @@ void UNiagaraComponent::GetResourceSizeEx(FResourceSizeEx& CumulativeResourceSiz
 	}
 }
 
+const UNiagaraComponent::FEmitterOverrideInfo* UNiagaraComponent::FindEmitterOverrideInfo(FName EmitterName) const
+{
+	for (const FEmitterOverrideInfo& Info : EmitterOverrideInfos)
+	{
+		if (Info.EmitterName == EmitterName)
+		{
+			return &Info;
+		}
+	}
+	return nullptr;
+}
+
+UNiagaraComponent::FEmitterOverrideInfo* UNiagaraComponent::FindEmitterOverrideInfo(FName EmitterName)
+{
+	for (FEmitterOverrideInfo& Info : EmitterOverrideInfos)
+	{
+		if (Info.EmitterName == EmitterName)
+		{
+			return &Info;
+		}
+	}
+	return nullptr;
+}
+
+UNiagaraComponent::FEmitterOverrideInfo& UNiagaraComponent::FindOrAddEmitterOverrideInfo(FName EmitterName)
+{
+	for (FEmitterOverrideInfo& Info : EmitterOverrideInfos)
+	{
+		if (Info.EmitterName == EmitterName)
+		{
+			return Info;
+		}
+	}
+	return EmitterOverrideInfos.Emplace_GetRef(EmitterName);
+}
+
 void UNiagaraComponent::CreateCullProxy(bool bForce)
 {
 	if (Asset->GetCullProxyMode() != ENiagaraCullProxyMode::None && CullProxy == nullptr)
diff --git a/Engine/Plugins/FX/Niagara/Source/Niagara/Private/NiagaraEmitterInstance.cpp b/Engine/Plugins/FX/Niagara/Source/Niagara/Private/NiagaraEmitterInstance.cpp
--- a/Engine/Plugins/FX/Niagara/Source/Niagara/Private/NiagaraEmitterInstance.cpp
+++ b/Engine/Plugins/FX/Niagara/Source/Niagara/Private/NiagaraEmitterInstance.cpp
@@ -555,6 +555,8 @@ void FNiagaraEmitterInstance::OnPooledReuse()
 	// Ensure we kill any existing particles and mark our buffers for reset
 	bResetPending = true;
 	TotalSpawnedParticles = 0;
+	bAllowSpawning_GT = true;
+	bAllowSpawning_CNC = true;
 
 	FixedBounds.Init();
 }
@@ -1244,6 +1246,8 @@ void FNiagaraEmitterInstance::PreTick()
 		}
 	}
 
+	bAllowSpawning_CNC = bAllowSpawning_GT;
+
 	if (ParentSystemInstance && RendererBindings.GetUObjectsDirty())
 	{
 		if ( const FNiagaraParameterStore* SrcStore = ParentSystemInstance->GetOverrideParameters() )
@@ -1403,7 +1407,7 @@ void FNiagaraEmitterInstance::Tick(float DeltaSeconds)
 
 	// Calculate number of new particles from regular spawning 
 	uint32 SpawnTotal = 0;
-	if (ExecutionState == ENiagaraExecutionState::Active)
+	if (ExecutionState == ENiagaraExecutionState::Active && bAllowSpawning_CNC)
 	{
 		for (FNiagaraSpawnInfo& Info : SpawnInfos)
 		{
diff --git a/Engine/Plugins/FX/Niagara/Source/Niagara/Private/NiagaraSystemInstance.cpp b/Engine/Plugins/FX/Niagara/Source/Niagara/Private/NiagaraSystemInstance.cpp
--- a/Engine/Plugins/FX/Niagara/Source/Niagara/Private/NiagaraSystemInstance.cpp
+++ b/Engine/Plugins/FX/Niagara/Source/Niagara/Private/NiagaraSystemInstance.cpp
@@ -165,62 +165,19 @@ FNiagaraSystemInstance::FNiagaraSystemInstance(UWorld& InWorld, UNiagaraSystem&
 
 void FNiagaraSystemInstance::SetEmitterEnable(FName EmitterName, bool bNewEnableState)
 {
-	// No need fo this code since it's not supported yet
-	//// Wait for any async operations, can complete the system
-	//WaitForAsyncTickAndFinalize();
-	//if (IsComplete())
-	//{
-	//	return;
-	//}
-
-	UE_LOG(LogNiagara, Warning, TEXT("SetEmitterEnable: Is not implemented in Niagara. Emitter(%s) System(%s) Component(%s)"), *EmitterName.ToString(), *GetNameSafe(System), *GetFullNameSafe(AttachComponent.Get()));
-	return;
-
-	/*
-	UNiagaraSystem* System = GetSystem();
-	if (System != nullptr)
+	for (const TSharedRef<FNiagaraEmitterInstance, ESPMode::ThreadSafe>& Emitter : Emitters)
 	{
-		const TArray<FNiagaraEmitterHandle>& EmitterHandles = GetSystem()->GetEmitterHandles();
-		int32 FoundIdx = INDEX_NONE;
-		for (int32 EmitterIdx = 0; EmitterIdx < GetSystem()->GetEmitterHandles().Num(); ++EmitterIdx)
+		if (Emitter->GetEmitterHandle().GetName() == EmitterName)
 		{
-			const FNiagaraEmitterHandle& EmitterHandle = EmitterHandles[EmitterIdx];
-			if (EmitterName == EmitterHandle.GetName())
-			{
-				FoundIdx = EmitterIdx;
-				break;
-			}
+			Emitter->SetEmitterEnable(bNewEnableState);
+			return;
 		}
+	}
 
-		if (FoundIdx != INDEX_NONE && Emitters.IsValidIndex(FoundIdx))
-		{
-			if (Emitters[FoundIdx]->IsAllowedToExecute())
-			{
-
-				{
-					if (bNewEnableState)
-					{
-						Emitters[FoundIdx]->SetExecutionState(ENiagaraExecutionState::Active);
-					}
-					else
-					{
-						Emitters[FoundIdx]->SetExecutionState(ENiagaraExecutionState::Inactive);
-					}
-				}
-			}
-			else
-			{
-				UE_LOG(LogNiagara, Log, TEXT("SetEmitterEnable: Emitter \"%s\" was found in the system's list of emitters, but it does not pass FNiagaraEmitterInstance::IsAllowedToExecute() and therefore cannot be manually enabled!"), *EmitterName.ToString());
-			}
-		}
-		else
-		{
-			UE_LOG(LogNiagara, Log, TEXT("SetEmitterEnable: Emitter \"%s\" was not found in the system's list of emitters!"), *EmitterName.ToString());
-		}
-	}*/
+	// Failed to find emitter
+	UE_LOG(LogNiagara, Warning, TEXT("SetEmitterEnable: Failed to find Emitter(%s) System(%s) Component(%s)"), *EmitterName.ToString(), *GetNameSafe(System), *GetFullNameSafe(AttachComponent.Get()));
 }
 
-
 void FNiagaraSystemInstance::Init(bool bInForceSolo)
 {
 	// We warn if async is not complete here as we should never wait
diff --git a/Engine/Plugins/FX/Niagara/Source/Niagara/Public/NiagaraComponent.h b/Engine/Plugins/FX/Niagara/Source/Niagara/Public/NiagaraComponent.h
--- a/Engine/Plugins/FX/Niagara/Source/Niagara/Public/NiagaraComponent.h
+++ b/Engine/Plugins/FX/Niagara/Source/Niagara/Public/NiagaraComponent.h
@@ -861,8 +861,19 @@ private:
 	FBox CurrLocalBounds;
 
 	FBox SystemFixedBounds;
-	TMap<FName, FBox> EmitterFixedBounds;
 
+	struct FEmitterOverrideInfo
+	{
+		explicit FEmitterOverrideInfo(FName InEmitterName) : EmitterName(InEmitterName) {}
+
+		FName	EmitterName;
+		bool	bEnabled = true;
+		FBox	FixedBounds = FBox(EForceInit::ForceInit);
+	};
+	TArray<FEmitterOverrideInfo> EmitterOverrideInfos;
+	const FEmitterOverrideInfo* FindEmitterOverrideInfo(FName EmitterName) const;
+	FEmitterOverrideInfo* FindEmitterOverrideInfo(FName EmitterName);
+	FEmitterOverrideInfo& FindOrAddEmitterOverrideInfo(FName EmitterName);
 
 	float CustomTimeDilation = 1.0f;
 
diff --git a/Engine/Source/Runtime/Engine/Classes/Particles/ParticleSystemComponent.h b/Engine/Source/Runtime/Engine/Classes/Particles/ParticleSystemComponent.h
--- a/Engine/Source/Runtime/Engine/Classes/Particles/ParticleSystemComponent.h
+++ b/Engine/Source/Runtime/Engine/Classes/Particles/ParticleSystemComponent.h
@@ -420,9 +420,10 @@ public:
 	virtual UFXSystemAsset* GetFXSystemAsset() const { return nullptr; };
 
 	/**
-	 *	Enables/Disables a sub-emitter
+	 *	Enables / disables an emitter by halting spawning of new particles.
+	 *  You will still pay the cost of the emitter update.
 	 *
-	 *	@param	EmitterName			The name of the sub-emitter to set it on
+	 *	@param	EmitterName			The name of the emitter
 	 *	@param	bNewEnableState		The value to set it to
 	 */
 	UFUNCTION(BlueprintCallable, Category = "Effects|Components|ParticleSystem")
