Add an option for NiagaraLandscape DI to disable use of VT

diff --git a/Engine/Plugins/FX/Niagara/Source/Niagara/Classes/NiagaraDataInterfaceLandscape.h b/Engine/Plugins/FX/Niagara/Source/Niagara/Classes/NiagaraDataInterfaceLandscape.h
--- a/Engine/Plugins/FX/Niagara/Source/Niagara/Classes/NiagaraDataInterfaceLandscape.h
+++ b/Engine/Plugins/FX/Niagara/Source/Niagara/Classes/NiagaraDataInterfaceLandscape.h
@@ -49,6 +49,10 @@ public:
 	UPROPERTY(EditAnywhere, Category = "Landscape")
 	TArray<TObjectPtr<UPhysicalMaterial>> PhysicalMaterials;
 
+	/** Can be used to ignore virtual textures even if they are defined for the landscape. */
+	UPROPERTY(EditAnywhere, Category = "Landscape")
+	bool bVirtualTexturesSupported;
+
 	//UObject Interface
 	NIAGARA_API virtual void PostInitProperties() override;	
 	//UObject Interface End
diff --git a/Engine/Plugins/FX/Niagara/Source/Niagara/Private/NiagaraDataInterfaceLandscape.cpp b/Engine/Plugins/FX/Niagara/Source/Niagara/Private/NiagaraDataInterfaceLandscape.cpp
--- a/Engine/Plugins/FX/Niagara/Source/Niagara/Private/NiagaraDataInterfaceLandscape.cpp
+++ b/Engine/Plugins/FX/Niagara/Source/Niagara/Private/NiagaraDataInterfaceLandscape.cpp
@@ -1156,6 +1156,7 @@ bool UNiagaraDataInterfaceLandscape::CopyToInternal(UNiagaraDataInterface* Desti
 	DestinationLandscape->SourceLandscape = SourceLandscape;	
 	DestinationLandscape->SourceMode = SourceMode;
 	DestinationLandscape->PhysicalMaterials = PhysicalMaterials;
+	DestinationLandscape->bVirtualTexturesSupported = bVirtualTexturesSupported;
 
 	return true;
 }
@@ -1169,7 +1170,8 @@ bool UNiagaraDataInterfaceLandscape::Equals(const UNiagaraDataInterface* Other)
 	const UNiagaraDataInterfaceLandscape* OtherLandscape = CastChecked<const UNiagaraDataInterfaceLandscape>(Other);
 	return OtherLandscape->SourceLandscape == SourceLandscape
 		&& OtherLandscape->SourceMode == SourceMode
-		&& OtherLandscape->PhysicalMaterials == PhysicalMaterials;
+		&& OtherLandscape->PhysicalMaterials == PhysicalMaterials
+		&& OtherLandscape->bVirtualTexturesSupported == bVirtualTexturesSupported;
 }
 
 void UNiagaraDataInterfaceLandscape::GetFunctions(TArray<FNiagaraFunctionSignature>& OutFunctions)
@@ -1524,7 +1526,7 @@ void UNiagaraDataInterfaceLandscape::ApplyLandscape(const FNiagaraSystemInstance
 	InstanceData.NormalVirtualTextureIndex = INDEX_NONE;
 
 	// only worry about virtual textures if our current platform supports them
-	if (UseVirtualTexturing(SystemInstance.GetFeatureLevel()))
+	if (bVirtualTexturesSupported && UseVirtualTexturing(SystemInstance.GetFeatureLevel()))
 	{
 		const int32 RuntimeVirtualTextureCount = InstanceData.Landscape->RuntimeVirtualTextures.Num();
 		for (int32 TextureIt = 0; TextureIt < RuntimeVirtualTextureCount; ++TextureIt)
