Niagara - Fix an issue where default linked data interfaces from modules and dynamic inputs wouldn't be displayed correctly in the stack.  Also added a test which verifies that the data side of this is working as expected.

diff --git a/Engine/Plugins/FX/Niagara/Source/NiagaraEditor/Private/ViewModels/Stack/NiagaraStackFunctionInput.cpp b/Engine/Plugins/FX/Niagara/Source/NiagaraEditor/Private/ViewModels/Stack/NiagaraStackFunctionInput.cpp
--- a/Engine/Plugins/FX/Niagara/Source/NiagaraEditor/Private/ViewModels/Stack/NiagaraStackFunctionInput.cpp
+++ b/Engine/Plugins/FX/Niagara/Source/NiagaraEditor/Private/ViewModels/Stack/NiagaraStackFunctionInput.cpp
@@ -1198,7 +1198,7 @@ void UNiagaraStackFunctionInput::RefreshValues()
 	}
 	else
 	{
-		if (InputType.IsDataInterface())
+		if (InputType.IsDataInterface() && DefaultInputValues.Mode == EValueMode::Data)
 		{
 			// If the input it a data interface but hasn't been edited yet, we need to provide a placeholder data interface to edit.
 			FGuid EmitterHandleId = GetEmitterViewModel().IsValid()
@@ -1212,7 +1212,7 @@ void UNiagaraStackFunctionInput::RefreshValues()
 				DefaultInputValues.DataObject->CopyTo(InputValues.DataObject.Get());
 			}
 		}
-		else if (InputType.IsUObject())
+		else if (InputType.IsUObject() && DefaultInputValues.Mode == EValueMode::ObjectAsset)
 		{
 			InputValues.Mode = EValueMode::ObjectAsset;
 			InputValues.ObjectAssetInputNode = DefaultInputValues.ObjectAssetInputNode;
@@ -2044,6 +2044,38 @@ void UNiagaraStackFunctionInput::SetLocalValue(TSharedRef<FStructOnScope> InLoca
 	}
 }
 
+void UNiagaraStackFunctionInput::SetDataInterfaceValue(const UNiagaraDataInterface& InDataInterface)
+{
+	checkf(InDataInterface.GetClass() == InputType.GetClass(), TEXT("Can not set an input to an unrelated type."));
+
+	if (InputValues.Mode == EValueMode::Data && InputValues.DataObject->Equals(&InDataInterface))
+	{
+		// The value matches the current value so noop.
+		return;
+	}
+
+	TGuardValue<bool> UpdateGuard(bUpdatingLocalValueDirectly, true);
+	FScopedTransaction ScopedTransaction(LOCTEXT("UpdateInputLocalDataValue", "Update input local data interface value"));
+	bool bGraphWillNeedRelayout = false;
+	UEdGraphPin* OverridePin = GetOverridePin();
+
+	if (OverridePin != nullptr && OverridePin->LinkedTo.Num() > 0)
+	{
+		// If there is an override pin and it's linked we'll need to remove all of the linked nodes to set a local value.
+		RemoveNodesForOverridePin(*OverridePin);
+		bGraphWillNeedRelayout = true;
+	}
+
+	if (OverridePin == nullptr)
+	{
+		OverridePin = &GetOrCreateOverridePin();
+	}
+
+	UNiagaraDataInterface* InputDataInterface = nullptr;
+	FNiagaraStackGraphUtilities::SetDataInterfaceValueForFunctionInput(*OverridePin, InDataInterface.GetClass(), AliasedInputParameterHandle.GetParameterHandleString().ToString(), InputDataInterface);
+
+}
+
 bool UNiagaraStackFunctionInput::CanReset() const
 {
 	if (bCanResetCache.IsSet() == false)
@@ -3133,16 +3165,23 @@ UEdGraphPin& UNiagaraStackFunctionInput::GetOrCreateOverridePin()
 
 void UNiagaraStackFunctionInput::GetDefaultDataInterfaceValueFromDefaultPin(UEdGraphPin* DefaultPin, UNiagaraStackFunctionInput::FInputValues& InInputValues) const
 {
-	// Default data interfaces are stored on the input node in the graph, but if it doesn't exist or it's data interface pointer is null, just use the CDO.
-	InInputValues.Mode = EValueMode::Data;
 	if (DefaultPin->LinkedTo.Num() == 1 && DefaultPin->LinkedTo[0]->GetOwningNode() != nullptr && DefaultPin->LinkedTo[0]->GetOwningNode()->IsA<UNiagaraNodeInput>())
 	{
+		// If a valid input node was linked, use the data interface from there.
+		InInputValues.Mode = EValueMode::Data;
 		UNiagaraNodeInput* DataInputNode = CastChecked<UNiagaraNodeInput>(DefaultPin->LinkedTo[0]->GetOwningNode());
 		InInputValues.DataObject = DataInputNode->GetDataInterface();
 	}
-	if (InInputValues.DataObject.IsValid() == false)
+	else
 	{
-		InInputValues.DataObject = Cast<UNiagaraDataInterface>(InputType.GetClass()->GetDefaultObject());
+		// If there was no input node, try to get a linked data interface default.
+		GetDefaultLinkedHandleOrLinkedFunctionFromDefaultPin(DefaultPin, InInputValues);
+		if (InInputValues.Mode == EValueMode::None)
+		{
+			// If there is no specified default input and no linked value, use the CDO as the default data value.
+			InInputValues.Mode = EValueMode::Data;
+			InInputValues.DataObject = Cast<UNiagaraDataInterface>(InputType.GetClass()->GetDefaultObject());
+		}
 	}
 }
 
diff --git a/Engine/Plugins/FX/Niagara/Source/NiagaraEditor/Public/ViewModels/Stack/NiagaraStackFunctionInput.h b/Engine/Plugins/FX/Niagara/Source/NiagaraEditor/Public/ViewModels/Stack/NiagaraStackFunctionInput.h
--- a/Engine/Plugins/FX/Niagara/Source/NiagaraEditor/Public/ViewModels/Stack/NiagaraStackFunctionInput.h
+++ b/Engine/Plugins/FX/Niagara/Source/NiagaraEditor/Public/ViewModels/Stack/NiagaraStackFunctionInput.h
@@ -190,6 +190,9 @@ public:
 
 	/** Sets this input's local value. */
 	NIAGARAEDITOR_API void SetLocalValue(TSharedRef<FStructOnScope> InLocalValue);
+
+	/** Sets this input's data interface value. */
+	NIAGARAEDITOR_API void SetDataInterfaceValue(const UNiagaraDataInterface& InDataInterface);
 	
 	/** Returns whether or not the value or handle of this input has been overridden and can be reset. */
 	NIAGARAEDITOR_API bool CanReset() const;
diff --git a/Engine/Plugins/FX/Niagara/Source/NiagaraEditorWidgets/Private/Stack/SNiagaraStackFunctionInputValue.cpp b/Engine/Plugins/FX/Niagara/Source/NiagaraEditorWidgets/Private/Stack/SNiagaraStackFunctionInputValue.cpp
--- a/Engine/Plugins/FX/Niagara/Source/NiagaraEditorWidgets/Private/Stack/SNiagaraStackFunctionInputValue.cpp
+++ b/Engine/Plugins/FX/Niagara/Source/NiagaraEditorWidgets/Private/Stack/SNiagaraStackFunctionInputValue.cpp
@@ -524,16 +524,24 @@ FSlateColor SNiagaraStackFunctionInputValue::GetVersionSelectorColor() const
 
 void SNiagaraStackFunctionInputValue::SetToLocalValue()
 {
-	const UScriptStruct* LocalValueStruct = FunctionInput->GetInputType().GetScriptStruct();
-	if (LocalValueStruct != nullptr)
+	if (FunctionInput->GetInputType().IsDataInterface())
 	{
-		TSharedRef<FStructOnScope> LocalValue = MakeShared<FStructOnScope>(LocalValueStruct);
-		TArray<uint8> DefaultValueData;
-		FNiagaraEditorUtilities::GetTypeDefaultValue(FunctionInput->GetInputType(), DefaultValueData);
-		if (DefaultValueData.Num() == LocalValueStruct->GetStructureSize())
+		const UNiagaraDataInterface* DefaultDataInterfaceValue = GetDefault<UNiagaraDataInterface>(FunctionInput->GetInputType().GetClass());
+		FunctionInput->SetDataInterfaceValue(*DefaultDataInterfaceValue);
+	}
+	else if (FunctionInput->GetInputType().IsUObject() == false)
+	{
+		const UScriptStruct* LocalValueStruct = FunctionInput->GetInputType().GetScriptStruct();
+		if (LocalValueStruct != nullptr)
 		{
-			FMemory::Memcpy(LocalValue->GetStructMemory(), DefaultValueData.GetData(), DefaultValueData.Num());
-			FunctionInput->SetLocalValue(LocalValue);
+			TSharedRef<FStructOnScope> LocalValue = MakeShared<FStructOnScope>(LocalValueStruct);
+			TArray<uint8> DefaultValueData;
+			FNiagaraEditorUtilities::GetTypeDefaultValue(FunctionInput->GetInputType(), DefaultValueData);
+			if (DefaultValueData.Num() == LocalValueStruct->GetStructureSize())
+			{
+				FMemory::Memcpy(LocalValue->GetStructMemory(), DefaultValueData.GetData(), DefaultValueData.Num());
+				FunctionInput->SetLocalValue(LocalValue);
+			}
 		}
 	}
 }
@@ -1208,9 +1216,10 @@ TArray<TSharedPtr<FNiagaraMenuAction_Generic>> SNiagaraStackFunctionInputValue::
 	FNiagaraActionSourceData NiagaraSourceData(EScriptSource::Niagara, FText::FromString(TEXT("Niagara")), true);
 	
 	// Set a local value
-	if(bIsDataInterfaceOrObject == false)
 	{
-		bool bCanSetLocalValue = FunctionInput->GetValueMode() != UNiagaraStackFunctionInput::EValueMode::Local;
+		bool bCanSetLocalValue = 
+			(FunctionInput->GetInputType().IsDataInterface() && FunctionInput->GetValueMode() != UNiagaraStackFunctionInput::EValueMode::Data) ||
+			(bIsDataInterfaceOrObject == false && FunctionInput->GetValueMode() != UNiagaraStackFunctionInput::EValueMode::Local);
 
 		const FText DisplayName = LOCTEXT("LocalValue", "New Local Value");
 		const FText Tooltip = FText::Format(LOCTEXT("LocalValueToolTip", "Set a local editable value for this input."), DisplayName);
