- Deprecate ambiguous call FindVariable in favor of a more explicit name
- Add a method to find variable by DI index

diff --git a/Engine/Plugins/FX/Niagara/Source/Niagara/Private/DataInterface/NiagaraDataInterfaceUtilities.cpp b/Engine/Plugins/FX/Niagara/Source/Niagara/Private/DataInterface/NiagaraDataInterfaceUtilities.cpp
--- a/Engine/Plugins/FX/Niagara/Source/Niagara/Private/DataInterface/NiagaraDataInterfaceUtilities.cpp
+++ b/Engine/Plugins/FX/Niagara/Source/Niagara/Private/DataInterface/NiagaraDataInterfaceUtilities.cpp
@@ -26,7 +26,7 @@ void ForEachVMFunctionImpl(class UNiagaraDataInterface* RuntimeResolvedDataInter
 	{
 		if (OverrideDI != nullptr && OverrideDI == RuntimeResolvedDataInterface)
 		{
-			if (const FNiagaraVariableBase* Variable = ParameterStore.FindVariable(OverrideDI))
+			if (const FNiagaraVariableBase* Variable = ParameterStore.FindVariableFromDataInterface(OverrideDI))
 			{
 				DataInterfaceUserVariable = *Variable;
 				break;
@@ -87,7 +87,7 @@ void ForEachGpuFunctionImpl(class UNiagaraDataInterface* RuntimeResolvedDataInte
 	{
 		if (OverrideDI != nullptr && OverrideDI == RuntimeResolvedDataInterface)
 		{
-			if (const FNiagaraVariableBase* Variable = ParameterStore.FindVariable(OverrideDI))
+			if (const FNiagaraVariableBase* Variable = ParameterStore.FindVariableFromDataInterface(OverrideDI))
 			{
 				DataInterfaceUserVariable = *Variable;
 				break;
diff --git a/Engine/Plugins/FX/Niagara/Source/Niagara/Private/NiagaraDataInterfaceVectorField.cpp b/Engine/Plugins/FX/Niagara/Source/Niagara/Private/NiagaraDataInterfaceVectorField.cpp
--- a/Engine/Plugins/FX/Niagara/Source/Niagara/Private/NiagaraDataInterfaceVectorField.cpp
+++ b/Engine/Plugins/FX/Niagara/Source/Niagara/Private/NiagaraDataInterfaceVectorField.cpp
@@ -236,7 +236,7 @@ void UNiagaraDataInterfaceVectorField::GetFeedback(UNiagaraSystem* InAsset, UNia
 			{
 				if (DI == this || DI->Equals(this))
 				{
-					const FNiagaraVariableBase* Var = InComponent->GetOverrideParameters().FindVariable(DI);
+					const FNiagaraVariableBase* Var = InComponent->GetOverrideParameters().FindVariableFromDataInterface(DI);
 					if (Var)
 					{
 						DIAliases.AddUnique(Var->GetName());
@@ -253,7 +253,7 @@ void UNiagaraDataInterfaceVectorField::GetFeedback(UNiagaraSystem* InAsset, UNia
 			{
 				if (DI == this || DI->Equals(this))
 				{
-					const FNiagaraVariableBase* Var = InAsset->GetExposedParameters().FindVariable(DI);
+					const FNiagaraVariableBase* Var = InAsset->GetExposedParameters().FindVariableFromDataInterface(DI);
 					if (Var)
 					{
 						DIAliases.AddUnique(Var->GetName());
diff --git a/Engine/Plugins/FX/Niagara/Source/Niagara/Private/NiagaraParameterStore.cpp b/Engine/Plugins/FX/Niagara/Source/Niagara/Private/NiagaraParameterStore.cpp
--- a/Engine/Plugins/FX/Niagara/Source/Niagara/Private/NiagaraParameterStore.cpp
+++ b/Engine/Plugins/FX/Niagara/Source/Niagara/Private/NiagaraParameterStore.cpp
@@ -1310,18 +1310,29 @@ void FNiagaraParameterStore::OnLayoutChange()
 #endif
 }
 
-const FNiagaraVariableBase* FNiagaraParameterStore::FindVariable(const UNiagaraDataInterface* Interface) const
+const FNiagaraVariableBase* FNiagaraParameterStore::FindVariableFromDataInterfaceIndex(int32 DataInterfaceIndex) const
 {
 	SCOPE_CYCLE_COUNTER(STAT_NiagaraParameterStoreFindVar);
-	int32 Idx = DataInterfaces.IndexOfByKey(Interface);
+	for (const FNiagaraVariableWithOffset& ParamWithOffset : ReadParameterVariables())
+	{
+		if (ParamWithOffset.Offset == DataInterfaceIndex)
+		{
+			return &ParamWithOffset;
+		}
+	}
+	return nullptr;
+}
+
+const FNiagaraVariableBase* FNiagaraParameterStore::FindVariableFromDataInterface(const UNiagaraDataInterface* Interface) const
+{
+	SCOPE_CYCLE_COUNTER(STAT_NiagaraParameterStoreFindVar);
+	const int32 Idx = DataInterfaces.IndexOfByKey(Interface);
 	if (Idx != INDEX_NONE)
 	{
-		for (const FNiagaraVariableWithOffset& ParamWithOffset : ReadParameterVariables())
+		const FNiagaraVariableBase* Variable = FindVariableFromDataInterfaceIndex(Idx);
+		if (Variable && ensure(Variable->GetType().GetClass() == Interface->GetClass()))
 		{
-			if (ParamWithOffset.Offset == Idx && ParamWithOffset.GetType().GetClass() == Interface->GetClass())
-			{
-				return &ParamWithOffset;
-			}
+			return Variable;
 		}
 	}
 	return nullptr;
diff --git a/Engine/Plugins/FX/Niagara/Source/Niagara/Public/NiagaraParameterStore.h b/Engine/Plugins/FX/Niagara/Source/Niagara/Public/NiagaraParameterStore.h
--- a/Engine/Plugins/FX/Niagara/Source/Niagara/Public/NiagaraParameterStore.h
+++ b/Engine/Plugins/FX/Niagara/Source/Niagara/Public/NiagaraParameterStore.h
@@ -467,8 +467,13 @@ public:
 	/** Returns a struct converter for the given variable, if the store contains the variable and it's a LWC type. */
 	NIAGARA_API FNiagaraLwcStructConverter GetStructConverter(const FNiagaraVariable& Parameter) const;
 
+	UE_DEPRECATED(5.4, "FindVariable has been replaced by FindVariableFromDataInterface.")
+	NIAGARA_API const FNiagaraVariableBase* FindVariable(const UNiagaraDataInterface* Interface) const { return FindVariableFromDataInterface(Interface); }
+
+	/** Returns the associated FNiagaraVariable for the passed data interface index if it exists in the store. Null if not.*/
+	NIAGARA_API const FNiagaraVariableBase* FindVariableFromDataInterfaceIndex(int32 DataInterfaceIndex) const;
 	/** Returns the associated FNiagaraVariable for the passed data interface if it exists in the store. Null if not.*/
-	NIAGARA_API const FNiagaraVariableBase* FindVariable(const UNiagaraDataInterface* Interface) const;
+	NIAGARA_API const FNiagaraVariableBase* FindVariableFromDataInterface(const UNiagaraDataInterface* Interface) const;
 
 	NIAGARA_API virtual const int32* FindParameterOffset(const FNiagaraVariableBase& Parameter, bool IgnoreType = false) const;
 
diff --git a/Engine/Plugins/FX/Niagara/Source/NiagaraEditor/Private/ViewModels/NiagaraSystemViewModel.cpp b/Engine/Plugins/FX/Niagara/Source/NiagaraEditor/Private/ViewModels/NiagaraSystemViewModel.cpp
--- a/Engine/Plugins/FX/Niagara/Source/NiagaraEditor/Private/ViewModels/NiagaraSystemViewModel.cpp
+++ b/Engine/Plugins/FX/Niagara/Source/NiagaraEditor/Private/ViewModels/NiagaraSystemViewModel.cpp
@@ -2465,7 +2465,7 @@ void FNiagaraSystemViewModel::UpdateCompiledDataInterfaces(UNiagaraDataInterface
 	else
 	{
 		// If the data interface wasn't owned by a script, try to find it in the exposed parameter data interfaces.
-		const FNiagaraVariableBase* FoundExposedDataInterface = GetSystem().GetExposedParameters().FindVariable(ChangedDataInterface);
+		const FNiagaraVariableBase* FoundExposedDataInterface = GetSystem().GetExposedParameters().FindVariableFromDataInterface(ChangedDataInterface);
 		if (FoundExposedDataInterface != nullptr)
 		{
 			GetSystem().GetExposedParameters().OnInterfaceChange();
