- Ensure we handle missing / invalid data interfaces when loading Niagara Systems

diff --git a/Engine/Plugins/FX/Niagara/Source/Niagara/Private/NiagaraComponent.cpp b/Engine/Plugins/FX/Niagara/Source/Niagara/Private/NiagaraComponent.cpp
--- a/Engine/Plugins/FX/Niagara/Source/Niagara/Private/NiagaraComponent.cpp
+++ b/Engine/Plugins/FX/Niagara/Source/Niagara/Private/NiagaraComponent.cpp
@@ -3371,6 +3371,13 @@ void UNiagaraComponent::CopyParametersFromAsset(bool bResetExistingOverrideParam
 	{
 		UNiagaraDataInterface* AssetDataInterface = AssetExposedParameters.GetDataInterface(i);
 		UNiagaraDataInterface* OverrideDataInterface = OverrideParameters.GetDataInterface(i);
+		if (!AssetDataInterface || !OverrideDataInterface)
+		{
+			const FNiagaraVariableBase* DIVariable = AssetExposedParameters.FindVariableFromDataInterfaceIndex(i);
+			UE_LOG(LogNiagara, Error, TEXT("null data interface found for Variable(%s) System(%s) will not run."), DIVariable ? *DIVariable->GetName().ToString() : TEXT("Unknown"), *GetNameSafe(Asset));
+			continue;
+		}
+
 		if (AssetDataInterface != OverrideDataInterface)
 		{
 			// We must copy the data regardless as there is an expectation that data interfaces are always reset to the original state
