- Fix for compiling when tweaking rapid iteration parameters on an emitter only asset

diff --git a/Engine/Plugins/FX/Niagara/Source/NiagaraEditor/Private/ViewModels/NiagaraSystemViewModel.cpp b/Engine/Plugins/FX/Niagara/Source/NiagaraEditor/Private/ViewModels/NiagaraSystemViewModel.cpp
--- a/Engine/Plugins/FX/Niagara/Source/NiagaraEditor/Private/ViewModels/NiagaraSystemViewModel.cpp
+++ b/Engine/Plugins/FX/Niagara/Source/NiagaraEditor/Private/ViewModels/NiagaraSystemViewModel.cpp
@@ -3275,7 +3275,7 @@ TSharedPtr<FNiagaraOverviewGraphViewModel> FNiagaraSystemViewModel::GetOverviewG
 
 bool FNiagaraSystemViewModel::CanImpactCompileForEdit() const
 {
-	return bSupportCompileForEdit && !bIsForDataProcessingOnly && EditMode == ENiagaraSystemViewModelEditMode::SystemAsset;
+	return bSupportCompileForEdit && !bIsForDataProcessingOnly && (EditMode == ENiagaraSystemViewModelEditMode::SystemAsset || EditMode == ENiagaraSystemViewModelEditMode::EmitterAsset);
 }
 
 #undef LOCTEXT_NAMESPACE // NiagaraSystemViewModel
