- Add warning to Niagara compile results if the constant table contains a NaN

diff --git a/Engine/Plugins/FX/Niagara/Source/NiagaraEditor/Private/NiagaraCompiler.cpp b/Engine/Plugins/FX/Niagara/Source/NiagaraEditor/Private/NiagaraCompiler.cpp
--- a/Engine/Plugins/FX/Niagara/Source/NiagaraEditor/Private/NiagaraCompiler.cpp
+++ b/Engine/Plugins/FX/Niagara/Source/NiagaraEditor/Private/NiagaraCompiler.cpp
@@ -2123,6 +2123,7 @@ TOptional<FNiagaraCompileResults> FHlslNiagaraCompiler::GetCompileResult(int32 J
 		}
 
 		Results.Data->InternalParameters.Empty();
+		bool bAllFloatsFinite = true;
 		for (int32 i = 0; i < CompilationOutput.InternalConstantOffsets.Num(); ++i)
 		{
 			const FName ConstantName(TEXT("InternalConstant"), i);
@@ -2134,6 +2135,7 @@ TOptional<FNiagaraCompileResults> FHlslNiagaraCompiler::GetCompileResult(int32 J
 			{
 				float Val = *(float*)(CompilationOutput.InternalConstantData.GetData() + Offset);
 				Results.Data->InternalParameters.SetOrAdd(FNiagaraVariable(FNiagaraTypeDefinition::GetFloatDef(), ConstantName))->SetValue(Val);
+				bAllFloatsFinite &= FMath::IsFinite(Val);
 			}
 			break;
 			case EVectorVMBaseTypes::Int:
@@ -2151,6 +2153,11 @@ TOptional<FNiagaraCompileResults> FHlslNiagaraCompiler::GetCompileResult(int32 J
 			}
 		}
 
+		if (!bAllFloatsFinite)
+		{
+			Warning(LOCTEXT("FloatConstantsNanOrInf", "Float constant table contains NaN or Inf, this may result in invalid simulation results."));
+		}
+
 		Results.CompilerWallTime = (float)(FPlatformTime::Seconds() - CompilationJob->StartTime);
 		Results.CompilerWorkerTime = (float)CompilationJob->ShaderCompileJob->Output.CompileTime;
 		Results.CompilerPreprocessTime = (float)CompilationJob->ShaderCompileJob->Output.PreprocessTime;
