Fixed landscape DI function pin to use correct position type (the shader already treated it as lwc position, so just a cosmetic change).

diff --git a/Engine/Plugins/FX/Niagara/Source/Niagara/Private/NiagaraDataInterfaceLandscape.cpp b/Engine/Plugins/FX/Niagara/Source/Niagara/Private/NiagaraDataInterfaceLandscape.cpp
--- a/Engine/Plugins/FX/Niagara/Source/Niagara/Private/NiagaraDataInterfaceLandscape.cpp
+++ b/Engine/Plugins/FX/Niagara/Source/Niagara/Private/NiagaraDataInterfaceLandscape.cpp
@@ -40,6 +40,7 @@ namespace NiagaraDataInterfaceLandscape
 		InitialVersion = 0,
 		SupportVirtualTextures = 1,
 		LWCPosition = 2,
+		MoreLWCPosition = 3,
 
 		VersionPlusOne,
 		LatestVersion = VersionPlusOne - 1
@@ -1184,7 +1185,7 @@ void UNiagaraDataInterfaceLandscape::GetFunctions(TArray<FNiagaraFunctionSignatu
 		Sig.bRequiresContext = false;
 		Sig.bSupportsCPU = false;
 		Sig.Inputs.Add(FNiagaraVariable(FNiagaraTypeDefinition(GetClass()), TEXT("Landscape")));
-		Sig.Inputs.Add(FNiagaraVariable(FNiagaraTypeDefinition::GetVec3Def(), TEXT("WorldPos")));
+		Sig.Inputs.Add(FNiagaraVariable(FNiagaraTypeDefinition::GetPositionDef(), TEXT("WorldPos")));
 		Sig.Outputs.Add(FNiagaraVariable(FNiagaraTypeDefinition::GetVec3Def(), TEXT("Color")));
 		Sig.Outputs.Add(FNiagaraVariable(FNiagaraTypeDefinition::GetBoolDef(), TEXT("IsValid")));
 
