Fixed ensure in FindVariableFromDataInterface

diff --git a/Engine/Plugins/FX/Niagara/Source/Niagara/Private/NiagaraParameterStore.cpp b/Engine/Plugins/FX/Niagara/Source/Niagara/Private/NiagaraParameterStore.cpp
--- a/Engine/Plugins/FX/Niagara/Source/Niagara/Private/NiagaraParameterStore.cpp
+++ b/Engine/Plugins/FX/Niagara/Source/Niagara/Private/NiagaraParameterStore.cpp
@@ -1315,7 +1315,7 @@ const FNiagaraVariableBase* FNiagaraParameterStore::FindVariableFromDataInterfac
 	SCOPE_CYCLE_COUNTER(STAT_NiagaraParameterStoreFindVar);
 	for (const FNiagaraVariableWithOffset& ParamWithOffset : ReadParameterVariables())
 	{
-		if (ParamWithOffset.Offset == DataInterfaceIndex)
+		if (ParamWithOffset.Offset == DataInterfaceIndex && ParamWithOffset.GetType().IsDataInterface())
 		{
 			return &ParamWithOffset;
 		}
