Fixed dynamic mesh c++ class description

diff --git a/Engine/Plugins/FX/Niagara/Source/Niagara/Private/DataInterface/NiagaraDataInterfaceDynamicMesh.h b/Engine/Plugins/FX/Niagara/Source/Niagara/Private/DataInterface/NiagaraDataInterfaceDynamicMesh.h
--- a/Engine/Plugins/FX/Niagara/Source/Niagara/Private/DataInterface/NiagaraDataInterfaceDynamicMesh.h
+++ b/Engine/Plugins/FX/Niagara/Source/Niagara/Private/DataInterface/NiagaraDataInterfaceDynamicMesh.h
@@ -48,7 +48,7 @@ struct FNiagaraDynamicMeshMaterial
 	}
 };
 
-/** Data Interface allowing sampling of a texture */
+/** Data Interface that can create mesh geometry at runtime from triangle data (e.g. to render a custom generated mesh per particle) */
 UCLASS(Experimental, EditInlineNew, Category = "Mesh", CollapseCategories, meta = (DisplayName = "Dynamic Mesh"), MinimalAPI)
 class UNiagaraDataInterfaceDynamicMesh : public UNiagaraDataInterface, public INiagaraRenderableMeshInterface
 {
