Fixed duplicate message key

diff --git a/Engine/Plugins/FX/Niagara/Source/NiagaraEditorWidgets/Private/SNiagaraStackNote.cpp b/Engine/Plugins/FX/Niagara/Source/NiagaraEditorWidgets/Private/SNiagaraStackNote.cpp
--- a/Engine/Plugins/FX/Niagara/Source/NiagaraEditorWidgets/Private/SNiagaraStackNote.cpp
+++ b/Engine/Plugins/FX/Niagara/Source/NiagaraEditorWidgets/Private/SNiagaraStackNote.cpp
@@ -178,7 +178,7 @@ void SNiagaraStackNote::Rebuild()
 				.OnClicked(this, &SNiagaraStackNote::OnEditBodyButtonClicked)
 				.ButtonStyle(&FAppStyle::GetWidgetStyle<FButtonStyle>("HoverHintOnly"))
 				.Visibility(this, &SNiagaraStackNote::GetEditNoteBodyButtonVisibility)
-				.ToolTipText(LOCTEXT("EditHeaderButtonTooltip", "Edit the message of this note."))
+				.ToolTipText(LOCTEXT("EditMessageButtonTooltip", "Edit the message of this note."))
 				[
 					SNew(SImage)
 					.Image(FAppStyle::GetBrush("Icons.Edit"))
