- Handle large bbox with per particle LODs

diff --git a/Engine/Plugins/FX/Niagara/Source/Niagara/Private/NiagaraRendererMeshes.cpp b/Engine/Plugins/FX/Niagara/Source/Niagara/Private/NiagaraRendererMeshes.cpp
--- a/Engine/Plugins/FX/Niagara/Source/Niagara/Private/NiagaraRendererMeshes.cpp
+++ b/Engine/Plugins/FX/Niagara/Source/Niagara/Private/NiagaraRendererMeshes.cpp
@@ -264,7 +264,7 @@ void FNiagaraRendererMeshes::Initialize(const UNiagaraRendererProperties* InProp
 				{
 					MeshData.LODScreenSize = RenderableMesh->GetLODScreenSize(LOD);
 					MeshData.LODScreenSize.X = LOD < LODRange.Y - 1 ? MeshData.LODScreenSize.X : 0.0f;
-					MeshData.LODScreenSize.Y = LOD > LODRange.X ? MeshData.LODScreenSize.Y : 2.0f;
+					MeshData.LODScreenSize.Y = LOD > LODRange.X ? MeshData.LODScreenSize.Y : -1.0f;
 				}
 			}
 
@@ -1012,7 +1012,7 @@ FVector4f FNiagaraRendererMeshes::GetShaderLODScreenSize(const FSceneView& View,
 
 	return FVector4f(
 		FMath::Square(MeshData.LODScreenSize.X * 0.5f),
-		FMath::Square(MeshData.LODScreenSize.Y * 0.5f),
+		MeshData.LODScreenSize.Y >= 0.0f ? FMath::Square(MeshData.LODScreenSize.Y * 0.5f) : UE_FLOAT_HUGE_DISTANCE,
 		FMath::Square(LODSize),
 		ProjMatrix.M[2][3] * LODView.LODDistanceFactor * MeshData.LODDistanceFactor
 	);
