Disable NiagaraDataChannels during cooking

diff --git a/Engine/Plugins/FX/Niagara/Source/Niagara/Private/NiagaraModule.cpp b/Engine/Plugins/FX/Niagara/Source/Niagara/Private/NiagaraModule.cpp
--- a/Engine/Plugins/FX/Niagara/Source/Niagara/Private/NiagaraModule.cpp
+++ b/Engine/Plugins/FX/Niagara/Source/Niagara/Private/NiagaraModule.cpp
@@ -121,6 +121,13 @@ static FAutoConsoleVariableRef CVarDataChannelEnabled(
 
 void INiagaraModule::OnDataChannelsEnabledChanged(IConsoleVariable* Variable)
 {
+	if(IsRunningCookCommandlet())
+	{
+		//Never allow NDCs when cooking.
+		INiagaraModule::bDataChannelsEnabled = false;
+		return;
+	}
+
 	if(DataChannelsEnabled())
 	{
 		ENiagaraTypeRegistryFlags Flags = ENiagaraTypeRegistryFlags::AllowAnyVariable | ENiagaraTypeRegistryFlags::AllowParameter;
@@ -277,6 +284,12 @@ void INiagaraModule::StartupModule()
 	LLM_SCOPE(ELLMTag::Niagara);
 	FNiagaraTypeDefinition::Init();
 
+	if (IsRunningCookCommandlet())
+	{
+		//Never allow NDCs when cooking.
+		INiagaraModule::bDataChannelsEnabled = false;
+	}
+
 #if UE_USE_OPENVDB
 	// Global registration of  the vdb types.
 	openvdb::initialize();
