- Output system name if we have a null view pointer

diff --git a/Engine/Plugins/FX/Niagara/Source/Niagara/Private/NiagaraRendererMeshes.cpp b/Engine/Plugins/FX/Niagara/Source/Niagara/Private/NiagaraRendererMeshes.cpp
--- a/Engine/Plugins/FX/Niagara/Source/Niagara/Private/NiagaraRendererMeshes.cpp
+++ b/Engine/Plugins/FX/Niagara/Source/Niagara/Private/NiagaraRendererMeshes.cpp
@@ -1072,8 +1072,16 @@ void FNiagaraRendererMeshes::SetupElementForGPUScene(
 		OutMeshBatchElement.bPreserveInstanceOrder = ParticleMeshRenderData.bNeedsSort;
 
 		GPUSceneRes.DynamicPrimitiveData.DataWriterGPU = FGPUSceneWriteDelegate::CreateLambda(				
-			[&GPUSceneRes](FRDGBuilder& GraphBuilder, const FGPUSceneWriteDelegateParams& Params)
+			[&GPUSceneRes,ResouceName=SceneProxy.GetResourceName()](FRDGBuilder& GraphBuilder, const FGPUSceneWriteDelegateParams& Params)
 			{
+			#if !UE_BUILD_SHIPPING
+				//-TEMP: We have a nullptr for the view in a very low repro situation this is to trap the system name to hopefully give some clues.
+				if (Params.View == nullptr)
+				{
+					UE_LOG(LogUObjectGlobals, Fatal, TEXT("Null view pointer in GPU scene update - %s"), *ResouceName.ToString());
+				}
+			#endif
+
 				GPUSceneRes.GPUWriteParams.GPUSceneWriterParameters	= Params.GPUWriteParams;
 				GPUSceneRes.GPUWriteParams.View						= Params.View->ViewUniformBuffer; // NOTE: Set here, not outside lambda
 				GPUSceneRes.GPUWriteParams.Scene					= GetSceneUniformBufferRef(GraphBuilder, *Params.View);
