Fixed small issues with the sim cache asset editor

diff --git a/Engine/Plugins/FX/Niagara/Source/NiagaraEditor/Private/ViewModels/NiagaraSimCacheViewModel.cpp b/Engine/Plugins/FX/Niagara/Source/NiagaraEditor/Private/ViewModels/NiagaraSimCacheViewModel.cpp
--- a/Engine/Plugins/FX/Niagara/Source/NiagaraEditor/Private/ViewModels/NiagaraSimCacheViewModel.cpp
+++ b/Engine/Plugins/FX/Niagara/Source/NiagaraEditor/Private/ViewModels/NiagaraSimCacheViewModel.cpp
@@ -165,6 +165,7 @@ void FNiagaraSimCacheViewModel::SetFrameIndex(const int32 InFrameIndex)
 		const float NormalizedFrame = FMath::Clamp( NumFrames == 0 ? 0.0f : float(InFrameIndex) / float(NumFrames - 1), 0.0f, 1.0f );
 		const float DesiredAge = FMath::Clamp(StartSeconds + (Duration * NormalizedFrame), StartSeconds, StartSeconds + Duration);
 		
+		PreviewComponent->Activate();
 		PreviewComponent->SetDesiredAge(DesiredAge);
 	}
 	OnViewDataChangedDelegate.Broadcast(false);
diff --git a/Engine/Plugins/FX/Niagara/Source/NiagaraEditor/Private/Widgets/SNiagaraSimCacheViewTransportControls.cpp b/Engine/Plugins/FX/Niagara/Source/NiagaraEditor/Private/Widgets/SNiagaraSimCacheViewTransportControls.cpp
--- a/Engine/Plugins/FX/Niagara/Source/NiagaraEditor/Private/Widgets/SNiagaraSimCacheViewTransportControls.cpp
+++ b/Engine/Plugins/FX/Niagara/Source/NiagaraEditor/Private/Widgets/SNiagaraSimCacheViewTransportControls.cpp
@@ -81,6 +81,7 @@ TSharedRef<SWidget> SNiagaraSimCacheViewTransportControls::CreateFrameWidget()
 		.Padding(0.5f)
 		[
 			SNew(SNumericEntryBox<int32>)
+			.MinDesiredValueWidth(25)
 			.AllowSpin(false)
 			.MinValue(0)
 			.MaxValue(this, &SNiagaraSimCacheViewTransportControls::GetMaxFrameIndex)
