- Fix for bad normals on some mesh renderers
- Remove normalize from CalcTangentToWorld in mesh particle renderer.  This is causing issue with some data from meshes which is invalid (i.e. 0,0,0) and matches what the local vertex factory does.  The per particle matrix we transform by is already no scale so the normalize is superfluous

diff --git a/Engine/Plugins/FX/Niagara/Shaders/Private/NiagaraMeshVertexFactory.ush b/Engine/Plugins/FX/Niagara/Shaders/Private/NiagaraMeshVertexFactory.ush
--- a/Engine/Plugins/FX/Niagara/Shaders/Private/NiagaraMeshVertexFactory.ush
+++ b/Engine/Plugins/FX/Niagara/Shaders/Private/NiagaraMeshVertexFactory.ush
@@ -684,8 +684,10 @@ void CalcTangentToWorld(FVertexFactoryInput Input, FVertexFactoryIntermediates I
 	const float3x3 TangentToWorld = mul(TangentToLocal, TransformNoScale);
 	const float DetSign = Intermediates.SceneData.InstanceData.DeterminantSign;
 
-	TangentToWorld0 = normalize(TangentToWorld[0]);
-	TangentToWorld2 = float4(normalize(TangentToWorld[2]), TangentBias(Intermediates.TangentDeterminant) * DetSign);
+	// Note: We do not normalize to match localvertexfactory.ush and the above transform does not contain scale
+	//       If for some reason we need to add the normalize back in you will need to consider UE-195371 as some mesh data is 0,0,0
+	TangentToWorld0 = TangentToWorld[0];
+	TangentToWorld2 = float4(TangentToWorld[2], TangentBias(Intermediates.TangentDeterminant) * DetSign);
 }
 
 FVertexFactoryInterpolantsVSToPS VertexFactoryGetInterpolantsVSToPS(FVertexFactoryInput Input, FVertexFactoryIntermediates Intermediates, FMaterialVertexParameters VertexParameters)
