- Fix a crash with main pass scene captures mesh renderers with low latency
- Don't build translucent batches for scene captures which perform opaque only draws

diff --git a/Engine/Plugins/FX/Niagara/Source/Niagara/Private/NiagaraRenderer.cpp b/Engine/Plugins/FX/Niagara/Source/Niagara/Private/NiagaraRenderer.cpp
--- a/Engine/Plugins/FX/Niagara/Source/Niagara/Private/NiagaraRenderer.cpp
+++ b/Engine/Plugins/FX/Niagara/Source/Niagara/Private/NiagaraRenderer.cpp
@@ -584,6 +584,11 @@ void FNiagaraRenderer::ProcessMaterialParameterBindings(const FNiagaraRendererMa
 	}
 }
 
+bool FNiagaraRenderer::IsViewRenderingOpaqueOnly(const FSceneView* View)
+{
+	return View->bIsSceneCapture && View->SceneCaptureRenderTarget;
+}
+
 void FNiagaraRenderer::SortIndices(const FNiagaraGPUSortInfo& SortInfo, const FNiagaraRendererVariableInfo& SortVariable, const FNiagaraDataBuffer& Buffer, FGlobalDynamicReadBuffer::FAllocation& OutIndices)
 {
 	SCOPE_CYCLE_COUNTER(STAT_NiagaraSortParticles);
diff --git a/Engine/Plugins/FX/Niagara/Source/Niagara/Private/NiagaraRendererMeshes.cpp b/Engine/Plugins/FX/Niagara/Source/Niagara/Private/NiagaraRendererMeshes.cpp
--- a/Engine/Plugins/FX/Niagara/Source/Niagara/Private/NiagaraRendererMeshes.cpp
+++ b/Engine/Plugins/FX/Niagara/Source/Niagara/Private/NiagaraRendererMeshes.cpp
@@ -328,6 +328,8 @@ void FNiagaraRendererMeshes::PrepareParticleMeshRenderData(FParticleMeshRenderDa
 
 	// Check if any materials are translucent and if we can pickup the low latency data
 	// If these conditions change please update the DebugHUD display also to reflect it
+	// Note: SceneCaptures will use latent data as GpuReadyTickStage < CurrentParticleData->GetGPUDataReadyStage()
+	//       For main pass scene captures we exclude the batches if they are translucent
 	const bool bIsWireframe = AllowDebugViewmodes() && ViewFamily.EngineShowFlags.Wireframe;
 	ParticleMeshRenderData.bIsGpuLowLatencyTranslucency =
 		bGpuLowLatencyTranslucency &&
@@ -1423,11 +1425,27 @@ void FNiagaraRendererMeshes::GetDynamicMeshElements(const TArray<const FSceneVie
 						const FMaterial& Material = MaterialProxy->GetIncompleteMaterialWithFallback(FeatureLevel);
 						const bool bTranslucent = IsTranslucentBlendMode(Material);
 						const bool bNeedsPrevTransform = !bTranslucent || Material.IsTranslucencyWritingVelocity();
-						if (bIsShadowView && bTranslucent && !SceneProxy->CastsVolumetricTranslucentShadow())
+						if (bTranslucent)
 						{
-							// Don't add translucent mesh batches for shadows
-							// TODO: Need a way to know if it's a volumetric translucent shadow view to make this logic better
-							continue;
+							if (bIsShadowView && !SceneProxy->CastsVolumetricTranslucentShadow())
+							{
+								// Don't add translucent mesh batches for shadows
+								// TODO: Need a way to know if it's a volumetric translucent shadow view to make this logic better
+								continue;
+							}
+
+							// If we are rendering opaque only we can skip this batch
+							//-OPT: If we only have opaque materials we can skip earlier however due to RemappedMaterialIndex potentially being invalid this is tricky
+							if (IsViewRenderingOpaqueOnly(View))
+							{
+								continue;
+							}
+
+							// This should never occur as the GPU data is considered not ready so low latency is disabled for regular scene captures
+							if (!ensureMsgf(!View->bIsSceneCapture || !ParticleMeshRenderData.bUseGPUScene || !ParticleMeshRenderData.bIsGpuLowLatencyTranslucency, TEXT("Attemping to render translucent mesh particles with low latency into a scenecapture this is not supported.")))
+							{
+								continue;
+							}
 						}
 
 						// When using GPU scene, the indirect draw is managed, so prevent creating indirect draw args
diff --git a/Engine/Plugins/FX/Niagara/Source/Niagara/Private/NiagaraRendererRibbons.cpp b/Engine/Plugins/FX/Niagara/Source/Niagara/Private/NiagaraRendererRibbons.cpp
--- a/Engine/Plugins/FX/Niagara/Source/Niagara/Private/NiagaraRendererRibbons.cpp
+++ b/Engine/Plugins/FX/Niagara/Source/Niagara/Private/NiagaraRendererRibbons.cpp
@@ -4,7 +4,6 @@
 
 #include "GlobalRenderResources.h"
 #include "GPUSortManager.h"
-#include "Materials/Material.h"
 #include "NiagaraRibbonVertexFactory.h"
 #include "NiagaraDataSet.h"
 #include "NiagaraDataSetAccessor.h"
@@ -15,8 +14,10 @@
 #include "NiagaraSystemInstance.h"
 #include "NiagaraComponent.h"
 #include "RayTracingInstance.h"
+#include "Materials/Material.h"
 #include "Materials/MaterialRenderProxy.h"
 #include "MaterialDomain.h"
+#include "MaterialShared.h"
 #include "Math/NumericLimits.h"
 #include "NiagaraCullProxyComponent.h"
 #include "NiagaraGpuComputeDispatchInterface.h"
@@ -851,7 +852,14 @@ void FNiagaraRendererRibbons::GetDynamicMeshElements(const TArray<const FSceneVi
 				// We don't have to generate batches for non-primary views in stereo instance rendering
 				continue;
 			}
-			
+
+			// If we are rendering opaque only we can skip this batch
+			//-OPT: If we only have opaque materials we can skip earlier however due to RemappedMaterialIndex potentially being invalid this is tricky
+			if (IsViewRenderingOpaqueOnly(View) && DynamicData->Material && IsTranslucentBlendMode(DynamicData->Material->GetIncompleteMaterialWithFallback(FeatureLevel)))
+			{
+				continue;
+			}
+
 			FMeshBatch& MeshBatch = Collector.AllocateMesh();
 			
 			const FVector ViewOriginForDistanceCulling = View->ViewMatrices.GetViewOrigin();
diff --git a/Engine/Plugins/FX/Niagara/Source/Niagara/Private/NiagaraRendererSprites.cpp b/Engine/Plugins/FX/Niagara/Source/Niagara/Private/NiagaraRendererSprites.cpp
--- a/Engine/Plugins/FX/Niagara/Source/Niagara/Private/NiagaraRendererSprites.cpp
+++ b/Engine/Plugins/FX/Niagara/Source/Niagara/Private/NiagaraRendererSprites.cpp
@@ -596,12 +596,10 @@ FNiagaraSpriteUniformBufferRef FNiagaraRendererSprites::CreateViewUniformBuffer(
 			{
 				if (IsTranslucentOnlyBlendMode(ParticleSpriteRenderData.BlendMode))
 				{
-					ParticleSpriteRenderData.bHasTranslucentMaterials = true;
 					PerViewUniformParameters.PixelCoverageColorBlend = FVector4f(0.0f, 0.0f, 0.0f, PixelCoverageBlend);
 				}
 				else if (IsAdditiveBlendMode(ParticleSpriteRenderData.BlendMode))
 				{
-					ParticleSpriteRenderData.bHasTranslucentMaterials = true;
 					PerViewUniformParameters.PixelCoverageColorBlend = FVector4f(PixelCoverageBlend, PixelCoverageBlend, PixelCoverageBlend, PixelCoverageBlend);
 				}
 				else
@@ -610,7 +608,7 @@ FNiagaraSpriteUniformBufferRef FNiagaraRendererSprites::CreateViewUniformBuffer(
 					//BLEND_Modulate
 					//BLEND_AlphaComposite
 					//BLEND_AlphaHoldout
-					ParticleSpriteRenderData.bHasTranslucentMaterials = false;
+					PerViewUniformParameters.PixelCoverageEnabled = false;
 				}
 			}
 		}
@@ -980,6 +978,12 @@ void FNiagaraRendererSprites::GetDynamicMeshElements(const TArray<const FSceneVi
 				continue;
 			}
 
+			// Scene captures that are rendered in with the regular views will run depth only so we can skip building the batches
+			if (ParticleSpriteRenderData.bHasTranslucentMaterials && IsViewRenderingOpaqueOnly(View))
+			{
+				continue;
+			}
+
 			if (SourceMode == ENiagaraRendererSourceDataMode::Emitter && bEnableDistanceCulling)
 			{
 				const FVector ViewOrigin = View->ViewMatrices.GetViewOrigin();
diff --git a/Engine/Plugins/FX/Niagara/Source/Niagara/Public/NiagaraRenderer.h b/Engine/Plugins/FX/Niagara/Source/Niagara/Public/NiagaraRenderer.h
--- a/Engine/Plugins/FX/Niagara/Source/Niagara/Public/NiagaraRenderer.h
+++ b/Engine/Plugins/FX/Niagara/Source/Niagara/Public/NiagaraRenderer.h
@@ -92,6 +92,11 @@ public:
 	virtual int32 GetDynamicDataSize()const { return 0; }
 	virtual bool IsMaterialValid(const UMaterialInterface* Mat)const { return Mat != nullptr; }
 
+	// Determine if we are rendering into an opaque only view
+	// This is an optimization for depth only scene capture which render in the main pass.
+	// These captures will have SceneCaptureRenderTarget set on the view, no other capture type does this.
+	static NIAGARA_API bool IsViewRenderingOpaqueOnly(const FSceneView* View);
+
 	static NIAGARA_API void SortIndices(const struct FNiagaraGPUSortInfo& SortInfo, const FNiagaraRendererVariableInfo& SortVariable, const FNiagaraDataBuffer& Buffer, FGlobalDynamicReadBuffer::FAllocation& OutIndices);
 	static NIAGARA_API int32 SortAndCullIndices(const FNiagaraGPUSortInfo& SortInfo, const FNiagaraDataBuffer& Buffer, FGlobalDynamicReadBuffer::FAllocation& OutIndices);
 
