* Added initial parameter support to the audio data interface
* Reworked the existing audio modules a bit with better default values + added a module to initialize parameters

diff --git a/Engine/Plugins/FX/Niagara/Source/Niagara/Classes/NiagaraDataInterfaceAudioPlayer.h b/Engine/Plugins/FX/Niagara/Source/Niagara/Classes/NiagaraDataInterfaceAudioPlayer.h
--- a/Engine/Plugins/FX/Niagara/Source/Niagara/Classes/NiagaraDataInterfaceAudioPlayer.h
+++ b/Engine/Plugins/FX/Niagara/Source/Niagara/Classes/NiagaraDataInterfaceAudioPlayer.h
@@ -2,7 +2,6 @@
 #pragma once
 
 #include "NiagaraCommon.h"
-#include "NiagaraShared.h"
 #include "NiagaraDataInterface.h"
 #include "NiagaraParameterStore.h"
 #include "Sound/SoundAttenuation.h"
@@ -21,6 +20,7 @@ struct FAudioParticleData
 	float Volume = 1;
 	float Pitch = 1;
 	float StartTime = 1;
+	int32 ParticleID = -1;
 };
 
 struct FPersistentAudioParticleData
@@ -31,6 +31,12 @@ struct FPersistentAudioParticleData
 	TFunction<void(struct FAudioPlayerInterface_InstanceData*,UAudioComponent*,FNiagaraSystemInstance*)> UpdateCallback;
 };
 
+struct FAudioInitialParamData
+{
+	int32 ParticleID = -1;
+	FNiagaraVariable Value;
+};
+
 UCLASS(EditInlineNew, Category = "Audio", meta = (DisplayName = "Niagara Audio Player Settings"), Blueprintable, BlueprintType, MinimalAPI)
 class UNiagaraDataInterfaceAudioPlayerSettings : public UObject
 {
@@ -56,8 +62,10 @@ struct FAudioPlayerInterface_InstanceData
 {
 	/** We use a lock-free queue here because multiple threads might try to push data to it at the same time. */
 	TQueue<FAudioParticleData, EQueueMode::Mpsc> PlayAudioQueue;
+	TQueue<FAudioInitialParamData, EQueueMode::Mpsc> InitialParamDataQueue;
 	TQueue<FPersistentAudioParticleData, EQueueMode::Mpsc> PersistentAudioActionQueue;
 	FThreadSafeCounter HandleCount;
+	int32 ParamCountEstimate = 0;
 
 	TSortedMap<int32, TWeakObjectPtr<UAudioComponent>> PersistentAudioMapping;
 
@@ -66,6 +74,7 @@ struct FAudioPlayerInterface_InstanceData
 	TWeakObjectPtr<USoundConcurrency> Concurrency;
 	TWeakObjectPtr<UNiagaraDataInterfaceAudioPlayerSettings> CachedUserParam;
 	TArray<FName> ParameterNames;
+	TSet<FNiagaraVariable> GlobalInitialParameterValues;
 
 	FNiagaraLWCConverter LWCConverter;
 	int32 MaxPlaysPerTick = 0;
@@ -149,7 +158,7 @@ public:
 
 	virtual bool HasPreSimulateTick() const override { return true; }
 	virtual bool HasPostSimulateTick() const override { return true; }
-	virtual bool PostSimulateCanOverlapFrames() const { return false; }
+	virtual bool PostSimulateCanOverlapFrames() const override { return false; }
 	//UNiagaraDataInterface Interface
 
 	NIAGARA_API virtual void PlayOneShotAudio(FVectorVMExternalFunctionContext& Context);
@@ -157,6 +166,9 @@ public:
 	NIAGARA_API virtual void SetParameterBool(FVectorVMExternalFunctionContext& Context);
 	NIAGARA_API virtual void SetParameterInteger(FVectorVMExternalFunctionContext& Context);
 	NIAGARA_API virtual void SetParameterFloat(FVectorVMExternalFunctionContext& Context);
+	NIAGARA_API virtual void SetInitialParameterBool(FVectorVMExternalFunctionContext& Context);
+	NIAGARA_API virtual void SetInitialParameterInteger(FVectorVMExternalFunctionContext& Context);
+	NIAGARA_API virtual void SetInitialParameterFloat(FVectorVMExternalFunctionContext& Context);
 	NIAGARA_API virtual void UpdateVolume(FVectorVMExternalFunctionContext& Context);
 	NIAGARA_API virtual void UpdatePitch(FVectorVMExternalFunctionContext& Context);
 	NIAGARA_API virtual void UpdateLocation(FVectorVMExternalFunctionContext& Context);
@@ -176,5 +188,8 @@ private:
 	static NIAGARA_API const FName SetPersistentAudioBoolParamName;
 	static NIAGARA_API const FName SetPersistentAudioIntegerParamName;
 	static NIAGARA_API const FName SetPersistentAudioFloatParamName;
+	static NIAGARA_API const FName SetInitialAudioBoolParamName;
+	static NIAGARA_API const FName SetInitialAudioIntegerParamName;
+	static NIAGARA_API const FName SetInitialAudioFloatParamName;
 	static NIAGARA_API const FName PausePersistentAudioName;
 };
diff --git a/Engine/Plugins/FX/Niagara/Source/Niagara/Private/NiagaraDataInterfaceAudioPlayer.cpp b/Engine/Plugins/FX/Niagara/Source/Niagara/Private/NiagaraDataInterfaceAudioPlayer.cpp
--- a/Engine/Plugins/FX/Niagara/Source/Niagara/Private/NiagaraDataInterfaceAudioPlayer.cpp
+++ b/Engine/Plugins/FX/Niagara/Source/Niagara/Private/NiagaraDataInterfaceAudioPlayer.cpp
@@ -26,6 +26,9 @@ const FName UNiagaraDataInterfaceAudioPlayer::SetPersistentAudioRotationName(TEX
 const FName UNiagaraDataInterfaceAudioPlayer::SetPersistentAudioBoolParamName(TEXT("SetBooleanParameter"));
 const FName UNiagaraDataInterfaceAudioPlayer::SetPersistentAudioIntegerParamName(TEXT("SetIntegerParameter"));
 const FName UNiagaraDataInterfaceAudioPlayer::SetPersistentAudioFloatParamName(TEXT("SetFloatParameter"));
+const FName UNiagaraDataInterfaceAudioPlayer::SetInitialAudioBoolParamName(TEXT("SetInitialBooleanParameter"));
+const FName UNiagaraDataInterfaceAudioPlayer::SetInitialAudioIntegerParamName(TEXT("SetInitialIntegerParameter"));
+const FName UNiagaraDataInterfaceAudioPlayer::SetInitialAudioFloatParamName(TEXT("SetInitialFloatParameter"));
 const FName UNiagaraDataInterfaceAudioPlayer::PausePersistentAudioName(TEXT("SetPaused"));
 
 struct FNiagaraAudioPlayerDIFunctionVersion
@@ -34,6 +37,7 @@ struct FNiagaraAudioPlayerDIFunctionVersion
 	{
 		InitialVersion = 0,
 		LWCConversion = 1,
+		InitialAudioParameters = 2,
 
 		VersionPlusOne,
 		LatestVersion = VersionPlusOne - 1
@@ -161,9 +165,28 @@ bool UNiagaraDataInterfaceAudioPlayer::PerInstanceTick(void* PerInstanceData, FN
 	return false;
 }
 
+FAudioParameter ConvertVariableToAudioParam(const FNiagaraVariable& NiagaraVariable)
+{
+	FAudioParameter Param;
+	Param.ParamName = NiagaraVariable.GetName();
+	if (NiagaraVariable.GetType() == FNiagaraTypeDefinition::GetBoolDef() && NiagaraVariable.IsDataAllocated())
+	{
+		Param.BoolParam = NiagaraVariable.GetValue<bool>();
+	}
+	else if (NiagaraVariable.GetType() == FNiagaraTypeDefinition::GetIntDef() && NiagaraVariable.IsDataAllocated())
+	{
+		Param.IntParam = NiagaraVariable.GetValue<int32>();
+	}
+	else if (NiagaraVariable.GetType() == FNiagaraTypeDefinition::GetFloatDef() && NiagaraVariable.IsDataAllocated())
+	{
+		Param.FloatParam = NiagaraVariable.GetValue<float>();
+	}
+	return Param;
+}
+
 bool UNiagaraDataInterfaceAudioPlayer::PerInstanceTickPostSimulate(void* PerInstanceData, FNiagaraSystemInstance* SystemInstance, float DeltaSeconds)
 {
-	FAudioPlayerInterface_InstanceData* PIData = (FAudioPlayerInterface_InstanceData*) PerInstanceData;
+	FAudioPlayerInterface_InstanceData* PIData = static_cast<FAudioPlayerInterface_InstanceData*>(PerInstanceData);
 	UNiagaraSystem* System = SystemInstance->GetSystem();
 	UWorld* World = SystemInstance->GetWorldManager()->GetWorld();
 
@@ -171,10 +194,28 @@ bool UNiagaraDataInterfaceAudioPlayer::PerInstanceTickPostSimulate(void* PerInst
 	if (World == nullptr || (World->HasBegunPlay() == false && PIData->bOnlyActiveDuringGameplay))
 	{
 		PIData->PlayAudioQueue.Empty();
+		PIData->InitialParamDataQueue.Empty();
 		PIData->PersistentAudioMapping.Empty();
 		return false;
 	}
 #endif
+
+	TMap<int32, TArray<FAudioInitialParamData>> PerParticleParams;
+	FAudioInitialParamData ParamData;
+	if (PIData->PlayAudioQueue.IsEmpty())
+	{
+		PIData->InitialParamDataQueue.Empty();
+	}
+	else
+	{
+		PerParticleParams.Reserve(PIData->ParamCountEstimate);
+		PIData->ParamCountEstimate = 0;
+	}
+	
+	while (PIData->InitialParamDataQueue.Dequeue(ParamData))
+	{
+		PerParticleParams.FindOrAdd(ParamData.ParticleID).Add(ParamData);
+	}
 	
 	if (!PIData->PlayAudioQueue.IsEmpty() && System)
 	{
@@ -205,8 +246,24 @@ bool UNiagaraDataInterfaceAudioPlayer::PerInstanceTickPostSimulate(void* PerInst
 			}
 			for (const FAudioParticleData& ParticleData : Data)
 			{
-				UGameplayStatics::PlaySoundAtLocation(World, PIData->SoundToPlay.Get(), ParticleData.Position, ParticleData.Rotation, ParticleData.Volume,
-					ParticleData.Pitch, ParticleData.StartTime, PIData->Attenuation.Get(), PIData->Concurrency.Get(), SoundOwner);
+				UInitialActiveSoundParams* InitialParams = nullptr;
+				TArray<FAudioInitialParamData>* ParticleAudioParams = PerParticleParams.Find(ParticleData.ParticleID);
+				if (PIData->GlobalInitialParameterValues.Num() > 0 || ParticleAudioParams)
+				{
+					InitialParams = NewObject<UInitialActiveSoundParams>(); // unfortunately the api needs a temp object as wrapper for the data
+					for (const FNiagaraVariable& Var : PIData->GlobalInitialParameterValues)
+					{
+						InitialParams->AudioParams.Add(ConvertVariableToAudioParam(Var));
+					}
+					if (ParticleAudioParams)
+					{
+						for (const FAudioInitialParamData& Var : *ParticleAudioParams)
+						{
+							InitialParams->AudioParams.Add(ConvertVariableToAudioParam(Var.Value));
+						}
+					}
+				}
+				UGameplayStatics::PlaySoundAtLocation(World, PIData->SoundToPlay.Get(), ParticleData.Position, ParticleData.Rotation, ParticleData.Volume, ParticleData.Pitch, ParticleData.StartTime, PIData->Attenuation.Get(), PIData->Concurrency.Get(), SoundOwner, InitialParams);
 			}
 		}
 	}
@@ -273,7 +330,6 @@ void UNiagaraDataInterfaceAudioPlayer::GetFunctions(TArray<FNiagaraFunctionSigna
 	Sig.FunctionVersion = FNiagaraAudioPlayerDIFunctionVersion::LatestVersion;
 #endif
 	Sig.bMemberFunction = true;
-	Sig.bRequiresContext = false;
 	Sig.bSupportsGPU = false;
 	Sig.bRequiresExecPin = true;
 	Sig.Inputs.Add(FNiagaraVariable(FNiagaraTypeDefinition(GetClass()), TEXT("Audio interface")));
@@ -283,6 +339,7 @@ void UNiagaraDataInterfaceAudioPlayer::GetFunctions(TArray<FNiagaraFunctionSigna
 	Sig.Inputs.Add(FNiagaraVariable(FNiagaraTypeDefinition::GetFloatDef(), TEXT("VolumeFactor")));
 	Sig.Inputs.Add(FNiagaraVariable(FNiagaraTypeDefinition::GetFloatDef(), TEXT("PitchFactor")));
 	Sig.Inputs.Add(FNiagaraVariable(FNiagaraTypeDefinition::GetFloatDef(), TEXT("StartTime")));
+	Sig.AddInput(FNiagaraVariable(FNiagaraTypeDefinition::GetIntDef(), TEXT("ParticleID")), NSLOCTEXT("Niagara", "PlayAudioParticleIDDescription", "This ID parameter is only needed when initial parameter values are used (to match the previously set parameters to the particle pülaying the audio)."));
 	Sig.Outputs.Add(FNiagaraVariable(FNiagaraTypeDefinition::GetBoolDef(), TEXT("Success")));
 	OutFunctions.Add(Sig);
 
@@ -293,7 +350,6 @@ void UNiagaraDataInterfaceAudioPlayer::GetFunctions(TArray<FNiagaraFunctionSigna
 	Sig.FunctionVersion = FNiagaraAudioPlayerDIFunctionVersion::LatestVersion;
 #endif
 	Sig.bMemberFunction = true;
-	Sig.bRequiresContext = false;
 	Sig.bSupportsGPU = false;
 	Sig.bRequiresExecPin = true;
 	Sig.Inputs.Add(FNiagaraVariable(FNiagaraTypeDefinition(GetClass()), TEXT("Audio Interface")));
@@ -316,7 +372,6 @@ void UNiagaraDataInterfaceAudioPlayer::GetFunctions(TArray<FNiagaraFunctionSigna
 	Sig.FunctionVersion = FNiagaraAudioPlayerDIFunctionVersion::LatestVersion;
 #endif
 	Sig.bMemberFunction = true;
-	Sig.bRequiresContext = false;
 	Sig.bSupportsGPU = false;
 	Sig.bRequiresExecPin = true;
 	Sig.Inputs.Add(FNiagaraVariable(FNiagaraTypeDefinition(GetClass()), TEXT("Audio Interface")));
@@ -332,7 +387,6 @@ void UNiagaraDataInterfaceAudioPlayer::GetFunctions(TArray<FNiagaraFunctionSigna
 	Sig.FunctionVersion = FNiagaraAudioPlayerDIFunctionVersion::LatestVersion;
 #endif
 	Sig.bMemberFunction = true;
-	Sig.bRequiresContext = false;
 	Sig.bSupportsGPU = false;
 	Sig.bRequiresExecPin = true;
 	Sig.Inputs.Add(FNiagaraVariable(FNiagaraTypeDefinition(GetClass()), TEXT("Audio Interface")));
@@ -348,7 +402,6 @@ void UNiagaraDataInterfaceAudioPlayer::GetFunctions(TArray<FNiagaraFunctionSigna
 	Sig.FunctionVersion = FNiagaraAudioPlayerDIFunctionVersion::LatestVersion;
 #endif
 	Sig.bMemberFunction = true;
-	Sig.bRequiresContext = false;
 	Sig.bSupportsGPU = false;
 	Sig.bRequiresExecPin = true;
 	Sig.Inputs.Add(FNiagaraVariable(FNiagaraTypeDefinition(GetClass()), TEXT("Audio Interface")));
@@ -357,6 +410,54 @@ void UNiagaraDataInterfaceAudioPlayer::GetFunctions(TArray<FNiagaraFunctionSigna
 	Sig.Inputs.Add(FNiagaraVariable(FNiagaraTypeDefinition::GetFloatDef(), TEXT("Parameter Value")));
 	OutFunctions.Add(Sig);
 
+	Sig = FNiagaraFunctionSignature();
+	Sig.Name = SetInitialAudioBoolParamName;
+#if WITH_EDITORONLY_DATA
+	Sig.Description = NSLOCTEXT("Niagara", "SetInitialAudioBoolParamFunctionDescription", "Sets a bool parameter for one-shot audio invocations.");
+	Sig.FunctionVersion = FNiagaraAudioPlayerDIFunctionVersion::LatestVersion;
+#endif
+	Sig.bMemberFunction = true;
+	Sig.bSupportsGPU = false;
+	Sig.bRequiresExecPin = true;
+	Sig.Inputs.Add(FNiagaraVariable(FNiagaraTypeDefinition(GetClass()), TEXT("Audio Interface")));
+	Sig.AddInput(FNiagaraVariable(FNiagaraTypeDefinition::GetBoolDef(), TEXT("Is Global Parameter")), NSLOCTEXT("Niagara", "SetInitialAudioBoolParamIsGlobalTooltip", "If true then this parameter is used for all one-shot sounds played with this data interface."));
+	Sig.AddInput(FNiagaraVariable(FNiagaraTypeDefinition::GetIntDef(), TEXT("Particle ID")), NSLOCTEXT("Niagara", "SetInitialAudioBoolParamParticleIDTooltip", "Only used when this is not a global parameter."));
+	Sig.Inputs.Add(FNiagaraVariable(FNiagaraTypeDefinition::GetIntDef(), TEXT("Parameter Name Index")));
+	Sig.Inputs.Add(FNiagaraVariable(FNiagaraTypeDefinition::GetBoolDef(), TEXT("Parameter Value")));
+	OutFunctions.Add(Sig);
+
+	Sig = FNiagaraFunctionSignature();
+	Sig.Name = SetInitialAudioIntegerParamName;
+#if WITH_EDITORONLY_DATA
+	Sig.Description = NSLOCTEXT("Niagara", "SetInitialAudioIntegerParamFunctionDescription", "Sets an int parameter for one-shot audio invocations.");
+	Sig.FunctionVersion = FNiagaraAudioPlayerDIFunctionVersion::LatestVersion;
+#endif
+	Sig.bMemberFunction = true;
+	Sig.bSupportsGPU = false;
+	Sig.bRequiresExecPin = true;
+	Sig.Inputs.Add(FNiagaraVariable(FNiagaraTypeDefinition(GetClass()), TEXT("Audio Interface")));
+	Sig.AddInput(FNiagaraVariable(FNiagaraTypeDefinition::GetBoolDef(), TEXT("Is Global Parameter")), NSLOCTEXT("Niagara", "SetInitialAudioIntegerParamIsGlobalTooltip", "If true then this parameter is used for all one-shot sounds played with this data interface."));
+	Sig.AddInput(FNiagaraVariable(FNiagaraTypeDefinition::GetIntDef(), TEXT("Particle ID")), NSLOCTEXT("Niagara", "SetInitialAudioIntegerParamParticleIDTooltip", "Only used when this is not a global parameter."));
+	Sig.Inputs.Add(FNiagaraVariable(FNiagaraTypeDefinition::GetIntDef(), TEXT("Parameter Name Index")));
+	Sig.Inputs.Add(FNiagaraVariable(FNiagaraTypeDefinition::GetIntDef(), TEXT("Parameter Value")));
+	OutFunctions.Add(Sig);
+
+	Sig = FNiagaraFunctionSignature();
+	Sig.Name = SetInitialAudioFloatParamName;
+#if WITH_EDITORONLY_DATA
+	Sig.Description = NSLOCTEXT("Niagara", "SetInitialAudioFloatParamFunctionDescription", "Sets a float parameter for one-shot audio invocations.");
+	Sig.FunctionVersion = FNiagaraAudioPlayerDIFunctionVersion::LatestVersion;
+#endif
+	Sig.bMemberFunction = true;
+	Sig.bSupportsGPU = false;
+	Sig.bRequiresExecPin = true;
+	Sig.Inputs.Add(FNiagaraVariable(FNiagaraTypeDefinition(GetClass()), TEXT("Audio Interface")));
+	Sig.AddInput(FNiagaraVariable(FNiagaraTypeDefinition::GetBoolDef(), TEXT("Is Global Parameter")), NSLOCTEXT("Niagara", "SetInitialAudioFloatParamIsGlobalTooltip", "If true then this parameter is used for all one-shot sounds played with this data interface."));
+	Sig.AddInput(FNiagaraVariable(FNiagaraTypeDefinition::GetIntDef(), TEXT("Particle ID")), NSLOCTEXT("Niagara", "SetInitialAudioFloatParamParticleIDTooltip", "Only used when this is not a global parameter."));
+	Sig.Inputs.Add(FNiagaraVariable(FNiagaraTypeDefinition::GetIntDef(), TEXT("Parameter Name Index")));
+	Sig.Inputs.Add(FNiagaraVariable(FNiagaraTypeDefinition::GetFloatDef(), TEXT("Parameter Value")));
+	OutFunctions.Add(Sig);
+
 	Sig = FNiagaraFunctionSignature();
 	Sig.Name = SetPersistentAudioVolumeName;
 #if WITH_EDITORONLY_DATA
@@ -364,7 +465,6 @@ void UNiagaraDataInterfaceAudioPlayer::GetFunctions(TArray<FNiagaraFunctionSigna
 	Sig.FunctionVersion = FNiagaraAudioPlayerDIFunctionVersion::LatestVersion;
 #endif
 	Sig.bMemberFunction = true;
-	Sig.bRequiresContext = false;
 	Sig.bSupportsGPU = false;
 	Sig.bRequiresExecPin = true;
 	Sig.Inputs.Add(FNiagaraVariable(FNiagaraTypeDefinition(GetClass()), TEXT("Audio Interface")));
@@ -379,7 +479,6 @@ void UNiagaraDataInterfaceAudioPlayer::GetFunctions(TArray<FNiagaraFunctionSigna
 	Sig.FunctionVersion = FNiagaraAudioPlayerDIFunctionVersion::LatestVersion;
 #endif
 	Sig.bMemberFunction = true;
-	Sig.bRequiresContext = false;
 	Sig.bSupportsGPU = false;
 	Sig.bRequiresExecPin = true;
 	Sig.Inputs.Add(FNiagaraVariable(FNiagaraTypeDefinition(GetClass()), TEXT("Audio Interface")));
@@ -394,7 +493,6 @@ void UNiagaraDataInterfaceAudioPlayer::GetFunctions(TArray<FNiagaraFunctionSigna
 	Sig.FunctionVersion = FNiagaraAudioPlayerDIFunctionVersion::LatestVersion;
 #endif
 	Sig.bMemberFunction = true;
-	Sig.bRequiresContext = false;
 	Sig.bSupportsGPU = false;
 	Sig.bRequiresExecPin = true;
 	Sig.Inputs.Add(FNiagaraVariable(FNiagaraTypeDefinition(GetClass()), TEXT("Audio Interface")));
@@ -409,7 +507,6 @@ void UNiagaraDataInterfaceAudioPlayer::GetFunctions(TArray<FNiagaraFunctionSigna
 	Sig.FunctionVersion = FNiagaraAudioPlayerDIFunctionVersion::LatestVersion;
 #endif
 	Sig.bMemberFunction = true;
-	Sig.bRequiresContext = false;
 	Sig.bSupportsGPU = false;
 	Sig.bRequiresExecPin = true;
 	Sig.Inputs.Add(FNiagaraVariable(FNiagaraTypeDefinition(GetClass()), TEXT("Audio Interface")));
@@ -424,7 +521,6 @@ void UNiagaraDataInterfaceAudioPlayer::GetFunctions(TArray<FNiagaraFunctionSigna
 	Sig.FunctionVersion = FNiagaraAudioPlayerDIFunctionVersion::LatestVersion;
 #endif
 	Sig.bMemberFunction = true;
-	Sig.bRequiresContext = false;
 	Sig.bSupportsGPU = false;
 	Sig.bRequiresExecPin = true;
 	Sig.Inputs.Add(FNiagaraVariable(FNiagaraTypeDefinition(GetClass()), TEXT("Audio Interface")));
@@ -438,6 +534,9 @@ DEFINE_NDI_DIRECT_FUNC_BINDER(UNiagaraDataInterfaceAudioPlayer, PlayPersistentAu
 DEFINE_NDI_DIRECT_FUNC_BINDER(UNiagaraDataInterfaceAudioPlayer, SetParameterBool);
 DEFINE_NDI_DIRECT_FUNC_BINDER(UNiagaraDataInterfaceAudioPlayer, SetParameterInteger);
 DEFINE_NDI_DIRECT_FUNC_BINDER(UNiagaraDataInterfaceAudioPlayer, SetParameterFloat);
+DEFINE_NDI_DIRECT_FUNC_BINDER(UNiagaraDataInterfaceAudioPlayer, SetInitialParameterBool);
+DEFINE_NDI_DIRECT_FUNC_BINDER(UNiagaraDataInterfaceAudioPlayer, SetInitialParameterInteger);
+DEFINE_NDI_DIRECT_FUNC_BINDER(UNiagaraDataInterfaceAudioPlayer, SetInitialParameterFloat);
 DEFINE_NDI_DIRECT_FUNC_BINDER(UNiagaraDataInterfaceAudioPlayer, UpdateVolume);
 DEFINE_NDI_DIRECT_FUNC_BINDER(UNiagaraDataInterfaceAudioPlayer, UpdatePitch);
 DEFINE_NDI_DIRECT_FUNC_BINDER(UNiagaraDataInterfaceAudioPlayer, UpdateLocation);
@@ -465,6 +564,18 @@ void UNiagaraDataInterfaceAudioPlayer::GetVMExternalFunction(const FVMExternalFu
 	{
 		NDI_FUNC_BINDER(UNiagaraDataInterfaceAudioPlayer, SetParameterFloat)::Bind(this, OutFunc);
 	}
+	else if (BindingInfo.Name == SetInitialAudioBoolParamName)
+	{
+		NDI_FUNC_BINDER(UNiagaraDataInterfaceAudioPlayer, SetInitialParameterBool)::Bind(this, OutFunc);
+	}
+	else if (BindingInfo.Name == SetInitialAudioIntegerParamName)
+	{
+		NDI_FUNC_BINDER(UNiagaraDataInterfaceAudioPlayer, SetInitialParameterInteger)::Bind(this, OutFunc);
+	}
+	else if (BindingInfo.Name == SetInitialAudioFloatParamName)
+	{
+		NDI_FUNC_BINDER(UNiagaraDataInterfaceAudioPlayer, SetInitialParameterFloat)::Bind(this, OutFunc);
+	}
 	else if (BindingInfo.Name == SetPersistentAudioVolumeName)
 	{
 		NDI_FUNC_BINDER(UNiagaraDataInterfaceAudioPlayer, UpdateVolume)::Bind(this, OutFunc);
@@ -498,7 +609,7 @@ void UNiagaraDataInterfaceAudioPlayer::SetParameterBool(FVectorVMExternalFunctio
 	FNDIInputParam<int32> AudioHandleInParam(Context);
 	FNDIInputParam<int32> NameIndexParam(Context);
 	FNDIInputParam<FNiagaraBool> ValueParam(Context);
-	checkfSlow(InstData.Get(), TEXT("Audio player interface has invalid instance data. %s"), *GetPathName());
+	checkf(InstData.Get(), TEXT("Audio player interface has invalid instance data. %s"), *GetPathName());
 
 	for (int32 i = 0; i < Context.GetNumInstances(); ++i)
 	{
@@ -530,7 +641,7 @@ void UNiagaraDataInterfaceAudioPlayer::SetParameterInteger(FVectorVMExternalFunc
 	FNDIInputParam<int32> AudioHandleInParam(Context);
 	FNDIInputParam<int32> NameIndexParam(Context);
 	FNDIInputParam<int32> ValueParam(Context);
-	checkfSlow(InstData.Get(), TEXT("Audio player interface has invalid instance data. %s"), *GetPathName());
+	checkf(InstData.Get(), TEXT("Audio player interface has invalid instance data. %s"), *GetPathName());
 
 	for (int32 i = 0; i < Context.GetNumInstances(); ++i)
 	{
@@ -562,7 +673,7 @@ void UNiagaraDataInterfaceAudioPlayer::SetParameterFloat(FVectorVMExternalFuncti
 	FNDIInputParam<int32> AudioHandleInParam(Context);
 	FNDIInputParam<int32> NameIndexParam(Context);
 	FNDIInputParam<float> ValueParam(Context);
-	checkfSlow(InstData.Get(), TEXT("Audio player interface has invalid instance data. %s"), *GetPathName());
+	checkf(InstData.Get(), TEXT("Audio player interface has invalid instance data. %s"), *GetPathName());
 
 	for (int32 i = 0; i < Context.GetNumInstances(); ++i)
 	{
@@ -587,13 +698,118 @@ void UNiagaraDataInterfaceAudioPlayer::SetParameterFloat(FVectorVMExternalFuncti
 	}
 }
 
+void UNiagaraDataInterfaceAudioPlayer::SetInitialParameterBool(FVectorVMExternalFunctionContext& Context)
+{
+	VectorVM::FUserPtrHandler<FAudioPlayerInterface_InstanceData> InstData(Context);
+
+	FNDIInputParam<FNiagaraBool> IsGlobalParam(Context);
+	FNDIInputParam<int32> ParticleIDParam(Context);
+	FNDIInputParam<int32> NameIndexParam(Context);
+	FNDIInputParam<FNiagaraBool> ValueParam(Context);
+	checkf(InstData.Get(), TEXT("Audio player interface has invalid instance data. %s"), *GetPathName());
+
+	for (int32 i = 0; i < Context.GetNumInstances(); ++i)
+	{
+		bool IsGlobal = IsGlobalParam.GetAndAdvance();
+		int32 ParticleID = ParticleIDParam.GetAndAdvance();
+		int32 NameIndex = NameIndexParam.GetAndAdvance();
+		bool Value = ValueParam.GetAndAdvance();
+
+		if (InstData->ParameterNames.IsValidIndex(NameIndex))
+		{
+			FNiagaraVariable Parameter(FNiagaraTypeDefinition::GetBoolDef(), InstData->ParameterNames[NameIndex]);
+			Parameter.SetValue(Value);
+
+			if (IsGlobal)
+			{
+				InstData->GlobalInitialParameterValues.Add(Parameter);
+			}
+			else if (ParticleID >= 0)
+			{
+				InstData->InitialParamDataQueue.Enqueue({ ParticleID, Parameter });
+			}
+		}
+	}
+	InstData->ParamCountEstimate = Context.GetNumInstances();
+}
+
+void UNiagaraDataInterfaceAudioPlayer::SetInitialParameterInteger(FVectorVMExternalFunctionContext& Context)
+{
+	VectorVM::FUserPtrHandler<FAudioPlayerInterface_InstanceData> InstData(Context);
+
+	FNDIInputParam<FNiagaraBool> IsGlobalParam(Context);
+	FNDIInputParam<int32> ParticleIDParam(Context);
+	FNDIInputParam<int32> NameIndexParam(Context);
+	FNDIInputParam<int32> ValueParam(Context);
+	checkf(InstData.Get(), TEXT("Audio player interface has invalid instance data. %s"), *GetPathName());
+
+	for (int32 i = 0; i < Context.GetNumInstances(); ++i)
+	{
+		bool IsGlobal = IsGlobalParam.GetAndAdvance();
+		int32 ParticleID = ParticleIDParam.GetAndAdvance();
+		int32 NameIndex = NameIndexParam.GetAndAdvance();
+		int32 Value = ValueParam.GetAndAdvance();
+
+		if (InstData->ParameterNames.IsValidIndex(NameIndex))
+		{
+			FNiagaraVariable Parameter(FNiagaraTypeDefinition::GetIntDef(), InstData->ParameterNames[NameIndex]);
+			Parameter.SetValue(Value);
+
+			if (IsGlobal)
+			{
+				InstData->GlobalInitialParameterValues.Add(Parameter);
+			}
+			else if (ParticleID >= 0)
+			{
+				InstData->InitialParamDataQueue.Enqueue({ ParticleID, Parameter });
+			}
+		}
+	}
+	InstData->ParamCountEstimate = Context.GetNumInstances();
+}
+
+void UNiagaraDataInterfaceAudioPlayer::SetInitialParameterFloat(FVectorVMExternalFunctionContext& Context)
+{
+	VectorVM::FUserPtrHandler<FAudioPlayerInterface_InstanceData> InstData(Context);
+
+	FNDIInputParam<FNiagaraBool> IsGlobalParam(Context);
+	FNDIInputParam<int32> ParticleIDParam(Context);
+	FNDIInputParam<int32> NameIndexParam(Context);
+	FNDIInputParam<float> ValueParam(Context);
+	checkf(InstData.Get(), TEXT("Audio player interface has invalid instance data. %s"), *GetPathName());
+
+	for (int32 i = 0; i < Context.GetNumInstances(); ++i)
+	{
+		bool IsGlobal = IsGlobalParam.GetAndAdvance();
+		int32 ParticleID = ParticleIDParam.GetAndAdvance();
+		int32 NameIndex = NameIndexParam.GetAndAdvance();
+		float Value = ValueParam.GetAndAdvance();
+
+		if (InstData->ParameterNames.IsValidIndex(NameIndex))
+		{
+			FNiagaraVariable Parameter(FNiagaraTypeDefinition::GetFloatDef(), InstData->ParameterNames[NameIndex]);
+			Parameter.SetValue(Value);
+
+			if (IsGlobal)
+			{
+				InstData->GlobalInitialParameterValues.Add(Parameter);
+			}
+			else if (ParticleID >= 0)
+			{
+				InstData->InitialParamDataQueue.Enqueue({ ParticleID, Parameter });
+			}
+		}
+	}
+	InstData->ParamCountEstimate = Context.GetNumInstances();
+}
+
 void UNiagaraDataInterfaceAudioPlayer::UpdateVolume(FVectorVMExternalFunctionContext& Context)
 {
 	VectorVM::FUserPtrHandler<FAudioPlayerInterface_InstanceData> InstData(Context);
 
 	FNDIInputParam<int32> AudioHandleInParam(Context);
 	FNDIInputParam<float> VolumeParam(Context);
-	checkfSlow(InstData.Get(), TEXT("Audio player interface has invalid instance data. %s"), *GetPathName());
+	checkf(InstData.Get(), TEXT("Audio player interface has invalid instance data. %s"), *GetPathName());
 
 	for (int32 i = 0; i < Context.GetNumInstances(); ++i)
 	{
@@ -622,7 +838,7 @@ void UNiagaraDataInterfaceAudioPlayer::UpdatePitch(FVectorVMExternalFunctionCont
 
 	FNDIInputParam<int32> AudioHandleInParam(Context);
 	FNDIInputParam<float> PitchParam(Context);
-	checkfSlow(InstData.Get(), TEXT("Audio player interface has invalid instance data. %s"), *GetPathName());
+	checkf(InstData.Get(), TEXT("Audio player interface has invalid instance data. %s"), *GetPathName());
 
 	for (int32 i = 0; i < Context.GetNumInstances(); ++i)
 	{
@@ -651,7 +867,7 @@ void UNiagaraDataInterfaceAudioPlayer::UpdateLocation(FVectorVMExternalFunctionC
 
 	FNDIInputParam<int32> AudioHandleInParam(Context);
 	FNDIInputParam<FVector3f> LocationParam(Context);
-	checkfSlow(InstData.Get(), TEXT("Audio player interface has invalid instance data. %s"), *GetPathName());
+	checkf(InstData.Get(), TEXT("Audio player interface has invalid instance data. %s"), *GetPathName());
 
 	for (int32 i = 0; i < Context.GetNumInstances(); ++i)
 	{
@@ -680,7 +896,7 @@ void UNiagaraDataInterfaceAudioPlayer::UpdateRotation(FVectorVMExternalFunctionC
 
 	FNDIInputParam<int32> AudioHandleInParam(Context);
 	FNDIInputParam<FVector3f> RotationParam(Context);
-	checkfSlow(InstData.Get(), TEXT("Audio player interface has invalid instance data. %s"), *GetPathName());
+	checkf(InstData.Get(), TEXT("Audio player interface has invalid instance data. %s"), *GetPathName());
 
 	for (int32 i = 0; i < Context.GetNumInstances(); ++i)
 	{
@@ -710,7 +926,7 @@ void UNiagaraDataInterfaceAudioPlayer::SetPausedState(FVectorVMExternalFunctionC
 
 	FNDIInputParam<int32> AudioHandleInParam(Context);
 	FNDIInputParam<FNiagaraBool> PausedParam(Context);
-	checkfSlow(InstData.Get(), TEXT("Audio player interface has invalid instance data. %s"), *GetPathName());
+	checkf(InstData.Get(), TEXT("Audio player interface has invalid instance data. %s"), *GetPathName());
 
 	for (int32 i = 0; i < Context.GetNumInstances(); ++i)
 	{
@@ -737,23 +953,21 @@ void UNiagaraDataInterfaceAudioPlayer::PlayOneShotAudio(FVectorVMExternalFunctio
 {
 	VectorVM::FUserPtrHandler<FAudioPlayerInterface_InstanceData> InstData(Context);
 
-	VectorVM::FExternalFuncInputHandler<FNiagaraBool> PlayDataParam(Context);
-
-	VectorVM::FExternalFuncInputHandler<float> PositionParamX(Context);
-	VectorVM::FExternalFuncInputHandler<float> PositionParamY(Context);
-	VectorVM::FExternalFuncInputHandler<float> PositionParamZ(Context);
+	FNDIInputParam<FNiagaraBool> PlayDataParam(Context);
+	FNDIInputParam<FNiagaraPosition> PositionParam(Context);
 	
-	VectorVM::FExternalFuncInputHandler<float> RotationParamX(Context);
-	VectorVM::FExternalFuncInputHandler<float> RotationParamY(Context);
-	VectorVM::FExternalFuncInputHandler<float> RotationParamZ(Context);
+	FNDIInputParam<float> RotationParamX(Context);
+	FNDIInputParam<float> RotationParamY(Context);
+	FNDIInputParam<float> RotationParamZ(Context);
 	
-	VectorVM::FExternalFuncInputHandler<float> VolumeParam(Context);
-	VectorVM::FExternalFuncInputHandler<float> PitchParam(Context);
-	VectorVM::FExternalFuncInputHandler<float> StartTimeParam(Context);
+	FNDIInputParam<float> VolumeParam(Context);
+	FNDIInputParam<float> PitchParam(Context);
+	FNDIInputParam<float> StartTimeParam(Context);
+	FNDIInputParam<int32> ParticleIDParam(Context);
 
-	VectorVM::FExternalFuncRegisterHandler<FNiagaraBool> OutSample(Context);
+	FNDIOutputParam<FNiagaraBool> OutSample(Context);
 
-	checkfSlow(InstData.Get(), TEXT("Audio player interface has invalid instance data. %s"), *GetPathName());
+	checkf(InstData.Get(), TEXT("Audio player interface has invalid instance data. %s"), *GetPathName());
 	bool ValidSoundData = InstData->SoundToPlay.IsValid() && InstData->bValidOneShotSound;
 
 #if WITH_EDITOR
@@ -767,19 +981,20 @@ void UNiagaraDataInterfaceAudioPlayer::PlayOneShotAudio(FVectorVMExternalFunctio
 	{
 		FNiagaraBool ShouldPlay = PlayDataParam.GetAndAdvance();
 		FAudioParticleData Data;
-		FNiagaraPosition SimulationPosition(PositionParamX.GetAndAdvance(), PositionParamY.GetAndAdvance(), PositionParamZ.GetAndAdvance());
+		FNiagaraPosition SimulationPosition = PositionParam.GetAndAdvance();
 		Data.Position = InstData->LWCConverter.ConvertSimulationPositionToWorld(SimulationPosition);
 		Data.Rotation = FRotator(RotationParamX.GetAndAdvance(), RotationParamY.GetAndAdvance(), RotationParamZ.GetAndAdvance());
 		Data.Volume = VolumeParam.GetAndAdvance();
 		Data.Pitch = PitchParam.GetAndAdvance();
 		Data.StartTime = StartTimeParam.GetAndAdvance();
+		Data.ParticleID = ParticleIDParam.GetAndAdvance();
 
 		FNiagaraBool Valid;
 		if (ValidSoundData && ShouldPlay)
 		{
 			Valid.SetValue(InstData->PlayAudioQueue.Enqueue(Data));
 		}
-		*OutSample.GetDestAndAdvance() = Valid;
+		OutSample.SetAndAdvance(Valid);
 	}
 }
 
@@ -799,7 +1014,7 @@ void UNiagaraDataInterfaceAudioPlayer::PlayPersistentAudio(FVectorVMExternalFunc
 
 	FNDIOutputParam<int32> AudioHandleOutParam(Context);
 
-	checkfSlow(InstData.Get(), TEXT("Audio player interface has invalid instance data. %s"), *GetPathName());
+	checkf(InstData.Get(), TEXT("Audio player interface has invalid instance data. %s"), *GetPathName());
 
 	for (int32 i = 0; i < Context.GetNumInstances(); ++i)
 	{
