Adding ability for NDCs to search by component and a separate location override.
Cleaned up the api a bit.

diff --git a/Engine/Plugins/FX/Niagara/Source/Niagara/Private/DataInterface/NiagaraDataInterfaceDataChannelRead.cpp b/Engine/Plugins/FX/Niagara/Source/Niagara/Private/DataInterface/NiagaraDataInterfaceDataChannelRead.cpp
--- a/Engine/Plugins/FX/Niagara/Source/Niagara/Private/DataInterface/NiagaraDataInterfaceDataChannelRead.cpp
+++ b/Engine/Plugins/FX/Niagara/Source/Niagara/Private/DataInterface/NiagaraDataInterfaceDataChannelRead.cpp
@@ -347,8 +347,7 @@ bool FNDIDataChannelReadInstanceData::Tick(UNiagaraDataInterfaceDataChannelRead*
 			if(DataChannelData == nullptr || Interface->bUpdateSourceDataEveryTick)
 			{
 				//TODO: Automatically modify tick group if we have DIs that require current frame info?
-				FNiagaraDataChannelSearchParameters SearchParams;
-				SearchParams.OwningComponent = Instance->GetAttachComponent();
+				FNiagaraDataChannelSearchParameters SearchParams(Instance->GetAttachComponent());
 				DataChannelData = DataChannelPtr->FindData(SearchParams, ENiagaraResourceAccess::ReadOnly);//TODO: Maybe should have two paths, one for system instances and another for SceneComponents...
 			}	
 
diff --git a/Engine/Plugins/FX/Niagara/Source/Niagara/Private/DataInterface/NiagaraDataInterfaceDataChannelWrite.cpp b/Engine/Plugins/FX/Niagara/Source/Niagara/Private/DataInterface/NiagaraDataInterfaceDataChannelWrite.cpp
--- a/Engine/Plugins/FX/Niagara/Source/Niagara/Private/DataInterface/NiagaraDataInterfaceDataChannelWrite.cpp
+++ b/Engine/Plugins/FX/Niagara/Source/Niagara/Private/DataInterface/NiagaraDataInterfaceDataChannelWrite.cpp
@@ -212,8 +212,7 @@ struct FNDIDataChannelWriteInstanceData
 			{
 				if (DataChannelData == nullptr || Interface->bUpdateDestinationDataEveryTick)
 				{
-					FNiagaraDataChannelSearchParameters SearchParams;
-					SearchParams.OwningComponent = Instance->GetAttachComponent();
+					FNiagaraDataChannelSearchParameters SearchParams(Instance->GetAttachComponent());
 					DataChannelData = DataChannelPtr->FindData(SearchParams, ENiagaraResourceAccess::WriteOnly);
 				}
 			}
diff --git a/Engine/Plugins/FX/Niagara/Source/Niagara/Private/NiagaraDataChannelCommon.cpp b/Engine/Plugins/FX/Niagara/Source/Niagara/Private/NiagaraDataChannelCommon.cpp
--- a/Engine/Plugins/FX/Niagara/Source/Niagara/Private/NiagaraDataChannelCommon.cpp
+++ b/Engine/Plugins/FX/Niagara/Source/Niagara/Private/NiagaraDataChannelCommon.cpp
@@ -5,5 +5,5 @@
 
 FVector FNiagaraDataChannelSearchParameters::GetLocation()const
 {
-	return OwningComponent ? OwningComponent->GetComponentLocation() : Location;
+	return OwningComponent && !bOverrideLocation ? OwningComponent->GetComponentLocation() : Location;
 }
\ No newline at end of file
diff --git a/Engine/Plugins/FX/Niagara/Source/Niagara/Private/NiagaraDataChannel_Islands.cpp b/Engine/Plugins/FX/Niagara/Source/Niagara/Private/NiagaraDataChannel_Islands.cpp
--- a/Engine/Plugins/FX/Niagara/Source/Niagara/Private/NiagaraDataChannel_Islands.cpp
+++ b/Engine/Plugins/FX/Niagara/Source/Niagara/Private/NiagaraDataChannel_Islands.cpp
@@ -188,12 +188,12 @@ FNDCIsland* UNiagaraDataChannelHandler_Islands::FindOrCreateIsland(const FNiagar
 		//When writing we'll grow/spawn islands to accommodate the data.
 
 		//First we see if this is an islands handler system.
-		if(SearchParams.OwningComponent)
+		if(SearchParams.GetOwner())
 		{
 			for (int32 i : ActiveIslands)
 			{
 				FNDCIsland& Island = IslandPool[i];
-				if (Island.IsHandlerSystem(SearchParams.OwningComponent))
+				if (Island.IsHandlerSystem(SearchParams.GetOwner()))
 				{
 					return &Island;
 				}
@@ -202,7 +202,7 @@ FNDCIsland* UNiagaraDataChannelHandler_Islands::FindOrCreateIsland(const FNiagar
 			for (int32 i : ActiveIslands)
 			{
 				FNDCIsland& Island = IslandPool[i];
-				if (Island.Intersects(SearchParams.OwningComponent->Bounds))
+				if (Island.Intersects(SearchParams.GetOwner()->Bounds))
 				{
 					return &Island;
 				}
diff --git a/Engine/Plugins/FX/Niagara/Source/Niagara/Public/NiagaraDataChannelPublic.h b/Engine/Plugins/FX/Niagara/Source/Niagara/Public/NiagaraDataChannelPublic.h
--- a/Engine/Plugins/FX/Niagara/Source/Niagara/Public/NiagaraDataChannelPublic.h
+++ b/Engine/Plugins/FX/Niagara/Source/Niagara/Public/NiagaraDataChannelPublic.h
@@ -45,24 +45,55 @@ Minimal set of types and declares required for external users of Niagara Data Ch
 */
 
 /**
-Parameters allowing users to search for the correct data channel data to read/write.
-Some data channels will sub divide their data internally in various ways, e.g., spacial partition.
-These parameters allow users to search for the correct internal data when reading and writing.
+Parameters used when retrieving a specific set of Data Channel Data to read or write.
+Many Data Channel types will have multiple internal sets of data and these parameters control which the Channel should return to users for access.
+An example of this would be the Islands Data Channel type which will subdivide the world and have a different set of data for each sub division.
+It will return to users the correct data for their location based on these parameters.
 */
 USTRUCT(BlueprintType)
 struct FNiagaraDataChannelSearchParameters
 {
 	GENERATED_BODY()
+	
+	FNiagaraDataChannelSearchParameters()
+	: bOverrideLocation(false)
+	{
+	}
+
+	FNiagaraDataChannelSearchParameters(USceneComponent* Owner)
+	: OwningComponent(Owner)
+	, bOverrideLocation(false)
+	{
+	}
+
+	FNiagaraDataChannelSearchParameters(USceneComponent* Owner, FVector LocationOverride)
+	: OwningComponent(Owner)
+	, Location(LocationOverride)
+	, bOverrideLocation(true)
+	{
+	}
+
+	FNiagaraDataChannelSearchParameters(FVector InLocation)
+		: OwningComponent(nullptr)
+		, Location(InLocation)
+		, bOverrideLocation(true)
+	{
+	}
+	
+	NIAGARA_API FVector GetLocation()const;
+	NIAGARA_API USceneComponent* GetOwner()const { return OwningComponent; }
 
 	/** In cases where there is an owning component such as an object spawning from itself etc, then we pass that component in. Some handlers may only use it's location but others may make use of more data. */
 	UPROPERTY(EditAnywhere, BlueprintReadWrite, Category = "Parameters")
 	TObjectPtr<USceneComponent> OwningComponent = nullptr;
 
-	/** In cases where there is no owning component for data being read or written to a data channel, we simply pass in a location. */
+	/** In cases where there is no owning component for data being read or written to a data channel, we simply pass in a location. We can also use this when bOverrideLocaiton is set. */
 	UPROPERTY(EditAnywhere, BlueprintReadWrite, Category = "Parameters")
 	FVector Location = FVector::ZeroVector;
 
-	NIAGARA_API FVector GetLocation()const;
+	/** If true, even if an owning component is set, the data channel should use the Location value rather than the component location. If this is false, the NDC will get any location needed from the owning component. */
+	UPROPERTY(EditAnywhere, BlueprintReadWrite, Category = "Parameters")
+	uint32 bOverrideLocation : 1;
 };
 
 USTRUCT()
diff --git a/Engine/Plugins/FX/Niagara/Source/Niagara/Public/NiagaraDataChannel_Islands.h b/Engine/Plugins/FX/Niagara/Source/Niagara/Public/NiagaraDataChannel_Islands.h
--- a/Engine/Plugins/FX/Niagara/Source/Niagara/Public/NiagaraDataChannel_Islands.h
+++ b/Engine/Plugins/FX/Niagara/Source/Niagara/Public/NiagaraDataChannel_Islands.h
@@ -126,11 +126,11 @@ public:
 	FVector GetMaxExtents()const { return MaxExtents; }
 	FVector GetPerElementExtents()const { return PerElementExtents; }
 
-	TConstArrayView<TObjectPtr<UNiagaraSystem>> GetSystems()const;
+	NIAGARA_API TConstArrayView<TObjectPtr<UNiagaraSystem>> GetSystems()const;
 	int32 GetIslandPoolSize()const { return IslandPoolSize; }
 	const FNDCIslandDebugDrawSettings& GetDebugDrawSettings()const { return DebugDrawSettings; }
 
-private:
+protected:
 
 	/** Controls how islands are placed and sized. */
 	UPROPERTY(EditAnywhere, Category = "Islands")
@@ -197,7 +197,7 @@ class UNiagaraDataChannelHandler_Islands : public UNiagaraDataChannelHandler
 	NIAGARA_API virtual void Tick(float DeltaTime, ETickingGroup TickGroup, FNiagaraWorldManager* OwningWorld) override;
 	NIAGARA_API virtual FNiagaraDataChannelDataPtr FindData(FNiagaraDataChannelSearchParameters SearchParams, ENiagaraResourceAccess AccessType) override;
 
-private:
+protected:
 
 	/** All currently active Islands for this channel. */
 	UPROPERTY()
