- Fix missing nullptr checks for component being valid

diff --git a/Engine/Plugins/FX/Niagara/Source/Niagara/Private/DataInterface/NiagaraDataInterfaceStaticMesh.cpp b/Engine/Plugins/FX/Niagara/Source/Niagara/Private/DataInterface/NiagaraDataInterfaceStaticMesh.cpp
--- a/Engine/Plugins/FX/Niagara/Source/Niagara/Private/DataInterface/NiagaraDataInterfaceStaticMesh.cpp
+++ b/Engine/Plugins/FX/Niagara/Source/Niagara/Private/DataInterface/NiagaraDataInterfaceStaticMesh.cpp
@@ -3511,7 +3511,10 @@ UStaticMesh* UNiagaraDataInterfaceStaticMesh::GetStaticMesh(USceneComponent*& Ou
 	else if (bTrySource && SoftSourceActor.Get())
 	{
 		OutStaticMeshComponent = FindActorMeshComponent(SoftSourceActor.Get());
-		OutMesh = OutStaticMeshComponent->GetStaticMesh();
+		if (OutStaticMeshComponent)
+		{
+			OutMesh = OutStaticMeshComponent->GetStaticMesh();
+		}
 	}
 	else
 	{
@@ -3521,7 +3524,10 @@ UStaticMesh* UNiagaraDataInterfaceStaticMesh::GetStaticMesh(USceneComponent*& Ou
 			if (AActor* Actor = Cast<AActor>(ParameterBindingValue))
 			{
 				OutStaticMeshComponent = FindActorMeshComponent(Actor, true);
-				OutMesh = OutStaticMeshComponent->GetStaticMesh();
+				if (OutStaticMeshComponent)
+				{
+					OutMesh = OutStaticMeshComponent->GetStaticMesh();
+				}
 			}
 			else
 			{
