Add a GPU Scene instance data flag for "Force Hidden" and use that instead of an invalid primitive ID to mark instances for cull.
This fixes issues with mobile GPU scene culling, which needs to access a valid primitive ID from forcibly culled instances' data.

diff --git a/Engine/Plugins/FX/Niagara/Shaders/Private/NiagaraUpdateMeshGPUSceneInstances.usf b/Engine/Plugins/FX/Niagara/Shaders/Private/NiagaraUpdateMeshGPUSceneInstances.usf
--- a/Engine/Plugins/FX/Niagara/Shaders/Private/NiagaraUpdateMeshGPUSceneInstances.usf
+++ b/Engine/Plugins/FX/Niagara/Shaders/Private/NiagaraUpdateMeshGPUSceneInstances.usf
@@ -74,8 +74,7 @@ uint GetPayloadDataFlags()
 void CullInstance(uint InstanceId)
 {
 	// This will minimally initialize the instance data such that it will be ignored and culled
-	// NOTE: Must use flags (instead of zero) to prevent the element getting flushed to zero when it happens to be denormal
-	WriteInstancePrimitiveIdAndFlags(InstanceId, INVALID_PRIMITIVE_ID, GetPayloadDataFlags());
+	WriteInstancePrimitiveIdAndFlags(InstanceId, PrimitiveId, INSTANCE_SCENE_DATA_FLAG_HIDDEN);
 }
 
 [numthreads(THREAD_GROUP_SIZE, 1, 1)]
@@ -123,11 +122,11 @@ void UpdateMeshInstancesCS(uint3 GroupId : SV_GroupID, uint GroupIndex : SV_Grou
 	Params.ParticleIndex				= ParticleIndex;
 	Params.SystemLWCTile				= SystemLWCTile;
 	Params.bLocalSpace					= bLocalSpace;
-	#if ENABLE_PRECISE_MOTION_VECTORS
-		Params.bPreciseMotionVectors 	= true;
-	#else
-		Params.bPreciseMotionVectors 	= false;		
-	#endif
+#if ENABLE_PRECISE_MOTION_VECTORS
+	Params.bPreciseMotionVectors 		= true;
+#else
+	Params.bPreciseMotionVectors 		= false;
+#endif
 	Params.FacingMode					= FacingMode;
 	Params.DeltaSeconds					= DeltaSeconds;
 	Params.MeshScale 					= MeshScale;
@@ -167,7 +166,7 @@ void UpdateMeshInstancesCS(uint3 GroupId : SV_GroupID, uint GroupIndex : SV_Grou
 	Params.PrimitiveInvNonUniformScale	= PrimitiveData.InvNonUniformScale;
 
 	NiagaraMeshParticleTransforms Transforms = NiagaraCalculateMeshParticleTransforms(Params);
-	
+
 	// Transform the bounding center to world-space and distance cull
 	{
 		const FLWCVector3 BCenter = LWCMultiply(LocalBoundingCenter, Transforms.LocalToWorld);
diff --git a/Engine/Shaders/Private/SceneData.ush b/Engine/Shaders/Private/SceneData.ush
--- a/Engine/Shaders/Private/SceneData.ush
+++ b/Engine/Shaders/Private/SceneData.ush
@@ -922,7 +922,8 @@ FInstanceSceneData GetInstanceSceneData(uint InstanceId, uint SOAStride, bool bC
 
 	// Only process valid instances
 	LoadInstancePrimitiveIdAndFlags(InstanceId, SOAStride, InstanceData.PrimitiveId, InstanceData.Flags);
-	InstanceData.ValidInstance = InstanceData.PrimitiveId != INVALID_PRIMITIVE_ID;
+	InstanceData.ValidInstance = InstanceData.PrimitiveId != INVALID_PRIMITIVE_ID &&
+								 ((InstanceData.Flags & INSTANCE_SCENE_DATA_FLAG_HIDDEN) == 0u);
 
 	// Payload Data Layout
 	// NOTE: Per-instance local bounds and hierarchy offset are always mutually inclusive, so pack together.
diff --git a/Engine/Shaders/Shared/SceneDefinitions.h b/Engine/Shaders/Shared/SceneDefinitions.h
--- a/Engine/Shaders/Shared/SceneDefinitions.h
+++ b/Engine/Shaders/Shared/SceneDefinitions.h
@@ -88,6 +88,7 @@
 #define INSTANCE_SCENE_DATA_FLAG_HAS_PAYLOAD_EXTENSION		0x80
 #define INSTANCE_SCENE_DATA_FLAG_HAS_EDITOR_DATA			0x100
 #define INSTANCE_SCENE_DATA_FLAG_IS_RAYTRACING_FAR_FIELD	0x200
+#define INSTANCE_SCENE_DATA_FLAG_HIDDEN						0x400
 
 #define INSTANCE_SCENE_DATA_FLAG_PAYLOAD_MASK (			\
 	INSTANCE_SCENE_DATA_FLAG_HAS_RANDOM					\
