Move RapidIterationParameter::DebugName to transient to prevent stale debug names

diff --git a/Engine/Plugins/FX/Niagara/Source/Niagara/Private/NiagaraComponent.cpp b/Engine/Plugins/FX/Niagara/Source/Niagara/Private/NiagaraComponent.cpp
--- a/Engine/Plugins/FX/Niagara/Source/Niagara/Private/NiagaraComponent.cpp
+++ b/Engine/Plugins/FX/Niagara/Source/Niagara/Private/NiagaraComponent.cpp
@@ -2877,7 +2877,7 @@ void UNiagaraComponent::PostLoad()
 	PostLoadVariantDIs(InstanceParameterOverrides);
 #endif
 
-	OverrideParameters.PostLoad();
+	OverrideParameters.PostLoad(this);
 	OverrideParameters.SanityCheckData();
 
 #if WITH_EDITOR
diff --git a/Engine/Plugins/FX/Niagara/Source/Niagara/Private/NiagaraParameterCollection.cpp b/Engine/Plugins/FX/Niagara/Source/Niagara/Private/NiagaraParameterCollection.cpp
--- a/Engine/Plugins/FX/Niagara/Source/Niagara/Private/NiagaraParameterCollection.cpp
+++ b/Engine/Plugins/FX/Niagara/Source/Niagara/Private/NiagaraParameterCollection.cpp
@@ -32,7 +32,7 @@ void UNiagaraParameterCollectionInstance::PostLoad()
 {
 	Super::PostLoad();
 
-	ParameterStorage.PostLoad();
+	ParameterStorage.PostLoad(this);
 
 	// Before calling SyncWithCollection we must ensure the collections parameter store is PostLoaded
 	// otherwise the parameters may not be sorted correctly.  If we are the default instance of the
diff --git a/Engine/Plugins/FX/Niagara/Source/Niagara/Private/NiagaraParameterStore.cpp b/Engine/Plugins/FX/Niagara/Source/Niagara/Private/NiagaraParameterStore.cpp
--- a/Engine/Plugins/FX/Niagara/Source/Niagara/Private/NiagaraParameterStore.cpp
+++ b/Engine/Plugins/FX/Niagara/Source/Niagara/Private/NiagaraParameterStore.cpp
@@ -1448,7 +1448,7 @@ const int32* FNiagaraParameterStore::FindParameterOffset(const FNiagaraVariableB
 	return nullptr;
 }
 
-void FNiagaraParameterStore::PostLoad()
+void FNiagaraParameterStore::PostLoad(UObject* InOwner)
 {
 #if WITH_EDITORONLY_DATA
 	// Convert ParameterOffsets map to the new SortedParameterOffsets array.
@@ -1474,6 +1474,8 @@ void FNiagaraParameterStore::PostLoad()
 	}
 #endif
 
+	SetOwner(InOwner);
+
 	// Not always required if NIAGARA_VARIABLE_LEXICAL_SORTING
 	SortParameters();
 }
diff --git a/Engine/Plugins/FX/Niagara/Source/Niagara/Private/NiagaraScript.cpp b/Engine/Plugins/FX/Niagara/Source/Niagara/Private/NiagaraScript.cpp
--- a/Engine/Plugins/FX/Niagara/Source/Niagara/Private/NiagaraScript.cpp
+++ b/Engine/Plugins/FX/Niagara/Source/Niagara/Private/NiagaraScript.cpp
@@ -2129,11 +2129,11 @@ void UNiagaraScript::PostLoad()
 	CheckVersionDataAvailable();
 #endif
 	
-	RapidIterationParameters.PostLoad();
+	RapidIterationParameters.PostLoad(this);
 
 	if (FPlatformProperties::RequiresCookedData() || IsScriptCooked())
 	{
-		ScriptExecutionParamStore.PostLoad();
+		ScriptExecutionParamStore.PostLoad(this);
 
 		// if our bindings aren't valid, then something has gone wrong with our cook and we need to disable this Script, which will in turn
 		// disable the owning script and system
diff --git a/Engine/Plugins/FX/Niagara/Source/Niagara/Private/NiagaraSystem.cpp b/Engine/Plugins/FX/Niagara/Source/Niagara/Private/NiagaraSystem.cpp
--- a/Engine/Plugins/FX/Niagara/Source/Niagara/Private/NiagaraSystem.cpp
+++ b/Engine/Plugins/FX/Niagara/Source/Niagara/Private/NiagaraSystem.cpp
@@ -940,10 +940,15 @@ void UNiagaraSystem::PostLoad()
 		}
 	}
 
-	ExposedParameters.PostLoad();
+	ExposedParameters.PostLoad(this);
 	ExposedParameters.SanityCheckData();
 
-	SystemCompiledData.InstanceParamStore.PostLoad();
+#if WITH_EDITORONLY_DATA
+	EditorOnlyAddedParameters.PostLoad(this);
+	EditorOnlyAddedParameters.SanityCheckData();
+#endif
+
+	SystemCompiledData.InstanceParamStore.PostLoad(this);
 
 	if (GIsEditor)
 	{
diff --git a/Engine/Plugins/FX/Niagara/Source/Niagara/Private/NiagaraUserRedirectionParameterStore.cpp b/Engine/Plugins/FX/Niagara/Source/Niagara/Private/NiagaraUserRedirectionParameterStore.cpp
--- a/Engine/Plugins/FX/Niagara/Source/Niagara/Private/NiagaraUserRedirectionParameterStore.cpp
+++ b/Engine/Plugins/FX/Niagara/Source/Niagara/Private/NiagaraUserRedirectionParameterStore.cpp
@@ -185,7 +185,7 @@ bool FNiagaraUserRedirectionParameterStore::SerializeFromMismatchedTag(const FPr
 		FNiagaraParameterStore OldStore;
 		FNiagaraParameterStore::StaticStruct()->SerializeItem(Slot, &OldStore, nullptr);
 		// Call PostLoad() to convert the serialized ParameterOffsets to SortedParameterOffsets.
-		OldStore.PostLoad();
+		OldStore.PostLoad(nullptr);
 		*this = OldStore;
 		return true;
 	}
diff --git a/Engine/Plugins/FX/Niagara/Source/Niagara/Public/NiagaraParameterStore.h b/Engine/Plugins/FX/Niagara/Source/Niagara/Public/NiagaraParameterStore.h
--- a/Engine/Plugins/FX/Niagara/Source/Niagara/Public/NiagaraParameterStore.h
+++ b/Engine/Plugins/FX/Niagara/Source/Niagara/Public/NiagaraParameterStore.h
@@ -256,7 +256,7 @@ public:
 	NIAGARA_API virtual ~FNiagaraParameterStore();
 	
 #if WITH_EDITORONLY_DATA
-	UPROPERTY()
+	UPROPERTY(Transient)
 	FString DebugName;
 
 	/** Guid data to remap rapid iteration parameters after a function input was renamed. */
@@ -477,7 +477,7 @@ public:
 
 	NIAGARA_API virtual const int32* FindParameterOffset(const FNiagaraVariableBase& Parameter, bool IgnoreType = false) const;
 
-	NIAGARA_API void PostLoad();
+	NIAGARA_API void PostLoad(UObject* InOwner);
 	NIAGARA_API void SortParameters();
 
 	/** Returns the UObject at the passed offset. */
diff --git a/Engine/Plugins/FX/Niagara/Source/NiagaraEditor/Private/NiagaraNodeAssignment.cpp b/Engine/Plugins/FX/Niagara/Source/NiagaraEditor/Private/NiagaraNodeAssignment.cpp
--- a/Engine/Plugins/FX/Niagara/Source/NiagaraEditor/Private/NiagaraNodeAssignment.cpp
+++ b/Engine/Plugins/FX/Niagara/Source/NiagaraEditor/Private/NiagaraNodeAssignment.cpp
@@ -203,7 +203,7 @@ void UNiagaraNodeAssignment::PostLoad()
 				check(Script);
 				if (Script->HasAnyFlags(RF_NeedPostLoad))
 				{
-					Script->RapidIterationParameters.PostLoad();
+					Script->RapidIterationParameters.PostLoad(Script);
 				}
 				if (Script->HandleVariableRenames(Converted, Emitter ? Emitter->GetUniqueEmitterName() : FString()))
 				{
