Fix for cleaning up inactive solo systems when they are still holding onto a simulation and the simulation is being changed

diff --git a/Engine/Plugins/FX/Niagara/Source/Niagara/Private/NiagaraCommon.cpp b/Engine/Plugins/FX/Niagara/Source/Niagara/Private/NiagaraCommon.cpp
--- a/Engine/Plugins/FX/Niagara/Source/Niagara/Private/NiagaraCommon.cpp
+++ b/Engine/Plugins/FX/Niagara/Source/Niagara/Private/NiagaraCommon.cpp
@@ -88,6 +88,16 @@ void FNiagaraSystemUpdateContext::CommitUpdate()
 	}
 	SystemSimsToDestroy.Empty();
 
+	for (UNiagaraComponent* Comp : ComponentsToDestroyInstance)
+	{
+		if (Comp)
+		{
+			Comp->DestroyInstanceNotComponent();
+			PostWork.ExecuteIfBound(Comp);
+		}
+	}
+	ComponentsToDestroyInstance.Empty();
+
 	bool bNeedsWaitOnGpu = true;
 	for (UNiagaraSystem* NiagaraSystem : SystemSimsToRecache)
 	{
@@ -285,9 +295,18 @@ void FNiagaraSystemUpdateContext::AddInternal(UNiagaraComponent* Comp, bool bReI
 		// Otherwise, they will hold reference and bind or remain bound to a system simulation that has been abandoned by the world manager
 		if (FNiagaraSystemInstanceControllerConstPtr SystemInstanceController = Comp->GetSystemInstanceController())
 		{
-			if (!SystemInstanceController->IsSolo() && SystemInstanceController->HasValidSimulation())
+			if (SystemInstanceController->HasValidSimulation())
 			{
-				ComponentsToNotifySimDestroy.Add(Comp);
+				if (!SystemInstanceController->IsSolo())
+				{
+					ComponentsToNotifySimDestroy.AddUnique(Comp);
+				}
+				// solo systems still need to be reinitialized because we don't want them to try to use stale data either (like if a compilation
+				// has changed the script data)
+				else
+				{
+					ComponentsToDestroyInstance.AddUnique(Comp);
+				}
 				return;
 			}
 		}
diff --git a/Engine/Plugins/FX/Niagara/Source/Niagara/Public/NiagaraCommon.h b/Engine/Plugins/FX/Niagara/Source/Niagara/Public/NiagaraCommon.h
--- a/Engine/Plugins/FX/Niagara/Source/Niagara/Public/NiagaraCommon.h
+++ b/Engine/Plugins/FX/Niagara/Source/Niagara/Public/NiagaraCommon.h
@@ -1014,6 +1014,8 @@ private:
 	UPROPERTY(transient)
 	TArray<TObjectPtr<UNiagaraComponent>> ComponentsToNotifySimDestroy;
 	UPROPERTY(transient)
+	TArray<TObjectPtr<UNiagaraComponent>> ComponentsToDestroyInstance;
+	UPROPERTY(transient)
 	TArray<TObjectPtr<UNiagaraSystem>> SystemSimsToDestroy;
 	UPROPERTY(transient)
 	TArray<TObjectPtr<UNiagaraSystem>> SystemSimsToRecache;
