- Fix for GPU ribbons max using the maximum ever allowed vs maximum estimate
- This saved a simple GPU ribbon test ~15mb of memory

diff --git a/Engine/Plugins/FX/Niagara/Source/Niagara/Private/NiagaraRendererRibbons.cpp b/Engine/Plugins/FX/Niagara/Source/Niagara/Private/NiagaraRendererRibbons.cpp
--- a/Engine/Plugins/FX/Niagara/Source/Niagara/Private/NiagaraRendererRibbons.cpp
+++ b/Engine/Plugins/FX/Niagara/Source/Niagara/Private/NiagaraRendererRibbons.cpp
@@ -290,10 +290,6 @@ struct FNiagaraDynamicDataRibbon : public FNiagaraDynamicDataBase
 {
 	FNiagaraDynamicDataRibbon(const FNiagaraEmitterInstance* InEmitter)
 		: FNiagaraDynamicDataBase(InEmitter)
-		, Material(nullptr)
-		, MaxAllocatedParticleCount(0)
-		, bUseGPUInit(false)
-		, bIsGPUSystem(false)
 	{
 	}
 
@@ -306,12 +302,12 @@ struct FNiagaraDynamicDataRibbon : public FNiagaraDynamicDataBase
 	}
 
 	
-	/** Material to use passed to the Renderer. */
-	FMaterialRenderProxy* Material;
-	int32 MaxAllocatedParticleCount;
+	FMaterialRenderProxy* Material = nullptr;		// MaterialProxy used on the renderer
+	uint32 MaxAllocationCount = 0;					// Maximum allocation count allowed (i.e. max we can fit in a buffer)
+	uint32 MaxAllocatedCountEstimate = 0;			// Maximum allocated count ever seen / estimate, since this is updated on the GT might be lower than actual particle count
 
-	bool bUseGPUInit;
-	bool bIsGPUSystem;
+	bool bUseGPUInit = false;
+	bool bIsGPUSystem = false;
 	
 	TSharedPtr<FNiagaraRibbonCPUGeneratedVertexData> GenerationOutput;
 	
@@ -451,35 +447,34 @@ FNiagaraRibbonVertexBuffers::FNiagaraRibbonVertexBuffers()
 {
 }
 
-void FNiagaraRibbonVertexBuffers::InitializeOrUpdateBuffers(FRHICommandListBase& RHICmdList, const FNiagaraRibbonGenerationConfig& GenerationConfig, const TSharedPtr<FNiagaraRibbonCPUGeneratedVertexData>& GeneratedGeometryData, const FNiagaraDataBuffer* SourceParticleData, int32 MaxAllocatedCount, bool bIsUsingGPUInit)
+void FNiagaraRibbonVertexBuffers::InitializeOrUpdateBuffers(FRHICommandListBase& RHICmdList, const FNiagaraRibbonGenerationConfig& GenerationConfig, const TSharedPtr<FNiagaraRibbonCPUGeneratedVertexData>& GeneratedGeometryData, const FNiagaraDataBuffer* SourceParticleData, int32 MaxAllocationCount, bool bIsUsingGPUInit)
 {	
-	const uint32 MaxAllocatedRibbons = GenerationConfig.HasRibbonIDs() ? (GenerationConfig.GetMaxNumRibbons() > 0 ? GenerationConfig.GetMaxNumRibbons() : MaxAllocatedCount) : 1;
-	
 	constexpr ERHIAccess InitialBufferAccessFlags = ERHIAccess::SRVMask | ERHIAccess::VertexOrIndexBuffer;
 
 	if (bIsUsingGPUInit)
 	{
 		const uint32 TotalParticles = SourceParticleData->GetNumInstances();
+		const uint32 MaxAllocatedRibbons = GenerationConfig.HasRibbonIDs() ? (GenerationConfig.GetMaxNumRibbons() > 0 ? GenerationConfig.GetMaxNumRibbons() : TotalParticles) : 1;
 
 		//-OPT:  We should be able to assume 2 particles per ribbon, however the compute pass does not cull our single particle ribbons therefore we need to allocate
 		//       enough space to assume each particle will be the start of a unique ribbon to avoid running OOB on the buffers.
 		const int32 TotalRibbons = FMath::Clamp<int32>(TotalParticles, 1, MaxAllocatedRibbons);
 
-		SortedIndicesBuffer.Allocate(RHICmdList, TotalParticles, MaxAllocatedCount, InitialBufferAccessFlags, false);
-		TangentsAndDistancesBuffer.Allocate(RHICmdList, TotalParticles * 4, MaxAllocatedCount * 4, InitialBufferAccessFlags, false);
-		MultiRibbonIndicesBuffer.Allocate(RHICmdList, GenerationConfig.HasRibbonIDs() ? TotalParticles : 0, MaxAllocatedCount, InitialBufferAccessFlags, false);
+		SortedIndicesBuffer.Allocate(RHICmdList, TotalParticles, MaxAllocationCount, InitialBufferAccessFlags, false);
+		TangentsAndDistancesBuffer.Allocate(RHICmdList, TotalParticles * 4, MaxAllocationCount * 4, InitialBufferAccessFlags, false);
+		MultiRibbonIndicesBuffer.Allocate(RHICmdList, GenerationConfig.HasRibbonIDs() ? TotalParticles : 0, MaxAllocationCount, InitialBufferAccessFlags, false);
 		RibbonLookupTableBuffer.Allocate(RHICmdList, TotalRibbons * FRibbonMultiRibbonInfoBufferEntry::NumElements, MaxAllocatedRibbons * FRibbonMultiRibbonInfoBufferEntry::NumElements, InitialBufferAccessFlags, false);
-		SegmentsBuffer.Allocate(RHICmdList, TotalParticles, MaxAllocatedCount, InitialBufferAccessFlags, false);
+		SegmentsBuffer.Allocate(RHICmdList, TotalParticles, MaxAllocationCount, InitialBufferAccessFlags, false);
 		bJustCreatedCommandBuffer |= GPUComputeCommandBuffer.Allocate(RHICmdList, FNiagaraRibbonCommandBufferLayout::NumElements, FNiagaraRibbonCommandBufferLayout::NumElements, InitialBufferAccessFlags | ERHIAccess::IndirectArgs, false, EBufferUsageFlags::DrawIndirect);
 	}
 	else
 	{		
 		check(GeneratedGeometryData.IsValid());
 
-		SortedIndicesBuffer.Allocate(RHICmdList, GeneratedGeometryData->SortedIndices.Num(), MaxAllocatedCount, InitialBufferAccessFlags, true);
-		TangentsAndDistancesBuffer.Allocate(RHICmdList, GeneratedGeometryData->TangentAndDistances.Num() * 4, MaxAllocatedCount * 4, InitialBufferAccessFlags, true);
-		MultiRibbonIndicesBuffer.Allocate(RHICmdList, GenerationConfig.HasRibbonIDs() ? GeneratedGeometryData->MultiRibbonIndices.Num() : 0, MaxAllocatedCount, InitialBufferAccessFlags, true);
-		RibbonLookupTableBuffer.Allocate(RHICmdList, GeneratedGeometryData->RibbonInfoLookup.Num() * FRibbonMultiRibbonInfoBufferEntry::NumElements, MaxAllocatedCount * FRibbonMultiRibbonInfoBufferEntry::NumElements, InitialBufferAccessFlags, true);
+		SortedIndicesBuffer.Allocate(RHICmdList, GeneratedGeometryData->SortedIndices.Num(), MaxAllocationCount, InitialBufferAccessFlags, true);
+		TangentsAndDistancesBuffer.Allocate(RHICmdList, GeneratedGeometryData->TangentAndDistances.Num() * 4, MaxAllocationCount * 4, InitialBufferAccessFlags, true);
+		MultiRibbonIndicesBuffer.Allocate(RHICmdList, GenerationConfig.HasRibbonIDs() ? GeneratedGeometryData->MultiRibbonIndices.Num() : 0, MaxAllocationCount, InitialBufferAccessFlags, true);
+		RibbonLookupTableBuffer.Allocate(RHICmdList, GeneratedGeometryData->RibbonInfoLookup.Num() * FRibbonMultiRibbonInfoBufferEntry::NumElements, MaxAllocationCount * FRibbonMultiRibbonInfoBufferEntry::NumElements, InitialBufferAccessFlags, true);
 		SegmentsBuffer.Release();
 		GPUComputeCommandBuffer.Release();
 		bJustCreatedCommandBuffer = false;
@@ -589,11 +584,11 @@ public:
 		return BufferEntry->Buffer;
 	}
 
-	FNiagaraRibbonIndexBuffer GetOrAllocateIndexBuffer(FRHICommandListBase& RHICmdList, int32 NumIndices, int32 MaxIndices)
+	FNiagaraRibbonIndexBuffer GetOrAllocateIndexBuffer(FRHICommandListBase& RHICmdList, int32 NumIndices, int32 MaxIndicesEstimate)
 	{
 		if (GNiagaraRibbonGpuBufferCachePurgeCounter >= 0)
 		{
-			NumIndices = GNiagaraRibbonGpuAllocateMaxCount == 0 ? Align(NumIndices, GNiagaraRibbonGpuBufferAlign) : MaxIndices;
+			NumIndices = GNiagaraRibbonGpuAllocateMaxCount == 0 ? Align(NumIndices, GNiagaraRibbonGpuBufferAlign) : MaxIndicesEstimate;
 		}
 
 		FIndexBufferEntry* BufferEntry = Index32BufferCache.FindByPredicate(
@@ -935,7 +930,8 @@ FNiagaraDynamicDataBase* FNiagaraRendererRibbons::GenerateDynamicData(const FNia
 			
 			DynamicData->bUseGPUInit = bIsGPUSystem || bWantsGPUInit;
 			DynamicData->bIsGPUSystem = bIsGPUSystem;
-			DynamicData->MaxAllocatedParticleCount = Emitter->GetData().GetMaxAllocationCount();
+			DynamicData->MaxAllocationCount = Emitter->GetData().GetMaxAllocationCount();
+			DynamicData->MaxAllocatedCountEstimate = FMath::Min<uint32>(Emitter->GetCachedEmitterData()->GetMaxParticleCountEstimate(), DynamicData->MaxAllocationCount);
 			
 			if (!DynamicData->bUseGPUInit)
 			{
@@ -1878,7 +1874,7 @@ void FNiagaraRendererRibbons::GenerateIndexBufferForView(
 		if (DynamicDataRibbon->bUseGPUInit)
 		{
 			RenderingViewResources->IndirectDrawBuffer = GpuRibbonsDataManager.GetOrAllocateIndirectDrawBuffer(RHICmdList);
-			RenderingViewResources->IndexBuffer = GpuRibbonsDataManager.GetOrAllocateIndexBuffer(RHICmdList, GeneratedData.TotalNumIndices, DynamicDataRibbon->MaxAllocatedParticleCount);
+			RenderingViewResources->IndexBuffer = GpuRibbonsDataManager.GetOrAllocateIndexBuffer(RHICmdList, GeneratedData.TotalNumIndices, FMath::Max(GeneratedData.TotalNumIndices, DynamicDataRibbon->MaxAllocatedCountEstimate));
 		}
 		else
 		{
@@ -2248,7 +2244,7 @@ void FNiagaraRendererRibbons::InitializeVertexBuffersResources(FRHICommandListBa
 {
 
 	// Make sure our ribbon data buffers are setup
-	VertexBuffers.InitializeOrUpdateBuffers(RHICmdList, GenerationConfig, DynamicDataRibbon->GenerationOutput, SourceParticleData, DynamicDataRibbon->MaxAllocatedParticleCount, bShouldUseGPUInit);
+	VertexBuffers.InitializeOrUpdateBuffers(RHICmdList, GenerationConfig, DynamicDataRibbon->GenerationOutput, SourceParticleData, DynamicDataRibbon->MaxAllocationCount, bShouldUseGPUInit);
 	
 	// Now we need to bind the source particle data, copying it to the gpu if necessary
 	if (DynamicDataRibbon->bIsGPUSystem)
diff --git a/Engine/Plugins/FX/Niagara/Source/Niagara/Public/NiagaraRendererRibbons.h b/Engine/Plugins/FX/Niagara/Source/Niagara/Public/NiagaraRendererRibbons.h
--- a/Engine/Plugins/FX/Niagara/Source/Niagara/Public/NiagaraRendererRibbons.h
+++ b/Engine/Plugins/FX/Niagara/Source/Niagara/Public/NiagaraRendererRibbons.h
@@ -243,7 +243,7 @@ struct FNiagaraRibbonVertexBuffers
 	FNiagaraRibbonGpuBuffer GPUComputeCommandBuffer;
 	bool bJustCreatedCommandBuffer = false;
 
-	void InitializeOrUpdateBuffers(FRHICommandListBase& RHICmdList, const FNiagaraRibbonGenerationConfig& GenerationConfig, const TSharedPtr<FNiagaraRibbonCPUGeneratedVertexData>& GeneratedGeometryData, const FNiagaraDataBuffer* SourceParticleData, int32 MaxAllocatedCount, bool bIsUsingGPUInit);
+	void InitializeOrUpdateBuffers(FRHICommandListBase& RHICmdList, const FNiagaraRibbonGenerationConfig& GenerationConfig, const TSharedPtr<FNiagaraRibbonCPUGeneratedVertexData>& GeneratedGeometryData, const FNiagaraDataBuffer* SourceParticleData, int32 MaxAllocationCount, bool bIsUsingGPUInit);
 
 	void Release()
 	{
