- Add cvar disable for light renderer
- Fix cvar disable for mesh renderer

diff --git a/Engine/Plugins/FX/Niagara/Source/Niagara/Private/NiagaraRendererLights.cpp b/Engine/Plugins/FX/Niagara/Source/Niagara/Private/NiagaraRendererLights.cpp
--- a/Engine/Plugins/FX/Niagara/Source/Niagara/Private/NiagaraRendererLights.cpp
+++ b/Engine/Plugins/FX/Niagara/Source/Niagara/Private/NiagaraRendererLights.cpp
@@ -18,6 +18,14 @@
 
 DECLARE_CYCLE_STAT(TEXT("Generate Particle Lights"), STAT_NiagaraGenLights, STATGROUP_Niagara);
 
+static int32 GbEnableNiagaraLightRendering = 1;
+static FAutoConsoleVariableRef CVarEnableNiagaraLightRendering(
+	TEXT("fx.EnableNiagaraLightRendering"),
+	GbEnableNiagaraLightRendering,
+	TEXT("If == 0, Niagara Light Renderers are disabled. \n"),
+	ECVF_Default
+);
+
 struct FNiagaraDynamicDataLights : public FNiagaraDynamicDataBase
 {
 	FNiagaraDynamicDataLights(const FNiagaraEmitterInstance* InEmitter)
@@ -177,6 +185,11 @@ FNiagaraDynamicDataBase* FNiagaraRendererLights::GenerateDynamicData(const FNiag
 
 void FNiagaraRendererLights::GatherSimpleLights(FSimpleLightArray& OutParticleLights)const
 {
+	if (GbEnableNiagaraLightRendering == 0)
+	{
+		return;
+	}
+
 	if (const FNiagaraDynamicDataLights* DynamicData = static_cast<const FNiagaraDynamicDataLights*>(DynamicDataRender))
 	{
 		const int32 LightCount = DynamicData->LightArray.Num();
diff --git a/Engine/Plugins/FX/Niagara/Source/Niagara/Private/NiagaraRendererMeshes.cpp b/Engine/Plugins/FX/Niagara/Source/Niagara/Private/NiagaraRendererMeshes.cpp
--- a/Engine/Plugins/FX/Niagara/Source/Niagara/Private/NiagaraRendererMeshes.cpp
+++ b/Engine/Plugins/FX/Niagara/Source/Niagara/Private/NiagaraRendererMeshes.cpp
@@ -1285,7 +1285,7 @@ void FNiagaraRendererMeshes::GetDynamicMeshElements(const TArray<const FSceneVie
 	FParticleMeshRenderData ParticleMeshRenderData;
 	PrepareParticleMeshRenderData(ParticleMeshRenderData, ViewFamily, Collector, DynamicDataRender, SceneProxy, false, ENiagaraGpuComputeTickStage::Last);
 
-	if (ParticleMeshRenderData.SourceParticleData == nullptr )
+	if (ParticleMeshRenderData.SourceParticleData == nullptr || GbEnableNiagaraMeshRendering == 0)
 	{
 		return;
 	}
