Adding a per FXType breakdown to the CSV Graphs.

diff --git a/Engine/Plugins/FX/Niagara/Source/Niagara/Classes/NiagaraEffectType.h b/Engine/Plugins/FX/Niagara/Source/Niagara/Classes/NiagaraEffectType.h
--- a/Engine/Plugins/FX/Niagara/Source/Niagara/Classes/NiagaraEffectType.h
+++ b/Engine/Plugins/FX/Niagara/Source/Niagara/Classes/NiagaraEffectType.h
@@ -8,10 +8,13 @@
 #include "NiagaraPerfBaseline.h"
 #include "NiagaraValidationRule.h"
 #include "NiagaraValidationRuleSet.h"
+#include "Particles/ParticlePerfStats.h"
 #include "NiagaraEffectType.generated.h"
 
 #define DEBUG_SCALABILITY_STATE (!UE_BUILD_SHIPPING)
 
+#define WITH_PER_FXTYPE_PARTICLE_PERF_STATS (WITH_PER_SYSTEM_PARTICLE_PERF_STATS && WITH_PARTICLE_PERF_CSV_STATS && !UE_BUILD_SHIPPING)
+
 /** Controls what action is taken by a Niagara system that fails it's cull checks. */
 UENUM()
 enum class ENiagaraCullReaction
@@ -397,6 +400,7 @@ class UNiagaraEffectType : public UObject
 	NIAGARA_API virtual void BeginDestroy()override;
 	NIAGARA_API virtual bool IsReadyForFinishDestroy()override;
 	NIAGARA_API virtual void Serialize(FArchive& Ar)override;
+	NIAGARA_API virtual void PostInitProperties()override;
 	NIAGARA_API virtual void PostLoad()override;
 #if WITH_EDITORONLY_DATA
 	static NIAGARA_API void DeclareConstructClasses(TArray<FTopLevelAssetPath>& OutConstructClasses, const UClass* SpecificSubclass);
@@ -500,6 +504,15 @@ public:
 	static FGeneratePerfBaselines& OnGeneratePerfBaselines(){ return GeneratePerfBaselinesDelegate; }
 	static NIAGARA_API void GeneratePerfBaselines();
 #endif
+
+#if WITH_PER_FXTYPE_PARTICLE_PERF_STATS
+	//Cached CSV Stat names for this system.
+	FName CSVStat_Count = NAME_None;
+	FName CSVStat_Total = NAME_None;
+	FName CSVStat_GTOnly = NAME_None;
+	FName CSVStat_RT = NAME_None;
+	FName CSVStat_GPU = NAME_None;
+#endif
 };
 
 /** Performs the passed action for all FNiagaraPlatformSets in this FXType. */
@@ -516,3 +529,25 @@ void UNiagaraEffectType::ForEachPlatformSet(TAction Func)
 		Func(Setting.Platforms);
 	}
 }
+
+//////////////////////////////////////////////////////////////////////////
+
+#if WITH_PER_FXTYPE_PARTICLE_PERF_STATS
+/** Listener to gather accumulated stats for all systems in a particular Effect Type. */
+class FParticlePerfStatsListener_EffectType : public FParticlePerfStatsListener
+{
+public:
+	FParticlePerfStatsListener_EffectType()
+	{}
+
+	virtual ~FParticlePerfStatsListener_EffectType() {}
+
+	virtual bool Tick() override;
+	virtual void TickRT() override;
+
+
+	virtual bool NeedsWorldStats() const { return false; }
+	virtual bool NeedsSystemStats() const { return true; }
+	virtual bool NeedsComponentStats() const { return false; }
+};
+#endif//WITH_PER_FXTYPE_PARTICLE_PERF_STATS
\ No newline at end of file
diff --git a/Engine/Plugins/FX/Niagara/Source/Niagara/Private/NiagaraEffectType.cpp b/Engine/Plugins/FX/Niagara/Source/Niagara/Private/NiagaraEffectType.cpp
--- a/Engine/Plugins/FX/Niagara/Source/Niagara/Private/NiagaraEffectType.cpp
+++ b/Engine/Plugins/FX/Niagara/Source/Niagara/Private/NiagaraEffectType.cpp
@@ -44,6 +44,18 @@ void UNiagaraEffectType::Serialize(FArchive& Ar)
 	Ar.UsingCustomVersion(FNiagaraCustomVersion::GUID);
 }
 
+void UNiagaraEffectType::PostInitProperties()
+{
+	Super::PostInitProperties();
+#if WITH_PER_FXTYPE_PARTICLE_PERF_STATS
+	CSVStat_Total = *FString::Printf(TEXT("FXType_Total/%s"), *GetFName().ToString());
+	CSVStat_GTOnly = *FString::Printf(TEXT("FXType_GTOnly/%s"), *GetFName().ToString());
+	CSVStat_RT = *FString::Printf(TEXT("FXType_RT/%s"), *GetFName().ToString());
+	CSVStat_GPU = *FString::Printf(TEXT("FXType_GPU/%s"), *GetFName().ToString());
+	CSVStat_Count = *FString::Printf(TEXT("FXType_Count/%s"), *GetFName().ToString());
+#endif
+}
+
 void UNiagaraEffectType::PostLoad()
 {
 	Super::PostLoad();
@@ -407,3 +419,110 @@ FAutoConsoleCommand InvalidatePerfBaselinesCommand(
 
 #endif
 //////////////////////////////////////////////////////////////////////////
+
+
+//////////////////////////////////////////////////////////////////////////
+
+#if WITH_PER_FXTYPE_PARTICLE_PERF_STATS
+
+bool FParticlePerfStatsListener_EffectType::Tick()
+{
+	if (FCsvProfiler* CSVProfiler = FCsvProfiler::Get())
+	{
+		//First we gather up stats for all systems in each FXType
+		TMap<UNiagaraEffectType*, FParticlePerfStats> FXTypeStatsMap;
+		FParticlePerfStatsManager::ForAllSystemStats(
+			[&](TWeakObjectPtr<const UFXSystemAsset>& WeakSystem, TUniquePtr<FParticlePerfStats>& Stats)
+			{
+				if (Stats->GetGameThreadStats().NumInstances > 0)
+				{
+					if (const UNiagaraSystem* System = Cast<UNiagaraSystem>(WeakSystem.Get()))
+					{
+						FParticlePerfStats& FXTypeStats = FXTypeStatsMap.FindOrAdd(System->GetEffectType());
+						FParticlePerfStats_GT& GTStats = FXTypeStats.GetGameThreadStats();
+
+						GTStats += Stats->GetGameThreadStats();
+					}
+				}
+			}
+		);
+
+		//Now add to CSV
+		for (auto It = FXTypeStatsMap.CreateIterator(); It; ++It)
+		{
+			if (UNiagaraEffectType* FXType = It.Key())
+			{
+				FParticlePerfStats& Stats = It.Value();
+
+				if (Stats.GetGameThreadStats().NumInstances > 0)
+				{
+					float TotalTime = FPlatformTime::ToMilliseconds64(Stats.GetGameThreadStats().GetTotalCycles()) * 1000.0f;
+					float GTTime = FPlatformTime::ToMilliseconds64(Stats.GetGameThreadStats().GetTotalCycles_GTOnly()) * 1000.0f;
+					//float AvgTime = FPlatformTime::ToMilliseconds64(Stats->GetGameThreadStats().GetPerInstanceAvgCycles()) * 1000.0f;
+					int32 Count = (int32)Stats.GetGameThreadStats().NumInstances;
+					//float Activation = FPlatformTime::ToMilliseconds64(Stats->GetGameThreadStats().ActivationCycles) * 1000.0f;
+					//float Wait = FPlatformTime::ToMilliseconds64(Stats->GetGameThreadStats().WaitCycles) * 1000.0f;
+
+					CSVProfiler->RecordCustomStat(FXType->CSVStat_Total, CSV_CATEGORY_INDEX(Particles), TotalTime, ECsvCustomStatOp::Set);
+					CSVProfiler->RecordCustomStat(FXType->CSVStat_GTOnly, CSV_CATEGORY_INDEX(Particles), GTTime, ECsvCustomStatOp::Set);
+					//CSVProfiler->RecordCustomStat(FXType->CSVStat_InstAvgGT, CSV_CATEGORY_INDEX(Particles), AvgTime, ECsvCustomStatOp::Set);
+					CSVProfiler->RecordCustomStat(FXType->CSVStat_Count, CSV_CATEGORY_INDEX(Particles), Count, ECsvCustomStatOp::Set);
+
+					//CSVProfiler->RecordCustomStat(System->CSVStat_Activation, CSV_CATEGORY_INDEX(Particles), Activation, ECsvCustomStatOp::Set);
+					//CSVProfiler->RecordCustomStat(System->CSVStat_Waits, CSV_CATEGORY_INDEX(Particles), Wait, ECsvCustomStatOp::Set);
+				}
+			}
+		}
+	}
+	return true;
+}
+
+void FParticlePerfStatsListener_EffectType::TickRT()
+{
+	if (FCsvProfiler* CSVProfiler = FCsvProfiler::Get())
+	{
+		//First we gather up stats for all systems in each FXType
+		TMap<UNiagaraEffectType*, FParticlePerfStats> FXTypeStatsMap;
+		FParticlePerfStatsManager::ForAllSystemStats(
+			[&](TWeakObjectPtr<const UFXSystemAsset>& WeakSystem, TUniquePtr<FParticlePerfStats>& Stats)
+			{
+				if (Stats->GetGameThreadStats().NumInstances > 0)
+				{
+					if (const UNiagaraSystem* System = Cast<UNiagaraSystem>(WeakSystem.Get()))
+					{
+						FParticlePerfStats& FXTypeStats = FXTypeStatsMap.FindOrAdd(System->GetEffectType());
+						FParticlePerfStats_RT& RTStats = FXTypeStats.GetRenderThreadStats();
+						FParticlePerfStats_GPU& GPUStats = FXTypeStats.GetGPUStats();
+
+						RTStats += Stats->GetRenderThreadStats();
+						GPUStats += Stats->GetGPUStats();
+					}
+				}
+			}
+		);
+
+		//Now add to CSV
+		for (auto It = FXTypeStatsMap.CreateIterator(); It; ++It)
+		{
+			if (UNiagaraEffectType* FXType = It.Key())
+			{
+				FParticlePerfStats& Stats = It.Value();
+
+				if (Stats.GetRenderThreadStats().NumInstances > 0)
+				{
+					const float RTTime = FPlatformTime::ToMilliseconds64(Stats.GetRenderThreadStats().GetTotalCycles()) * 1000.0f;
+					//const float RTAvgTime = FPlatformTime::ToMilliseconds64(Stats->GetRenderThreadStats().GetPerInstanceAvgCycles()) * 1000.0f;
+					CSVProfiler->RecordCustomStat(FXType->CSVStat_RT, CSV_CATEGORY_INDEX(Particles), RTTime, ECsvCustomStatOp::Set);
+					//CSVProfiler->RecordCustomStat(System->CSVStat_InstAvgRT, CSV_CATEGORY_INDEX(Particles), RTAvgTime, ECsvCustomStatOp::Set);
+
+					const float GpuTime = float(Stats.GetGPUStats().GetTotalMicroseconds());
+					//const float GpuAvgTime = float(Stats->GetGPUStats().GetPerInstanceAvgMicroseconds());
+					CSVProfiler->RecordCustomStat(FXType->CSVStat_GPU, CSV_CATEGORY_INDEX(Particles), GpuTime, ECsvCustomStatOp::Set);
+					//CSVProfiler->RecordCustomStat(System->CSVStat_InstAvgGPU, CSV_CATEGORY_INDEX(Particles), GpuAvgTime, ECsvCustomStatOp::Set);
+				}
+			}
+		}
+	}
+}
+
+#endif //WITH_PER_FXTYPE_PARTICLE_PERF_STATS
\ No newline at end of file
diff --git a/Engine/Plugins/FX/Niagara/Source/Niagara/Private/NiagaraWorldManager.cpp b/Engine/Plugins/FX/Niagara/Source/Niagara/Private/NiagaraWorldManager.cpp
--- a/Engine/Plugins/FX/Niagara/Source/Niagara/Private/NiagaraWorldManager.cpp
+++ b/Engine/Plugins/FX/Niagara/Source/Niagara/Private/NiagaraWorldManager.cpp
@@ -1088,11 +1088,20 @@ void FNiagaraWorldManager::PostActorTick(float DeltaSeconds)
 		}
 	}
 
+	HandleCSVStats(DeltaSeconds);
+
+	DataChannelManager->EndFrame(DeltaSeconds);
+}
+
+void FNiagaraWorldManager::HandleCSVStats(float DeltaSeconds)
+{
+	bool bCSVStatsEnabled = false;
 #if WITH_PARTICLE_PERF_CSV_STATS
 	if (FCsvProfiler* CSVProfiler = FCsvProfiler::Get())
 	{
 		if (CSVProfiler->IsCapturing() && FParticlePerfStats::GetCSVStatsEnabled())
 		{
+			bCSVStatsEnabled = true;
 			//Record custom events marking split times at set intervals. Allows us to generate summary tables for averages over shorter bursts.
 			if (GNiagaraCSVSplitTime > 0.0f)
 			{
@@ -1115,11 +1124,24 @@ void FNiagaraWorldManager::PostActorTick(float DeltaSeconds)
 
 				ScalabilityMan.CSVProfilerUpdate(CSVProfiler);
 			}
+			#if WITH_PER_FXTYPE_PARTICLE_PERF_STATS
+			if (FXTypeCSVListener.IsValid() == false)
+			{
+				FXTypeCSVListener = MakeShared<FParticlePerfStatsListener_EffectType, ESPMode::ThreadSafe>();
+				FParticlePerfStatsManager::AddListener(FXTypeCSVListener);
+			}
+			#endif
 		}
 	}
 #endif 
 
-	DataChannelManager->EndFrame(DeltaSeconds);
+#if WITH_PER_FXTYPE_PARTICLE_PERF_STATS
+	if (bCSVStatsEnabled == false && FXTypeCSVListener.IsValid())
+	{
+		FParticlePerfStatsManager::RemoveListener(FXTypeCSVListener.Get());
+		FXTypeCSVListener.Reset();
+	}
+#endif
 }
 
 void FNiagaraWorldManager::PreSendAllEndOfFrameUpdates()
diff --git a/Engine/Plugins/FX/Niagara/Source/Niagara/Public/NiagaraWorldManager.h b/Engine/Plugins/FX/Niagara/Source/Niagara/Public/NiagaraWorldManager.h
--- a/Engine/Plugins/FX/Niagara/Source/Niagara/Public/NiagaraWorldManager.h
+++ b/Engine/Plugins/FX/Niagara/Source/Niagara/Public/NiagaraWorldManager.h
@@ -335,6 +335,8 @@ private:
 	// Returns scalability state if one exists, this function is not designed for runtime performance and for debugging only
 	bool GetScalabilityState(UNiagaraComponent* Component, FNiagaraScalabilityState& OutState) const;
 
+	void HandleCSVStats(float DeltaSeconds);
+
 	static FDelegateHandle OnWorldInitHandle;
 	static FDelegateHandle OnWorldCleanupHandle;
 	static FDelegateHandle OnPostWorldCleanupHandle;
@@ -407,6 +409,10 @@ private:
 	TUniquePtr<FNiagaraDataChannelManager> DataChannelManager;
 
 	TUniquePtr<FNiagaraSimpleObjectPool> ObjectPool;
+
+#if WITH_PER_FXTYPE_PARTICLE_PERF_STATS
+	FParticlePerfStatsListenerPtr FXTypeCSVListener;
+#endif
 };
 
 
diff --git a/Engine/Source/Runtime/Engine/Classes/Particles/ParticlePerfStats.h b/Engine/Source/Runtime/Engine/Classes/Particles/ParticlePerfStats.h
--- a/Engine/Source/Runtime/Engine/Classes/Particles/ParticlePerfStats.h
+++ b/Engine/Source/Runtime/Engine/Classes/Particles/ParticlePerfStats.h
@@ -83,6 +83,18 @@ struct FParticlePerfStats_GT
 		return *this;
 	}
 
+	FParticlePerfStats_GT& operator+=(FParticlePerfStats_GT& Other)
+	{
+		NumInstances += Other.NumInstances;
+		TickGameThreadCycles += Other.TickGameThreadCycles;
+		TickConcurrentCycles += Other.TickConcurrentCycles.Load();
+		FinalizeCycles += Other.FinalizeCycles;
+		EndOfFrameCycles += Other.EndOfFrameCycles.Load();
+		ActivationCycles += Other.ActivationCycles.Load();
+		WaitCycles += Other.WaitCycles;
+		return *this;
+	}
+
 	FORCEINLINE void Reset()
 	{
 		NumInstances = 0;
@@ -114,6 +126,14 @@ struct FParticlePerfStats_RT
 	}
 	FORCEINLINE uint64 GetTotalCycles() const { return RenderUpdateCycles + GetDynamicMeshElementsCycles; }
 	FORCEINLINE uint64 GetPerInstanceAvgCycles() const { return NumInstances > 0 ? (RenderUpdateCycles + GetDynamicMeshElementsCycles) / NumInstances : 0; }
+
+	FParticlePerfStats_RT& operator+=(FParticlePerfStats_RT& Other)
+	{
+		NumInstances += Other.NumInstances;
+		RenderUpdateCycles += Other.RenderUpdateCycles;
+		GetDynamicMeshElementsCycles += Other.GetDynamicMeshElementsCycles;
+		return *this;
+	}
 };
 
 /** Stats gathered from the GPU */
@@ -131,6 +151,13 @@ struct FParticlePerfStats_GPU
 		NumInstances = 0;
 		TotalMicroseconds = 0;
 	}
+
+	FParticlePerfStats_GPU& operator+=(FParticlePerfStats_GPU& Other)
+	{
+		NumInstances += Other.NumInstances;
+		TotalMicroseconds += Other.TotalMicroseconds;
+		return *this;
+	}
 };
 
 struct FParticlePerfStats
