- Handle nullptr for GraphSource

diff --git a/Engine/Plugins/FX/Niagara/Source/NiagaraEditor/Private/NiagaraEmitterEditorData.cpp b/Engine/Plugins/FX/Niagara/Source/NiagaraEditor/Private/NiagaraEmitterEditorData.cpp
--- a/Engine/Plugins/FX/Niagara/Source/NiagaraEditor/Private/NiagaraEmitterEditorData.cpp
+++ b/Engine/Plugins/FX/Niagara/Source/NiagaraEditor/Private/NiagaraEmitterEditorData.cpp
@@ -373,12 +373,15 @@ void UNiagaraEmitterEditorData::PostLoad_TransferModuleStackNotesToNewFormat(UOb
 	UNiagaraEmitter* Emitter = Cast<UNiagaraEmitter>(Owner);
 	TArray<FNiagaraAssetVersion> AssetVersions = Emitter->GetAllAvailableVersions();
 	UNiagaraGraph* EmitterGraph = nullptr;
-	for(FNiagaraAssetVersion& AssetVersion : AssetVersions)
+	for (FNiagaraAssetVersion& AssetVersion : AssetVersions)
 	{
 		FVersionedNiagaraEmitterData* EmitterData = Emitter->GetEmitterData(AssetVersion.VersionGuid);
 		if(EmitterData->GetEditorData() == this)
 		{
-			EmitterGraph = Cast<UNiagaraScriptSource>(EmitterData->GraphSource)->NodeGraph;
+			if (UNiagaraScriptSource* GraphSource = Cast<UNiagaraScriptSource>(EmitterData->GraphSource))
+			{
+				EmitterGraph = GraphSource->NodeGraph;
+			}
 		}
 	}
 
