Added debug logging to check niagara nodes for pin connections that are nullptrs.

diff --git a/Engine/Plugins/FX/Niagara/Source/NiagaraEditor/Private/NiagaraNode.cpp b/Engine/Plugins/FX/Niagara/Source/NiagaraEditor/Private/NiagaraNode.cpp
--- a/Engine/Plugins/FX/Niagara/Source/NiagaraEditor/Private/NiagaraNode.cpp
+++ b/Engine/Plugins/FX/Niagara/Source/NiagaraEditor/Private/NiagaraNode.cpp
@@ -39,9 +39,32 @@ void UNiagaraNode::PostLoad()
 	{
 		SetFlags(RF_Transactional);
 	}
-
+	
 	for (UEdGraphPin* Pin : Pins)
 	{
+		if(Pin == nullptr)
+		{
+			UNiagaraGraph* OwningGraph = GetNiagaraGraph();
+			FString AssetPath = OwningGraph ? OwningGraph->GetPathName() : "Undetermined";
+			FString PinName = Pin->PinName.ToString();
+			FString NodeName = Pin->GetOwningNode() ? Pin->GetOwningNode()->GetNodeTitle(ENodeTitleType::FullTitle).ToString() : "Undetermined";
+			UE_LOG(LogNiagaraEditor, Log, TEXT("Pin of node %s in asset %s is nullptr!"), *NodeName, *AssetPath);
+		}
+		if(Pin->LinkedTo.Num() > 0)
+		{
+			for(UEdGraphPin* LinkedPin : Pin->LinkedTo)
+			{
+				if(LinkedPin == nullptr)
+				{
+					UNiagaraGraph* OwningGraph = GetNiagaraGraph();
+					FString AssetPath = OwningGraph ? OwningGraph->GetPathName() : "Undetermined";
+					FString PinName = Pin->PinName.ToString();
+					FString NodeName = Pin->GetOwningNode() ? Pin->GetOwningNode()->GetNodeTitle(ENodeTitleType::FullTitle).ToString() : "Undetermined";
+					UE_LOG(LogNiagaraEditor, Log, TEXT("Connected pin of Pin %s of node %s in asset %s is nullptr!"), *PinName, *NodeName, *AssetPath);
+				}
+			}
+		}
+		
 		const UEdGraphSchema_Niagara* Schema = CastChecked<UEdGraphSchema_Niagara>(GetSchema());
 		check(Schema);
 
