Removed unnecessary line that would crash in debug logging

diff --git a/Engine/Plugins/FX/Niagara/Source/NiagaraEditor/Private/NiagaraNode.cpp b/Engine/Plugins/FX/Niagara/Source/NiagaraEditor/Private/NiagaraNode.cpp
--- a/Engine/Plugins/FX/Niagara/Source/NiagaraEditor/Private/NiagaraNode.cpp
+++ b/Engine/Plugins/FX/Niagara/Source/NiagaraEditor/Private/NiagaraNode.cpp
@@ -46,8 +46,7 @@ void UNiagaraNode::PostLoad()
 		{
 			UNiagaraGraph* OwningGraph = GetNiagaraGraph();
 			FString AssetPath = OwningGraph ? OwningGraph->GetPathName() : "Undetermined";
-			FString PinName = Pin->PinName.ToString();
-			FString NodeName = Pin->GetOwningNode() ? Pin->GetOwningNode()->GetNodeTitle(ENodeTitleType::FullTitle).ToString() : "Undetermined";
+			FString NodeName = GetNodeTitle(ENodeTitleType::FullTitle).ToString();
 			UE_LOG(LogNiagaraEditor, Log, TEXT("Pin of node %s in asset %s is nullptr!"), *NodeName, *AssetPath);
 		}
 		if(Pin->LinkedTo.Num() > 0)
