Niagara - Fix an issue when copy/pasting modules with static switches where the static switch changes the set of input categories.

diff --git a/Engine/Plugins/FX/Niagara/Source/NiagaraEditor/Private/ViewModels/Stack/NiagaraStackFunctionInputCollection.cpp b/Engine/Plugins/FX/Niagara/Source/NiagaraEditor/Private/ViewModels/Stack/NiagaraStackFunctionInputCollection.cpp
--- a/Engine/Plugins/FX/Niagara/Source/NiagaraEditor/Private/ViewModels/Stack/NiagaraStackFunctionInputCollection.cpp
+++ b/Engine/Plugins/FX/Niagara/Source/NiagaraEditor/Private/ViewModels/Stack/NiagaraStackFunctionInputCollection.cpp
@@ -751,19 +751,34 @@ void UNiagaraStackFunctionInputCollection::ToClipboardFunctionInputs(UObject* In
 
 void UNiagaraStackFunctionInputCollection::SetValuesFromClipboardFunctionInputs(const TArray<const UNiagaraClipboardFunctionInput*>& ClipboardFunctionInputs)
 {
-	TArray<const UNiagaraClipboardFunctionInput*> StaticSwitchInputs;
-	TArray<const UNiagaraClipboardFunctionInput*> StandardInputs;
-
+	// First try to set each input as a static switch, and if a switch is set refresh the categories
+	// before applying additional inputs.  This is necessary because static switches can change the set
+	// of exposed inputs.
+	// NOTE: It's still possible that inputs could end up missing in cases where the switch dependencies are 
+	// especially complex since we're not doing a full switch dependency check, but this should handle the
+	// vast majority of cases.
 	TArray<UNiagaraStackInputCategory*> ChildCategories;
 	GetUnfilteredChildrenOfType(ChildCategories);
-
-	// Set static switches first so that other inputs will be available to set.
-	for (UNiagaraStackInputCategory* ChildCategory : ChildCategories)
+	for (const UNiagaraClipboardFunctionInput* ClipboardFunctionInput : ClipboardFunctionInputs)
 	{
-		ChildCategory->SetStaticSwitchValuesFromClipboardFunctionInputs(ClipboardFunctionInputs, *this);
+		bool bInputSetAsSwitch = false;
+		for (UNiagaraStackInputCategory* ChildCategory : ChildCategories)
+		{
+			if (ChildCategory->TrySetStaticSwitchValuesFromClipboardFunctionInput(*ClipboardFunctionInput))
+			{
+				bInputSetAsSwitch = true;
+				break;
+			}
+		}
+		if (bInputSetAsSwitch)
+		{
+			RefreshChildren();
+			ChildCategories.Empty();
+			GetUnfilteredChildrenOfType(ChildCategories);
+		}
 	}
 
-	RefreshChildren();
+	// After all static switches have been set the remaining standard inputs can be set without additional refreshes.
 	for (UNiagaraStackInputCategory* ChildCategory : ChildCategories)
 	{
 		ChildCategory->SetStandardValuesFromClipboardFunctionInputs(ClipboardFunctionInputs);
diff --git a/Engine/Plugins/FX/Niagara/Source/NiagaraEditor/Private/ViewModels/Stack/NiagaraStackInputCategory.cpp b/Engine/Plugins/FX/Niagara/Source/NiagaraEditor/Private/ViewModels/Stack/NiagaraStackInputCategory.cpp
--- a/Engine/Plugins/FX/Niagara/Source/NiagaraEditor/Private/ViewModels/Stack/NiagaraStackInputCategory.cpp
+++ b/Engine/Plugins/FX/Niagara/Source/NiagaraEditor/Private/ViewModels/Stack/NiagaraStackInputCategory.cpp
@@ -176,17 +176,41 @@ void UNiagaraStackInputCategory::ToClipboardFunctionInputs(UObject* InOuter, TAr
 	}
 }
 
-template<typename Predicate>
-void SetValuesFromFunctionInputsInternal(UNiagaraStackInputCategory* Category, const TArray<const UNiagaraClipboardFunctionInput*>& ClipboardFunctionInputs, Predicate InputMatchesFilter, UNiagaraStackFunctionInputCollection* OwningFunctionCollection)
+bool UNiagaraStackInputCategory::TrySetStaticSwitchValuesFromClipboardFunctionInput(const UNiagaraClipboardFunctionInput& ClipboardFunctionInput)
 {
+	TArray<UNiagaraStackFunctionInput*> ChildInputs;
+	GetUnfilteredChildrenOfType(ChildInputs);
+	for (UNiagaraStackFunctionInput* ChildInput : ChildInputs)
+	{
+		if (ChildInput->IsStaticParameter() &&
+			ChildInput->GetInputParameterHandle().GetName() == ClipboardFunctionInput.InputName &&
+			ChildInput->GetInputType() == ClipboardFunctionInput.InputType)
+		{
+			if (ClipboardFunctionInput.ValueMode == ENiagaraClipboardFunctionInputValueMode::ResetToDefault)
+			{
+				ChildInput->Reset();
+			}
+			else
+			{
+				ChildInput->SetValueFromClipboardFunctionInput(ClipboardFunctionInput);
+			}
+			return true;
+		}
+	}
+	return false;
+}
+
+void  UNiagaraStackInputCategory::SetStandardValuesFromClipboardFunctionInputs(const TArray<const UNiagaraClipboardFunctionInput*>& ClipboardFunctionInputs)
+{
+	TArray<UNiagaraStackFunctionInput*> ChildInputs;
+	GetUnfilteredChildrenOfType(ChildInputs);
 	for (const UNiagaraClipboardFunctionInput* ClipboardFunctionInput : ClipboardFunctionInputs)
 	{
-		TArray<UNiagaraStackFunctionInput*> ChildInputs;
-		Category->GetUnfilteredChildrenOfType(ChildInputs);
-		
 		for (UNiagaraStackFunctionInput* ChildInput : ChildInputs)
 		{
-			if (InputMatchesFilter(ChildInput) && ChildInput->GetInputParameterHandle().GetName() == ClipboardFunctionInput->InputName && ChildInput->GetInputType() == ClipboardFunctionInput->InputType)
+			if (ChildInput->IsStaticParameter() == false && 
+				ChildInput->GetInputParameterHandle().GetName() == ClipboardFunctionInput->InputName &&
+				ChildInput->GetInputType() == ClipboardFunctionInput->InputType)
 			{
 				if (ClipboardFunctionInput->ValueMode == ENiagaraClipboardFunctionInputValueMode::ResetToDefault)
 				{
@@ -196,26 +220,11 @@ void SetValuesFromFunctionInputsInternal(UNiagaraStackInputCategory* Category, c
 				{
 					ChildInput->SetValueFromClipboardFunctionInput(*ClipboardFunctionInput);
 				}
-				if(OwningFunctionCollection)
-				{
-					OwningFunctionCollection->RefreshChildren();
-				}
-				break;
 			}
 		}
 	}
 }
 
-void  UNiagaraStackInputCategory::SetStaticSwitchValuesFromClipboardFunctionInputs(const TArray<const UNiagaraClipboardFunctionInput*>& ClipboardFunctionInputs, UNiagaraStackFunctionInputCollection& OwningFunctionCollection)
-{
-	SetValuesFromFunctionInputsInternal(this, ClipboardFunctionInputs, [](UNiagaraStackFunctionInput* ChildInput) { return ChildInput->IsStaticParameter(); }, &OwningFunctionCollection);
-}
-
-void  UNiagaraStackInputCategory::SetStandardValuesFromClipboardFunctionInputs(const TArray<const UNiagaraClipboardFunctionInput*>& ClipboardFunctionInputs)
-{
-	SetValuesFromFunctionInputsInternal(this, ClipboardFunctionInputs, [](UNiagaraStackFunctionInput* ChildInput) { return ChildInput->IsStaticParameter() == false; }, nullptr);
-}
-
 void UNiagaraStackInputCategory::GetFilteredChildInputs(TArray<UNiagaraStackFunctionInput*>& OutFilteredChildInputs) const
 {
 	GetFilteredChildrenOfType(OutFilteredChildInputs);
diff --git a/Engine/Plugins/FX/Niagara/Source/NiagaraEditor/Public/ViewModels/Stack/NiagaraStackInputCategory.h b/Engine/Plugins/FX/Niagara/Source/NiagaraEditor/Public/ViewModels/Stack/NiagaraStackInputCategory.h
--- a/Engine/Plugins/FX/Niagara/Source/NiagaraEditor/Public/ViewModels/Stack/NiagaraStackInputCategory.h
+++ b/Engine/Plugins/FX/Niagara/Source/NiagaraEditor/Public/ViewModels/Stack/NiagaraStackInputCategory.h
@@ -73,9 +73,7 @@ public:
 
 	NIAGARAEDITOR_API void ToClipboardFunctionInputs(UObject* InOuter, TArray<const UNiagaraClipboardFunctionInput*>& OutClipboardFunctionInputs) const;
 
-	// We pass in the owning input collection to refresh the inputs of a category after setting a static switch value.
-	// This is done to support copy paste for chained static switches, as otherwise the stack entry for the 2nd static switch won't exist for pasting data
-	NIAGARAEDITOR_API void SetStaticSwitchValuesFromClipboardFunctionInputs(const TArray<const UNiagaraClipboardFunctionInput*>& ClipboardFunctionInputs, class UNiagaraStackFunctionInputCollection& OwningFunctionCollection);
+	NIAGARAEDITOR_API bool TrySetStaticSwitchValuesFromClipboardFunctionInput(const UNiagaraClipboardFunctionInput& ClipboardFunctionInput);
 
 	NIAGARAEDITOR_API void SetStandardValuesFromClipboardFunctionInputs(const TArray<const UNiagaraClipboardFunctionInput*>& ClipboardFunctionInputs);
 
