- Set RDG internal access mode on RTs
- This fixes an issue when an external RT is used in one pass, set to external access, then used again in another pass

diff --git a/Engine/Plugins/FX/Niagara/Source/Niagara/Private/NiagaraDataInterfaceIntRenderTarget2D.cpp b/Engine/Plugins/FX/Niagara/Source/Niagara/Private/NiagaraDataInterfaceIntRenderTarget2D.cpp
--- a/Engine/Plugins/FX/Niagara/Source/Niagara/Private/NiagaraDataInterfaceIntRenderTarget2D.cpp
+++ b/Engine/Plugins/FX/Niagara/Source/Niagara/Private/NiagaraDataInterfaceIntRenderTarget2D.cpp
@@ -587,6 +587,8 @@ void UNiagaraDataInterfaceIntRenderTarget2D::SetShaderParameters(const FNiagaraD
 				InstanceData.TransientRDGTexture = GraphBuilder.RegisterExternalTexture(CreateRenderTarget(InstanceData.TextureRHI, TEXT("NiagaraIntRenderTarget2D")));
 			}
 			InstanceData.TransientRDGUAV = GraphBuilder.CreateUAV(InstanceData.TransientRDGTexture);
+
+			GraphBuilder.UseInternalAccessMode(InstanceData.TransientRDGTexture);
 			Context.GetRDGExternalAccessQueue().Add(InstanceData.TransientRDGTexture);
 		}
 
diff --git a/Engine/Plugins/FX/Niagara/Source/Niagara/Private/NiagaraDataInterfaceRenderTarget2D.cpp b/Engine/Plugins/FX/Niagara/Source/Niagara/Private/NiagaraDataInterfaceRenderTarget2D.cpp
--- a/Engine/Plugins/FX/Niagara/Source/Niagara/Private/NiagaraDataInterfaceRenderTarget2D.cpp
+++ b/Engine/Plugins/FX/Niagara/Source/Niagara/Private/NiagaraDataInterfaceRenderTarget2D.cpp
@@ -396,6 +396,8 @@ void UNiagaraDataInterfaceRenderTarget2D::SetShaderParameters(const FNiagaraData
 		}
 		InstanceData_RT->TransientRDGSRV = GraphBuilder.CreateSRV(InstanceData_RT->TransientRDGTexture);
 		InstanceData_RT->TransientRDGUAV = GraphBuilder.CreateUAV(InstanceData_RT->TransientRDGTexture);
+
+		GraphBuilder.UseInternalAccessMode(InstanceData_RT->TransientRDGTexture);
 		Context.GetRDGExternalAccessQueue().Add(InstanceData_RT->TransientRDGTexture);
 	}
 
diff --git a/Engine/Plugins/FX/Niagara/Source/Niagara/Private/NiagaraDataInterfaceRenderTarget2DArray.cpp b/Engine/Plugins/FX/Niagara/Source/Niagara/Private/NiagaraDataInterfaceRenderTarget2DArray.cpp
--- a/Engine/Plugins/FX/Niagara/Source/Niagara/Private/NiagaraDataInterfaceRenderTarget2DArray.cpp
+++ b/Engine/Plugins/FX/Niagara/Source/Niagara/Private/NiagaraDataInterfaceRenderTarget2DArray.cpp
@@ -395,6 +395,8 @@ void UNiagaraDataInterfaceRenderTarget2DArray::SetShaderParameters(const FNiagar
 		InstanceData_RT->TransientRDGTexture = GraphBuilder.RegisterExternalTexture(InstanceData_RT->RenderTarget);
 		InstanceData_RT->TransientRDGSRV = GraphBuilder.CreateSRV(InstanceData_RT->TransientRDGTexture);
 		InstanceData_RT->TransientRDGUAV = GraphBuilder.CreateUAV(InstanceData_RT->TransientRDGTexture);
+
+		GraphBuilder.UseInternalAccessMode(InstanceData_RT->TransientRDGTexture);
 		Context.GetRDGExternalAccessQueue().Add(InstanceData_RT->TransientRDGTexture);
 	}
 
diff --git a/Engine/Plugins/FX/Niagara/Source/Niagara/Private/NiagaraDataInterfaceRenderTargetCube.cpp b/Engine/Plugins/FX/Niagara/Source/Niagara/Private/NiagaraDataInterfaceRenderTargetCube.cpp
--- a/Engine/Plugins/FX/Niagara/Source/Niagara/Private/NiagaraDataInterfaceRenderTargetCube.cpp
+++ b/Engine/Plugins/FX/Niagara/Source/Niagara/Private/NiagaraDataInterfaceRenderTargetCube.cpp
@@ -360,6 +360,8 @@ void UNiagaraDataInterfaceRenderTargetCube::SetShaderParameters(const FNiagaraDa
 		InstanceData_RT->TransientRDGTexture = GraphBuilder.RegisterExternalTexture(InstanceData_RT->RenderTarget);
 		InstanceData_RT->TransientRDGSRV = GraphBuilder.CreateSRV(InstanceData_RT->TransientRDGTexture);
 		InstanceData_RT->TransientRDGUAV = GraphBuilder.CreateUAV(InstanceData_RT->TransientRDGTexture);
+
+		GraphBuilder.UseInternalAccessMode(InstanceData_RT->TransientRDGTexture);
 		Context.GetRDGExternalAccessQueue().Add(InstanceData_RT->TransientRDGTexture);
 	}
 
diff --git a/Engine/Plugins/FX/Niagara/Source/Niagara/Private/NiagaraDataInterfaceRenderTargetVolume.cpp b/Engine/Plugins/FX/Niagara/Source/Niagara/Private/NiagaraDataInterfaceRenderTargetVolume.cpp
--- a/Engine/Plugins/FX/Niagara/Source/Niagara/Private/NiagaraDataInterfaceRenderTargetVolume.cpp
+++ b/Engine/Plugins/FX/Niagara/Source/Niagara/Private/NiagaraDataInterfaceRenderTargetVolume.cpp
@@ -457,6 +457,7 @@ void UNiagaraDataInterfaceRenderTargetVolume::SetShaderParameters(const FNiagara
 		InstanceData_RT->TransientRDGTexture = GraphBuilder.RegisterExternalTexture(InstanceData_RT->RenderTarget);
 		InstanceData_RT->TransientRDGSRV = GraphBuilder.CreateSRV(InstanceData_RT->TransientRDGTexture);
 		InstanceData_RT->TransientRDGUAV = GraphBuilder.CreateUAV(InstanceData_RT->TransientRDGTexture);
+		GraphBuilder.UseInternalAccessMode(InstanceData_RT->TransientRDGTexture);
 		Context.GetRDGExternalAccessQueue().Add(InstanceData_RT->TransientRDGTexture);
 	}
 
