fix issues with packaged builds containing Volume Renderers

diff --git a/Engine/Plugins/FX/Niagara/Source/Niagara/Private/NiagaraRendererVolumes.cpp b/Engine/Plugins/FX/Niagara/Source/Niagara/Private/NiagaraRendererVolumes.cpp
--- a/Engine/Plugins/FX/Niagara/Source/Niagara/Private/NiagaraRendererVolumes.cpp
+++ b/Engine/Plugins/FX/Niagara/Source/Niagara/Private/NiagaraRendererVolumes.cpp
@@ -133,9 +133,8 @@ void FNiagaraRendererVolumes::ReleaseRenderThreadResources()
 }
 
 bool FNiagaraRendererVolumes::IsMaterialValid(const UMaterialInterface* Material) const
-{
-	const UMaterial* BaseMaterial = Material ? Material->GetMaterial() : nullptr;
-	return BaseMaterial ? BaseMaterial->CheckMaterialUsage_Concurrent(MATUSAGE_HeterogeneousVolumes) : false;
+{	
+	return Material ? Material->CheckMaterialUsage_Concurrent(MATUSAGE_HeterogeneousVolumes) : false;
 }
 
 FNiagaraDynamicDataBase* FNiagaraRendererVolumes::GenerateDynamicData(const FNiagaraSceneProxy* Proxy, const UNiagaraRendererProperties* InProperties, const FNiagaraEmitterInstance* Emitter) const
