- Add auto attachment information to the HUD

diff --git a/Engine/Plugins/FX/Niagara/Source/Niagara/Private/NiagaraDebugHud.cpp b/Engine/Plugins/FX/Niagara/Source/Niagara/Private/NiagaraDebugHud.cpp
--- a/Engine/Plugins/FX/Niagara/Source/Niagara/Private/NiagaraDebugHud.cpp
+++ b/Engine/Plugins/FX/Niagara/Source/Niagara/Private/NiagaraDebugHud.cpp
@@ -2877,6 +2877,11 @@ void FNiagaraDebugHud::DrawComponents(FNiagaraWorldManager* WorldManager, UCanva
 						StringBuilder.Append(TEXT("\n"));
 					}
 
+					if (NiagaraComponent->bAutoManageAttachment)
+					{
+						StringBuilder.Appendf(TEXT("Auto Attachment - Parent(%s) Socket(%s)\n"), *GetNameSafe(NiagaraComponent->AutoAttachParent.Get()), *NiagaraComponent->AutoAttachSocketName.ToString());
+					}
+
 					if (UNiagaraSimCache* SimCache = NiagaraComponent->GetSimCache())
 					{
 						StringBuilder.Appendf(TEXT("SimCache(%s)\n"), *GetFullNameSafe(SimCache));
