- Don't detach components when inside non game worlds
- Don't call Modify on the attach parent when we detach a Cascade / Niagara component due to system completion
- This fixes an issue where levels are being marked dirty if they have one shot FX inside them that play on level load

diff --git a/Engine/Plugins/FX/Niagara/Source/Niagara/Private/NiagaraComponent.cpp b/Engine/Plugins/FX/Niagara/Source/Niagara/Private/NiagaraComponent.cpp
--- a/Engine/Plugins/FX/Niagara/Source/Niagara/Private/NiagaraComponent.cpp
+++ b/Engine/Plugins/FX/Niagara/Source/Niagara/Private/NiagaraComponent.cpp
@@ -1864,9 +1864,7 @@ void UNiagaraComponent::OnRegister()
 				// Only detach if we are not about to auto attach to the same target, that would be wasteful.
 				if (!bAutoActivate || (AutoAttachLocationRule != EAttachmentRule::KeepRelative && AutoAttachRotationRule != EAttachmentRule::KeepRelative && AutoAttachScaleRule != EAttachmentRule::KeepRelative) || (AutoAttachSocketName != GetAttachSocketName()) || (AutoAttachParent != GetAttachParent()))
 				{
-					//bIsChangingAutoAttachment = true;
 					DetachFromComponent(FDetachmentTransformRules(EDetachmentRule::KeepRelative, /*bCallModify=*/ false));
-					//bIsChangingAutoAttachment = false;
 				}
 			}
 			else
@@ -4145,9 +4143,11 @@ void UNiagaraComponent::CancelAutoAttachment(bool bDetachFromParent)
 
 		if (bDetachFromParent)
 		{
-			//bIsChangingAutoAttachment = true;
-			DetachFromComponent(FDetachmentTransformRules::KeepRelativeTransform);
-			//bIsChangingAutoAttachment = false;
+			UWorld* World = GetWorld();
+			if (!World || World->IsGameWorld())
+			{
+				DetachFromComponent(FDetachmentTransformRules(EDetachmentRule::KeepRelative, /*bCallModify=*/ false));
+			}
 		}
 	}
 }
diff --git a/Engine/Source/Runtime/Engine/Private/Particles/ParticleComponents.cpp b/Engine/Source/Runtime/Engine/Private/Particles/ParticleComponents.cpp
--- a/Engine/Source/Runtime/Engine/Private/Particles/ParticleComponents.cpp
+++ b/Engine/Source/Runtime/Engine/Private/Particles/ParticleComponents.cpp
@@ -6571,7 +6571,11 @@ void UParticleSystemComponent::CancelAutoAttachment(bool bDetachFromParent, cons
 
 		if (bDetachFromParent)
 		{
-			DetachFromComponent(FDetachmentTransformRules::KeepRelativeTransform);
+			UWorld* World = GetWorld();
+			if (!World || World->IsGameWorld())
+			{
+				DetachFromComponent(FDetachmentTransformRules(EDetachmentRule::KeepRelative, /*bCallModify=*/ false));
+			}
 		}
 	}
 }
