Fix for default value for flag added to landscape DI (so that we will, by default, allow sampling of the RVT if available).

diff --git a/Engine/Plugins/FX/Niagara/Source/Niagara/Classes/NiagaraDataInterfaceLandscape.h b/Engine/Plugins/FX/Niagara/Source/Niagara/Classes/NiagaraDataInterfaceLandscape.h
--- a/Engine/Plugins/FX/Niagara/Source/Niagara/Classes/NiagaraDataInterfaceLandscape.h
+++ b/Engine/Plugins/FX/Niagara/Source/Niagara/Classes/NiagaraDataInterfaceLandscape.h
@@ -51,7 +51,7 @@ public:
 
 	/** Can be used to ignore virtual textures even if they are defined for the landscape. */
 	UPROPERTY(EditAnywhere, Category = "Landscape")
-	bool bVirtualTexturesSupported;
+	bool bVirtualTexturesSupported = true;
 
 	//UObject Interface
 	NIAGARA_API virtual void PostInitProperties() override;	
