Niagara - Don't try processing a null pin when post loading a niagara node.  This will likely break elsewhere, but the null check and then use is triggering the the static analysis tools.

diff --git a/Engine/Plugins/FX/Niagara/Source/NiagaraEditor/Private/NiagaraNode.cpp b/Engine/Plugins/FX/Niagara/Source/NiagaraEditor/Private/NiagaraNode.cpp
--- a/Engine/Plugins/FX/Niagara/Source/NiagaraEditor/Private/NiagaraNode.cpp
+++ b/Engine/Plugins/FX/Niagara/Source/NiagaraEditor/Private/NiagaraNode.cpp
@@ -48,6 +48,7 @@ void UNiagaraNode::PostLoad()
 			FString AssetPath = OwningGraph ? OwningGraph->GetPathName() : "Undetermined";
 			FString NodeName = GetNodeTitle(ENodeTitleType::FullTitle).ToString();
 			UE_LOG(LogNiagaraEditor, Log, TEXT("Pin of node %s in asset %s is nullptr!"), *NodeName, *AssetPath);
+			continue;
 		}
 		if(Pin->LinkedTo.Num() > 0)
 		{
