- Include Cascade in debug hud bounds rendering

diff --git a/Engine/Plugins/FX/Niagara/Source/Niagara/Private/NiagaraDebugHud.cpp b/Engine/Plugins/FX/Niagara/Source/Niagara/Private/NiagaraDebugHud.cpp
--- a/Engine/Plugins/FX/Niagara/Source/Niagara/Private/NiagaraDebugHud.cpp
+++ b/Engine/Plugins/FX/Niagara/Source/Niagara/Private/NiagaraDebugHud.cpp
@@ -1059,9 +1059,9 @@ void FNiagaraDebugHud::GatherSystemInfo()
 				bIsMatch &= FXComponent->GetName().MatchesWildcard(Settings.ComponentFilter);
 			}
 
-			if (bIsMatch && NiagaraComponent)//TODO: Handle Cascade?
+			if (bIsMatch)
 			{
-				InWorldComponents.Add(NiagaraComponent);
+				InWorldComponents.Add(FXComponent);
 			}
 		}
 
@@ -2497,9 +2497,9 @@ void FNiagaraDebugHud::DrawValidation(class FNiagaraWorldManager* WorldManager,
 		return;
 	}
 
-	for (TWeakObjectPtr<UNiagaraComponent> WeakComponent : InWorldComponents)
+	for (TWeakObjectPtr<UFXSystemComponent> WeakComponent : InWorldComponents)
 	{
-		UNiagaraComponent* NiagaraComponent = WeakComponent.Get();
+		UNiagaraComponent* NiagaraComponent = Cast<UNiagaraComponent>(WeakComponent.Get());
 		if (NiagaraComponent == nullptr)
 		{
 			continue;
@@ -2681,36 +2681,42 @@ void FNiagaraDebugHud::DrawComponents(FNiagaraWorldManager* WorldManager, UCanva
 	UEnum* ExecutionStateEnum = StaticEnum<ENiagaraExecutionState>();
 	UEnum* PoolingMethodEnum = StaticEnum<ENCPoolMethod>();
 	UEnum* SystemInstanceState = StaticEnum<ENiagaraSystemInstanceState>();
-	for (TWeakObjectPtr<UNiagaraComponent> WeakComponent : InWorldComponents)
+	for (TWeakObjectPtr<UFXSystemComponent> WeakComponent : InWorldComponents)
 	{
-		UNiagaraComponent* NiagaraComponent = WeakComponent.Get();
-		if (NiagaraComponent == nullptr)
-		{
-			continue;
-		}
-
-		UNiagaraSystem* NiagaraSystem = NiagaraComponent->GetAsset();
-		FNiagaraSystemInstanceControllerPtr SystemInstanceController = NiagaraComponent->GetSystemInstanceController();
-		FNiagaraSystemInstance* SystemInstance = SystemInstanceController.IsValid() ? SystemInstanceController->GetSystemInstance_Unsafe() : nullptr;
-		if (NiagaraSystem == nullptr || SystemInstance == nullptr)
+		UFXSystemComponent* FXComponent = WeakComponent.Get();
+		if (FXComponent == nullptr)
 		{
 			continue;
 		}
 
-		const FVector ComponentLocation = NiagaraComponent->GetComponentLocation();
-		const FRotator ComponentRotation = NiagaraComponent->GetComponentRotation();
-		const bool bIsActive = NiagaraComponent->IsActive();
+		const FVector ComponentLocation = FXComponent->GetComponentLocation();
+		const FRotator ComponentRotation = FXComponent->GetComponentRotation();
+		const bool bIsActive = FXComponent->IsActive();
 
 		// Show system bounds (only active components)
 		if (Settings.bSystemShowBounds && bIsActive)
 		{
-			const FBox Bounds = NiagaraComponent->CalcBounds(NiagaraComponent->GetComponentTransform()).GetBox();
+			const FBox Bounds = FXComponent->CalcBounds(FXComponent->GetComponentTransform()).GetBox();
 			if (Bounds.IsValid)
 			{
 				DrawBox(World, Bounds.GetCenter(), Bounds.GetExtent(), FColor::Red, Settings.SystemBoundsSolidBoxAlpha);
 			}
 		}
 
+		UNiagaraComponent* NiagaraComponent = Cast<UNiagaraComponent>(FXComponent);
+		if (NiagaraComponent == nullptr)
+		{
+			continue;
+		}
+
+		UNiagaraSystem* NiagaraSystem = NiagaraComponent->GetAsset();
+		FNiagaraSystemInstanceControllerPtr SystemInstanceController = NiagaraComponent->GetSystemInstanceController();
+		FNiagaraSystemInstance* SystemInstance = SystemInstanceController.IsValid() ? SystemInstanceController->GetSystemInstance_Unsafe() : nullptr;
+		if (NiagaraSystem == nullptr || SystemInstance == nullptr)
+		{
+			continue;
+		}
+
 		// Get system simulation
 		auto SystemSimulation = SystemInstance->GetSystemSimulation();
 		const bool bSystemSimulationValid = SystemSimulation.IsValid() && SystemSimulation->IsValid();
diff --git a/Engine/Plugins/FX/Niagara/Source/Niagara/Private/NiagaraDebugHud.h b/Engine/Plugins/FX/Niagara/Source/Niagara/Private/NiagaraDebugHud.h
--- a/Engine/Plugins/FX/Niagara/Source/Niagara/Private/NiagaraDebugHud.h
+++ b/Engine/Plugins/FX/Niagara/Source/Niagara/Private/NiagaraDebugHud.h
@@ -267,7 +267,7 @@ private:
 
 	TMap<FName, FSystemDebugInfo>	PerSystemDebugInfo;
 
-	TArray<TWeakObjectPtr<class UNiagaraComponent>>	InWorldComponents;
+	TArray<TWeakObjectPtr<class UFXSystemComponent>>	InWorldComponents;
 
 	TMap<TWeakObjectPtr<class UNiagaraComponent>, FValidationErrorInfo> ValidationErrors;
 
