- While active always include the calculated bounds when using auto attachments
- This fixes an issue where the effect would no longer be visible due to incorrect bounds if we are attached, become detached (potentially forced via BP) and wish the effect to continue playing

diff --git a/Engine/Plugins/FX/Niagara/Source/Niagara/Private/NiagaraComponent.cpp b/Engine/Plugins/FX/Niagara/Source/Niagara/Private/NiagaraComponent.cpp
--- a/Engine/Plugins/FX/Niagara/Source/Niagara/Private/NiagaraComponent.cpp
+++ b/Engine/Plugins/FX/Niagara/Source/Niagara/Private/NiagaraComponent.cpp
@@ -2234,8 +2234,10 @@ FBoxSphereBounds UNiagaraComponent::CalcBounds(const FTransform& LocalToWorld) c
 {
 	FScopeCycleCounter SystemStatCounter(Asset ? Asset->GetStatID(true, false) : TStatId());
 
+	// When inactive and using auto attachments do not include our bounds as they will be in an invalid location
+	// While active it's more complicated as we could become detatched and wish to play the remainder of the effect so we must include them
 	const USceneComponent* UseAutoParent = (bAutoManageAttachment && GetAttachParent() == nullptr) ? AutoAttachParent.Get() : nullptr;
-	if (UseAutoParent)
+	if (UseAutoParent && !IsActive())
 	{
 		// We use auto attachment but have detached, don't use our own bogus bounds (we're off near 0,0,0), use the usual parent's bounds.
 		return UseAutoParent->Bounds;
diff --git a/Engine/Source/Runtime/Engine/Private/Particles/ParticleComponents.cpp b/Engine/Source/Runtime/Engine/Private/Particles/ParticleComponents.cpp
--- a/Engine/Source/Runtime/Engine/Private/Particles/ParticleComponents.cpp
+++ b/Engine/Source/Runtime/Engine/Private/Particles/ParticleComponents.cpp
@@ -4974,8 +4974,10 @@ FBoxSphereBounds UParticleSystemComponent::CalcBounds(const FTransform& LocalToW
 	FBox BoundingBox;
 	BoundingBox.Init();
 
+	// When inactive and using auto attachments do not include our bounds as they will be in an invalid location
+	// While active it's more complicated as we could become detatched and wish to play the remainder of the effect so we must include them
 	const USceneComponent* UseAutoParent = (bAutoManageAttachment && GetAttachParent() == nullptr) ? AutoAttachParent.Get() : nullptr;
-	if (UseAutoParent)
+	if (UseAutoParent && !IsActive())
 	{
 		// We use auto attachment but have detached, don't use our own bogus bounds (we're off near 0,0,0), use the usual parent's bounds.
 		return UseAutoParent->Bounds;
