Minor post load changes.
1) Ensuring all nodes are post loaded before we to further fixup
2) Calling PostLoad directly on script variables that may have been postloaded before

diff --git a/Engine/Plugins/FX/Niagara/Source/NiagaraEditor/Private/NiagaraGraph.cpp b/Engine/Plugins/FX/Niagara/Source/NiagaraEditor/Private/NiagaraGraph.cpp
--- a/Engine/Plugins/FX/Niagara/Source/NiagaraEditor/Private/NiagaraGraph.cpp
+++ b/Engine/Plugins/FX/Niagara/Source/NiagaraEditor/Private/NiagaraGraph.cpp
@@ -675,7 +675,10 @@ void UNiagaraGraph::PostLoad()
 		for (UEdGraphNode* Node : Nodes)
 		{
 			Node->ConditionalPostLoad();
-
+		}
+		
+		for (UEdGraphNode* Node : Nodes)
+		{
 			if (UNiagaraNode* NiagaraNode = Cast<UNiagaraNode>(Node))
 			{
 				// Assume that all externally referenced assets have changed, so update to match. They will return true if they have changed.
@@ -684,7 +687,7 @@ void UNiagaraGraph::PostLoad()
 					ReferencedAsset->ConditionalPostLoad();
 					NiagaraNode->RefreshFromExternalChanges();
 				}
-
+				
 				if (UNiagaraNodeInput* InputNode = Cast<UNiagaraNodeInput>(NiagaraNode))
 				{
 					InputNodes.Add(InputNode);
@@ -1469,7 +1472,7 @@ void UNiagaraGraph::StandardizeParameterNames()
 	{
 		FNiagaraVariable Variable = VariableScriptVariablePair.Key;
 		UNiagaraScriptVariable* ScriptVariable = VariableScriptVariablePair.Value;
-		ScriptVariable->PostLoad();
+		ScriptVariable->ConditionalPostLoad();
 		
 		if (ScriptVariable->GetIsStaticSwitch() == false)
 		{
