Added sanity check & some conditional post load calls

diff --git a/Engine/Plugins/FX/Niagara/Source/NiagaraEditor/Private/NiagaraEmitterEditorData.cpp b/Engine/Plugins/FX/Niagara/Source/NiagaraEditor/Private/NiagaraEmitterEditorData.cpp
--- a/Engine/Plugins/FX/Niagara/Source/NiagaraEditor/Private/NiagaraEmitterEditorData.cpp
+++ b/Engine/Plugins/FX/Niagara/Source/NiagaraEditor/Private/NiagaraEmitterEditorData.cpp
@@ -380,6 +380,7 @@ void UNiagaraEmitterEditorData::PostLoad_TransferModuleStackNotesToNewFormat(UOb
 		{
 			if (UNiagaraScriptSource* GraphSource = Cast<UNiagaraScriptSource>(EmitterData->GraphSource))
 			{
+				GraphSource->ConditionalPostLoad();
 				EmitterGraph = GraphSource->NodeGraph;
 			}
 		}
diff --git a/Engine/Plugins/FX/Niagara/Source/NiagaraEditor/Private/NiagaraSystemEditorData.cpp b/Engine/Plugins/FX/Niagara/Source/NiagaraEditor/Private/NiagaraSystemEditorData.cpp
--- a/Engine/Plugins/FX/Niagara/Source/NiagaraEditor/Private/NiagaraSystemEditorData.cpp
+++ b/Engine/Plugins/FX/Niagara/Source/NiagaraEditor/Private/NiagaraSystemEditorData.cpp
@@ -239,24 +239,29 @@ void UNiagaraSystemEditorData::PostLoad_TransferModuleStackNotesToNewFormat(UObj
 	if(UNiagaraScript* SystemSpawnScript = System->GetSystemSpawnScript())
 	{
 		SystemSpawnScript->ConditionalPostLoad();
-		SourceBases.Add(SystemSpawnScript->GetLatestSource());
+		
+		if (UNiagaraScriptSource* GraphSource = Cast<UNiagaraScriptSource>(SystemSpawnScript->GetLatestSource()))
+        {
+			GraphSource->ConditionalPostLoad();
+			SourceBases.Add(GraphSource);
+        }
 	}
 	
 	if(UNiagaraScript* SystemUpdateScript = System->GetSystemUpdateScript())
 	{
 		SystemUpdateScript->ConditionalPostLoad();
-		SourceBases.Add(SystemUpdateScript->GetLatestSource());
+
+		if (UNiagaraScriptSource* GraphSource = Cast<UNiagaraScriptSource>(SystemUpdateScript->GetLatestSource()))
+		{
+			GraphSource->ConditionalPostLoad();
+			SourceBases.Add(GraphSource);
+		}
 	}
 	
 	for(UNiagaraScriptSourceBase* SourceBase : SourceBases)
-	{
+	{		
 		UNiagaraScriptSource* Source = CastChecked<UNiagaraScriptSource>(SourceBase);
 
-		if(Source->NodeGraph == nullptr)
-		{
-			continue;
-		}
-
 		TArray<UNiagaraNodeFunctionCall*> FunctionCallNodes;
 		Source->NodeGraph->GetNodesOfClass(FunctionCallNodes);
 
