Unhid system asset options

diff --git a/Engine/Plugins/FX/Niagara/Source/NiagaraEditor/Private/NiagaraSystemDetailsCustomization.cpp b/Engine/Plugins/FX/Niagara/Source/NiagaraEditor/Private/NiagaraSystemDetailsCustomization.cpp
--- a/Engine/Plugins/FX/Niagara/Source/NiagaraEditor/Private/NiagaraSystemDetailsCustomization.cpp
+++ b/Engine/Plugins/FX/Niagara/Source/NiagaraEditor/Private/NiagaraSystemDetailsCustomization.cpp
@@ -16,7 +16,6 @@ void FNiagaraSystemDetails::CustomizeDetails(IDetailLayoutBuilder& InDetailLayou
 {
 	// we display the scalability category within scalability mode, which is why we hide it here
 	InDetailLayout.HideCategory("Scalability");
-	InDetailLayout.HideCategory("Asset Options");
 
 	InDetailLayout.SortCategories(
 		[](const TMap<FName, IDetailCategoryBuilder*>& CategoryMap)
