- Niagara renderer can now provide a streaming scale for materials
- This fixes an issue with sub image streaming for sprite renderers where the calculation did not take into account that the UVs would be scaled

diff --git a/Engine/Plugins/FX/Niagara/Source/Niagara/Private/NiagaraComponent.cpp b/Engine/Plugins/FX/Niagara/Source/Niagara/Private/NiagaraComponent.cpp
--- a/Engine/Plugins/FX/Niagara/Source/Niagara/Private/NiagaraComponent.cpp
+++ b/Engine/Plugins/FX/Niagara/Source/Niagara/Private/NiagaraComponent.cpp
@@ -2318,11 +2318,28 @@ void UNiagaraComponent::GetUsedMaterials(TArray<UMaterialInterface*>& OutMateria
 
 void UNiagaraComponent::GetStreamingRenderAssetInfo(FStreamingTextureLevelContext& LevelContext, TArray<FStreamingRenderAssetPrimitiveInfo>& OutStreamingRenderAssets) const
 {
-	// UPrimitiveComponent::GetStreamingRenderAssetInfo takes care of getting the referenced textures
-	Super::GetStreamingRenderAssetInfo(LevelContext, OutStreamingRenderAssets);
-
 	if (SystemInstanceController.IsValid())
 	{
+		FNiagaraMaterialAndScaleArray MaterialAndScales;
+		SystemInstanceController->GetMaterialStreamingInfo(MaterialAndScales);
+
+		LevelContext.BindBuildData(nullptr);
+		if (MaterialAndScales.Num() > 0)
+		{
+			static const FMeshUVChannelInfo UVChannelData(1.f);
+
+			FPrimitiveMaterialInfo MaterialData;
+			MaterialData.PackedRelativeBox = PackedRelativeBox_Identity;
+			MaterialData.UVChannelData = &UVChannelData;
+
+			for (const FNiagaraMaterialAndScale& MaterialAndScale : MaterialAndScales)
+			{
+				check(MaterialAndScale.Material != nullptr);
+				MaterialData.Material = MaterialAndScale.Material;
+				LevelContext.ProcessMaterial(Bounds, MaterialData, Bounds.SphereRadius * MaterialAndScale.Scale, OutStreamingRenderAssets, bIsValidTextureStreamingBuiltData, this);
+			}
+		}
+
 		SystemInstanceController->GetStreamingMeshInfo(Bounds, LevelContext, OutStreamingRenderAssets);
 	}
 }
diff --git a/Engine/Plugins/FX/Niagara/Source/Niagara/Private/NiagaraSystemInstanceController.cpp b/Engine/Plugins/FX/Niagara/Source/Niagara/Private/NiagaraSystemInstanceController.cpp
--- a/Engine/Plugins/FX/Niagara/Source/Niagara/Private/NiagaraSystemInstanceController.cpp
+++ b/Engine/Plugins/FX/Niagara/Source/Niagara/Private/NiagaraSystemInstanceController.cpp
@@ -114,16 +114,14 @@ void FNiagaraSystemInstanceController::GetUsedMaterials(TArray<UMaterialInterfac
 			EmitterData->ForEachEnabledRenderer(
 				[&](UNiagaraRendererProperties* Properties)
 				{
-					bool bCreateMidsForUsedMaterials = Properties->NeedsMIDsForMaterials();
 					TArray<UMaterialInterface*> Mats;
 					Properties->GetUsedMaterials(&EmitterInst.Get(), Mats);
 
-					if (bCreateMidsForUsedMaterials)
+					if (Properties->NeedsMIDsForMaterials())
 					{
 						for (const FMaterialOverride& Override : EmitterMaterials)
 						{
-							if (Override.EmitterRendererProperty == Properties &&
-								Mats.IsValidIndex(Override.MaterialSubIndex))
+							if (Override.EmitterRendererProperty == Properties && Mats.IsValidIndex(Override.MaterialSubIndex))
 							{
 								Mats[Override.MaterialSubIndex] = Override.Material;
 							}
@@ -137,6 +135,61 @@ void FNiagaraSystemInstanceController::GetUsedMaterials(TArray<UMaterialInterfac
 	}
 }
 
+void FNiagaraSystemInstanceController::GetMaterialStreamingInfo(FNiagaraMaterialAndScaleArray& OutMaterialAndScales) const
+{
+	if (!SystemInstance.IsValid())
+	{
+		return;
+	}
+
+	for (const TSharedRef<FNiagaraEmitterInstance, ESPMode::ThreadSafe>& EmitterInst : SystemInstance->GetEmitters())
+	{
+		if (FVersionedNiagaraEmitterData* EmitterData = EmitterInst->GetCachedEmitterData())
+		{
+			EmitterData->ForEachEnabledRenderer(
+				[&](UNiagaraRendererProperties* Properties)
+				{
+					TArray<UMaterialInterface*> UsedMaterials;
+					Properties->GetUsedMaterials(&EmitterInst.Get(), UsedMaterials);
+					if (UsedMaterials.Num() == 0)
+					{
+						return;
+					}
+
+					if (Properties->NeedsMIDsForMaterials())
+					{
+						for (const FMaterialOverride& Override : EmitterMaterials)
+						{
+							if (Override.EmitterRendererProperty == Properties && UsedMaterials.IsValidIndex(Override.MaterialSubIndex))
+							{
+								UsedMaterials[Override.MaterialSubIndex] = Override.Material;
+							}
+						}
+					}
+
+					const float StreamingScale = Properties->GetMaterialStreamingScale();
+					for (UMaterialInterface* UsedMaterial : UsedMaterials)
+					{
+						if (UsedMaterial == nullptr)
+						{
+							continue;
+						}
+
+						if (FNiagaraMaterialAndScale* Existing = OutMaterialAndScales.FindByPredicate([UsedMaterial](const FNiagaraMaterialAndScale& Existing) { return Existing.Material == UsedMaterial; }))
+						{
+							Existing->Scale = FMath::Max(Existing->Scale, StreamingScale);
+						}
+						else
+						{
+							OutMaterialAndScales.Emplace(UsedMaterial, StreamingScale);
+						}
+					}
+				}
+			);
+		}
+	}
+}
+
 void FNiagaraSystemInstanceController::GetStreamingMeshInfo(const FBoxSphereBounds& OwnerBounds, FStreamingTextureLevelContext& LevelContext, TArray<FStreamingRenderAssetPrimitiveInfo>& OutStreamingRenderAssets) const
 {
 	if (!SystemInstance.IsValid())
diff --git a/Engine/Plugins/FX/Niagara/Source/Niagara/Public/NiagaraRendererProperties.h b/Engine/Plugins/FX/Niagara/Source/Niagara/Public/NiagaraRendererProperties.h
--- a/Engine/Plugins/FX/Niagara/Source/Niagara/Public/NiagaraRendererProperties.h
+++ b/Engine/Plugins/FX/Niagara/Source/Niagara/Public/NiagaraRendererProperties.h
@@ -308,6 +308,7 @@ public:
 	virtual const FVertexFactoryType* GetVertexFactoryType() const { return nullptr; }
 	virtual bool IsBackfaceCullingDisabled() const { return false; }
 
+	virtual float GetMaterialStreamingScale() const { return 1.0f; }
 	virtual void GetStreamingMeshInfo(const FBoxSphereBounds& OwnerBounds, const FNiagaraEmitterInstance* InEmitter, TArray<FStreamingRenderAssetPrimitiveInfo>& OutStreamingRenderAssets) const {}
 
 	virtual bool IsSimTargetSupported(ENiagaraSimTarget InSimTarget) const { return false; };
diff --git a/Engine/Plugins/FX/Niagara/Source/Niagara/Public/NiagaraSpriteRendererProperties.h b/Engine/Plugins/FX/Niagara/Source/Niagara/Public/NiagaraSpriteRendererProperties.h
--- a/Engine/Plugins/FX/Niagara/Source/Niagara/Public/NiagaraSpriteRendererProperties.h
+++ b/Engine/Plugins/FX/Niagara/Source/Niagara/Public/NiagaraSpriteRendererProperties.h
@@ -133,6 +133,7 @@ public:
 	NIAGARA_API virtual FNiagaraRenderer* CreateEmitterRenderer(ERHIFeatureLevel::Type FeatureLevel, const FNiagaraEmitterInstance* Emitter, const FNiagaraSystemInstanceController& InController) override;
 	NIAGARA_API virtual class FNiagaraBoundsCalculator* CreateBoundsCalculator() override;
 	NIAGARA_API virtual void GetUsedMaterials(const FNiagaraEmitterInstance* InEmitter, TArray<UMaterialInterface*>& OutMaterials) const override;
+	NIAGARA_API virtual float GetMaterialStreamingScale() const override { return FMath::Max(FMath::Max(SubImageSize.X, SubImageSize.Y), 1.0f); }
 	NIAGARA_API virtual const FVertexFactoryType* GetVertexFactoryType() const override;
 	virtual bool IsSimTargetSupported(ENiagaraSimTarget InSimTarget) const override { return true; };
 	NIAGARA_API virtual bool PopulateRequiredBindings(FNiagaraParameterStore& InParameterStore)  override;
diff --git a/Engine/Plugins/FX/Niagara/Source/Niagara/Public/NiagaraSystemInstanceController.h b/Engine/Plugins/FX/Niagara/Source/Niagara/Public/NiagaraSystemInstanceController.h
--- a/Engine/Plugins/FX/Niagara/Source/Niagara/Public/NiagaraSystemInstanceController.h
+++ b/Engine/Plugins/FX/Niagara/Source/Niagara/Public/NiagaraSystemInstanceController.h
@@ -38,6 +38,17 @@ class FNiagaraSceneProxy;
 using FNiagaraSystemInstanceControllerPtr = TSharedPtr<FNiagaraSystemInstanceController, ESPMode::ThreadSafe>;
 using FNiagaraSystemInstanceControllerConstPtr = TSharedPtr<const FNiagaraSystemInstanceController, ESPMode::ThreadSafe>;
 
+struct FNiagaraMaterialAndScale
+{
+	FNiagaraMaterialAndScale() = default;
+	explicit FNiagaraMaterialAndScale(UMaterialInterface* InMaterial, float InScale) : Material(InMaterial), Scale(InScale) {}
+
+	UMaterialInterface* Material = nullptr;
+	float Scale = 1.0f;
+};
+
+using FNiagaraMaterialAndScaleArray = TArray<FNiagaraMaterialAndScale, TInlineAllocator<16>>;
+
 /**
  * This is the main asynchronous interface for controlling operation of a single instance of a Niagara System.
  */
@@ -81,6 +92,7 @@ public:
 	UMaterialInterface* GetMaterialOverride(const UNiagaraRendererProperties* InProps, int32 InMaterialSubIndex) const;
 	void SetOnMaterialsUpdated(const FOnMaterialsUpdated& Delegate) { OnMaterialsUpdatedDelegate = Delegate; }
 
+	void GetMaterialStreamingInfo(FNiagaraMaterialAndScaleArray& OutMaterialAndScales) const;
 	void GetStreamingMeshInfo(const FBoxSphereBounds& OwnerBounds, FStreamingTextureLevelContext& LevelContext, TArray<FStreamingRenderAssetPrimitiveInfo>& OutStreamingRenderAssets) const;
 
 	/** Dumps system instance state and info to the log (Used by fx.Niagara.DumpComponents commandlet) */
