Fix Niagara landscape DI sampling of RVT to use SRV so that we can properly handle the different texture formats that the data may be stored in (i.e. BC3_UNORM vs BC3_SRGB)

diff --git a/Engine/Plugins/FX/Niagara/Source/Niagara/Private/NiagaraDataInterfaceLandscape.cpp b/Engine/Plugins/FX/Niagara/Source/Niagara/Private/NiagaraDataInterfaceLandscape.cpp
--- a/Engine/Plugins/FX/Niagara/Source/Niagara/Private/NiagaraDataInterfaceLandscape.cpp
+++ b/Engine/Plugins/FX/Niagara/Source/Niagara/Private/NiagaraDataInterfaceLandscape.cpp
@@ -49,7 +49,7 @@ namespace NiagaraDataInterfaceLandscape
 	static const TCHAR* TemplateShaderFile = TEXT("/Plugin/FX/Niagara/Private/NiagaraDataInterfaceLandscape.ush");
 
 	BEGIN_SHADER_PARAMETER_STRUCT(FShaderParameters,)
-		SHADER_PARAMETER_TEXTURE(Texture2D,			BaseColorVirtualTexture)
+		SHADER_PARAMETER_SRV(Texture2D,				BaseColorVirtualTexture)
 		SHADER_PARAMETER_TEXTURE(Texture2D<uint4>,	BaseColorVirtualTexturePageTable)
 		SHADER_PARAMETER_SAMPLER(SamplerState,		BaseColorVirtualTextureSampler)
 		SHADER_PARAMETER(FVector3f,					BaseColorVirtualTextureLWCTile)
@@ -59,7 +59,7 @@ namespace NiagaraDataInterfaceLandscape
 		SHADER_PARAMETER(FUintVector4,				BaseColorVirtualTexturePackedUniform0)
 		SHADER_PARAMETER(FUintVector4,				BaseColorVirtualTexturePackedUniform1)
 		SHADER_PARAMETER(FUintVector4,				BaseColorVirtualTextureUniforms)
-		SHADER_PARAMETER_TEXTURE(Texture2D,			HeightVirtualTexture)
+		SHADER_PARAMETER_SRV(Texture2D,				HeightVirtualTexture)
 		SHADER_PARAMETER_TEXTURE(Texture2D<uint4>,	HeightVirtualTexturePageTable)
 		SHADER_PARAMETER_SAMPLER(SamplerState,		HeightVirtualTextureSampler)
 		SHADER_PARAMETER(FVector3f,					HeightVirtualTextureLWCTile)
@@ -68,8 +68,8 @@ namespace NiagaraDataInterfaceLandscape
 		SHADER_PARAMETER(FUintVector4,				HeightVirtualTexturePackedUniform0)
 		SHADER_PARAMETER(FUintVector4,				HeightVirtualTexturePackedUniform1)
 		SHADER_PARAMETER(FUintVector4,				HeightVirtualTextureUniforms)
-		SHADER_PARAMETER_TEXTURE(Texture2D,			NormalVirtualTexture0)
-		SHADER_PARAMETER_TEXTURE(Texture2D,			NormalVirtualTexture1)
+		SHADER_PARAMETER_SRV(Texture2D,				NormalVirtualTexture0)
+		SHADER_PARAMETER_SRV(Texture2D,				NormalVirtualTexture1)
 		SHADER_PARAMETER_TEXTURE(Texture2D<uint4>,	NormalVirtualTexturePageTable)
 		SHADER_PARAMETER_SAMPLER(SamplerState,		NormalVirtualTexture0Sampler)
 		SHADER_PARAMETER_SAMPLER(SamplerState,		NormalVirtualTexture1Sampler)
@@ -275,16 +275,16 @@ struct FNDILandscapeData_RenderThread
 	{
 		bool IsValid() const
 		{
-			return TextureRef.IsValid();
+			return TextureSRV.IsValid();
 		}
 
 		void Reset()
 		{
-			TextureRef = nullptr;
+			TextureSRV = nullptr;
 			TextureUniforms = FUintVector4(0, 0, 0, 0);
 		}
 
-		void Update(const URuntimeVirtualTexture* VirtualTexture, uint32 LayerIndex)
+		void Update(const URuntimeVirtualTexture* VirtualTexture, uint32 LayerIndex, bool bSRGB)
 		{
 			Reset();
 
@@ -292,16 +292,16 @@ struct FNDILandscapeData_RenderThread
 			{
 				if (IAllocatedVirtualTexture* AllocatedTexture = VirtualTexture->GetAllocatedVirtualTexture())
 				{
-					if (FRHITexture* PhysicalTexture = AllocatedTexture->GetPhysicalTexture(LayerIndex))
+					if (FRHIShaderResourceView* PhysicalTextureSRV = AllocatedTexture->GetPhysicalTextureSRV(LayerIndex, bSRGB))
 					{
-						TextureRef = PhysicalTexture;
+						TextureSRV = PhysicalTextureSRV;
 						AllocatedTexture->GetPackedUniform(&TextureUniforms, LayerIndex);
 					}
 				}
 			}
 		}
 
-		FTextureRHIRef TextureRef;
+		FShaderResourceViewRHIRef TextureSRV;
 		FUintVector4 TextureUniforms = FUintVector4(0, 0, 0, 0);
 	};
 
@@ -379,7 +379,7 @@ struct FNDILandscapeData_RenderThread
 			return false;
 		}
 
-		ShaderParameters->BaseColorVirtualTexture = BaseColorVirtualLayer.TextureRef;
+		ShaderParameters->BaseColorVirtualTexture = BaseColorVirtualLayer.TextureSRV;
 		ShaderParameters->BaseColorVirtualTexturePageTable = BaseColorVirtualPage.PageTableRef;
 
 		FLargeWorldRenderPosition BaseColorVirtualTextureOrigin(BaseColorVirtualPage.WorldToUvParameters[0]);
@@ -403,7 +403,7 @@ struct FNDILandscapeData_RenderThread
 	{
 		const FUintVector4 DummyUint4(ForceInitToZero);
 
-		ShaderParameters->BaseColorVirtualTexture = GBlackTexture->TextureRHI;
+		ShaderParameters->BaseColorVirtualTexture = GBlackTextureWithSRV->ShaderResourceViewRHI;
 		ShaderParameters->BaseColorVirtualTexturePageTable = GBlackUintTexture->TextureRHI;
 		ShaderParameters->BaseColorVirtualTextureWorldToUvTransform = FMatrix44f::Identity;
 		ShaderParameters->BaseColorVirtualTextureSRGB = 0;
@@ -420,7 +420,7 @@ struct FNDILandscapeData_RenderThread
 			return false;
 		}
 
-		ShaderParameters->HeightVirtualTexture = HeightVirtualLayer.TextureRef;
+		ShaderParameters->HeightVirtualTexture = HeightVirtualLayer.TextureSRV;
 		ShaderParameters->HeightVirtualTexturePageTable = HeightVirtualPage.PageTableRef;
 
 		FLargeWorldRenderPosition HeightVirtualTextureOrigin(HeightVirtualPage.WorldToUvParameters[0]);
@@ -444,7 +444,7 @@ struct FNDILandscapeData_RenderThread
 	{
 		const FUintVector4 DummyUint4(ForceInitToZero);
 
-		ShaderParameters->HeightVirtualTexture = GBlackTexture->TextureRHI;
+		ShaderParameters->HeightVirtualTexture = GBlackTextureWithSRV->ShaderResourceViewRHI;
 		ShaderParameters->HeightVirtualTexturePageTable = GBlackUintTexture->TextureRHI;
 		ShaderParameters->HeightVirtualTextureWorldToUvTransform = FMatrix44f::Identity;
 		ShaderParameters->HeightVirtualTextureEnabled = 0;
@@ -460,8 +460,8 @@ struct FNDILandscapeData_RenderThread
 			return false;
 		}
 
-		ShaderParameters->NormalVirtualTexture0 = NormalVirtualLayer0.TextureRef;
-		ShaderParameters->NormalVirtualTexture1 = NormalVirtualLayer1.TextureRef;
+		ShaderParameters->NormalVirtualTexture0 = NormalVirtualLayer0.TextureSRV;
+		ShaderParameters->NormalVirtualTexture1 = NormalVirtualLayer1.TextureSRV;
 		ShaderParameters->NormalVirtualTexturePageTable = NormalVirtualPage.PageTableRef;
 
 		FLargeWorldRenderPosition NormalVirtualTextureOrigin(NormalVirtualPage.WorldToUvParameters[0]);
@@ -487,8 +487,8 @@ struct FNDILandscapeData_RenderThread
 	{
 		const FUintVector4 DummyUint4(ForceInitToZero);
 
-		ShaderParameters->NormalVirtualTexture0 =  GBlackTexture->TextureRHI;
-		ShaderParameters->NormalVirtualTexture1 = GBlackTexture->TextureRHI;
+		ShaderParameters->NormalVirtualTexture0 = GBlackTextureWithSRV->ShaderResourceViewRHI;
+		ShaderParameters->NormalVirtualTexture1 = GBlackTextureWithSRV->ShaderResourceViewRHI;
 		ShaderParameters->NormalVirtualTexturePageTable = GBlackUintTexture->TextureRHI;
 		ShaderParameters->NormalVirtualTextureWorldToUvTransform = FMatrix44f::Identity;
 		ShaderParameters->NormalVirtualTextureEnabled = 0;
@@ -593,6 +593,7 @@ struct FNiagaraDataInterfaceProxyLandscape : public FNiagaraDataInterfaceProxy
 	void UpdateProxy_RT(
 		FNiagaraSystemInstanceID ID,
 		const URuntimeVirtualTexture* BaseColorVirtualTexture,
+		bool bBaseColorTextureSRGB,
 		const URuntimeVirtualTexture* HeightVirtualTexture,
 		const URuntimeVirtualTexture* NormalVirtualTexture,
 		ERuntimeVirtualTextureMaterialType NormalVirtualTextureMode)
@@ -602,12 +603,12 @@ struct FNiagaraDataInterfaceProxyLandscape : public FNiagaraDataInterfaceProxy
 			// todo - need to figure out a way to confirm that this is in fact the best/only option for the page/layer indices
 			constexpr uint32 BaseColorVirtualTextureLayerIndex = 0;
 			constexpr uint32 BaseColorVirtualTexturePageIndex = 0;
-			Proxy_RT->BaseColorVirtualLayer.Update(BaseColorVirtualTexture, BaseColorVirtualTextureLayerIndex);
+			Proxy_RT->BaseColorVirtualLayer.Update(BaseColorVirtualTexture, BaseColorVirtualTextureLayerIndex, bBaseColorTextureSRGB);
 			Proxy_RT->BaseColorVirtualPage.Update(BaseColorVirtualTexture, BaseColorVirtualTexturePageIndex, true /*bIncludeWorldToUv*/, false /*bIncludeHeightUnpack*/);
 
 			constexpr uint32 HeightVirtualTextureLayerIndex = 0;
 			constexpr uint32 HeightVirtualTexturePageIndex = 0;
-			Proxy_RT->HeightVirtualLayer.Update(HeightVirtualTexture, HeightVirtualTextureLayerIndex);
+			Proxy_RT->HeightVirtualLayer.Update(HeightVirtualTexture, HeightVirtualTextureLayerIndex, false /* bSRGB */);
 			Proxy_RT->HeightVirtualPage.Update(HeightVirtualTexture, HeightVirtualTexturePageIndex, true /*bIncludeWorldToUv*/, true /*bIncludeHeightUnpack*/);
 
 			constexpr uint32 NormalVirtualTexturePageIndex = 0;
@@ -616,20 +617,20 @@ struct FNiagaraDataInterfaceProxyLandscape : public FNiagaraDataInterfaceProxy
 			switch (NormalVirtualTextureMode)
 			{
 			case ERuntimeVirtualTextureMaterialType::BaseColor_Normal_Roughness:
-				Proxy_RT->NormalVirtualLayer0.Update(NormalVirtualTexture, 0);
-				Proxy_RT->NormalVirtualLayer1.Update(NormalVirtualTexture, 1);
+				Proxy_RT->NormalVirtualLayer0.Update(NormalVirtualTexture, 0, false /* bSRGB */);
+				Proxy_RT->NormalVirtualLayer1.Update(NormalVirtualTexture, 1, false /* bSRGB */);
 				Proxy_RT->NormalUnpackMode = FNDILandscapeData_RenderThread::ENormalUnpackType::B5G6R5;
 				break;
 			case ERuntimeVirtualTextureMaterialType::BaseColor_Normal_Specular:
-				Proxy_RT->NormalVirtualLayer0.Update(NormalVirtualTexture, 0);
-				Proxy_RT->NormalVirtualLayer1.Update(NormalVirtualTexture, 1);
+				Proxy_RT->NormalVirtualLayer0.Update(NormalVirtualTexture, 0, false /* bSRGB */);
+				Proxy_RT->NormalVirtualLayer1.Update(NormalVirtualTexture, 1, false /* bSRGB */);
 				Proxy_RT->NormalUnpackMode = FNDILandscapeData_RenderThread::ENormalUnpackType::BC3BC3;
 				break;
 
 			case ERuntimeVirtualTextureMaterialType::BaseColor_Normal_Specular_YCoCg:
 			case ERuntimeVirtualTextureMaterialType::BaseColor_Normal_Specular_Mask_YCoCg:
-				Proxy_RT->NormalVirtualLayer0.Update(NormalVirtualTexture, 1);
-				Proxy_RT->NormalVirtualLayer1.Update(NormalVirtualTexture, 2);
+				Proxy_RT->NormalVirtualLayer0.Update(NormalVirtualTexture, 1, false /* bSRGB */);
+				Proxy_RT->NormalVirtualLayer1.Update(NormalVirtualTexture, 2, false /* bSRGB */);
 				Proxy_RT->NormalUnpackMode = FNDILandscapeData_RenderThread::ENormalUnpackType::BC5BC1;
 				break;
 
@@ -1490,12 +1491,13 @@ bool UNiagaraDataInterfaceLandscape::PerInstanceTick(void* PerInstanceData, FNia
 		ENQUEUE_RENDER_COMMAND(FNiagaraDIUpdateProxy) (
 			[RT_Proxy,
 			BaseColorVirtualTexture,
+			bBaseColorTextureSRGB = InstanceData->BaseColorVirtualTextureSRGB,
 			HeightVirtualTexture,
 			NormalVirtualTexture,
 			NormalVirtualTextureMode = InstanceData->NormalVirtualTextureMode,
 			InstanceID = SystemInstance->GetId()](FRHICommandListImmediate& CmdList)
 		{
-			RT_Proxy->UpdateProxy_RT(InstanceID, BaseColorVirtualTexture, HeightVirtualTexture, NormalVirtualTexture, NormalVirtualTextureMode);
+			RT_Proxy->UpdateProxy_RT(InstanceID, BaseColorVirtualTexture, bBaseColorTextureSRGB, HeightVirtualTexture, NormalVirtualTexture, NormalVirtualTextureMode);
 		});
 	}
 
