Added refresh to inline parameter box when a module item gets refreshed.
Also removed ensure that could trigger if clicking on an inline edit condition bool

diff --git a/Engine/Plugins/FX/Niagara/Source/NiagaraEditorWidgets/Private/SNiagaraOverviewInlineParameterBox.cpp b/Engine/Plugins/FX/Niagara/Source/NiagaraEditorWidgets/Private/SNiagaraOverviewInlineParameterBox.cpp
--- a/Engine/Plugins/FX/Niagara/Source/NiagaraEditorWidgets/Private/SNiagaraOverviewInlineParameterBox.cpp
+++ b/Engine/Plugins/FX/Niagara/Source/NiagaraEditorWidgets/Private/SNiagaraOverviewInlineParameterBox.cpp
@@ -40,6 +40,8 @@ void SNiagaraOverviewInlineParameterBox::Construct(const FArguments& InArgs, UNi
 	];
 
 	ConstructChildren();
+	
+	ModuleItem->OnStructureChanged().AddSP(this, &SNiagaraOverviewInlineParameterBox::OnModuleItemStructureChanged);
 }
 
 SNiagaraOverviewInlineParameterBox::~SNiagaraOverviewInlineParameterBox()
@@ -74,7 +76,7 @@ FReply SNiagaraOverviewInlineParameterBox::NavigateToStack(TWeakObjectPtr<const
 	
 	TArray<UNiagaraStackEntry::FStackSearchItem> SearchItems;
 	FunctionInput->GetSearchItems(SearchItems);
-	if(ensure(SearchItems.Num() > 0))
+	if(SearchItems.Num() > 0)
 	{
 		ModuleItem->GetSystemViewModel()->GetSelectionViewModel()->GetSelectionStackViewModel()->SetSearchTextExternal(SearchItems[0].Value);
 		return FReply::Handled();
@@ -420,4 +422,9 @@ TWeakObjectPtr<const UNiagaraStackFunctionInput> SNiagaraOverviewInlineParameter
 	return nullptr;
 }
 
+void SNiagaraOverviewInlineParameterBox::OnModuleItemStructureChanged(ENiagaraStructureChangedFlags)
+{
+	ConstructChildren();
+}
+
 #undef LOCTEXT_NAMESPACE
diff --git a/Engine/Plugins/FX/Niagara/Source/NiagaraEditorWidgets/Private/SNiagaraOverviewInlineParameterBox.h b/Engine/Plugins/FX/Niagara/Source/NiagaraEditorWidgets/Private/SNiagaraOverviewInlineParameterBox.h
--- a/Engine/Plugins/FX/Niagara/Source/NiagaraEditorWidgets/Private/SNiagaraOverviewInlineParameterBox.h
+++ b/Engine/Plugins/FX/Niagara/Source/NiagaraEditorWidgets/Private/SNiagaraOverviewInlineParameterBox.h
@@ -42,6 +42,7 @@ private:
 	/** We build a substitute map of entries that we want to */
 	TWeakObjectPtr<const UNiagaraStackFunctionInput> FindSubstituteEntry(const UNiagaraStackFunctionInput* Input);
 
+	void OnModuleItemStructureChanged(ENiagaraStructureChangedFlags);
 private:
 	/** The module item whose parameters this box is representing */
 	TWeakObjectPtr<UNiagaraStackModuleItem> ModuleItem;
