Changed Lerp C pin's friendly name to 'Alpha'.

Also changed NiagaraNodeOp to support pin friendly names. This might cause new pin names to appear elsewhere as an unintended consequence but seems fine from my tests

diff --git a/Engine/Plugins/FX/Niagara/Source/NiagaraEditor/Private/NiagaraEditorCommon.cpp b/Engine/Plugins/FX/Niagara/Source/NiagaraEditor/Private/NiagaraEditorCommon.cpp
--- a/Engine/Plugins/FX/Niagara/Source/NiagaraEditor/Private/NiagaraEditorCommon.cpp
+++ b/Engine/Plugins/FX/Niagara/Source/NiagaraEditor/Private/NiagaraEditorCommon.cpp
@@ -365,9 +365,12 @@ void FNiagaraOpInfo::Init()
 		Op->bShowPinNamesInCompactMode = true;
 		Op->Description = NSLOCTEXT("NiagaraOpInfo", "Lerp Desc", "Result = (A * (1 - C)) + (B * C)");
 		Op->Keywords = FText::FromString(TEXT("lerp"));
+
+		FText AlphaPinFriendlyName = FText::FromString("Alpha");
+		FText AlphaTooltip = NSLOCTEXT("NiagaraOpInfo", "Lerp Alpha Tooltip", "A value typically between 0 and 1. Determines the percentage to use for interpolating from A to B.");
 		Op->Inputs.Add(FNiagaraOpInOutInfo(A, NumericType, AText, AText, DefaultStr_Zero));
 		Op->Inputs.Add(FNiagaraOpInOutInfo(B, NumericType, BText, BText, DefaultStr_One));
-		Op->Inputs.Add(FNiagaraOpInOutInfo(C, NumericType, CText, CText, DefaultStr_Zero));
+		Op->Inputs.Add(FNiagaraOpInOutInfo(C, NumericType, AlphaPinFriendlyName, AlphaTooltip, DefaultStr_Zero));
 		Op->Outputs.Add(FNiagaraOpInOutInfo(Result, NumericType, ResultText, ResultText, DefaultStr_Zero, TEXT("lerp({0},{1},{2})")));
 		Op->BuildName(TEXT("Lerp"), CategoryName);
 		Op->InputTypeValidationFunction.BindLambda([=](const TArray<FNiagaraTypeDefinition>& InputTypes, FText& ErrorMessage)
diff --git a/Engine/Plugins/FX/Niagara/Source/NiagaraEditor/Private/NiagaraNodeOp.cpp b/Engine/Plugins/FX/Niagara/Source/NiagaraEditor/Private/NiagaraNodeOp.cpp
--- a/Engine/Plugins/FX/Niagara/Source/NiagaraEditor/Private/NiagaraNodeOp.cpp
+++ b/Engine/Plugins/FX/Niagara/Source/NiagaraEditor/Private/NiagaraNodeOp.cpp
@@ -32,6 +32,12 @@ void UNiagaraNodeOp::AllocateDefaultPins()
 		const FNiagaraOpInOutInfo& InOutInfo = OpInfo->Inputs[SrcIndex];
 		UEdGraphPin* Pin = CreatePin(EGPD_Input, Schema->TypeDefinitionToPinType(InOutInfo.DataType), *InOutInfo.Name.ToString());
 		check(Pin);
+
+		if(InOutInfo.FriendlyName.IsEmpty() == false)
+		{
+			Pin->PinFriendlyName = InOutInfo.FriendlyName;
+		}
+		
 		Pin->bDefaultValueIsIgnored = false;
 		Pin->bDefaultValueIsReadOnly = false;
 		Pin->bNotConnectable = false;
