Remove an unnecessary assert in mobile path for Niagara mesh rendering that was triggerning during PSO pre-caching

diff --git a/Engine/Plugins/FX/Niagara/Source/NiagaraVertexFactories/Private/NiagaraMeshVertexFactory.cpp b/Engine/Plugins/FX/Niagara/Source/NiagaraVertexFactories/Private/NiagaraMeshVertexFactory.cpp
--- a/Engine/Plugins/FX/Niagara/Source/NiagaraVertexFactories/Private/NiagaraMeshVertexFactory.cpp
+++ b/Engine/Plugins/FX/Niagara/Source/NiagaraVertexFactories/Private/NiagaraMeshVertexFactory.cpp
@@ -213,9 +213,7 @@ void FNiagaraMeshVertexFactory::GetVertexElements(ERHIFeatureLevel::Type Feature
 #if NIAGARA_ENABLE_GPU_SCENE_MESHES
 	if (UseGPUScene(GMaxRHIShaderPlatform, GMaxRHIFeatureLevel))
 	{
-		// For ES3.1 attribute ID needs to be done differently
-		check(FeatureLevel > ERHIFeatureLevel::ES3_1);
-		Elements.Add(FVertexElement(InOutStreams.Num(), 0, VET_UInt, 13, 0, true));
+		Elements.Add(FVertexElement(InOutStreams.Num(), 0, VET_UInt, 13, sizeof(uint32), true));
 	}
 #endif
 }
