Some small UX improvements for NDCs

- Initial customization for NDC assets. Tidies up the assset details view considerablly.
- Adding const to a few NDC library func params.
- Preserving cached NDC settings ptr when clearing to null and returning to a valid object.
- Adding asset menu option to dirty depednant compilable assets for NDCs.
- Tweak to MarkDependentCompilableAssetsDirty() to allow top level assets to not need to be compilable, i.e. NDCs

diff --git a/Engine/Plugins/FX/Niagara/Source/Niagara/Private/NiagaraDataChannel.cpp b/Engine/Plugins/FX/Niagara/Source/Niagara/Private/NiagaraDataChannel.cpp
--- a/Engine/Plugins/FX/Niagara/Source/Niagara/Private/NiagaraDataChannel.cpp
+++ b/Engine/Plugins/FX/Niagara/Source/Niagara/Private/NiagaraDataChannel.cpp
@@ -970,22 +970,22 @@ UNiagaraDataChannelLibrary::UNiagaraDataChannelLibrary(const FObjectInitializer&
 {
 }
 
-UNiagaraDataChannelHandler* UNiagaraDataChannelLibrary::GetNiagaraDataChannel(const UObject* WorldContextObject, UNiagaraDataChannelAsset* Channel)
+UNiagaraDataChannelHandler* UNiagaraDataChannelLibrary::GetNiagaraDataChannel(const UObject* WorldContextObject, const UNiagaraDataChannelAsset* Channel)
 {
 	return GetNiagaraDataChannel(WorldContextObject, Channel->Get());
 }
 
-UNiagaraDataChannelWriter* UNiagaraDataChannelLibrary::WriteToNiagaraDataChannel(const UObject* WorldContextObject, UNiagaraDataChannelAsset* Channel, FNiagaraDataChannelSearchParameters SearchParams, int32 Count, bool bVisibleToGame, bool bVisibleToCPU, bool bVisibleToGPU)
+UNiagaraDataChannelWriter* UNiagaraDataChannelLibrary::WriteToNiagaraDataChannel(const UObject* WorldContextObject, const UNiagaraDataChannelAsset* Channel, FNiagaraDataChannelSearchParameters SearchParams, int32 Count, bool bVisibleToGame, bool bVisibleToCPU, bool bVisibleToGPU)
 {
 	return WriteToNiagaraDataChannel(WorldContextObject, Channel->Get(), SearchParams, Count, bVisibleToGame, bVisibleToCPU, bVisibleToGPU);
 }
 
-UNiagaraDataChannelReader* UNiagaraDataChannelLibrary::ReadFromNiagaraDataChannel(const UObject* WorldContextObject, UNiagaraDataChannelAsset* Channel, FNiagaraDataChannelSearchParameters SearchParams, bool bReadPreviousFrame)
+UNiagaraDataChannelReader* UNiagaraDataChannelLibrary::ReadFromNiagaraDataChannel(const UObject* WorldContextObject, const UNiagaraDataChannelAsset* Channel, FNiagaraDataChannelSearchParameters SearchParams, bool bReadPreviousFrame)
 {
 	return ReadFromNiagaraDataChannel(WorldContextObject, Channel->Get(), SearchParams, bReadPreviousFrame);
 }
 
-UNiagaraDataChannelHandler* UNiagaraDataChannelLibrary::GetNiagaraDataChannel(const UObject* WorldContextObject, UNiagaraDataChannel* Channel)
+UNiagaraDataChannelHandler* UNiagaraDataChannelLibrary::GetNiagaraDataChannel(const UObject* WorldContextObject, const UNiagaraDataChannel* Channel)
 {
 	UWorld* World = (WorldContextObject != nullptr) ? GEngine->GetWorldFromContextObject(WorldContextObject, EGetWorldErrorMode::LogAndReturnNull) : nullptr;
 	if (World)
@@ -998,7 +998,7 @@ UNiagaraDataChannelHandler* UNiagaraDataChannelLibrary::GetNiagaraDataChannel(co
 	return nullptr;
 }
 
-UNiagaraDataChannelWriter* UNiagaraDataChannelLibrary::WriteToNiagaraDataChannel(const UObject* WorldContextObject, UNiagaraDataChannel* Channel, FNiagaraDataChannelSearchParameters SearchParams, int32 Count, bool bVisibleToGame, bool bVisibleToCPU, bool bVisibleToGPU)
+UNiagaraDataChannelWriter* UNiagaraDataChannelLibrary::WriteToNiagaraDataChannel(const UObject* WorldContextObject, const UNiagaraDataChannel* Channel, FNiagaraDataChannelSearchParameters SearchParams, int32 Count, bool bVisibleToGame, bool bVisibleToCPU, bool bVisibleToGPU)
 {
 	check(IsInGameThread());
 	UWorld* World = (WorldContextObject != nullptr) ? GEngine->GetWorldFromContextObject(WorldContextObject, EGetWorldErrorMode::LogAndReturnNull) : nullptr;
@@ -1021,7 +1021,7 @@ UNiagaraDataChannelWriter* UNiagaraDataChannelLibrary::WriteToNiagaraDataChannel
 	return nullptr;
 }
 
-UNiagaraDataChannelReader* UNiagaraDataChannelLibrary::ReadFromNiagaraDataChannel(const UObject* WorldContextObject, UNiagaraDataChannel* Channel, FNiagaraDataChannelSearchParameters SearchParams, bool bReadPreviousFrame)
+UNiagaraDataChannelReader* UNiagaraDataChannelLibrary::ReadFromNiagaraDataChannel(const UObject* WorldContextObject, const UNiagaraDataChannel* Channel, FNiagaraDataChannelSearchParameters SearchParams, bool bReadPreviousFrame)
 {
 	UWorld* World = (WorldContextObject != nullptr) ? GEngine->GetWorldFromContextObject(WorldContextObject, EGetWorldErrorMode::LogAndReturnNull) : nullptr;
 	if (World)
@@ -1049,7 +1049,7 @@ UNiagaraDataChannelReader* UNiagaraDataChannelLibrary::ReadFromNiagaraDataChanne
 
 void UNiagaraDataChannelAsset::PreEditChange(class FProperty* PropertyAboutToChange)
 {
-	if(PropertyAboutToChange->GetFName() == GET_MEMBER_NAME_CHECKED(UNiagaraDataChannelAsset, DataChannel))
+	if(PropertyAboutToChange->GetFName() == GET_MEMBER_NAME_CHECKED(UNiagaraDataChannelAsset, DataChannel) && DataChannel)
 	{
 		CachedPreChangeDataChannel = DataChannel;
 	}
@@ -1057,7 +1057,7 @@ void UNiagaraDataChannelAsset::PreEditChange(class FProperty* PropertyAboutToCha
 
 void UNiagaraDataChannelAsset::PostEditChangeProperty(FPropertyChangedEvent& PropertyChangedEvent)
 {
-	if(CachedPreChangeDataChannel)
+	if(CachedPreChangeDataChannel && DataChannel)
 	{
 		UEngine::FCopyPropertiesForUnrelatedObjectsParams Params;
 		UEngine::CopyPropertiesForUnrelatedObjects(CachedPreChangeDataChannel, DataChannel, Params);
diff --git a/Engine/Plugins/FX/Niagara/Source/Niagara/Public/NiagaraDataChannel.h b/Engine/Plugins/FX/Niagara/Source/Niagara/Public/NiagaraDataChannel.h
--- a/Engine/Plugins/FX/Niagara/Source/Niagara/Public/NiagaraDataChannel.h
+++ b/Engine/Plugins/FX/Niagara/Source/Niagara/Public/NiagaraDataChannel.h
@@ -81,7 +81,7 @@ struct FNiagaraDataChannelDataProxy
 
 DECLARE_MULTICAST_DELEGATE_OneParam(FOnDataChannelCreated, const UNiagaraDataChannel*);
 
-UCLASS(Experimental, abstract, EditInlineNew, MinimalAPI)
+UCLASS(Experimental, abstract, EditInlineNew, MinimalAPI, prioritizeCategories=("Data Channel"))
 class UNiagaraDataChannel : public UObject
 {
 public:
@@ -184,17 +184,17 @@ class NIAGARA_API UNiagaraDataChannelLibrary : public UBlueprintFunctionLibrary
 	GENERATED_UCLASS_BODY()
 
 	UFUNCTION(BlueprintCallable, Category = NiagaraDataChannel, meta = (Keywords = "niagara DataChannel", WorldContext = "WorldContextObject", UnsafeDuringActorConstruction = "true"))
-	static UNiagaraDataChannelHandler* GetNiagaraDataChannel(const UObject* WorldContextObject, UNiagaraDataChannelAsset* Channel);
+	static UNiagaraDataChannelHandler* GetNiagaraDataChannel(const UObject* WorldContextObject, const UNiagaraDataChannelAsset* Channel);
 
 	/** Initializes and returns the Niagara Data Channel writer to write N elements to the given data channel. */
 	UFUNCTION(BlueprintCallable, Category = NiagaraDataChannel, meta = (Keywords = "niagara DataChannel", WorldContext = "WorldContextObject", UnsafeDuringActorConstruction = "true"))
-	static UNiagaraDataChannelWriter* WriteToNiagaraDataChannel(const UObject* WorldContextObject, UNiagaraDataChannelAsset* Channel, FNiagaraDataChannelSearchParameters SearchParams, int32 Count, bool bVisibleToGame, bool bVisibleToCPU, bool bVisibleToGPU);
+	static UNiagaraDataChannelWriter* WriteToNiagaraDataChannel(const UObject* WorldContextObject, const UNiagaraDataChannelAsset* Channel, FNiagaraDataChannelSearchParameters SearchParams, int32 Count, bool bVisibleToGame, bool bVisibleToCPU, bool bVisibleToGPU);
 
 	/** Initializes and returns the Niagara Data Channel reader for the given data channel. */
 	UFUNCTION(BlueprintCallable, Category = NiagaraDataChannel, meta = (Keywords = "niagara DataChannel", WorldContext = "WorldContextObject", UnsafeDuringActorConstruction = "true"))
-	static UNiagaraDataChannelReader* ReadFromNiagaraDataChannel(const UObject* WorldContextObject, UNiagaraDataChannelAsset* Channel, FNiagaraDataChannelSearchParameters SearchParams, bool bReadPreviousFrame);
+	static UNiagaraDataChannelReader* ReadFromNiagaraDataChannel(const UObject* WorldContextObject, const UNiagaraDataChannelAsset* Channel, FNiagaraDataChannelSearchParameters SearchParams, bool bReadPreviousFrame);
 
-	static UNiagaraDataChannelHandler* GetNiagaraDataChannel(const UObject* WorldContextObject, UNiagaraDataChannel* Channel);
-	static UNiagaraDataChannelWriter* WriteToNiagaraDataChannel(const UObject* WorldContextObject, UNiagaraDataChannel* Channel, FNiagaraDataChannelSearchParameters SearchParams, int32 Count, bool bVisibleToGame, bool bVisibleToCPU, bool bVisibleToGPU);
-	static UNiagaraDataChannelReader* ReadFromNiagaraDataChannel(const UObject* WorldContextObject, UNiagaraDataChannel* Channel, FNiagaraDataChannelSearchParameters SearchParams, bool bReadPreviousFrame);
+	static UNiagaraDataChannelHandler* GetNiagaraDataChannel(const UObject* WorldContextObject, const UNiagaraDataChannel* Channel);
+	static UNiagaraDataChannelWriter* WriteToNiagaraDataChannel(const UObject* WorldContextObject, const UNiagaraDataChannel* Channel, FNiagaraDataChannelSearchParameters SearchParams, int32 Count, bool bVisibleToGame, bool bVisibleToCPU, bool bVisibleToGPU);
+	static UNiagaraDataChannelReader* ReadFromNiagaraDataChannel(const UObject* WorldContextObject, const UNiagaraDataChannel* Channel, FNiagaraDataChannelSearchParameters SearchParams, bool bReadPreviousFrame);
 };
diff --git a/Engine/Plugins/FX/Niagara/Source/NiagaraEditor/Private/AssetDefinitions/AssetDefinition_NiagaraDataChannel.cpp b/Engine/Plugins/FX/Niagara/Source/NiagaraEditor/Private/AssetDefinitions/AssetDefinition_NiagaraDataChannel.cpp
--- a/Engine/Plugins/FX/Niagara/Source/NiagaraEditor/Private/AssetDefinitions/AssetDefinition_NiagaraDataChannel.cpp
+++ b/Engine/Plugins/FX/Niagara/Source/NiagaraEditor/Private/AssetDefinitions/AssetDefinition_NiagaraDataChannel.cpp
@@ -4,6 +4,9 @@
 
 #include "NiagaraEditorStyle.h"
 #include "NiagaraDataChannel.h"
+#include "NiagaraEditorUtilities.h"
+#include "ContentBrowserMenuContexts.h"
+#include "ToolMenus.h"
 
 #define LOCTEXT_NAMESPACE "UAssetDefinition_NiagaraDataChannel"
 
@@ -23,6 +26,39 @@ TConstArrayView<FAssetCategoryPath> UAssetDefinition_NiagaraDataChannel::GetAsse
 	return AssetPaths;
 }
 
+
+
+namespace MenuExtension_NiagaraDataChannel
+{
+	static void ExecuteMarkDependentCompilableAssetsDirty(const FToolMenuContext& InContext)
+	{
+		const UContentBrowserAssetContextMenuContext* CBContext = UContentBrowserAssetContextMenuContext::FindContextWithAssets(InContext);
+		FNiagaraEditorUtilities::MarkDependentCompilableAssetsDirty(CBContext->LoadSelectedObjects<UObject>());
+	}
+
+	static FDelayedAutoRegisterHelper DelayedAutoRegister(EDelayedRegisterRunPhase::EndOfEngineInit, [] {
+		UToolMenus::RegisterStartupCallback(FSimpleMulticastDelegate::FDelegate::CreateLambda([]()
+			{
+				FToolMenuOwnerScoped OwnerScoped(UE_MODULE_NAME);
+				UToolMenu* Menu = UE::ContentBrowser::ExtendToolMenu_AssetContextMenu(UNiagaraDataChannelAsset::StaticClass());
+
+				FToolMenuSection& Section = Menu->FindOrAddSection("GetAssetActions");
+				Section.AddDynamicEntry(NAME_None, FNewToolMenuSectionDelegate::CreateLambda([](FToolMenuSection& InSection)
+					{
+						{
+							const TAttribute<FText> Label = LOCTEXT("MarkDependentCompilableAssetsDirtyLabel", "Mark Dependent Compilable Assets Dirty");
+							const TAttribute<FText> ToolTip = LOCTEXT("MarkDependentCompilableAssetsDirtyToolTip", "Finds all niagara assets which depend on this asset either directly or indirectly, and marks them dirty so they can be saved with the latest version.");
+							const FSlateIcon Icon = FSlateIcon();
+
+							FToolUIAction UIAction;
+							UIAction.ExecuteAction = FToolMenuExecuteAction::CreateStatic(&ExecuteMarkDependentCompilableAssetsDirty);
+							InSection.AddMenuEntry("MarkDependentCompilableAssetsDirty", Label, ToolTip, Icon, UIAction);
+						}
+					}));
+			}));
+		});
+}
+
 #undef LOCTEXT_NAMESPACE
 
 
diff --git a/Engine/Plugins/FX/Niagara/Source/NiagaraEditor/Private/Customizations/NiagaraDataChannelDetails.cpp b/Engine/Plugins/FX/Niagara/Source/NiagaraEditor/Private/Customizations/NiagaraDataChannelDetails.cpp
new file mode 100644
--- /dev/null
+++ b/Engine/Plugins/FX/Niagara/Source/NiagaraEditor/Private/Customizations/NiagaraDataChannelDetails.cpp
@@ -0,0 +1,137 @@
+// Copyright Epic Games, Inc. All Rights Reserved.
+
+#include "NiagaraDataChannelDetails.h"
+
+#include "NiagaraDataChannel.h"
+
+#include "PropertyHandle.h"
+#include "DetailLayoutBuilder.h"
+#include "DetailCategoryBuilder.h"
+#include "IDetailChildrenBuilder.h"
+#include "DetailWidgetRow.h"
+#include "PropertyCustomizationHelpers.h"
+#include "ObjectEditorUtils.h"
+
+#include "Widgets/Text/STextBlock.h"
+#include "Widgets/Layout/SBox.h"
+#include "Widgets/Input/SButton.h"
+
+#include "Modules/ModuleManager.h"
+
+#define LOCTEXT_NAMESPACE "NiagaraDataChannelDetails"
+
+TSharedRef<IDetailCustomization> FNiagaraDataChannelAssetDetails::MakeInstance()
+{
+	return MakeShareable(new FNiagaraDataChannelAssetDetails);
+}
+
+FNiagaraDataChannelAssetDetails::~FNiagaraDataChannelAssetDetails()
+{
+
+}
+
+void FNiagaraDataChannelAssetDetails::CustomizeDetails(IDetailLayoutBuilder& DetailBuilder)
+{
+	FPropertyEditorModule& PropertyModule = FModuleManager::GetModuleChecked<FPropertyEditorModule>("PropertyEditor");
+
+	TArray<TWeakObjectPtr<UObject>> CustomizedObjects;
+	DetailBuilder.GetObjectsBeingCustomized(CustomizedObjects);
+
+	if(CustomizedObjects.Num() != 1)
+	{
+		return;
+	}
+	
+	DataChannelAsset = Cast<UNiagaraDataChannelAsset>(CustomizedObjects[0].Get());
+	if(DataChannelAsset == nullptr)
+	{
+		return;
+	}
+	
+	FName DataChannelCatName("Data Channel");
+	IDetailCategoryBuilder& DataChannelCat = DetailBuilder.EditCategory(DataChannelCatName);
+
+	TSharedPtr<IPropertyHandle> DataChannelProperty = DetailBuilder.GetProperty("DataChannel");
+	DataChannelProperty->SetPropertyDisplayName(LOCTEXT("DataChannelTypeName", "Data Channel Class"));
+	
+	// we need to build a mapping between categories and properties first because the CLASS_CollapseCategories class flag can
+	// remove the categories from top-level properties, so we need to traverse them differently
+	TMap<FName, TArray<TSharedPtr<IPropertyHandle>>> CategoryMapping;
+
+	//Always add the NDC type property.
+	CategoryMapping.Add(DataChannelCatName).Add(DataChannelProperty);
+
+	UNiagaraDataChannel* NDC = DataChannelAsset->Get();
+	if(NDC != nullptr)
+	{		
+		TSharedPtr<IPropertyHandle> InstancedNDCHandle = DataChannelProperty->GetChildHandle(0);
+		check(InstancedNDCHandle);
+
+		uint32 NumChildren = 0;
+		InstancedNDCHandle->GetNumChildren(NumChildren);
+		UClass* NDCClass = NDC->GetClass();
+		bool HasCollapsedCategories = NDCClass->HasAnyClassFlags(CLASS_CollapseCategories);
+
+		for (uint32 i = 0; i < NumChildren; i++)
+		{
+			TSharedPtr<IPropertyHandle> TopLevelHandle = InstancedNDCHandle->GetChildHandle(i);
+
+			if (HasCollapsedCategories)
+			{
+				if (!TopLevelHandle.IsValid() || !TopLevelHandle->GetProperty() || !TopLevelHandle->GetProperty()->GetClass())
+				{
+					continue;
+				}
+
+				// with collapsed categories, the properties are given to us directly, without the category properties as parents
+				FProperty* Property = TopLevelHandle->GetProperty();
+				FName CategoryFName = FObjectEditorUtils::GetCategoryFName(Property);
+				CategoryMapping.FindOrAdd(CategoryFName).Add(TopLevelHandle);
+			}
+			else
+			{
+				uint32 CategoryChildren;
+				TopLevelHandle->GetNumChildren(CategoryChildren);
+				for (uint32 k = 0; k < CategoryChildren; k++)
+				{
+					TSharedPtr<IPropertyHandle> ChildHandle = TopLevelHandle->GetChildHandle(k);
+					if (!ChildHandle->GetProperty() || !ChildHandle->GetProperty()->GetClass())
+					{
+						continue;
+					}
+
+					FName CategoryFName = FName(TopLevelHandle->GetPropertyDisplayName().ToString());
+					CategoryMapping.FindOrAdd(CategoryFName).Add(ChildHandle);
+				}
+			}
+		}
+	}
+
+	auto AddCategory = [&DetailBuilder](FName CategoryName, TArray<TSharedPtr<IPropertyHandle>>& Props)
+	{
+		IDetailCategoryBuilder& CatBuilder = DetailBuilder.EditCategory(CategoryName);
+		for(TSharedPtr<IPropertyHandle>& Prop : Props)
+		{
+			CatBuilder.AddProperty(Prop);
+		}
+	};
+	
+	//Add the basic data channel props first.
+	if(auto* DataChannelProps = CategoryMapping.Find(DataChannelCatName))
+	{
+		AddCategory(DataChannelCatName, *DataChannelProps);
+	}
+	for (TPair<FName, TArray<TSharedPtr<IPropertyHandle>>> CategoryPair : CategoryMapping)
+	{
+		FName CatName = CategoryPair.Key;
+		TArray<TSharedPtr<IPropertyHandle>>& ChildProps = CategoryPair.Value;
+		if (ChildProps.Num() == 0 || CatName == DataChannelCatName)
+		{
+			continue;
+		}
+
+		AddCategory(CatName, ChildProps);
+	}
+}
+
+#undef LOCTEXT_NAMESPACE
diff --git a/Engine/Plugins/FX/Niagara/Source/NiagaraEditor/Private/Customizations/NiagaraDataChannelDetails.h b/Engine/Plugins/FX/Niagara/Source/NiagaraEditor/Private/Customizations/NiagaraDataChannelDetails.h
new file mode 100644
--- /dev/null
+++ b/Engine/Plugins/FX/Niagara/Source/NiagaraEditor/Private/Customizations/NiagaraDataChannelDetails.h
@@ -0,0 +1,20 @@
+// Copyright Epic Games, Inc. All Rights Reserved.
+
+#pragma once
+
+#include "CoreMinimal.h"
+#include "IDetailCustomization.h"
+#include "NiagaraDataChannelPublic.h"
+
+class FNiagaraDataChannelAssetDetails : public IDetailCustomization
+{
+public:
+	virtual ~FNiagaraDataChannelAssetDetails() override;
+
+	static TSharedRef<IDetailCustomization> MakeInstance();
+	virtual void CustomizeDetails(IDetailLayoutBuilder& DetailBuilder) override;
+
+protected:
+	TWeakObjectPtr<UNiagaraDataChannelAsset> DataChannelAsset;
+};
+
diff --git a/Engine/Plugins/FX/Niagara/Source/NiagaraEditor/Private/NiagaraEditorModule.cpp b/Engine/Plugins/FX/Niagara/Source/NiagaraEditor/Private/NiagaraEditorModule.cpp
--- a/Engine/Plugins/FX/Niagara/Source/NiagaraEditor/Private/NiagaraEditorModule.cpp
+++ b/Engine/Plugins/FX/Niagara/Source/NiagaraEditor/Private/NiagaraEditorModule.cpp
@@ -110,6 +110,7 @@
 #include "Customizations/NiagaraBakerSettingsDetails.h"
 #include "Customizations/NiagaraOutlinerCustomization.h"
 #include "Customizations/NiagaraSimulationStageCustomization.h"
+#include "Customizations/NiagaraDataChannelDetails.h"
 
 #include "NiagaraComponent.h"
 #include "NiagaraNodeStaticSwitch.h"
@@ -1040,7 +1041,11 @@ void FNiagaraEditorModule::StartupModule()
 	PropertyModule.RegisterCustomClassLayout(
 		UNiagaraComponentRendererProperties::StaticClass()->GetFName(), 
 		FOnGetDetailCustomizationInstance::CreateStatic(&FNiagaraComponentRendererPropertiesDetails::MakeInstance));
-	
+
+	PropertyModule.RegisterCustomClassLayout(
+		UNiagaraDataChannelAsset::StaticClass()->GetFName(),
+		FOnGetDetailCustomizationInstance::CreateStatic(&FNiagaraDataChannelAssetDetails::MakeInstance));
+
 	PropertyModule.RegisterCustomPropertyTypeLayout(
 		FNiagaraFloat::StaticStruct()->GetFName(),
 		FOnGetPropertyTypeCustomizationInstance::CreateStatic(&FNiagaraNumericCustomization::MakeInstance));
diff --git a/Engine/Plugins/FX/Niagara/Source/NiagaraEditor/Private/NiagaraEditorUtilities.cpp b/Engine/Plugins/FX/Niagara/Source/NiagaraEditor/Private/NiagaraEditorUtilities.cpp
--- a/Engine/Plugins/FX/Niagara/Source/NiagaraEditor/Private/NiagaraEditorUtilities.cpp
+++ b/Engine/Plugins/FX/Niagara/Source/NiagaraEditor/Private/NiagaraEditorUtilities.cpp
@@ -1063,7 +1063,8 @@ void FNiagaraEditorUtilities::MarkDependentCompilableAssetsDirty(const TArray<FA
 	{
 		FAssetData AssetToCheck = AssetsToCheck[0];
 		AssetsToCheck.RemoveAtSwap(0);
-		if (IsCompilableAssetClass(AssetToCheck.GetClass()))
+		bool bMustBeCompilableAsset = !InAssetsToCheck.Contains(AssetToCheck);//Don't require the original assets to be compilable.
+		if (!bMustBeCompilableAsset || IsCompilableAssetClass(AssetToCheck.GetClass()))
 		{
 			if (AssetsToLoadAndMarkDirty.Contains(AssetToCheck) == false)
 			{
