Fixed average perf stats being wrong when multiple systems are active at the same time

diff --git a/Engine/Plugins/FX/Niagara/Source/Niagara/Private/NiagaraCommon.cpp b/Engine/Plugins/FX/Niagara/Source/Niagara/Private/NiagaraCommon.cpp
--- a/Engine/Plugins/FX/Niagara/Source/Niagara/Private/NiagaraCommon.cpp
+++ b/Engine/Plugins/FX/Niagara/Source/Niagara/Private/NiagaraCommon.cpp
@@ -381,10 +381,9 @@ float FNiagaraStatDatabase::GetRuntimeStat(FName StatName, ENiagaraScriptUsage U
 		{
 			if (MinimalNameToName(StatEntry.Key->Name) == StatName)
 			{
-				ValueCount = StatEntry.Value.CapturedTimings.Num();
-				for (int i = 0; i < ValueCount; i++)
+				ValueCount += StatEntry.Value.CapturedTimings.Num();
+				for (const float& Value : StatEntry.Value.CapturedTimings)
 				{
-					float Value = StatEntry.Value.CapturedTimings[i];
 					Max = FMath::Max(Max, Value);
 					Sum += Value;
 				}
