Wrapped AddSortedGPUSimulation in a mutex in preparation for parallel GDME.

diff --git a/Engine/Plugins/FX/Niagara/Source/Niagara/Private/NiagaraGpuComputeDispatch.cpp b/Engine/Plugins/FX/Niagara/Source/Niagara/Private/NiagaraGpuComputeDispatch.cpp
--- a/Engine/Plugins/FX/Niagara/Source/Niagara/Private/NiagaraGpuComputeDispatch.cpp
+++ b/Engine/Plugins/FX/Niagara/Source/Niagara/Private/NiagaraGpuComputeDispatch.cpp
@@ -2187,6 +2187,7 @@ void FNiagaraGpuComputeDispatch::OnDestroy()
 
 bool FNiagaraGpuComputeDispatch::AddSortedGPUSimulation(FRHICommandListBase& RHICmdList, FNiagaraGPUSortInfo& SortInfo)
 {
+	UE::TScopeLock Lock(AddSortedGPUSimulationMutex);
 	if (GPUSortManager && GPUSortManager->AddTask(RHICmdList, SortInfo.AllocationInfo, SortInfo.ParticleCount, SortInfo.SortFlags))
 	{
 		// It's not worth currently to have a map between SortInfo.AllocationInfo.SortBatchId and the relevant indices in SimulationsToSort
diff --git a/Engine/Plugins/FX/Niagara/Source/Niagara/Private/NiagaraGpuComputeDispatch.h b/Engine/Plugins/FX/Niagara/Source/Niagara/Private/NiagaraGpuComputeDispatch.h
--- a/Engine/Plugins/FX/Niagara/Source/Niagara/Private/NiagaraGpuComputeDispatch.h
+++ b/Engine/Plugins/FX/Niagara/Source/Niagara/Private/NiagaraGpuComputeDispatch.h
@@ -235,6 +235,8 @@ private:
 
 	int32 MaxTicksToFlush = TNumericLimits<int32>::Max();
 
+	UE::FMutex AddSortedGPUSimulationMutex;
+
 	bool bRequiresReadback = false;
 	TArray<FNiagaraSystemGpuComputeProxy*> ProxiesPerStage[ENiagaraGpuComputeTickStage::Max];
 
diff --git a/Engine/Source/Runtime/Engine/Private/Particles/FXSystemPrivate.h b/Engine/Source/Runtime/Engine/Private/Particles/FXSystemPrivate.h
--- a/Engine/Source/Runtime/Engine/Private/Particles/FXSystemPrivate.h
+++ b/Engine/Source/Runtime/Engine/Private/Particles/FXSystemPrivate.h
@@ -384,6 +384,9 @@ private:
 
 	/** Previous frame new particles for multi-gpu simulation*/
 	TArray<FNewParticle> LastFrameNewParticles;
+
+	UE::FMutex AddSortedGPUSimulationMutex;
+
 #if WITH_EDITOR
 	/** true if the system has been suspended. */
 	bool bSuspended;
diff --git a/Engine/Source/Runtime/Engine/Private/Particles/ParticleGpuSimulation.cpp b/Engine/Source/Runtime/Engine/Private/Particles/ParticleGpuSimulation.cpp
--- a/Engine/Source/Runtime/Engine/Private/Particles/ParticleGpuSimulation.cpp
+++ b/Engine/Source/Runtime/Engine/Private/Particles/ParticleGpuSimulation.cpp
@@ -4280,6 +4280,7 @@ bool FFXSystem::AddSortedGPUSimulation(FRHICommandListBase& RHICmdList, FParticl
 		EGPUSortFlags::LowPrecisionKeys | 
 		EGPUSortFlags::SortAfterPostRenderOpaque;
 
+	UE::TScopeLock Lock(AddSortedGPUSimulationMutex);
 	// Currently opaque materials would need SortAfterPreRender but this is incompatible with KeyGenAfterPostRenderOpaque
 	if (bIsTranslucent && GPUSortManager && GPUSortManager->AddTask(RHICmdList, OutInfo, Simulation->VertexBuffer.ParticleCount, SortFlags))
 	{
