Converted screenshot test to cached tests; improved comparison method

diff --git a/Engine/Plugins/FX/Niagara/Source/Niagara/Private/NiagaraSimCache.cpp b/Engine/Plugins/FX/Niagara/Source/Niagara/Private/NiagaraSimCache.cpp
--- a/Engine/Plugins/FX/Niagara/Source/Niagara/Private/NiagaraSimCache.cpp
+++ b/Engine/Plugins/FX/Niagara/Source/Niagara/Private/NiagaraSimCache.cpp
@@ -1257,9 +1257,9 @@ bool UNiagaraSimCache::IsDataEqual(const UNiagaraSimCache& OtherCache, float Err
 
 			for (const FNiagaraSimCacheVariable& Var : CacheLayout.EmitterLayouts[EmitterIndex].Variables)
 			{
-				if (Var.Variable.GetName() == FName("ID"))
+				if (Var.Variable.GetName() == FName("ID") || Var.Variable.GetName() == FName("MaterialRandom"))
 				{
-					// persistent ID is not deterministic, skip
+					// persistent ID and material random are not deterministic, skip
 					continue;
 				}
 				TArray<float> ExpectedFloats;
