- Change message for an invalid sim target for a renderer & add a fix option

diff --git a/Engine/Plugins/FX/Niagara/Source/NiagaraEditor/Private/ViewModels/Stack/NiagaraStackRendererItem.cpp b/Engine/Plugins/FX/Niagara/Source/NiagaraEditor/Private/ViewModels/Stack/NiagaraStackRendererItem.cpp
--- a/Engine/Plugins/FX/Niagara/Source/NiagaraEditor/Private/ViewModels/Stack/NiagaraStackRendererItem.cpp
+++ b/Engine/Plugins/FX/Niagara/Source/NiagaraEditor/Private/ViewModels/Stack/NiagaraStackRendererItem.cpp
@@ -595,15 +595,44 @@ void UNiagaraStackRendererItem::RefreshIssues(TArray<FStackIssue>& NewIssues)
 
 	if (RendererProperties->GetIsEnabled() && !RendererProperties->IsSimTargetSupported(EmitterData->SimTarget))
 	{
-		
-		FStackIssue TargetSupportError(
+		const ENiagaraSimTarget SimTargets[] = {ENiagaraSimTarget::CPUSim, ENiagaraSimTarget::GPUComputeSim};
+
+		TArray<FStackIssueFix> Fixes;
+		for (ENiagaraSimTarget SimTarget : SimTargets)
+		{
+			if ( !RendererProperties->IsSimTargetSupported(SimTarget) )
+			{
+				continue;
+			}
+			Fixes.Emplace(
+				FText::Format(LOCTEXT("RendererChangeSimTargetFix", "Change Sim Target to \"{0}\""), UEnum::GetDisplayValueAsText(SimTarget)),
+				FStackIssueFixDelegate::CreateLambda(
+					[WeakEmitterPtr=GetEmitterViewModel()->GetEmitter().ToWeakPtr(), SimTarget]()
+					{
+						FVersionedNiagaraEmitter VersionedEmitter = WeakEmitterPtr.ResolveWeakPtr();
+						if (FVersionedNiagaraEmitterData* VersionedEmitterData = VersionedEmitter.GetEmitterData())
+						{
+							const FScopedTransaction Transaction(LOCTEXT("ChangeSimTarget", "Change Sim Target"));
+							VersionedEmitter.Emitter->Modify();
+							VersionedEmitterData->SimTarget = SimTarget;
+
+							FProperty* SimTargetProperty = FindFProperty<FProperty>(FVersionedNiagaraEmitterData::StaticStruct(), GET_MEMBER_NAME_CHECKED(FVersionedNiagaraEmitterData, SimTarget));
+							FPropertyChangedEvent PropertyChangedEvent(SimTargetProperty);
+							VersionedEmitter.Emitter->PostEditChangeVersionedProperty(PropertyChangedEvent, VersionedEmitter.Version);
+						}
+					}
+				)
+			);
+		}
+
+		NewIssues.Emplace(
 			EStackIssueSeverity::Error,
-			LOCTEXT("FailedRendererDueToSimTarget", "Renderer incompatible with SimTarget mode."),
-			FText::Format(LOCTEXT("FailedRendererDueToSimTargetLong", "Renderer incompatible with SimTarget mode \"{0}\"."), FText::FromName(UEnum::GetValueAsName(EmitterData->SimTarget))),
+			LOCTEXT("FailedRendererDueToSimTarget", "Renderer incompatible with chosen Sim Target in Emitter Properties."),
+			FText::Format(LOCTEXT("FailedRendererDueToSimTargetLong", "Renderer incompatible with Sim Target \"{0}\"."), UEnum::GetDisplayValueAsText(EmitterData->SimTarget)),
 			GetStackEditorDataKey(),
-			false);
-
-		NewIssues.Add(TargetSupportError);
+			false,
+			Fixes
+		);
 	}
 
 	if (RendererProperties->GetIsEnabled())
