- Add GetTriangleIndices to skel mesh DI (This mirrors the functionality from static mesh DI)
- Add Interpolated methods to get skinned vertex data to match triangle / bone

diff --git a/Engine/Plugins/FX/Niagara/Shaders/Private/NiagaraDataInterfaceSkeletalMeshTemplate.ush b/Engine/Plugins/FX/Niagara/Shaders/Private/NiagaraDataInterfaceSkeletalMeshTemplate.ush
--- a/Engine/Plugins/FX/Niagara/Shaders/Private/NiagaraDataInterfaceSkeletalMeshTemplate.ush
+++ b/Engine/Plugins/FX/Niagara/Shaders/Private/NiagaraDataInterfaceSkeletalMeshTemplate.ush
@@ -638,6 +638,11 @@ void GetTriCoordVertices_{ParameterName}(in int TriangleIndex, out int VertexInd
 	GetTriVertices_{ParameterName}(TriangleIndex, VertexIndex0, VertexIndex1, VertexIndex2);
 }
 
+void GetTriangleIndices_{ParameterName}(int Tri, out int Index0, out int Index1, out int Index2)
+{
+	GetTriVertices_{ParameterName}(Tri, Index0, Index1, Index2);
+}
+
 void GetTriColor_{ParameterName}(in uint TriangleIndex, in float3 BaryCoord, out float4 OutColor)
 {
 	// Early out for bad data
@@ -676,6 +681,28 @@ void GetSkinnedVertexDataWS_{ParameterName}(in int VertexIndex, out float3 OutPo
     OutTangent  = normalize(mul(float4(SkinnedVertex.TangentX, 0.0f), {ParameterName}_InstanceTransform).xyz);
 }
 
+void GetSkinnedVertexDataInterpolated_{ParameterName}(in int VertexIndex, in float Interp, out float3 OutPosition, out float3 OutVelocity, out float3 OutNormal, out float3 OutBinormal, out float3 OutTangent)
+{
+    FDISkelMeshSkinnedVertex SkinnedVertex = GetSkinnedVertex_{ParameterName}(VertexIndex);
+    OutPosition = lerp(SkinnedVertex.PrevPosition, SkinnedVertex.Position, Interp);
+    OutVelocity = (SkinnedVertex.Position - SkinnedVertex.PrevPosition) * {ParameterName}_InstanceInvDeltaTime;
+	OutNormal = SkinnedVertex.TangentZ;
+	OutBinormal = SkinnedVertex.TangentY;
+	OutTangent = SkinnedVertex.TangentX;
+}
+
+void GetSkinnedVertexDataInterpolatedWS_{ParameterName}(in int VertexIndex, in float Interp, out float3 OutPosition, out float3 OutVelocity, out float3 OutNormal, out float3 OutBinormal, out float3 OutTangent)
+{
+    FDISkelMeshSkinnedVertex SkinnedVertex = GetSkinnedVertex_{ParameterName}(VertexIndex);
+    OutPosition = mul(float4(SkinnedVertex.Position, 1.0f), {ParameterName}_InstanceTransform).xyz;
+    const float3 PrevPosition = mul(float4(SkinnedVertex.PrevPosition, 1.0f), {ParameterName}_InstancePrevTransform).xyz;
+    OutPosition	= lerp(OutPosition, PrevPosition, Interp);
+    OutVelocity = (SkinnedVertex.Position - SkinnedVertex.PrevPosition) * {ParameterName}_InstanceInvDeltaTime;
+    OutNormal   = normalize(mul(float4(SkinnedVertex.TangentZ, 0.0f), {ParameterName}_InstanceTransform).xyz);
+    OutBinormal = normalize(mul(float4(SkinnedVertex.TangentY, 0.0f), {ParameterName}_InstanceTransform).xyz);
+    OutTangent  = normalize(mul(float4(SkinnedVertex.TangentX, 0.0f), {ParameterName}_InstanceTransform).xyz);
+}
+
 void GetVertexColor_{ParameterName}(in int VertexIndex, out float4 OutColor)
 {
     OutColor = HasMeshColors_{ParameterName}() ? {ParameterName}_MeshColorBuffer[VertexIndex] FMANUALFETCH_COLOR_COMPONENT_SWIZZLE : float4(1, 1, 1, 1);
diff --git a/Engine/Plugins/FX/Niagara/Source/Niagara/Classes/NiagaraDataInterfaceSkeletalMesh.h b/Engine/Plugins/FX/Niagara/Source/Niagara/Classes/NiagaraDataInterfaceSkeletalMesh.h
--- a/Engine/Plugins/FX/Niagara/Source/Niagara/Classes/NiagaraDataInterfaceSkeletalMesh.h
+++ b/Engine/Plugins/FX/Niagara/Source/Niagara/Classes/NiagaraDataInterfaceSkeletalMesh.h
@@ -855,6 +855,7 @@ public:
 	void IsValidTriCoord(FVectorVMExternalFunctionContext& Context);
 
 	NIAGARA_API void GetTriangleData(FVectorVMExternalFunctionContext& Context);
+	void GetTriangleIndices(FVectorVMExternalFunctionContext& Context);
 
 	template<typename SkinningHandlerType, typename TransformHandlerType, typename VertexAccessorType, typename bInterpolated>
 	void GetTriCoordSkinnedData(FVectorVMExternalFunctionContext& Context);
@@ -869,9 +870,6 @@ public:
 	template<typename VertexAccessorType>
 	void GetTriCoordUV(FVectorVMExternalFunctionContext& Context);
 
-	template<typename SkinningHandlerType>
-	void GetTriCoordVertices(FVectorVMExternalFunctionContext& Context);
-
 	template<typename VertexAccessorType>
 	void GetTriangleCoordAtUV(FVectorVMExternalFunctionContext& Context);
 
@@ -917,7 +915,7 @@ public:
 	template<typename VertexAccessorType>
 	void GetVertexData(FVectorVMExternalFunctionContext& Context);
 
-	template<typename SkinningHandlerType, typename TransformHandlerType, typename VertexAccessorType>
+	template<typename SkinningHandlerType, typename TransformHandlerType, typename Interpolated, typename VertexAccessorType>
 	void GetVertexSkinnedData(FVectorVMExternalFunctionContext& Context);
 
 	NIAGARA_API void GetVertexColor(FVectorVMExternalFunctionContext& Context);
@@ -989,6 +987,7 @@ public:
 	static const FName RandomTriCoordName;
 	static const FName IsValidTriCoordName;
 	static const FName GetTriangleDataName;
+	static const FName GetTriangleIndicesName;
 	static const FName GetSkinnedTriangleDataName;
 	static const FName GetSkinnedTriangleDataWSName;
 	static const FName GetSkinnedTriangleDataInterpName;
@@ -1033,6 +1032,8 @@ public:
 	static const FName GetVertexDataName;
 	static const FName GetSkinnedVertexDataName;
 	static const FName GetSkinnedVertexDataWSName;
+	static const FName GetSkinnedVertexDataInterpolatedName;
+	static const FName GetSkinnedVertexDataInterpolatedWSName;
 	static const FName GetVertexColorName;
 	static const FName GetVertexUVName;
 
diff --git a/Engine/Plugins/FX/Niagara/Source/Niagara/Private/NDISkeletalMesh_TriangleSampling.cpp b/Engine/Plugins/FX/Niagara/Source/Niagara/Private/NDISkeletalMesh_TriangleSampling.cpp
--- a/Engine/Plugins/FX/Niagara/Source/Niagara/Private/NDISkeletalMesh_TriangleSampling.cpp
+++ b/Engine/Plugins/FX/Niagara/Source/Niagara/Private/NDISkeletalMesh_TriangleSampling.cpp
@@ -25,13 +25,13 @@ DEFINE_NDI_FUNC_BINDER(UNiagaraDataInterfaceSkeletalMesh, GetTriCoordUV);
 DEFINE_NDI_FUNC_BINDER(UNiagaraDataInterfaceSkeletalMesh, IsValidTriCoord);
 DEFINE_NDI_FUNC_BINDER(UNiagaraDataInterfaceSkeletalMesh, GetFilteredTriangleCount);
 DEFINE_NDI_FUNC_BINDER(UNiagaraDataInterfaceSkeletalMesh, GetFilteredTriangleAt);
-DEFINE_NDI_FUNC_BINDER(UNiagaraDataInterfaceSkeletalMesh, GetTriCoordVertices);
 DEFINE_NDI_FUNC_BINDER(UNiagaraDataInterfaceSkeletalMesh, GetTriangleCoordAtUV);
 DEFINE_NDI_FUNC_BINDER(UNiagaraDataInterfaceSkeletalMesh, GetTriangleCoordInAabb);
 
 const FName FSkeletalMeshInterfaceHelper::RandomTriCoordName("RandomTriCoord");
 const FName FSkeletalMeshInterfaceHelper::IsValidTriCoordName("IsValidTriCoord");
 const FName FSkeletalMeshInterfaceHelper::GetTriangleDataName("GetTriangleData");
+const FName FSkeletalMeshInterfaceHelper::GetTriangleIndicesName("GetTriangleIndices");
 const FName FSkeletalMeshInterfaceHelper::GetSkinnedTriangleDataName("GetSkinnedTriangleData");
 const FName FSkeletalMeshInterfaceHelper::GetSkinnedTriangleDataWSName("GetSkinnedTriangleDataWS");
 const FName FSkeletalMeshInterfaceHelper::GetSkinnedTriangleDataInterpName("GetSkinnedTriangleDataInterpolated");
@@ -91,6 +91,21 @@ void UNiagaraDataInterfaceSkeletalMesh::GetTriangleSamplingFunctions(TArray<FNia
 #endif
 	}
 
+	{
+		FNiagaraFunctionSignature & Sig = OutFunctions.AddDefaulted_GetRef();
+		Sig.Name = FSkeletalMeshInterfaceHelper::GetTriangleIndicesName;
+		Sig.Inputs.Emplace(FNiagaraTypeDefinition(GetClass()), TEXT("SkeletalMesh"));
+		Sig.Inputs.Emplace(FNiagaraTypeDefinition::GetIntDef(), TEXT("Triangle"));
+		Sig.Outputs.Emplace(FNiagaraTypeDefinition::GetIntDef(), TEXT("Index0"));
+		Sig.Outputs.Emplace(FNiagaraTypeDefinition::GetIntDef(), TEXT("Index1"));
+		Sig.Outputs.Emplace(FNiagaraTypeDefinition::GetIntDef(), TEXT("Index2"));
+		Sig.bMemberFunction = true;
+		Sig.bRequiresContext = false;
+#if WITH_EDITORONLY_DATA
+		Sig.Description = LOCTEXT("GetTriangleDataDesc", "Returns bind pose triangle data.");
+#endif
+	}
+
 	{
 		FNiagaraFunctionSignature& Sig = OutFunctions.AddDefaulted_GetRef();
 		Sig.Name = FSkeletalMeshInterfaceHelper::GetSkinnedTriangleDataName;
@@ -334,6 +349,10 @@ void UNiagaraDataInterfaceSkeletalMesh::BindTriangleSamplingFunction(const FVMEx
 		check(BindingInfo.GetNumInputs() == 5 && BindingInfo.GetNumOutputs() == 12);
 		OutFunc = FVMExternalFunction::CreateLambda([this](FVectorVMExternalFunctionContext& Context) {this->GetTriangleData(Context);});
 	}
+	else if (BindingInfo.Name == FSkeletalMeshInterfaceHelper::GetTriangleIndicesName)
+	{
+		OutFunc = FVMExternalFunction::CreateLambda([this](FVectorVMExternalFunctionContext& Context) {this->GetTriangleIndices(Context); });
+	}
 	else if (BindingInfo.Name == FSkeletalMeshInterfaceHelper::GetSkinnedTriangleDataName)
 	{
 		check(BindingInfo.GetNumInputs() == 5 && BindingInfo.GetNumOutputs() == 15);
@@ -401,8 +420,7 @@ void UNiagaraDataInterfaceSkeletalMesh::BindTriangleSamplingFunction(const FVMEx
 	}
 	else if (BindingInfo.Name == FSkeletalMeshInterfaceHelper::GetTriCoordVerticesName)
 	{
-		check(BindingInfo.GetNumInputs() == 2 && BindingInfo.GetNumOutputs() == 3);
-		TSkinningModeBinder<NDI_FUNC_BINDER(UNiagaraDataInterfaceSkeletalMesh, GetTriCoordVertices)>::BindCheckCPUAccess(this, BindingInfo, InstanceData, OutFunc);		
+		OutFunc = FVMExternalFunction::CreateLambda([this](FVectorVMExternalFunctionContext& Context) {this->GetTriangleIndices(Context); });
 	}
 	else if (BindingInfo.Name == FSkeletalMeshInterfaceHelper::RandomTriangleName)
 	{
@@ -1294,6 +1312,48 @@ void UNiagaraDataInterfaceSkeletalMesh::GetTriangleData(FVectorVMExternalFunctio
 	}
 }
 
+void UNiagaraDataInterfaceSkeletalMesh::GetTriangleIndices(FVectorVMExternalFunctionContext& Context)
+{
+	SCOPE_CYCLE_COUNTER(STAT_NiagaraSkel_Sample);
+	VectorVM::FUserPtrHandler<FNDISkeletalMesh_InstanceData> InstanceData(Context);
+	FNDIInputParam<int32> TriParam(Context);
+	FNDIOutputParam<int32> OutIndex0Param(Context);
+	FNDIOutputParam<int32> OutIndex1Param(Context);
+	FNDIOutputParam<int32> OutIndex2Param(Context);
+
+	FSkeletalMeshAccessorHelper Accessor;
+	Accessor.Init<TNDISkelMesh_FilterModeNone, TNDISkelMesh_AreaWeightingOff>(InstanceData);
+
+	// Is the data valid?
+	if (!InstanceData->bAllowCPUMeshDataAccess || !Accessor.IsLODAccessible())
+	{
+		for (int32 i = 0; i < Context.GetNumInstances(); ++i)
+		{
+			OutIndex0Param.SetAndAdvance(0);
+			OutIndex1Param.SetAndAdvance(0);
+			OutIndex2Param.SetAndAdvance(0);
+		}
+		return;
+	}
+
+	// Data should be considered valid here
+	const FSkeletalMeshLODRenderData* LODData = Accessor.LODData;
+	const int32 TriMax = (Accessor.IndexBuffer->Num() / 3) - 1;
+	FSkinnedPositionAccessorHelper<TNDISkelMesh_SkinningModeNone> SkinningHandler;
+
+	for (int32 i = 0; i < Context.GetNumInstances(); ++i)
+	{
+		const int Triangle = FMath::Clamp(TriParam.GetAndAdvance(), 0, TriMax);
+
+		int32 Indices[3];
+		SkinningHandler.GetTriangleIndices(Accessor, Triangle, Indices[0], Indices[1], Indices[2]);
+
+		OutIndex0Param.SetAndAdvance(Indices[0]);
+		OutIndex1Param.SetAndAdvance(Indices[1]);
+		OutIndex2Param.SetAndAdvance(Indices[2]);
+	}
+}
+
 template<typename SkinningHandlerType, typename TransformHandlerType, typename VertexAccessorType, typename bInterpolated>
 void UNiagaraDataInterfaceSkeletalMesh::GetTriCoordSkinnedData(FVectorVMExternalFunctionContext& Context)
 {
@@ -1550,51 +1610,4 @@ void UNiagaraDataInterfaceSkeletalMesh::GetTriCoordSkinnedDataFallback(FVectorVM
 	}
 }
 
-template<typename SkinningHandlerType>
-void UNiagaraDataInterfaceSkeletalMesh::GetTriCoordVertices(FVectorVMExternalFunctionContext& Context)
-{
-	SCOPE_CYCLE_COUNTER(STAT_NiagaraSkel_Sample);
-	VectorVM::FUserPtrHandler<FNDISkeletalMesh_InstanceData> InstData(Context);
-
-	SkinningHandlerType SkinningHandler;
-	FNDIInputParam<int32> TriParam(Context);
-
-	checkf(InstData.Get(), TEXT("Skeletal Mesh Interface has invalid instance data. %s"), *GetPathName());
-	checkf(InstData->bMeshValid, TEXT("Skeletal Mesh Interface has invalid mesh. %s"), *GetPathName());
-
-	FNDIOutputParam<int32> OutV0(Context);
-	FNDIOutputParam<int32> OutV1(Context);
-	FNDIOutputParam<int32> OutV2(Context);
-
-	int32 Idx0; int32 Idx1; int32 Idx2;
-	FSkeletalMeshAccessorHelper Accessor;
-	Accessor.Init<TNDISkelMesh_FilterModeNone, TNDISkelMesh_AreaWeightingOff>(InstData);
-
-	if ( Accessor.IsLODAccessible() )
-	{
-		const int32 TriMax = (Accessor.IndexBuffer->Num() / 3) - 1;
-		if (TriMax >= 0)
-		{
-			for (int32 i = 0; i < Context.GetNumInstances(); ++i)
-			{
-				const int32 Tri = FMath::Clamp(TriParam.GetAndAdvance(), 0, TriMax);
-				SkinningHandler.GetTriangleIndices(Accessor, Tri, Idx0, Idx1, Idx2);
-				OutV0.SetAndAdvance(Idx0);
-				OutV1.SetAndAdvance(Idx1);
-				OutV2.SetAndAdvance(Idx2);
-			}
-			// We are done
-			return;
-		}
-	}
-
-	// Bad or missing data
-	for (int32 i = 0; i < Context.GetNumInstances(); ++i)
-	{
-		OutV0.SetAndAdvance(0);
-		OutV1.SetAndAdvance(0);
-		OutV2.SetAndAdvance(0);
-	}
-}
-
 #undef LOCTEXT_NAMESPACE
diff --git a/Engine/Plugins/FX/Niagara/Source/Niagara/Private/NDISkeletalMesh_VertexSampling.cpp b/Engine/Plugins/FX/Niagara/Source/Niagara/Private/NDISkeletalMesh_VertexSampling.cpp
--- a/Engine/Plugins/FX/Niagara/Source/Niagara/Private/NDISkeletalMesh_VertexSampling.cpp
+++ b/Engine/Plugins/FX/Niagara/Source/Niagara/Private/NDISkeletalMesh_VertexSampling.cpp
@@ -26,6 +26,8 @@ DEFINE_NDI_FUNC_BINDER(UNiagaraDataInterfaceSkeletalMesh, GetFilteredVertexAt);
 const FName FSkeletalMeshInterfaceHelper::GetVertexDataName("GetVertexData");
 const FName FSkeletalMeshInterfaceHelper::GetSkinnedVertexDataName("GetSkinnedVertexData");
 const FName FSkeletalMeshInterfaceHelper::GetSkinnedVertexDataWSName("GetSkinnedVertexDataWS");
+const FName FSkeletalMeshInterfaceHelper::GetSkinnedVertexDataInterpolatedName("GetSkinnedVertexDataInterpolated");
+const FName FSkeletalMeshInterfaceHelper::GetSkinnedVertexDataInterpolatedWSName("GetSkinnedVertexDataInterpolatedWS");
 const FName FSkeletalMeshInterfaceHelper::GetVertexColorName("GetVertexColor");
 const FName FSkeletalMeshInterfaceHelper::GetVertexUVName("GetVertexUV");
 
@@ -92,6 +94,44 @@ void UNiagaraDataInterfaceSkeletalMesh::GetVertexSamplingFunctions(TArray<FNiaga
 		OutFunctions.Add(Sig);
 	}
 
+	{
+		FNiagaraFunctionSignature Sig;
+		Sig.Name = FSkeletalMeshInterfaceHelper::GetSkinnedVertexDataInterpolatedName;
+		Sig.Inputs.Add(FNiagaraVariable(FNiagaraTypeDefinition(GetClass()), TEXT("SkeletalMesh")));
+		Sig.Inputs.Add(FNiagaraVariable(FNiagaraTypeDefinition::GetIntDef(), TEXT("Vertex")));
+		Sig.Inputs.Add(FNiagaraVariable(FNiagaraTypeDefinition::GetFloatDef(), TEXT("Interp")));
+		Sig.Outputs.Add(FNiagaraVariable(FNiagaraTypeDefinition::GetPositionDef(), TEXT("Position")));
+		Sig.Outputs.Add(FNiagaraVariable(FNiagaraTypeDefinition::GetVec3Def(), TEXT("Velocity")));
+		Sig.Outputs.Add(FNiagaraVariable(FNiagaraTypeDefinition::GetVec3Def(), TEXT("Normal")));
+		Sig.Outputs.Add(FNiagaraVariable(FNiagaraTypeDefinition::GetVec3Def(), TEXT("Binormal")));
+		Sig.Outputs.Add(FNiagaraVariable(FNiagaraTypeDefinition::GetVec3Def(), TEXT("Tangent")));
+		Sig.bMemberFunction = true;
+		Sig.bRequiresContext = false;
+#if WITH_EDITORONLY_DATA
+		Sig.Description = LOCTEXT("GetSkinnedDataDesc", "Returns skinning dependant data for the pased vertex in local space. All outputs are optional and you will incur zero to minimal cost if they are not connected.");
+#endif
+		OutFunctions.Add(Sig);
+	}
+
+	{
+		FNiagaraFunctionSignature Sig;
+		Sig.Name = FSkeletalMeshInterfaceHelper::GetSkinnedVertexDataInterpolatedWSName;
+		Sig.Inputs.Add(FNiagaraVariable(FNiagaraTypeDefinition(GetClass()), TEXT("SkeletalMesh")));
+		Sig.Inputs.Add(FNiagaraVariable(FNiagaraTypeDefinition::GetIntDef(), TEXT("Vertex")));
+		Sig.Inputs.Add(FNiagaraVariable(FNiagaraTypeDefinition::GetFloatDef(), TEXT("Interp")));
+		Sig.Outputs.Add(FNiagaraVariable(FNiagaraTypeDefinition::GetPositionDef(), TEXT("Position")));
+		Sig.Outputs.Add(FNiagaraVariable(FNiagaraTypeDefinition::GetVec3Def(), TEXT("Velocity")));
+		Sig.Outputs.Add(FNiagaraVariable(FNiagaraTypeDefinition::GetVec3Def(), TEXT("Normal")));
+		Sig.Outputs.Add(FNiagaraVariable(FNiagaraTypeDefinition::GetVec3Def(), TEXT("Binormal")));
+		Sig.Outputs.Add(FNiagaraVariable(FNiagaraTypeDefinition::GetVec3Def(), TEXT("Tangent")));
+		Sig.bMemberFunction = true;
+		Sig.bRequiresContext = false;
+#if WITH_EDITORONLY_DATA
+		Sig.Description = LOCTEXT("GetSkinnedDataWSDesc", "Returns skinning dependant data for the pased vertex in world space. All outputs are optional and you will incur zero to minimal cost if they are not connected.");
+#endif
+		OutFunctions.Add(Sig);
+	}
+
 	{
 		FNiagaraFunctionSignature Sig;
 		Sig.Name = FSkeletalMeshInterfaceHelper::GetVertexColorName;
@@ -189,6 +229,9 @@ void UNiagaraDataInterfaceSkeletalMesh::GetVertexSamplingFunctions(TArray<FNiaga
 
 void UNiagaraDataInterfaceSkeletalMesh::BindVertexSamplingFunction(const FVMExternalFunctionBindingInfo& BindingInfo, FNDISkeletalMesh_InstanceData* InstanceData, FVMExternalFunction &OutFunc)
 {
+	using TInterpOff = TIntegralConstant<bool, false>;
+	using TInterpOn = TIntegralConstant<bool, true>;
+
 	//////////////////////////////////////////////////////////////////////////
 	if (BindingInfo.Name == FSkeletalMeshInterfaceHelper::GetVertexDataName)
 	{
@@ -198,12 +241,22 @@ void UNiagaraDataInterfaceSkeletalMesh::BindVertexSamplingFunction(const FVMExte
 	else if (BindingInfo.Name == FSkeletalMeshInterfaceHelper::GetSkinnedVertexDataName)
 	{
 		check(BindingInfo.GetNumInputs() == 2 && BindingInfo.GetNumOutputs() == 15);
-		TSkinningModeBinder<TNDIExplicitBinder<FNDITransformHandlerNoop, TVertexAccessorBinder<NDI_FUNC_BINDER(UNiagaraDataInterfaceSkeletalMesh, GetVertexSkinnedData)>>>::BindCheckCPUAccess(this, BindingInfo, InstanceData, OutFunc);
+		TSkinningModeBinder<TNDIExplicitBinder<FNDITransformHandlerNoop, TNDIExplicitBinder<TInterpOff, TVertexAccessorBinder<NDI_FUNC_BINDER(UNiagaraDataInterfaceSkeletalMesh, GetVertexSkinnedData)>>>>::BindCheckCPUAccess(this, BindingInfo, InstanceData, OutFunc);
 	}
 	else if (BindingInfo.Name == FSkeletalMeshInterfaceHelper::GetSkinnedVertexDataWSName)
 	{
 		check(BindingInfo.GetNumInputs() == 2 && BindingInfo.GetNumOutputs() == 15);
-		TSkinningModeBinder<TNDIExplicitBinder<FNDITransformHandler, TVertexAccessorBinder<NDI_FUNC_BINDER(UNiagaraDataInterfaceSkeletalMesh, GetVertexSkinnedData)>>>::BindCheckCPUAccess(this, BindingInfo, InstanceData, OutFunc);
+		TSkinningModeBinder<TNDIExplicitBinder<FNDITransformHandler, TNDIExplicitBinder<TInterpOff, TVertexAccessorBinder<NDI_FUNC_BINDER(UNiagaraDataInterfaceSkeletalMesh, GetVertexSkinnedData)>>>>::BindCheckCPUAccess(this, BindingInfo, InstanceData, OutFunc);
+	}
+	else if (BindingInfo.Name == FSkeletalMeshInterfaceHelper::GetSkinnedVertexDataInterpolatedName)
+	{
+		check(BindingInfo.GetNumInputs() == 3 && BindingInfo.GetNumOutputs() == 15);
+		TSkinningModeBinder<TNDIExplicitBinder<FNDITransformHandlerNoop, TNDIExplicitBinder<TInterpOn, TVertexAccessorBinder<NDI_FUNC_BINDER(UNiagaraDataInterfaceSkeletalMesh, GetVertexSkinnedData)>>>>::BindCheckCPUAccess(this, BindingInfo, InstanceData, OutFunc);
+	}
+	else if (BindingInfo.Name == FSkeletalMeshInterfaceHelper::GetSkinnedVertexDataInterpolatedWSName)
+	{
+		check(BindingInfo.GetNumInputs() == 3 && BindingInfo.GetNumOutputs() == 15);
+		TSkinningModeBinder<TNDIExplicitBinder<FNDITransformHandler, TNDIExplicitBinder<TInterpOn, TVertexAccessorBinder<NDI_FUNC_BINDER(UNiagaraDataInterfaceSkeletalMesh, GetVertexSkinnedData)>>>>::BindCheckCPUAccess(this, BindingInfo, InstanceData, OutFunc);
 	}
 	else if (BindingInfo.Name == FSkeletalMeshInterfaceHelper::GetVertexColorName)
 	{
@@ -750,7 +803,7 @@ void UNiagaraDataInterfaceSkeletalMesh::GetVertexData(FVectorVMExternalFunctionC
 	}
 }
 
-template<typename SkinningHandlerType, typename TransformHandlerType, typename VertexAccessorType>
+template<typename SkinningHandlerType, typename TransformHandlerType, typename Interpolated, typename VertexAccessorType>
 void UNiagaraDataInterfaceSkeletalMesh::GetVertexSkinnedData(FVectorVMExternalFunctionContext& Context)
 {
 	SCOPE_CYCLE_COUNTER(STAT_NiagaraSkel_Vertex_Sample);
@@ -759,6 +812,12 @@ void UNiagaraDataInterfaceSkeletalMesh::GetVertexSkinnedData(FVectorVMExternalFu
 	SkinningHandlerType SkinningHandler;
 	TransformHandlerType TransformHandler;
 	FNDIInputParam<int32> VertParam(Context);
+	FNDIInputParam<float> InterpParam;
+
+	if constexpr (Interpolated::Value)
+	{
+		InterpParam.Init(Context);
+	}
 
 	checkfSlow(InstData.Get(), TEXT("Skeletal Mesh Interface has invalid instance data. %s"), *GetPathName());
 
@@ -780,20 +839,37 @@ void UNiagaraDataInterfaceSkeletalMesh::GetVertexSkinnedData(FVectorVMExternalFu
 			for (int32 i = 0; i < Context.GetNumInstances(); ++i)
 			{
 				const int32 Vertex = FMath::Clamp(VertParam.GetAndAdvance(), 0, VertMax);
+				float Interp = 1.0f;
+				if constexpr (Interpolated::Value)
+				{
+					Interp = InterpParam.GetAndAdvance();
+				}
 
-				FVector3f Pos = FVector3f::ZeroVector;
+				FVector3f CurrPos = FVector3f::ZeroVector;
+				FVector3f PrevPos = FVector3f::ZeroVector;
 				if (Output.bNeedsPosition || Output.bNeedsVelocity)
 				{
-					Pos = SkinningHandler.GetSkinnedVertexPosition(Accessor, Vertex);
-					TransformHandler.TransformPosition(Pos, Transform);
-					Output.Position.SetAndAdvance(Pos);
+					CurrPos = SkinningHandler.GetSkinnedVertexPosition(Accessor, Vertex);
+					TransformHandler.TransformPosition(CurrPos, Transform);
+					if (Output.bNeedsVelocity || Interpolated::Value)
+					{
+						PrevPos = SkinningHandler.GetSkinnedVertexPreviousPosition(Accessor, Vertex);
+						TransformHandler.TransformPosition(PrevPos, PrevTransform);
+					}
+
+					if constexpr (Interpolated::Value)
+					{
+						Output.Position.SetAndAdvance(FMath::Lerp(PrevPos, CurrPos, Interp));
+					}
+					else
+					{
+						Output.Position.SetAndAdvance(CurrPos);
+					}
 				}
 
 				if (Output.bNeedsVelocity)
 				{
-					FVector3f Prev = SkinningHandler.GetSkinnedVertexPreviousPosition(Accessor, Vertex);
-					TransformHandler.TransformPosition(Prev, PrevTransform);
-					const FVector3f Velocity = (Pos - Prev) * InvDt;
+					const FVector3f Velocity = (CurrPos - PrevPos) * InvDt;
 					Output.Velocity.SetAndAdvance(Velocity);
 				}
 
diff --git a/Engine/Plugins/FX/Niagara/Source/Niagara/Private/NiagaraDataInterfaceSkeletalMesh.cpp b/Engine/Plugins/FX/Niagara/Source/Niagara/Private/NiagaraDataInterfaceSkeletalMesh.cpp
--- a/Engine/Plugins/FX/Niagara/Source/Niagara/Private/NiagaraDataInterfaceSkeletalMesh.cpp
+++ b/Engine/Plugins/FX/Niagara/Source/Niagara/Private/NiagaraDataInterfaceSkeletalMesh.cpp
@@ -2824,6 +2824,8 @@ bool UNiagaraDataInterfaceSkeletalMesh::GetFunctionHLSL(const FNiagaraDataInterf
 	{
 		NDISkelMeshLocal::NAME_GetPreSkinnedLocalBounds,
 		//////////////////////////////////////////////////////////////////////////////////////////////////////////////////////////////
+		// Triangle Sampling
+		FSkeletalMeshInterfaceHelper::GetTriangleIndicesName,
 		FSkeletalMeshInterfaceHelper::GetTriCoordVerticesName,
 		FSkeletalMeshInterfaceHelper::GetTriangleCountName,
 		FSkeletalMeshInterfaceHelper::GetFilteredTriangleCountName,
@@ -2838,6 +2840,8 @@ bool UNiagaraDataInterfaceSkeletalMesh::GetFunctionHLSL(const FNiagaraDataInterf
 		FSkeletalMeshInterfaceHelper::GetVertexDataName,
 		FSkeletalMeshInterfaceHelper::GetSkinnedVertexDataName,
 		FSkeletalMeshInterfaceHelper::GetSkinnedVertexDataWSName,
+		FSkeletalMeshInterfaceHelper::GetSkinnedVertexDataInterpolatedName,
+		FSkeletalMeshInterfaceHelper::GetSkinnedVertexDataInterpolatedWSName,
 		FSkeletalMeshInterfaceHelper::GetVertexColorName,
 		FSkeletalMeshInterfaceHelper::GetVertexUVName,
 		FSkeletalMeshInterfaceHelper::IsValidVertexName,
