- Fix for Engine.Emitter.NumParticles & Engine.System.NumParticles not working with GPU emitters
- Engine.Emitter.NumParticles will be the GPU count no the GPU count so will be latent

diff --git a/Engine/Plugins/FX/Niagara/Source/NiagaraShader/Public/NiagaraShader.h b/Engine/Plugins/FX/Niagara/Source/NiagaraShader/Public/NiagaraShader.h
--- a/Engine/Plugins/FX/Niagara/Source/NiagaraShader/Public/NiagaraShader.h
+++ b/Engine/Plugins/FX/Niagara/Source/NiagaraShader/Public/NiagaraShader.h
@@ -69,9 +69,9 @@ public:
 		SHADER_PARAMETER(int32,		Engine_System_CurrentTimeStep)
 		SHADER_PARAMETER(int32,		Engine_System_NumTimeSteps)
 		SHADER_PARAMETER(float,		Engine_System_TimeStepFraction)
+		SHADER_PARAMETER(int32,		Engine_System_NumParticles)
 		SHADER_PARAMETER(int32,		System_Pad0)
 		SHADER_PARAMETER(int32,		System_Pad1)
-		SHADER_PARAMETER(int32,		System_Pad2)
 
 		SHADER_PARAMETER(float,		PREV_Engine_Owner_TimeSinceRendered)
 		SHADER_PARAMETER(float,		PREV_Engine_Owner_LODDistance)
@@ -86,9 +86,9 @@ public:
 		SHADER_PARAMETER(int32,		PREV_CurrentTimeStep)
 		SHADER_PARAMETER(int32,		PREV_NumTimeSteps)
 		SHADER_PARAMETER(float,		PREV_TimeStepFraction)
+		SHADER_PARAMETER(int32,		PREV_Engine_System_NumParticles)
 		SHADER_PARAMETER(int32,		PREV_System_Pad0)
 		SHADER_PARAMETER(int32,		PREV_System_Pad1)
-		SHADER_PARAMETER(int32,		PREV_System_Pad2)
 	END_SHADER_PARAMETER_STRUCT()
 
 	// This structure is a replication of FNiagaraOwnerParameters with interpolated parameters includes
@@ -138,7 +138,7 @@ public:
 
 	// This structure is a replication of FNiagaraEmitterParameters with interpolated parameters includes
 	BEGIN_SHADER_PARAMETER_STRUCT(FEmitterParameters, )
-		SHADER_PARAMETER(int32,		Engine_Emitter_NumParticle)
+		SHADER_PARAMETER(int32,		Engine_Emitter_NumParticles)
 		SHADER_PARAMETER(int32,		Engine_Emitter_TotalSpawnedParticles)
 		SHADER_PARAMETER(float,		Engine_Emitter_SpawnCountScale)
 		SHADER_PARAMETER(float,		Emitter_Age)
@@ -147,7 +147,7 @@ public:
 		SHADER_PARAMETER(int32,		Emitter_Pad0)
 		SHADER_PARAMETER(int32,		Emitter_Pad1)
 
-		SHADER_PARAMETER(int32,		PREV_Engine_Emitter_NumParticle)
+		SHADER_PARAMETER(int32,		PREV_Engine_Emitter_NumParticles)
 		SHADER_PARAMETER(int32,		PREV_Engine_Emitter_TotalSpawnedParticles)
 		SHADER_PARAMETER(float,		PREV_Engine_Emitter_SpawnCountScale)
 		SHADER_PARAMETER(float,		PREV_Emitter_Age)
