- Change OnSystemComplete delegate to be DuplicateTransient

diff --git a/Engine/Plugins/FX/Niagara/Source/Niagara/Public/NiagaraComponent.h b/Engine/Plugins/FX/Niagara/Source/Niagara/Public/NiagaraComponent.h
--- a/Engine/Plugins/FX/Niagara/Source/Niagara/Public/NiagaraComponent.h
+++ b/Engine/Plugins/FX/Niagara/Source/Niagara/Public/NiagaraComponent.h
@@ -667,7 +667,7 @@ public:
 	//~ End UObject Interface.
 
 	// Called when the particle system is done
-	UPROPERTY(BlueprintAssignable)
+	UPROPERTY(BlueprintAssignable, DuplicateTransient)
 	FOnNiagaraSystemFinished OnSystemFinished;
 
 	/** Removes all local overrides and replaces them with the values from the source System - note: this also removes the editor overrides from the component as it is used by the pooling mechanism to prevent values leaking between different instances. */
