- Methods to get the vertex data for a skinned triangle
- This is to reduce VVM calls by bundling all data in one function

diff --git a/Engine/Plugins/FX/Niagara/Shaders/Private/NiagaraDataInterfaceSkeletalMeshTemplate.ush b/Engine/Plugins/FX/Niagara/Shaders/Private/NiagaraDataInterfaceSkeletalMeshTemplate.ush
--- a/Engine/Plugins/FX/Niagara/Shaders/Private/NiagaraDataInterfaceSkeletalMeshTemplate.ush
+++ b/Engine/Plugins/FX/Niagara/Shaders/Private/NiagaraDataInterfaceSkeletalMeshTemplate.ush
@@ -269,6 +269,17 @@ FDISkelMeshSkinnedVertex GetSkinnedVertex_{ParameterName}(uint VertexIndex)
 	return SkinnedVertex;
 }
 
+FDISkelMeshSkinnedVertex GetSkinnedVertexWS_{ParameterName}(uint VertexIndex)
+{
+	FDISkelMeshSkinnedVertex SkinnedVertex = GetSkinnedVertex_{ParameterName}(VertexIndex);
+    SkinnedVertex.Position		= mul(float4(SkinnedVertex.Position, 1.0f), {ParameterName}_InstanceTransform).xyz;
+    SkinnedVertex.PrevPosition	= mul(float4(SkinnedVertex.PrevPosition, 1.0f), {ParameterName}_InstancePrevTransform).xyz;
+    SkinnedVertex.TangentX		= normalize(mul(float4(SkinnedVertex.TangentX, 0.0f), {ParameterName}_InstanceTransform).xyz);
+    SkinnedVertex.TangentY		= normalize(mul(float4(SkinnedVertex.TangentY, 0.0f), {ParameterName}_InstanceTransform).xyz);
+    SkinnedVertex.TangentZ		= normalize(mul(float4(SkinnedVertex.TangentZ, 0.0f), {ParameterName}_InstanceTransform).xyz);
+	return SkinnedVertex;
+}
+
 FDISkelMeshSkinnedVertex GetSkinnedPointOnTriangle_{ParameterName}(uint TriangleIndex, float3 BaryCoord)
 {
 	// Early out for bad data
@@ -404,6 +415,118 @@ void GetSkinnedTriangleDataInterpolated_{ParameterName}(in uint TriangleIndex, i
 	OutTangent	= normalize(SkinnedVertex.TangentX);
 }
 
+void GetSkinnedTriangleVertexDataHelper_{ParameterName}(
+	int TriangleIndex, 
+	float Interp,
+	bool bWorldSpace,
+	out float3 Position0, out float3 Velocity0, out float3 Normal0, out float3 Binormal0, out float3 Tangent0,
+	out float3 Position1, out float3 Velocity1, out float3 Normal1, out float3 Binormal1, out float3 Tangent1,
+	out float3 Position2, out float3 Velocity2, out float3 Normal2, out float3 Binormal2, out float3 Tangent2
+)
+{
+    int VertexIndex0 = 0;
+    int VertexIndex1 = 0;
+    int VertexIndex2 = 0;
+	if ( {ParameterName}_MeshTriangleCount != 0 )
+	{
+		TriangleIndex = clamp(TriangleIndex, 0, {ParameterName}_MeshTriangleCount - 1);
+		const uint IndexBufferOffset = TriangleIndex * 3;
+		VertexIndex0 = {ParameterName}_MeshIndexBuffer[IndexBufferOffset];
+		VertexIndex1 = {ParameterName}_MeshIndexBuffer[IndexBufferOffset + 1];
+		VertexIndex2 = {ParameterName}_MeshIndexBuffer[IndexBufferOffset + 2];
+	}
+
+	FDISkelMeshSkinnedVertex SkinnedVertex0;
+	FDISkelMeshSkinnedVertex SkinnedVertex1;
+	FDISkelMeshSkinnedVertex SkinnedVertex2;
+
+	BRANCH
+	if ( bWorldSpace )
+	{
+		SkinnedVertex0 = GetSkinnedVertexWS_{ParameterName}(VertexIndex0);
+		SkinnedVertex1 = GetSkinnedVertexWS_{ParameterName}(VertexIndex1);
+		SkinnedVertex2 = GetSkinnedVertexWS_{ParameterName}(VertexIndex2);
+	}
+	else
+	{
+		SkinnedVertex0 = GetSkinnedVertex_{ParameterName}(VertexIndex0);
+		SkinnedVertex1 = GetSkinnedVertex_{ParameterName}(VertexIndex1);
+		SkinnedVertex2 = GetSkinnedVertex_{ParameterName}(VertexIndex2);
+	}
+
+	Position0		= lerp(SkinnedVertex0.PrevPosition, SkinnedVertex0.Position, Interp);
+	Velocity0		= SkinnedVertex0.Position - SkinnedVertex0.PrevPosition;
+	Normal0			= SkinnedVertex0.TangentZ;
+	Binormal0		= SkinnedVertex0.TangentY;
+	Tangent0		= SkinnedVertex0.TangentX;
+
+	Position1		= lerp(SkinnedVertex1.PrevPosition, SkinnedVertex1.Position, Interp);
+	Velocity1		= SkinnedVertex1.Position - SkinnedVertex1.PrevPosition;
+	Normal1			= SkinnedVertex1.TangentZ;
+	Binormal1		= SkinnedVertex1.TangentY;
+	Tangent1		= SkinnedVertex1.TangentX;
+
+	Position2		= lerp(SkinnedVertex2.PrevPosition, SkinnedVertex2.Position, Interp);
+	Velocity2		= SkinnedVertex2.Position - SkinnedVertex2.PrevPosition;
+	Normal2			= SkinnedVertex2.TangentZ;
+	Binormal2		= SkinnedVertex2.TangentY;
+	Tangent2		= SkinnedVertex2.TangentX;
+}
+
+void GetSkinnedTriangleVertexData_{ParameterName}(
+	int TriangleIndex, 
+	out float3 Position0, out float3 Velocity0, out float3 Normal0, out float3 Binormal0, out float3 Tangent0,
+	out float3 Position1, out float3 Velocity1, out float3 Normal1, out float3 Binormal1, out float3 Tangent1,
+	out float3 Position2, out float3 Velocity2, out float3 Normal2, out float3 Binormal2, out float3 Tangent2
+)
+{
+	GetSkinnedTriangleVertexDataHelper_{ParameterName}(
+		TriangleIndex, 1.0f, false,
+		Position0, Velocity0, Normal0, Binormal0, Tangent0, Position1, Velocity1, Normal1, Binormal1, Tangent1, Position2, Velocity2, Normal2, Binormal2, Tangent2
+	);
+}
+
+void GetSkinnedTriangleVertexDataWS_{ParameterName}(
+	int TriangleIndex, 
+	out float3 Position0, out float3 Velocity0, out float3 Normal0, out float3 Binormal0, out float3 Tangent0,
+	out float3 Position1, out float3 Velocity1, out float3 Normal1, out float3 Binormal1, out float3 Tangent1,
+	out float3 Position2, out float3 Velocity2, out float3 Normal2, out float3 Binormal2, out float3 Tangent2
+)
+{
+	GetSkinnedTriangleVertexDataHelper_{ParameterName}(
+		TriangleIndex, 1.0f, true,
+		Position0, Velocity0, Normal0, Binormal0, Tangent0, Position1, Velocity1, Normal1, Binormal1, Tangent1, Position2, Velocity2, Normal2, Binormal2, Tangent2
+	);
+}
+
+void GetSkinnedTriangleVertexDataInterpolated_{ParameterName}(
+	int TriangleIndex, 
+	float Interp,
+	out float3 Position0, out float3 Velocity0, out float3 Normal0, out float3 Binormal0, out float3 Tangent0,
+	out float3 Position1, out float3 Velocity1, out float3 Normal1, out float3 Binormal1, out float3 Tangent1,
+	out float3 Position2, out float3 Velocity2, out float3 Normal2, out float3 Binormal2, out float3 Tangent2
+)
+{
+	GetSkinnedTriangleVertexDataHelper_{ParameterName}(
+		TriangleIndex, Interp, false,
+		Position0, Velocity0, Normal0, Binormal0, Tangent0, Position1, Velocity1, Normal1, Binormal1, Tangent1, Position2, Velocity2, Normal2, Binormal2, Tangent2
+	);
+}
+
+void GetSkinnedTriangleVertexDataWSInterpolated_{ParameterName}(
+	int TriangleIndex, 
+	float Interp,
+	out float3 Position0, out float3 Velocity0, out float3 Normal0, out float3 Binormal0, out float3 Tangent0,
+	out float3 Position1, out float3 Velocity1, out float3 Normal1, out float3 Binormal1, out float3 Tangent1,
+	out float3 Position2, out float3 Velocity2, out float3 Normal2, out float3 Binormal2, out float3 Tangent2
+)
+{
+	GetSkinnedTriangleVertexDataHelper_{ParameterName}(
+		TriangleIndex, Interp, true,
+		Position0, Velocity0, Normal0, Binormal0, Tangent0, Position1, Velocity1, Normal1, Binormal1, Tangent1, Position2, Velocity2, Normal2, Binormal2, Tangent2
+	);
+}
+
 void GetTriUV_{ParameterName}(in uint TriangleIndex, in float3 BaryCoord, in int UVSet, out float2 OutUV)
 {
     if ({ParameterName}_MeshNumTexCoord > 0)
diff --git a/Engine/Plugins/FX/Niagara/Source/Niagara/Classes/NiagaraDataInterfaceSkeletalMesh.h b/Engine/Plugins/FX/Niagara/Source/Niagara/Classes/NiagaraDataInterfaceSkeletalMesh.h
--- a/Engine/Plugins/FX/Niagara/Source/Niagara/Classes/NiagaraDataInterfaceSkeletalMesh.h
+++ b/Engine/Plugins/FX/Niagara/Source/Niagara/Classes/NiagaraDataInterfaceSkeletalMesh.h
@@ -863,6 +863,9 @@ public:
 	template<typename TransformHandlerType, typename bInterpolated>
 	void GetTriCoordSkinnedDataFallback(FVectorVMExternalFunctionContext& Context);
 
+	template<typename SkinningHandlerType, typename TransformHandlerType, typename VertexAccessorType, typename bInterpolated>
+	void VMGetSkinnedTriangleVertexData(FVectorVMExternalFunctionContext& Context);
+
 	NIAGARA_API void GetTriCoordColor(FVectorVMExternalFunctionContext& Context);
 
 	NIAGARA_API void GetTriCoordColorFallback(FVectorVMExternalFunctionContext& Context);
@@ -992,6 +995,10 @@ public:
 	static const FName GetSkinnedTriangleDataWSName;
 	static const FName GetSkinnedTriangleDataInterpName;
 	static const FName GetSkinnedTriangleDataWSInterpName;
+	static const FName GetSkinnedTriangleVertexDataName;
+	static const FName GetSkinnedTriangleVertexDataWSName;
+	static const FName GetSkinnedTriangleVertexDataInterpName;
+	static const FName GetSkinnedTriangleVertexDataWSInterpName;
 	static const FName GetTriColorName;
 	static const FName GetTriUVName;
 	static const FName GetTriCoordVerticesName;
diff --git a/Engine/Plugins/FX/Niagara/Source/Niagara/Private/NDISkeletalMesh_TriangleSampling.cpp b/Engine/Plugins/FX/Niagara/Source/Niagara/Private/NDISkeletalMesh_TriangleSampling.cpp
--- a/Engine/Plugins/FX/Niagara/Source/Niagara/Private/NDISkeletalMesh_TriangleSampling.cpp
+++ b/Engine/Plugins/FX/Niagara/Source/Niagara/Private/NDISkeletalMesh_TriangleSampling.cpp
@@ -19,6 +19,7 @@ DECLARE_CYCLE_STAT(TEXT("Skel Mesh Sampling"), STAT_NiagaraSkel_Sample, STATGROU
 DEFINE_NDI_FUNC_BINDER(UNiagaraDataInterfaceSkeletalMesh, RandomTriCoord);
 DEFINE_NDI_FUNC_BINDER(UNiagaraDataInterfaceSkeletalMesh, GetTriCoordSkinnedData);
 DEFINE_NDI_FUNC_BINDER(UNiagaraDataInterfaceSkeletalMesh, GetTriCoordSkinnedDataFallback);
+DEFINE_NDI_FUNC_BINDER(UNiagaraDataInterfaceSkeletalMesh, VMGetSkinnedTriangleVertexData);
 DEFINE_NDI_DIRECT_FUNC_BINDER(UNiagaraDataInterfaceSkeletalMesh, GetTriCoordColor);
 DEFINE_NDI_DIRECT_FUNC_BINDER(UNiagaraDataInterfaceSkeletalMesh, GetTriCoordColorFallback);
 DEFINE_NDI_FUNC_BINDER(UNiagaraDataInterfaceSkeletalMesh, GetTriCoordUV);
@@ -36,6 +37,10 @@ const FName FSkeletalMeshInterfaceHelper::GetSkinnedTriangleDataName("GetSkinned
 const FName FSkeletalMeshInterfaceHelper::GetSkinnedTriangleDataWSName("GetSkinnedTriangleDataWS");
 const FName FSkeletalMeshInterfaceHelper::GetSkinnedTriangleDataInterpName("GetSkinnedTriangleDataInterpolated");
 const FName FSkeletalMeshInterfaceHelper::GetSkinnedTriangleDataWSInterpName("GetSkinnedTriangleDataWSInterpolated");
+const FName FSkeletalMeshInterfaceHelper::GetSkinnedTriangleVertexDataName("GetSkinnedTriangleVertexData");
+const FName FSkeletalMeshInterfaceHelper::GetSkinnedTriangleVertexDataWSName("GetSkinnedTriangleVertexDataWS");
+const FName FSkeletalMeshInterfaceHelper::GetSkinnedTriangleVertexDataInterpName("GetSkinnedTriangleVertexDataInterpolated");
+const FName FSkeletalMeshInterfaceHelper::GetSkinnedTriangleVertexDataWSInterpName("GetSkinnedTriangleVertexDataWSInterpolated");
 const FName FSkeletalMeshInterfaceHelper::GetTriColorName("GetTriColor");
 const FName FSkeletalMeshInterfaceHelper::GetTriUVName("GetTriUV");
 const FName FSkeletalMeshInterfaceHelper::GetTriCoordVerticesName("GetTriCoordVertices");
@@ -176,6 +181,36 @@ void UNiagaraDataInterfaceSkeletalMesh::GetTriangleSamplingFunctions(TArray<FNia
 #endif
 	}
 
+	{
+		FNiagaraFunctionSignature BaseSig;
+		BaseSig.Inputs.Emplace(FNiagaraTypeDefinition(GetClass()), TEXT("SkeletalMesh"));
+		BaseSig.Inputs.Emplace(FNiagaraTypeDefinition::GetIntDef(), TEXT("Triangle"));
+		BaseSig.Outputs.Emplace(FNiagaraTypeDefinition::GetPositionDef(), TEXT("Position 0"));
+		BaseSig.Outputs.Emplace(FNiagaraTypeDefinition::GetVec3Def(), TEXT("Velocity 0"));
+		BaseSig.Outputs.Emplace(FNiagaraTypeDefinition::GetVec3Def(), TEXT("Normal 0"));
+		BaseSig.Outputs.Emplace(FNiagaraTypeDefinition::GetVec3Def(), TEXT("Binormal 0"));
+		BaseSig.Outputs.Emplace(FNiagaraTypeDefinition::GetVec3Def(), TEXT("Tangent 0"));
+		BaseSig.Outputs.Emplace(FNiagaraTypeDefinition::GetPositionDef(), TEXT("Position 1"));
+		BaseSig.Outputs.Emplace(FNiagaraTypeDefinition::GetVec3Def(), TEXT("Velocity 1"));
+		BaseSig.Outputs.Emplace(FNiagaraTypeDefinition::GetVec3Def(), TEXT("Normal 1"));
+		BaseSig.Outputs.Emplace(FNiagaraTypeDefinition::GetVec3Def(), TEXT("Binormal 1"));
+		BaseSig.Outputs.Emplace(FNiagaraTypeDefinition::GetVec3Def(), TEXT("Tangent 1"));
+		BaseSig.Outputs.Emplace(FNiagaraTypeDefinition::GetPositionDef(), TEXT("Position 2"));
+		BaseSig.Outputs.Emplace(FNiagaraTypeDefinition::GetVec3Def(), TEXT("Velocity 2"));
+		BaseSig.Outputs.Emplace(FNiagaraTypeDefinition::GetVec3Def(), TEXT("Normal 2"));
+		BaseSig.Outputs.Emplace(FNiagaraTypeDefinition::GetVec3Def(), TEXT("Binormal 2"));
+		BaseSig.Outputs.Emplace(FNiagaraTypeDefinition::GetVec3Def(), TEXT("Tangent 2"));
+		BaseSig.bMemberFunction = true;
+		BaseSig.bRequiresContext = false;
+
+		OutFunctions.Add_GetRef(BaseSig).Name = FSkeletalMeshInterfaceHelper::GetSkinnedTriangleVertexDataName;
+		OutFunctions.Add_GetRef(BaseSig).Name = FSkeletalMeshInterfaceHelper::GetSkinnedTriangleVertexDataWSName;
+
+		BaseSig.Inputs.Emplace(FNiagaraTypeDefinition::GetFloatDef(), TEXT("Interp"));
+		OutFunctions.Add_GetRef(BaseSig).Name = FSkeletalMeshInterfaceHelper::GetSkinnedTriangleVertexDataInterpName;
+		OutFunctions.Add_GetRef(BaseSig).Name = FSkeletalMeshInterfaceHelper::GetSkinnedTriangleVertexDataWSInterpName;
+	}
+
 	{
 		FNiagaraFunctionSignature& Sig = OutFunctions.AddDefaulted_GetRef();
 		Sig.Name = FSkeletalMeshInterfaceHelper::GetTriColorName;
@@ -401,6 +436,22 @@ void UNiagaraDataInterfaceSkeletalMesh::BindTriangleSamplingFunction(const FVMEx
 			NDI_FUNC_BINDER(UNiagaraDataInterfaceSkeletalMesh, GetTriCoordSkinnedDataFallback)::Bind<FNDITransformHandler, TInterpOn>(this, BindingInfo, InstanceData, OutFunc);
 		}
 	}
+	else if (BindingInfo.Name == FSkeletalMeshInterfaceHelper::GetSkinnedTriangleVertexDataName)
+	{
+		TSkinningModeBinder<TNDIExplicitBinder<FNDITransformHandlerNoop, TVertexAccessorBinder<TNDIExplicitBinder<TIntegralConstant<bool, false>, NDI_FUNC_BINDER(UNiagaraDataInterfaceSkeletalMesh, VMGetSkinnedTriangleVertexData)>>>>::BindIgnoreCPUAccess(this, BindingInfo, InstanceData, OutFunc);
+	}
+	else if (BindingInfo.Name == FSkeletalMeshInterfaceHelper::GetSkinnedTriangleVertexDataWSName)
+	{
+		TSkinningModeBinder<TNDIExplicitBinder<FNDITransformHandler, TVertexAccessorBinder<TNDIExplicitBinder<TIntegralConstant<bool, false>, NDI_FUNC_BINDER(UNiagaraDataInterfaceSkeletalMesh, VMGetSkinnedTriangleVertexData)>>>>::BindIgnoreCPUAccess(this, BindingInfo, InstanceData, OutFunc);
+	}
+	else if (BindingInfo.Name == FSkeletalMeshInterfaceHelper::GetSkinnedTriangleVertexDataInterpName)
+	{
+		TSkinningModeBinder<TNDIExplicitBinder<FNDITransformHandlerNoop, TVertexAccessorBinder<TNDIExplicitBinder<TIntegralConstant<bool, true>, NDI_FUNC_BINDER(UNiagaraDataInterfaceSkeletalMesh, VMGetSkinnedTriangleVertexData)>>>>::BindIgnoreCPUAccess(this, BindingInfo, InstanceData, OutFunc);
+	}
+	else if (BindingInfo.Name == FSkeletalMeshInterfaceHelper::GetSkinnedTriangleVertexDataWSInterpName)
+	{
+		TSkinningModeBinder<TNDIExplicitBinder<FNDITransformHandler, TVertexAccessorBinder<TNDIExplicitBinder<TIntegralConstant<bool, true>, NDI_FUNC_BINDER(UNiagaraDataInterfaceSkeletalMesh, VMGetSkinnedTriangleVertexData)>>>>::BindIgnoreCPUAccess(this, BindingInfo, InstanceData, OutFunc);
+	}
 	else if (BindingInfo.Name == FSkeletalMeshInterfaceHelper::GetTriColorName)
 	{
 		check(BindingInfo.GetNumInputs() == 5 && BindingInfo.GetNumOutputs() == 4);
@@ -1610,4 +1661,156 @@ void UNiagaraDataInterfaceSkeletalMesh::GetTriCoordSkinnedDataFallback(FVectorVM
 	}
 }
 
+template<typename SkinningHandlerType, typename TransformHandlerType, typename VertexAccessorType, typename bInterpolated>
+void UNiagaraDataInterfaceSkeletalMesh::VMGetSkinnedTriangleVertexData(FVectorVMExternalFunctionContext& Context)
+{
+	SCOPE_CYCLE_COUNTER(STAT_NiagaraSkel_Sample);
+
+	VectorVM::FUserPtrHandler<FNDISkeletalMesh_InstanceData> InstData(Context);
+	FNDIInputParam<int32> TriangleParam(Context);
+	FNDIInputParam<float> InterpParam;
+	if constexpr (bInterpolated::Value)
+	{
+		InterpParam.Init(Context);
+	}
+
+	struct FVertexOutput
+	{
+		explicit FVertexOutput(FVectorVMExternalFunctionContext& Context)
+			: OutPosition(Context)
+			, OutVelocity(Context)
+			, OutNormal(Context)
+			, OutBinormal(Context)
+			, OutTangent(Context)
+		{
+		}
+
+		bool WritesPosition() const { return OutPosition.IsValid(); }
+		bool WritesVelocity() const { return OutVelocity.IsValid(); }
+		bool WritesTangents() const { return OutNormal.IsValid() || OutBinormal.IsValid() || OutTangent.IsValid(); }
+
+		FNDIOutputParam<FVector3f>	OutPosition;
+		FNDIOutputParam<FVector3f>	OutVelocity;
+		FNDIOutputParam<FVector3f>	OutNormal;
+		FNDIOutputParam<FVector3f>	OutBinormal;
+		FNDIOutputParam<FVector3f>	OutTangent;
+	};
+	FVertexOutput VertexOutput[] = {FVertexOutput(Context), FVertexOutput(Context), FVertexOutput(Context)};
+
+	const bool bWritesCurrPos = VertexOutput[0].WritesPosition() || VertexOutput[1].WritesPosition() || VertexOutput[2].WritesPosition();
+	const bool bWritesVelocity = VertexOutput[0].WritesVelocity() || VertexOutput[1].WritesVelocity() || VertexOutput[2].WritesVelocity();
+	const bool bWritesTangentBasis = VertexOutput[0].WritesTangents() || VertexOutput[1].WritesTangents() || VertexOutput[2].WritesTangents();
+	const bool bNeedsPrevPos = bInterpolated::Value || bWritesVelocity;
+	const float InvDt = InstData->DeltaSeconds > 0.0f ? 1.0f / InstData->DeltaSeconds : 0.0f;
+
+	SkinningHandlerType SkinningHandler;
+	TransformHandlerType TransformHandler;
+	const FMatrix44f Transform(InstData->Transform);			// LWC_TODO: Precision loss
+	const FMatrix44f PrevTransform(InstData->PrevTransform);
+
+	FSkeletalMeshAccessorHelper Accessor;
+	Accessor.Init<TNDISkelMesh_FilterModeNone, TNDISkelMesh_AreaWeightingOff>(InstData);
+
+	// Invalid of missing data path
+	if (!Accessor.IsLODAccessible())
+	{
+		for (int32 i=0; i < Context.GetNumInstances(); ++i)
+		{
+			const int32 Triangle = TriangleParam.GetAndAdvance();
+			const float Interp = bInterpolated::Value ? InterpParam.GetAndAdvance() : 1.0f;
+
+			for (int32 Vertex=0; Vertex < 3; ++Vertex)
+			{
+				FVector3f CurrPosition = FVector3f::ZeroVector;
+				if (bWritesCurrPos || bWritesVelocity)
+				{
+					TransformHandler.TransformPosition(CurrPosition, Transform);
+				}
+
+				FVector3f PrevPosition = FVector3f::ZeroVector;
+				if (bNeedsPrevPos)
+				{
+					TransformHandler.TransformPosition(PrevPosition, PrevTransform);
+				}
+
+				if (bWritesCurrPos)
+				{
+					if constexpr (bInterpolated::Value)
+					{
+						VertexOutput[Vertex].OutPosition.SetAndAdvance(FMath::Lerp(PrevPosition, CurrPosition, Interp));
+					}
+					else
+					{
+						VertexOutput[Vertex].OutPosition.SetAndAdvance(CurrPosition);
+					}
+				}
+				if (bWritesVelocity)
+				{
+					VertexOutput[Vertex].OutVelocity.SetAndAdvance((CurrPosition - PrevPosition)* InvDt);
+				}
+				if (bWritesTangentBasis)
+				{
+					VertexOutput[Vertex].OutNormal.SetAndAdvance(FVector3f::ZAxisVector);
+					VertexOutput[Vertex].OutBinormal.SetAndAdvance(FVector3f::YAxisVector);
+					VertexOutput[Vertex].OutTangent.SetAndAdvance(FVector3f::XAxisVector);
+				}
+			}
+		}
+		return;
+	}
+
+	// Do the sampling
+	for (int32 i=0; i < Context.GetNumInstances(); ++i)
+	{
+		const int32 Triangle = TriangleParam.GetAndAdvance();
+		const float Interp = bInterpolated::Value ? InterpParam.GetAndAdvance() : 1.0f;
+
+		int32 Indices[3];
+		SkinningHandler.GetTriangleIndices(Accessor, Triangle, Indices[0], Indices[1], Indices[2]);
+
+		for (int32 Index=0; Index < 3; ++Index)
+		{
+			const int32 Vertex = Indices[Index];
+
+			FVector3f CurrPosition = FVector3f::ZeroVector;
+			if (bWritesCurrPos || bWritesVelocity)
+			{
+				CurrPosition = SkinningHandler.GetSkinnedVertexPosition(Accessor, Vertex);
+				TransformHandler.TransformPosition(CurrPosition, Transform);
+			}
+
+			FVector3f PrevPosition = FVector3f::ZeroVector;
+			if (bNeedsPrevPos)
+			{
+				PrevPosition = SkinningHandler.GetSkinnedVertexPreviousPosition(Accessor, Vertex);
+				TransformHandler.TransformPosition(PrevPosition, PrevTransform);
+			}
+
+			if (bWritesCurrPos)
+			{
+				const FVector3f InterpPosition = bInterpolated::Value ? FMath::Lerp(PrevPosition, CurrPosition, Interp) : CurrPosition;
+				VertexOutput[Index].OutPosition.SetAndAdvance(InterpPosition);
+			}
+			if (bWritesVelocity)
+			{
+				VertexOutput[Index].OutVelocity.SetAndAdvance((CurrPosition - PrevPosition)* InvDt);
+			}
+			if (bWritesTangentBasis)
+			{
+				FVector3f TangentX = FVector3f::ZeroVector;
+				FVector3f TangentY = FVector3f::ZeroVector;
+				FVector3f TangentZ = FVector3f::ZeroVector;
+				SkinningHandler.GetSkinnedTangentBasis(Accessor, Vertex, TangentX, TangentY, TangentZ);
+
+				TransformHandler.TransformUnitVector(TangentX, Transform);
+				TransformHandler.TransformUnitVector(TangentY, Transform);
+				TransformHandler.TransformUnitVector(TangentZ, Transform);
+				VertexOutput[Index].OutNormal.SetAndAdvance(TangentX);
+				VertexOutput[Index].OutBinormal.SetAndAdvance(TangentY);
+				VertexOutput[Index].OutTangent.SetAndAdvance(TangentZ);
+			}
+		}
+	}
+}
+
 #undef LOCTEXT_NAMESPACE
diff --git a/Engine/Plugins/FX/Niagara/Source/Niagara/Private/NiagaraDataInterfaceSkeletalMesh.cpp b/Engine/Plugins/FX/Niagara/Source/Niagara/Private/NiagaraDataInterfaceSkeletalMesh.cpp
--- a/Engine/Plugins/FX/Niagara/Source/Niagara/Private/NiagaraDataInterfaceSkeletalMesh.cpp
+++ b/Engine/Plugins/FX/Niagara/Source/Niagara/Private/NiagaraDataInterfaceSkeletalMesh.cpp
@@ -2829,6 +2829,10 @@ bool UNiagaraDataInterfaceSkeletalMesh::GetFunctionHLSL(const FNiagaraDataInterf
 		FSkeletalMeshInterfaceHelper::GetTriCoordVerticesName,
 		FSkeletalMeshInterfaceHelper::GetTriangleCountName,
 		FSkeletalMeshInterfaceHelper::GetFilteredTriangleCountName,
+		FSkeletalMeshInterfaceHelper::GetSkinnedTriangleVertexDataName,
+		FSkeletalMeshInterfaceHelper::GetSkinnedTriangleVertexDataWSName,
+		FSkeletalMeshInterfaceHelper::GetSkinnedTriangleVertexDataInterpName,
+		FSkeletalMeshInterfaceHelper::GetSkinnedTriangleVertexDataWSInterpName,
 		//////////////////////////////////////////////////////////////////////////////////////////////////////////////////////////////
 		// Bone Sampling
 		FSkeletalMeshInterfaceHelper::GetSkinnedBoneDataName,
