Couple of fixes for Niagara script compilation
-make sure that static variable collection is run nested for emitters
-handle shader compilation when threading isn't supported
-forward VM backend compiler warnings (typically div/0) back to the log

diff --git a/Engine/Plugins/FX/Niagara/Source/NiagaraEditor/Private/NiagaraCompilationTasks.cpp b/Engine/Plugins/FX/Niagara/Source/NiagaraEditor/Private/NiagaraCompilationTasks.cpp
--- a/Engine/Plugins/FX/Niagara/Source/NiagaraEditor/Private/NiagaraCompilationTasks.cpp
+++ b/Engine/Plugins/FX/Niagara/Source/NiagaraEditor/Private/NiagaraCompilationTasks.cpp
@@ -636,6 +636,13 @@ bool FNiagaraSystemCompilationTask::FCompileTaskInfo::RetrieveCompilationResult(
 
 	check(!ExeData.IsValid());
 
+	// in cases where asynchronous shader compiling isn't allowed we'll simply block execution here till the task
+	// is complete
+	if (!GShaderCompilingManager->AllowAsynchronousShaderCompiling())
+	{
+		bWait = true;
+	}
+
 	TOptional<FNiagaraCompileResults> CompileResult;
 	CompileResult = Compiler->GetCompileResult(CompilationJobId, bWait);
 	if (!CompileResult)
@@ -1776,7 +1783,7 @@ UE::Tasks::FTask FNiagaraSystemCompilationTask::BuildRapidIterationParametersAsy
 			FTask EmitterCollectStaticVariableTask = Launch(UE_SOURCE_LOCATION, [this, &EmitterInfo]
 			{
 				FCollectStaticVariablesTaskBuilder EmitterTaskBuilder(*this, EmitterInfo);
-				FTask EmitterStaticVariableTask = EmitterTaskBuilder.LaunchCollectedTasks();
+				AddNested(EmitterTaskBuilder.LaunchCollectedTasks());
 			}, SystemCollectStaticVariableTask);
 
 			CollectStaticVariableTasks.Add(EmitterInfo.DigestedEmitterIndex, EmitterCollectStaticVariableTask);
diff --git a/Engine/Plugins/FX/Niagara/Source/NiagaraEditor/Private/NiagaraCompiler.cpp b/Engine/Plugins/FX/Niagara/Source/NiagaraEditor/Private/NiagaraCompiler.cpp
--- a/Engine/Plugins/FX/Niagara/Source/NiagaraEditor/Private/NiagaraCompiler.cpp
+++ b/Engine/Plugins/FX/Niagara/Source/NiagaraEditor/Private/NiagaraCompiler.cpp
@@ -2082,6 +2082,11 @@ TOptional<FNiagaraCompileResults> FHlslNiagaraCompiler::GetCompileResult(int32 J
 		Ar.SetLimitSize(ShaderCode.GetActualShaderCodeSize());
 		Ar << CompilationOutput;
 
+		if (!CompilationOutput.Errors.IsEmpty())
+		{
+			Warning(FText::Format(LOCTEXT("VectorVMCompileWarningMessageFormat", "The Vector VM compile generated warnings:\n{0}"), FText::FromString(CompilationOutput.Errors)));
+		}
+
 		Results.bVMSucceeded = true;
 	}
 	else if (CompilationJob->ShaderCompileJob->Output.Errors.Num() > 0)
