Add braces to the list of characters that are treated as separation tokens for niagara's custom hlsl node processing

diff --git a/Engine/Plugins/FX/Niagara/Source/NiagaraEditor/Private/NiagaraNodeCustomHlsl.cpp b/Engine/Plugins/FX/Niagara/Source/NiagaraEditor/Private/NiagaraNodeCustomHlsl.cpp
--- a/Engine/Plugins/FX/Niagara/Source/NiagaraEditor/Private/NiagaraNodeCustomHlsl.cpp
+++ b/Engine/Plugins/FX/Niagara/Source/NiagaraEditor/Private/NiagaraNodeCustomHlsl.cpp
@@ -120,7 +120,7 @@ bool UNiagaraNodeCustomHlsl::GetTokensFromString(const FString& InHlsl, TArray<F
 		return false;
 	}
 
-	FString Separators = TEXT(";/*+-=)(?:, []<>\"\t\r\n");
+	FString Separators = TEXT(";/*+-=)(?:, []<>\"\t\r\n{}");
 	const int32 TargetLength = InHlsl.Len();
 
 	int32 TokenStart = 0;
