Changed the Sequencer Notify function call to 'Unknown'.
This is to accomodate a change in core sequencer logic that calls RefreshUI now instead of RefreshTree. RefreshUI clears the selection, which propagates to the Niagara stack selection.
Using 'Unknown' instead should work as before, calling RefreshTree instead, and seems to be more correct as our Notify function is generic and doesn't contain information about what changed.

diff --git a/Engine/Plugins/FX/Niagara/Source/NiagaraEditor/Private/ViewModels/NiagaraSystemViewModel.cpp b/Engine/Plugins/FX/Niagara/Source/NiagaraEditor/Private/ViewModels/NiagaraSystemViewModel.cpp
--- a/Engine/Plugins/FX/Niagara/Source/NiagaraEditor/Private/ViewModels/NiagaraSystemViewModel.cpp
+++ b/Engine/Plugins/FX/Niagara/Source/NiagaraEditor/Private/ViewModels/NiagaraSystemViewModel.cpp
@@ -2184,7 +2184,10 @@ void FNiagaraSystemViewModel::UpdateSequencerTracksForEmitters(const TArray<FGui
 				EmitterTrack->UpdateTrackFromEmitterGraphChange(NiagaraSequence->GetMovieScene()->GetTickResolution());
 			}
 		}
-		Sequencer->NotifyMovieSceneDataChanged(EMovieSceneDataChangeType::MovieSceneStructureItemsChanged);
+
+		// We use 'Unknown' here because this is a generic Refresh function with no context as to what changed.
+		// Also, using 'MovieSceneStructureItemsChanged' will clear the selection
+		Sequencer->NotifyMovieSceneDataChanged(EMovieSceneDataChangeType::Unknown);
 	}
 }
 
