- Optimizations to activation saves ~10% - 20% depending on system

diff --git a/Engine/Plugins/FX/Niagara/Source/Niagara/Classes/NiagaraEmitter.h b/Engine/Plugins/FX/Niagara/Source/Niagara/Classes/NiagaraEmitter.h
--- a/Engine/Plugins/FX/Niagara/Source/Niagara/Classes/NiagaraEmitter.h
+++ b/Engine/Plugins/FX/Niagara/Source/Niagara/Classes/NiagaraEmitter.h
@@ -150,8 +150,6 @@ struct FNiagaraEmitterScriptProperties
 	TArray<FNiagaraEventGeneratorProperties> EventGenerators;
 
 	NIAGARA_API void InitDataSetAccess();
-
-	NIAGARA_API bool DataSetAccessSynchronized() const;
 };
 
 USTRUCT()
@@ -374,6 +372,7 @@ struct FVersionedNiagaraEmitterData
 	FORCEINLINE const TArray<FNiagaraEventScriptProperties>& GetEventHandlers() const { return EventHandlerScriptProps; }
 	NIAGARA_API void CacheFromCompiledData(const FNiagaraDataSetCompiledData* CompiledData, const UNiagaraEmitter& Emitter);
 	NIAGARA_API void CacheFromShaderCompiled();
+	bool IsAllowedToExecute() const { return bIsAllowedToExecute; }
 
 	NIAGARA_API FGraphEventArray PrecacheComputePSOs(const UNiagaraEmitter& NiagaraEmitter);
 	bool DidPSOPrecacheFail() const { return PSOPrecacheResult == EPSOPrecacheResult::NotSupported; }
@@ -499,7 +498,10 @@ private:
 
 	UPROPERTY(Transient)
 	FNiagaraEmitterScalabilitySettings CurrentScalabilitySettings;
-	
+
+	/** Can this emitter run with the current scalability settings, etc */
+	uint32 bIsAllowedToExecute : 1;
+
 	/** Indicates that the GPU script requires the view uniform buffer. */
 	uint32 bRequiresViewUniformBuffer : 1;
 
@@ -542,6 +544,13 @@ private:
 	NIAGARA_API void OnPostCompile(const UNiagaraEmitter& InEmitter);
 
 	EPSOPrecacheResult PSOPrecacheResult = EPSOPrecacheResult::Unknown;
+
+	bool IsValidInternal() const;
+	bool IsReadyToRunInternal() const;
+#if !WITH_EDITORONLY_DATA
+	mutable TOptional<bool> IsValidCached;
+	mutable TOptional<bool> IsReadyToRunCached;
+#endif
 };
 
 /** 
diff --git a/Engine/Plugins/FX/Niagara/Source/Niagara/Private/NiagaraCommon.cpp b/Engine/Plugins/FX/Niagara/Source/Niagara/Private/NiagaraCommon.cpp
--- a/Engine/Plugins/FX/Niagara/Source/Niagara/Private/NiagaraCommon.cpp
+++ b/Engine/Plugins/FX/Niagara/Source/Niagara/Private/NiagaraCommon.cpp
@@ -1089,12 +1089,14 @@ bool FNiagaraScriptDataInterfaceInfo::IsUserDataInterface() const
 {
 	TStringBuilder<128> NameBuilder;
 	Name.ToString(NameBuilder);
-	return FCString::Strnicmp(NameBuilder.ToString(), TEXT("user."), 5) == 0;
+	return FCString::Strnicmp(NameBuilder.ToString(), PARAM_MAP_USER_STR, 5) == 0;
 }
 
 bool FNiagaraScriptResolvedDataInterfaceInfo::NeedsPerInstanceBinding() const
 {
-	return ResolvedVariable.GetName().ToString().StartsWith(TEXT("User."));
+	FNameBuilder NameBuilder;
+	ResolvedVariable.GetName().ToString(NameBuilder);
+	return NameBuilder.ToView().StartsWith(PARAM_MAP_USER_STR);
 }
 
 bool FNiagaraScriptDataInterfaceCompileInfo::CanExecuteOnTarget(ENiagaraSimTarget SimTarget) const
@@ -1132,11 +1134,18 @@ UNiagaraDataInterface* FNiagaraScriptDataInterfaceCompileInfo::GetDefaultDataInt
 
 bool FNiagaraScriptDataInterfaceCompileInfo::NeedsPerInstanceBinding()const
 {
-	if (Name.ToString().StartsWith(TEXT("User.")))
+	FNameBuilder NameBuilder;
+	Name.ToString(NameBuilder);
+	if (NameBuilder.ToView().StartsWith(TEXT("User.")))
+	{
 		return true;
+	}
+
 	UNiagaraDataInterface* Obj = GetDefaultDataInterface();
 	if (Obj && Obj->PerInstanceDataSize() > 0)
+	{
 		return true;
+	}
 	return false;
 }
 
diff --git a/Engine/Plugins/FX/Niagara/Source/Niagara/Private/NiagaraComponentSettings.cpp b/Engine/Plugins/FX/Niagara/Source/Niagara/Private/NiagaraComponentSettings.cpp
--- a/Engine/Plugins/FX/Niagara/Source/Niagara/Private/NiagaraComponentSettings.cpp
+++ b/Engine/Plugins/FX/Niagara/Source/Niagara/Private/NiagaraComponentSettings.cpp
@@ -260,28 +260,22 @@ namespace FNiagaraComponentSettings
 		return bUseSystemDenyList ? !SystemDenyList.Contains(System->GetFName()) : true;
 	}
 
-	bool IsEmitterAllowedToRun(const FNiagaraEmitterInstance* EmitterInstance)
+	bool IsEmitterAllowedToRun(const FVersionedNiagaraEmitterData& EmitterData, const UNiagaraEmitter& NiagaraEmitter)
 	{
 		if (bUseEmitterDenyList)
 		{
-			const UNiagaraSystem* NiagaraSystem = EmitterInstance->GetParentSystemInstance()->GetSystem();
-			const FVersionedNiagaraEmitter CachedEmitter = EmitterInstance->GetCachedEmitter();
-			FVersionedNiagaraEmitterData* EmitterData = CachedEmitter.GetEmitterData();
-			if (EmitterData == nullptr)
-			{
-				return false;
-			}
+			const UNiagaraSystem* NiagaraSystem = NiagaraEmitter.GetTypedOuter<UNiagaraSystem>();
 
 			FNiagaraEmitterNameSettingsRef EmitterRef;
 			EmitterRef.SystemName = NiagaraSystem ? NiagaraSystem->GetFName() : NAME_None;
-			EmitterRef.EmitterName = FName(*CachedEmitter.Emitter->GetUniqueEmitterName());
+			EmitterRef.EmitterName = FName(*NiagaraEmitter.GetUniqueEmitterName());
 
 			if (EmitterDenyList.Contains(EmitterRef))
 			{
 				return false;
 			}
 
-			if (EmitterData->SimTarget == ENiagaraSimTarget::GPUComputeSim)
+			if (EmitterData.SimTarget == ENiagaraSimTarget::GPUComputeSim)
 			{
 				if (GpuEmitterDenyList.Contains(EmitterRef))
 				{
@@ -290,7 +284,7 @@ namespace FNiagaraComponentSettings
 
 				if (GpuDataInterfaceDenyList.Num() > 0)
 				{
-					if (const UNiagaraScript* GPUComputeScript = EmitterData->GetGPUComputeScript())
+					if (const UNiagaraScript* GPUComputeScript = EmitterData.GetGPUComputeScript())
 					{
 						for (const FNiagaraScriptResolvedDataInterfaceInfo& DefaultDIInfo : GPUComputeScript->GetResolvedDataInterfaces())
 						{
@@ -305,9 +299,7 @@ namespace FNiagaraComponentSettings
 		}
 		if (bAllowGpuEmitters == false)
 		{
-			const FVersionedNiagaraEmitter CachedEmitter = EmitterInstance->GetCachedEmitter();
-			FVersionedNiagaraEmitterData* EmitterData = CachedEmitter.GetEmitterData();
-			return EmitterData && (EmitterData->SimTarget != ENiagaraSimTarget::GPUComputeSim);
+			return EmitterData.SimTarget != ENiagaraSimTarget::GPUComputeSim;
 		}
 		return true;
 	}
diff --git a/Engine/Plugins/FX/Niagara/Source/Niagara/Private/NiagaraEmitter.cpp b/Engine/Plugins/FX/Niagara/Source/Niagara/Private/NiagaraEmitter.cpp
--- a/Engine/Plugins/FX/Niagara/Source/Niagara/Private/NiagaraEmitter.cpp
+++ b/Engine/Plugins/FX/Niagara/Source/Niagara/Private/NiagaraEmitter.cpp
@@ -6,8 +6,8 @@
 #include "INiagaraEditorOnlyDataUtlities.h"
 #include "NiagaraBoundsCalculator.h"
 #include "NiagaraCustomVersion.h"
+#include "NiagaraComponentSettings.h"
 #include "NiagaraMessageDataBase.h"
-#include "UObject/UE5MainStreamObjectVersion.h"
 #include "NiagaraEditorDataBase.h"
 #include "NiagaraModule.h"
 #include "NiagaraRendererProperties.h"
@@ -20,10 +20,13 @@
 #include "NiagaraStats.h"
 #include "NiagaraSystem.h"
 #include "NiagaraTrace.h"
+
+#include "Engine/Engine.h"
 #include "HAL/LowLevelMemTracker.h"
 #include "Interfaces/ITargetPlatform.h"
 #include "Modules/ModuleManager.h"
 #include "Templates/Greater.h"
+#include "UObject/UE5MainStreamObjectVersion.h"
 #include "UObject/LinkerLoad.h"
 #include "UObject/ObjectSaveContext.h"
 #include "UObject/Package.h"
@@ -143,26 +146,6 @@ void FNiagaraEmitterScriptProperties::InitDataSetAccess()
 	}
 }
 
-bool FNiagaraEmitterScriptProperties::DataSetAccessSynchronized() const
-{
-	if (Script && Script->IsReadyToRun(ENiagaraSimTarget::CPUSim))
-	{
-		if (Script->GetVMExecutableData().ReadDataSets.Num() != EventReceivers.Num())
-		{
-			return false;
-		}
-		if (Script->GetVMExecutableData().WriteDataSets.Num() != EventGenerators.Num())
-		{
-			return false;
-		}
-		return true;
-	}
-	else
-	{
-		return EventReceivers.Num() == 0 && EventGenerators.Num() == 0;
-	}
-}
-
 //////////////////////////////////////////////////////////////////////////
 
 UNiagaraEmitter::UNiagaraEmitter(const FObjectInitializer& Initializer)
@@ -1459,7 +1442,7 @@ bool UNiagaraEmitter::IsEnabledOnPlatform(const FString& PlatformName)const
 	return false;
 }
 
-bool FVersionedNiagaraEmitterData::IsValid() const
+bool FVersionedNiagaraEmitterData::IsValidInternal() const
 {
 	if (!SpawnScriptProps.Script || !UpdateScriptProps.Script)
 	{
@@ -1500,7 +1483,7 @@ bool FVersionedNiagaraEmitterData::IsValid() const
 	return true;
 }
 
-bool FVersionedNiagaraEmitterData::IsReadyToRun() const
+bool FVersionedNiagaraEmitterData::IsReadyToRunInternal() const
 {
 	//Check for various failure conditions and bail.
 	if (!UpdateScriptProps.Script || !SpawnScriptProps.Script)
@@ -1541,6 +1524,32 @@ bool FVersionedNiagaraEmitterData::IsReadyToRun() const
 	return true;
 }
 
+bool FVersionedNiagaraEmitterData::IsValid() const
+{
+#if WITH_EDITORONLY_DATA
+	return IsValidInternal();
+#else
+	if (!IsValidCached.IsSet())
+	{
+		IsValidCached = IsValidInternal();
+	}
+	return IsValidCached.GetValue();
+#endif
+}
+
+bool FVersionedNiagaraEmitterData::IsReadyToRun() const
+{
+#if WITH_EDITORONLY_DATA
+	return IsReadyToRunInternal();
+#else
+	if (!IsReadyToRunCached.IsSet())
+	{
+		IsReadyToRunCached = IsReadyToRunInternal();
+	}
+	return IsReadyToRunCached.GetValue();
+#endif
+}
+
 void FVersionedNiagaraEmitterData::GetScripts(TArray<UNiagaraScript*>& OutScripts, bool bCompilableOnly, bool bEnabledOnly) const
 {
 	OutScripts.Add(SpawnScriptProps.Script);
@@ -1597,6 +1606,7 @@ UNiagaraScript* FVersionedNiagaraEmitterData::GetScript(ENiagaraScriptUsage Usag
 
 void FVersionedNiagaraEmitterData::CacheFromCompiledData(const FNiagaraDataSetCompiledData* CompiledData, const UNiagaraEmitter& Emitter)
 {
+	bIsAllowedToExecute = true;
 	bRequiresViewUniformBuffer = false;
 	bNeedsPartialDepthTexture = false;
 
@@ -1703,6 +1713,38 @@ void FVersionedNiagaraEmitterData::CacheFromCompiledData(const FNiagaraDataSetCo
 		MaxInstanceCount = 0;
 	}
 	UpdateDebugName(Emitter, CompiledData);
+
+	// Detemine if we are allowed to execute or not
+	bIsAllowedToExecute = IsAllowedByScalability();
+	if (bIsAllowedToExecute && (SimTarget == ENiagaraSimTarget::GPUComputeSim))
+	{
+		if (const FNiagaraShaderScript* ShaderScript = GPUComputeScript ? GPUComputeScript->GetRenderThreadScript() : nullptr)
+		{
+			TSharedRef<FNiagaraShaderScriptParametersMetadata> ShaderParametersMetadata = ShaderScript->GetScriptParametersMetadata();
+			if (ShaderParametersMetadata->ShaderParametersMetadata != nullptr)
+			{
+				const uint64 ScriptCBufferSize = ShaderParametersMetadata->ShaderParametersMetadata->GetLayout().ConstantBufferSize;
+				const uint64 RHIMaxCBufferSize = GetMaxConstantBufferByteSize();
+				if (ScriptCBufferSize > RHIMaxCBufferSize)
+				{
+#if !(UE_BUILD_SHIPPING || UE_BUILD_TEST)
+					GEngine->AddOnScreenDebugMessage(uint64(this), 1.f, FColor::Red, *FString::Printf(TEXT("GPU Simulation(%s) is disabled due to using too much constant buffer space (%d/%d)."), GetDebugSimName(), ScriptCBufferSize, RHIMaxCBufferSize));
+#endif
+					bIsAllowedToExecute = false;
+				}
+			}
+			else
+			{
+				bIsAllowedToExecute = false;
+			}
+		}
+		else
+		{
+			bIsAllowedToExecute = false;
+		}
+	}
+
+	bIsAllowedToExecute &= FNiagaraComponentSettings::IsEmitterAllowedToRun(*this, Emitter);
 }
 
 void FVersionedNiagaraEmitterData::UpdateDebugName(const UNiagaraEmitter& Emitter, const FNiagaraDataSetCompiledData* CompiledData)
diff --git a/Engine/Plugins/FX/Niagara/Source/Niagara/Private/NiagaraEmitterInstance.cpp b/Engine/Plugins/FX/Niagara/Source/Niagara/Private/NiagaraEmitterInstance.cpp
--- a/Engine/Plugins/FX/Niagara/Source/Niagara/Private/NiagaraEmitterInstance.cpp
+++ b/Engine/Plugins/FX/Niagara/Source/Niagara/Private/NiagaraEmitterInstance.cpp
@@ -214,43 +214,7 @@ void FNiagaraEmitterInstance::Dump()const
 bool FNiagaraEmitterInstance::IsAllowedToExecute() const
 {
 	FVersionedNiagaraEmitterData* EmitterData = CachedEmitter.GetEmitterData();
-	if (!GetEmitterHandle().GetIsEnabled() || EmitterData == nullptr || !EmitterData->IsAllowedByScalability())
-	{
-		return false;
-	}
-
-	if (EmitterData->SimTarget == ENiagaraSimTarget::GPUComputeSim)
-	{
-		//-TODO: Could replace with CPU side sim in some cases
-		if ( !ComputeDispatchInterface || !FNiagaraUtilities::AllowGPUParticles(ComputeDispatchInterface->GetShaderPlatform()) )
-		{
-			return false;
-		}
-
-		if (EmitterData->DidPSOPrecacheFail() == true)
-		{
-			return false;
-		}
-
-		if (const UNiagaraScript* GPUComputeScript = EmitterData->GetGPUComputeScript())
-		{
-			if (const FNiagaraShaderScript* ShaderScript = GPUComputeScript->GetRenderThreadScript())
-			{
-				const uint64 ScriptCBufferSize = ShaderScript->GetScriptParametersMetadata()->ShaderParametersMetadata->GetLayout().ConstantBufferSize;
-				const uint64 RHIMaxCBufferSize = GetMaxConstantBufferByteSize();
-				if (ScriptCBufferSize > RHIMaxCBufferSize)
-				{
-				#if !(UE_BUILD_SHIPPING || UE_BUILD_TEST)
-					GEngine->AddOnScreenDebugMessage(uint64(EmitterData), 1.f, FColor::Red, *FString::Printf(TEXT("GPU Simulation(%s) is disabled due to using too much constant buffer space (%d/%d)."), EmitterData->GetDebugSimName(), ScriptCBufferSize, RHIMaxCBufferSize));
-				#endif
-					return false;
-				}
-			}
-		}
-	}
-
-	// Do we allow the emitter to run or not?
-	return FNiagaraComponentSettings::IsEmitterAllowedToRun(this);
+	return GetEmitterHandle().GetIsEnabled() && EmitterData && EmitterData->IsAllowedToExecute();
 }
 
 void FNiagaraEmitterInstance::Init(int32 InEmitterIdx, FNiagaraSystemInstanceID InSystemInstanceID)
@@ -342,12 +306,8 @@ void FNiagaraEmitterInstance::Init(int32 InEmitterIdx, FNiagaraSystemInstanceID
 	}
 
 	{
-		ensure(EmitterData->UpdateScriptProps.DataSetAccessSynchronized());
 		const int32 UpdateEventGeneratorCount = EmitterData->UpdateScriptProps.EventGenerators.Num();
-
-		ensure(EmitterData->SpawnScriptProps.DataSetAccessSynchronized());
 		const int32 SpawnEventGeneratorCount = EmitterData->SpawnScriptProps.EventGenerators.Num();
-
 		const int32 NumEvents = EmitterData->GetEventHandlers().Num();
 
 		if (UpdateEventGeneratorCount || SpawnEventGeneratorCount || NumEvents)
@@ -384,8 +344,6 @@ void FNiagaraEmitterInstance::Init(int32 InEmitterIdx, FNiagaraSystemInstanceID
 
 			for (int32 i = 0; i < NumEvents; i++)
 			{
-				ensure(EmitterData->GetEventHandlers()[i].DataSetAccessSynchronized());
-
 				UNiagaraScript* EventScript = EmitterData->GetEventHandlers()[i].Script;
 
 				//This is cpu explicitly? Are we doing event handlers on GPU?
diff --git a/Engine/Plugins/FX/Niagara/Source/Niagara/Public/NiagaraComponentSettings.h b/Engine/Plugins/FX/Niagara/Source/Niagara/Public/NiagaraComponentSettings.h
--- a/Engine/Plugins/FX/Niagara/Source/Niagara/Public/NiagaraComponentSettings.h
+++ b/Engine/Plugins/FX/Niagara/Source/Niagara/Public/NiagaraComponentSettings.h
@@ -4,12 +4,13 @@
 
 #include "NiagaraCommon.h"
 
+class UNiagaraEmitter;
 class UNiagaraSystem;
-class FNiagaraEmitterInstance;
+struct FVersionedNiagaraEmitterData;
 
 namespace FNiagaraComponentSettings
 {
 	extern void UpdateSettings();
 	extern NIAGARA_API bool IsSystemAllowedToRun(const UNiagaraSystem* System);
-	extern NIAGARA_API bool IsEmitterAllowedToRun(const FNiagaraEmitterInstance* EmitterInstance);
+	extern NIAGARA_API bool IsEmitterAllowedToRun(const FVersionedNiagaraEmitterData& EmitterData, const UNiagaraEmitter& NiagaraEmitter);
 };
