* Converted screenshot tests to cached tests
* Improved comparison method for system data
* Added support for gpu readback tests

diff --git a/Engine/Plugins/FX/Niagara/Source/Niagara/Private/NiagaraSimCache.cpp b/Engine/Plugins/FX/Niagara/Source/Niagara/Private/NiagaraSimCache.cpp
--- a/Engine/Plugins/FX/Niagara/Source/Niagara/Private/NiagaraSimCache.cpp
+++ b/Engine/Plugins/FX/Niagara/Source/Niagara/Private/NiagaraSimCache.cpp
@@ -1204,7 +1204,7 @@ bool UNiagaraSimCache::IsDataEqual(const UNiagaraSimCache& OtherCache, float Err
 		const FNiagaraSimCacheFrame& ExpectedFrame = CacheFrames[FrameIndex];
 		const FNiagaraSimCacheFrame& OtherFrame = OtherFrames[FrameIndex];
 
-		if (OtherFrame.SimulationAge != ExpectedFrame.SimulationAge)
+		if (!FMath::IsNearlyEqual(OtherFrame.SimulationAge, ExpectedFrame.SimulationAge, UE_KINDA_SMALL_NUMBER))
 		{
 			CacheCompare::AddError(Errors, TEXT("Simulation age different - frame ") + Frame + '\n');
 			bEqual = false;
@@ -1215,12 +1215,63 @@ bool UNiagaraSimCache::IsDataEqual(const UNiagaraSimCache& OtherCache, float Err
 			bEqual = false;
 			continue;
 		}
-		
-		if (OtherFrame.SystemData.SystemDataBuffers != ExpectedFrame.SystemData.SystemDataBuffers)
+
+		// compare system data buffers
+		if (OtherFrame.SystemData.SystemDataBuffers.NumInstances != ExpectedFrame.SystemData.SystemDataBuffers.NumInstances)
+		{
+			CacheCompare::AddError(Errors, TEXT("SystemDataBuffers.NumInstances different - frame ") + Frame + TEXT(": expected ") + FString::FromInt(ExpectedFrame.SystemData.SystemDataBuffers.NumInstances) + TEXT(", got ") + FString::FromInt(OtherFrame.SystemData.SystemDataBuffers.NumInstances) + '\n');
+			bEqual = false;
+		}
+		if (OtherFrame.SystemData.SystemDataBuffers.Int32Data != ExpectedFrame.SystemData.SystemDataBuffers.Int32Data)
+		{
+			CacheCompare::AddError(Errors, TEXT("SystemDataBuffers.Int32Data different - frame ") + Frame + '\n');
+			bEqual = false;
+		}
+		if (OtherFrame.SystemData.SystemDataBuffers.IDToIndexTable != ExpectedFrame.SystemData.SystemDataBuffers.IDToIndexTable)
+		{
+			CacheCompare::AddError(Errors, TEXT("SystemDataBuffers.IDToIndexTable different - frame ") + Frame + '\n');
+			bEqual = false;
+		}
+		if (OtherFrame.SystemData.SystemDataBuffers.FloatData.Num() == ExpectedFrame.SystemData.SystemDataBuffers.FloatData.Num())
+		{
+			for (int i = 0; i < OtherFrame.SystemData.SystemDataBuffers.FloatData.Num(); i += sizeof(float))
+			{
+				float ExpectedVal = *reinterpret_cast<const float*>(ExpectedFrame.SystemData.SystemDataBuffers.FloatData.GetData() + i);
+				float OtherVal    = *reinterpret_cast<const float*>(OtherFrame.SystemData.SystemDataBuffers.FloatData.GetData() + i);
+				if (!FMath::IsNearlyEqual(OtherVal, ExpectedVal, ErrorTolerance))
+				{
+					CacheCompare::AddError(Errors, TEXT("SystemDataBuffers.FloatData different - frame ") + Frame + '\n');
+					bEqual = false;
+					break;
+				}
+			}
+		}
+		else
 		{
-			CacheCompare::AddError(Errors, TEXT("Cached system data different - frame ") + Frame + '\n');
+			CacheCompare::AddError(Errors, TEXT("SystemDataBuffers.FloatData count different - frame ") + Frame + '\n');
 			bEqual = false;
 		}
+		if (OtherFrame.SystemData.SystemDataBuffers.HalfData.Num() == ExpectedFrame.SystemData.SystemDataBuffers.HalfData.Num())
+		{
+			for (int i = 0; i < OtherFrame.SystemData.SystemDataBuffers.HalfData.Num(); i += sizeof(FFloat16))
+			{
+				FFloat16 ExpectedVal = *reinterpret_cast<const FFloat16*>(ExpectedFrame.SystemData.SystemDataBuffers.HalfData.GetData() + i);
+				FFloat16 OtherVal    = *reinterpret_cast<const FFloat16*>(OtherFrame.SystemData.SystemDataBuffers.HalfData.GetData() + i);
+				if (!FMath::IsNearlyEqual(OtherVal, ExpectedVal, ErrorTolerance))
+				{
+					CacheCompare::AddError(Errors, TEXT("SystemDataBuffers.HalfData different - frame ") + Frame + '\n');
+					bEqual = false;
+					break;
+				}
+			}
+		}
+		else
+		{
+			CacheCompare::AddError(Errors, TEXT("SystemDataBuffers.HalfData count different - frame ") + Frame + '\n');
+			bEqual = false;
+		}
+
+		// compare emitter data
 		for (int EmitterIndex = 0; EmitterIndex < ExpectedFrame.EmitterData.Num(); EmitterIndex++)
 		{
 			FName EmitterName = CacheLayout.EmitterLayouts[EmitterIndex].LayoutName;
