Added a little performance ticker to the Niagara system view (enabled via fx.ShowNiagaraDeveloperWindows 1)

diff --git a/Engine/Plugins/FX/Niagara/Source/NiagaraEditor/Private/Toolkits/SystemToolkitModes/NiagaraSystemToolkitMode_Default.cpp b/Engine/Plugins/FX/Niagara/Source/NiagaraEditor/Private/Toolkits/SystemToolkitModes/NiagaraSystemToolkitMode_Default.cpp
--- a/Engine/Plugins/FX/Niagara/Source/NiagaraEditor/Private/Toolkits/SystemToolkitModes/NiagaraSystemToolkitMode_Default.cpp
+++ b/Engine/Plugins/FX/Niagara/Source/NiagaraEditor/Private/Toolkits/SystemToolkitModes/NiagaraSystemToolkitMode_Default.cpp
@@ -285,6 +285,37 @@ void FNiagaraSystemToolkitMode_Default::ExtendToolbar()
 				}
 				ToolbarBuilder.EndSection();
 			}
+
+			if (GbShowNiagaraDeveloperWindows != 0)
+			{
+				// a simple ticker to visually see how fast the ui is ticking and to see any hitches due to stack actions 
+				ToolbarBuilder.AddSeparator();
+				ToolbarBuilder.AddWidget(SNew(STextBlock)
+					.Margin(FMargin(5, 10, 0, 0))
+					.ColorAndOpacity_Lambda([]()
+					{
+						static int32 Counter = 0;
+						Counter = (Counter + 1) % 255;
+						return FLinearColor::MakeFromHSV8(Counter, 192, 255);
+					})
+					.Text_Lambda([]()
+					{
+						static int32 Counter = 0;
+						Counter = (Counter + 1) % 16;
+						switch (Counter / 2)
+						{
+						default:
+						case 0: return FText::FromString(".oO     ");
+						case 1: return FText::FromString(" .oO    ");
+						case 2: return FText::FromString("  .oO   ");
+						case 3: return FText::FromString("   .oO  ");
+						case 4: return FText::FromString("    .oO ");
+						case 5: return FText::FromString("     .oO");
+						case 6: return FText::FromString("O     .o");
+						case 7: return FText::FromString("oO     .");
+						}
+					}));
+			}
 		}
 	};
 
