Added Debug Draw Message for scripts to display on the Debug Draw actions

diff --git a/Engine/Plugins/FX/Niagara/Source/Niagara/Classes/NiagaraScript.h b/Engine/Plugins/FX/Niagara/Source/Niagara/Classes/NiagaraScript.h
--- a/Engine/Plugins/FX/Niagara/Source/Niagara/Classes/NiagaraScript.h
+++ b/Engine/Plugins/FX/Niagara/Source/Niagara/Classes/NiagaraScript.h
@@ -739,6 +739,10 @@ public:
 	UPROPERTY(EditAnywhere, Category = Script, meta = (MultiLine = true))
 	FText NoteMessage;
 
+	/** A message to display on UI actions handling debug draw state. */
+	UPROPERTY(EditAnywhere, Category = Script, meta = (MultiLine = true))
+	FText DebugDrawMessage;
+
 	/* Defines if this script is visible to the user when searching for modules to add to an emitter.  */
 	UPROPERTY(AssetRegistrySearchable, EditAnywhere, Category = Script)
 	ENiagaraScriptLibraryVisibility LibraryVisibility;
diff --git a/Engine/Plugins/FX/Niagara/Source/NiagaraEditor/Private/NiagaraEditorUtilities.cpp b/Engine/Plugins/FX/Niagara/Source/NiagaraEditor/Private/NiagaraEditorUtilities.cpp
--- a/Engine/Plugins/FX/Niagara/Source/NiagaraEditor/Private/NiagaraEditorUtilities.cpp
+++ b/Engine/Plugins/FX/Niagara/Source/NiagaraEditor/Private/NiagaraEditorUtilities.cpp
@@ -88,6 +88,8 @@ static FAutoConsoleVariableRef CVarDeletePythonFilesOnError(
 	ECVF_Default
 );
 
+const FText FNiagaraEditorSharedTexts::DebugDrawUIActionBaseText = LOCTEXT("EnableDebugDrawCheckBoxToolTip", "Enable or disable debug drawing for this item.");
+
 TSet<FName> FNiagaraEditorUtilities::GetSystemConstantNames()
 {
 	TSet<FName> SystemConstantNames;
diff --git a/Engine/Plugins/FX/Niagara/Source/NiagaraEditor/Public/NiagaraEditorUtilities.h b/Engine/Plugins/FX/Niagara/Source/NiagaraEditor/Public/NiagaraEditorUtilities.h
--- a/Engine/Plugins/FX/Niagara/Source/NiagaraEditor/Public/NiagaraEditorUtilities.h
+++ b/Engine/Plugins/FX/Niagara/Source/NiagaraEditor/Public/NiagaraEditorUtilities.h
@@ -60,6 +60,10 @@ struct FRefreshAllScriptsFromExternalChangesArgs
 	UNiagaraParameterDefinitions* OriginatingParameterDefinitions = nullptr;
 };
 
+struct NIAGARAEDITOR_API FNiagaraEditorSharedTexts
+{
+	static const FText DebugDrawUIActionBaseText;
+};
 namespace FNiagaraEditorUtilities
 {
 	/** Determines if the contents of two sets matches */
diff --git a/Engine/Plugins/FX/Niagara/Source/NiagaraEditorWidgets/Private/NiagaraEditorWidgetsUtilities.cpp b/Engine/Plugins/FX/Niagara/Source/NiagaraEditorWidgets/Private/NiagaraEditorWidgetsUtilities.cpp
--- a/Engine/Plugins/FX/Niagara/Source/NiagaraEditorWidgets/Private/NiagaraEditorWidgetsUtilities.cpp
+++ b/Engine/Plugins/FX/Niagara/Source/NiagaraEditorWidgets/Private/NiagaraEditorWidgetsUtilities.cpp
@@ -382,12 +382,22 @@ bool FNiagaraStackEditorWidgetsUtilities::AddStackItemContextMenuActions(FMenuBu
 				
 				if (ModuleItem->GetModuleNode().ContainsDebugSwitch())
 				{
+					FText TooltipText = FNiagaraEditorSharedTexts::DebugDrawUIActionBaseText;
+
+					if(FVersionedNiagaraScriptData* ScriptData = ModuleItem->GetModuleNode().GetScriptData())
+					{
+						if(ScriptData->DebugDrawMessage.IsEmpty() == false)
+						{
+							TooltipText = ScriptData->DebugDrawMessage;
+						}
+					}
+					
 					FUIAction Action(FExecuteAction::CreateStatic(&ToggleShouldDebugDraw, TWeakObjectPtr<UNiagaraStackItem>(&StackItem)),
 						FCanExecuteAction(),
 						FIsActionChecked::CreateUObject(ModuleItem, &UNiagaraStackModuleItem::IsDebugDrawEnabled));
 					MenuBuilder.AddMenuEntry(
 						LOCTEXT("ShouldDebugDraw", "Enable Debug Draw"),
-						LOCTEXT("ToggleShouldDebugDrawToolTip", "Toggle debug draw enable/disabled"),
+						TooltipText,
 						FSlateIcon(),
 						Action,
 						NAME_None,
diff --git a/Engine/Plugins/FX/Niagara/Source/NiagaraEditorWidgets/Private/SNiagaraOverviewStack.cpp b/Engine/Plugins/FX/Niagara/Source/NiagaraEditorWidgets/Private/SNiagaraOverviewStack.cpp
--- a/Engine/Plugins/FX/Niagara/Source/NiagaraEditorWidgets/Private/SNiagaraOverviewStack.cpp
+++ b/Engine/Plugins/FX/Niagara/Source/NiagaraEditorWidgets/Private/SNiagaraOverviewStack.cpp
@@ -1029,6 +1029,16 @@ TSharedRef<ITableRow> SNiagaraOverviewStack::OnGenerateRowForEntry(UNiagaraStack
 			// Debug draw 
 			if(StackModuleItem->GetModuleNode().ContainsDebugSwitch())
 			{
+				FText TooltipText = FNiagaraEditorSharedTexts::DebugDrawUIActionBaseText;
+				
+				if(FVersionedNiagaraScriptData* ScriptData = StackModuleItem->GetModuleNode().GetScriptData())
+				{
+					if(ScriptData->DebugDrawMessage.IsEmpty() == false)
+					{
+						TooltipText = ScriptData->DebugDrawMessage;
+					}
+				}
+				
 				OptionsBox->AddSlot()
 					.VAlign(VAlign_Center)
 					.AutoWidth()
@@ -1038,7 +1048,7 @@ TSharedRef<ITableRow> SNiagaraOverviewStack::OnGenerateRowForEntry(UNiagaraStack
 						.HAlign(HAlign_Center)
 						.VAlign(VAlign_Center)
 						.ForegroundColor(FLinearColor::Transparent)
-						.ToolTipText(LOCTEXT("EnableDebugDrawCheckBoxToolTip", "Enable or disable debug drawing for this item."))
+						.ToolTipText(TooltipText)
 						.OnClicked(this, &SNiagaraOverviewStack::ToggleModuleDebugDraw, StackItem)
 						.ContentPadding(FMargin(1.0f))
 						.ButtonStyle(FNiagaraEditorWidgetsStyle::Get(), "NiagaraEditor.Stack.SimpleButton")
