Improved data channel ux:
* improved display of variables in asset
* enforce unique variable names in asset
* improved property naming and tooltips in various places
* fixed asset color
* added hyperlink to docs

diff --git a/Engine/Plugins/FX/Niagara/Source/Niagara/Internal/DataInterface/NiagaraDataInterfaceDataChannelRead.h b/Engine/Plugins/FX/Niagara/Source/Niagara/Internal/DataInterface/NiagaraDataInterfaceDataChannelRead.h
--- a/Engine/Plugins/FX/Niagara/Source/Niagara/Internal/DataInterface/NiagaraDataInterfaceDataChannelRead.h
+++ b/Engine/Plugins/FX/Niagara/Source/Niagara/Internal/DataInterface/NiagaraDataInterfaceDataChannelRead.h
@@ -68,7 +68,7 @@ public:
 // 	UPROPERTY(EditAnywhere, Category="Data Channel", meta=(EditCondition = "Scope == ENiagaraDataChannelScope::Local"))
 // 	FName Source;
 
-	/** When reading from external, the channel to consume. */
+	/** The data channel to access and read from. */
 	UPROPERTY(EditAnywhere, Category="Data Channel")
 	TObjectPtr<UNiagaraDataChannelAsset> Channel;
 	
diff --git a/Engine/Plugins/FX/Niagara/Source/Niagara/Internal/DataInterface/NiagaraDataInterfaceDataChannelWrite.h b/Engine/Plugins/FX/Niagara/Source/Niagara/Internal/DataInterface/NiagaraDataInterfaceDataChannelWrite.h
--- a/Engine/Plugins/FX/Niagara/Source/Niagara/Internal/DataInterface/NiagaraDataInterfaceDataChannelWrite.h
+++ b/Engine/Plugins/FX/Niagara/Source/Niagara/Internal/DataInterface/NiagaraDataInterfaceDataChannelWrite.h
@@ -50,15 +50,15 @@ public:
 	uint32 AllocationCount = 0;
 
 	/** Whether the data generated by this DI should be published to the world game data channel. This is require to allow game BP and C++ to read this data. */
-	UPROPERTY(EditAnywhere, Category = "Data Channel")
+	UPROPERTY(EditAnywhere, Category = "Data Channel", meta = (DisplayName = "Visible to Blueprint"))
 	bool bPublishToGame = false;
 
-	/** Whether the data generated by this DI should be published to the world CPUSim data channel. This is required for CPU emitters in other Niagara Systems to read this data. */
-	UPROPERTY(EditAnywhere, Category = "Data Channel")
+	/** Whether the data generated by this DI should be published to CPU emitters in other Niagara systems. */
+	UPROPERTY(EditAnywhere, Category = "Data Channel", meta = (DisplayName = "Visible to CPU systems"))
 	bool bPublishToCPU = false;
 
-	/** Whether the data generated by this DI should be published to a world data channel. This is required to allow GPU emitters in other Niagara Systems to read this data.  */
-	UPROPERTY(EditAnywhere, Category = "Data Channel")
+	/** Whether the data generated by this DI should be published to GPU emitters in other Niagara systems.  */
+	UPROPERTY(EditAnywhere, Category = "Data Channel", meta = (DisplayName = "Visible to GPU systems"))
 	bool bPublishToGPU = false;
 	
 	/**
@@ -70,8 +70,8 @@ public:
 	UPROPERTY(EditAnywhere, Category = "Data Channel", AdvancedDisplay)
 	bool bUpdateDestinationDataEveryTick = true;
 
-	/** When writing externally, the channel to use. */
-	UPROPERTY(EditAnywhere, Category = "Data Channel", meta = (EditCondition = "bPublishToWorld"))
+	/** The data channel to write to. */
+	UPROPERTY(EditAnywhere, Category = "Data Channel")
 	TObjectPtr<UNiagaraDataChannelAsset> Channel;
 
 	//UObject Interface
diff --git a/Engine/Plugins/FX/Niagara/Source/Niagara/Private/NiagaraDataChannel.cpp b/Engine/Plugins/FX/Niagara/Source/Niagara/Private/NiagaraDataChannel.cpp
--- a/Engine/Plugins/FX/Niagara/Source/Niagara/Private/NiagaraDataChannel.cpp
+++ b/Engine/Plugins/FX/Niagara/Source/Niagara/Private/NiagaraDataChannel.cpp
@@ -891,6 +891,33 @@ void UNiagaraDataChannel::PreEditChange(FProperty* PropertyAboutToChange)
 
 void UNiagaraDataChannel::PostEditChangeProperty(FPropertyChangedEvent& PropertyChangedEvent)
 {
+	FName VariablesMemberName = GET_MEMBER_NAME_CHECKED(UNiagaraDataChannel, Variables);
+	if (PropertyChangedEvent.ChangeType == EPropertyChangeType::ArrayAdd && PropertyChangedEvent.GetPropertyName() == VariablesMemberName)
+	{
+		TSet<FName> ExistingNames;
+		for (const FNiagaraVariable& Var : Variables)
+		{
+			ExistingNames.Add(Var.GetName());
+		}
+		FName UniqueName = FNiagaraUtilities::GetUniqueName(FName("MyNewVar"), ExistingNames);
+		Variables.Last().SetName(UniqueName);
+	}
+	if (PropertyChangedEvent.ChangeType == EPropertyChangeType::Duplicate && PropertyChangedEvent.GetPropertyName() == VariablesMemberName)
+	{
+		int32 ArrayIndex = PropertyChangedEvent.GetArrayIndex(VariablesMemberName.ToString());
+		if (Variables.IsValidIndex(ArrayIndex + 1))
+		{
+			TSet<FName> ExistingNames;
+			for (const FNiagaraVariable& Var : Variables)
+			{
+				ExistingNames.Add(Var.GetName());
+			}
+			FNiagaraVariable& NewEntry = Variables[ArrayIndex + 1];
+			FName UniqueName = FNiagaraUtilities::GetUniqueName(NewEntry.GetName(), ExistingNames);
+			NewEntry.SetName(UniqueName);
+		}
+	}
+	
 	Super::PostEditChangeProperty(PropertyChangedEvent);
 	
 	check(IsInGameThread());
diff --git a/Engine/Plugins/FX/Niagara/Source/Niagara/Public/NiagaraDataChannel.h b/Engine/Plugins/FX/Niagara/Source/Niagara/Public/NiagaraDataChannel.h
--- a/Engine/Plugins/FX/Niagara/Source/Niagara/Public/NiagaraDataChannel.h
+++ b/Engine/Plugins/FX/Niagara/Source/Niagara/Public/NiagaraDataChannel.h
@@ -131,7 +131,7 @@ public:
 private:
 
 	/** The variables that define the data contained in this Data Channel. */
-	UPROPERTY(EditAnywhere, Category = "Data Channel")
+	UPROPERTY(EditAnywhere, Category = "Data Channel", meta=(EnforceUniqueNames = true))
 	TArray<FNiagaraVariable> Variables;
 
 	/** If true, we keep our previous frame's data. This comes at a memory and performance cost but allows users to avoid tick order dependency by reading last frame's data. Some users will prefer a frame of latency to tick order dependency. */
@@ -186,11 +186,30 @@ class NIAGARA_API UNiagaraDataChannelLibrary : public UBlueprintFunctionLibrary
 	UFUNCTION(BlueprintCallable, Category = NiagaraDataChannel, meta = (Keywords = "niagara DataChannel", WorldContext = "WorldContextObject", UnsafeDuringActorConstruction = "true"))
 	static UNiagaraDataChannelHandler* GetNiagaraDataChannel(const UObject* WorldContextObject, const UNiagaraDataChannelAsset* Channel);
 
-	/** Initializes and returns the Niagara Data Channel writer to write N elements to the given data channel. */
+	/**
+	 * Initializes and returns the Niagara Data Channel writer to write N elements to the given data channel.
+	 *
+	 * @param WorldContextObject	World to execute in
+	 * @param Channel				The channel to write to
+	 * @param SearchParams			Parameters used when retrieving a specific set of Data Channel Data to read or write like the islands data channel type.
+	 * @param Count					The number of elements to write 
+	 * @param bVisibleToGame	If true, the data written to this data channel is visible to Blueprint and C++ logic reading from it
+	 * @param bVisibleToCPU	If true, the data written to this data channel is visible to Niagara CPU emitters
+	 * @param bVisibleToGPU	If true, the data written to this data channel is visible to Niagara GPU emitters
+	 */
 	UFUNCTION(BlueprintCallable, Category = NiagaraDataChannel, meta = (Keywords = "niagara DataChannel", WorldContext = "WorldContextObject", UnsafeDuringActorConstruction = "true"))
-	static UNiagaraDataChannelWriter* WriteToNiagaraDataChannel(const UObject* WorldContextObject, const UNiagaraDataChannelAsset* Channel, FNiagaraDataChannelSearchParameters SearchParams, int32 Count, bool bVisibleToGame, bool bVisibleToCPU, bool bVisibleToGPU);
+	static UNiagaraDataChannelWriter* WriteToNiagaraDataChannel(const UObject* WorldContextObject, const UNiagaraDataChannelAsset* Channel, FNiagaraDataChannelSearchParameters SearchParams, int32 Count, UPARAM(DisplayName = "Visible to Blueprint") bool bVisibleToGame, UPARAM(DisplayName = "Visible to Niagara CPU") bool bVisibleToCPU, UPARAM(DisplayName = "Visible to Niagara GPU") bool bVisibleToGPU);
 
-	/** Initializes and returns the Niagara Data Channel reader for the given data channel. */
+	/**
+	 * Initializes and returns the Niagara Data Channel reader for the given data channel.
+	 *
+	 * @param WorldContextObject	World to execute in
+	 * @param Channel				The channel to read from
+	 * @param SearchParams			Parameters used when retrieving a specific set of Data Channel Data to read or write like the islands data channel type.
+	 * @param bReadPreviousFrame	True if this reader will read the previous frame's data. If false, we read the current frame.
+	 *								Reading the current frame allows for zero latency reads, but any data elements that are generated after this reader is used are missed.
+	 *								Reading the previous frame's data introduces a frame of latency but ensures we never miss any data as we have access to the whole frame.
+	 */
 	UFUNCTION(BlueprintCallable, Category = NiagaraDataChannel, meta = (Keywords = "niagara DataChannel", WorldContext = "WorldContextObject", UnsafeDuringActorConstruction = "true"))
 	static UNiagaraDataChannelReader* ReadFromNiagaraDataChannel(const UObject* WorldContextObject, const UNiagaraDataChannelAsset* Channel, FNiagaraDataChannelSearchParameters SearchParams, bool bReadPreviousFrame);
 
diff --git a/Engine/Plugins/FX/Niagara/Source/Niagara/Public/NiagaraDataChannelPublic.h b/Engine/Plugins/FX/Niagara/Source/Niagara/Public/NiagaraDataChannelPublic.h
--- a/Engine/Plugins/FX/Niagara/Source/Niagara/Public/NiagaraDataChannelPublic.h
+++ b/Engine/Plugins/FX/Niagara/Source/Niagara/Public/NiagaraDataChannelPublic.h
@@ -15,7 +15,15 @@ struct FNiagaraDataChannelData;
 using FNiagaraDataChannelGameDataPtr = TSharedPtr<FNiagaraDataChannelGameData>;
 using FNiagaraDataChannelDataPtr = TSharedPtr<FNiagaraDataChannelData>;
 
-/** Wrapper asset class for UNiagaraDataChannel which is instanced. */
+/** Niagara Data Channels are a system for communication between Niagara Systems and with game code/Blueprint.
+
+Data channel assets define the payload as well as some transfer settings.
+Niagara Systems can read from and write to data channels via data interfaces.
+Blueprint and C++ code can also read from and write to data channels using its API functions.
+
+EXPERIMENTAL:
+Data Channels are currently experimental and undergoing heavy development.
+ */
 UCLASS(Experimental, DisplayName = "Niagara Data Channel", MinimalAPI)
 class UNiagaraDataChannelAsset : public UObject
 {
diff --git a/Engine/Plugins/FX/Niagara/Source/Niagara/Public/NiagaraDataChannel_Global.h b/Engine/Plugins/FX/Niagara/Source/Niagara/Public/NiagaraDataChannel_Global.h
--- a/Engine/Plugins/FX/Niagara/Source/Niagara/Public/NiagaraDataChannel_Global.h
+++ b/Engine/Plugins/FX/Niagara/Source/Niagara/Public/NiagaraDataChannel_Global.h
@@ -6,7 +6,9 @@
 #include "NiagaraDataChannelHandler.h"
 #include "NiagaraDataChannel_Global.generated.h"
 
-
+/**
+Simple DataChannel handler that makes all data visible globally.
+*/
 UCLASS(Experimental, MinimalAPI)
 class UNiagaraDataChannel_Global : public UNiagaraDataChannel
 {
diff --git a/Engine/Plugins/FX/Niagara/Source/Niagara/Public/NiagaraDataChannel_Islands.h b/Engine/Plugins/FX/Niagara/Source/Niagara/Public/NiagaraDataChannel_Islands.h
--- a/Engine/Plugins/FX/Niagara/Source/Niagara/Public/NiagaraDataChannel_Islands.h
+++ b/Engine/Plugins/FX/Niagara/Source/Niagara/Public/NiagaraDataChannel_Islands.h
@@ -106,6 +106,9 @@ struct FNDCIslandDebugDrawSettings
 	bool ShowBounds()const { return bEnabled && bShowIslandBounds; }
 };
 
+/**
+Data channel that will automatically sub-divide the world into discreet "islands" based on location.
+*/
 UCLASS(Experimental, MinimalAPI)
 class UNiagaraDataChannel_Islands : public UNiagaraDataChannel
 {
diff --git a/Engine/Plugins/FX/Niagara/Source/NiagaraEditor/Private/AssetDefinitions/AssetDefinition_NiagaraDataChannel.cpp b/Engine/Plugins/FX/Niagara/Source/NiagaraEditor/Private/AssetDefinitions/AssetDefinition_NiagaraDataChannel.cpp
--- a/Engine/Plugins/FX/Niagara/Source/NiagaraEditor/Private/AssetDefinitions/AssetDefinition_NiagaraDataChannel.cpp
+++ b/Engine/Plugins/FX/Niagara/Source/NiagaraEditor/Private/AssetDefinitions/AssetDefinition_NiagaraDataChannel.cpp
@@ -12,7 +12,7 @@
 
 FLinearColor UAssetDefinition_NiagaraDataChannel::GetAssetColor() const
 {
-	return FNiagaraEditorStyle::Get().GetColor("NiagaraEditor.AssetColors.DataChannel");
+	return FNiagaraEditorStyle::Get().GetColor("NiagaraEditor.AssetColors.DataChannelDefinitions");
 }
 
 TSoftClassPtr<> UAssetDefinition_NiagaraDataChannel::GetAssetClass() const
diff --git a/Engine/Plugins/FX/Niagara/Source/NiagaraEditor/Private/Customizations/NiagaraDataChannelDetails.cpp b/Engine/Plugins/FX/Niagara/Source/NiagaraEditor/Private/Customizations/NiagaraDataChannelDetails.cpp
--- a/Engine/Plugins/FX/Niagara/Source/NiagaraEditor/Private/Customizations/NiagaraDataChannelDetails.cpp
+++ b/Engine/Plugins/FX/Niagara/Source/NiagaraEditor/Private/Customizations/NiagaraDataChannelDetails.cpp
@@ -7,9 +7,7 @@
 #include "PropertyHandle.h"
 #include "DetailLayoutBuilder.h"
 #include "DetailCategoryBuilder.h"
-#include "IDetailChildrenBuilder.h"
 #include "DetailWidgetRow.h"
-#include "PropertyCustomizationHelpers.h"
 #include "ObjectEditorUtils.h"
 
 #include "Widgets/Text/STextBlock.h"
@@ -17,6 +15,7 @@
 #include "Widgets/Input/SButton.h"
 
 #include "Modules/ModuleManager.h"
+#include "Widgets/Input/SHyperlink.h"
 
 #define LOCTEXT_NAMESPACE "NiagaraDataChannelDetails"
 
@@ -32,8 +31,6 @@ FNiagaraDataChannelAssetDetails::~FNiagaraDataChannelAssetDetails()
 
 void FNiagaraDataChannelAssetDetails::CustomizeDetails(IDetailLayoutBuilder& DetailBuilder)
 {
-	FPropertyEditorModule& PropertyModule = FModuleManager::GetModuleChecked<FPropertyEditorModule>("PropertyEditor");
-
 	TArray<TWeakObjectPtr<UObject>> CustomizedObjects;
 	DetailBuilder.GetObjectsBeingCustomized(CustomizedObjects);
 
@@ -50,59 +47,62 @@ void FNiagaraDataChannelAssetDetails::CustomizeDetails(IDetailLayoutBuilder& Det
 	
 	FName DataChannelCatName("Data Channel");
 	IDetailCategoryBuilder& DataChannelCat = DetailBuilder.EditCategory(DataChannelCatName);
-
+	
 	TSharedPtr<IPropertyHandle> DataChannelProperty = DetailBuilder.GetProperty("DataChannel");
-	DataChannelProperty->SetPropertyDisplayName(LOCTEXT("DataChannelTypeName", "Data Channel Class"));
+	DataChannelProperty->SetPropertyDisplayName(LOCTEXT("DataChannelTypeName", "Data Channel Type"));
 	
 	// we need to build a mapping between categories and properties first because the CLASS_CollapseCategories class flag can
 	// remove the categories from top-level properties, so we need to traverse them differently
 	TMap<FName, TArray<TSharedPtr<IPropertyHandle>>> CategoryMapping;
 
 	//Always add the NDC type property.
-	CategoryMapping.Add(DataChannelCatName).Add(DataChannelProperty);
+	DataChannelCat.AddProperty(DataChannelProperty);
+	DataChannelCat.AddCustomRow(FText::FromString("ChannelTypeHelp"))
+	.Visibility(TAttribute<EVisibility>::CreateLambda([this]()
+	{
+		return DataChannelAsset.IsValid() && DataChannelAsset->Get() ? EVisibility::Collapsed : EVisibility::Visible;
+	}))
+	.WholeRowWidget
+	[
+		SNew(SBox)
+		.HAlign(HAlign_Left)
+		.VAlign(VAlign_Center)
+		[
+			SNew(SHyperlink)
+			.Text(LOCTEXT("DataChannelTypePage", "Which data channel type should I choose?"))
+			.TextStyle(&FCoreStyle::Get().GetWidgetStyle<FTextBlockStyle>("SmallText"))
+			.ToolTipText(LOCTEXT("DataChannelTypePageTooltip", "Opens a docs page that explains data channel types"))
+			.OnNavigate_Lambda([]()
+			{
+				//TODO (mga) set correct url here
+				FPlatformProcess::LaunchURL(TEXT("https://dev.epicgames.com/community/unreal-engine/learning?types=tutorial"), nullptr, nullptr);
+			})
+		]
+	];
 
-	UNiagaraDataChannel* NDC = DataChannelAsset->Get();
-	if(NDC != nullptr)
+	if(DataChannelAsset->Get())
 	{		
 		TSharedPtr<IPropertyHandle> InstancedNDCHandle = DataChannelProperty->GetChildHandle(0);
 		check(InstancedNDCHandle);
 
 		uint32 NumChildren = 0;
 		InstancedNDCHandle->GetNumChildren(NumChildren);
-		UClass* NDCClass = NDC->GetClass();
-		bool HasCollapsedCategories = NDCClass->HasAnyClassFlags(CLASS_CollapseCategories);
 
 		for (uint32 i = 0; i < NumChildren; i++)
 		{
 			TSharedPtr<IPropertyHandle> TopLevelHandle = InstancedNDCHandle->GetChildHandle(i);
-
-			if (HasCollapsedCategories)
+			uint32 CategoryChildren;
+			TopLevelHandle->GetNumChildren(CategoryChildren);
+			for (uint32 k = 0; k < CategoryChildren; k++)
 			{
-				if (!TopLevelHandle.IsValid() || !TopLevelHandle->GetProperty() || !TopLevelHandle->GetProperty()->GetClass())
+				TSharedPtr<IPropertyHandle> ChildHandle = TopLevelHandle->GetChildHandle(k);
+				if (!ChildHandle->GetProperty() || !ChildHandle->GetProperty()->GetClass())
 				{
 					continue;
 				}
 
-				// with collapsed categories, the properties are given to us directly, without the category properties as parents
-				FProperty* Property = TopLevelHandle->GetProperty();
-				FName CategoryFName = FObjectEditorUtils::GetCategoryFName(Property);
-				CategoryMapping.FindOrAdd(CategoryFName).Add(TopLevelHandle);
-			}
-			else
-			{
-				uint32 CategoryChildren;
-				TopLevelHandle->GetNumChildren(CategoryChildren);
-				for (uint32 k = 0; k < CategoryChildren; k++)
-				{
-					TSharedPtr<IPropertyHandle> ChildHandle = TopLevelHandle->GetChildHandle(k);
-					if (!ChildHandle->GetProperty() || !ChildHandle->GetProperty()->GetClass())
-					{
-						continue;
-					}
-
-					FName CategoryFName = FName(TopLevelHandle->GetPropertyDisplayName().ToString());
-					CategoryMapping.FindOrAdd(CategoryFName).Add(ChildHandle);
-				}
+				FName CategoryFName = FName(TopLevelHandle->GetPropertyDisplayName().ToString());
+				CategoryMapping.FindOrAdd(CategoryFName).Add(ChildHandle);
 			}
 		}
 	}
diff --git a/Engine/Plugins/FX/Niagara/Source/NiagaraEditor/Private/Customizations/NiagaraTypeCustomizations.cpp b/Engine/Plugins/FX/Niagara/Source/NiagaraEditor/Private/Customizations/NiagaraTypeCustomizations.cpp
--- a/Engine/Plugins/FX/Niagara/Source/NiagaraEditor/Private/Customizations/NiagaraTypeCustomizations.cpp
+++ b/Engine/Plugins/FX/Niagara/Source/NiagaraEditor/Private/Customizations/NiagaraTypeCustomizations.cpp
@@ -5,6 +5,7 @@
 #include "CoreMinimal.h"
 #include "DetailLayoutBuilder.h"
 #include "DetailWidgetRow.h"
+#include "EdGraphSchema_Niagara.h"
 #include "MaterialTypes.h"
 #include "IDetailChildrenBuilder.h"
 #include "IDetailGroup.h"
@@ -12,7 +13,6 @@
 #include "PropertyHandle.h"
 #include "ScopedTransaction.h"
 #include "SGraphActionMenu.h"
-#include "DeviceProfiles/DeviceProfileManager.h"
 #include "Framework/Application/SlateApplication.h"
 #include "Framework/MultiBox/MultiBoxBuilder.h"
 #include "Materials/MaterialInterface.h"
@@ -44,7 +44,6 @@
 #include "ViewModels/NiagaraEmitterHandleViewModel.h"
 
 #include "Widgets/SNiagaraParameterMenu.h"
-#include "Widgets/Input/SSuggestionTextBox.h"
 
 #include UE_INLINE_GENERATED_CPP_BY_NAME(NiagaraTypeCustomizations)
 
@@ -2598,85 +2597,144 @@ void FNiagaraVariableDetailsCustomization::CustomizeChildren(TSharedRef<IPropert
 			[
 				PropertyHandle->CreatePropertyValueWidget()
 			];
-
 	}
 	else
 	{
 		void* VarPtr = nullptr;
 		if (PropertyHandle->GetValueData(VarPtr) == FPropertyAccess::Success)
 		{
-			FNiagaraVariable* Var = (FNiagaraVariable*)VarPtr;
+			bool bEnforceUniqueName = PropertyHandle->HasMetaData(TEXT("EnforceUniqueNames"));
+			FNiagaraVariable* Variable = static_cast<FNiagaraVariable*>(VarPtr);
 
 			Row.NameContent()
+			[
+				SNew(SHorizontalBox)
+				+ SHorizontalBox::Slot()
+				.AutoWidth()
+				.VAlign(VAlign_Center)
 				[
-					//PropertyHandle->CreatePropertyNameWidget()
-					SNew(SHorizontalBox)
-					+ SHorizontalBox::Slot()
-					.AutoWidth()
-					.Padding(4.0f, 4.0f, 4.0f, 4.0f)
-					[
-						SNew(SComboButton)
-						.ButtonStyle(FAppStyle::Get(), "RoundButton")
-						.ForegroundColor(FAppStyle::GetSlateColor("DefaultForeground"))
-						.ContentPadding(FMargin(0, 0))
-						.OnGetMenuContent(this, &FNiagaraVariableDetailsCustomization::GetTypeMenu, TypeDefHandleHandle, Var)
-						//.IsEnabled(this, &SNiagaraParameterPanel::GetCanAddParametersToCategory, Category)
-						.HAlign(HAlign_Center)
-						.VAlign(VAlign_Center)
-						.HasDownArrow(false)
-						.ButtonContent()
-						[
-							SNew(SHorizontalBox)
-							+ SHorizontalBox::Slot()
-							.AutoWidth()
-							.Padding(FMargin(0, 1))
-							[
-								SNew(SImage)
-								.Image(FAppStyle::GetBrush("Icons.Edit"))
-							]
-						]
-					]
-					+ SHorizontalBox::Slot()
-					.AutoWidth()
-					.HAlign(HAlign_Center)
-					.VAlign(VAlign_Center)
-					.Padding(4.0f, 4.0f, 4.0f, 4.0f)
-					[
-						SNew(SNiagaraConstrainedBox)
-						.MinWidth(75.0f)
-						[
-							SNew(STextBlock)
-							.Text_Lambda([Var] {return Var ? Var->GetType().GetNameText() : FText::GetEmpty(); })
-						]
-					]
-				];
-
-			Row.ValueContent()
-				[
-					SNew(SHorizontalBox)
-			+ SHorizontalBox::Slot()
+					SNew(STextBlock)
+					.Font(IDetailLayoutBuilder::GetDetailFont())
+					.Text(LOCTEXT("DataChannelVarNameText", "Name: "))
+				]
+				+ SHorizontalBox::Slot()
 				.AutoWidth()
-				.HAlign(HAlign_Center)
 				.VAlign(VAlign_Center)
 				.Padding(4.0f, 4.0f, 4.0f, 4.0f)
 				[
 					SNew(SNiagaraConstrainedBox)
 					.MinWidth(150.0f)
-				[
-					NameHandle->CreatePropertyValueWidget()
-				]
+					[
+						SNew( SEditableTextBox )
+						.Text_Lambda([NameHandle]() -> FText { FText OutText; NameHandle->GetValueAsDisplayText(OutText); return OutText; })
+						.Font(IDetailLayoutBuilder::GetDetailFont())
+						.SelectAllTextWhenFocused(true)
+						.ClearKeyboardFocusOnCommit(false)
+						.OnTextCommitted_Lambda([NameHandle](const FText& NewName, ETextCommit::Type)
+						{
+							FScopedTransaction Transaction(LOCTEXT("ChangeVarName", "Change variable name"));
+							TArray<UObject*> OuterObjects;
+							NameHandle->GetOuterObjects(OuterObjects);
+							for (UObject* OuterObj : OuterObjects)
+							{
+								OuterObj->Modify();
+							}
+							NameHandle->SetValue(FName(NewName.ToString()));
+						} )
+						.OnVerifyTextChanged_Lambda([bEnforceUniqueName, PropertyHandle](const FText& InNewText, FText& OutErrorMessage) -> bool
+						{
+							// if necessary, validate that the entered name is unique among all entries of the variable array
+							if (!bEnforceUniqueName)
+							{
+								return true;
+							}
+							FName NewName = FName(InNewText.ToString());
+							if (TSharedPtr<IPropertyHandle> Handle = PropertyHandle->GetParentHandle())
+							{
+								if (TSharedPtr<IPropertyHandleArray> PropertyHandleArray = Handle->AsArray())
+								{
+									uint32 NumElements;
+									PropertyHandleArray->GetNumElements(NumElements);
+									bool bNewNameValid = true;
+									for (uint32 i = 0; i < NumElements; i++)
+									{
+										TSharedRef<IPropertyHandle> ElementHandle = PropertyHandleArray->GetElement(i);
+										void* VarPtr = nullptr;
+										if (!ElementHandle->IsSamePropertyNode(PropertyHandle) && ElementHandle->GetValueData(VarPtr) == FPropertyAccess::Success)
+										{
+											if (NewName == static_cast<FNiagaraVariable*>(VarPtr)->GetName())
+											{
+												bNewNameValid = false;
+												break;
+											}
+										}
+									}
+									if (!bNewNameValid)
+									{
+										OutErrorMessage = LOCTEXT("DuplicateNameError", "Variable name has to be unique!");
+										return false;
+									}
+								}
+							}
+							return true;
+						})
+						.SelectAllTextOnCommit( true )
+					]
 				]
-			+ SHorizontalBox::Slot()
+				+ SHorizontalBox::Slot()
 				.AutoWidth()
-				.HAlign(HAlign_Center)
 				.VAlign(VAlign_Center)
-				.Padding(4.0f, 4.0f, 4.0f, 4.0f)
+				.Padding(4.0f, 2.0f, 4.0f, 2.0f)
 				[
 					SNew(SNiagaraConstrainedBox)
 					.MinWidth(150.0f)
+					[
+						PropertyHandle->CreateDefaultPropertyButtonWidgets()
+					]
+				]
+			];
+			
+			Row.ValueContent()
+			[
+				SNew(SHorizontalBox)
+				+ SHorizontalBox::Slot()
+				.AutoWidth()
+				.VAlign(VAlign_Center)
+				.Padding(0, 0, 2.0f, 0)
 				[
-					PropertyHandle->CreateDefaultPropertyButtonWidgets()
+					SNew(STextBlock)
+					.Font(IDetailLayoutBuilder::GetDetailFont())
+					.Text(LOCTEXT("DataChannelVarTypeText", "Type: "))
 				]
+				+ SHorizontalBox::Slot()
+				.VAlign(VAlign_Center)
+				.Padding(2.0f)
+				[
+					SNew(SComboButton)
+					.HasDownArrow(true)
+					.ContentPadding(0)
+					.OnGetMenuContent(this, &FNiagaraVariableDetailsCustomization::GetTypeMenu, TypeDefHandleHandle, Variable)
+					.ButtonContent()
+					[
+						SNew(SHorizontalBox)
+						+ SHorizontalBox::Slot()
+						.VAlign(VAlign_Center)
+						.HAlign(HAlign_Center)
+						.AutoWidth()
+						[
+							SNew(SImage)
+							.ColorAndOpacity_Lambda([Variable]() {return UEdGraphSchema_Niagara::GetTypeColor(Variable->GetType());})
+							.Image(FNiagaraEditorStyle::Get().GetBrush("NiagaraEditor.Module.TypeIconPill"))
+						]
+						+ SHorizontalBox::Slot()
+						.Padding(4, 2, 2, 2)
+						[
+							SNew(STextBlock)
+							.MinDesiredWidth(150)
+							.Text_Lambda([Variable] {return Variable ? Variable->GetType().GetNameText() : FText::GetEmpty(); })
+							.Font(IDetailLayoutBuilder::GetDetailFont())
+						]
+					]
 				]
 			];
 		}
diff --git a/Engine/Plugins/FX/Niagara/Source/NiagaraEditor/Private/Customizations/NiagaraTypeCustomizations.h b/Engine/Plugins/FX/Niagara/Source/NiagaraEditor/Private/Customizations/NiagaraTypeCustomizations.h
--- a/Engine/Plugins/FX/Niagara/Source/NiagaraEditor/Private/Customizations/NiagaraTypeCustomizations.h
+++ b/Engine/Plugins/FX/Niagara/Source/NiagaraEditor/Private/Customizations/NiagaraTypeCustomizations.h
@@ -450,9 +450,9 @@ public:
 		return MakeShared<FNiagaraVariableDetailsCustomization>();
 	}
 
-	virtual void CustomizeHeader(TSharedRef<IPropertyHandle> PropertyHandle, FDetailWidgetRow& HeaderRow, IPropertyTypeCustomizationUtils& CustomizationUtils);
+	virtual void CustomizeHeader(TSharedRef<IPropertyHandle> PropertyHandle, FDetailWidgetRow& HeaderRow, IPropertyTypeCustomizationUtils& CustomizationUtils) override;
 
-	virtual void CustomizeChildren(TSharedRef<IPropertyHandle> PropertyHandle, IDetailChildrenBuilder& ChildBuilder, IPropertyTypeCustomizationUtils& CustomizationUtils);
+	virtual void CustomizeChildren(TSharedRef<IPropertyHandle> PropertyHandle, IDetailChildrenBuilder& ChildBuilder, IPropertyTypeCustomizationUtils& CustomizationUtils) override;
 
 private:
 	TSharedRef<SWidget> GetTypeMenu(TSharedPtr<IPropertyHandle> InPropertyHandle, FNiagaraVariable* Var);
