Migrate fx.NiagaraScript.StripByteCodeOnLoad over to a NiagaraSetting property to more easily control it per project

diff --git a/Engine/Plugins/FX/Niagara/Source/Niagara/Private/NiagaraScript.cpp b/Engine/Plugins/FX/Niagara/Source/Niagara/Private/NiagaraScript.cpp
--- a/Engine/Plugins/FX/Niagara/Source/Niagara/Private/NiagaraScript.cpp
+++ b/Engine/Plugins/FX/Niagara/Source/Niagara/Private/NiagaraScript.cpp
@@ -3,40 +3,49 @@
 #include "NiagaraScript.h"
 
 #include "Algo/RemoveIf.h"
+#include "Algo/Sort.h"
 #include "Algo/Unique.h"
 #include "DataDrivenShaderPlatformInfo.h"
+#include "DataInterface/NiagaraDataInterfaceStaticMesh.h"
+#include "HAL/PlatformFileManager.h"
+#include "Interfaces/ITargetPlatform.h"
 #include "Misc/Compression.h"
+#include "Misc/CoreMiscDefines.h"
+#include "Misc/FileHelper.h"
+#include "Misc/SecureHash.h"
 #include "Modules/ModuleManager.h"
 #include "NiagaraCompileHashVisitor.h"
-#include "NiagaraScriptSourceBase.h"
+#include "NiagaraCustomVersion.h"
+#include "NiagaraDataInterfaceSkeletalMesh.h"
 #include "NiagaraEmitter.h"
-#include "NiagaraShader.h"
-#include "UObject/Package.h"
 #include "NiagaraModule.h"
-#include "NiagaraCustomVersion.h"
+#include "NiagaraScriptSourceBase.h"
+#include "NiagaraSettings.h"
+#include "NiagaraShader.h"
+#include "NiagaraSimulationStageBase.h"
 #include "NiagaraStackSection.h"
 #include "NiagaraSystem.h"
-#include "Serialization/MemoryReader.h"
-#include "Misc/SecureHash.h"
-
 #include "ProfilingDebugging/CookStats.h"
+#include "Serialization/MemoryReader.h"
+#include "Serialization/ObjectAndNameAsStringProxyArchive.h"
 #include "Stats/Stats.h"
-#include "HAL/PlatformFileManager.h"
-#include "Misc/FileHelper.h"
 #include "UObject/EditorObjectVersion.h"
+#include "UObject/FortniteMainBranchObjectVersion.h"
 #include "UObject/ObjectSaveContext.h"
+#include "UObject/Package.h"
 #include "UObject/ReleaseObjectVersion.h"
-#include "NiagaraDataInterfaceSkeletalMesh.h"
-#include "DataInterface/NiagaraDataInterfaceStaticMesh.h"
-#include "Interfaces/ITargetPlatform.h"
+#include "UObject/RenderingObjectVersion.h"
+#include "UObject/UObjectIterator.h"
+#include "UObject/UObjectThreadContext.h"
+#include "VectorVM.h"
+
 #if WITH_EDITOR
 	#include "DerivedDataCacheInterface.h"
 	#include "DerivedDataCacheKey.h"
+	#include "INiagaraEditorOnlyDataUtlities.h"
 	#include "Interfaces/ITargetPlatform.h"
-	#include "NiagaraSettings.h"
 	#include "Internationalization/Regex.h"
 	#include "ShaderCodeLibrary.h"
-	#include "INiagaraEditorOnlyDataUtlities.h"
 
 	// This is a version string that mimics the old versioning scheme. In case of merge conflicts with DDC versions,
 	// you MUST generate a new GUID and set this new version. If you want to bump this version, generate a new guid
@@ -44,18 +53,6 @@
 	#define NIAGARASCRIPT_DERIVEDDATA_VER		TEXT("AB7397ACFEFD46158A87743735E3C773")
 #endif
 
-#include "UObject/FortniteMainBranchObjectVersion.h"
-#include "UObject/RenderingObjectVersion.h"
-#include "UObject/UObjectIterator.h"
-#include "UObject/UObjectThreadContext.h"
-#include "Serialization/ObjectAndNameAsStringProxyArchive.h"
-
-#include "VectorVM.h"
-#include "NiagaraSimulationStageBase.h"
-
-#include "Algo/Sort.h"
-#include "Misc/CoreMiscDefines.h"
-
 #include UE_INLINE_GENERATED_CPP_BY_NAME(NiagaraScript)
 
 #if ENABLE_COOK_STATS
@@ -109,14 +106,6 @@
 	ECVF_Default
 );
 
-int32 GNiagaraScriptStripByteCodeOnLoad = 0;
-static FAutoConsoleVariableRef CVarGNiagaraScriptStripByteCodeOnLoad(
-	TEXT("fx.NiagaraScript.StripByteCodeOnLoad"),
-	GNiagaraScriptStripByteCodeOnLoad,
-	TEXT("If > 0 all scripts will have their legacy byte code stripped on load.  If < 0 all scripts will have their experimental data stripped on load. \n"),
-	ECVF_Default
-);
-
 static int32 GbNiagaraEventSpawnsUpdateInitialAttributeValues = 1;
 static FAutoConsoleVariableRef CVarNiagaraEventSpawnsUpdateInitialAttributeValues(
 	TEXT("fx.Niagara.EventSpawnsUpdateAttributeInitialValues"),
@@ -429,13 +418,23 @@
 #if VECTORVM_SUPPORTS_EXPERIMENTAL && VECTORVM_SUPPORTS_LEGACY && !WITH_EDITOR
 	if (Ar.IsLoading())
 	{
-		if (GNiagaraScriptStripByteCodeOnLoad > 0)
+		// only worry about stripping out byte code if we have both sets loaded
+		if (ByteCode.GetLength() > 0 && !ExperimentalContextData.IsEmpty())
 		{
-			ByteCode.Reset();
+			switch (GetDefault<UNiagaraSettings>()->ByteCodeStripOption)
+			{
+			case ENiagaraStripScriptByteCodeOption::Strip_Original:
+				ByteCode.Reset();
+				break;
+
+			case ENiagaraStripScriptByteCodeOption::Strip_Experimental:
+				ExperimentalContextData.Empty();
+				break;
+			
+			default:
+			case ENiagaraStripScriptByteCodeOption::Default:
+				break;
+			}
 		}
-		else if (GNiagaraScriptStripByteCodeOnLoad < 0)
-		{
-			ExperimentalContextData.Empty();
-		}
 	}
 #endif
diff --git a/Engine/Plugins/FX/Niagara/Source/Niagara/Public/NiagaraSettings.h b/Engine/Plugins/FX/Niagara/Source/Niagara/Public/NiagaraSettings.h
--- a/Engine/Plugins/FX/Niagara/Source/Niagara/Public/NiagaraSettings.h
+++ b/Engine/Plugins/FX/Niagara/Source/Niagara/Public/NiagaraSettings.h
@@ -93,6 +93,14 @@ namespace ENDICollisionQuery_AsyncGpuTraceProvider
 	};
 }
 
+UENUM()
+enum class ENiagaraStripScriptByteCodeOption : uint8
+{
+	Default = 0 UMETA(DisplayName = "No Stripping"),
+	Strip_Original = 1 UMETA(DisplayName = "Strip Original ByteCode"),
+	Strip_Experimental = 2 UMETA(DisplayName = "Strip Experimental ByteCode"),
+};
+
 UCLASS(config = Niagara, defaultconfig, meta=(DisplayName="Niagara"), MinimalAPI)
 class UNiagaraSettings : public UDeveloperSettings
 {
@@ -140,6 +148,12 @@ class UNiagaraSettings : public UDeveloperSettings
 	UPROPERTY(config, EditAnywhere, Category = Niagara, meta = (DisplayName = "Enable building data for Experimental VM"))
 	bool bExperimentalVMEnabled = false;
 
+	/**
+	 Controls how byte code will be stripped when loading assets that have multiple sets of bytecode (i.e. experimental).
+	 */
+	UPROPERTY(config, EditAnywhere, Category = Niagara, meta = (DisplayName = "Option for how to strip bytecode"))
+	ENiagaraStripScriptByteCodeOption ByteCodeStripOption = ENiagaraStripScriptByteCodeOption::Default;
+
 	/** Whether to limit the max tick delta time or not. */
 	UPROPERTY(config, EditAnywhere, BlueprintReadWrite, AdvancedDisplay, Category = "Niagara", meta = (InlineEditConditionToggle))
 	bool bLimitDeltaTime = true;
