Allow enum inputs in the stack to be displayed as button grid

diff --git a/Engine/Plugins/FX/Niagara/Source/Niagara/Public/NiagaraTypes.h b/Engine/Plugins/FX/Niagara/Source/Niagara/Public/NiagaraTypes.h
--- a/Engine/Plugins/FX/Niagara/Source/Niagara/Public/NiagaraTypes.h
+++ b/Engine/Plugins/FX/Niagara/Source/Niagara/Public/NiagaraTypes.h
@@ -659,8 +659,12 @@ enum class ENiagaraInputWidgetType : uint8
 	// a numeric input, but also has a dropdown with named values
 	NumericDropdown,
 
-	// A dropdown that behaves like an enum; only allows the exact pre-defined values, for integer inputs only
-	EnumStyle
+	// (for integer inputs only) A dropdown that behaves like an enum; only allows the exact pre-defined values.
+	EnumStyle,
+
+	// (for enum inputs only) Instead of the normal dropdown, the enum values are all displayed in a button grid.
+	// This shows all possible values at once, so only makes sense if there are few input values. 
+	SegmentedButtons
 };
 
 USTRUCT()
@@ -678,6 +682,26 @@ struct FWidgetNamedInputValue
 	FText Tooltip;
 };
 
+USTRUCT()
+struct FWidgetSegmentValueOverride
+{
+	GENERATED_BODY()
+
+	UPROPERTY(EditAnywhere, Category="Customization")
+	int32 EnumIndexToOverride = 0;
+
+	UPROPERTY(EditAnywhere, Category="Customization", meta=(InlineEditConditionToggle))
+	bool bOverrideDisplayName = false;
+	
+	// This will be used as display name instead of the enum value
+	UPROPERTY(EditAnywhere, Category="Customization", meta=(EditCondition="bOverrideDisplayName"))
+	FText DisplayNameOverride;
+
+	// If set, then this icon will be displayed on the button
+	UPROPERTY(EditAnywhere, Category="Customization")
+	TObjectPtr<UTexture2D> DisplayIcon;
+};
+
 /** A struct that serves as display metadata for integer type static switches. Is used in conjuction with the 'EnumStyle' widget customization. */
 USTRUCT()
 struct FNiagaraWidgetNamedIntegerInputValue
@@ -727,6 +751,13 @@ struct FNiagaraInputParameterCustomization
 	UPROPERTY(EditAnywhere, Category="Customization", meta=(EditCondition="WidgetType == ENiagaraInputWidgetType::EnumStyle", EditConditionHides))
 	TArray<FNiagaraWidgetNamedIntegerInputValue> EnumStyleDropdownValues;
 
+	// Limits the number of buttons shown per row, 0 = unlimited
+	UPROPERTY(EditAnywhere, Category="Customization", meta=(EditCondition="WidgetType == ENiagaraInputWidgetType::SegmentedButtons", EditConditionHides))
+	int32 MaxSegmentsPerRow = 0;
+
+	UPROPERTY(EditAnywhere, Category="Customization", meta=(EditCondition="WidgetType == ENiagaraInputWidgetType::SegmentedButtons", EditConditionHides))
+	TArray<FWidgetSegmentValueOverride> SegmentValueOverrides;
+
 	UPROPERTY()
 	bool bBroadcastValueChangesOnCommitOnly = false;
 	
diff --git a/Engine/Plugins/FX/Niagara/Source/NiagaraEditor/Private/NiagaraEditorStyle.cpp b/Engine/Plugins/FX/Niagara/Source/NiagaraEditor/Private/NiagaraEditorStyle.cpp
--- a/Engine/Plugins/FX/Niagara/Source/NiagaraEditor/Private/NiagaraEditorStyle.cpp
+++ b/Engine/Plugins/FX/Niagara/Source/NiagaraEditor/Private/NiagaraEditorStyle.cpp
@@ -8,6 +8,7 @@
 #include "Styling/SlateStyleRegistry.h"
 #include "Styling/SlateTypes.h"
 #include "Styling/CoreStyle.h"
+#include "Styling/SegmentedControlStyle.h"
 #include "Styling/StarshipCoreStyle.h"
 #include "Styling/StyleColors.h"
 
@@ -516,6 +517,16 @@ void FNiagaraEditorStyle::InitStackWidgets()
 
 	Set("NiagaraEditor.Stack.Note.HeaderText", NoteHeaderStyle);
 	Set("NiagaraEditor.Stack.Note.HeaderEditableText", NoteHeaderInlineEditableStyle);
+
+	FSegmentedControlStyle SegmentedButtonStyle = FCoreStyle::Get().GetWidgetStyle<FSegmentedControlStyle>("SegmentedControl");
+	FCheckBoxStyle SegmentedControlStyle = FCoreStyle::Get().GetWidgetStyle<FCheckBoxStyle>("DetailsView.SectionButton");
+	SegmentedButtonStyle
+	.SetControlStyle(SegmentedControlStyle)
+	.SetFirstControlStyle(SegmentedControlStyle)
+	.SetLastControlStyle(SegmentedControlStyle)
+	.SetBackgroundBrush(FSlateNoResource());
+	
+	Set("NiagaraEditor.Stack.SegmentedControl", SegmentedButtonStyle);
 }
 
 void FNiagaraEditorStyle::InitNiagaraSequence()
diff --git a/Engine/Plugins/FX/Niagara/Source/NiagaraEditor/Private/TypeEditorUtilities/NiagaraEnumTypeEditorUtilities.cpp b/Engine/Plugins/FX/Niagara/Source/NiagaraEditor/Private/TypeEditorUtilities/NiagaraEnumTypeEditorUtilities.cpp
--- a/Engine/Plugins/FX/Niagara/Source/NiagaraEditor/Private/TypeEditorUtilities/NiagaraEnumTypeEditorUtilities.cpp
+++ b/Engine/Plugins/FX/Niagara/Source/NiagaraEditor/Private/TypeEditorUtilities/NiagaraEnumTypeEditorUtilities.cpp
@@ -3,12 +3,14 @@
 #include "NiagaraEnumTypeEditorUtilities.h"
 #include "SNiagaraParameterEditor.h"
 #include "NiagaraTypes.h"
-#include "NiagaraEditorStyle.h"
 #include "NiagaraEditorCommon.h"
+#include "NiagaraEditorStyle.h"
 
 #include "SEnumCombo.h"
+#include "Engine/Texture2D.h"
 #include "Styling/AppStyle.h"
-#include "Widgets/Input/SSpinBox.h"
+#include "Brushes/SlateImageBrush.h"
+#include "Widgets/Input/SSegmentedControl.h"
 
 class SNiagaraEnumParameterEditor : public SNiagaraParameterEditor
 {
@@ -16,32 +18,90 @@ public:
 	SLATE_BEGIN_ARGS(SNiagaraEnumParameterEditor) { }
 	SLATE_END_ARGS();
 
-	void Construct(const FArguments& InArgs, const UEnum* Enum)
+	void Construct(const FArguments& InArgs, const UEnum* Enum, const FNiagaraInputParameterCustomization& WidgetCustomization)
 	{
 		SNiagaraParameterEditor::Construct(SNiagaraParameterEditor::FArguments()
 			.MinimumDesiredWidth(DefaultInputSize)
-			.MaximumDesiredWidth(2 * DefaultInputSize));
-
-		ChildSlot
-		[
-			SNew(SEnumComboBox, Enum)
-			.CurrentValue(this, &SNiagaraEnumParameterEditor::GetValue)
-			.ContentPadding(FMargin(2, 0))
-			.Font(FAppStyle::Get().GetFontStyle("PropertyWindow.NormalFont"))
-			.OnEnumSelectionChanged(SEnumComboBox::FOnEnumSelectionChanged::CreateSP(this, &SNiagaraEnumParameterEditor::ValueChanged))
-		];
+			.MaximumDesiredWidth(3 * DefaultInputSize));
+
+		if (WidgetCustomization.WidgetType == ENiagaraInputWidgetType::SegmentedButtons)
+		{
+			TSharedRef<SSegmentedControl<int32>> ButtonOptionsPanel =
+			SNew(SSegmentedControl<int32>)
+			.UniformPadding(FMargin(10, 5))
+			.MaxSegmentsPerLine(WidgetCustomization.MaxSegmentsPerRow)
+			.TextStyle(FNiagaraEditorStyle::Get(), "NiagaraEditor.ParameterText")
+			.Style(FNiagaraEditorStyle::Get(), "NiagaraEditor.Stack.SegmentedControl")
+			.Value(this, &SNiagaraEnumParameterEditor::GetValue)
+			.OnValueChanged(this, &SNiagaraEnumParameterEditor::SegmentedValueChanged);
+
+			TMap<int32, FWidgetSegmentValueOverride> Overrides;
+			for (const FWidgetSegmentValueOverride& WidgetOverride : WidgetCustomization.SegmentValueOverrides)
+			{
+				Overrides.Add(WidgetOverride.EnumIndexToOverride, WidgetOverride);
+			}
+			for (int i = 0; i < Enum->NumEnums() - 1; i++)
+			{
+				if (Enum->HasMetaData(TEXT("Hidden"), i))
+				{
+					continue;
+				}
+				FWidgetSegmentValueOverride* Override = Overrides.Find(i);
+				int32 Value = static_cast<int32>(Enum->GetValueByIndex(i));
+				FText DisplayName = Enum->GetDisplayNameTextByIndex(i);
+				FText TooltipText = Enum->GetToolTipTextByIndex(i);
+				if (TooltipText.IsEmpty())
+				{
+					TooltipText = DisplayName;
+				}
+				
+				if (Override && Override->bOverrideDisplayName)
+				{
+					DisplayName = Override->DisplayNameOverride;
+				}
+
+				SSegmentedControl<int>::FScopedWidgetSlotArguments Slot = ButtonOptionsPanel->AddSlot(Value);
+				Slot
+                  .HAlign(HAlign_Center)
+                  .VAlign(VAlign_Center)
+                  .Text(DisplayName)
+                  .ToolTip(TooltipText);
+
+				if (Override && Override->DisplayIcon)
+				{
+					ImageBrushes.Add(MakeShared<FSlateImageBrush>(Override->DisplayIcon, FVector2D(16.f, 16.f)));
+					Slot.Icon(ImageBrushes.Last().Get());
+				}
+			}
+			
+			ChildSlot
+			[
+				ButtonOptionsPanel
+			];
+		}
+		else
+		{
+			ChildSlot
+			[
+				SNew(SEnumComboBox, Enum)
+				.CurrentValue(this, &SNiagaraEnumParameterEditor::GetValue)
+				.ContentPadding(FMargin(2, 0))
+				.Font(FAppStyle::Get().GetFontStyle("PropertyWindow.NormalFont"))
+				.OnEnumSelectionChanged(SEnumComboBox::FOnEnumSelectionChanged::CreateSP(this, &SNiagaraEnumParameterEditor::ValueChanged))
+			];
+		}
 	}
 
 	virtual void UpdateInternalValueFromStruct(TSharedRef<FStructOnScope> Struct) override
 	{
 		checkf(Struct->GetStruct() == FNiagaraTypeDefinition::GetIntStruct(), TEXT("Struct type not supported."));
-		IntValue = ((FNiagaraInt32*)Struct->GetStructMemory())->Value;
+		IntValue = reinterpret_cast<FNiagaraInt32*>(Struct->GetStructMemory())->Value;
 	}
 
 	virtual void UpdateStructFromInternalValue(TSharedRef<FStructOnScope> Struct) override
 	{
 		checkf(Struct->GetStruct() == FNiagaraTypeDefinition::GetIntStruct(), TEXT("Struct type not supported."));
-		((FNiagaraInt32*)Struct->GetStructMemory())->Value = IntValue;
+		reinterpret_cast<FNiagaraInt32*>(Struct->GetStructMemory())->Value = IntValue;
 	}
 
 private:
@@ -56,8 +116,14 @@ private:
 		ExecuteOnValueChanged();
 	}
 
-private:
-	int32 IntValue;
+	void SegmentedValueChanged(int32 NewValue)
+	{
+		IntValue = NewValue;
+		ExecuteOnValueChanged();
+	}
+
+	int32 IntValue = 0;
+	TArray<TSharedPtr<FSlateImageBrush>> ImageBrushes;
 };
 
 bool FNiagaraEditorEnumTypeUtilities::CanProvideDefaultValue() const
@@ -78,7 +144,7 @@ void FNiagaraEditorEnumTypeUtilities::UpdateVariableWithDefaultValue(FNiagaraVar
 
 TSharedPtr<SNiagaraParameterEditor> FNiagaraEditorEnumTypeUtilities::CreateParameterEditor(const FNiagaraTypeDefinition& ParameterType, EUnit DisplayUnit, const FNiagaraInputParameterCustomization& WidgetCustomization) const
 {
-	return SNew(SNiagaraEnumParameterEditor, ParameterType.GetEnum());
+	return SNew(SNiagaraEnumParameterEditor, ParameterType.GetEnum(), WidgetCustomization);
 }
 
 bool FNiagaraEditorEnumTypeUtilities::CanHandlePinDefaults() const
@@ -93,7 +159,7 @@ FString FNiagaraEditorEnumTypeUtilities::GetPinDefaultStringFromValue(const FNia
 	const UEnum* Enum = AllocatedVariable.GetType().GetEnum();
 	checkf(Enum != nullptr, TEXT("Variable is not an enum type."));
 
-	int64 EnumValue = AllocatedVariable.GetValue<int32>();
+	int32 EnumValue = AllocatedVariable.GetValue<int32>();
 	if(Enum->IsValidEnumValue(EnumValue))
 	{
 		return Enum->GetNameStringByValue(EnumValue);
@@ -113,7 +179,7 @@ bool FNiagaraEditorEnumTypeUtilities::SetValueFromPinDefaultString(const FString
 	checkf(Enum != nullptr, TEXT("Variable is not an enum type."));
 
 	FNiagaraInt32 EnumValue;
-	EnumValue.Value = (int32)Enum->GetValueByNameString(StringValue);
+	EnumValue.Value = static_cast<int32>(Enum->GetValueByNameString(StringValue));
 	if(EnumValue.Value != INDEX_NONE)
 	{
 		Variable.AllocateData();
@@ -146,7 +212,7 @@ bool FNiagaraEditorEnumTypeUtilities::SetValueFromDisplayName(const FText& TextV
 	if (FoundIndex != INDEX_NONE)
 	{
 		FNiagaraInt32 EnumValue;
-		EnumValue.Value = (int32)Enum->GetValueByIndex(FoundIndex);
+		EnumValue.Value = static_cast<int32>(Enum->GetValueByIndex(FoundIndex));
 		Variable.AllocateData();
 		Variable.SetValue<FNiagaraInt32>(EnumValue);
 		return true;
