Fixed non-unity header includes

diff --git a/Engine/Plugins/FX/Niagara/Source/NiagaraEditor/Private/Customizations/NiagaraMetaDataCustomNodeBuilder.h b/Engine/Plugins/FX/Niagara/Source/NiagaraEditor/Private/Customizations/NiagaraMetaDataCustomNodeBuilder.h
--- a/Engine/Plugins/FX/Niagara/Source/NiagaraEditor/Private/Customizations/NiagaraMetaDataCustomNodeBuilder.h
+++ b/Engine/Plugins/FX/Niagara/Source/NiagaraEditor/Private/Customizations/NiagaraMetaDataCustomNodeBuilder.h
@@ -4,6 +4,7 @@
 
 #include "IDetailCustomNodeBuilder.h"
 #include "NiagaraTypes.h"
+#include "NiagaraVariableMetaData.h"
 #include "NiagaraMetaDataCustomNodeBuilder.generated.h"
 
 class UNiagaraGraph;
diff --git a/Engine/Plugins/FX/Niagara/Source/NiagaraEditor/Private/Customizations/NiagaraOutlinerCustomization.h b/Engine/Plugins/FX/Niagara/Source/NiagaraEditor/Private/Customizations/NiagaraOutlinerCustomization.h
--- a/Engine/Plugins/FX/Niagara/Source/NiagaraEditor/Private/Customizations/NiagaraOutlinerCustomization.h
+++ b/Engine/Plugins/FX/Niagara/Source/NiagaraEditor/Private/Customizations/NiagaraOutlinerCustomization.h
@@ -10,6 +10,7 @@
 #include "Layout/Visibility.h"
 #include "NiagaraDebuggerCommon.h"
 #include "Widgets/SCompoundWidget.h"
+#include "UObject/StructOnScope.h"
 
 #include "NiagaraOutlinerCustomization.generated.h"
 
diff --git a/Engine/Plugins/FX/Niagara/Source/NiagaraEditor/Private/Customizations/NiagaraParameterBindingCustomization.cpp b/Engine/Plugins/FX/Niagara/Source/NiagaraEditor/Private/Customizations/NiagaraParameterBindingCustomization.cpp
--- a/Engine/Plugins/FX/Niagara/Source/NiagaraEditor/Private/Customizations/NiagaraParameterBindingCustomization.cpp
+++ b/Engine/Plugins/FX/Niagara/Source/NiagaraEditor/Private/Customizations/NiagaraParameterBindingCustomization.cpp
@@ -23,6 +23,7 @@
 #include "NiagaraScriptSource.h"
 #include "NiagaraSimulationStageBase.h"
 #include "NiagaraSystem.h"
+#include "NiagaraVariableMetaData.h"
 #include "NiagaraParameterBinding.h"
 #include "SNiagaraParameterEditor.h"
 #include "Customizations/NiagaraTypeCustomizations.h"
diff --git a/Engine/Plugins/FX/Niagara/Source/NiagaraEditor/Private/TypeEditorUtilities/NiagaraEnumTypeEditorUtilities.cpp b/Engine/Plugins/FX/Niagara/Source/NiagaraEditor/Private/TypeEditorUtilities/NiagaraEnumTypeEditorUtilities.cpp
--- a/Engine/Plugins/FX/Niagara/Source/NiagaraEditor/Private/TypeEditorUtilities/NiagaraEnumTypeEditorUtilities.cpp
+++ b/Engine/Plugins/FX/Niagara/Source/NiagaraEditor/Private/TypeEditorUtilities/NiagaraEnumTypeEditorUtilities.cpp
@@ -5,6 +5,7 @@
 #include "NiagaraTypes.h"
 #include "NiagaraEditorCommon.h"
 #include "NiagaraEditorStyle.h"
+#include "NiagaraVariableMetaData.h"
 
 #include "SEnumCombo.h"
 #include "Engine/Texture2D.h"
diff --git a/Engine/Plugins/FX/Niagara/Source/NiagaraEditor/Private/TypeEditorUtilities/NiagaraVectorTypeEditorUtilities.h b/Engine/Plugins/FX/Niagara/Source/NiagaraEditor/Private/TypeEditorUtilities/NiagaraVectorTypeEditorUtilities.h
--- a/Engine/Plugins/FX/Niagara/Source/NiagaraEditor/Private/TypeEditorUtilities/NiagaraVectorTypeEditorUtilities.h
+++ b/Engine/Plugins/FX/Niagara/Source/NiagaraEditor/Private/TypeEditorUtilities/NiagaraVectorTypeEditorUtilities.h
@@ -3,6 +3,7 @@
 #pragma once
 
 #include "INiagaraEditorTypeUtilities.h"
+#include "NiagaraVariableMetaData.h"
 
 class SNiagaraParameterEditor;
 
diff --git a/Engine/Plugins/FX/Niagara/Source/NiagaraEditor/Public/INiagaraEditorTypeUtilities.h b/Engine/Plugins/FX/Niagara/Source/NiagaraEditor/Public/INiagaraEditorTypeUtilities.h
--- a/Engine/Plugins/FX/Niagara/Source/NiagaraEditor/Public/INiagaraEditorTypeUtilities.h
+++ b/Engine/Plugins/FX/Niagara/Source/NiagaraEditor/Public/INiagaraEditorTypeUtilities.h
@@ -6,6 +6,7 @@
 #include "Internationalization/Text.h"
 #include "Math/UnitConversion.h"
 #include "NiagaraTypes.h"
+#include "NiagaraVariableMetaData.h"
 
 class FStructOnScope;
 class SNiagaraParameterEditor;
