Fixed non-unity header includes

diff --git a/Engine/Plugins/FX/Niagara/Source/NiagaraEditor/Private/Customizations/NiagaraParameterBindingCustomization.h b/Engine/Plugins/FX/Niagara/Source/NiagaraEditor/Private/Customizations/NiagaraParameterBindingCustomization.h
--- a/Engine/Plugins/FX/Niagara/Source/NiagaraEditor/Private/Customizations/NiagaraParameterBindingCustomization.h
+++ b/Engine/Plugins/FX/Niagara/Source/NiagaraEditor/Private/Customizations/NiagaraParameterBindingCustomization.h
@@ -6,6 +6,7 @@
 #include "CoreMinimal.h"
 #include "IPropertyTypeCustomization.h"
 #include "Layout/Visibility.h"
+#include "UObject/StructOnScope.h"
 
 #include "NiagaraCommon.h"
 #include "NiagaraParameterBindingCustomization.generated.h"
