Set segmented control buttons to default styling

diff --git a/Engine/Plugins/FX/Niagara/Source/NiagaraEditor/Private/TypeEditorUtilities/NiagaraEnumTypeEditorUtilities.cpp b/Engine/Plugins/FX/Niagara/Source/NiagaraEditor/Private/TypeEditorUtilities/NiagaraEnumTypeEditorUtilities.cpp
--- a/Engine/Plugins/FX/Niagara/Source/NiagaraEditor/Private/TypeEditorUtilities/NiagaraEnumTypeEditorUtilities.cpp
+++ b/Engine/Plugins/FX/Niagara/Source/NiagaraEditor/Private/TypeEditorUtilities/NiagaraEnumTypeEditorUtilities.cpp
@@ -32,7 +32,7 @@ public:
 			.UniformPadding(FMargin(10, 5))
 			.MaxSegmentsPerLine(WidgetCustomization.MaxSegmentsPerRow)
 			.TextStyle(FNiagaraEditorStyle::Get(), "NiagaraEditor.ParameterText")
-			.Style(FNiagaraEditorStyle::Get(), "NiagaraEditor.Stack.SegmentedControl")
+			//.Style(FNiagaraEditorStyle::Get(), "NiagaraEditor.Stack.SegmentedControl")
 			.Value(this, &SNiagaraEnumParameterEditor::GetValue)
 			.OnValueChanged(this, &SNiagaraEnumParameterEditor::SegmentedValueChanged);
 
