- Fix various persistent debug draw functions not compiling on GPU

diff --git a/Engine/Plugins/FX/Niagara/Shaders/Private/NiagaraDataInterfaceDebugDraw.ush b/Engine/Plugins/FX/Niagara/Shaders/Private/NiagaraDataInterfaceDebugDraw.ush
--- a/Engine/Plugins/FX/Niagara/Shaders/Private/NiagaraDataInterfaceDebugDraw.ush
+++ b/Engine/Plugins/FX/Niagara/Shaders/Private/NiagaraDataInterfaceDebugDraw.ush
@@ -456,14 +456,14 @@ void NDIDebugDraw_Torus(bool bDraw, float3 Location, float3 Axis, float MajorRad
 }
 
 
-#define MakeDrawLinePersistent(NAME)				void NAME(float3 StartLocation, int StartLocationCoordinateSpace, float3 EndLocation, int EndLocationCoordinateSpace, float4 Color) {}
-#define MakeDrawRectanglePersistent(NAME)			void NAME(float3 Center, int CenterCoordinateSpace, float3 Offset, int OffsetCoordinateSpace, float3 Extents, bool HalfExtents, float3 XAxis, float3 YAxis, int AxisCoordinateSpace, int NumXSegments, int NumYSegments, bool bUnboundedPlane, float4 Color) {}
+#define MakeDrawBoxPersistent(NAME)					void NAME(float3 Center, int CenterCoordinateSpace, float3 Offset, int OffsetCoordinateSpace, float3 Extents, bool HalfExtents, float3 RotationAxis, float RotationNormalizedAngle, int RotationCoordinateSpace, float4 Color) {}
 #define MakeDrawCirclePersistent(NAME)				void NAME(float3 Center, int CenterCoordinateSpace, float3 Offset, int OffsetCoordinateSpace, float Radius, float3 XAxis, float3 YAxis, int AxisCoordinateSpace, int NumSegments, float4 Color) {}
-#define MakeDrawBoxPersistent(NAME)					void NAME(float3 Center, int CenterCoordinateSpace, float3 Offset, int OffsetCoordinateSpace, float3 Extents, bool HalfExtents, float3 RotationAxis, float RotationAngle, int RotationCoordinateSpace, float4 Color) {}
-#define MakeDrawSpherePersistent(NAME)				void NAME(float3 Center, int CenterCoordinateSpace, float3 OffsetFromCenter, int OffsetCoordinateSpace, float Radius,  int RadiusCoordinateSpace, bool HemisphereX, bool HemisphereY, bool HemisphereZ, float3 OrientationAxis, int RotationCoordinateSpace, float4 AdditionalRotation, float3 NonUniformScale, int NumSegments, float4 Color) {}
-#define MakeDrawCylinderPersistent(NAME)			void NAME(float3 Center, int CenterCoordinateSpace, float3 OffsetFromCenter, int OffsetCoordinateSpace, float Radius, float Height, bool IsHalfHeight, bool HemisphereX, bool HemisphereY, float3 OrientationAxis, int OrientationAxisCoordinateSpace, float3 NonUniformScale, int NumRadiusSegments, int NumHeightSegments, float4 Color) {}
-#define MakeDrawConePersistent(NAME)				void NAME(float3 Center, int CenterCoordinateSpace, float3 OffsetFromCenter, int OffsetCoordinateSpace, float Angle, float Length, int OrientationAxis, int OrientationAxisCoordinateSpace, float3 NonUniformScale, int NumRadiusSegments, int NumHeightSegments, float4 Color) {}
-#define MakeDrawTorusPersistent(NAME)				void NAME(float3 Center, int CenterCoordinateSpace, float3 OffsetFromCenter, int OffsetCoordinateSpace, float MajorRadius, float MinorRadius, float3 OrientationAxis, int OrientationAxisCoordinateSpace, float3 NonUniformScale, int MajorRadiusSegments, int MinorRadiusSegments, float4 Color) {}
+#define MakeDrawConePersistent(NAME)				void NAME(float3 Center, int CenterCoordinateSpace, float3 Offset, int OffsetCoordinateSpace, float Angle, float Length, float3 OrientationAxis, int OrientationAxisCoordinateSpace, float3 NonUniformScale, int NumRadiusSegments, int NumHeightSegments, float4 Color) {}
 #define MakeDrawCoordinateSystemPersistent(NAME)	void NAME(float3 Center, int CenterCoordinateSpace, float3 OffsetFromCenter, int OffsetCoordinateSpace, float4 Rotation, int RotationCoordinateSpace, float Scale) {}
+#define MakeDrawCylinderPersistent(NAME)			void NAME(float3 Center, int CenterCoordinateSpace, float3 Offset, int OffsetCoordinateSpace, float Radius, float Height, bool IsHalfHeight, bool HemisphereX, bool HemisphereY, int OrientationAxis, int OrientationAxisCoordinateSpace, float3 NonUniformScale, int NumRadiusSegments, int NumHeightSegments, float4 Color) {}
 #define MakeDrawGrid2DPersistent(NAME)				void NAME(float3 Center, int CenterCoordinateSpace, float3 OffsetFromCenter, int OffsetCoordinateSpace, float4 Rotation, int RotationCoordinateSpace, float2 Extents, int NumCellsX, int NumCellsY, float4 Color) {}
 #define MakeDrawGrid3DPersistent(NAME)				void NAME(float3 Center, int CenterCoordinateSpace, float3 OffsetFromCenter, int OffsetCoordinateSpace, float4 Rotation, int RotationCoordinateSpace, float3 Extents, int NumCellsX, int NumCellsY, int NumCellsZ, float4 Color) {}
+#define MakeDrawLinePersistent(NAME)				void NAME(float3 StartLocation, int StartLocationCoordinateSpace, float3 EndLocation, int EndLocationCoordinateSpace, float4 Color) {}
+#define MakeDrawRectanglePersistent(NAME)			void NAME(float3 Center, int CenterCoordinateSpace, float3 Offset, int OffsetCoordinateSpace, float2 Extents, bool HalfExtents, float3 XAxis, float3 YAxis, int AxisVectorsCoordinateSpace, int NumXSegments, int NumYSegments, bool UnboundedPlane, float4 Color) {}
+#define MakeDrawSpherePersistent(NAME)				void NAME(float3 Center, int CenterCoordinateSpace, float3 OffsetFromCenter, int OffsetCoordinateSpace, float Radius, int RadiusCoordinateSpace, bool HemisphereX, bool HemisphereY, bool HemisphereZ, float3 SphereOrientationAxis, int RotationCoordinateSpace, float4 AdditionalRotation, float3 NonUniformScale, int NumSegments, float4 Color) {}
+#define MakeDrawTorusPersistent(NAME)				void NAME(float3 Center, int CenterCoordinateSpace, float3 OffsetFromCenter, int OffsetCoordinateSpace, float MajorRadius, float MinorRadius, float3 OrientationAxis, int OrientationAxisCoordinateSpace, float3 NonUniformScale, int MajorRadiusSegments, int MinorRadiusSegments, float4 Color) {}
