- Fix tooltip for Ribbon DrawDirection & warn about lack of support on GPU

diff --git a/Engine/Plugins/FX/Niagara/Source/Niagara/Private/NiagaraRibbonRendererProperties.cpp b/Engine/Plugins/FX/Niagara/Source/Niagara/Private/NiagaraRibbonRendererProperties.cpp
--- a/Engine/Plugins/FX/Niagara/Source/Niagara/Private/NiagaraRibbonRendererProperties.cpp
+++ b/Engine/Plugins/FX/Niagara/Source/Niagara/Private/NiagaraRibbonRendererProperties.cpp
@@ -689,6 +689,17 @@ void UNiagaraRibbonRendererProperties::GetRendererFeedback(const FVersionedNiaga
 
 		CheckUVSettingsForChannel(UV0Settings, 0);
 		CheckUVSettingsForChannel(UV1Settings, 1);
+
+		if (DrawDirection != ENiagaraRibbonDrawDirection::FrontToBack)
+		{
+			OutWarnings.Emplace(
+				LOCTEXT("GpuDrawDirectionNoSupportDesc", "Gpu ribbons only support the default Draw Direction for 'Front To Back'"),
+				LOCTEXT("GpuDrawDirectionNoSupportSummary", "Gpu ribbons do not support this Draw Direction mode it will be ignored"),
+				FText(),
+				FNiagaraRendererFeedbackFix(),
+				true
+			);
+		}
 	}
 
 	// If we're in multiplane shape, and multiplane count is even while we're in camera facing mode then one
diff --git a/Engine/Plugins/FX/Niagara/Source/Niagara/Public/NiagaraRibbonRendererProperties.h b/Engine/Plugins/FX/Niagara/Source/Niagara/Public/NiagaraRibbonRendererProperties.h
--- a/Engine/Plugins/FX/Niagara/Source/Niagara/Public/NiagaraRibbonRendererProperties.h
+++ b/Engine/Plugins/FX/Niagara/Source/Niagara/Public/NiagaraRibbonRendererProperties.h
@@ -286,12 +286,11 @@ private:
 #endif
 
 public:
-
 	UPROPERTY(EditAnywhere, Category = "Ribbon Rendering")	
 	int32 MaxNumRibbons;
 	
-	/** If true, the particles are only sorted when using a translucent material. */
-	UPROPERTY(EditAnywhere, Category = "Ribbon Rendering")
+	/** Controls the order the ribbon segments will be rendered. */
+	UPROPERTY(EditAnywhere, Category = "Sorting")
 	ENiagaraRibbonDrawDirection DrawDirection;
 
 	/** Shape of the ribbon, from flat plane, multiplane, 3d tube, and custom shapes. */
