- Mesh Renderer GDME optimization when using MeshIndex / Renderer Vis
- While below the cvar controlled instance count look over the mesh index & renderer vis information to determine what meshes are relevant
- Because the information is generally hot in cache from copying to GPU buffers, the cost for this is negligible, but can save an entire mesh submit

diff --git a/Engine/Plugins/FX/Niagara/Source/Niagara/Private/NiagaraRendererMeshes.cpp b/Engine/Plugins/FX/Niagara/Source/Niagara/Private/NiagaraRendererMeshes.cpp
--- a/Engine/Plugins/FX/Niagara/Source/Niagara/Private/NiagaraRendererMeshes.cpp
+++ b/Engine/Plugins/FX/Niagara/Source/Niagara/Private/NiagaraRendererMeshes.cpp
@@ -43,6 +43,14 @@ static FAutoConsoleVariableRef CVarEnableNiagaraMeshRendering(
 	ECVF_Default
 );
 
+static int32 GNiagaraMeshRendererCalcMeshUsedParticleCount = 64;
+static FAutoConsoleVariableRef CVarNiagaraMeshRendererCalcMeshUsedParticleCount(
+	TEXT("fx.Niagara.MeshRenderer.CalcMeshUsedParticleCount"),
+	GNiagaraMeshRendererCalcMeshUsedParticleCount,
+	TEXT("Optimization which will inspect renderer vis / mesh index to determine which set of meshes we can potential render with."),
+	ECVF_Default
+);
+
 #if RHI_RAYTRACING
 static TAutoConsoleVariable<int32> CVarRayTracingNiagaraMeshes(
 	TEXT("r.RayTracing.Geometry.NiagaraMeshes"),
@@ -212,6 +220,8 @@ void FNiagaraRendererMeshes::Initialize(const UNiagaraRendererProperties* InProp
 	// Initialize the valid mesh slots, and prep them with the data for every mesh, LOD, and section we'll be needing over the lifetime of the renderer
 	const uint32 MaxMeshes = Properties->Meshes.Num();
 	Meshes.Empty(MaxMeshes);
+	MeshUsedMax = 0;
+
 	for (uint32 SourceMeshIndex = 0; SourceMeshIndex < MaxMeshes; ++SourceMeshIndex)
 	{
 		const auto& MeshProperties = Properties->Meshes[SourceMeshIndex];
@@ -223,6 +233,8 @@ void FNiagaraRendererMeshes::Initialize(const UNiagaraRendererProperties* InProp
 			continue;
 		}
 
+		MeshUsedMax = SourceMeshIndex + 1;
+
 		FIntVector2 LODRange = FIntVector2(0, 1);
 		if (MeshProperties.LODMode == ENiagaraMeshLODMode::LODLevel)
 		{
@@ -461,6 +473,68 @@ void FNiagaraRendererMeshes::PrepareParticleMeshRenderData(FParticleMeshRenderDa
 	}
 }
 
+bool FNiagaraRendererMeshes::CalculateMeshUsed(FParticleMeshRenderData& ParticleMeshRenderData) const
+{
+	//-OPT: Should we handle modes where the tags are inside none-particle data?  i.e. EmitterRendererVisTagOffset | EmitterMeshIndexOffset
+	if (ParticleRendererVisTagOffset == INDEX_NONE && ParticleMeshIndexOffset == INDEX_NONE)
+	{
+		return true;
+	}
+
+	if (ParticleMeshRenderData.DynamicDataMesh->GetSimTarget() != ENiagaraSimTarget::CPUSim)
+	{
+		return true;
+	}
+
+	const int32 NumInstances = ParticleMeshRenderData.SourceParticleData->GetNumInstances();
+	if ( NumInstances > GNiagaraMeshRendererCalcMeshUsedParticleCount )
+	{
+		return true;
+	}
+
+	if (SourceMode != ENiagaraRendererSourceDataMode::Particles)
+	{
+		//-OPT: Should we handle emitter mode?
+		return true;
+	}
+
+	const FNiagaraDataBuffer* DataToRender = ParticleMeshRenderData.DynamicDataMesh->GetParticleDataToRender();
+	if (ParticleRendererVisTagOffset != INDEX_NONE)
+	{
+		const int32* RendererVisValues = reinterpret_cast<const int32*>(DataToRender->GetComponentPtrInt32(ParticleRendererVisTagOffset));
+		int32 FoundIndex = 0;
+		while (FoundIndex < NumInstances)
+		{
+			if (RendererVisValues[FoundIndex] == RendererVisibility)
+			{
+				break;
+			}
+			++FoundIndex;
+		}
+		if (FoundIndex == NumInstances)
+		{
+			return false;
+		}
+	}
+
+	ParticleMeshRenderData.MeshUsed.AddUninitialized(MeshUsedMax);
+	ParticleMeshRenderData.MeshUsed.SetRange(0, MeshUsedMax, false);
+	bool bAnyVisible = false;
+
+	const int32* MeshIndexValues = reinterpret_cast<const int32*>(DataToRender->GetComponentPtrInt32(ParticleMeshIndexOffset));
+	for (int32 i=0; i < NumInstances; ++i)
+	{
+		const int32 MeshIndex = MeshIndexValues[i];
+		if (MeshIndex >= 0 && MeshIndex < MeshUsedMax)
+		{
+			ParticleMeshRenderData.MeshUsed[MeshIndex] = true;
+			bAnyVisible = true;
+		}
+	}
+
+	return bAnyVisible;
+}
+
 void FNiagaraRendererMeshes::PrepareParticleRenderBuffers(FRHICommandListBase& RHICmdList, FParticleMeshRenderData& ParticleMeshRenderData, FGlobalDynamicReadBuffer& DynamicReadBuffer) const
 {
 	if ( SourceMode == ENiagaraRendererSourceDataMode::Particles )
@@ -1294,6 +1368,11 @@ void FNiagaraRendererMeshes::GetDynamicMeshElements(const TArray<const FSceneVie
 	FScopeCycleCounter EmitterStatsCounter(EmitterStatID);
 #endif
 
+	if (!CalculateMeshUsed(ParticleMeshRenderData))
+	{
+		return;
+	}
+
 	PrepareParticleRenderBuffers(RHICmdList, ParticleMeshRenderData, Collector.GetDynamicReadBuffer());
 
 	// If mesh index comes from the parameter store grab the information now
@@ -1356,6 +1435,11 @@ void FNiagaraRendererMeshes::GetDynamicMeshElements(const TArray<const FSceneVie
 					continue;
 				}
 
+				if (ParticleMeshRenderData.MeshUsed.Num() > 0 && !ParticleMeshRenderData.MeshUsed[Meshes[MeshIndex].SourceMeshIndex])
+				{
+					continue;
+				}
+
 				// If we can no perform sorting / culling on the GPU we need to skip meshes which are LODs
 				if (MeshData.LODMode == ENiagaraMeshLODMode::PerParticle && !ParticleMeshRenderData.bAllowPerParticleMeshLODs)
 				{
diff --git a/Engine/Plugins/FX/Niagara/Source/Niagara/Public/NiagaraRendererMeshes.h b/Engine/Plugins/FX/Niagara/Source/Niagara/Public/NiagaraRendererMeshes.h
--- a/Engine/Plugins/FX/Niagara/Source/Niagara/Public/NiagaraRendererMeshes.h
+++ b/Engine/Plugins/FX/Niagara/Source/Niagara/Public/NiagaraRendererMeshes.h
@@ -74,6 +74,7 @@ protected:
 
 		uint32							RendererVisTagOffset = INDEX_NONE;
 		uint32							MeshIndexOffset = INDEX_NONE;
+		TBitArray<TInlineAllocator<2>>	MeshUsed;
 
 		FVector							WorldSpacePivotOffset = FVector::ZeroVector;
 		FSphere							CullingSphere = FSphere(EForceInit::ForceInit);
@@ -127,6 +128,7 @@ protected:
 	};
 
 	NIAGARA_API void PrepareParticleMeshRenderData(FParticleMeshRenderData& ParticleMeshRenderData, const FSceneViewFamily& ViewFamily, FMeshElementCollector& Collector, FNiagaraDynamicDataBase* InDynamicData, const FNiagaraSceneProxy* SceneProxy, bool bRayTracing, ENiagaraGpuComputeTickStage::Type GpuReadyTickStage) const;
+	NIAGARA_API bool CalculateMeshUsed(FParticleMeshRenderData& ParticleMeshRenderData) const;
 	NIAGARA_API void PrepareParticleRenderBuffers(FRHICommandListBase& RHICmdList, FParticleMeshRenderData& ParticleMeshRenderData, FGlobalDynamicReadBuffer& DynamicReadBuffer) const;
 	NIAGARA_API void InitializeSortInfo(const FParticleMeshRenderData& ParticleMeshRenderData, const FNiagaraSceneProxy& SceneProxy, const FSceneView& View, int32 ViewIndex, bool bIsInstancedStereo, FNiagaraGPUSortInfo& OutSortInfo) const;
 	NIAGARA_API void PreparePerMeshData(FParticleMeshRenderData& ParticleMeshRenderData, const FNiagaraMeshVertexFactory& VertexFactory, const FNiagaraSceneProxy& SceneProxy, const FMeshData& MeshData) const;
@@ -169,6 +171,7 @@ protected:
 
 private:
 	TArray<FMeshData, TInlineAllocator<1>> Meshes;
+	int32 MeshUsedMax = 0;
 
 	ENiagaraRendererSourceDataMode SourceMode;
 	ENiagaraSortMode SortMode;
