- Fix warning about accessing a CVar that isn't flagged as ECVF_RenderThreadSafe

diff --git a/Engine/Plugins/FX/Niagara/Source/Niagara/Private/NiagaraCommon.cpp b/Engine/Plugins/FX/Niagara/Source/Niagara/Private/NiagaraCommon.cpp
--- a/Engine/Plugins/FX/Niagara/Source/Niagara/Private/NiagaraCommon.cpp
+++ b/Engine/Plugins/FX/Niagara/Source/Niagara/Private/NiagaraCommon.cpp
@@ -20,6 +20,7 @@
 #include "String/ParseTokens.h"
 #include "UObject/Class.h"
 #include "UObject/UObjectIterator.h"
+#include "FXSystem.h"
 
 DECLARE_CYCLE_STAT(TEXT("Niagara - Utilities - PrepareRapidIterationParameters"), STAT_Niagara_Utilities_PrepareRapidIterationParameters, STATGROUP_Niagara);
 
@@ -951,8 +952,7 @@ bool FNiagaraUtilities::AllowComputeShaders(EShaderPlatform ShaderPlatform)
 
 bool FNiagaraUtilities::AllowGPUSorting(EShaderPlatform ShaderPlatform)
 {
-	static const IConsoleVariable* AllowGPUSortingCVar = IConsoleManager::Get().FindConsoleVariable(TEXT("FX.AllowGPUSorting"));
-	return ensure(AllowGPUSortingCVar) && (AllowGPUSortingCVar->GetInt() != 0);
+	return FXConsoleVariables::bAllowGPUSorting != 0;
 }
 
 bool FNiagaraUtilities::AllowGPUCulling(EShaderPlatform ShaderPlatform)
diff --git a/Engine/Source/Runtime/Engine/Public/FXSystem.h b/Engine/Source/Runtime/Engine/Public/FXSystem.h
--- a/Engine/Source/Runtime/Engine/Public/FXSystem.h
+++ b/Engine/Source/Runtime/Engine/Public/FXSystem.h
@@ -63,7 +63,7 @@ namespace FXConsoleVariables
 	/** Visualize GPU particle simulation. */
 	extern int32 VisualizeGPUSimulation;
 	/** true if GPU emitters are permitted to sort. */
-	extern int32 bAllowGPUSorting;
+	ENGINE_API extern int32 bAllowGPUSorting;
 	/** true if emitters can be culled. */
 	extern int32 bAllowCulling;
 	/** true if GPU particle simulation is frozen. */
