- Fix for cull proxies and used materials
- We need to gather the list from the cull proxy not our controller as it could be different with regards to MIDs, etc
- Note: This change does not fix the fact the cull proxies do not map to MIDs correctly

diff --git a/Engine/Plugins/FX/Niagara/Source/Niagara/Private/NiagaraComponent.cpp b/Engine/Plugins/FX/Niagara/Source/Niagara/Private/NiagaraComponent.cpp
--- a/Engine/Plugins/FX/Niagara/Source/Niagara/Private/NiagaraComponent.cpp
+++ b/Engine/Plugins/FX/Niagara/Source/Niagara/Private/NiagaraComponent.cpp
@@ -2104,8 +2104,7 @@ void UNiagaraComponent::CreateCullProxy(bool bForce)
 		FNiagaraWorldManager* WorldMan = FNiagaraWorldManager::Get(GetWorld());
 		check(WorldMan);
 		
-		UNiagaraCullProxyComponent* SystemCullProxy = WorldMan->GetCullProxy(this);
-		if (SystemCullProxy)
+		if ( UNiagaraCullProxyComponent* SystemCullProxy = WorldMan->GetCullProxy(this) )
 		{
 			if (SystemCullProxy->RegisterCulledComponent(this, bForce))
 			{
@@ -2313,7 +2312,11 @@ FPrimitiveSceneProxy* UNiagaraComponent::CreateSceneProxy()
 
 void UNiagaraComponent::GetUsedMaterials(TArray<UMaterialInterface*>& OutMaterials, bool bGetDebugMaterials) const
 {
-	if (SystemInstanceController.IsValid())
+	if (CullProxy)
+	{
+		CullProxy->GetUsedMaterials(OutMaterials, bGetDebugMaterials);
+	}
+	else if (SystemInstanceController.IsValid())
 	{
 		SystemInstanceController->GetUsedMaterials(OutMaterials);
 	}
diff --git a/Engine/Plugins/FX/Niagara/Source/Niagara/Private/NiagaraCullProxyComponent.cpp b/Engine/Plugins/FX/Niagara/Source/Niagara/Private/NiagaraCullProxyComponent.cpp
--- a/Engine/Plugins/FX/Niagara/Source/Niagara/Private/NiagaraCullProxyComponent.cpp
+++ b/Engine/Plugins/FX/Niagara/Source/Niagara/Private/NiagaraCullProxyComponent.cpp
@@ -77,7 +77,7 @@ bool UNiagaraCullProxyComponent::RegisterCulledComponent(UNiagaraComponent* Comp
 
 		SetPaused(false);
 
-		Component->MarkRenderDynamicDataDirty();
+		Component->MarkRenderStateDirty();
 
 		return true;
 	}
@@ -98,7 +98,7 @@ void UNiagaraCullProxyComponent::UnregisterCulledComponent(UNiagaraComponent* Co
 		SetPaused(true);
 	}
 
-	Component->MarkRenderDynamicDataDirty();
+	Component->MarkRenderStateDirty();
 }
 
 //////////////////////////////////////////////////////////////////////////
