Added emitter properties as an item for summary view
Additionally refactored how individual properties are added to the stack via summary.

Instead of creating a property row, we create a stack object that has its root nodes filtered.
This allows stack object callbacks to be called as intended, which wasn't the case previously.
For example, when changing the material as an individual property in the summary, the system wouldn't restart due to missing callbacks

diff --git a/Engine/Plugins/FX/Niagara/Source/NiagaraEditor/Private/ViewModels/HierarchyEditor/NiagaraHierarchyViewModelBase.cpp b/Engine/Plugins/FX/Niagara/Source/NiagaraEditor/Private/ViewModels/HierarchyEditor/NiagaraHierarchyViewModelBase.cpp
--- a/Engine/Plugins/FX/Niagara/Source/NiagaraEditor/Private/ViewModels/HierarchyEditor/NiagaraHierarchyViewModelBase.cpp
+++ b/Engine/Plugins/FX/Niagara/Source/NiagaraEditor/Private/ViewModels/HierarchyEditor/NiagaraHierarchyViewModelBase.cpp
@@ -573,25 +573,6 @@ FString UNiagaraHierarchyViewModelBase::OnItemToStringDebug(TSharedPtr<FNiagaraH
 	return ItemBaseViewModel->ToString();	
 }
 
-const TArray<TSharedRef<IDetailTreeNode>>& UNiagaraHierarchyViewModelBase::RequestDetailTreeNodesForObject(UObject* Object)
-{
-	if(ObjectTreeNodeMap.Contains(Object))
-	{
-		return ObjectTreeNodeMap[Object];
-	}
-
-	if(ObjectToPropertyRowGeneratorMap.Contains(Object))
-	{
-		ObjectTreeNodeMap.Add(Object, ObjectToPropertyRowGeneratorMap[Object]->GetRootTreeNodes());
-		return ObjectTreeNodeMap[Object];
-	}
-	
-	FPropertyEditorModule& PropertyEditorModule = FModuleManager::Get().LoadModuleChecked<FPropertyEditorModule>("PropertyEditor");
-	ObjectToPropertyRowGeneratorMap.Add(Object, PropertyEditorModule.CreatePropertyRowGenerator(FPropertyRowGeneratorArgs()));
-	ObjectToPropertyRowGeneratorMap[Object]->SetObjects({Object});
-	return RequestDetailTreeNodesForObject(Object);
-}
-
 FNiagaraHierarchyItemViewModelBase::~FNiagaraHierarchyItemViewModelBase()
 {
 	Children.Empty();
diff --git a/Engine/Plugins/FX/Niagara/Source/NiagaraEditor/Private/ViewModels/HierarchyEditor/NiagaraSummaryViewViewModel.cpp b/Engine/Plugins/FX/Niagara/Source/NiagaraEditor/Private/ViewModels/HierarchyEditor/NiagaraSummaryViewViewModel.cpp
--- a/Engine/Plugins/FX/Niagara/Source/NiagaraEditor/Private/ViewModels/HierarchyEditor/NiagaraSummaryViewViewModel.cpp
+++ b/Engine/Plugins/FX/Niagara/Source/NiagaraEditor/Private/ViewModels/HierarchyEditor/NiagaraSummaryViewViewModel.cpp
@@ -47,6 +47,15 @@ void UNiagaraHierarchyAssignmentInput::Initialize(const UNiagaraNodeAssignment&
 	SetIdentity(InputIdentity);
 }
 
+void UNiagaraHierarchyEmitterProperties::Initialize(const FVersionedNiagaraEmitter& Emitter)
+{
+	FNiagaraHierarchyIdentity InputIdentity;
+	InputIdentity.Names.Add(FName(Emitter.Emitter->GetUniqueEmitterName()));
+	InputIdentity.Names.Add("Category");
+	InputIdentity.Names.Add("Properties");
+	SetIdentity(InputIdentity);
+}
+
 void UNiagaraHierarchyRenderer::Initialize(const UNiagaraRendererProperties& Renderer)
 {
 	FNiagaraHierarchyIdentity RendererIdentity;
@@ -901,6 +910,34 @@ const UNiagaraHierarchySection* FNiagaraHierarchyRendererViewModel::GetSectionIn
 	return Section.IsValid() ? Section.Get() : nullptr;
 }
 
+FString FNiagaraHierarchyEmitterPropertiesViewModel::ToString() const
+{
+	return TEXT("Emitter Properties");
+}
+
+bool FNiagaraHierarchyEmitterPropertiesViewModel::IsFromBaseEmitter() const
+{
+	if(!IsFromBaseEmitterCache.IsSet())
+	{
+		UNiagaraSummaryViewViewModel* ViewModel = Cast<UNiagaraSummaryViewViewModel>(GetHierarchyViewModel());
+		IsFromBaseEmitterCache = GetIsFromBaseEmitter(ViewModel->GetEmitterViewModel()->GetEmitter(), GetData()->GetPersistentIdentity());
+	}
+
+	return IsFromBaseEmitterCache.GetValue();
+}
+
+FNiagaraHierarchyItemViewModelBase::FCanPerformActionResults FNiagaraHierarchyEmitterPropertiesViewModel::IsEditableByUser()
+{
+	FCanPerformActionResults CanEditResults(IsFromBaseEmitter() == false);
+	CanEditResults.CanPerformMessage = CanEditResults.bCanPerform == false ? LOCTEXT("EmitterPropertiesIsFromBaseEmitter", "These emitter properties were added in the parent emitter and can not be edited.") : FText::GetEmpty();
+	return CanEditResults;
+}
+
+const UNiagaraHierarchySection* FNiagaraHierarchyEmitterPropertiesViewModel::GetSectionInternal() const
+{
+	return Section.IsValid() ? Section.Get() : nullptr;
+}
+
 FString FNiagaraHierarchyEventHandlerViewModel::ToString() const
 {
 	FNiagaraEventScriptProperties* ScriptProperties = GetEventScriptProperties();
@@ -1340,6 +1377,10 @@ TSharedPtr<FNiagaraHierarchyItemViewModelBase> UNiagaraSummaryViewViewModel::Cre
 	{
 		return MakeShared<FNiagaraHierarchyRendererViewModel>(Renderer, Parent, this, bIsForHierarchy);
 	}
+	else if(UNiagaraHierarchyEmitterProperties* EmitterProperties = Cast<UNiagaraHierarchyEmitterProperties>(ItemBase))
+	{
+		return MakeShared<FNiagaraHierarchyEmitterPropertiesViewModel>(EmitterProperties, Parent, this, bIsForHierarchy);
+	}
 	else if(UNiagaraHierarchyEventHandler* EventHandler = Cast<UNiagaraHierarchyEventHandler>(ItemBase))
 	{
 		return MakeShared<FNiagaraHierarchyEventHandlerViewModel>(EventHandler, Parent, this, bIsForHierarchy);
@@ -1392,6 +1433,30 @@ void UNiagaraSummaryViewViewModel::PrepareSourceItems(UNiagaraHierarchyRoot* Sou
 	TArray<UNiagaraNodeFunctionCall*> SimStageModules = FNiagaraStackGraphUtilities::GetAllSimStagesModuleNodes(EmitterViewModel.ToSharedRef());
 	TArray<UNiagaraNodeFunctionCall*> EventHandlerModules = FNiagaraStackGraphUtilities::GetAllEventHandlerModuleNodes(EmitterViewModel.ToSharedRef());
 
+	FNiagaraHierarchyIdentity EmitterPropertiesIdentity;
+	EmitterPropertiesIdentity.Guids.Add(EmitterViewModel->GetEmitter().Version);
+	EmitterPropertiesIdentity.Names.Add("Category");
+	EmitterPropertiesIdentity.Names.Add("Properties");
+
+	UNiagaraHierarchyEmitterProperties* EmitterProperties = nullptr;
+	
+	auto* FoundEmitterProperties = SourceRoot->GetChildrenMutable().FindByPredicate([EmitterPropertiesIdentity](UNiagaraHierarchyItemBase* ItemBase)
+		{
+			return ItemBase->GetPersistentIdentity() == EmitterPropertiesIdentity;
+		});
+
+	if(FoundEmitterProperties == nullptr)
+	{
+		EmitterProperties = SourceRoot->AddChild<UNiagaraHierarchyEmitterProperties>();
+		EmitterProperties->Initialize(EmitterViewModel->GetEmitter());
+	}
+	else
+	{
+		EmitterProperties = CastChecked<UNiagaraHierarchyEmitterProperties>(*FoundEmitterProperties);	
+	}
+
+	NewItems.Add(EmitterProperties);
+	
 	// We create hierarchy modules here. We attempt to maintain as many previous elements as possible in order to maintain UI state
 	TArray<UNiagaraNodeFunctionCall*> ModuleNodes = FNiagaraStackGraphUtilities::GetAllModuleNodes(EmitterViewModel.ToSharedRef());
 	for(UNiagaraNodeFunctionCall* ModuleNode : ModuleNodes)
diff --git a/Engine/Plugins/FX/Niagara/Source/NiagaraEditor/Private/ViewModels/Stack/NiagaraStackEventScriptItemGroup.cpp b/Engine/Plugins/FX/Niagara/Source/NiagaraEditor/Private/ViewModels/Stack/NiagaraStackEventScriptItemGroup.cpp
--- a/Engine/Plugins/FX/Niagara/Source/NiagaraEditor/Private/ViewModels/Stack/NiagaraStackEventScriptItemGroup.cpp
+++ b/Engine/Plugins/FX/Niagara/Source/NiagaraEditor/Private/ViewModels/Stack/NiagaraStackEventScriptItemGroup.cpp
@@ -121,7 +121,7 @@ void UNiagaraStackEventHandlerPropertiesItem::RefreshChildrenInternal(const TArr
 		EmitterObject->RegisterInstancedCustomPropertyTypeLayout(FNiagaraEventScriptProperties::StaticStruct()->GetFName(), 
 			FOnGetPropertyTypeCustomizationInstance::CreateStatic(&FNiagaraEventScriptPropertiesCustomization::MakeInstance, 
 			TWeakObjectPtr<UNiagaraSystem>(&GetSystemViewModel()->GetSystem()), GetEmitterViewModel()->GetEmitter().ToWeakPtr()));
-		EmitterObject->SetOnSelectRootNodes(UNiagaraStackObject::FOnSelectRootNodes::CreateUObject(this, &UNiagaraStackEventHandlerPropertiesItem::SelectEmitterStackObjectRootTreeNodes));
+		EmitterObject->SetOnGetCustomRootNodes(UNiagaraStackObject::FOnGetCustomRootNodes::CreateUObject(this, &UNiagaraStackEventHandlerPropertiesItem::SelectEmitterStackObjectRootTreeNodes));
 	}
 
 	FVersionedNiagaraEmitterData* EmitterData = GetEmitterViewModel()->GetEmitter().GetEmitterData();
diff --git a/Engine/Plugins/FX/Niagara/Source/NiagaraEditor/Private/ViewModels/Stack/NiagaraStackInputCategory.cpp b/Engine/Plugins/FX/Niagara/Source/NiagaraEditor/Private/ViewModels/Stack/NiagaraStackInputCategory.cpp
--- a/Engine/Plugins/FX/Niagara/Source/NiagaraEditor/Private/ViewModels/Stack/NiagaraStackInputCategory.cpp
+++ b/Engine/Plugins/FX/Niagara/Source/NiagaraEditor/Private/ViewModels/Stack/NiagaraStackInputCategory.cpp
@@ -11,9 +11,11 @@
 #include "NiagaraSimulationStageBase.h"
 #include "ViewModels/NiagaraEmitterViewModel.h"
 #include "ViewModels/HierarchyEditor/NiagaraSummaryViewViewModel.h"
+#include "ViewModels/Stack/NiagaraStackEmitterSettingsGroup.h"
 #include "ViewModels/Stack/NiagaraStackEventScriptItemGroup.h"
 #include "ViewModels/Stack/NiagaraStackFunctionInputCollection.h"
 #include "ViewModels/Stack/NiagaraStackModuleItem.h"
+#include "ViewModels/Stack/NiagaraStackObject.h"
 #include "ViewModels/Stack/NiagaraStackPropertyRow.h"
 #include "ViewModels/Stack/NiagaraStackRendererItem.h"
 #include "ViewModels/Stack/NiagaraStackSimulationStageGroup.h"
@@ -534,6 +536,19 @@ void UNiagaraStackSummaryCategory::RefreshChildrenInternal(const TArray<UNiagara
 				NewChildren.Add(StackRenderer);	
 			}
 		}
+		else if(UNiagaraHierarchyEmitterProperties* EmitterProperties = Cast<UNiagaraHierarchyEmitterProperties>(Data))
+		{
+			UNiagaraStackEmitterPropertiesItem* StackEmitterPropertiesItem = FindCurrentChildOfTypeByPredicate<UNiagaraStackEmitterPropertiesItem>(CurrentChildren,
+			[&](UNiagaraStackEmitterPropertiesItem* CurrentEmitterProperties) { return CurrentEmitterProperties->GetEmitterViewModel()->GetEmitter().Emitter->GetUniqueEmitterName() == EmitterProperties->GetPersistentIdentity().Names[0]; });
+	
+			if (StackEmitterPropertiesItem == nullptr)
+			{
+				StackEmitterPropertiesItem = NewObject<UNiagaraStackEmitterPropertiesItem>(this);
+				StackEmitterPropertiesItem->Initialize(CreateDefaultChildRequiredData());
+			}
+		
+			NewChildren.Add(StackEmitterPropertiesItem);	
+		}
 		else if(UNiagaraHierarchySimStage* SummarySimStage = Cast<UNiagaraHierarchySimStage>(Data))
 		{
 			TArray<UNiagaraSimulationStageBase*> SimStages = GetEmitterViewModel()->GetEmitter().GetEmitterData()->GetSimulationStages();
@@ -624,30 +639,51 @@ void UNiagaraStackSummaryCategory::RefreshChildrenInternal(const TArray<UNiagara
 		else if(const UNiagaraHierarchyObjectProperty* ObjectProperty = Cast<UNiagaraHierarchyObjectProperty>(Data))
 		{
 			TMap<FGuid, UObject*> ObjectsForProperties = GetEmitterViewModel()->GetSummaryHierarchyViewModel()->GetObjectsForProperties();
+			FGuid ObjectGuid = ObjectProperty->GetPersistentIdentity().Guids.Num() > 0 ? ObjectProperty->GetPersistentIdentity().Guids[0] : FGuid();
 			
-			FGuid OwnerGuid = ObjectProperty->GetPersistentIdentity().Guids.Num() > 0 ? ObjectProperty->GetPersistentIdentity().Guids[0] : FGuid();
-			
-			if(ObjectsForProperties.Contains(OwnerGuid))
+			if(ObjectsForProperties.Contains(ObjectGuid))
 			{
 				UObject* Object = ObjectsForProperties[ObjectProperty->GetPersistentIdentity().Guids[0]];
-				TArray<TSharedRef<IDetailTreeNode>> RootNodes = GetEmitterViewModel()->GetSummaryHierarchyViewModel()->RequestDetailTreeNodesForObject(Object);
-				
-				for(TSharedRef<IDetailTreeNode>& RootNode : RootNodes)
+
+				UNiagaraStackObject* StackObjectWithProperty = FindCurrentChildOfTypeByPredicate<UNiagaraStackObject>(CurrentChildren,
+				[&](UNiagaraStackObject* StackObjectCandidate)
 				{
-					TArray<TSharedRef<IDetailTreeNode>> ChildrenNodes;
-					RootNode->GetChildren(ChildrenNodes);
+					if(StackObjectCandidate->GetObject() != Object)
+					{
+						return false;
+					}
 
-					for(TSharedRef<IDetailTreeNode> ChildNode : ChildrenNodes)
+					if(StackObjectCandidate->GetCustomRootNodeNames().Num() == 1 && StackObjectCandidate->GetCustomRootNodeNames()[0] == ObjectProperty->GetPersistentIdentity().Names[0])
 					{
-						if(ChildNode->GetNodeName() == ObjectProperty->GetPersistentIdentity().Names[0])
+						return true;
+					}
+
+					return false;
+				});
+
+				if(StackObjectWithProperty == nullptr)
+				{
+					StackObjectWithProperty = NewObject<UNiagaraStackObject>(this);
+					StackObjectWithProperty->Initialize(CreateDefaultChildRequiredData(), Object, false, GetStackEditorDataKey(), nullptr);
+					StackObjectWithProperty->SetOnGetCustomRootNodes(UNiagaraStackObject::FOnGetCustomRootNodes::CreateLambda([PropertyName = ObjectProperty->GetPersistentIdentity().Names[0]](TArray<TSharedRef<IDetailTreeNode>> DefaultRootNodes, TArray<TSharedRef<IDetailTreeNode>>* Result)
+					{
+						for(TSharedRef<IDetailTreeNode>& RootNode : DefaultRootNodes)
 						{
-							UNiagaraStackPropertyRow* PropertyRow = NewObject<UNiagaraStackPropertyRow>(this);
-							PropertyRow->Initialize(CreateDefaultChildRequiredData(), ChildNode, false, GetOwnerStackItemEditorDataKey(), GetOwnerStackItemEditorDataKey(), nullptr);
-							PropertyRow->SetOwnerGuid(ObjectProperty->GetPersistentIdentity().Guids[0]);
-							NewChildren.Add(PropertyRow);
+							TArray<TSharedRef<IDetailTreeNode>> ChildrenNodes;
+							RootNode->GetChildren(ChildrenNodes);
+							
+							for(TSharedRef<IDetailTreeNode> ChildNode : ChildrenNodes)
+							{
+								if(ChildNode->GetNodeName() == PropertyName)
+								{
+									Result->Add(ChildNode);
+								}
+							}
 						}
-					}
+					}));
 				}
+				
+				NewChildren.Add(StackObjectWithProperty);
 			}
 		}
 	}
diff --git a/Engine/Plugins/FX/Niagara/Source/NiagaraEditor/Private/ViewModels/Stack/NiagaraStackObject.cpp b/Engine/Plugins/FX/Niagara/Source/NiagaraEditor/Private/ViewModels/Stack/NiagaraStackObject.cpp
--- a/Engine/Plugins/FX/Niagara/Source/NiagaraEditor/Private/ViewModels/Stack/NiagaraStackObject.cpp
+++ b/Engine/Plugins/FX/Niagara/Source/NiagaraEditor/Private/ViewModels/Stack/NiagaraStackObject.cpp
@@ -44,11 +44,23 @@ void UNiagaraStackObject::Initialize(FRequiredEntryData InRequiredEntryData, UOb
 	).BindUObject(this, &UNiagaraStackObject::OnMessageManagerRefresh);
 }
 
-void UNiagaraStackObject::SetOnSelectRootNodes(FOnSelectRootNodes OnSelectRootNodes)
+void UNiagaraStackObject::SetOnGetCustomRootNodes(FOnGetCustomRootNodes OnSelectRootNodes)
 {
 	OnSelectRootNodesDelegate = OnSelectRootNodes;
 }
 
+TArray<FName> UNiagaraStackObject::GetCustomRootNodeNames() const
+{
+	TArray<FName> Result;
+
+	Algo::Transform(CustomRootNodes, Result, [](TSharedRef<IDetailTreeNode> RootNote)
+	{
+		return RootNote->GetNodeName();
+	});
+
+	return Result;
+}
+
 void UNiagaraStackObject::RegisterInstancedCustomPropertyLayout(UStruct* Class, FOnGetDetailCustomizationInstance DetailLayoutDelegate)
 {
 	checkf(PropertyRowGenerator.IsValid() == false, TEXT("Can not add additional customizations after children have been refreshed."));
@@ -97,6 +109,8 @@ FNiagaraHierarchyIdentity UNiagaraStackObject::DetermineSummaryIdentity() const
 
 void UNiagaraStackObject::FinalizeInternal()
 {
+	CustomRootNodes.Empty();
+	
 	if (PropertyRowGenerator.IsValid())
 	{
 		PropertyRowGenerator->OnRowsRefreshed().RemoveAll(this);
@@ -290,6 +304,7 @@ void UNiagaraStackObject::RefreshChildrenInternal(const TArray<UNiagaraStackEntr
 		if (OnSelectRootNodesDelegate.IsBound())
 		{
 			OnSelectRootNodesDelegate.Execute(DefaultRootTreeNodes, &RootTreeNodes);
+			CustomRootNodes = RootTreeNodes;
 		}
 		else
 		{
diff --git a/Engine/Plugins/FX/Niagara/Source/NiagaraEditor/Private/ViewModels/Stack/NiagaraStackSummaryViewInputCollection.cpp b/Engine/Plugins/FX/Niagara/Source/NiagaraEditor/Private/ViewModels/Stack/NiagaraStackSummaryViewInputCollection.cpp
--- a/Engine/Plugins/FX/Niagara/Source/NiagaraEditor/Private/ViewModels/Stack/NiagaraStackSummaryViewInputCollection.cpp
+++ b/Engine/Plugins/FX/Niagara/Source/NiagaraEditor/Private/ViewModels/Stack/NiagaraStackSummaryViewInputCollection.cpp
@@ -13,11 +13,14 @@
 #include "ViewModels/NiagaraSystemViewModel.h"
 #include "ViewModels/NiagaraScriptViewModel.h"
 #include "ViewModels/HierarchyEditor/NiagaraSummaryViewViewModel.h"
+#include "ViewModels/Stack/NiagaraStackEmitterPropertiesGroup.h"
+#include "ViewModels/Stack/NiagaraStackEmitterSettingsGroup.h"
 #include "ViewModels/Stack/NiagaraStackEventScriptItemGroup.h"
 #include "ViewModels/Stack/NiagaraStackGraphUtilities.h"
 #include "ViewModels/Stack/NiagaraStackFunctionInputCollection.h"
 #include "ViewModels/Stack/NiagaraStackInputCategory.h"
 #include "ViewModels/Stack/NiagaraStackModuleItem.h"
+#include "ViewModels/Stack/NiagaraStackObject.h"
 #include "ViewModels/Stack/NiagaraStackPropertyRow.h"
 #include "ViewModels/Stack/NiagaraStackRendererItem.h"
 #include "ViewModels/Stack/NiagaraStackSimulationStageGroup.h"
@@ -333,6 +336,19 @@ void UNiagaraStackSummaryViewCollection::RefreshChildrenInternal(const TArray<UN
 				NewChildren.Add(StackRenderer);	
 			}
 		}
+		else if(UNiagaraHierarchyEmitterProperties* EmitterProperties = Cast<UNiagaraHierarchyEmitterProperties>(Data))
+		{
+			UNiagaraStackEmitterPropertiesItem* StackEmitterPropertiesItem = FindCurrentChildOfTypeByPredicate<UNiagaraStackEmitterPropertiesItem>(CurrentChildren,
+			[&](UNiagaraStackEmitterPropertiesItem* CurrentEmitterProperties) { return CurrentEmitterProperties->GetEmitterViewModel()->GetEmitter().Emitter->GetUniqueEmitterName() == EmitterProperties->GetPersistentIdentity().Names[0]; });
+	
+			if (StackEmitterPropertiesItem == nullptr)
+			{
+				StackEmitterPropertiesItem = NewObject<UNiagaraStackEmitterPropertiesItem>(this);
+				StackEmitterPropertiesItem->Initialize(CreateDefaultChildRequiredData());
+			}
+		
+			NewChildren.Add(StackEmitterPropertiesItem);	
+		}
 		else if(UNiagaraHierarchyEventHandler* EventHandler = Cast<UNiagaraHierarchyEventHandler>(Data))
 		{
 			const TArray<FNiagaraEventScriptProperties>& EventScriptProperties = GetEmitterViewModel()->GetEmitter().GetEmitterData()->GetEventHandlers();
@@ -474,23 +490,46 @@ void UNiagaraStackSummaryViewCollection::RefreshChildrenInternal(const TArray<UN
 			if(ObjectsForProperties.Contains(ObjectGuid))
 			{
 				UObject* Object = ObjectsForProperties[ObjectProperty->GetPersistentIdentity().Guids[0]];
-				TArray<TSharedRef<IDetailTreeNode>> RootNodes = GetEmitterViewModel()->GetSummaryHierarchyViewModel()->RequestDetailTreeNodesForObject(Object);
-				
-				for(TSharedRef<IDetailTreeNode>& RootNode : RootNodes)
+
+				UNiagaraStackObject* StackObjectWithProperty = FindCurrentChildOfTypeByPredicate<UNiagaraStackObject>(CurrentChildren,
+				[&](UNiagaraStackObject* StackObjectCandidate)
 				{
-					TArray<TSharedRef<IDetailTreeNode>> ChildrenNodes;
-					RootNode->GetChildren(ChildrenNodes);
+					if(StackObjectCandidate->GetObject() != Object)
+					{
+						return false;
+					}
+					
+					if(StackObjectCandidate->GetCustomRootNodeNames().Num() == 1 && StackObjectCandidate->GetCustomRootNodeNames()[0] == ObjectProperty->GetPersistentIdentity().Names[0])
+					{
+						return true;
+					}
+
+					return false;
+				});
 
-					for(TSharedRef<IDetailTreeNode> ChildNode : ChildrenNodes)
+				if(StackObjectWithProperty == nullptr)
+				{
+					StackObjectWithProperty = NewObject<UNiagaraStackObject>(this);
+					StackObjectWithProperty->Initialize(CreateDefaultChildRequiredData(), Object, false, GetStackEditorDataKey(), nullptr);
+					StackObjectWithProperty->SetOnGetCustomRootNodes(UNiagaraStackObject::FOnGetCustomRootNodes::CreateLambda([PropertyName = ObjectProperty->GetPersistentIdentity().Names[0]](TArray<TSharedRef<IDetailTreeNode>> DefaultRootNodes, TArray<TSharedRef<IDetailTreeNode>>* Result)
 					{
-						if(ChildNode->GetNodeName() == ObjectProperty->GetPersistentIdentity().Names[0])
+						for(TSharedRef<IDetailTreeNode>& RootNode : DefaultRootNodes)
 						{
-							UNiagaraStackPropertyRow* PropertyRow = NewObject<UNiagaraStackPropertyRow>(this);
-							PropertyRow->Initialize(CreateDefaultChildRequiredData(), ChildNode, false, GetOwnerStackItemEditorDataKey(), GetOwnerStackItemEditorDataKey(), nullptr);
-							NewChildren.Add(PropertyRow);
+							TArray<TSharedRef<IDetailTreeNode>> ChildrenNodes;
+							RootNode->GetChildren(ChildrenNodes);
+							
+							for(TSharedRef<IDetailTreeNode> ChildNode : ChildrenNodes)
+							{
+								if(ChildNode->GetNodeName() == PropertyName)
+								{
+									Result->Add(ChildNode);
+								}
+							}
 						}
-					}
+					}));
 				}
+				
+				NewChildren.Add(StackObjectWithProperty);
 			}
 		}
 	}
diff --git a/Engine/Plugins/FX/Niagara/Source/NiagaraEditor/Private/Widgets/SNiagaraHierarchy.cpp b/Engine/Plugins/FX/Niagara/Source/NiagaraEditor/Private/Widgets/SNiagaraHierarchy.cpp
--- a/Engine/Plugins/FX/Niagara/Source/NiagaraEditor/Private/Widgets/SNiagaraHierarchy.cpp
+++ b/Engine/Plugins/FX/Niagara/Source/NiagaraEditor/Private/Widgets/SNiagaraHierarchy.cpp
@@ -1155,10 +1155,7 @@ bool SNiagaraHierarchy::FilterForSourceSection(TSharedPtr<const FNiagaraHierarch
 {
 	if(ActiveSourceSection.IsValid())
 	{
-		if(const UNiagaraHierarchySection* Section = ItemViewModel->GetSection())
-		{
-			return GetActiveSourceSectionData() == Section;
-		}
+		return GetActiveSourceSectionData() == ItemViewModel->GetSection();
 	}
 
 	return true;
diff --git a/Engine/Plugins/FX/Niagara/Source/NiagaraEditor/Public/ViewModels/HierarchyEditor/NiagaraHierarchyViewModelBase.h b/Engine/Plugins/FX/Niagara/Source/NiagaraEditor/Public/ViewModels/HierarchyEditor/NiagaraHierarchyViewModelBase.h
--- a/Engine/Plugins/FX/Niagara/Source/NiagaraEditor/Public/ViewModels/HierarchyEditor/NiagaraHierarchyViewModelBase.h
+++ b/Engine/Plugins/FX/Niagara/Source/NiagaraEditor/Public/ViewModels/HierarchyEditor/NiagaraHierarchyViewModelBase.h
@@ -464,7 +464,6 @@ public:
 
 	FString OnItemToStringDebug(TSharedPtr<FNiagaraHierarchyItemViewModelBase> ItemBaseViewModel) const;
 
-	const TArray<TSharedRef<IDetailTreeNode>>& RequestDetailTreeNodesForObject(UObject* Object);
 protected:
 	virtual void InitializeInternal() {}
 	virtual void FinalizeInternal() {}
diff --git a/Engine/Plugins/FX/Niagara/Source/NiagaraEditor/Public/ViewModels/HierarchyEditor/NiagaraSummaryViewViewModel.h b/Engine/Plugins/FX/Niagara/Source/NiagaraEditor/Public/ViewModels/HierarchyEditor/NiagaraSummaryViewViewModel.h
--- a/Engine/Plugins/FX/Niagara/Source/NiagaraEditor/Public/ViewModels/HierarchyEditor/NiagaraSummaryViewViewModel.h
+++ b/Engine/Plugins/FX/Niagara/Source/NiagaraEditor/Public/ViewModels/HierarchyEditor/NiagaraSummaryViewViewModel.h
@@ -75,6 +75,14 @@ private:
 	FText TooltipOverride;
 };
 
+UCLASS(MinimalAPI)
+class UNiagaraHierarchyEmitterProperties : public UNiagaraHierarchyItem
+{
+	GENERATED_BODY()
+public:
+	void Initialize(const FVersionedNiagaraEmitter& Emitter);
+};
+
 UCLASS(MinimalAPI)
 class UNiagaraHierarchyRenderer : public UNiagaraHierarchyItem
 {
@@ -343,6 +351,31 @@ protected:
 	mutable TOptional<bool> IsFromBaseEmitterCache;
 };
 
+/** Emitter Properties currently don't list their individual properties since it's a mix of data of FVersionedNiagaraEmitterData & actual properties on the emitter which requires customization. */
+struct FNiagaraHierarchyEmitterPropertiesViewModel : public FNiagaraHierarchyItemViewModel
+{
+	FNiagaraHierarchyEmitterPropertiesViewModel(UNiagaraHierarchyEmitterProperties* EmitterProperties, TSharedPtr<FNiagaraHierarchyItemViewModelBase> InParent, TWeakObjectPtr<UNiagaraSummaryViewViewModel> ViewModel, bool bInIsForHierarchy)
+	: FNiagaraHierarchyItemViewModel(EmitterProperties, InParent, ViewModel, bInIsForHierarchy) {}
+
+	virtual FString ToString() const override;
+	virtual bool CanRenameInternal() override { return false; }
+
+	void SetSection(UNiagaraHierarchySection& InSection) { Section = &InSection; }
+
+	bool IsFromBaseEmitter() const;
+protected:	
+	virtual FCanPerformActionResults IsEditableByUser() override;
+	virtual bool CanHaveChildren() const override { return bIsForHierarchy == false; }
+	
+	virtual bool RepresentsExternalData() const override { return true; }
+	virtual bool DoesExternalDataStillExist(const UNiagaraHierarchyDataRefreshContext* Context) const override { return true; }
+
+	virtual const UNiagaraHierarchySection* GetSectionInternal() const override;
+	
+	TWeakObjectPtr<UNiagaraHierarchySection> Section;
+	mutable TOptional<bool> IsFromBaseEmitterCache;
+};
+
 struct FNiagaraHierarchyEventHandlerViewModel : public FNiagaraHierarchyItemViewModel
 {
 	FNiagaraHierarchyEventHandlerViewModel(UNiagaraHierarchyEventHandler* EventHandler, TSharedPtr<FNiagaraHierarchyItemViewModelBase> InParent, TWeakObjectPtr<UNiagaraSummaryViewViewModel> ViewModel, bool bInIsForHierarchy)
diff --git a/Engine/Plugins/FX/Niagara/Source/NiagaraEditor/Public/ViewModels/Stack/NiagaraStackObject.h b/Engine/Plugins/FX/Niagara/Source/NiagaraEditor/Public/ViewModels/Stack/NiagaraStackObject.h
--- a/Engine/Plugins/FX/Niagara/Source/NiagaraEditor/Public/ViewModels/Stack/NiagaraStackObject.h
+++ b/Engine/Plugins/FX/Niagara/Source/NiagaraEditor/Public/ViewModels/Stack/NiagaraStackObject.h
@@ -18,7 +18,7 @@ class UNiagaraStackObject : public UNiagaraStackItemContent, public FNotifyHook
 	GENERATED_BODY()
 		
 public:
-	DECLARE_DELEGATE_TwoParams(FOnSelectRootNodes, TArray<TSharedRef<IDetailTreeNode>>, TArray<TSharedRef<IDetailTreeNode>>*);
+	DECLARE_DELEGATE_TwoParams(FOnGetCustomRootNodes, TArray<TSharedRef<IDetailTreeNode>>, TArray<TSharedRef<IDetailTreeNode>>*);
 
 public:
 	NIAGARAEDITOR_API UNiagaraStackObject();
@@ -30,9 +30,10 @@ public:
 		@param InOwnerStackItemEditorDataKey The stack editor data key of the owning stack entry. 
 		@param InOwningNiagaraNode An optional niagara node which owns this object. */
 	NIAGARAEDITOR_API void Initialize(FRequiredEntryData InRequiredEntryData, UObject* InObject, bool bInIsTopLevelObject, FString InOwnerStackItemEditorDataKey, UNiagaraNode* InOwningNiagaraNode = nullptr);
-
-	NIAGARAEDITOR_API void SetOnSelectRootNodes(FOnSelectRootNodes OnSelectRootNodes);
-
+	
+	NIAGARAEDITOR_API void SetOnGetCustomRootNodes(FOnGetCustomRootNodes OnSelectRootNodes);
+	NIAGARAEDITOR_API TArray<FName> GetCustomRootNodeNames() const;
+	
 	NIAGARAEDITOR_API void RegisterInstancedCustomPropertyLayout(UStruct* Class, FOnGetDetailCustomizationInstance DetailLayoutDelegate);
 	NIAGARAEDITOR_API void RegisterInstancedCustomPropertyTypeLayout(FName PropertyTypeName, FOnGetPropertyTypeCustomizationInstance PropertyTypeLayoutDelegate, TSharedPtr<IPropertyTypeIdentifier> Identifier = nullptr);
 
@@ -79,7 +80,7 @@ private:
 	bool bIsTopLevelObject;
 
 	UNiagaraNode* OwningNiagaraNode;
-	FOnSelectRootNodes OnSelectRootNodesDelegate;
+	FOnGetCustomRootNodes OnSelectRootNodesDelegate;
 	TArray<FRegisteredClassCustomization> RegisteredClassCustomizations;
 	TArray<FRegisteredPropertyCustomization> RegisteredPropertyCustomizations;
 	TSharedPtr<IPropertyRowGenerator> PropertyRowGenerator;
@@ -87,6 +88,9 @@ private:
 
 	/** An optional object guid that can be provided to identify the object this stack entry represents. This can be used for summary view purposes and is also given to the child property rows. */
 	TOptional<FGuid> ObjectGuid;
+
+	/** We keep track of our custom root nodes, if available. This helps us determine identities of stack objects. */
+	TArray<TSharedRef<IDetailTreeNode>> CustomRootNodes;
 	
 	FGuid MessageManagerRegistrationKey;
 	TArray<FStackIssue> MessageManagerIssues;
