- Fix transitions issues with Index to ID table & readbacks

diff --git a/Engine/Plugins/FX/Niagara/Source/Niagara/Private/NiagaraDataSet.cpp b/Engine/Plugins/FX/Niagara/Source/Niagara/Private/NiagaraDataSet.cpp
--- a/Engine/Plugins/FX/Niagara/Source/Niagara/Private/NiagaraDataSet.cpp
+++ b/Engine/Plugins/FX/Niagara/Source/Niagara/Private/NiagaraDataSet.cpp
@@ -887,7 +887,7 @@ void FNiagaraDataBuffer::AllocateGPU(FRHICommandList& RHICmdList, uint32 InNumIn
 				NiagaraWarnGpuBufferSize(NumNeededElems * sizeof(int32), DebugSimName);
 				TStringBuilder<128> DebugBufferName;
 				DebugBufferName.Appendf(TEXT("NiagaraIDToIndexTable_%s_%p"), DebugSimName ? DebugSimName : TEXT(""), this);
-				GPUIDToIndexTable.Initialize(RHICmdList, DebugBufferName.ToString(), sizeof(int32), NumNeededElems, EPixelFormat::PF_R32_SINT, ERHIAccess::SRVCompute, BUF_Static);
+				GPUIDToIndexTable.Initialize(RHICmdList, DebugBufferName.ToString(), sizeof(int32), NumNeededElems, EPixelFormat::PF_R32_SINT, ERHIAccess::SRVCompute, GPUBufferFlags);
 			}
 		}
 	}
diff --git a/Engine/Plugins/FX/Niagara/Source/Niagara/Private/NiagaraDataSetReadback.cpp b/Engine/Plugins/FX/Niagara/Source/Niagara/Private/NiagaraDataSetReadback.cpp
--- a/Engine/Plugins/FX/Niagara/Source/Niagara/Private/NiagaraDataSetReadback.cpp
+++ b/Engine/Plugins/FX/Niagara/Source/Niagara/Private/NiagaraDataSetReadback.cpp
@@ -136,7 +136,11 @@ void FNiagaraDataSetReadback::GPUReadbackInternal(FRHICommandListImmediate& RHIC
 	Transitions.Emplace(ReadbackBuffers[0].Buffer, FNiagaraGPUInstanceCountManager::kCountBufferDefaultState, ERHIAccess::CopySrc);
 	for (int32 i=1; i < ReadbackBuffers.Num(); ++i)
 	{
-		Transitions.Emplace(ReadbackBuffers[i].Buffer, ERHIAccess::SRVMask, ERHIAccess::CopySrc);
+		Transitions.Emplace(
+			ReadbackBuffers[i].Buffer,
+			(i == IDtoIndexBufferIndex) ? ERHIAccess::SRVCompute : ERHIAccess::SRVMask,
+			ERHIAccess::CopySrc
+		);
 	}
 	RHICmdList.Transition(Transitions);
 
