- Fix lack of initial transition for dynamic mesh buffers

diff --git a/Engine/Plugins/FX/Niagara/Source/Niagara/Private/DataInterface/NiagaraDataInterfaceDynamicMesh.cpp b/Engine/Plugins/FX/Niagara/Source/Niagara/Private/DataInterface/NiagaraDataInterfaceDynamicMesh.cpp
--- a/Engine/Plugins/FX/Niagara/Source/Niagara/Private/DataInterface/NiagaraDataInterfaceDynamicMesh.cpp
+++ b/Engine/Plugins/FX/Niagara/Source/Niagara/Private/DataInterface/NiagaraDataInterfaceDynamicMesh.cpp
@@ -421,10 +421,10 @@ namespace NDIDynamicMeshLocal
 				IndexBufferDesc.Usage = (IndexBufferDesc.Usage & ~EBufferUsageFlags::VertexBuffer) | EBufferUsageFlags::IndexBuffer;
 
 				//-OPT: We don't need to create a SectionDefaultPooledBuffer if we aren't doing dynamic allocations
-				ResizeBufferIfNeeded(GraphBuilder, SectionDefaultPooledBuffer, EPixelFormat::PF_R32_UINT, FMath::Max(NumSections * 2u, 1u), TEXT("NiagaraDynamicMeshSectionBuffer"));
+				FRDGBufferRef SectionDefaultBuffer = ResizeBufferIfNeeded(GraphBuilder, SectionDefaultPooledBuffer, EPixelFormat::PF_R32_UINT, FMath::Max(NumSections * 2u, 1u), TEXT("NiagaraDynamicMeshSectionBuffer"));
 				ResizeBufferIfNeeded(GraphBuilder, SectionPooledBuffer, EPixelFormat::PF_R32_UINT, FMath::Max(NumSections * 2u, 1u), TEXT("NiagaraDynamicMeshSectionBuffer"));
-				ResizeBufferIfNeeded(GraphBuilder, IndexPooledBuffer, IndexBufferDesc, TEXT("NiagaraDynamicMeshIndexBuffer"));
-				ResizeBufferIfNeeded(GraphBuilder, VertexPooledBuffer, EPixelFormat::PF_R32_UINT, FMath::Max(GameToRenderData.VertexBufferSize >> 2u, 1u), TEXT("NiagaraDynamicMeshVertexBuffer"));
+				FRDGBufferRef RDGIndexBuffer = ResizeBufferIfNeeded(GraphBuilder, IndexPooledBuffer, IndexBufferDesc, TEXT("NiagaraDynamicMeshIndexBuffer"));
+				FRDGBufferRef RDGVertexBuffer = ResizeBufferIfNeeded(GraphBuilder, VertexPooledBuffer, EPixelFormat::PF_R32_UINT, FMath::Max(GameToRenderData.VertexBufferSize >> 2u, 1u), TEXT("NiagaraDynamicMeshVertexBuffer"));
 
 				MeshSections.SetNum(GameToRenderData.MeshSections.Num());
 				if (NumSections > 0)
@@ -451,19 +451,30 @@ namespace NDIDynamicMeshLocal
 						GpuMeshSections[i * 2 + 0] = InSection.MaxTriangles;
 						GpuMeshSections[i * 2 + 1] = InSection.AllocatedTriangles;
 					}
-					GraphBuilder.QueueBufferUpload(GraphBuilder.RegisterExternalBuffer(SectionDefaultPooledBuffer), GpuMeshSections, GpuMeshSectionsBytes, ERDGInitialDataFlags::NoCopy);
+					GraphBuilder.QueueBufferUpload(SectionDefaultBuffer, GpuMeshSections, GpuMeshSectionsBytes, ERDGInitialDataFlags::NoCopy);
 				}
 				//-OPT: We can remove the copying here.
 				if (GameToRenderData.IndexData.Num() > 0)
 				{
-					GraphBuilder.QueueBufferUpload(GraphBuilder.RegisterExternalBuffer(IndexPooledBuffer), GameToRenderData.IndexData.GetData(), GameToRenderData.IndexData.Num());
+					GraphBuilder.QueueBufferUpload(RDGIndexBuffer, GameToRenderData.IndexData.GetData(), GameToRenderData.IndexData.Num());
 				}
 
 				if (GameToRenderData.VertexData.Num() > 0)
 				{
-					GraphBuilder.QueueBufferUpload(GraphBuilder.RegisterExternalBuffer(VertexPooledBuffer), GameToRenderData.VertexData.GetData(), GameToRenderData.VertexData.Num());
+					GraphBuilder.QueueBufferUpload(RDGVertexBuffer, GameToRenderData.VertexData.GetData(), GameToRenderData.VertexData.Num());
 				}
 
+				FRDGExternalAccessQueue ExternalAccessQueue;
+				if (RDGIndexBuffer != nullptr)
+				{
+					ExternalAccessQueue.Add(RDGIndexBuffer, ERHIAccess::VertexOrIndexBuffer | ERHIAccess::SRVMask);
+				}
+				if (RDGVertexBuffer != nullptr)
+				{
+					ExternalAccessQueue.Add(RDGVertexBuffer, ERHIAccess::VertexOrIndexBuffer | ERHIAccess::SRVMask);
+				}
+				ExternalAccessQueue.Submit(GraphBuilder);
+
 				GraphBuilder.Execute();
 			}
 
