Fix Niagara vector field sampling via ispc supporting masked output (if some of the channels are unused)

diff --git a/Engine/Plugins/FX/Niagara/Source/Niagara/Private/NiagaraDataInterfaceVectorField.cpp b/Engine/Plugins/FX/Niagara/Source/Niagara/Private/NiagaraDataInterfaceVectorField.cpp
--- a/Engine/Plugins/FX/Niagara/Source/Niagara/Private/NiagaraDataInterfaceVectorField.cpp
+++ b/Engine/Plugins/FX/Niagara/Source/Niagara/Private/NiagaraDataInterfaceVectorField.cpp
@@ -589,12 +589,16 @@ void UNiagaraDataInterfaceVectorField::SampleVectorField(FVectorVMExternalFuncti
 #if INTEL_ISPC && VECTOR_FIELD_DATA_AS_HALF
 			if (GNiagaraVectorFieldUseIspc)
 			{
-				TConstArrayView<FFloat16> FieldSamples = StaticVectorField->ReadCPUData();
+				constexpr int32 DataSamplesPerVector = sizeof(ispc::FHalfVector) / sizeof(FFloat16);
+
+				TConstArrayView<FFloat16> FieldSamplesData = StaticVectorField->ReadCPUData();
+				TConstArrayView<ispc::FHalfVector> FieldSamples(reinterpret_cast<const ispc::FHalfVector*>(FieldSamplesData.GetData()), FieldSamplesData.Num() / DataSamplesPerVector);
 
 				ispc::SampleVectorField(XParam.GetDest(), YParam.GetDest(), ZParam.GetDest(),
 					XParam.IsConstant(), YParam.IsConstant(), ZParam.IsConstant(),
 					OutSampleX.GetDest(), OutSampleY.GetDest(), OutSampleZ.GetDest(),
-					(ispc::FHalfVector*) FieldSamples.GetData(), FieldSamples.Num() - sizeof(ispc::FHalfVector), (ispc::FVector&)MinBounds, (ispc::FVector&)OneOverBoundSize,
+					OutSampleX.IsValid(), OutSampleY.IsValid(), OutSampleZ.IsValid(),
+					FieldSamples.GetData(), FieldSamples.Num() - 1, (ispc::FVector&)MinBounds, (ispc::FVector&)OneOverBoundSize,
 					(ispc::FVector&)Size, (ispc::FVector&)TilingAxes, Context.GetNumInstances());
 			}
 			else
diff --git a/Engine/Plugins/FX/Niagara/Source/Niagara/Private/NiagaraDataInterfaceVectorField.ispc b/Engine/Plugins/FX/Niagara/Source/Niagara/Private/NiagaraDataInterfaceVectorField.ispc
--- a/Engine/Plugins/FX/Niagara/Source/Niagara/Private/NiagaraDataInterfaceVectorField.ispc
+++ b/Engine/Plugins/FX/Niagara/Source/Niagara/Private/NiagaraDataInterfaceVectorField.ispc
@@ -21,6 +21,7 @@ inline varying FVector ReadHalfVector(const uniform FHalfVector Data[], int Data
 export void SampleVectorField(const uniform float XParam[], const uniform float YParam[], const uniform float ZParam[],
 							const uniform bool bXParamConst, const uniform bool bYParamConst, const uniform bool bZParamConst,
 							uniform float OutSampleX[], uniform float OutSampleY[], uniform float OutSampleZ[],
+							const uniform bool bOutputXConst, const uniform bool bOutputYConst, const uniform bool bOutputZConst,
 							const uniform FHalfVector Data[], const uniform int DataSize, const uniform FVector &MinBounds, const uniform FVector &OneOverBoundSize,
 							const uniform FVector &Size, const uniform FVector &TilingAxes, const uniform int NumInstances)
 {
@@ -84,8 +85,19 @@ export void SampleVectorField(const uniform float XParam[], const uniform float
 		const FVector V = VectorLerp(V0, V1, Fraction.V[2]);
 
 		// Write final output...
-		OutSampleX[InstanceIdx] = V.V[0];
-		OutSampleY[InstanceIdx] = V.V[1];
-		OutSampleZ[InstanceIdx] = V.V[2];
+		if (bOutputXConst)
+		{
+			OutSampleX[InstanceIdx] = V.V[0];
+		}
+
+		if (bOutputYConst)
+		{
+			OutSampleY[InstanceIdx] = V.V[1];
+		}
+
+		if (bOutputZConst)
+		{
+			OutSampleZ[InstanceIdx] = V.V[2];
+		}
 	}
 }
