- If pooling was disabled for an asset with a pooled component in flight we can incorrectly warn about the pool being missing on completion

diff --git a/Engine/Plugins/FX/Niagara/Source/Niagara/Private/NiagaraComponentPool.cpp b/Engine/Plugins/FX/Niagara/Source/Niagara/Private/NiagaraComponentPool.cpp
--- a/Engine/Plugins/FX/Niagara/Source/Niagara/Private/NiagaraComponentPool.cpp
+++ b/Engine/Plugins/FX/Niagara/Source/Niagara/Private/NiagaraComponentPool.cpp
@@ -499,6 +499,8 @@ void UNiagaraComponentPool::ReclaimWorldParticleSystem(UNiagaraComponent* Compon
 	}
 	else
 	{
+		// We've stopped pooling while some effects were in flight so ensure they're destroyed now.
+		Component->PoolingMethod = ENCPoolMethod::None;
 		Component->DestroyComponent();
 	}
 }
