- Niagara Custom HLSL tab handling changes
- Make sure we don't loose focus when using tab in the text box
- Add support for Shift+Tab to trim whitespace / move cursor similar to other text editors

diff --git a/Engine/Plugins/FX/Niagara/Source/NiagaraEditor/Private/Widgets/SNiagaraGraphNodeCustomHlsl.cpp b/Engine/Plugins/FX/Niagara/Source/NiagaraEditor/Private/Widgets/SNiagaraGraphNodeCustomHlsl.cpp
--- a/Engine/Plugins/FX/Niagara/Source/NiagaraEditor/Private/Widgets/SNiagaraGraphNodeCustomHlsl.cpp
+++ b/Engine/Plugins/FX/Niagara/Source/NiagaraEditor/Private/Widgets/SNiagaraGraphNodeCustomHlsl.cpp
@@ -110,6 +110,7 @@ void SNiagaraGraphNodeCustomHlsl::CreateBelowPinControls(TSharedPtr<SVerticalBox
 		.Margin(FMargin(5, 5, 5, 5))
 		.Text(GetText)
 		.Marshaller(SyntaxHighlighter)
+		.OnKeyDownHandler(this, &SNiagaraGraphNodeCustomHlsl::OnKeyDown)
 		.OnKeyCharHandler(this, &SNiagaraGraphNodeCustomHlsl::OnShaderTextKeyChar)
 		.OnTextCommitted(TextCommit)
 		.Visibility(this, &SNiagaraGraphNodeCustomHlsl::GetShaderTextVisibility);
@@ -197,6 +198,59 @@ FString SNiagaraGraphNodeCustomHlsl::GetIdentifierUnderCursor() const
 	return FString();
 }
 
+FReply SNiagaraGraphNodeCustomHlsl::OnKeyDown(const FGeometry& MyGeometry, const FKeyEvent& InKeyEvent)
+{
+	if (InKeyEvent.GetKey() == EKeys::Tab)
+	{
+		constexpr int TabToSpaceLen = 4;
+
+		const FTextLocation CursorLocation = ShaderTextBox->GetCursorLocation();
+		if (InKeyEvent.IsShiftDown())
+		{
+			const FTextLocation NewCursorLocation(CursorLocation.GetLineIndex(), FMath::Max(AlignDown(CursorLocation.GetOffset() - 1, TabToSpaceLen), 0));
+
+			FString CurrentLine;
+			ShaderTextBox->GetCurrentTextLine(CurrentLine);
+
+			if (NewCursorLocation.GetOffset() != CursorLocation.GetOffset())
+			{
+				bool bTrimWhitespace = true;
+				for (int32 i = 0; i < CursorLocation.GetOffset() && bTrimWhitespace; ++i)
+				{
+					bTrimWhitespace = FChar::IsWhitespace(CurrentLine[i]);
+				}
+
+				if (bTrimWhitespace)
+				{
+					CurrentLine.RemoveAt(NewCursorLocation.GetOffset(), CursorLocation.GetOffset() - NewCursorLocation.GetOffset());
+
+					// We lack methods exposed to select or delete data from the ShaderTextBox so we rebuild all the text
+					TArray<FString> AllLines;
+					ShaderTextBox->GetText().ToString().ParseIntoArrayLines(AllLines, false);
+					AllLines[CursorLocation.GetLineIndex()] = CurrentLine;
+
+					const FString NewText = FString::Join(AllLines, TEXT("\n"));
+					ShaderTextBox->SetText(FText::FromString(NewText));
+				}
+			}
+
+			ShaderTextBox->GoTo(NewCursorLocation);
+
+			return FReply::Handled();
+		}
+		else
+		{
+			const FTextLocation NewCursorLocation(CursorLocation.GetLineIndex(), Align(CursorLocation.GetOffset() + 1, TabToSpaceLen));
+			const FString SpaceString = FString::ChrN(NewCursorLocation.GetOffset() - CursorLocation.GetOffset(), TEXT(' '));
+			ShaderTextBox->InsertTextAtCursor(SpaceString);
+			ShaderTextBox->GoTo(NewCursorLocation);
+			return FReply::Handled();
+		}
+	}
+
+	return FReply::Unhandled();
+}
+
 FReply SNiagaraGraphNodeCustomHlsl::OnShaderTextKeyChar(const FGeometry&, const FCharacterEvent& InCharacterEvent)
 {
 	const TCHAR Character = InCharacterEvent.GetCharacter();
@@ -223,12 +277,6 @@ FReply SNiagaraGraphNodeCustomHlsl::OnShaderTextKeyChar(const FGeometry&, const
 	MenuAnchor->SetIsOpen(false);
 	bIsAutoCompleteActive = false;
 	
-	if (Character == TEXT('\t'))
-	{
-		// Convert tab to four spaces
-		ShaderTextBox->InsertTextAtCursor(TEXT("    "));
-		return FReply::Handled();
-	}
 	if (Character == TEXT('.'))
 	{
 		// Show auto-complete hint for data interfaces
diff --git a/Engine/Plugins/FX/Niagara/Source/NiagaraEditor/Private/Widgets/SNiagaraGraphNodeCustomHlsl.h b/Engine/Plugins/FX/Niagara/Source/NiagaraEditor/Private/Widgets/SNiagaraGraphNodeCustomHlsl.h
--- a/Engine/Plugins/FX/Niagara/Source/NiagaraEditor/Private/Widgets/SNiagaraGraphNodeCustomHlsl.h
+++ b/Engine/Plugins/FX/Niagara/Source/NiagaraEditor/Private/Widgets/SNiagaraGraphNodeCustomHlsl.h
@@ -33,6 +33,7 @@ protected:
 	TSharedPtr<SMenuAnchor> MenuAnchor;
 
 private:
+	FReply OnKeyDown(const FGeometry& MyGeometry, const FKeyEvent& InKeyEvent);
 	FReply OnShaderTextKeyChar(const FGeometry& MyGeometry, const FCharacterEvent& InCharacterEvent);
 	FString GetIdentifierUnderCursor() const;
 
