Fix for dangling compilation event when tasks are aborted

diff --git a/Engine/Plugins/FX/Niagara/Source/NiagaraEditor/Private/NiagaraCompilationTasks.cpp b/Engine/Plugins/FX/Niagara/Source/NiagaraEditor/Private/NiagaraCompilationTasks.cpp
--- a/Engine/Plugins/FX/Niagara/Source/NiagaraEditor/Private/NiagaraCompilationTasks.cpp
+++ b/Engine/Plugins/FX/Niagara/Source/NiagaraEditor/Private/NiagaraCompilationTasks.cpp
@@ -342,6 +342,7 @@ FNiagaraSystemCompilationTask::FNiagaraSystemCompilationTask(FNiagaraCompilation
 void FNiagaraSystemCompilationTask::Abort()
 {
 	bAborting = true;
+	CompileCompletionEvent.Trigger();
 }
 
 
