Niagara - Fix a bug in DI resolving where DIs which aren't modified in the stack (placeholders) aren't correctly bound between scripts.

diff --git a/Engine/Plugins/FX/Niagara/Source/Niagara/Private/NiagaraResolveDIHelpers.cpp b/Engine/Plugins/FX/Niagara/Source/Niagara/Private/NiagaraResolveDIHelpers.cpp
--- a/Engine/Plugins/FX/Niagara/Source/Niagara/Private/NiagaraResolveDIHelpers.cpp
+++ b/Engine/Plugins/FX/Niagara/Source/Niagara/Private/NiagaraResolveDIHelpers.cpp
@@ -23,10 +23,11 @@ namespace FNiagaraResolveDIHelpers
 		FString SourceEmitterName;
 	};
 
-	void CollectDIBindingsAndAssignmentsForScript(
+	void CollectDIReadsAssignmentsAndBindingsForScript(
 		const UNiagaraScript* TargetScript,
 		const UNiagaraSystem* System,
 		const FString& EmitterName,
+		TMap<FNiagaraVariableBase, FDataInterfaceSourceEmitterNamePair>& OutVariableReadMap,
 		TMap<FNiagaraVariableBase, FDataInterfaceSourceEmitterNamePair>& OutVariableAssignmentMap,
 		TMap<FNiagaraVariableBase, FNiagaraVariableBase>& OutVariableBindingMap,
 		TArray<FText>& OutErrorMessages)
@@ -46,7 +47,7 @@ namespace FNiagaraResolveDIHelpers
 					WriteVariable = FNiagaraUtilities::ResolveAliases(WriteVariable, FNiagaraAliasContext().ChangeEmitterNameToEmitter(EmitterName));
 				}
 
-				// Handle bindings.
+				// Handle reads and bindings.
 				if (CachedDefaultDataInterface.RegisteredParameterMapRead != NAME_None)
 				{
 					FNiagaraVariable ReadVariable(CachedDefaultDataInterface.Type, CachedDefaultDataInterface.RegisteredParameterMapRead);
@@ -55,6 +56,15 @@ namespace FNiagaraResolveDIHelpers
 						ReadVariable = FNiagaraUtilities::ResolveAliases(ReadVariable, FNiagaraAliasContext().ChangeEmitterToEmitterName(EmitterName));
 					}
 
+					FDataInterfaceSourceEmitterNamePair* CurrentRead = OutVariableReadMap.Find(ReadVariable);
+					if (CurrentRead == nullptr)
+					{
+						FDataInterfaceSourceEmitterNamePair& NewRead = OutVariableReadMap.Add(ReadVariable);
+						NewRead.DataInterface = CachedDefaultDataInterface.DataInterface;
+						NewRead.SourceName = CachedDefaultDataInterface.Name;
+						NewRead.SourceEmitterName = CachedDefaultDataInterface.SourceEmitterName;
+					}
+
 					if (ReadVariable != WriteVariable)
 					{
 						FNiagaraVariableBase* CurrentBinding = OutVariableBindingMap.Find(WriteVariable);
@@ -102,12 +112,14 @@ namespace FNiagaraResolveDIHelpers
 		}
 	}
 
-	void CollectDIBindingsAndAssignments(
+	void CollectDIReadsAssignmentsAndBindings(
 		const UNiagaraSystem* System,
+		TMap<FGuid, TMap<FNiagaraVariableBase, FDataInterfaceSourceEmitterNamePair>>& OutEmitterIdToVariableReadsMap,
 		TMap<FGuid, TMap<FNiagaraVariableBase, FDataInterfaceSourceEmitterNamePair>>& OutEmitterIdToVariableAssignmentsMap,
 		TMap<FGuid, TMap<FNiagaraVariableBase, FNiagaraVariableBase>>& OutEmitterIdToVariableBindingsMap,
 		TArray<FText>& OutErrorMessages)
 	{
+		TMap<FNiagaraVariableBase, FDataInterfaceSourceEmitterNamePair> VariableReadMap;
 		TMap<FNiagaraVariableBase, FDataInterfaceSourceEmitterNamePair> VariableAssignmentMap;
 		TMap<FNiagaraVariableBase, FNiagaraVariableBase> VariableBindingMap;
 
@@ -128,8 +140,9 @@ namespace FNiagaraResolveDIHelpers
 		}
 
 		// Collect system and emitter scripts.
-		CollectDIBindingsAndAssignmentsForScript(System->GetSystemSpawnScript(), System, FString(), VariableAssignmentMap, VariableBindingMap, OutErrorMessages);
-		CollectDIBindingsAndAssignmentsForScript(System->GetSystemUpdateScript(), System, FString(), VariableAssignmentMap, VariableBindingMap, OutErrorMessages);
+		CollectDIReadsAssignmentsAndBindingsForScript(System->GetSystemSpawnScript(), System, FString(), VariableReadMap, VariableAssignmentMap, VariableBindingMap, OutErrorMessages);
+		CollectDIReadsAssignmentsAndBindingsForScript(System->GetSystemUpdateScript(), System, FString(), VariableReadMap, VariableAssignmentMap, VariableBindingMap, OutErrorMessages);
+		OutEmitterIdToVariableReadsMap.Add(FGuid(), VariableReadMap);
 		OutEmitterIdToVariableAssignmentsMap.Add(FGuid(), VariableAssignmentMap);
 		OutEmitterIdToVariableBindingsMap.Add(FGuid(), VariableBindingMap);
 
@@ -147,15 +160,16 @@ namespace FNiagaraResolveDIHelpers
 				continue;
 			}
 
+			TMap<FNiagaraVariableBase, FDataInterfaceSourceEmitterNamePair> EmitterVariableReadMap = VariableReadMap;
 			TMap<FNiagaraVariableBase, FDataInterfaceSourceEmitterNamePair> EmitterVariableAssignmentMap = VariableAssignmentMap;
 			TMap<FNiagaraVariableBase, FNiagaraVariableBase> EmitterVariableBindingMap = VariableBindingMap;
 
-			CollectDIBindingsAndAssignmentsForScript(VersionedNiagaraEmitterData->SpawnScriptProps.Script, System, EmitterHandle.GetUniqueInstanceName(), EmitterVariableAssignmentMap, EmitterVariableBindingMap, OutErrorMessages);
-			CollectDIBindingsAndAssignmentsForScript(VersionedNiagaraEmitterData->UpdateScriptProps.Script, System, EmitterHandle.GetUniqueInstanceName(), EmitterVariableAssignmentMap, EmitterVariableBindingMap, OutErrorMessages);
+			CollectDIReadsAssignmentsAndBindingsForScript(VersionedNiagaraEmitterData->SpawnScriptProps.Script, System, EmitterHandle.GetUniqueInstanceName(), EmitterVariableReadMap, EmitterVariableAssignmentMap, EmitterVariableBindingMap, OutErrorMessages);
+			CollectDIReadsAssignmentsAndBindingsForScript(VersionedNiagaraEmitterData->UpdateScriptProps.Script, System, EmitterHandle.GetUniqueInstanceName(), EmitterVariableReadMap, EmitterVariableAssignmentMap, EmitterVariableBindingMap, OutErrorMessages);
 
 			for (const FNiagaraEventScriptProperties& EventHandler : VersionedNiagaraEmitterData->GetEventHandlers())
 			{
-				CollectDIBindingsAndAssignmentsForScript(EventHandler.Script, System, EmitterHandle.GetUniqueInstanceName(), EmitterVariableAssignmentMap, EmitterVariableBindingMap, OutErrorMessages);
+				CollectDIReadsAssignmentsAndBindingsForScript(EventHandler.Script, System, EmitterHandle.GetUniqueInstanceName(), EmitterVariableReadMap, EmitterVariableAssignmentMap, EmitterVariableBindingMap, OutErrorMessages);
 			}
 
 			for (const UNiagaraSimulationStageBase* SimulationStage : VersionedNiagaraEmitterData->GetSimulationStages())
@@ -164,15 +178,16 @@ namespace FNiagaraResolveDIHelpers
 				{
 					continue;
 				}
-				CollectDIBindingsAndAssignmentsForScript(SimulationStage->Script, System, EmitterHandle.GetUniqueInstanceName(), EmitterVariableAssignmentMap, EmitterVariableBindingMap, OutErrorMessages);
+				CollectDIReadsAssignmentsAndBindingsForScript(SimulationStage->Script, System, EmitterHandle.GetUniqueInstanceName(), EmitterVariableReadMap, EmitterVariableAssignmentMap, EmitterVariableBindingMap, OutErrorMessages);
 			}
 
+			OutEmitterIdToVariableReadsMap.Add(EmitterHandle.GetId(), EmitterVariableReadMap);
 			OutEmitterIdToVariableAssignmentsMap.Add(EmitterHandle.GetId(), EmitterVariableAssignmentMap);
 			OutEmitterIdToVariableBindingsMap.Add(EmitterHandle.GetId(), EmitterVariableBindingMap);
 		}
 	}
 
-	void CompactBindings(const TMap<FNiagaraVariableBase, FNiagaraVariableBase> InVariableBindingMap, TMap<FNiagaraVariableBase, FNiagaraVariableBase>& OutCompactedVariableBindingMap, TArray<FText>& OutErrorMessages)
+	void CompactBindings(const TMap<FNiagaraVariableBase, FNiagaraVariableBase>& InVariableBindingMap, TMap<FNiagaraVariableBase, FNiagaraVariableBase>& OutCompactedVariableBindingMap, TArray<FText>& OutErrorMessages)
 	{
 		TArray<FNiagaraVariableBase> SourceVariables;
 		InVariableBindingMap.GetKeys(SourceVariables);
@@ -205,6 +220,7 @@ namespace FNiagaraResolveDIHelpers
 		const UNiagaraSystem* System,
 		UNiagaraScript* TargetScript,
 		const FString& EmitterName,
+		const TMap<FNiagaraVariableBase, FDataInterfaceSourceEmitterNamePair>& VariableReadMap,
 		const TMap<FNiagaraVariableBase, FDataInterfaceSourceEmitterNamePair>& VariableAssignmentMap,
 		const TMap<FNiagaraVariableBase, FNiagaraVariableBase>& CompactedVariableBindingMap,
 		TArray<FText>& OutErrorMessages)
@@ -223,7 +239,7 @@ namespace FNiagaraResolveDIHelpers
 
 			if (CachedDefaultDataInterface.RegisteredParameterMapRead != NAME_None)
 			{
-				// If this DI is read from a parameter try to resolve it through the binding and assignment maps.
+				// If this DI is read from a parameter try to resolve it through the binding and assignment or read maps.
 				FNiagaraVariable ReadVariable(CachedDefaultDataInterface.Type, CachedDefaultDataInterface.RegisteredParameterMapRead);
 				if (EmitterName.IsEmpty() == false)
 				{
@@ -236,14 +252,28 @@ namespace FNiagaraResolveDIHelpers
 					BoundVariable = &ReadVariable;
 				}
 
-				const FDataInterfaceSourceEmitterNamePair* DataInterfaceAssignment = VariableAssignmentMap.Find(*BoundVariable);
-				if (DataInterfaceAssignment != nullptr && ParameterStoreVariables.Contains(ReadVariable) == false)
+				const FDataInterfaceSourceEmitterNamePair* BoundDataInterface = VariableAssignmentMap.Find(*BoundVariable);
+				if (BoundDataInterface == nullptr)
+				{
+					// If the data interface was not written to explicitly, then it needs to use the default which was cached in the first read.
+					BoundDataInterface = VariableReadMap.Find(*BoundVariable);
+				}
+
+				if (BoundDataInterface == nullptr)
+				{
+					OutErrorMessages.Add(FText::Format(
+						LOCTEXT("MissingReadDataInterface", "A data interface was read from the parameter map, but its default could not be found.  The data interface used in the simulation may be incorrect.  Read Parameter: {0} Bound Parameter: {1}"),
+						FText::FromName(ReadVariable.GetName()),
+						FText::FromName(BoundVariable->GetName())));
+				}
+
+				if (BoundDataInterface != nullptr && ParameterStoreVariables.Contains(ReadVariable) == false)
 				{
 					ResolvedDataInterface.ResolvedVariable = *BoundVariable;
 					ResolvedDataInterface.ParameterStoreVariable = ReadVariable;
 					ResolvedDataInterface.bIsInternal = false;
-					ResolvedDataInterface.ResolvedDataInterface = DataInterfaceAssignment->DataInterface;
-					ResolvedDataInterface.ResolvedSourceEmitterName = DataInterfaceAssignment->SourceEmitterName;
+					ResolvedDataInterface.ResolvedDataInterface = BoundDataInterface->DataInterface;
+					ResolvedDataInterface.ResolvedSourceEmitterName = BoundDataInterface->SourceEmitterName;
 					ParameterStoreVariables.Add(ResolvedDataInterface.ParameterStoreVariable);
 				}
 
@@ -326,16 +356,18 @@ namespace FNiagaraResolveDIHelpers
 
 	void ResolveDIsInternal(
 		UNiagaraSystem* System,
+		const TMap<FGuid, TMap<FNiagaraVariableBase, FDataInterfaceSourceEmitterNamePair>>& EmitterIdToVariableReadsMap,
 		const TMap<FGuid, TMap<FNiagaraVariableBase, FDataInterfaceSourceEmitterNamePair>>& EmitterIdToVariableAssignmentsMap,
 		const TMap<FGuid, TMap<FNiagaraVariableBase, FNiagaraVariableBase>>& EmitterIdToVariableBindingsMap,
 		TArray<FText>& OutErrorMessages)
 	{
+		TMap<FNiagaraVariableBase, FDataInterfaceSourceEmitterNamePair> VariableReadMap = EmitterIdToVariableReadsMap[FGuid()];
 		TMap<FNiagaraVariableBase, FDataInterfaceSourceEmitterNamePair> VariableAssignmentMap = EmitterIdToVariableAssignmentsMap[FGuid()];
 		TMap<FNiagaraVariableBase, FNiagaraVariableBase> VariableBindingMap = EmitterIdToVariableBindingsMap[FGuid()];
 		TMap<FNiagaraVariableBase, FNiagaraVariableBase> CompactedVariableBindingMap;
 		CompactBindings(VariableBindingMap, CompactedVariableBindingMap, OutErrorMessages);
-		ResolveDIsForScript(System, System->GetSystemSpawnScript(), FString(), VariableAssignmentMap, CompactedVariableBindingMap, OutErrorMessages);
-		ResolveDIsForScript(System, System->GetSystemUpdateScript(), FString(), VariableAssignmentMap, CompactedVariableBindingMap, OutErrorMessages);
+		ResolveDIsForScript(System, System->GetSystemSpawnScript(), FString(), VariableReadMap, VariableAssignmentMap, CompactedVariableBindingMap, OutErrorMessages);
+		ResolveDIsForScript(System, System->GetSystemUpdateScript(), FString(), VariableReadMap, VariableAssignmentMap, CompactedVariableBindingMap, OutErrorMessages);
 
 		for (const FNiagaraEmitterHandle& EmitterHandle : System->GetEmitterHandles())
 		{
@@ -350,13 +382,14 @@ namespace FNiagaraResolveDIHelpers
 				continue;
 			}
 
+			TMap<FNiagaraVariableBase, FDataInterfaceSourceEmitterNamePair> EmitterVariableReadMap = EmitterIdToVariableReadsMap[EmitterHandle.GetId()];
 			TMap<FNiagaraVariableBase, FDataInterfaceSourceEmitterNamePair> EmitterVariableAssignmentMap = EmitterIdToVariableAssignmentsMap[EmitterHandle.GetId()];
 			TMap<FNiagaraVariableBase, FNiagaraVariableBase> EmitterVariableBindingMap = EmitterIdToVariableBindingsMap[EmitterHandle.GetId()];
 			TMap<FNiagaraVariableBase, FNiagaraVariableBase> EmitterCompactedVariableBindingMap;
 			CompactBindings(EmitterVariableBindingMap, EmitterCompactedVariableBindingMap, OutErrorMessages);
-			ResolveDIsForScript(System, VersionedNiagaraEmitterData->SpawnScriptProps.Script, EmitterHandle.GetUniqueInstanceName(), EmitterVariableAssignmentMap, EmitterCompactedVariableBindingMap, OutErrorMessages);
-			ResolveDIsForScript(System, VersionedNiagaraEmitterData->UpdateScriptProps.Script, EmitterHandle.GetUniqueInstanceName(), EmitterVariableAssignmentMap, EmitterCompactedVariableBindingMap, OutErrorMessages);
-			ResolveDIsForScript(System, VersionedNiagaraEmitterData->GetGPUComputeScript(), EmitterHandle.GetUniqueInstanceName(), EmitterVariableAssignmentMap, EmitterCompactedVariableBindingMap, OutErrorMessages);
+			ResolveDIsForScript(System, VersionedNiagaraEmitterData->SpawnScriptProps.Script, EmitterHandle.GetUniqueInstanceName(), EmitterVariableReadMap, EmitterVariableAssignmentMap, EmitterCompactedVariableBindingMap, OutErrorMessages);
+			ResolveDIsForScript(System, VersionedNiagaraEmitterData->UpdateScriptProps.Script, EmitterHandle.GetUniqueInstanceName(), EmitterVariableReadMap, EmitterVariableAssignmentMap, EmitterCompactedVariableBindingMap, OutErrorMessages);
+			ResolveDIsForScript(System, VersionedNiagaraEmitterData->GetGPUComputeScript(), EmitterHandle.GetUniqueInstanceName(), EmitterVariableReadMap, EmitterVariableAssignmentMap, EmitterCompactedVariableBindingMap, OutErrorMessages);
 
 			if (VersionedNiagaraEmitterData->bInterpolatedSpawning)
 			{
@@ -366,7 +399,7 @@ namespace FNiagaraResolveDIHelpers
 
 			for (const FNiagaraEventScriptProperties& EventHandler : VersionedNiagaraEmitterData->GetEventHandlers())
 			{
-				ResolveDIsForScript(System, EventHandler.Script, EmitterHandle.GetUniqueInstanceName(), EmitterVariableAssignmentMap, EmitterCompactedVariableBindingMap, OutErrorMessages);
+				ResolveDIsForScript(System, EventHandler.Script, EmitterHandle.GetUniqueInstanceName(), EmitterVariableReadMap, EmitterVariableAssignmentMap, EmitterCompactedVariableBindingMap, OutErrorMessages);
 			}
 
 			for (const UNiagaraSimulationStageBase* SimulationStage : VersionedNiagaraEmitterData->GetSimulationStages())
@@ -375,7 +408,7 @@ namespace FNiagaraResolveDIHelpers
 				{
 					continue;
 				}
-				ResolveDIsForScript(System, SimulationStage->Script, EmitterHandle.GetUniqueInstanceName(), EmitterVariableAssignmentMap, EmitterCompactedVariableBindingMap, OutErrorMessages);
+				ResolveDIsForScript(System, SimulationStage->Script, EmitterHandle.GetUniqueInstanceName(), EmitterVariableReadMap, EmitterVariableAssignmentMap, EmitterCompactedVariableBindingMap, OutErrorMessages);
 			}
 
 			if (VersionedNiagaraEmitterData->SimTarget == ENiagaraSimTarget::GPUComputeSim)
@@ -397,10 +430,11 @@ namespace FNiagaraResolveDIHelpers
 
 	void ResolveDIs(UNiagaraSystem* System, TArray<FText>& OutErrorMessages)
 	{
+		TMap<FGuid, TMap<FNiagaraVariableBase, FDataInterfaceSourceEmitterNamePair>> EmitterIdToVariableReadsMap;
 		TMap<FGuid, TMap<FNiagaraVariableBase, FDataInterfaceSourceEmitterNamePair>> EmitterIdToVariableAssignmentsMap;
 		TMap<FGuid, TMap<FNiagaraVariableBase, FNiagaraVariableBase>> EmitterIdToVariableBindingsMap;
-		FNiagaraResolveDIHelpers::CollectDIBindingsAndAssignments(System, EmitterIdToVariableAssignmentsMap, EmitterIdToVariableBindingsMap, OutErrorMessages);
-		FNiagaraResolveDIHelpers::ResolveDIsInternal(System, EmitterIdToVariableAssignmentsMap, EmitterIdToVariableBindingsMap, OutErrorMessages);
+		FNiagaraResolveDIHelpers::CollectDIReadsAssignmentsAndBindings(System, EmitterIdToVariableReadsMap, EmitterIdToVariableAssignmentsMap, EmitterIdToVariableBindingsMap, OutErrorMessages);
+		FNiagaraResolveDIHelpers::ResolveDIsInternal(System, EmitterIdToVariableReadsMap, EmitterIdToVariableAssignmentsMap, EmitterIdToVariableBindingsMap, OutErrorMessages);
 	}
 }
 
