- GDME Niagara shadow opts
- Early out if all views can never create a mesh batch & improve per view early out for shadow passes
- With 30 instances of a test asset this was ~3.0ms -> ~1.8ms for GDME.  We still had a couple of renderers casting shadows but the majority do not.

diff --git a/Engine/Plugins/FX/Niagara/Source/Niagara/Private/NiagaraRenderer.cpp b/Engine/Plugins/FX/Niagara/Source/Niagara/Private/NiagaraRenderer.cpp
--- a/Engine/Plugins/FX/Niagara/Source/Niagara/Private/NiagaraRenderer.cpp
+++ b/Engine/Plugins/FX/Niagara/Source/Niagara/Private/NiagaraRenderer.cpp
@@ -584,9 +584,23 @@ void FNiagaraRenderer::ProcessMaterialParameterBindings(const FNiagaraRendererMa
 	}
 }
 
-bool FNiagaraRenderer::IsViewRenderingOpaqueOnly(const FSceneView* View)
+bool FNiagaraRenderer::IsViewRenderingOpaqueOnly(const FSceneView* View, bool bCastsVolumetricTranslucentShadow)
 {
-	return View->bIsSceneCapture && View->SceneCaptureRenderTarget;
+	const bool bOpaqueSceneCapture = View->bIsSceneCapture && View->SceneCaptureRenderTarget;
+	const bool bShadowView = View->GetDynamicMeshElementsShadowCullFrustum() != nullptr;
+	return bOpaqueSceneCapture || (bShadowView && !bCastsVolumetricTranslucentShadow);
+}
+
+bool FNiagaraRenderer::AreViewsRenderingOpaqueOnly(const TArray<const FSceneView*>& Views, int32 ViewVisibilityMask, bool bCastsVolumetricTranslucentShadow)
+{
+	for ( int32 i=0; i < Views.Num(); ++i )
+	{
+		if (((ViewVisibilityMask & (1 << i)) != 0) && !IsViewRenderingOpaqueOnly(Views[i], bCastsVolumetricTranslucentShadow))
+		{
+			return false;
+		}
+	}
+	return true;
 }
 
 void FNiagaraRenderer::SortIndices(const FNiagaraGPUSortInfo& SortInfo, const FNiagaraRendererVariableInfo& SortVariable, const FNiagaraDataBuffer& Buffer, FGlobalDynamicReadBuffer::FAllocation& OutIndices)
diff --git a/Engine/Plugins/FX/Niagara/Source/Niagara/Private/NiagaraRendererMeshes.cpp b/Engine/Plugins/FX/Niagara/Source/Niagara/Private/NiagaraRendererMeshes.cpp
--- a/Engine/Plugins/FX/Niagara/Source/Niagara/Private/NiagaraRendererMeshes.cpp
+++ b/Engine/Plugins/FX/Niagara/Source/Niagara/Private/NiagaraRendererMeshes.cpp
@@ -353,6 +353,7 @@ void FNiagaraRendererMeshes::PrepareParticleMeshRenderData(FParticleMeshRenderDa
 		ViewFamilySupportLowLatencyTranslucency(ViewFamily);
 
 	ParticleMeshRenderData.bHasTranslucentMaterials = false;
+	ParticleMeshRenderData.bAllTranslucentMaterials = ParticleMeshRenderData.DynamicDataMesh->Materials.Num() > 0;
 	for (FMaterialRenderProxy* MaterialProxy : ParticleMeshRenderData.DynamicDataMesh->Materials)
 	{
 		check(MaterialProxy);
@@ -360,6 +361,7 @@ void FNiagaraRendererMeshes::PrepareParticleMeshRenderData(FParticleMeshRenderDa
 		const bool bTranslucent = IsTranslucentBlendMode(Material);
 
 		ParticleMeshRenderData.bHasTranslucentMaterials |= bTranslucent;
+		ParticleMeshRenderData.bAllTranslucentMaterials &= bTranslucent;
 
 		// If even one material can cause the mesh to render before FFXSystemInterface::PostRenderOpaque, we cannot use low latency data
 		ParticleMeshRenderData.bIsGpuLowLatencyTranslucency =
@@ -1369,6 +1371,11 @@ void FNiagaraRendererMeshes::GetDynamicMeshElements(const TArray<const FSceneVie
 		return;
 	}
 
+	if (ParticleMeshRenderData.bAllTranslucentMaterials && AreViewsRenderingOpaqueOnly(Views, VisibilityMap, SceneProxy->CastsVolumetricTranslucentShadow()))
+	{
+		return;
+	}
+
 #if STATS
 	FScopeCycleCounter EmitterStatsCounter(EmitterStatID);
 #endif
@@ -1516,16 +1523,9 @@ void FNiagaraRendererMeshes::GetDynamicMeshElements(const TArray<const FSceneVie
 						const bool bNeedsPrevTransform = !bTranslucent || Material.IsTranslucencyWritingVelocity();
 						if (bTranslucent)
 						{
-							if (bIsShadowView && !SceneProxy->CastsVolumetricTranslucentShadow())
-							{
-								// Don't add translucent mesh batches for shadows
-								// TODO: Need a way to know if it's a volumetric translucent shadow view to make this logic better
-								continue;
-							}
-
 							// If we are rendering opaque only we can skip this batch
 							//-OPT: If we only have opaque materials we can skip earlier however due to RemappedMaterialIndex potentially being invalid this is tricky
-							if (IsViewRenderingOpaqueOnly(View))
+							if (IsViewRenderingOpaqueOnly(View, SceneProxy->CastsVolumetricTranslucentShadow()))
 							{
 								continue;
 							}
diff --git a/Engine/Plugins/FX/Niagara/Source/Niagara/Private/NiagaraRendererRibbons.cpp b/Engine/Plugins/FX/Niagara/Source/Niagara/Private/NiagaraRendererRibbons.cpp
--- a/Engine/Plugins/FX/Niagara/Source/Niagara/Private/NiagaraRendererRibbons.cpp
+++ b/Engine/Plugins/FX/Niagara/Source/Niagara/Private/NiagaraRendererRibbons.cpp
@@ -821,6 +821,12 @@ void FNiagaraRendererRibbons::GetDynamicMeshElements(const TArray<const FSceneVi
 		}
 	}
 
+	const bool bTranslucentMaterial = DynamicData->Material && IsTranslucentBlendMode(DynamicData->Material->GetIncompleteMaterialWithFallback(FeatureLevel));
+	if (bTranslucentMaterial && AreViewsRenderingOpaqueOnly(Views, VisibilityMap, SceneProxy->CastsVolumetricTranslucentShadow()))
+	{
+		return;
+	}
+
 	FRHICommandListBase& RHICmdList = Collector.GetRHICommandList();
 
 #if STATS
@@ -850,7 +856,7 @@ void FNiagaraRendererRibbons::GetDynamicMeshElements(const TArray<const FSceneVi
 
 			// If we are rendering opaque only we can skip this batch
 			//-OPT: If we only have opaque materials we can skip earlier however due to RemappedMaterialIndex potentially being invalid this is tricky
-			if (IsViewRenderingOpaqueOnly(View) && DynamicData->Material && IsTranslucentBlendMode(DynamicData->Material->GetIncompleteMaterialWithFallback(FeatureLevel)))
+			if (bTranslucentMaterial && IsViewRenderingOpaqueOnly(View, SceneProxy->CastsVolumetricTranslucentShadow()))
 			{
 				continue;
 			}
diff --git a/Engine/Plugins/FX/Niagara/Source/Niagara/Private/NiagaraRendererSprites.cpp b/Engine/Plugins/FX/Niagara/Source/Niagara/Private/NiagaraRendererSprites.cpp
--- a/Engine/Plugins/FX/Niagara/Source/Niagara/Private/NiagaraRendererSprites.cpp
+++ b/Engine/Plugins/FX/Niagara/Source/Niagara/Private/NiagaraRendererSprites.cpp
@@ -961,6 +961,11 @@ void FNiagaraRendererSprites::GetDynamicMeshElements(const TArray<const FSceneVi
 		return;
 	}
 
+	if (ParticleSpriteRenderData.bHasTranslucentMaterials && AreViewsRenderingOpaqueOnly(Views, VisibilityMap, SceneProxy->CastsVolumetricTranslucentShadow()))
+	{
+		return;
+	}
+
 #if STATS
 	FScopeCycleCounter EmitterStatsCounter(EmitterStatID);
 #endif
@@ -979,7 +984,7 @@ void FNiagaraRendererSprites::GetDynamicMeshElements(const TArray<const FSceneVi
 			}
 
 			// Scene captures that are rendered in with the regular views will run depth only so we can skip building the batches
-			if (ParticleSpriteRenderData.bHasTranslucentMaterials && IsViewRenderingOpaqueOnly(View))
+			if (ParticleSpriteRenderData.bHasTranslucentMaterials && IsViewRenderingOpaqueOnly(View, SceneProxy->CastsVolumetricTranslucentShadow()))
 			{
 				continue;
 			}
diff --git a/Engine/Plugins/FX/Niagara/Source/Niagara/Public/NiagaraRenderer.h b/Engine/Plugins/FX/Niagara/Source/Niagara/Public/NiagaraRenderer.h
--- a/Engine/Plugins/FX/Niagara/Source/Niagara/Public/NiagaraRenderer.h
+++ b/Engine/Plugins/FX/Niagara/Source/Niagara/Public/NiagaraRenderer.h
@@ -95,7 +95,8 @@ public:
 	// Determine if we are rendering into an opaque only view
 	// This is an optimization for depth only scene capture which render in the main pass.
 	// These captures will have SceneCaptureRenderTarget set on the view, no other capture type does this.
-	static NIAGARA_API bool IsViewRenderingOpaqueOnly(const FSceneView* View);
+	static NIAGARA_API bool IsViewRenderingOpaqueOnly(const FSceneView* View, bool bCastsVolumetricTranslucentShadow);
+	static NIAGARA_API bool AreViewsRenderingOpaqueOnly(const TArray<const FSceneView*>& Views, int32 ViewVisibilityMask, bool bCastsVolumetricTranslucentShadow);
 
 	static NIAGARA_API void SortIndices(const struct FNiagaraGPUSortInfo& SortInfo, const FNiagaraRendererVariableInfo& SortVariable, const FNiagaraDataBuffer& Buffer, FGlobalDynamicReadBuffer::FAllocation& OutIndices);
 	static NIAGARA_API int32 SortAndCullIndices(const FNiagaraGPUSortInfo& SortInfo, const FNiagaraDataBuffer& Buffer, FGlobalDynamicReadBuffer::FAllocation& OutIndices);
diff --git a/Engine/Plugins/FX/Niagara/Source/Niagara/Public/NiagaraRendererMeshes.h b/Engine/Plugins/FX/Niagara/Source/Niagara/Public/NiagaraRendererMeshes.h
--- a/Engine/Plugins/FX/Niagara/Source/Niagara/Public/NiagaraRendererMeshes.h
+++ b/Engine/Plugins/FX/Niagara/Source/Niagara/Public/NiagaraRendererMeshes.h
@@ -49,6 +49,7 @@ protected:
 
 		bool 							bUseGPUScene = false;
 		bool							bHasTranslucentMaterials = false;
+		bool							bAllTranslucentMaterials = false;
 		bool							bSortCullOnGpu = false;
 		bool							bNeedsSort = false;
 		bool							bNeedsCull = false;
