Niagara - Disable new data interface warning while investigating issues.

diff --git a/Engine/Plugins/FX/Niagara/Source/Niagara/Private/NiagaraResolveDIHelpers.cpp b/Engine/Plugins/FX/Niagara/Source/Niagara/Private/NiagaraResolveDIHelpers.cpp
--- a/Engine/Plugins/FX/Niagara/Source/Niagara/Private/NiagaraResolveDIHelpers.cpp
+++ b/Engine/Plugins/FX/Niagara/Source/Niagara/Private/NiagaraResolveDIHelpers.cpp
@@ -259,6 +259,7 @@ namespace FNiagaraResolveDIHelpers
 					BoundDataInterface = VariableReadMap.Find(*BoundVariable);
 				}
 
+				/*
 				if (BoundDataInterface == nullptr)
 				{
 					OutErrorMessages.Add(FText::Format(
@@ -266,6 +267,7 @@ namespace FNiagaraResolveDIHelpers
 						FText::FromName(ReadVariable.GetName()),
 						FText::FromName(BoundVariable->GetName())));
 				}
+				*/
 
 				if (BoundDataInterface != nullptr && ParameterStoreVariables.Contains(ReadVariable) == false)
 				{
