Disabled bHasDeformableMesh for FNiagaraSceneProxy to avoid continuous invalidation from the primitive bounds.

diff --git a/Engine/Plugins/FX/Niagara/Source/Niagara/Private/NiagaraComponent.cpp b/Engine/Plugins/FX/Niagara/Source/Niagara/Private/NiagaraComponent.cpp
--- a/Engine/Plugins/FX/Niagara/Source/Niagara/Private/NiagaraComponent.cpp
+++ b/Engine/Plugins/FX/Niagara/Source/Niagara/Private/NiagaraComponent.cpp
@@ -208,6 +208,8 @@ FNiagaraSceneProxy::FNiagaraSceneProxy(UNiagaraComponent* InComponent)
 #endif
 	}
 
+	// Prevent continuous VSM invalidation from the bounds of the primitive.
+	bHasDeformableMesh = false;
 	// Niagara renderers reference a lot of common contexts that aren't locked (and would otherwise introduce a lot of contention).
 	bSupportsParallelGDME = false;
 }
