Improve attribute trimming when dealing with impure functions in non-expressions

diff --git a/Engine/Plugins/FX/Niagara/Source/NiagaraEditor/Private/NiagaraAttributeTrimmer.cpp b/Engine/Plugins/FX/Niagara/Source/NiagaraEditor/Private/NiagaraAttributeTrimmer.cpp
--- a/Engine/Plugins/FX/Niagara/Source/NiagaraEditor/Private/NiagaraAttributeTrimmer.cpp
+++ b/Engine/Plugins/FX/Niagara/Source/NiagaraEditor/Private/NiagaraAttributeTrimmer.cpp
@@ -93,7 +93,8 @@ template<typename GraphBridge>
 class FNiagaraAttributeTrimmerHelper<GraphBridge>::FImpureFunctionParser
 {
 public:
-	FImpureFunctionParser(const FGraph* Graph, ENiagaraScriptUsage Usage)
+	FImpureFunctionParser(const FGraph* Graph, ENiagaraScriptUsage Usage, FTrimAttributeCache& InCache)
+		: Cache(InCache)
 	{
 		Traverse(TEXT(""), Graph, Usage);
 	}
@@ -125,7 +126,20 @@ private:
 						Traverse(NamespacePrefix + GraphBridge::GetFunctionName(FunctionNode), FunctionGraph, GraphBridge::GetFunctionUsage(FunctionNode));
 					}
 
+					bool bEvaluatePins = false;
 					if (FunctionNode->Signature.bRequiresExecPin)
+					{
+						bEvaluatePins = true;
+					}
+					else if(const FCustomHlslNode* CustomHlslNode = GraphBridge::AsCustomHlslNode(Node))
+					{
+						FCustomHlslNodeInfo NodeInfo;
+						BuildCustomHlslNodeInfo(Cache, CustomHlslNode, NodeInfo);
+
+						bEvaluatePins = NodeInfo.bHasImpureFunctionText;
+					}
+
+					if (bEvaluatePins)
 					{
 						for (const FInputPin* Pin : GraphBridge::GetInputPins(FunctionNode))
 						{
@@ -149,6 +163,7 @@ private:
 		}
 	}
 
+	FTrimAttributeCache& Cache;
 	TArray<FModuleScopedPin> ImpureFunctionInputs;
 };
 
@@ -157,9 +172,10 @@ class FNiagaraAttributeTrimmerHelper<GraphBridge>::FExpressionBuilder
 {
 public:
 
-	FExpressionBuilder(const FParamMapHistory& InParamMap, const FFunctionInputResolver& InInputResolver)
+	FExpressionBuilder(const FParamMapHistory& InParamMap, const FFunctionInputResolver& InInputResolver, FTrimAttributeCache& InCache)
 		: ParamMap(InParamMap)
 		, InputResolver(InInputResolver)
+		, Cache(InCache)
 	{}
 
 	// generates a set of dependencies for a specific pin
@@ -294,7 +310,7 @@ private:
 						if (!Dependencies.CustomNodes.Contains(CustomHlslNode))
 						{
 							FCustomHlslNodeInfo& NodeInfo = Dependencies.CustomNodes.Add(CustomHlslNode);
-							BuildCustomHlslNodeInfo(CustomHlslNode, NodeInfo);
+							BuildCustomHlslNodeInfo(Cache, CustomHlslNode, NodeInfo);
 						}
 					}
 				}
@@ -311,64 +327,31 @@ private:
 		}
 	}
 
-	void BuildCustomHlslNodeInfo(const FCustomHlslNode* CustomNode, FCustomHlslNodeInfo& NodeInfo)
-	{
-		check(CustomNode);
-
-		NodeInfo.bHasDataInterfaceInputs = false;
-		NodeInfo.bHasImpureFunctionText = false;
-
-		TArray<FString> ImpureFunctionNames;
-
-		for (const FInputPin* InputPin : GraphBridge::GetInputPins(CustomNode))
-		{
-			FNiagaraTypeDefinition NiagaraType = GraphBridge::GetPinType(InputPin, ENiagaraStructConversion::Simulation);
-			if (NiagaraType.IsDataInterface())
-			{
-				NodeInfo.bHasDataInterfaceInputs = true;
-
-				if (UNiagaraDataInterface* DataInterfaceClass = CastChecked<UNiagaraDataInterface>(NiagaraType.GetClass()->ClassDefaultObject))
-				{
-					TArray<FNiagaraFunctionSignature> FunctionSignatures;
-					DataInterfaceClass->GetFunctions(FunctionSignatures);
-
-					for (const FNiagaraFunctionSignature& FunctionSignature : FunctionSignatures)
-					{
-						if (FunctionSignature.bRequiresExecPin)
-						{
-							TStringBuilder<256> Builder;
-							InputPin->PinName.AppendString(Builder);
-							Builder.AppendChar(TCHAR('.'));
-							FunctionSignature.Name.AppendString(Builder);
+	const FParamMapHistory& ParamMap;
+	const FFunctionInputResolver& InputResolver;
+	FTrimAttributeCache& Cache;
 
-							ImpureFunctionNames.AddUnique(Builder.ToString());
-						}
-					}
-				}
-			}
-		}
+	const bool EvaluateStaticSwitches = false;
+};
 
-		if (!ImpureFunctionNames.IsEmpty())
-		{
-			TArray<FStringView> ImpureFunctionNameViews;
-			ImpureFunctionNameViews.Reserve(ImpureFunctionNames.Num());
-			Algo::Transform(ImpureFunctionNames, ImpureFunctionNameViews, [](const FString& FunctionName) -> FStringView
-				{
-					return FunctionName;
-				});
+template<typename GraphBridge>
+void FNiagaraAttributeTrimmerHelper<GraphBridge>::BuildCustomHlslNodeInfo(FTrimAttributeCache& InCache, const FCustomHlslNode* CustomNode, FCustomHlslNodeInfo& NodeInfo)
+{
+	if (CustomNode == nullptr)
+	{
+		return;
+	}
 
-			if (GraphBridge::CustomHlslReferencesTokens(CustomNode, ImpureFunctionNameViews))
-			{
-				return;
-			}
-		}
+	if (FCustomHlslNodeInfo* ExistingNodeInfo = InCache.CustomNodeCache.Find(CustomNode))
+	{
+		NodeInfo = *ExistingNodeInfo;
+		return;
 	}
 
-	const FParamMapHistory& ParamMap;
-	const FFunctionInputResolver& InputResolver;
+	NodeInfo.bHasImpureFunctionText = GraphBridge::GetCustomNodeUsesImpureFunctions(CustomNode);
 
-	const bool EvaluateStaticSwitches = false;
-};
+	InCache.CustomNodeCache.Add(CustomNode, NodeInfo);
+}
 
 // For a specific read of a variable finds the corresponding PreviousWritePin if one exists (only considers actual writes
 // rather than default pins on a MapGet)
@@ -613,6 +596,8 @@ void FNiagaraAttributeTrimmerHelper<GraphBridge>::TrimAttributes_Safe(TConstArra
 template<typename GraphBridge>
 void FNiagaraAttributeTrimmerHelper<GraphBridge>::TrimAttributes_Aggressive(const FCompilationCopy* CompileDuplicateData, TConstArrayView<const FParamMapHistory*> LocalParamHistories, TSet<FName>& AttributesToPreserve, TArray<FNiagaraVariable>& Attributes)
 {
+	FTrimAttributeCache Cache;
+
 	// variable references hidden in custom hlsl nodes may not be present in a specific stages parameter map
 	// so we consolidate all the variables into one list to use when going through custom hlsl nodes
 	TSet<FNiagaraVariableBase> UnifiedVariables;
@@ -641,7 +626,7 @@ void FNiagaraAttributeTrimmerHelper<GraphBridge>::TrimAttributes_Aggressive(cons
 			{
 				FDependencyChain& Dependencies = PerVariableDependencySets.Add(WritePin);
 
-				FExpressionBuilder Builder(*ParamMap, InputResolver);
+				FExpressionBuilder Builder(*ParamMap, InputResolver, Cache);
 				Builder.FindDependencies(WritePin, Dependencies);
 
 				for (typename FCustomHlslNodeMap::TConstIterator CustomNodeIt = Dependencies.CustomNodes.CreateConstIterator(); CustomNodeIt; ++CustomNodeIt)
@@ -654,12 +639,12 @@ void FNiagaraAttributeTrimmerHelper<GraphBridge>::TrimAttributes_Aggressive(cons
 			}
 		}
 
-		FImpureFunctionParser ImpureFunctionParser(GraphBridge::GetGraph(CompileDuplicateData), ParamMap->OriginatingScriptUsage);
+		FImpureFunctionParser ImpureFunctionParser(GraphBridge::GetGraph(CompileDuplicateData), ParamMap->OriginatingScriptUsage, Cache);
 		for (const FModuleScopedPin& RequiredFunctionInput : ImpureFunctionParser.ReadFunctionInputs())
 		{
 			FDependencyChain Dependencies;
 
-			FExpressionBuilder Builder(*ParamMap, InputResolver);
+			FExpressionBuilder Builder(*ParamMap, InputResolver, Cache);
 			Builder.FindDependencies(RequiredFunctionInput, Dependencies);
 
 			for (const FModuleScopedPin& DependentPin : Dependencies.Pins)
diff --git a/Engine/Plugins/FX/Niagara/Source/NiagaraEditor/Private/NiagaraAttributeTrimmer.h b/Engine/Plugins/FX/Niagara/Source/NiagaraEditor/Private/NiagaraAttributeTrimmer.h
--- a/Engine/Plugins/FX/Niagara/Source/NiagaraEditor/Private/NiagaraAttributeTrimmer.h
+++ b/Engine/Plugins/FX/Niagara/Source/NiagaraEditor/Private/NiagaraAttributeTrimmer.h
@@ -26,7 +26,6 @@ class FNiagaraAttributeTrimmerHelper : public GraphBridge
 	using FDependencySet = TSet<FModuleScopedPin>;
 	struct FCustomHlslNodeInfo
 	{
-		bool bHasDataInterfaceInputs = false;
 		bool bHasImpureFunctionText = false;
 	};
 
@@ -38,6 +37,11 @@ class FNiagaraAttributeTrimmerHelper : public GraphBridge
 		FCustomHlslNodeMap CustomNodes;
 	};
 
+	struct FTrimAttributeCache
+	{
+		FCustomHlslNodeMap CustomNodeCache;
+	};
+
 	using FDependencyMap = TMap<FModuleScopedPin, FDependencyChain>;
 
 	class FFunctionInputResolver;
@@ -55,6 +59,8 @@ class FNiagaraAttributeTrimmerHelper : public GraphBridge
 	// given the set of expressions (as defined in FindDependencies above) we resolve the named attribute aggregating the dependent reads and custom nodes
 	static void ResolveDependencyChain(const FParamMapHistory& ParamMap, const FDependencyMap& DependencyData, const FName& AttributeName, FDependencyChain& ResolvedDependencies);
 
+	static void BuildCustomHlslNodeInfo(FTrimAttributeCache& Cache, const FCustomHlslNode* CustomNode, FCustomHlslNodeInfo& NodeInfo);
+
 public:
 	static void TrimAttributes_Safe(TConstArrayView<const FParamMapHistory*> LocalParamHistories, TSet<FName>& AttributesToPreserve, TArray<FNiagaraVariable>& Attributes);
 	static void TrimAttributes_Aggressive(const FCompilationCopy* CompileDuplicateData, TConstArrayView<const FParamMapHistory*> LocalParamHistories, TSet<FName>& AttributesToPreserve, TArray<FNiagaraVariable>& Attributes);
diff --git a/Engine/Plugins/FX/Niagara/Source/NiagaraEditor/Private/NiagaraCompilationDigestBridgeImpl.cpp b/Engine/Plugins/FX/Niagara/Source/NiagaraEditor/Private/NiagaraCompilationDigestBridgeImpl.cpp
--- a/Engine/Plugins/FX/Niagara/Source/NiagaraEditor/Private/NiagaraCompilationDigestBridgeImpl.cpp
+++ b/Engine/Plugins/FX/Niagara/Source/NiagaraEditor/Private/NiagaraCompilationDigestBridgeImpl.cpp
@@ -492,6 +492,11 @@ const FNiagaraCompilationNodeEmitter* FNiagaraCompilationDigestBridge::GetNodeAs
 	return Node->AsType<FEmitterNode>();
 }
 
+bool FNiagaraCompilationDigestBridge::GetCustomNodeUsesImpureFunctions(const FNiagaraCompilationNodeCustomHlsl* CustomNode)
+{
+	return CustomNode->bCallsImpureFunctions;
+}
+
 //////////////////////////////////////////////////////////////////////////
 
 
diff --git a/Engine/Plugins/FX/Niagara/Source/NiagaraEditor/Private/NiagaraCompilationGraphBridgeImpl.cpp b/Engine/Plugins/FX/Niagara/Source/NiagaraEditor/Private/NiagaraCompilationGraphBridgeImpl.cpp
--- a/Engine/Plugins/FX/Niagara/Source/NiagaraEditor/Private/NiagaraCompilationGraphBridgeImpl.cpp
+++ b/Engine/Plugins/FX/Niagara/Source/NiagaraEditor/Private/NiagaraCompilationGraphBridgeImpl.cpp
@@ -551,6 +551,11 @@ const UNiagaraNodeEmitter* FNiagaraCompilationGraphBridge::GetNodeAsEmitter(cons
 	return Cast<const UNiagaraNodeEmitter>(Node);
 }
 
+bool FNiagaraCompilationGraphBridge::GetCustomNodeUsesImpureFunctions(const UNiagaraNodeCustomHlsl* CustomNode)
+{
+	return CustomNode->CallsImpureDataInterfaceFunctions();
+}
+
 UNiagaraParameterCollection* FNiagaraCompilationGraphBridge::FAvailableParameterCollections::FindCollection(const FNiagaraVariable& Variable) const
 {
 	const UEdGraphSchema_Niagara* Schema = GetDefault<UEdGraphSchema_Niagara>();
diff --git a/Engine/Plugins/FX/Niagara/Source/NiagaraEditor/Private/NiagaraGraphDigest.cpp b/Engine/Plugins/FX/Niagara/Source/NiagaraEditor/Private/NiagaraGraphDigest.cpp
--- a/Engine/Plugins/FX/Niagara/Source/NiagaraEditor/Private/NiagaraGraphDigest.cpp
+++ b/Engine/Plugins/FX/Niagara/Source/NiagaraEditor/Private/NiagaraGraphDigest.cpp
@@ -2948,8 +2948,9 @@ FNiagaraCompilationNodeCustomHlsl::FNiagaraCompilationNodeCustomHlsl(const UNiag
 		Tokens.Push(FString(View));
 	}
 	InNode->GetIncludeFilePaths(CustomIncludePaths);
-}
 
+	bCallsImpureFunctions = InNode->CallsImpureDataInterfaceFunctions();
+}
 
 FNiagaraCompilationNodeCustomHlsl::FNiagaraCompilationNodeCustomHlsl(const FNiagaraCompilationNodeCustomHlsl& InNode, FNiagaraCompilationGraphDuplicateContext& Context)
 	: FNiagaraCompilationNodeFunctionCall(InNode, Context)
@@ -2958,6 +2959,7 @@ FNiagaraCompilationNodeCustomHlsl::FNiagaraCompilationNodeCustomHlsl(const FNiag
 	, CustomHlsl(InNode.CustomHlsl)
 	, Tokens(InNode.Tokens)
 	, CustomIncludePaths(InNode.CustomIncludePaths)
+	, bCallsImpureFunctions(InNode.bCallsImpureFunctions)
 {
 
 }
diff --git a/Engine/Plugins/FX/Niagara/Source/NiagaraEditor/Private/NiagaraGraphDigest.h b/Engine/Plugins/FX/Niagara/Source/NiagaraEditor/Private/NiagaraGraphDigest.h
--- a/Engine/Plugins/FX/Niagara/Source/NiagaraEditor/Private/NiagaraGraphDigest.h
+++ b/Engine/Plugins/FX/Niagara/Source/NiagaraEditor/Private/NiagaraGraphDigest.h
@@ -460,6 +460,7 @@ public:
 	FString CustomHlsl;
 	TArray<FString> Tokens;
 	TArray<FNiagaraCustomHlslInclude> CustomIncludePaths;
+	bool bCallsImpureFunctions = false;
 };
 
 class FNiagaraCompilationNodeIf : public FNiagaraCompilationNode
diff --git a/Engine/Plugins/FX/Niagara/Source/NiagaraEditor/Private/NiagaraNodeCustomHlsl.cpp b/Engine/Plugins/FX/Niagara/Source/NiagaraEditor/Private/NiagaraNodeCustomHlsl.cpp
--- a/Engine/Plugins/FX/Niagara/Source/NiagaraEditor/Private/NiagaraNodeCustomHlsl.cpp
+++ b/Engine/Plugins/FX/Niagara/Source/NiagaraEditor/Private/NiagaraNodeCustomHlsl.cpp
@@ -277,6 +277,57 @@ void UNiagaraNodeCustomHlsl::InitAsCustomHlslDynamicInput(const FNiagaraTypeDefi
 	ScriptUsage = ENiagaraScriptUsage::DynamicInput;
 }
 
+bool UNiagaraNodeCustomHlsl::CallsImpureDataInterfaceFunctions() const
+{
+	TArray<FString> ImpureFunctionNames;
+
+	FPinCollectorArray InputPins;
+	GetInputPins(InputPins);
+
+	for (const UEdGraphPin* InputPin : InputPins)
+	{
+		FNiagaraTypeDefinition NiagaraType = UEdGraphSchema_Niagara::PinToTypeDefinition(InputPin, ENiagaraStructConversion::Simulation);
+		if (NiagaraType.IsDataInterface())
+		{
+			if (UNiagaraDataInterface* DataInterfaceClass = CastChecked<UNiagaraDataInterface>(NiagaraType.GetClass()->ClassDefaultObject))
+			{
+				TArray<FNiagaraFunctionSignature> FunctionSignatures;
+				DataInterfaceClass->GetFunctions(FunctionSignatures);
+
+				for (const FNiagaraFunctionSignature& FunctionSignature : FunctionSignatures)
+				{
+					if (FunctionSignature.bRequiresExecPin)
+					{
+						TStringBuilder<256> Builder;
+						InputPin->PinName.AppendString(Builder);
+						Builder.AppendChar(TCHAR('.'));
+						FunctionSignature.Name.AppendString(Builder);
+
+						ImpureFunctionNames.AddUnique(Builder.ToString());
+					}
+				}
+			}
+		}
+	}
+
+	if (!ImpureFunctionNames.IsEmpty())
+	{
+		TArray<FStringView> CustomHlslTokens;
+		GetTokensFromString(CustomHlsl, CustomHlslTokens, false, false);
+
+		for (const FString& ImpureFunctionName : ImpureFunctionNames)
+		{
+			if (CustomHlslTokens.Contains(ImpureFunctionName))
+			{
+				return true;
+			}
+		}
+	}
+
+	return false;
+
+}
+
 bool UNiagaraNodeCustomHlsl::IsPinNameEditableUponCreation(const UEdGraphPin* GraphPinObj) const
 {
 	if (GraphPinObj == PinPendingRename && ScriptUsage != ENiagaraScriptUsage::DynamicInput)
diff --git a/Engine/Plugins/FX/Niagara/Source/NiagaraEditor/Public/NiagaraCompilationBridge.h b/Engine/Plugins/FX/Niagara/Source/NiagaraEditor/Public/NiagaraCompilationBridge.h
--- a/Engine/Plugins/FX/Niagara/Source/NiagaraEditor/Public/NiagaraCompilationBridge.h
+++ b/Engine/Plugins/FX/Niagara/Source/NiagaraEditor/Public/NiagaraCompilationBridge.h
@@ -217,6 +217,8 @@ struct FNiagaraCompilationGraphBridge
 	static bool IsParameterMapPin(const FPin* Pin);
 	static bool GetGraphReferencesStaticVariables(const FGraph* Graph, FNiagaraStaticVariableSearchContext& StaticVariableContext);
 	static const FEmitterNode* GetNodeAsEmitter(const FNode* Node);
+
+	static bool GetCustomNodeUsesImpureFunctions(const FCustomHlslNode* CustomNode);
 };
 
 struct FNiagaraCompilationDigestBridge
@@ -346,4 +348,6 @@ struct FNiagaraCompilationDigestBridge
 	static bool IsParameterMapPin(const FPin* Pin);
 	static bool GetGraphReferencesStaticVariables(const FGraph* Graph, FNiagaraStaticVariableSearchContext& StaticVariableContext);
 	static const FEmitterNode* GetNodeAsEmitter(const FNode* Node);
+
+	static bool GetCustomNodeUsesImpureFunctions(const FCustomHlslNode* CustomNode);
 };
diff --git a/Engine/Plugins/FX/Niagara/Source/NiagaraEditor/Public/NiagaraNodeCustomHlsl.h b/Engine/Plugins/FX/Niagara/Source/NiagaraEditor/Public/NiagaraNodeCustomHlsl.h
--- a/Engine/Plugins/FX/Niagara/Source/NiagaraEditor/Public/NiagaraNodeCustomHlsl.h
+++ b/Engine/Plugins/FX/Niagara/Source/NiagaraEditor/Public/NiagaraNodeCustomHlsl.h
@@ -54,6 +54,8 @@ public:
 
 	void InitAsCustomHlslDynamicInput(const FNiagaraTypeDefinition& OutputType);
 
+	bool CallsImpureDataInterfaceFunctions() const;
+
 protected:
 	virtual bool AllowDynamicPins() const override { return true; }
 	virtual bool GetValidateDataInterfaces() const override { return false; }
diff --git a/Engine/Source/Runtime/Core/Private/UObject/FortniteMainBranchObjectVersion.cpp b/Engine/Source/Runtime/Core/Private/UObject/FortniteMainBranchObjectVersion.cpp
--- a/Engine/Source/Runtime/Core/Private/UObject/FortniteMainBranchObjectVersion.cpp
+++ b/Engine/Source/Runtime/Core/Private/UObject/FortniteMainBranchObjectVersion.cpp
@@ -11,7 +11,7 @@ TMap<FGuid, FGuid> FFortniteMainBranchObjectVersion::GetSystemGuids()
 	SystemGuids.Add(DevGuids.MATERIALSHADERMAP_DERIVEDDATA_VER, FGuid("CE7776E63E9E43FBB8C103EE1AB40B5D"));
 	SystemGuids.Add(DevGuids.NANITE_DERIVEDDATA_VER, FGuid("D28C003F825F4C178F698938FB0E47F7"));
 	SystemGuids.Add(DevGuids.NIAGARASHADERMAP_DERIVEDDATA_VER, FGuid("D0DF49D408F34E38A98FA097E0280EC7"));
-	SystemGuids.Add(DevGuids.Niagara_LatestScriptCompileVersion, FGuid("2017726CDB77425CB8F3D6278D491E65"));
+	SystemGuids.Add(DevGuids.Niagara_LatestScriptCompileVersion, FGuid("9B2B82C4C8CD4EFDB9576920A5800343"));
 	SystemGuids.Add(DevGuids.SkeletalMeshDerivedDataVersion, FGuid("518D7C51D31D9BE4A9F4B81542AC8119"));
 	SystemGuids.Add(DevGuids.STATICMESH_DERIVEDDATA_VER, FGuid("543843B103794E0E9BA4BE80FB602F79"));
 	return SystemGuids;
