Fix memory stomp when using SimCache Capture from the niagara component's details panel

diff --git a/Engine/Plugins/FX/Niagara/Source/Niagara/Classes/NiagaraSimCacheCapture.cpp b/Engine/Plugins/FX/Niagara/Source/Niagara/Classes/NiagaraSimCacheCapture.cpp
--- a/Engine/Plugins/FX/Niagara/Source/Niagara/Classes/NiagaraSimCacheCapture.cpp
+++ b/Engine/Plugins/FX/Niagara/Source/Niagara/Classes/NiagaraSimCacheCapture.cpp
@@ -28,10 +28,13 @@ void FNiagaraSimCacheCapture::FinishCapture()
 	{
 		CaptureSimCache->EndWrite();
 	}
-	CaptureComplete.Broadcast(CaptureSimCache);
-	CaptureComplete.Clear();
 
+	// in order to support the broadcasted message resulting in us being deleted we need to
+	// be sure that we don't manipulate any data post broadcast
+	FOnCaptureComplete LocalDelegate = MoveTemp(CaptureComplete);
 	WeakCaptureComponent = nullptr;
+
+	LocalDelegate.Broadcast(CaptureSimCache);
 }
 
 void FNiagaraSimCacheCapture::CaptureNiagaraSimCache(UNiagaraSimCache* SimCache, FNiagaraSimCacheCreateParameters CreateParameters, UNiagaraComponent* NiagaraComponent, FNiagaraSimCacheCaptureParameters InCaptureParameters)
diff --git a/Engine/Plugins/FX/Niagara/Source/Niagara/Classes/NiagaraSimCacheCapture.h b/Engine/Plugins/FX/Niagara/Source/Niagara/Classes/NiagaraSimCacheCapture.h
--- a/Engine/Plugins/FX/Niagara/Source/Niagara/Classes/NiagaraSimCacheCapture.h
+++ b/Engine/Plugins/FX/Niagara/Source/Niagara/Classes/NiagaraSimCacheCapture.h
@@ -74,8 +74,11 @@ struct FNiagaraSimCacheCaptureParameters
 class FNiagaraSimCacheCapture
 {
 public:
+	UE_NONCOPYABLE(FNiagaraSimCacheCapture);
+
 	DECLARE_MULTICAST_DELEGATE_OneParam(FOnCaptureComplete, UNiagaraSimCache*);
 
+	FNiagaraSimCacheCapture() = default;
 	NIAGARA_API ~FNiagaraSimCacheCapture();
 
 	FOnCaptureComplete& OnCaptureComplete() {return CaptureComplete;}
diff --git a/Engine/Plugins/FX/Niagara/Source/NiagaraEditor/Private/Customizations/NiagaraComponentDetails.cpp b/Engine/Plugins/FX/Niagara/Source/NiagaraEditor/Private/Customizations/NiagaraComponentDetails.cpp
--- a/Engine/Plugins/FX/Niagara/Source/NiagaraEditor/Private/Customizations/NiagaraComponentDetails.cpp
+++ b/Engine/Plugins/FX/Niagara/Source/NiagaraEditor/Private/Customizations/NiagaraComponentDetails.cpp
@@ -356,12 +356,10 @@ FReply FNiagaraComponentDetails::OnCaptureSelectedSystem()
 
 	if (NiagaraComponentToUse)
 	{
-		
-		int32 CaptureIndex = ComponentCaptures.AddDefaulted();
+		TSharedPtr<FNiagaraSimCacheCapture> SimCacheCapture = MakeShared<FNiagaraSimCacheCapture>();
 		UNiagaraSimCache* SimCache = NewObject<UNiagaraSimCache>(GetTransientPackage());
 		SimCache->SetFlags(RF_Transient);
 		CapturedCaches.Emplace(SimCache);
-		FNiagaraSimCacheCapture& SimCacheCapture = ComponentCaptures[CaptureIndex];
 
 		FNiagaraSimCacheCreateParameters CreateParameters;
 		
@@ -374,23 +372,15 @@ FReply FNiagaraComponentDetails::OnCaptureSelectedSystem()
 		
 		if(NiagaraComponentToUse->GetWorld())
 		{
-			
-			SimCacheCapture.CaptureNiagaraSimCache(SimCache, CreateParameters, NiagaraComponentToUse, CaptureParameters);
-
-			SimCacheCapture.OnCaptureComplete().AddLambda(
-			[CaptureIndex, this](UNiagaraSimCache* CapturedSimCache)
+			SimCacheCapture->OnCaptureComplete().AddLambda(
+			[SimCacheCapture, this](UNiagaraSimCache* CapturedSimCache)
 			{
 				if(CapturedSimCache)
 				{
 					GEditor->GetEditorSubsystem<UAssetEditorSubsystem>()->OpenEditorForAsset(CapturedSimCache, EToolkitMode::Standalone);
 				}
-
-				if(ComponentCaptures.IsValidIndex(CaptureIndex))
-				{
-					ComponentCaptures.RemoveAt(CaptureIndex);
-				}
-				
 			});
+			SimCacheCapture->CaptureNiagaraSimCache(SimCache, CreateParameters, NiagaraComponentToUse, CaptureParameters);
 		}
 
 		
diff --git a/Engine/Plugins/FX/Niagara/Source/NiagaraEditor/Private/Customizations/NiagaraComponentDetails.h b/Engine/Plugins/FX/Niagara/Source/NiagaraEditor/Private/Customizations/NiagaraComponentDetails.h
--- a/Engine/Plugins/FX/Niagara/Source/NiagaraEditor/Private/Customizations/NiagaraComponentDetails.h
+++ b/Engine/Plugins/FX/Niagara/Source/NiagaraEditor/Private/Customizations/NiagaraComponentDetails.h
@@ -54,7 +54,6 @@ private:
 	TWeakObjectPtr<UNiagaraComponent> Component;
 	IDetailLayoutBuilder* Builder = nullptr;
 
-	TArray<FNiagaraSimCacheCapture> ComponentCaptures;
 	TArray<TWeakObjectPtr<UNiagaraSimCache>> CapturedCaches;
 };
 
