Move Niagara CompilationMode option to NiagaraSettings

diff --git a/Engine/Plugins/FX/Niagara/Source/Niagara/Private/NiagaraCompilationTypes.cpp b/Engine/Plugins/FX/Niagara/Source/Niagara/Private/NiagaraCompilationTypes.cpp
--- a/Engine/Plugins/FX/Niagara/Source/Niagara/Private/NiagaraCompilationTypes.cpp
+++ b/Engine/Plugins/FX/Niagara/Source/Niagara/Private/NiagaraCompilationTypes.cpp
@@ -8,19 +8,13 @@
 #include "NiagaraAsyncCompile.h"
 #include "NiagaraModule.h"
 #include "NiagaraScriptSourceBase.h"
+#include "NiagaraSettings.h"
 #include "NiagaraSystem.h"
 
 #define LOCTEXT_NAMESPACE "NiagaraSystem"
 
 #if WITH_EDITORONLY_DATA
 
-enum class ENiagaraCompilationMode : int32
-{
-	Original,
-	AsyncTasks,
-	Verify
-};
-
 enum class ENiagaraCompilationValidateMode : int32
 {
 	Ignore,
@@ -29,14 +23,6 @@ enum class ENiagaraCompilationValidateMode : int32
 	Assert
 };
 
-static int32 GNiagaraSystemCompileMode = (int32) ENiagaraCompilationMode::Original;
-static FAutoConsoleVariableRef CVarNiagaraSystemCompileMode(
-	TEXT("fx.Niagara.SystemCompileMode"),
-	GNiagaraSystemCompileMode,
-	TEXT("Defines how NiagaraSystem will be compiled"),
-	ECVF_Default
-);
-
 static bool GNiagaraCompileDumpTimings = false;
 static FAutoConsoleVariableRef CVarNiagaraCompileDumpTimings(
 	TEXT("fx.Niagara.CompileDumpTimings"),
@@ -771,11 +757,13 @@ private:
 TUniquePtr<FNiagaraActiveCompilation> FNiagaraActiveCompilation::CreateCompilation()
 {
 #if WITH_EDITORONLY_DATA
-	if (GNiagaraSystemCompileMode == (int32)ENiagaraCompilationMode::AsyncTasks)
+	const ENiagaraCompilationMode Compilationmode = GetDefault<UNiagaraSettings>()->CompilationMode;
+
+	if (Compilationmode == ENiagaraCompilationMode::AsyncTasks)
 	{
 		return MakeUnique<FNiagaraActiveCompilationAsyncTask>();
 	}
-	else if (GNiagaraSystemCompileMode == (int32)ENiagaraCompilationMode::Verify)
+	else if (Compilationmode == ENiagaraCompilationMode::Verify)
 	{
 		return MakeUnique<FNiagaraActiveCompilationVerify>();
 	}
diff --git a/Engine/Plugins/FX/Niagara/Source/Niagara/Public/NiagaraSettings.h b/Engine/Plugins/FX/Niagara/Source/Niagara/Public/NiagaraSettings.h
--- a/Engine/Plugins/FX/Niagara/Source/Niagara/Public/NiagaraSettings.h
+++ b/Engine/Plugins/FX/Niagara/Source/Niagara/Public/NiagaraSettings.h
@@ -101,6 +101,16 @@ enum class ENiagaraStripScriptByteCodeOption : uint8
 	Strip_Experimental = 2 UMETA(DisplayName = "Strip Experimental ByteCode"),
 };
 
+#if WITH_EDITORONLY_DATA
+UENUM()
+enum class ENiagaraCompilationMode : int32
+{
+	Original = 0	UMETA(DisplayName = "Standard Compilation"),
+	AsyncTasks = 1	UMETA(DisplayName = "Experimental - Async Compilation"),
+	Verify = 2		UMETA(DisplayName = "Experimental - Validate Async")
+};
+#endif
+
 UCLASS(config = Niagara, defaultconfig, meta=(DisplayName="Niagara"), MinimalAPI)
 class UNiagaraSettings : public UDeveloperSettings
 {
@@ -148,12 +158,6 @@ class UNiagaraSettings : public UDeveloperSettings
 	UPROPERTY(config, EditAnywhere, Category = Niagara, meta = (DisplayName = "Enable building data for Experimental VM"))
 	bool bExperimentalVMEnabled = false;
 
-	/**
-	 Controls how byte code will be stripped when loading assets that have multiple sets of bytecode (i.e. experimental).
-	 */
-	UPROPERTY(config, EditAnywhere, Category = Niagara, meta = (DisplayName = "Option for how to strip bytecode"))
-	ENiagaraStripScriptByteCodeOption ByteCodeStripOption = ENiagaraStripScriptByteCodeOption::Default;
-
 	/** Whether to limit the max tick delta time or not. */
 	UPROPERTY(config, EditAnywhere, BlueprintReadWrite, AdvancedDisplay, Category = "Niagara", meta = (InlineEditConditionToggle))
 	bool bLimitDeltaTime = true;
@@ -174,6 +178,17 @@ class UNiagaraSettings : public UDeveloperSettings
 	UPROPERTY(config, EditAnywhere, Category=Niagara)
 	FLinearColor PositionPinTypeColor;
 
+	/**
+	 Controls how byte code will be stripped when loading assets that have multiple sets of bytecode (i.e. experimental).
+	 */
+	UPROPERTY(config, EditAnywhere, Category = Niagara, meta = (DisplayName = "Option for how to strip bytecode"))
+	ENiagaraStripScriptByteCodeOption ByteCodeStripOption = ENiagaraStripScriptByteCodeOption::Default;
+
+#if WITH_EDITORONLY_DATA
+	UPROPERTY(config, EditAnywhere, Category = Niagara, meta = (DisplayName = "Option for how to compile Niagara scripts"))
+	ENiagaraCompilationMode CompilationMode = ENiagaraCompilationMode::Original;
+#endif
+
 	/** The quality levels Niagara uses. */
 	UPROPERTY(config, EditAnywhere, Category = Scalability)
 	TArray<FText> QualityLevels;
