Only log lwc reset warning when the system already ticked

diff --git a/Engine/Plugins/FX/Niagara/Source/Niagara/Private/NiagaraSystemSimulation.cpp b/Engine/Plugins/FX/Niagara/Source/Niagara/Private/NiagaraSystemSimulation.cpp
--- a/Engine/Plugins/FX/Niagara/Source/Niagara/Private/NiagaraSystemSimulation.cpp
+++ b/Engine/Plugins/FX/Niagara/Source/Niagara/Private/NiagaraSystemSimulation.cpp
@@ -1186,7 +1186,10 @@ void FNiagaraSystemSimulation::Tick_GameThread_Internal(float DeltaSeconds, cons
 						if (UFXSystemComponent::RequiresLWCTileRecache(Instance->GetLWCTile(), SceneComponent->GetComponentLocation()))
 						{
 							//-OPT: For safety we reset everything, but if everything is local space we may not need to, or we could rebase.
-							UE_LOG(LogNiagara, Warning, TEXT("NiagaraComponent(%s - %s) required LWC tile recache and was reset."), *GetFullNameSafe(SceneComponent), *GetFullNameSafe(System));
+							if (Instance->GetTickCount() > 1)
+							{
+								UE_LOG(LogNiagara, Warning, TEXT("NiagaraComponent(%s - %s) required LWC tile recache and was reset."), *GetFullNameSafe(SceneComponent), *GetFullNameSafe(System));
+							}
 							Instance->Reset(FNiagaraSystemInstance::EResetMode::ResetAll);
 						}
 					}
