Go to source files when double-clicking a data interface function

diff --git a/Engine/Plugins/FX/Niagara/Source/Niagara/Private/DataInterface/NiagaraDataInterfaceDataChannelRead.cpp b/Engine/Plugins/FX/Niagara/Source/Niagara/Private/DataInterface/NiagaraDataInterfaceDataChannelRead.cpp
--- a/Engine/Plugins/FX/Niagara/Source/Niagara/Private/DataInterface/NiagaraDataInterfaceDataChannelRead.cpp
+++ b/Engine/Plugins/FX/Niagara/Source/Niagara/Private/DataInterface/NiagaraDataInterfaceDataChannelRead.cpp
@@ -69,6 +69,7 @@ namespace NDIDataChannelReadLocal
 			Sig.Name = TEXT("Num");
 #if WITH_EDITORONLY_DATA
 			Sig.Description = LOCTEXT("NumFunctionDescription", "Returns the current number of DataChannel accessible by this interface.");
+			NIAGARA_ADD_FUNCTION_SOURCE_INFO(Sig)
 #endif
 			Sig.bMemberFunction = true;
 			Sig.bExperimental = true;
@@ -86,6 +87,7 @@ namespace NDIDataChannelReadLocal
 			Sig.Name = TEXT("Read");
 #if WITH_EDITORONLY_DATA
 			Sig.Description = LOCTEXT("ReadFunctionDescription", "Reads DataChannel data at a specific index. Any values we read that are not in the DataChannel data are set to their default values. Returns success if there was a valid DataChannel to read from.");
+			NIAGARA_ADD_FUNCTION_SOURCE_INFO(Sig)
 #endif
 			Sig.bMemberFunction = true;
 			Sig.bExperimental = true;
@@ -105,6 +107,7 @@ namespace NDIDataChannelReadLocal
 			Sig.Name = TEXT("Consume");
 #if WITH_EDITORONLY_DATA
 			Sig.Description = LOCTEXT("ConsumeFunctionDescription", "Consumes an DataChannel from the end of the DataChannel array and reads the specified values. Any values we read that are not in the DataChannel data are set to their default values. Returns success if an DataChannel was available to pop.");
+			NIAGARA_ADD_FUNCTION_SOURCE_INFO(Sig)
 #endif
 			Sig.bMemberFunction = true;
 			Sig.bExperimental = true;
@@ -133,6 +136,7 @@ namespace NDIDataChannelReadLocal
 
 			Sig.Name = TEXT("SpawnConditional");
 #if WITH_EDITORONLY_DATA
+			NIAGARA_ADD_FUNCTION_SOURCE_INFO(Sig)
 			Sig.Description = LOCTEXT("SpawnCustomFunctionDescription", "Will Spawn particles into the bound Emitter between Min and Max counts for every element in the bound Data Channel.\n\
 		Can take optional additional parameters as conditions on spawning that will be compared against the contents of each data channel element.\n\
 		For example you could spawn only for a particular value of an enum.\n\
diff --git a/Engine/Plugins/FX/Niagara/Source/Niagara/Public/NiagaraCommon.h b/Engine/Plugins/FX/Niagara/Source/Niagara/Public/NiagaraCommon.h
--- a/Engine/Plugins/FX/Niagara/Source/Niagara/Public/NiagaraCommon.h
+++ b/Engine/Plugins/FX/Niagara/Source/Niagara/Public/NiagaraCommon.h
@@ -343,6 +343,9 @@ struct FNiagaraScriptDataUsageInfo
 	bool bReadsAttributeData;
 };
 
+// automatically fills in source file and line for a FNiagaraFunctionSignature, so the editor can jump to it on double click
+#define NIAGARA_ADD_FUNCTION_SOURCE_INFO(FunctionSignature) { FunctionSignature.SourceFile = __FILE__; FunctionSignature.SourceLine = __LINE__; }
+
 USTRUCT()
 struct FNiagaraFunctionSignature
 {
@@ -443,6 +446,12 @@ struct FNiagaraFunctionSignature
 
 	UPROPERTY(meta = (SkipForCompileHash = true))
 	TMap<FNiagaraVariableBase, FText> OutputDescriptions;
+
+	UPROPERTY(meta = (SkipForCompileHash = true))
+	FString SourceFile;
+
+	UPROPERTY(meta = (SkipForCompileHash = true))
+	int32 SourceLine = 0;
 #endif
 
 	FNiagaraFunctionSignature() 
diff --git a/Engine/Plugins/FX/Niagara/Source/NiagaraEditor/Private/NiagaraNodeFunctionCall.cpp b/Engine/Plugins/FX/Niagara/Source/NiagaraEditor/Private/NiagaraNodeFunctionCall.cpp
--- a/Engine/Plugins/FX/Niagara/Source/NiagaraEditor/Private/NiagaraNodeFunctionCall.cpp
+++ b/Engine/Plugins/FX/Niagara/Source/NiagaraEditor/Private/NiagaraNodeFunctionCall.cpp
@@ -29,6 +29,10 @@
 #include "NiagaraNodeStaticSwitch.h"
 #include "NiagaraSettings.h"
 #include "ScopedTransaction.h"
+#include "SourceCodeNavigation.h"
+#include "UnrealEdGlobals.h"
+#include "Editor/UnrealEdEngine.h"
+#include "Preferences/UnrealEdOptions.h"
 #include "Widgets/Input/SComboButton.h"
 
 #include UE_INLINE_GENERATED_CPP_BY_NAME(NiagaraNodeFunctionCall)
@@ -864,6 +868,28 @@ void UNiagaraNodeFunctionCall::OpenReferencedAsset() const
 		}
 		#endif
 	}
+	else if (GUnrealEd && GUnrealEd->GetUnrealEdOptions()->IsCPPAllowed())
+	{
+		if (!Signature.SourceFile.IsEmpty() && FPaths::FileExists(Signature.SourceFile))
+		{
+			FSourceCodeNavigation::OpenSourceFile(Signature.SourceFile, Signature.SourceLine);
+		}
+		else if (Signature.Inputs.Num() > 0)
+		{
+			// if the signature doesn't have the source info stored we try to find the source file
+			// by looking at the data interface on the input pin as a fallback
+			UClass* Class = Signature.Inputs[0].GetType().GetClass();
+			if (Class && Class->IsChildOf(UNiagaraDataInterface::StaticClass()))
+			{
+				FString ClassSourcePath;
+				if(FSourceCodeNavigation::FindClassSourcePath(Class, ClassSourcePath))
+				{
+					const FString AbsoluteSourcePath = IFileManager::Get().ConvertToAbsolutePathForExternalAppForRead(*ClassSourcePath);
+					FSourceCodeNavigation::OpenSourceFile(AbsoluteSourcePath);
+				}
+			}
+		}
+	}
 }
 
 void UNiagaraNodeFunctionCall::UpdateNodeErrorMessage()
