- Downgrade ensure to log

diff --git a/Engine/Plugins/FX/Niagara/Source/NiagaraEditor/Private/ViewModels/Stack/NiagaraStackFunctionInput.cpp b/Engine/Plugins/FX/Niagara/Source/NiagaraEditor/Private/ViewModels/Stack/NiagaraStackFunctionInput.cpp
--- a/Engine/Plugins/FX/Niagara/Source/NiagaraEditor/Private/ViewModels/Stack/NiagaraStackFunctionInput.cpp
+++ b/Engine/Plugins/FX/Niagara/Source/NiagaraEditor/Private/ViewModels/Stack/NiagaraStackFunctionInput.cpp
@@ -78,8 +78,7 @@ void GenerateInputParameterHandlePath(UNiagaraNodeFunctionCall& ModuleNode, UNia
 	{
 		FunctionOutputPins.Reset();
 		CurrentFunctionCallNode->GetOutputPins(FunctionOutputPins);
-		if (ensureMsgf(FunctionOutputPins.Num() == 1 && FunctionOutputPins[0]->LinkedTo.Num() == 1 && FunctionOutputPins[0]->LinkedTo[0]->GetOwningNode()->IsA<UNiagaraNodeParameterMapSet>(),
-			TEXT("Invalid Stack Graph - Dynamic Input Function call didn't have a valid connected output.")))
+		if (FunctionOutputPins.Num() == 1 && FunctionOutputPins[0]->LinkedTo.Num() == 1 && FunctionOutputPins[0]->LinkedTo[0]->GetOwningNode()->IsA<UNiagaraNodeParameterMapSet>())
 		{
 			FNiagaraParameterHandle AliasedHandle(FunctionOutputPins[0]->LinkedTo[0]->PinName);
 			OutHandlePath.Add(FNiagaraParameterHandle::CreateModuleParameterHandle(AliasedHandle.GetName()));
@@ -97,7 +96,6 @@ void GenerateInputParameterHandlePath(UNiagaraNodeFunctionCall& ModuleNode, UNia
 				}
 			}
 
-
 			if (ensureMsgf(CurrentFunctionCallNode != nullptr, TEXT("Invalid Stack Graph - Function call node for override pin %s could not be found."), *FunctionOutputPins[0]->PinName.ToString()) == false)
 			{
 				OutHandlePath.Empty();
@@ -106,6 +104,8 @@ void GenerateInputParameterHandlePath(UNiagaraNodeFunctionCall& ModuleNode, UNia
 		}
 		else
 		{
+			UE_LOG(LogNiagaraEditor, Log, TEXT("Invalid Stack Graph - Dynamic Input Function call didn't have a valid connected output."));
+
 			OutHandlePath.Empty();
 			return;
 		}
