- Add IncludeOnlyPackages file list to audit commandlet

diff --git a/Engine/Plugins/FX/Niagara/Source/NiagaraEditor/Private/Commandlets/NiagaraSystemAuditCommandlet.cpp b/Engine/Plugins/FX/Niagara/Source/NiagaraEditor/Private/Commandlets/NiagaraSystemAuditCommandlet.cpp
--- a/Engine/Plugins/FX/Niagara/Source/NiagaraEditor/Private/Commandlets/NiagaraSystemAuditCommandlet.cpp
+++ b/Engine/Plugins/FX/Niagara/Source/NiagaraEditor/Private/Commandlets/NiagaraSystemAuditCommandlet.cpp
@@ -4,9 +4,11 @@
 #include "DeviceProfiles/DeviceProfile.h"
 #include "DeviceProfiles/DeviceProfileManager.h"
 #include "HAL/FileManager.h"
+#include "HAL/PlatformFileManager.h"
+#include "Misc/FileHelper.h"
 #include "Misc/Paths.h"
-#include "Modules/ModuleManager.h"
 #include "Misc/PackageName.h"
+#include "Modules/ModuleManager.h"
 #include "AssetRegistry/AssetData.h"
 #include "AssetRegistry/ARFilter.h"
 #include "AssetRegistry/AssetRegistryModule.h"
@@ -178,6 +180,22 @@ int32 UNiagaraSystemAuditCommandlet::Main(const FString& Params)
 		PackagePaths.Add(FName(TEXT("/Game")));
 	}
 
+	// Include only assets
+	FString IncludeOnlyPackagesFileName;
+	if (FParse::Value(*Params, TEXT("IncludeOnlyPackages="), IncludeOnlyPackagesFileName, false))
+	{
+		IPlatformFile& PlatformFile = FPlatformFileManager::Get().GetPlatformFile();
+		if (PlatformFile.FileExists(*IncludeOnlyPackagesFileName))
+		{
+			TArray<FString> IncludeOnlyPackagesStringArray;
+			FFileHelper::LoadFileToStringArray(IncludeOnlyPackagesStringArray, *IncludeOnlyPackagesFileName);
+			for (const FString& Entry : IncludeOnlyPackagesStringArray)
+			{
+				IncludeOnlyPackages.Add(FName(Entry));
+			}
+		}
+	}
+
 	ProcessNiagaraSystems();
 	DumpResults();
 	NiagaraValidationIssues.Reset();
@@ -226,6 +244,14 @@ bool UNiagaraSystemAuditCommandlet::ProcessNiagaraSystems()
 		const FString SystemName = AssetIt.GetObjectPathString();
 		const FString PackageName = AssetIt.PackageName.ToString();
 
+		if (IncludeOnlyPackages.Num())
+		{
+			if (!IncludeOnlyPackages.Contains(AssetIt.PackagePath) && !IncludeOnlyPackages.Contains(FName(SystemName)))
+			{
+				continue;
+			}
+		}
+
 		if (PackageName.StartsWith(DevelopersFolder))
 		{
 			// Skip developer folders
diff --git a/Engine/Plugins/FX/Niagara/Source/NiagaraEditor/Public/Commandlets/NiagaraSystemAuditCommandlet.h b/Engine/Plugins/FX/Niagara/Source/NiagaraEditor/Public/Commandlets/NiagaraSystemAuditCommandlet.h
--- a/Engine/Plugins/FX/Niagara/Source/NiagaraEditor/Public/Commandlets/NiagaraSystemAuditCommandlet.h
+++ b/Engine/Plugins/FX/Niagara/Source/NiagaraEditor/Public/Commandlets/NiagaraSystemAuditCommandlet.h
@@ -55,6 +55,9 @@ class UNiagaraSystemAuditCommandlet : public UCommandlet
 	/** Package paths to include */
 	TArray<FName> PackagePaths;
 
+	/** List of assets to only include (i.e. process only these) */
+	TSet<FName> IncludeOnlyPackages;
+
 	/** Systems using specific data interfaces to find */
 	TSet<UClass*> DataInterfacesToFind;
 
