Fix build warning.

diff --git a/Engine/Plugins/FX/Niagara/Source/Niagara/Public/NiagaraDataChannelAccessor.h b/Engine/Plugins/FX/Niagara/Source/Niagara/Public/NiagaraDataChannelAccessor.h
--- a/Engine/Plugins/FX/Niagara/Source/Niagara/Public/NiagaraDataChannelAccessor.h
+++ b/Engine/Plugins/FX/Niagara/Source/Niagara/Public/NiagaraDataChannelAccessor.h
@@ -170,7 +170,8 @@ struct FNiagaraDataChannelGameDataWriter : public FNiagaraDataChannelGameDataWri
 	/** Write data to a specific variable index. */
 	bool Write(int32 Index, const T& InData)
 	{
-		//static_assert(false, "All implementations must provide a write function");		
+		//static_assert(false, "All implementations must provide a write function");
+		return false;
 	}
 };
 
