Removed parameter action favoring for Niagara and halved NPC parameter search score on top

diff --git a/Engine/Plugins/FX/Niagara/Source/NiagaraEditor/Private/NiagaraEditorUtilities.cpp b/Engine/Plugins/FX/Niagara/Source/NiagaraEditor/Private/NiagaraEditorUtilities.cpp
--- a/Engine/Plugins/FX/Niagara/Source/NiagaraEditor/Private/NiagaraEditorUtilities.cpp
+++ b/Engine/Plugins/FX/Niagara/Source/NiagaraEditor/Private/NiagaraEditorUtilities.cpp
@@ -1879,10 +1879,12 @@ int32 FNiagaraEditorUtilities::GetWeightForItem(const TSharedPtr<FNiagaraMenuAct
 		}
 	}
 
-	// parameter actions get favored
 	if(InCurrentAction->GetParameterVariable().IsSet())
 	{
-		TotalWeight *= 2;
+		if(InCurrentAction->GetParameterVariable().GetValue().IsInNameSpace(FNiagaraConstants::ParameterCollectionNamespaceString))
+		{
+			TotalWeight *= 0.5f;
+		}
 	}
 
 	// suggested actions get favored massively
