Request a compile of the system when handling a change to a user parameter name.  Avoids a crash that can happen because of a mismatch between the compiled constant table layouts and the parameter stores.

diff --git a/Engine/Plugins/FX/Niagara/Source/NiagaraEditor/Private/ViewModels/NiagaraSystemViewModel.cpp b/Engine/Plugins/FX/Niagara/Source/NiagaraEditor/Private/ViewModels/NiagaraSystemViewModel.cpp
--- a/Engine/Plugins/FX/Niagara/Source/NiagaraEditor/Private/ViewModels/NiagaraSystemViewModel.cpp
+++ b/Engine/Plugins/FX/Niagara/Source/NiagaraEditor/Private/ViewModels/NiagaraSystemViewModel.cpp
@@ -1362,6 +1362,10 @@ bool FNiagaraSystemViewModel::RenameParameter(const FNiagaraVariable TargetParam
 			System->HandleVariableRenamed(TargetParameter, FNiagaraVariableBase(TargetParameter.GetType(), NewName), true);
 		}
 
+		// changing a parameter name can result in conflicts between the runtime data and the parameter stores that we've just
+		// changed.  Currently the safest way to deal with that is a recompile of the system
+		CompileSystem(false);
+
 		GetSystemStackViewModel()->RequestRefreshDeferred();
 		return true;
 	}
