Fix NiagaraCacheTests for game configuration

diff --git a/Engine/Plugins/FX/Niagara/Source/Niagara/Classes/NiagaraSimCache.h b/Engine/Plugins/FX/Niagara/Source/Niagara/Classes/NiagaraSimCache.h
--- a/Engine/Plugins/FX/Niagara/Source/Niagara/Classes/NiagaraSimCache.h
+++ b/Engine/Plugins/FX/Niagara/Source/Niagara/Classes/NiagaraSimCache.h
@@ -547,10 +547,8 @@ public:
 	UFUNCTION(BlueprintCallable, Category = NiagaraSimCache)
 	NIAGARA_API void ReadQuatAttributeWithRebase(TArray<FQuat>& OutValues, FQuat Quat, FName AttributeName = FName("MeshOrientation"), FName EmitterName = NAME_None, int FrameIndex = 0) const;
 
-#if WITH_EDITORONLY_DATA
 	UE_DEPRECATED(5.3, "For test purposes only, don't use in production code.")
 	NIAGARA_API bool IsDataEqual(const UNiagaraSimCache& OtherCache, float ErrorTolerance, FString& OutDifference) const;
-#endif
 	
 private:
 	UPROPERTY(VisibleAnywhere, Category=SimCache)
diff --git a/Engine/Plugins/FX/Niagara/Source/Niagara/Private/NiagaraSimCache.cpp b/Engine/Plugins/FX/Niagara/Source/Niagara/Private/NiagaraSimCache.cpp
--- a/Engine/Plugins/FX/Niagara/Source/Niagara/Private/NiagaraSimCache.cpp
+++ b/Engine/Plugins/FX/Niagara/Source/Niagara/Private/NiagaraSimCache.cpp
@@ -1112,8 +1112,6 @@ bool FNiagaraSimCacheDataBuffers::operator!=(const FNiagaraSimCacheDataBuffers&
 	return !(*this == Other);
 }
 
-#if WITH_EDITORONLY_DATA
-
 namespace CacheCompare
 {
 	constexpr int32 ErrorStrLen = 8000;
@@ -1338,5 +1336,5 @@ bool UNiagaraSimCache::IsDataEqual(const UNiagaraSimCache& OtherCache, float Err
 	OutDifference = Errors.ToString();
 	return bEqual;
 }
-#endif
+
 
