Fix new compilation mode so that it always updates the compilation id

diff --git a/Engine/Plugins/FX/Niagara/Source/Niagara/Classes/NiagaraScript.h b/Engine/Plugins/FX/Niagara/Source/Niagara/Classes/NiagaraScript.h
--- a/Engine/Plugins/FX/Niagara/Source/Niagara/Classes/NiagaraScript.h
+++ b/Engine/Plugins/FX/Niagara/Source/Niagara/Classes/NiagaraScript.h
@@ -1155,7 +1155,7 @@ public:
 
 #if WITH_EDITORONLY_DATA
 	template<typename T>
-	TOptional<T> GetStaticVariableValue(const FNiagaraVariableBase& InVar) const
+	TOptional<T> GetCompiledStaticVariableValue(const FNiagaraVariableBase& InVar) const
 	{
 		for (const FNiagaraVariable& StaticVariable : CachedScriptVM.StaticVariablesWritten)
 		{
@@ -1170,6 +1170,11 @@ public:
 
 		return TOptional<T>();
 	}
+
+	// use of this function differs from the above in that the data is being pulled from the RapidIterationParameters rather than
+	// the compiled data
+	void GatherScriptStaticVariables(TArray<FNiagaraVariable>& StaticVariables) const;
+	bool HasScriptStaticVariables() const;
 #endif
 
 	template<typename T>
diff --git a/Engine/Plugins/FX/Niagara/Source/Niagara/Private/NiagaraCompilationTypes.cpp b/Engine/Plugins/FX/Niagara/Source/Niagara/Private/NiagaraCompilationTypes.cpp
--- a/Engine/Plugins/FX/Niagara/Source/Niagara/Private/NiagaraCompilationTypes.cpp
+++ b/Engine/Plugins/FX/Niagara/Source/Niagara/Private/NiagaraCompilationTypes.cpp
@@ -289,8 +289,6 @@ public:
 
 	virtual void Apply(const FNiagaraQueryCompilationOptions& Options) override
 	{
-		bool bAnyRIParametersDirtied = false;
-
 		// we need to do this apply in multiple passes.  Because the script's VMCompilationId has dependencies
 		// on the rapid iteration parameters of all the related scripts then we need to be sure to update the
 		// RI parameters with the data that we've collected before we actually try to set the results.
@@ -307,10 +305,7 @@ public:
 			UNiagaraScript* TargetScript = ResultIt->Key;
 			const FNiagaraScriptAsyncCompileData& ScriptCompileData = ResultIt->Value;
 
-			if (TargetScript->ApplyRapidIterationParameters(ScriptCompileData.RapidIterationParameters, bAllowParameterRemoval))
-			{
-				bAnyRIParametersDirtied = true;
-			}
+			TargetScript->ApplyRapidIterationParameters(ScriptCompileData.RapidIterationParameters, bAllowParameterRemoval);
 		}
 
 		// Now that the above code says they are all complete, go ahead and resolve them all at once.
@@ -339,19 +334,15 @@ public:
 					});
 				}
 
+				// because our compilation process includes the generation of rapid iteration parameters and static
+				// variables we need to generate the ExecutableDataId
 				// if we dirtied any RI parameters then we need to regenerate our CompilationId
 				FNiagaraVMExecutableDataId UpdatedCompileId;
-				const FNiagaraVMExecutableDataId* CompileIdToUse = &ScriptCompileData.CompileId;
-
-				if (bAnyRIParametersDirtied)
-				{
-					TargetScript->ComputeVMCompilationId(UpdatedCompileId, FGuid());
-					CompileIdToUse = &UpdatedCompileId;
-				}
+				TargetScript->ComputeVMCompilationId(UpdatedCompileId, FGuid());
 
 				constexpr bool bApplyRapidIterationParameters = false;
 				TargetScript->SetVMCompilationResults(
-					*CompileIdToUse,
+					UpdatedCompileId,
 					*ScriptCompileData.ExeData,
 					ScriptCompileData.UniqueEmitterName,
 					ObjectNameMap,
diff --git a/Engine/Plugins/FX/Niagara/Source/Niagara/Private/NiagaraEmitter.cpp b/Engine/Plugins/FX/Niagara/Source/Niagara/Private/NiagaraEmitter.cpp
--- a/Engine/Plugins/FX/Niagara/Source/Niagara/Private/NiagaraEmitter.cpp
+++ b/Engine/Plugins/FX/Niagara/Source/Niagara/Private/NiagaraEmitter.cpp
@@ -429,21 +429,9 @@ void FVersionedNiagaraEmitterData::GatherStaticVariables(TArray<FNiagaraVariable
 
 	for (UNiagaraScript* Script : OutScripts)
 	{
-		TArray<FNiagaraVariable> StoreParams;
-		Script->RapidIterationParameters.GetParameters(StoreParams);
-
-		for (int32 i = 0; i < StoreParams.Num(); i++)
+		if (Script)
 		{
-			if (StoreParams[i].GetType().IsStatic())
-			{
-				const int32* Index = Script->RapidIterationParameters.FindParameterOffset(StoreParams[i]);
-				if (Index != nullptr)
-				{
-					StoreParams[i].SetData(Script->RapidIterationParameters.GetParameterData(*Index)); // This will memcopy the data in.			
-					OutVars.AddUnique(StoreParams[i]);					
-					//UE_LOG(LogNiagara, Log, TEXT("UNiagaraEmitter::GatherStaticVariables Added %s"), *StoreParams[i].ToString());
-				}
-			}
+			Script->GatherScriptStaticVariables(OutVars);
 		}
 	}
 }
diff --git a/Engine/Plugins/FX/Niagara/Source/Niagara/Private/NiagaraMeshRendererProperties.cpp b/Engine/Plugins/FX/Niagara/Source/Niagara/Private/NiagaraMeshRendererProperties.cpp
--- a/Engine/Plugins/FX/Niagara/Source/Niagara/Private/NiagaraMeshRendererProperties.cpp
+++ b/Engine/Plugins/FX/Niagara/Source/Niagara/Private/NiagaraMeshRendererProperties.cpp
@@ -616,7 +616,7 @@ void UNiagaraMeshRendererProperties::CacheFromCompiledData(const FNiagaraDataSet
 				{
 					if (Mesh.LODLevelBinding.AliasedParameter.IsValid())
 					{
-						TOptional<int32> VariableValue = NiagaraScript->GetStaticVariableValue<int32>(Mesh.LODLevelBinding.ResolvedParameter);
+						TOptional<int32> VariableValue = NiagaraScript->GetCompiledStaticVariableValue<int32>(Mesh.LODLevelBinding.ResolvedParameter);
 						if (VariableValue.IsSet())
 						{
 							Mesh.LODLevel = VariableValue.GetValue();
@@ -624,7 +624,7 @@ void UNiagaraMeshRendererProperties::CacheFromCompiledData(const FNiagaraDataSet
 					}
 					if (Mesh.LODBiasBinding.AliasedParameter.IsValid())
 					{
-						TOptional<int32> VariableValue = NiagaraScript->GetStaticVariableValue<int32>(Mesh.LODBiasBinding.ResolvedParameter);
+						TOptional<int32> VariableValue = NiagaraScript->GetCompiledStaticVariableValue<int32>(Mesh.LODBiasBinding.ResolvedParameter);
 						if (VariableValue.IsSet())
 						{
 							Mesh.LODBias = VariableValue.GetValue();
diff --git a/Engine/Plugins/FX/Niagara/Source/Niagara/Private/NiagaraRendererProperties.cpp b/Engine/Plugins/FX/Niagara/Source/Niagara/Private/NiagaraRendererProperties.cpp
--- a/Engine/Plugins/FX/Niagara/Source/Niagara/Private/NiagaraRendererProperties.cpp
+++ b/Engine/Plugins/FX/Niagara/Source/Niagara/Private/NiagaraRendererProperties.cpp
@@ -631,7 +631,7 @@ int32 UNiagaraRendererProperties::GetDynamicParameterChannelMask(const FVersione
 			NameBuilder.Append(TEXT("ChannelMask"));
 
 			const FNiagaraVariableBase ChannelMaskVariable(FNiagaraTypeDefinition::GetIntDef().ToStaticDef(), FName(NameBuilder));
-			TOptional<int32> ChannelMaskValue = NiagaraScript->GetStaticVariableValue<int32>(ChannelMaskVariable);
+			TOptional<int32> ChannelMaskValue = NiagaraScript->GetCompiledStaticVariableValue<int32>(ChannelMaskVariable);
 			if (ChannelMaskValue.IsSet())
 			{
 				ChannelMask = ChannelMask.Get(0) | ChannelMaskValue.GetValue();
diff --git a/Engine/Plugins/FX/Niagara/Source/Niagara/Private/NiagaraScript.cpp b/Engine/Plugins/FX/Niagara/Source/Niagara/Private/NiagaraScript.cpp
--- a/Engine/Plugins/FX/Niagara/Source/Niagara/Private/NiagaraScript.cpp
+++ b/Engine/Plugins/FX/Niagara/Source/Niagara/Private/NiagaraScript.cpp
@@ -3991,6 +3991,35 @@ NIAGARA_API bool UNiagaraScript::DidScriptCompilationSucceed(bool bGPUScript) co
 	return false;
 }
 
+#if WITH_EDITORONLY_DATA
+void UNiagaraScript::GatherScriptStaticVariables(TArray<FNiagaraVariable>& StaticVariables) const
+{
+	for (const FNiagaraVariableWithOffset& Param : RapidIterationParameters.ReadParameterVariables())
+	{
+		if (Param.GetType().IsStatic() && Param.Offset != INDEX_NONE)
+		{
+			FNiagaraVariable StaticVariable = Param;
+			StaticVariable.SetData(RapidIterationParameters.GetParameterData(Param.Offset));
+
+			StaticVariables.AddUnique(StaticVariable);
+		}
+	}
+}
+
+bool UNiagaraScript::HasScriptStaticVariables() const
+{
+	for (const FNiagaraVariableWithOffset& Param : RapidIterationParameters.ReadParameterVariables())
+	{
+		if (Param.GetType().IsStatic() && Param.Offset != INDEX_NONE)
+		{
+			return true;
+		}
+	}
+
+	return false;
+}
+#endif
+
 void UNiagaraScript::SerializeNiagaraShaderMaps(FArchive& Ar, int32 NiagaraVer, bool IsValidShaderScript)
 {
 #if WITH_EDITORONLY_DATA
diff --git a/Engine/Plugins/FX/Niagara/Source/Niagara/Private/NiagaraSystem.cpp b/Engine/Plugins/FX/Niagara/Source/Niagara/Private/NiagaraSystem.cpp
--- a/Engine/Plugins/FX/Niagara/Source/Niagara/Private/NiagaraSystem.cpp
+++ b/Engine/Plugins/FX/Niagara/Source/Niagara/Private/NiagaraSystem.cpp
@@ -3468,50 +3468,29 @@ void UNiagaraSystem::SetEffectType(UNiagaraEffectType* InEffectType)
 
 void UNiagaraSystem::GatherStaticVariables(TArray<FNiagaraVariable>& OutVars, TArray<FNiagaraVariable>& OutEmitterVars) const
 {
-	TArray<UNiagaraScript*> OutScripts;
-	OutScripts.Add(SystemSpawnScript);
-	OutScripts.Add(SystemUpdateScript);
+	TArray<UNiagaraScript*> ScriptsToProcess;
+	ScriptsToProcess.Reserve(2 * (EmitterHandles.Num() + 1));
+	ScriptsToProcess.Add(SystemSpawnScript);
+	ScriptsToProcess.Add(SystemUpdateScript);
 
-	auto GatherFromScripts = [&](TArray<UNiagaraScript*> Scripts, TArray<FNiagaraVariable>& Vars)
-	{
-		for (UNiagaraScript* Script : Scripts)
-		{
-			TArray<FNiagaraVariable> StoreParams;
-			Script->RapidIterationParameters.GetParameters(StoreParams);
-
-			for (int32 i = 0; i < StoreParams.Num(); i++)
-			{
-				if (StoreParams[i].GetType().IsStatic())
-				{
-					const int32* Index = Script->RapidIterationParameters.FindParameterOffset(StoreParams[i]);
-					if (Index != nullptr)
-					{
-						StoreParams[i].SetData(Script->RapidIterationParameters.GetParameterData(*Index)); // This will memcopy the data in.
-						Vars.AddUnique(StoreParams[i]);
-
-						//UE_LOG(LogNiagara, Log, TEXT("UNiagaraSystem::GatherStaticVariables Added %s"), *StoreParams[i].ToString());
-					}
-				}
-			}
-		}
-	};
-
-	GatherFromScripts(OutScripts, OutVars);
-
-	TArray<UNiagaraScript*> EmitterScripts; 
-	//const TArray<FNiagaraEmitterHandle>& EmitterHandles = GetEmitterHandles();
 	for (int32 i = 0; i < EmitterHandles.Num(); i++)
 	{
 		const FNiagaraEmitterHandle& Handle = EmitterHandles[i];
 		FVersionedNiagaraEmitterData* EmitterData = Handle.GetEmitterData();
 		if (EmitterData && Handle.GetIsEnabled())
 		{
-			EmitterScripts.Add(EmitterData->EmitterSpawnScriptProps.Script);
-			EmitterScripts.Add(EmitterData->EmitterUpdateScriptProps.Script);
+			ScriptsToProcess.Add(EmitterData->EmitterSpawnScriptProps.Script);
+			ScriptsToProcess.Add(EmitterData->EmitterUpdateScriptProps.Script);
 		}		
 	}
 
-	GatherFromScripts(EmitterScripts, OutEmitterVars);
+	for (const UNiagaraScript* ScriptToProcess : ScriptsToProcess)
+	{
+		if (ScriptToProcess)
+		{
+			ScriptToProcess->GatherScriptStaticVariables(OutEmitterVars);
+		}
+	}
 }
 #endif
 
