- Add fx.Niagara.ObjectNeedsLoadMode which will validate / handle object loading
- This is to help track down / diagnose issues where the async loader is returning subobjects which are not fully loaded

diff --git a/Engine/Plugins/FX/Niagara/Source/Niagara/Private/NiagaraSystem.cpp b/Engine/Plugins/FX/Niagara/Source/Niagara/Private/NiagaraSystem.cpp
--- a/Engine/Plugins/FX/Niagara/Source/Niagara/Private/NiagaraSystem.cpp
+++ b/Engine/Plugins/FX/Niagara/Source/Niagara/Private/NiagaraSystem.cpp
@@ -39,6 +39,7 @@
 #include "ProfilingDebugging/CookStats.h"
 #include "ProfilingDebugging/ScopedTimers.h"
 #include "UObject/ObjectSaveContext.h"
+#include "UObject/LinkerLoad.h"
 #include "UObject/Package.h"
 #include "PipelineStateCache.h"
 #include "NiagaraDataChannel.h"
@@ -103,6 +104,17 @@ static FAutoConsoleVariableRef CVarNiagaraOnDemandCompileEnabled(
 	TEXT("Compiles Niagara Systems on demand rather than on post load."),
 	ECVF_Default
 );
+
+static int GNiagaraObjectNeedsLoadMode = 1;
+static FAutoConsoleVariableRef CVarNiagaraObjectNeedsLoadMode(
+	TEXT("fx.Niagara.ObjectNeedsLoadMode"),
+	GNiagaraObjectNeedsLoadMode,
+	TEXT("How we decide to handle objects that need loading\n")
+	TEXT("0 - Do nothing\n")
+	TEXT("1 - Validate objects are loaded\n")
+	TEXT("2 - Validate objects are loaded and force preload\n"),
+	ECVF_Default
+);
 #endif
 
 //////////////////////////////////////////////////////////////////////////
@@ -927,6 +939,28 @@ void UNiagaraSystem::PostLoad()
 {
 	Super::PostLoad();
 
+#if WITH_EDITORONLY_DATA
+	// Validate that all our sub-objects have been loaded before we start post loading
+	// This is to help track down issues where subobjects have not be loaded and to warn that it's the case
+	if (GNiagaraObjectNeedsLoadMode > 0)
+	{
+		TArray<UObject*> ObjectReferences;
+		FReferenceFinder(ObjectReferences, this, false, true, true, true).FindReferences(this);
+
+		for (UObject* Dependency : ObjectReferences)
+		{
+			if (Dependency->HasAnyFlags(RF_NeedLoad))
+			{
+				UE_LOG(LogNiagara, Log, TEXT("NiagaraSystem::PostLoad() - SubObject(%s) RF_NeedLoad"), *GetFullNameSafe(Dependency));
+				if (GNiagaraObjectNeedsLoadMode == 2)
+				{
+					Dependency->GetLinker()->Preload(Dependency);
+				}
+			}
+		}
+	}
+#endif
+
 	// Workaround for UE-104235 where a CDO loads a NiagaraSystem before the NiagaraModule has had a chance to load
 	// We force the module to load here we makes sure the type registry, etc, is all setup in time.
 	static bool bLoadChecked = false;
