- Fix for not updating the system when using Fix Issue button for sim stages

diff --git a/Engine/Plugins/FX/Niagara/Source/NiagaraEditor/Private/ViewModels/Stack/NiagaraStackSimulationStageGroup.cpp b/Engine/Plugins/FX/Niagara/Source/NiagaraEditor/Private/ViewModels/Stack/NiagaraStackSimulationStageGroup.cpp
--- a/Engine/Plugins/FX/Niagara/Source/NiagaraEditor/Private/ViewModels/Stack/NiagaraStackSimulationStageGroup.cpp
+++ b/Engine/Plugins/FX/Niagara/Source/NiagaraEditor/Private/ViewModels/Stack/NiagaraStackSimulationStageGroup.cpp
@@ -122,11 +122,16 @@ void UNiagaraStackSimulationStagePropertiesItem::RefreshChildrenInternal(const T
 				FStackIssueFixDelegate::CreateLambda(
 					[WeakEmitter=GetEmitterViewModel()->GetEmitter().ToWeakPtr()]()
 					{
-						if (WeakEmitter.IsValid())
+						const FVersionedNiagaraEmitter VersionedEmitter = WeakEmitter.ResolveWeakPtr();
+						if (VersionedEmitter.Emitter && VersionedEmitter.GetEmitterData() )
 						{
 							FScopedTransaction Transaction(LOCTEXT("SetGpuSimulation", "Set Gpu Simulation"));
-							WeakEmitter.Emitter.Get()->Modify();
-							WeakEmitter.GetEmitterData()->SimTarget = ENiagaraSimTarget::GPUComputeSim;
+							VersionedEmitter.Emitter->Modify();
+							VersionedEmitter.GetEmitterData()->SimTarget = ENiagaraSimTarget::GPUComputeSim;
+
+							FProperty* SimTargetProperty = FindFProperty<FProperty>(FVersionedNiagaraEmitterData::StaticStruct(), GET_MEMBER_NAME_CHECKED(FVersionedNiagaraEmitterData, SimTarget));
+							FPropertyChangedEvent PropertyChangedEvent(SimTargetProperty);
+							VersionedEmitter.Emitter->PostEditChangeVersionedProperty(PropertyChangedEvent, VersionedEmitter.Version);
 						}
 					}
 				)
