fix incorrectly excluded function for grid3d reader

diff --git a/Engine/Plugins/FX/Niagara/Source/Niagara/Private/NiagaraDataInterfaceGrid3DCollectionReader.cpp b/Engine/Plugins/FX/Niagara/Source/Niagara/Private/NiagaraDataInterfaceGrid3DCollectionReader.cpp
--- a/Engine/Plugins/FX/Niagara/Source/Niagara/Private/NiagaraDataInterfaceGrid3DCollectionReader.cpp
+++ b/Engine/Plugins/FX/Niagara/Source/Niagara/Private/NiagaraDataInterfaceGrid3DCollectionReader.cpp
@@ -182,8 +182,7 @@ void UNiagaraDataInterfaceGrid3DCollectionReader::GetFunctions(TArray<FNiagaraFu
 		SetVector4ValueFunctionName,
 		SetVector3ValueFunctionName,
 		SetVector2ValueFunctionName,
-		SetFloatValueFunctionName,
-		GetPreviousVector4ValueFunctionName,
+		SetFloatValueFunctionName,		
 		SetVectorValueFunctionName,
 		SetVector2DValueFunctionName,
 		SetNumCellsFunctionName,
