- Fix SetEmitterEnabled for GPU & Event Spawns

diff --git a/Engine/Plugins/FX/Niagara/Source/Niagara/Private/NiagaraEmitterInstance.cpp b/Engine/Plugins/FX/Niagara/Source/Niagara/Private/NiagaraEmitterInstance.cpp
--- a/Engine/Plugins/FX/Niagara/Source/Niagara/Private/NiagaraEmitterInstance.cpp
+++ b/Engine/Plugins/FX/Niagara/Source/Niagara/Private/NiagaraEmitterInstance.cpp
@@ -1344,6 +1344,7 @@ void FNiagaraEmitterInstance::Tick(float DeltaSeconds)
 		// we're receiving from, so we can use the data buffers that our sources have computed this tick.
 		const int32 NumEventHandlers = EmitterData->GetEventHandlers().Num();
 		EventInstanceData->EventSpawnTotal = 0;
+
 		for (int32 i = 0; i < NumEventHandlers; i++)
 		{
 			const FNiagaraEventScriptProperties& EventHandlerProps = EmitterData->GetEventHandlers()[i];
@@ -1353,12 +1354,15 @@ void FNiagaraEmitterInstance::Tick(float DeltaSeconds)
 			Info.SpawnCounts.Reset();
 
 			//TODO: We can move this lookup into the init and just store a ptr to the other set?
-			if (FNiagaraDataSet* EventSet = ParentSystemInstance->GetEventDataSet(Info.SourceEmitterName, EventHandlerProps.SourceEventName))
+			if (bAllowSpawning_CNC)
 			{
-				Info.SetEventData(&EventSet->GetCurrentDataChecked());
-				uint32 EventSpawnNum = CalculateEventSpawnCount(EventHandlerProps, Info.SpawnCounts, EventSet);
-				Info.TotalSpawnCount += EventSpawnNum;
-				EventInstanceData->EventSpawnTotal += EventSpawnNum;
+				if (FNiagaraDataSet* EventSet = ParentSystemInstance->GetEventDataSet(Info.SourceEmitterName, EventHandlerProps.SourceEventName))
+				{
+					Info.SetEventData(&EventSet->GetCurrentDataChecked());
+					uint32 EventSpawnNum = CalculateEventSpawnCount(EventHandlerProps, Info.SpawnCounts, EventSet);
+					Info.TotalSpawnCount += EventSpawnNum;
+					EventInstanceData->EventSpawnTotal += EventSpawnNum;
+				}
 			}
 		}
 	}
@@ -1468,7 +1472,7 @@ void FNiagaraEmitterInstance::Tick(float DeltaSeconds)
 
 			int NumSpawnInfos = 0;
 			int32 NumSpawnedOnGPUThisFrame = 0;
-			if (ExecutionState == ENiagaraExecutionState::Active)
+			if (ExecutionState == ENiagaraExecutionState::Active && bAllowSpawning_CNC)
 			{
 				for (int32 SpawnInfoIdx = 0; SpawnInfoIdx < SpawnInfos.Num(); SpawnInfoIdx++)
 				{
