Improved Blueprint API to write to data channels.
* added a new write node that automatically generates matching pins for the data channel used
* added versioning information to data channel assets, so dependent assets can pick up on changes

diff --git a/Engine/Plugins/FX/Niagara/Niagara.uplugin b/Engine/Plugins/FX/Niagara/Niagara.uplugin
--- a/Engine/Plugins/FX/Niagara/Niagara.uplugin
+++ b/Engine/Plugins/FX/Niagara/Niagara.uplugin
@@ -47,6 +47,11 @@
 			"Type" : "Runtime",
 			"LoadingPhase" : "PostConfigInit"
 		},
+		{
+			"Name" : "NiagaraBlueprintNodes",
+			"Type" : "UncookedOnly",
+			"LoadingPhase" : "Default"
+		},
 		{
 			"Name" : "NiagaraEditor",
 			"Type" : "Editor",
diff --git a/Engine/Plugins/FX/Niagara/Source/Niagara/Private/DataInterface/NiagaraDataInterfaceDataChannelRead.cpp b/Engine/Plugins/FX/Niagara/Source/Niagara/Private/DataInterface/NiagaraDataInterfaceDataChannelRead.cpp
--- a/Engine/Plugins/FX/Niagara/Source/Niagara/Private/DataInterface/NiagaraDataInterfaceDataChannelRead.cpp
+++ b/Engine/Plugins/FX/Niagara/Source/Niagara/Private/DataInterface/NiagaraDataInterfaceDataChannelRead.cpp
@@ -737,7 +737,7 @@ void UNiagaraDataInterfaceDataChannelRead::GetFeedback(UNiagaraSystem* InAsset,
 		if (const UNiagaraDataChannel* DataChannel = RuntimeReadDI->Channel->Get())
 		{
 			//Ensure the data channel contains all the parameters this function is requesting.
-			TConstArrayView<FNiagaraVariable> ChannelVars = DataChannel->GetVariables();
+			TConstArrayView<FNiagaraDataChannelVariable> ChannelVars = DataChannel->GetVariables();
 			for (const FNDIDataChannelFunctionInfo& FuncInfo : RuntimeReadDI->GetCompiledData().GetFunctionInfo())
 			{
 				TArray<FNiagaraVariableBase> MissingParams;
@@ -747,7 +747,7 @@ void UNiagaraDataInterfaceDataChannelRead::GetFeedback(UNiagaraSystem* InAsset,
 					for (const FNiagaraVariableBase& FuncParam : Parameters)
 					{
 						bool bParamFound = false;
-						for (const FNiagaraVariable& ChannelVar : ChannelVars)
+						for (const FNiagaraDataChannelVariable& ChannelVar : ChannelVars)
 						{
 							//We have to convert each channel var to SWC for comparison with the function variables as there is no reliable way to go back from the SWC function var to the originating LWC var.
 							FNiagaraVariable SWCVar(ChannelVar);
diff --git a/Engine/Plugins/FX/Niagara/Source/Niagara/Private/DataInterface/NiagaraDataInterfaceDataChannelWrite.cpp b/Engine/Plugins/FX/Niagara/Source/Niagara/Private/DataInterface/NiagaraDataInterfaceDataChannelWrite.cpp
--- a/Engine/Plugins/FX/Niagara/Source/Niagara/Private/DataInterface/NiagaraDataInterfaceDataChannelWrite.cpp
+++ b/Engine/Plugins/FX/Niagara/Source/Niagara/Private/DataInterface/NiagaraDataInterfaceDataChannelWrite.cpp
@@ -489,7 +489,7 @@ void UNiagaraDataInterfaceDataChannelWrite::GetFeedback(UNiagaraSystem* InAsset,
 	if (const UNiagaraDataChannel* DataChannel = RuntimeDI->Channel->Get())
 	{
 		//Ensure the data channel contains all the parameters this function is requesting.
-		TConstArrayView<FNiagaraVariable> ChannelVars = DataChannel->GetVariables();
+		TConstArrayView<FNiagaraDataChannelVariable> ChannelVars = DataChannel->GetVariables();
 		for (const FNDIDataChannelFunctionInfo& FuncInfo : RuntimeDI->GetCompiledData().GetFunctionInfo())
 		{
 			TArray<FNiagaraVariableBase> MissingParams;
@@ -499,7 +499,7 @@ void UNiagaraDataInterfaceDataChannelWrite::GetFeedback(UNiagaraSystem* InAsset,
 				for (const FNiagaraVariableBase& FuncParam : Parameters)
 				{
 					bool bParamFound = false;
-					for (const FNiagaraVariable& ChannelVar : ChannelVars)
+					for (const FNiagaraDataChannelVariable& ChannelVar : ChannelVars)
 					{
 						FNiagaraVariable SWCVar(ChannelVar);
 
diff --git a/Engine/Plugins/FX/Niagara/Source/Niagara/Private/NiagaraComponentRendererProperties.cpp b/Engine/Plugins/FX/Niagara/Source/Niagara/Private/NiagaraComponentRendererProperties.cpp
--- a/Engine/Plugins/FX/Niagara/Source/Niagara/Private/NiagaraComponentRendererProperties.cpp
+++ b/Engine/Plugins/FX/Niagara/Source/Niagara/Private/NiagaraComponentRendererProperties.cpp
@@ -145,44 +145,32 @@ FNiagaraTypeDefinition UNiagaraComponentRendererProperties::GetFRotatorDef()
 
 FNiagaraTypeDefinition UNiagaraComponentRendererProperties::GetFVector2DDef()
 {
-	static UPackage* CoreUObjectPkg = FindObjectChecked<UPackage>(nullptr, TEXT("/Script/CoreUObject"));
-	static UScriptStruct* VectorStruct = FindObjectChecked<UScriptStruct>(CoreUObjectPkg, TEXT("Vector2D"));
-	return FNiagaraTypeDefinition(VectorStruct, FNiagaraTypeDefinition::EAllowUnfriendlyStruct::Allow);
+	return FNiagaraTypeHelper::GetVector2DDef();
 }
 
 FNiagaraTypeDefinition UNiagaraComponentRendererProperties::GetFVectorDef()
 {
-	static UPackage* CoreUObjectPkg = FindObjectChecked<UPackage>(nullptr, TEXT("/Script/CoreUObject"));
-	static UScriptStruct* VectorStruct = FindObjectChecked<UScriptStruct>(CoreUObjectPkg, TEXT("Vector"));
-	return FNiagaraTypeDefinition(VectorStruct, FNiagaraTypeDefinition::EAllowUnfriendlyStruct::Allow);
+	return FNiagaraTypeHelper::GetVectorDef();
 }
 
 FNiagaraTypeDefinition UNiagaraComponentRendererProperties::GetFVector4Def()
 {
-	static UPackage* CoreUObjectPkg = FindObjectChecked<UPackage>(nullptr, TEXT("/Script/CoreUObject"));
-	static UScriptStruct* Vector4Struct = FindObjectChecked<UScriptStruct>(CoreUObjectPkg, TEXT("Vector4"));
-	return FNiagaraTypeDefinition(Vector4Struct, FNiagaraTypeDefinition::EAllowUnfriendlyStruct::Allow);
+	return FNiagaraTypeHelper::GetVector4Def();
 }
 
 FNiagaraTypeDefinition UNiagaraComponentRendererProperties::GetFVector3fDef()
 {
-	static UPackage* CoreUObjectPkg = FindObjectChecked<UPackage>(nullptr, TEXT("/Script/CoreUObject"));
-	static UScriptStruct* VectorStruct = FindObjectChecked<UScriptStruct>(CoreUObjectPkg, TEXT("Vector3f"));
-	return FNiagaraTypeDefinition(VectorStruct, FNiagaraTypeDefinition::EAllowUnfriendlyStruct::Allow);
+	return FNiagaraTypeDefinition::GetVec3Def();
 }
 
 FNiagaraTypeDefinition UNiagaraComponentRendererProperties::GetFQuatDef()
 {
-	static UPackage* CoreUObjectPkg = FindObjectChecked<UPackage>(nullptr, TEXT("/Script/CoreUObject"));
-	static UScriptStruct* Struct = FindObjectChecked<UScriptStruct>(CoreUObjectPkg, TEXT("Quat"));
-	return FNiagaraTypeDefinition(Struct, FNiagaraTypeDefinition::EAllowUnfriendlyStruct::Allow);
+	return FNiagaraTypeHelper::GetQuatDef();
 }
 
 FNiagaraTypeDefinition UNiagaraComponentRendererProperties::GetDoubleDef()
 {
-	static UPackage* NiagaraPkg = FindObjectChecked<UPackage>(nullptr, TEXT("/Script/Niagara"));
-	static UScriptStruct* DoubleStruct = FindObjectChecked<UScriptStruct>(NiagaraPkg, TEXT("NiagaraDouble"));
-	return FNiagaraTypeDefinition(DoubleStruct, FNiagaraTypeDefinition::EAllowUnfriendlyStruct::Allow);
+	return FNiagaraTypeHelper::GetDoubleDef();
 }
 
 TArray<TWeakObjectPtr<UNiagaraComponentRendererProperties>> UNiagaraComponentRendererProperties::ComponentRendererPropertiesToDeferredInit;
diff --git a/Engine/Plugins/FX/Niagara/Source/Niagara/Private/NiagaraDataChannel.cpp b/Engine/Plugins/FX/Niagara/Source/Niagara/Private/NiagaraDataChannel.cpp
--- a/Engine/Plugins/FX/Niagara/Source/Niagara/Private/NiagaraDataChannel.cpp
+++ b/Engine/Plugins/FX/Niagara/Source/Niagara/Private/NiagaraDataChannel.cpp
@@ -25,11 +25,11 @@ namespace NDCCVars
 
 //////////////////////////////////////////////////////////////////////////
 
-void FNiagaraDataChannelGameDataLayout::Init(const TArray<FNiagaraVariable>& Variables)
+void FNiagaraDataChannelGameDataLayout::Init(const TArray<FNiagaraDataChannelVariable>& Variables)
 {
 	VariableIndices.Reset();
 	LwcConverters.Reserve(Variables.Num());
-	for (const FNiagaraVariable& Var : Variables)
+	for (const FNiagaraDataChannelVariable& Var : Variables)
 	{
 		//Sigh.
 		//We must convert from the variable stored var in the data channels definition as we currently cannot serialize/store actual LWC types in FNiagawraTypeDefinitions.
@@ -787,87 +787,65 @@ const FNiagaraDataSetCompiledData& FNiagaraDataChannelData::GetCompiledData(ENia
 
 //////////////////////////////////////////////////////////////////////////
 
-namespace NiagaraDataChannel
-{
-	FNiagaraTypeDefinition GetFVectorDef()
-	{
-		static UPackage* CoreUObjectPkg = FindObjectChecked<UPackage>(nullptr, TEXT("/Script/CoreUObject"));
-		static UScriptStruct* VectorStruct = FindObjectChecked<UScriptStruct>(CoreUObjectPkg, TEXT("Vector"));
-		return FNiagaraTypeDefinition(VectorStruct, FNiagaraTypeDefinition::EAllowUnfriendlyStruct::Allow);
-	}
-	
-	FNiagaraTypeDefinition GetDoubleDef()
-	{
-		static UPackage* NiagaraPkg = FindObjectChecked<UPackage>(nullptr, TEXT("/Script/Niagara"));
-		static UScriptStruct* DoubleStruct = FindObjectChecked<UScriptStruct>(NiagaraPkg, TEXT("NiagaraDouble"));
-		return FNiagaraTypeDefinition(DoubleStruct, FNiagaraTypeDefinition::EAllowUnfriendlyStruct::Allow);
-	}
-
-	FNiagaraTypeDefinition GetFQuatDef()
-	{
-		static UPackage* CoreUObjectPkg = FindObjectChecked<UPackage>(nullptr, TEXT("/Script/CoreUObject"));
-		static UScriptStruct* Struct = FindObjectChecked<UScriptStruct>(CoreUObjectPkg, TEXT("Quat"));
-		return FNiagaraTypeDefinition(Struct, FNiagaraTypeDefinition::EAllowUnfriendlyStruct::Allow);
-	}
-	
-	FNiagaraTypeDefinition GetFVector2DDef()
-	{
-		static UPackage* CoreUObjectPkg = FindObjectChecked<UPackage>(nullptr, TEXT("/Script/CoreUObject"));
-		static UScriptStruct* VectorStruct = FindObjectChecked<UScriptStruct>(CoreUObjectPkg, TEXT("Vector2D"));
-		return FNiagaraTypeDefinition(VectorStruct, FNiagaraTypeDefinition::EAllowUnfriendlyStruct::Allow);
-	}
-
-	FNiagaraTypeDefinition GetFVector4Def()
-	{
-		static UPackage* CoreUObjectPkg = FindObjectChecked<UPackage>(nullptr, TEXT("/Script/CoreUObject"));
-		static UScriptStruct* Vector4Struct = FindObjectChecked<UScriptStruct>(CoreUObjectPkg, TEXT("Vector4"));
-		return FNiagaraTypeDefinition(Vector4Struct, FNiagaraTypeDefinition::EAllowUnfriendlyStruct::Allow);
-	}
-}
-
 void UNiagaraDataChannel::PostInitProperties()
 {
 	Super::PostInitProperties();
 	INiagaraModule::RequestRefreshDataChannels();
 }
 
+void AddVarToHash(const FNiagaraVariable& Var, FBlake3& Builder)
+{
+	FNameBuilder VarName(Var.GetName());
+	FStringView Name = VarName.ToView();
+	uint32 ClassHash = GetTypeHash(Var.GetType().ClassStructOrEnum);
+	Builder.Update(Name.GetData(), Name.Len());
+	Builder.Update(&ClassHash, sizeof(uint32));
+	Builder.Update(&Var.GetType().UnderlyingType, sizeof(uint16));
+}
+
 void UNiagaraDataChannel::PostLoad()
 {
 	Super::PostLoad();
 
-	for (int i = 0; i < Variables.Num(); i++)
+#if WITH_EDITORONLY_DATA
+	static FGuid BaseVersion(TEXT("182b8dd3-f963-477f-a57d-70a449d922d8"));
+	for (const FNiagaraVariable& Var : Variables_DEPRECATED)
 	{
-		// fix up variables deserialized with wrong type
-		// TODO (mga) find a better solution than this
-		FNiagaraVariable& Var = Variables[i];
-		if (Var.GetType() == FNiagaraTypeDefinition::GetVec3Def())
-		{
-			Var.SetType(NiagaraDataChannel::GetFVectorDef());
-		}
-		else if (Var.GetType() == FNiagaraTypeDefinition::GetFloatDef())
-		{
-			Var.SetType(NiagaraDataChannel::GetDoubleDef());
-		}
-		else if (Var.GetType() == FNiagaraTypeDefinition::GetQuatDef())
-		{
-			Var.SetType(NiagaraDataChannel::GetFQuatDef());
-		}
-		else if (Var.GetType() == FNiagaraTypeDefinition::GetVec2Def())
-		{
-			Var.SetType(NiagaraDataChannel::GetFVector2DDef());
-		}
-		else if (Var.GetType() == FNiagaraTypeDefinition::GetVec4Def())
+		FNiagaraDataChannelVariable ChannelVar;
+		ChannelVar.SetName(Var.GetName());
+		ChannelVar.SetType(FNiagaraDataChannelVariable::ToDataChannelType(Var.GetType()));
+
+		FBlake3 VarHashBuilder;
+		VarHashBuilder.Update(&BaseVersion, sizeof(FGuid));
+		AddVarToHash(Var, VarHashBuilder);
+		FBlake3Hash VarHash = VarHashBuilder.Finalize();
+		ChannelVar.Version = FGuid::NewGuidFromHash(VarHash);
+		
+		ChannelVariables.Add(ChannelVar);
+	}
+	Variables_DEPRECATED.Empty();
+	
+	if (!VersionGuid.IsValid())
+	{
+		// If we don't have a guid yet we create one by hashing the existing variables to get a deterministic start guid
+		FBlake3 Builder;
+		Builder.Update(&BaseVersion, sizeof(FGuid));
+		for (const FNiagaraDataChannelVariable& Var : ChannelVariables)
 		{
-			Var.SetType(NiagaraDataChannel::GetFVector4Def());
+			AddVarToHash(Var, Builder);
 		}
+
+		FBlake3Hash Hash = Builder.Finalize();
+		VersionGuid = FGuid::NewGuidFromHash(Hash);
 	}
+#endif
 
 	//Init compiled data. These are not currently serialized as we have no mechanism to rebuild then on internal data format changes like those in scripts.
 	GetCompiledData(ENiagaraSimTarget::CPUSim);
 	GetCompiledData(ENiagaraSimTarget::GPUComputeSim);
 
 	//TODO: Can serialize?
-	GameDataLayout.Init(Variables);
+	GameDataLayout.Init(ChannelVariables);
 
 	INiagaraModule::RequestRefreshDataChannels();
 }
@@ -891,32 +869,38 @@ void UNiagaraDataChannel::PreEditChange(FProperty* PropertyAboutToChange)
 
 void UNiagaraDataChannel::PostEditChangeProperty(FPropertyChangedEvent& PropertyChangedEvent)
 {
-	FName VariablesMemberName = GET_MEMBER_NAME_CHECKED(UNiagaraDataChannel, Variables);
+	FName VariablesMemberName = GET_MEMBER_NAME_CHECKED(UNiagaraDataChannel, ChannelVariables);
 	if (PropertyChangedEvent.ChangeType == EPropertyChangeType::ArrayAdd && PropertyChangedEvent.GetPropertyName() == VariablesMemberName)
 	{
 		TSet<FName> ExistingNames;
-		for (const FNiagaraVariable& Var : Variables)
+		for (const FNiagaraDataChannelVariable& Var : ChannelVariables)
 		{
 			ExistingNames.Add(Var.GetName());
 		}
 		FName UniqueName = FNiagaraUtilities::GetUniqueName(FName("MyNewVar"), ExistingNames);
-		Variables.Last().SetName(UniqueName);
+		ChannelVariables.Last().SetName(UniqueName);
 	}
 	if (PropertyChangedEvent.ChangeType == EPropertyChangeType::Duplicate && PropertyChangedEvent.GetPropertyName() == VariablesMemberName)
 	{
 		int32 ArrayIndex = PropertyChangedEvent.GetArrayIndex(VariablesMemberName.ToString());
-		if (Variables.IsValidIndex(ArrayIndex + 1))
+		if (ChannelVariables.IsValidIndex(ArrayIndex + 1))
 		{
 			TSet<FName> ExistingNames;
-			for (const FNiagaraVariable& Var : Variables)
+			for (const FNiagaraDataChannelVariable& Var : ChannelVariables)
 			{
 				ExistingNames.Add(Var.GetName());
 			}
-			FNiagaraVariable& NewEntry = Variables[ArrayIndex + 1];
+			FNiagaraDataChannelVariable& NewEntry = ChannelVariables[ArrayIndex + 1];
 			FName UniqueName = FNiagaraUtilities::GetUniqueName(NewEntry.GetName(), ExistingNames);
 			NewEntry.SetName(UniqueName);
+			NewEntry.Version = FGuid::NewGuid();
 		}
 	}
+	if (PropertyChangedEvent.GetPropertyName() == VariablesMemberName || PropertyChangedEvent.GetMemberPropertyName() == VariablesMemberName)
+	{
+		VersionGuid = FGuid::NewGuid();
+		// the guid of the variable is updated by the details customization, as we don't want to change it when just renaming a variable
+	}
 	
 	Super::PostEditChangeProperty(PropertyChangedEvent);
 	
@@ -928,7 +912,7 @@ void UNiagaraDataChannel::PostEditChangeProperty(FPropertyChangedEvent& Property
 	GetCompiledData(ENiagaraSimTarget::CPUSim);
 	GetCompiledData(ENiagaraSimTarget::GPUComputeSim);
 
-	GameDataLayout.Init(Variables);
+	GameDataLayout.Init(ChannelVariables);
 
 	INiagaraModule::RequestRefreshDataChannels();
 }
@@ -939,12 +923,12 @@ const FNiagaraDataSetCompiledData& UNiagaraDataChannel::GetCompiledData(ENiagara
 {
 	if(SimTarget == ENiagaraSimTarget::CPUSim)
 	{
-		if (CompiledData.Variables.Num() != Variables.Num())
+		if (CompiledData.Variables.Num() != ChannelVariables.Num())
 		{
 			//Build the compiled data from the current variables but convert to Simulation types.
 			CompiledData.Empty();
 			CompiledData.SimTarget = ENiagaraSimTarget::CPUSim;
-			for (FNiagaraVariableBase Var : Variables)
+			for (FNiagaraVariableBase Var : ChannelVariables)
 			{
 				if(Var.GetType().IsEnum() == false)
 				{
@@ -958,13 +942,13 @@ const FNiagaraDataSetCompiledData& UNiagaraDataChannel::GetCompiledData(ENiagara
 	}
 	else
 	{
-		if (CompiledDataGPU.Variables.Num() != Variables.Num())
+		if (CompiledDataGPU.Variables.Num() != ChannelVariables.Num())
 		{
 			check(SimTarget == ENiagaraSimTarget::GPUComputeSim);
 			//Build the compiled data from the current variables but convert to Simulation types.
 			CompiledDataGPU.Empty();
 			CompiledDataGPU.SimTarget = ENiagaraSimTarget::GPUComputeSim;
-			for (FNiagaraVariableBase Var : Variables)
+			for (FNiagaraVariableBase Var : ChannelVariables)
 			{
 				if (Var.GetType().IsEnum() == false)
 				{
@@ -987,7 +971,7 @@ FNiagaraDataChannelGameDataPtr UNiagaraDataChannel::CreateGameData()const
 
 bool UNiagaraDataChannel::IsValid()const
 {
-	return Variables.Num() > 0 && CompiledData.Variables.Num() == Variables.Num() && CompiledDataGPU.Variables.Num() == Variables.Num();
+	return ChannelVariables.Num() > 0 && CompiledData.Variables.Num() == ChannelVariables.Num() && CompiledDataGPU.Variables.Num() == ChannelVariables.Num();
 }
 
 //////////////////////////////////////////////////////////////////////////
@@ -999,20 +983,25 @@ UNiagaraDataChannelLibrary::UNiagaraDataChannelLibrary(const FObjectInitializer&
 
 UNiagaraDataChannelHandler* UNiagaraDataChannelLibrary::GetNiagaraDataChannel(const UObject* WorldContextObject, const UNiagaraDataChannelAsset* Channel)
 {
-	return GetNiagaraDataChannel(WorldContextObject, Channel->Get());
+	return FindDataChannelHandler(WorldContextObject, Channel->Get());
 }
 
 UNiagaraDataChannelWriter* UNiagaraDataChannelLibrary::WriteToNiagaraDataChannel(const UObject* WorldContextObject, const UNiagaraDataChannelAsset* Channel, FNiagaraDataChannelSearchParameters SearchParams, int32 Count, bool bVisibleToGame, bool bVisibleToCPU, bool bVisibleToGPU)
 {
-	return WriteToNiagaraDataChannel(WorldContextObject, Channel->Get(), SearchParams, Count, bVisibleToGame, bVisibleToCPU, bVisibleToGPU);
+	return CreateDataChannelWriter(WorldContextObject, Channel->Get(), SearchParams, Count, bVisibleToGame, bVisibleToCPU, bVisibleToGPU);
 }
 
 UNiagaraDataChannelReader* UNiagaraDataChannelLibrary::ReadFromNiagaraDataChannel(const UObject* WorldContextObject, const UNiagaraDataChannelAsset* Channel, FNiagaraDataChannelSearchParameters SearchParams, bool bReadPreviousFrame)
 {
-	return ReadFromNiagaraDataChannel(WorldContextObject, Channel->Get(), SearchParams, bReadPreviousFrame);
+	return CreateDataChannelReader(WorldContextObject, Channel->Get(), SearchParams, bReadPreviousFrame);
+}
+
+void UNiagaraDataChannelLibrary::WriteToNiagaraDataChannelSingle(const UObject*, const UNiagaraDataChannelAsset*, FNiagaraDataChannelSearchParameters, bool, bool, bool)
+{
+	// this function is just a placeholder and calls into CreateDataChannelWriter and its individual write functions from the BP node
 }
 
-UNiagaraDataChannelHandler* UNiagaraDataChannelLibrary::GetNiagaraDataChannel(const UObject* WorldContextObject, const UNiagaraDataChannel* Channel)
+UNiagaraDataChannelHandler* UNiagaraDataChannelLibrary::FindDataChannelHandler(const UObject* WorldContextObject, const UNiagaraDataChannel* Channel)
 {
 	UWorld* World = (WorldContextObject != nullptr) ? GEngine->GetWorldFromContextObject(WorldContextObject, EGetWorldErrorMode::LogAndReturnNull) : nullptr;
 	if (World)
@@ -1025,7 +1014,7 @@ UNiagaraDataChannelHandler* UNiagaraDataChannelLibrary::GetNiagaraDataChannel(co
 	return nullptr;
 }
 
-UNiagaraDataChannelWriter* UNiagaraDataChannelLibrary::WriteToNiagaraDataChannel(const UObject* WorldContextObject, const UNiagaraDataChannel* Channel, FNiagaraDataChannelSearchParameters SearchParams, int32 Count, bool bVisibleToGame, bool bVisibleToCPU, bool bVisibleToGPU)
+UNiagaraDataChannelWriter* UNiagaraDataChannelLibrary::CreateDataChannelWriter(const UObject* WorldContextObject, const UNiagaraDataChannel* Channel, FNiagaraDataChannelSearchParameters SearchParams, int32 Count, bool bVisibleToGame, bool bVisibleToCPU, bool bVisibleToGPU)
 {
 	check(IsInGameThread());
 	UWorld* World = (WorldContextObject != nullptr) ? GEngine->GetWorldFromContextObject(WorldContextObject, EGetWorldErrorMode::LogAndReturnNull) : nullptr;
@@ -1048,7 +1037,7 @@ UNiagaraDataChannelWriter* UNiagaraDataChannelLibrary::WriteToNiagaraDataChannel
 	return nullptr;
 }
 
-UNiagaraDataChannelReader* UNiagaraDataChannelLibrary::ReadFromNiagaraDataChannel(const UObject* WorldContextObject, const UNiagaraDataChannel* Channel, FNiagaraDataChannelSearchParameters SearchParams, bool bReadPreviousFrame)
+UNiagaraDataChannelReader* UNiagaraDataChannelLibrary::CreateDataChannelReader(const UObject* WorldContextObject, const UNiagaraDataChannel* Channel, FNiagaraDataChannelSearchParameters SearchParams, bool bReadPreviousFrame)
 {
 	UWorld* World = (WorldContextObject != nullptr) ? GEngine->GetWorldFromContextObject(WorldContextObject, EGetWorldErrorMode::LogAndReturnNull) : nullptr;
 	if (World)
@@ -1094,6 +1083,52 @@ void UNiagaraDataChannelAsset::PostEditChangeProperty(FPropertyChangedEvent& Pro
 
 #endif
 
+bool FNiagaraDataChannelVariable::Serialize(FArchive& Ar)
+{
+	if (FNiagaraVariableBase::Serialize(Ar))
+	{
+		if (Ar.IsLoading())
+		{
+			// fix up variables serialized with wrong type
+			// this happens because we only save swc types
+			SetType(ToDataChannelType(GetType()));
+		}
+		return true;
+	}
+	return false;
+}
+
+#if WITH_EDITORONLY_DATA
+bool FNiagaraDataChannelVariable::IsAllowedType(const FNiagaraTypeDefinition& Type)
+{
+	return !(Type.IsDataInterface() || Type.GetClass() || Type == FNiagaraTypeDefinition::GetParameterMapDef() || Type == FNiagaraTypeDefinition::GetGenericNumericDef() || Type == FNiagaraTypeDefinition::GetHalfDef() || Type == FNiagaraTypeDefinition::GetMatrix4Def());
+}
+#endif
+
+FNiagaraTypeDefinition FNiagaraDataChannelVariable::ToDataChannelType(const FNiagaraTypeDefinition& Type)
+{
+	if (Type == FNiagaraTypeDefinition::GetVec3Def())
+	{
+		return FNiagaraTypeHelper::GetVectorDef();
+	}
+	if (Type == FNiagaraTypeDefinition::GetFloatDef())
+	{
+		return FNiagaraTypeHelper::GetDoubleDef();
+	}
+	if (Type == FNiagaraTypeDefinition::GetQuatDef())
+	{
+		return FNiagaraTypeHelper::GetQuatDef();
+	}
+	if (Type == FNiagaraTypeDefinition::GetVec2Def())
+	{
+		return FNiagaraTypeHelper::GetVector2DDef();
+	}
+	if (Type == FNiagaraTypeDefinition::GetVec4Def())
+	{
+		return FNiagaraTypeHelper::GetVector4Def();
+	}
+	return Type;
+}
 
 //////////////////////////////////////////////////////////////////////////
 
diff --git a/Engine/Plugins/FX/Niagara/Source/Niagara/Private/NiagaraDataChannelAccessor.cpp b/Engine/Plugins/FX/Niagara/Source/Niagara/Private/NiagaraDataChannelAccessor.cpp
--- a/Engine/Plugins/FX/Niagara/Source/Niagara/Private/NiagaraDataChannelAccessor.cpp
+++ b/Engine/Plugins/FX/Niagara/Source/Niagara/Private/NiagaraDataChannelAccessor.cpp
@@ -14,62 +14,6 @@
 
 //////////////////////////////////////////////////////////////////////////
 
-namespace NiagaraDataChannelAccesorLocal
-{
-	//TODO: Move into FNiagaraTypeDefinition.
-	static FNiagaraTypeDefinition DoubleDef;
-	static FNiagaraTypeDefinition Vector2DDef;
-	static FNiagaraTypeDefinition VectorDef;
-	static FNiagaraTypeDefinition Vector4Def;
-	static FNiagaraTypeDefinition QuatDef;
-
-	void InitDefs()
-	{
-		static UPackage* CoreUObjectPkg = FindObjectChecked<UPackage>(nullptr, TEXT("/Script/CoreUObject"));
-		static UScriptStruct* Vector2Struct = FindObjectChecked<UScriptStruct>(CoreUObjectPkg, TEXT("Vector2D"));
-		static UScriptStruct* Vector3Struct = FindObjectChecked<UScriptStruct>(CoreUObjectPkg, TEXT("Vector"));
-		static UScriptStruct* Vector4Struct = FindObjectChecked<UScriptStruct>(CoreUObjectPkg, TEXT("Vector4"));
-		static UScriptStruct* Quat4Struct = FindObjectChecked<UScriptStruct>(CoreUObjectPkg, TEXT("Quat"));
-		DoubleDef = FNiagaraTypeDefinition(FNiagaraDouble::StaticStruct(), FNiagaraTypeDefinition::EAllowUnfriendlyStruct::Allow);
-		Vector2DDef = FNiagaraTypeDefinition(Vector2Struct, FNiagaraTypeDefinition::EAllowUnfriendlyStruct::Allow);
-		VectorDef = FNiagaraTypeDefinition(Vector3Struct, FNiagaraTypeDefinition::EAllowUnfriendlyStruct::Allow);
-		Vector4Def = FNiagaraTypeDefinition(Vector4Struct, FNiagaraTypeDefinition::EAllowUnfriendlyStruct::Allow);
-		QuatDef = FNiagaraTypeDefinition(Quat4Struct, FNiagaraTypeDefinition::EAllowUnfriendlyStruct::Allow);
-	}
-
-	const FNiagaraTypeDefinition& GetDoubleDef()
-	{
-		if(DoubleDef.IsValid() == false){ InitDefs(); }
-		return DoubleDef;
-	}
-
-	const FNiagaraTypeDefinition& GetVector2DDef()
-	{
-		if (Vector2DDef.IsValid() == false) { InitDefs(); }
-		return Vector2DDef;
-	}
-
-	const FNiagaraTypeDefinition& GetVectorDef()
-	{
-		if (VectorDef.IsValid() == false) { InitDefs(); }
-		return VectorDef;
-	}
-
-	const FNiagaraTypeDefinition& GetVector4Def()
-	{
-		if (Vector4Def.IsValid() == false) { InitDefs(); }
-		return Vector4Def;
-	}
-
-	const FNiagaraTypeDefinition& GetQuatDef()
-	{
-		if (QuatDef.IsValid() == false) { InitDefs(); }
-		return QuatDef;
-	}
-}
-
-//////////////////////////////////////////////////////////////////////////
-
 bool UNiagaraDataChannelReader::InitAccess(FNiagaraDataChannelSearchParameters SearchParams, bool bReadPreviousFrameData)
 {
 	Data = nullptr;
@@ -105,35 +49,35 @@ bool UNiagaraDataChannelReader::ReadData(const FNiagaraVariableBase& Var, int32
 double UNiagaraDataChannelReader::ReadFloat(FName VarName, int32 Index)const
 {
 	double RetVal = 0.0f;
-	ReadData<double>(FNiagaraVariableBase(NiagaraDataChannelAccesorLocal::GetDoubleDef(), VarName), Index, RetVal);
+	ReadData<double>(FNiagaraVariableBase(FNiagaraTypeHelper::GetDoubleDef(), VarName), Index, RetVal);
 	return RetVal;
 }
 
 FVector2D UNiagaraDataChannelReader::ReadVector2D(FName VarName, int32 Index)const
 {
 	FVector2D RetVal = FVector2D::ZeroVector;
-	ReadData<FVector2D>(FNiagaraVariableBase(NiagaraDataChannelAccesorLocal::GetVector2DDef(), VarName), Index, RetVal);
+	ReadData<FVector2D>(FNiagaraVariableBase(FNiagaraTypeHelper::GetVector2DDef(), VarName), Index, RetVal);
 	return RetVal;
 }
 
 FVector UNiagaraDataChannelReader::ReadVector(FName VarName, int32 Index)const
 {
 	FVector RetVal = FVector::ZeroVector;
-	ReadData<FVector>(FNiagaraVariableBase(NiagaraDataChannelAccesorLocal::GetVectorDef(), VarName), Index, RetVal);
+	ReadData<FVector>(FNiagaraVariableBase(FNiagaraTypeHelper::GetVectorDef(), VarName), Index, RetVal);
 	return RetVal;
 }
 
 FVector4 UNiagaraDataChannelReader::ReadVector4(FName VarName, int32 Index)const
 {
 	FVector4 RetVal = FVector4(0.0f);
-	ReadData<FVector4>(FNiagaraVariableBase(NiagaraDataChannelAccesorLocal::GetVector4Def(), VarName), Index, RetVal);
+	ReadData<FVector4>(FNiagaraVariableBase(FNiagaraTypeHelper::GetVector4Def(), VarName), Index, RetVal);
 	return RetVal;
 }
 
 FQuat UNiagaraDataChannelReader::ReadQuat(FName VarName, int32 Index)const
 {
 	FQuat RetVal = FQuat::Identity;
-	ReadData<FQuat>(FNiagaraVariableBase(NiagaraDataChannelAccesorLocal::GetQuatDef(), VarName), Index, RetVal);
+	ReadData<FQuat>(FNiagaraVariableBase(FNiagaraTypeHelper::GetQuatDef(), VarName), Index, RetVal);
 	return RetVal;
 }
 
@@ -151,6 +95,11 @@ int32 UNiagaraDataChannelReader::ReadInt(FName VarName, int32 Index)const
 	return RetVal;
 }
 
+uint8 UNiagaraDataChannelReader::ReadEnum(FName VarName, int32 Index) const
+{
+	return static_cast<uint8>(ReadInt(VarName, Index));
+}
+
 bool UNiagaraDataChannelReader::ReadBool(FName VarName, int32 Index)const
 {
 	bool RetVal = false;
@@ -167,7 +116,6 @@ FVector UNiagaraDataChannelReader::ReadPosition(FName VarName, int32 Index)const
 
 //////////////////////////////////////////////////////////////////////////
 
-
 template<typename T>
 void UNiagaraDataChannelWriter::WriteData(const FNiagaraVariableBase& Var, int32 Index, const T& InData)
 {
@@ -220,27 +168,27 @@ int32 UNiagaraDataChannelWriter::Num()const
 
 void UNiagaraDataChannelWriter::WriteFloat(FName VarName, int32 Index, double InData)
 {
-	WriteData(FNiagaraVariableBase(NiagaraDataChannelAccesorLocal::GetDoubleDef(), VarName), Index, InData);
+	WriteData(FNiagaraVariableBase(FNiagaraTypeHelper::GetDoubleDef(), VarName), Index, InData);
 }
 
 void UNiagaraDataChannelWriter::WriteVector2D(FName VarName, int32 Index, FVector2D InData)
 {
-	WriteData(FNiagaraVariableBase(NiagaraDataChannelAccesorLocal::GetVector2DDef(), VarName), Index, InData);
+	WriteData(FNiagaraVariableBase(FNiagaraTypeHelper::GetVector2DDef(), VarName), Index, InData);
 }
 
 void UNiagaraDataChannelWriter::WriteVector(FName VarName, int32 Index, FVector InData)
 {
-	WriteData(FNiagaraVariableBase(NiagaraDataChannelAccesorLocal::GetVectorDef(), VarName), Index, InData);
+	WriteData(FNiagaraVariableBase(FNiagaraTypeHelper::GetVectorDef(), VarName), Index, InData);
 }
 
 void UNiagaraDataChannelWriter::WriteVector4(FName VarName, int32 Index, FVector4 InData)
 {
-	WriteData(FNiagaraVariableBase(NiagaraDataChannelAccesorLocal::GetVector4Def(), VarName), Index, InData);
+	WriteData(FNiagaraVariableBase(FNiagaraTypeHelper::GetVector4Def(), VarName), Index, InData);
 }
 
 void UNiagaraDataChannelWriter::WriteQuat(FName VarName, int32 Index, FQuat InData)
 {
-	WriteData(FNiagaraVariableBase(NiagaraDataChannelAccesorLocal::GetQuatDef(), VarName), Index, InData);
+	WriteData(FNiagaraVariableBase(FNiagaraTypeHelper::GetQuatDef(), VarName), Index, InData);
 }
 
 void UNiagaraDataChannelWriter::WriteLinearColor(FName VarName, int32 Index, FLinearColor InData)
@@ -253,6 +201,11 @@ void UNiagaraDataChannelWriter::WriteInt(FName VarName, int32 Index, int32 InDat
 	WriteData(FNiagaraVariableBase(FNiagaraTypeDefinition::GetIntDef(), VarName), Index, InData);
 }
 
+void UNiagaraDataChannelWriter::WriteEnum(FName VarName, int32 Index, uint8 InData)
+{
+	WriteInt(VarName, Index, InData);
+}
+
 void UNiagaraDataChannelWriter::WriteBool(FName VarName, int32 Index, bool InData)
 {
 	WriteData(FNiagaraVariableBase(FNiagaraTypeDefinition::GetBoolDef(), VarName), Index, FNiagaraBool(InData));
diff --git a/Engine/Plugins/FX/Niagara/Source/Niagara/Private/NiagaraModule.cpp b/Engine/Plugins/FX/Niagara/Source/Niagara/Private/NiagaraModule.cpp
--- a/Engine/Plugins/FX/Niagara/Source/Niagara/Private/NiagaraModule.cpp
+++ b/Engine/Plugins/FX/Niagara/Source/Niagara/Private/NiagaraModule.cpp
@@ -283,6 +283,7 @@ void INiagaraModule::StartupModule()
 
 	LLM_SCOPE(ELLMTag::Niagara);
 	FNiagaraTypeDefinition::Init();
+	FNiagaraTypeHelper::RegisterStaticTypes();
 
 	if (IsRunningCookCommandlet())
 	{
diff --git a/Engine/Plugins/FX/Niagara/Source/Niagara/Private/NiagaraTypes.cpp b/Engine/Plugins/FX/Niagara/Source/Niagara/Private/NiagaraTypes.cpp
--- a/Engine/Plugins/FX/Niagara/Source/Niagara/Private/NiagaraTypes.cpp
+++ b/Engine/Plugins/FX/Niagara/Source/Niagara/Private/NiagaraTypes.cpp
@@ -369,11 +369,10 @@ FString FNiagaraTypeHelper::ToString(const uint8* ValueData, const UObject* Stru
 
 FNiagaraLwcStructConverter BuildSWCStructure(UScriptStruct* NewStruct, UScriptStruct* InStruct)
 {
-	static UPackage* CoreUObjectPkg = FindObjectChecked<UPackage>(nullptr, TEXT("/Script/CoreUObject"));
-	static UScriptStruct* Vector2fStruct = FindObjectChecked<UScriptStruct>(CoreUObjectPkg, TEXT("Vector2f"));
-	static UScriptStruct* Vector3fStruct = FindObjectChecked<UScriptStruct>(CoreUObjectPkg, TEXT("Vector3f"));
-	static UScriptStruct* Vector4fStruct = FindObjectChecked<UScriptStruct>(CoreUObjectPkg, TEXT("Vector4f"));
-	static UScriptStruct* Quat4fStruct = FindObjectChecked<UScriptStruct>(CoreUObjectPkg, TEXT("Quat4f"));
+	static UScriptStruct* Vector2fStruct = FNiagaraTypeDefinition::GetVec2Struct();
+	static UScriptStruct* Vector3fStruct = FNiagaraTypeDefinition::GetVec3Struct();
+	static UScriptStruct* Vector4fStruct = FNiagaraTypeDefinition::GetVec4Struct();
+	static UScriptStruct* Quat4fStruct = FNiagaraTypeDefinition::GetQuatStruct();
 
 	FNiagaraLwcStructConverter StructConverter;
 	int32 AlignedOffset = 0;
@@ -554,34 +553,25 @@ UScriptStruct* FNiagaraTypeHelper::GetSWCStruct(UScriptStruct* LWCStruct)
 			if (IsLWCStructure(LWCStruct))
 			{
 				//Return common SWC variant of common LWC Types.
-				static UPackage* CoreUObjectPkg = FindObjectChecked<UPackage>(nullptr, TEXT("/Script/CoreUObject"));
-				static UScriptStruct* Vector2Struct = FindObjectChecked<UScriptStruct>(CoreUObjectPkg, TEXT("Vector2D"));
-				static UScriptStruct* Vector3Struct = FindObjectChecked<UScriptStruct>(CoreUObjectPkg, TEXT("Vector"));
-				static UScriptStruct* Vector4Struct = FindObjectChecked<UScriptStruct>(CoreUObjectPkg, TEXT("Vector4"));
-				static UScriptStruct* QuatStruct = FindObjectChecked<UScriptStruct>(CoreUObjectPkg, TEXT("Quat"));
-				if(LWCStruct == FNiagaraDouble::StaticStruct())
+				if (LWCStruct == FNiagaraDouble::StaticStruct())
 				{
 					SWCStruct = FNiagaraFloat::StaticStruct();
 				}
-				else if(LWCStruct == Vector2Struct) 
+				else if(LWCStruct == GetVector2DDef().GetStruct()) 
 				{
-					static UScriptStruct* Vector2fStruct = FindObjectChecked<UScriptStruct>(CoreUObjectPkg, TEXT("Vector2f"));
-					SWCStruct = Vector2fStruct;
+					SWCStruct = FNiagaraTypeDefinition::GetVec2Struct();
 				}
-				else if (LWCStruct == Vector3Struct)
+				else if (LWCStruct == GetVectorDef().GetStruct())
 				{
-					static UScriptStruct* Vector3fStruct = FindObjectChecked<UScriptStruct>(CoreUObjectPkg, TEXT("Vector3f"));
-					SWCStruct = Vector3fStruct;
+					SWCStruct = FNiagaraTypeDefinition::GetVec3Struct();
 				}
-				else if (LWCStruct == Vector4Struct)
+				else if (LWCStruct == GetVector4Def().GetStruct())
 				{
-					static UScriptStruct* Vector4fStruct = FindObjectChecked<UScriptStruct>(CoreUObjectPkg, TEXT("Vector4f"));
-					SWCStruct = Vector4fStruct;
+					SWCStruct = FNiagaraTypeDefinition::GetVec4Struct();
 				}
-				else if (LWCStruct == QuatStruct)
+				else if (LWCStruct == GetQuatDef().GetStruct())
 				{
-					static UScriptStruct* Quat4fStruct = FindObjectChecked<UScriptStruct>(CoreUObjectPkg, TEXT("Quat4f"));
-					SWCStruct = Quat4fStruct;
+					SWCStruct = FNiagaraTypeDefinition::GetQuatStruct();
 				}
 				//More?
 				else
@@ -619,7 +609,7 @@ UScriptStruct* FNiagaraTypeHelper::GetSWCStruct(UScriptStruct* LWCStruct)
 			{
 				SWCStruct = LWCStruct;
 			}
-			FRemapEntry& RemapEntry = RemapTable.Add(LWCStruct);
+			FRemapEntry& RemapEntry = RemapTable.Add(TWeakObjectPtr<UScriptStruct>(LWCStruct));
 			RemapEntry.Struct = SWCStruct;
 		#if WITH_EDITORONLY_DATA
 			RemapEntry.SerialNumber = LWCStruct->FieldPathSerialNumber;
@@ -718,6 +708,15 @@ void FNiagaraTypeHelper::InitStaticTypes()
 	}
 }
 
+void FNiagaraTypeHelper::RegisterStaticTypes()
+{
+	FNiagaraTypeRegistry::Register(Vector2DDef, ENiagaraTypeRegistryFlags::None);
+	FNiagaraTypeRegistry::Register(VectorDef, ENiagaraTypeRegistryFlags::None);
+	FNiagaraTypeRegistry::Register(Vector4Def, ENiagaraTypeRegistryFlags::None);
+	FNiagaraTypeRegistry::Register(QuatDef, ENiagaraTypeRegistryFlags::None);
+	FNiagaraTypeRegistry::Register(DoubleDef, ENiagaraTypeRegistryFlags::None);
+}
+
 FNiagaraTypeDefinition FNiagaraTypeHelper::GetVector2DDef()
 {
 	return Vector2DDef;
diff --git a/Engine/Plugins/FX/Niagara/Source/Niagara/Public/NiagaraDataChannel.h b/Engine/Plugins/FX/Niagara/Source/Niagara/Public/NiagaraDataChannel.h
--- a/Engine/Plugins/FX/Niagara/Source/Niagara/Public/NiagaraDataChannel.h
+++ b/Engine/Plugins/FX/Niagara/Source/Niagara/Public/NiagaraDataChannel.h
@@ -88,32 +88,32 @@ public:
 	GENERATED_BODY()
 
 	//UObject Interface
-	NIAGARA_API virtual void PostInitProperties()override;
-	NIAGARA_API virtual void PostLoad();
-	NIAGARA_API virtual void BeginDestroy();
+	NIAGARA_API virtual void PostInitProperties() override;
+	NIAGARA_API virtual void PostLoad() override;
+	NIAGARA_API virtual void BeginDestroy() override;
 #if WITH_EDITOR
-	NIAGARA_API virtual void PreEditChange(FProperty* PropertyAboutToChange)override;
-	NIAGARA_API virtual void PostEditChangeProperty(FPropertyChangedEvent& PropertyChangedDataChannel)override;
+	NIAGARA_API virtual void PreEditChange(FProperty* PropertyAboutToChange) override;
+	NIAGARA_API virtual void PostEditChangeProperty(FPropertyChangedEvent& PropertyChangedDataChannel) override;
 #endif
 	//UObject Interface End.
 
-	const UNiagaraDataChannelAsset* GetAsset()const { return CastChecked<UNiagaraDataChannelAsset>(GetOuter()); }
-	TConstArrayView<FNiagaraVariable> GetVariables()const { return Variables; }
+	const UNiagaraDataChannelAsset* GetAsset() const { return CastChecked<UNiagaraDataChannelAsset>(GetOuter()); }
+	TConstArrayView<FNiagaraDataChannelVariable> GetVariables() const { return ChannelVariables; }
 
 	/** If true, we keep our previous frame's data. Some users will prefer a frame of latency to tick dependency. */
-	bool KeepPreviousFrameData()const { return bKeepPreviousFrameData; }
+	bool KeepPreviousFrameData() const { return bKeepPreviousFrameData; }
 
 	/** Returns the compiled data describing the data layout for DataChannels in this channel. */
-	NIAGARA_API const FNiagaraDataSetCompiledData& GetCompiledData(ENiagaraSimTarget SimTarget)const;
+	NIAGARA_API const FNiagaraDataSetCompiledData& GetCompiledData(ENiagaraSimTarget SimTarget) const;
 
 	/** Create the appropriate handler object for this data channel. */
 	NIAGARA_API virtual UNiagaraDataChannelHandler* CreateHandler(UWorld* OwningWorld) const PURE_VIRTUAL(UNiagaraDataChannel::CreateHandler, {return nullptr;} );
 	
-	const FNiagaraDataChannelGameDataLayout& GetGameDataLayout()const { return GameDataLayout; }
+	const FNiagaraDataChannelGameDataLayout& GetGameDataLayout() const { return GameDataLayout; }
 
-	NIAGARA_API FNiagaraDataChannelGameDataPtr CreateGameData()const;
+	NIAGARA_API FNiagaraDataChannelGameDataPtr CreateGameData() const;
 
-	bool IsValid()const;
+	bool IsValid() const;
 
 	#if !UE_BUILD_SHIPPING
 	void SetVerboseLogging(bool bValue){ bVerboseLogging = bValue; }
@@ -123,16 +123,23 @@ public:
 	template<typename TFunc>
 	static void ForEachDataChannel(TFunc Func);
 
-	bool ShouldEnforceTickGroupReadWriteOrder()const {return bEnforceTickGroupReadWriteOrder;}
+	bool ShouldEnforceTickGroupReadWriteOrder() const {return bEnforceTickGroupReadWriteOrder;}
 	
 	/** If we are enforcing tick group read/write ordering the this returns the final tick group that this NDC can be written to. All reads must happen in Tick groups after this or next frame. */
-	ETickingGroup GetFinalWriteTickGroup()const { return FinalWriteTickGroup; }
+	ETickingGroup GetFinalWriteTickGroup() const { return FinalWriteTickGroup; }
 
+#if WITH_EDITORONLY_DATA
+	/** Can be used to track structural changes that would need recompilation of downstream assets. */
+	NIAGARA_API FGuid GetVersion() const { return VersionGuid; }
+#endif
+	
 private:
 
+	//TODO: add default values for editor previews
+	
 	/** The variables that define the data contained in this Data Channel. */
 	UPROPERTY(EditAnywhere, Category = "Data Channel", meta=(EnforceUniqueNames = true))
-	TArray<FNiagaraVariable> Variables;
+	TArray<FNiagaraDataChannelVariable> ChannelVariables;
 
 	/** If true, we keep our previous frame's data. This comes at a memory and performance cost but allows users to avoid tick order dependency by reading last frame's data. Some users will prefer a frame of latency to tick order dependency. */
 	UPROPERTY(EditAnywhere, Category = "Data Channel")
@@ -146,6 +153,14 @@ private:
 	UPROPERTY(EditAnywhere, Category = "Data Channel", meta=(EditCondition="bEnforceTickGroupReadWriteOrder"))
 	TEnumAsByte<ETickingGroup> FinalWriteTickGroup = ETickingGroup::TG_EndPhysics;
 
+#if WITH_EDITORONLY_DATA
+	UPROPERTY()
+	FGuid VersionGuid;
+
+	UPROPERTY(meta=(DeprecatedProperty))
+	TArray<FNiagaraVariable> Variables_DEPRECATED;
+#endif
+
 	/**
 	Data layout for payloads in Niagara datasets.
 	*/
@@ -197,7 +212,7 @@ class NIAGARA_API UNiagaraDataChannelLibrary : public UBlueprintFunctionLibrary
 	 * @param bVisibleToCPU	If true, the data written to this data channel is visible to Niagara CPU emitters
 	 * @param bVisibleToGPU	If true, the data written to this data channel is visible to Niagara GPU emitters
 	 */
-	UFUNCTION(BlueprintCallable, Category = NiagaraDataChannel, meta = (Keywords = "niagara DataChannel", WorldContext = "WorldContextObject", UnsafeDuringActorConstruction = "true"))
+	UFUNCTION(BlueprintCallable, Category = NiagaraDataChannel, DisplayName="Write To Niagara Data Channel (Batch)", meta = (AdvancedDisplay = "SearchParams", Keywords = "niagara DataChannel", WorldContext = "WorldContextObject", UnsafeDuringActorConstruction = "true"))
 	static UNiagaraDataChannelWriter* WriteToNiagaraDataChannel(const UObject* WorldContextObject, const UNiagaraDataChannelAsset* Channel, FNiagaraDataChannelSearchParameters SearchParams, int32 Count, UPARAM(DisplayName = "Visible to Blueprint") bool bVisibleToGame, UPARAM(DisplayName = "Visible to Niagara CPU") bool bVisibleToCPU, UPARAM(DisplayName = "Visible to Niagara GPU") bool bVisibleToGPU);
 
 	/**
@@ -213,7 +228,20 @@ class NIAGARA_API UNiagaraDataChannelLibrary : public UBlueprintFunctionLibrary
 	UFUNCTION(BlueprintCallable, Category = NiagaraDataChannel, meta = (Keywords = "niagara DataChannel", WorldContext = "WorldContextObject", UnsafeDuringActorConstruction = "true"))
 	static UNiagaraDataChannelReader* ReadFromNiagaraDataChannel(const UObject* WorldContextObject, const UNiagaraDataChannelAsset* Channel, FNiagaraDataChannelSearchParameters SearchParams, bool bReadPreviousFrame);
 
-	static UNiagaraDataChannelHandler* GetNiagaraDataChannel(const UObject* WorldContextObject, const UNiagaraDataChannel* Channel);
-	static UNiagaraDataChannelWriter* WriteToNiagaraDataChannel(const UObject* WorldContextObject, const UNiagaraDataChannel* Channel, FNiagaraDataChannelSearchParameters SearchParams, int32 Count, bool bVisibleToGame, bool bVisibleToCPU, bool bVisibleToGPU);
-	static UNiagaraDataChannelReader* ReadFromNiagaraDataChannel(const UObject* WorldContextObject, const UNiagaraDataChannel* Channel, FNiagaraDataChannelSearchParameters SearchParams, bool bReadPreviousFrame);
+	/**
+	 * Writes a single element to a Niagara Data Channel
+	 *
+	 * @param WorldContextObject	World to execute in
+	 * @param Channel				The channel to write to
+	 * @param SearchParams			Parameters used when retrieving a specific set of Data Channel Data to read or write like the islands data channel type.
+	 * @param bVisibleToBlueprint	If true, the data written to this data channel is visible to Blueprint and C++ logic reading from it
+	 * @param bVisibleToNiagaraCPU	If true, the data written to this data channel is visible to Niagara CPU emitters
+	 * @param bVisibleToNiagaraGPU	If true, the data written to this data channel is visible to Niagara GPU emitters
+	 */
+	UFUNCTION(BlueprintInternalUseOnly, Category = NiagaraDataChannel, DisplayName="Write To Niagara Data Channel", meta = (bVisibleToBlueprint="true", bVisibleToNiagaraCPU="true", bVisibleToNiagaraGPU="true", AdvancedDisplay = "bVisibleToBlueprint, bVisibleToNiagaraCPU, bVisibleToNiagaraGPU, SearchParams", Keywords = "niagara DataChannel event writer", WorldContext = "WorldContextObject", UnsafeDuringActorConstruction = "true"))
+	static void WriteToNiagaraDataChannelSingle(const UObject* WorldContextObject, const UNiagaraDataChannelAsset* Channel, FNiagaraDataChannelSearchParameters SearchParams, bool bVisibleToBlueprint, bool bVisibleToNiagaraCPU, bool bVisibleToNiagaraGPU);
+
+	static UNiagaraDataChannelHandler* FindDataChannelHandler(const UObject* WorldContextObject, const UNiagaraDataChannel* Channel);
+	static UNiagaraDataChannelWriter* CreateDataChannelWriter(const UObject* WorldContextObject, const UNiagaraDataChannel* Channel, FNiagaraDataChannelSearchParameters SearchParams, int32 Count, bool bVisibleToGame, bool bVisibleToCPU, bool bVisibleToGPU);
+	static UNiagaraDataChannelReader* CreateDataChannelReader(const UObject* WorldContextObject, const UNiagaraDataChannel* Channel, FNiagaraDataChannelSearchParameters SearchParams, bool bReadPreviousFrame);
 };
diff --git a/Engine/Plugins/FX/Niagara/Source/Niagara/Public/NiagaraDataChannelAccessor.h b/Engine/Plugins/FX/Niagara/Source/Niagara/Public/NiagaraDataChannelAccessor.h
--- a/Engine/Plugins/FX/Niagara/Source/Niagara/Public/NiagaraDataChannelAccessor.h
+++ b/Engine/Plugins/FX/Niagara/Source/Niagara/Public/NiagaraDataChannelAccessor.h
@@ -61,6 +61,9 @@ public:
 	UFUNCTION(BlueprintCallable, Category = NiagaraDataChannel, meta = (Keywords = "Niagara DataChannel"))
 	NIAGARA_API int32 ReadInt(FName VarName, int32 Index)const;
 
+	UFUNCTION(BlueprintCallable, Category = NiagaraDataChannel, meta = (Keywords = "Niagara DataChannel"))
+	NIAGARA_API uint8 ReadEnum(FName VarName, int32 Index)const;
+
 	UFUNCTION(BlueprintCallable, Category = NiagaraDataChannel, meta = (Keywords = "Niagara DataChannel"))
 	NIAGARA_API bool ReadBool(FName VarName, int32 Index)const;
 
@@ -113,6 +116,9 @@ public:
 	UFUNCTION(BlueprintCallable, Category = NiagaraDataChannel, meta = (Keywords = "niagara DataChannel"))
 	NIAGARA_API void WriteInt(FName VarName, int32 Index, int32 InData);
 
+	UFUNCTION(BlueprintCallable, Category = NiagaraDataChannel, meta = (Keywords = "niagara DataChannel"))
+	NIAGARA_API void WriteEnum(FName VarName, int32 Index, uint8 InData);
+
 	UFUNCTION(BlueprintCallable, Category = NiagaraDataChannel, meta = (Keywords = "niagara DataChannel"))
 	NIAGARA_API void WriteBool(FName VarName, int32 Index, bool InData);
 
diff --git a/Engine/Plugins/FX/Niagara/Source/Niagara/Public/NiagaraDataChannelPublic.h b/Engine/Plugins/FX/Niagara/Source/Niagara/Public/NiagaraDataChannelPublic.h
--- a/Engine/Plugins/FX/Niagara/Source/Niagara/Public/NiagaraDataChannelPublic.h
+++ b/Engine/Plugins/FX/Niagara/Source/Niagara/Public/NiagaraDataChannelPublic.h
@@ -48,6 +48,35 @@ public:
 	UNiagaraDataChannel* Get() const { return DataChannel; }
 };
 
+USTRUCT(BlueprintType)
+struct FNiagaraDataChannelVariable : public FNiagaraVariableBase
+{
+	GENERATED_USTRUCT_BODY()
+	
+	bool Serialize(FArchive& Ar);
+
+#if WITH_EDITORONLY_DATA
+	
+	/** Can be used to track renamed data channel variables */
+	UPROPERTY(meta=(IgnoreForMemberInitializationTest))
+	FGuid Version = FGuid::NewGuid();
+
+	NIAGARA_API static bool IsAllowedType(const FNiagaraTypeDefinition& Type);
+#endif
+	
+	NIAGARA_API static FNiagaraTypeDefinition ToDataChannelType(const FNiagaraTypeDefinition& Type);
+
+};
+
+template<>
+struct TStructOpsTypeTraits<FNiagaraDataChannelVariable> : public TStructOpsTypeTraitsBase2<FNiagaraDataChannelVariable>
+{
+	enum
+	{
+		WithSerializer = true,
+	};
+};
+
 /**
 Minimal set of types and declares required for external users of Niagara Data Channels.
 */
@@ -117,7 +146,7 @@ struct FNiagaraDataChannelGameDataLayout
 	UPROPERTY()
 	TArray<FNiagaraLwcStructConverter> LwcConverters;
 
-	void Init(const TArray<FNiagaraVariable>& Variables);
+	void Init(const TArray<FNiagaraDataChannelVariable>& Variables);
 };
 
 
diff --git a/Engine/Plugins/FX/Niagara/Source/Niagara/Public/NiagaraTypes.h b/Engine/Plugins/FX/Niagara/Source/Niagara/Public/NiagaraTypes.h
--- a/Engine/Plugins/FX/Niagara/Source/Niagara/Public/NiagaraTypes.h
+++ b/Engine/Plugins/FX/Niagara/Source/Niagara/Public/NiagaraTypes.h
@@ -438,6 +438,7 @@ public:
 	static NIAGARA_API FNiagaraTypeDefinition GetDoubleDef();
 
 	static NIAGARA_API void InitStaticTypes();
+	static NIAGARA_API void RegisterStaticTypes();
 private:
 	struct FRemapEntry
 	{
diff --git a/Engine/Plugins/FX/Niagara/Source/NiagaraBlueprintNodes/NiagaraBlueprintNodes.build.cs b/Engine/Plugins/FX/Niagara/Source/NiagaraBlueprintNodes/NiagaraBlueprintNodes.build.cs
new file mode 100644
--- /dev/null
+++ b/Engine/Plugins/FX/Niagara/Source/NiagaraBlueprintNodes/NiagaraBlueprintNodes.build.cs
@@ -0,0 +1,33 @@
+// Copyright Epic Games, Inc. All Rights Reserved.
+
+using UnrealBuildTool;
+
+public class NiagaraBlueprintNodes : ModuleRules
+{
+	public NiagaraBlueprintNodes(ReadOnlyTargetRules Target) : base(Target)
+	{
+		PCHUsage = ModuleRules.PCHUsageMode.UseExplicitOrSharedPCHs;
+		
+		PublicDependencyModuleNames.AddRange(
+			new string[]
+			{
+				"Core",
+			}
+		);
+			
+		
+		PrivateDependencyModuleNames.AddRange(
+			new string[]
+			{
+				"BlueprintGraph",
+				"CoreUObject",
+				"UnrealEd",
+				"Engine",
+				"NiagaraCore",
+				"Niagara",
+				"Kismet",
+				"KismetCompiler",
+			}
+		);
+	}
+}
diff --git a/Engine/Plugins/FX/Niagara/Source/NiagaraBlueprintNodes/Private/K2Node_WriteDataChannel.cpp b/Engine/Plugins/FX/Niagara/Source/NiagaraBlueprintNodes/Private/K2Node_WriteDataChannel.cpp
new file mode 100644
--- /dev/null
+++ b/Engine/Plugins/FX/Niagara/Source/NiagaraBlueprintNodes/Private/K2Node_WriteDataChannel.cpp
@@ -0,0 +1,255 @@
+﻿// Copyright Epic Games, Inc. All Rights Reserved.
+
+#include "K2Node_WriteDataChannel.h"
+
+#include "BlueprintActionDatabaseRegistrar.h"
+#include "BlueprintNodeSpawner.h"
+#include "KismetCompiler.h"
+#include "NiagaraBlueprintUtil.h"
+#include "NiagaraDataChannel.h"
+#include "NiagaraDataChannelAccessor.h"
+
+#define LOCTEXT_NAMESPACE "K2Node_WriteDataChannel"
+
+UK2Node_WriteDataChannel::UK2Node_WriteDataChannel()
+{
+	FunctionReference.SetExternalMember(GET_FUNCTION_NAME_CHECKED(UNiagaraDataChannelLibrary, WriteToNiagaraDataChannelSingle), UNiagaraDataChannelLibrary::StaticClass());
+}
+
+void UK2Node_WriteDataChannel::PostLoad()
+{
+	Super::PostLoad();
+
+#if WITH_EDITORONLY_DATA
+	if (DataChannel && DataChannel->Get()->GetVersion() != DataChannelVersion && HasValidBlueprint())
+	{
+		ReconstructNode();
+	}
+#endif
+}
+
+void UK2Node_WriteDataChannel::AllocateDefaultPins()
+{
+	Super::AllocateDefaultPins();
+
+	if (DataChannel)
+	{
+		for (const FNiagaraDataChannelVariable& InVar : DataChannel->Get()->GetVariables())
+		{
+			UEdGraphPin* NewPin = CreatePin(EGPD_Input, FNiagaraBlueprintUtil::TypeDefinitionToBlueprintType(InVar.GetType()), InVar.GetName());
+
+#if WITH_EDITORONLY_DATA
+			NewPin->PersistentGuid = InVar.Version;
+#endif
+		}
+
+#if WITH_EDITORONLY_DATA
+		DataChannelVersion = DataChannel->Get()->GetVersion();
+#endif
+	}
+}
+
+void UK2Node_WriteDataChannel::PinDefaultValueChanged(UEdGraphPin* Pin)
+{
+	Super::PinDefaultValueChanged(Pin);
+
+#if WITH_EDITORONLY_DATA
+	if (Pin == GetChannelSelectorPin())
+	{
+		if (UNiagaraDataChannelAsset* ChannelAsset = Cast<UNiagaraDataChannelAsset>(Pin->DefaultObject))
+		{
+			DataChannel = ChannelAsset;
+			ReconstructNode();
+		}
+	}
+#endif
+}
+
+void UK2Node_WriteDataChannel::PinConnectionListChanged(UEdGraphPin* Pin)
+{
+	Super::PinConnectionListChanged(Pin);
+
+#if WITH_EDITORONLY_DATA
+	if (Pin == GetChannelSelectorPin())
+	{
+		if (UNiagaraDataChannelAsset* ChannelAsset = Cast<UNiagaraDataChannelAsset>(Pin->DefaultObject))
+		{
+			DataChannel = Pin->LinkedTo.Num() == 0 ? ChannelAsset : nullptr;
+			ReconstructNode();
+		}
+	}
+#endif
+}
+
+void UK2Node_WriteDataChannel::GetMenuActions(FBlueprintActionDatabaseRegistrar& InActionRegistrar) const
+{
+	const UClass* ActionKey = GetClass();
+	if (InActionRegistrar.IsOpenForRegistration(ActionKey))
+	{
+		UBlueprintNodeSpawner* NodeSpawner = UBlueprintNodeSpawner::Create(GetClass());
+		InActionRegistrar.AddBlueprintAction(ActionKey, NodeSpawner);
+	}
+}
+
+FText UK2Node_WriteDataChannel::GetMenuCategory() const
+{
+	static FText MenuCategory = LOCTEXT("MenuCategory", "Niagara Data Channel");
+	return MenuCategory;
+}
+
+void UK2Node_WriteDataChannel::ExpandNode(FKismetCompilerContext& CompilerContext, UEdGraph* SourceGraph)
+{
+	if (!DataChannel)
+	{
+		return;
+	}
+	ExpandSplitPins(CompilerContext, SourceGraph);
+	
+	// create function call node to init the writer object 
+	UK2Node_CallFunction* CreateWriterNode = CompilerContext.SpawnIntermediateNode<UK2Node_CallFunction>(this, SourceGraph);
+	CreateWriterNode->SetFromFunction(UNiagaraDataChannelLibrary::StaticClass()->FindFunctionByName(GET_FUNCTION_NAME_CHECKED(UNiagaraDataChannelLibrary, WriteToNiagaraDataChannel)));
+	CreateWriterNode->AllocateDefaultPins();
+
+	// transfer the input pins over
+	static TArray<TPair<FName, FName>> PinsToTransfer = { {"Channel", "Channel"}, {"SearchParams", "SearchParams"}, {"bVisibleToBlueprint", "bVisibleToGame"},
+		{"bVisibleToNiagaraCPU", "bVisibleToCPU"}, {"bVisibleToNiagaraGPU", "bVisibleToGPU"}};
+	
+	for (TPair<FName, FName> Pair : PinsToTransfer)
+	{
+		UEdGraphPin* OrgInputPin = FindPinChecked(Pair.Key, EGPD_Input);
+		UEdGraphPin* NewInputPin = CreateWriterNode->FindPinChecked(Pair.Value, EGPD_Input);
+		CompilerContext.MovePinLinksToIntermediate(*OrgInputPin, *NewInputPin);
+	}
+	CreateWriterNode->FindPinChecked(FName("Count"))->DefaultValue = FString::FromInt(1);
+	
+	UEdGraphPin* OldExecPin = FindPinChecked(UEdGraphSchema_K2::PN_Execute, EGPD_Input);
+	UEdGraphPin* NewExecPin = CreateWriterNode->FindPinChecked(UEdGraphSchema_K2::PN_Execute, EGPD_Input);
+	CompilerContext.MovePinLinksToIntermediate(*OldExecPin, *NewExecPin);
+
+	// create the write function nodes
+	UEdGraphPin* LastExecPin = CreateWriterNode->FindPinChecked(UEdGraphSchema_K2::PN_Then, EGPD_Output);
+	UEdGraphPin* WriterResultPin = CreateWriterNode->FindPinChecked(UEdGraphSchema_K2::PN_ReturnValue, EGPD_Output);
+	for (const FNiagaraVariableBase& InVar : DataChannel->Get()->GetVariables())
+	{
+		UEdGraphPin* VarInputPin = FindPinChecked(InVar.GetName(), EGPD_Input);
+		if (VarInputPin == nullptr)
+		{
+			CompilerContext.MessageLog.Error(*FText::Format(LOCTEXT("NoInputPinFound", "Missing input pin for variable '{0}' - @@"), FText::FromName(InVar.GetName())).ToString(), this);
+			continue;
+		}
+		
+		UFunction* WriteFunc = GetWriteFunctionForType(InVar.GetType());
+		if (WriteFunc == nullptr)
+		{
+			CompilerContext.MessageLog.Error(*FText::Format(LOCTEXT("NoWriteFuncFound", "Unable to find a write function for data channel variable '{0}' type {1}, looks like the type is not yet supported by UNiagaraDataChannelWriter. (Source Pin @@)"), FText::FromName(InVar.GetName()), FText::FromString(InVar.GetType().GetName())).ToString(), VarInputPin);
+			continue;
+		}
+		
+		UK2Node_CallFunction* WriteDataNode = CompilerContext.SpawnIntermediateNode<UK2Node_CallFunction>(this, SourceGraph);
+		WriteDataNode->SetFromFunction(WriteFunc);
+		WriteDataNode->AllocateDefaultPins();
+
+		// connect input pins of the write function
+		if (!CompilerContext.GetSchema()->TryCreateConnection(WriterResultPin, WriteDataNode->FindPinChecked(UEdGraphSchema_K2::PN_Self)))
+		{
+			CompilerContext.MessageLog.Error(*LOCTEXT("NoWriterConnection", "Unable to connect writer object result (UNiagaraDataChannelWriter) to write function pins. @@").ToString(), this);
+			continue;
+		}
+		WriteDataNode->FindPinChecked(FName("VarName"), EGPD_Input)->DefaultValue = InVar.GetName().ToString();
+		CompilerContext.MovePinLinksToIntermediate(*VarInputPin, *WriteDataNode->FindPinChecked(FName("InData"), EGPD_Input));
+
+		// connect exec pins
+		CompilerContext.GetSchema()->TryCreateConnection(LastExecPin, WriteDataNode->FindPinChecked(UEdGraphSchema_K2::PN_Execute, EGPD_Input));
+		LastExecPin = WriteDataNode->FindPinChecked(UEdGraphSchema_K2::PN_Then, EGPD_Output);
+	}
+
+
+	// connect the last exec pin
+	UEdGraphPin* OldThenPin = FindPinChecked(UEdGraphSchema_K2::PN_Then, EGPD_Output);
+	CompilerContext.MovePinLinksToIntermediate(*OldThenPin, *LastExecPin);
+}
+
+UK2Node::ERedirectType UK2Node_WriteDataChannel::DoPinsMatchForReconstruction(const UEdGraphPin* NewPin, int32 NewPinIndex, const UEdGraphPin* OldPin, int32 OldPinIndex) const
+{
+	const UEdGraphSchema_K2* K2Schema = GetDefault<UEdGraphSchema_K2>();
+	
+	ERedirectType Result = UK2Node::DoPinsMatchForReconstruction(NewPin, NewPinIndex, OldPin, OldPinIndex);
+	if (ERedirectType_None == Result && K2Schema && K2Schema->ArePinTypesCompatible(NewPin->PinType, OldPin->PinType) && (NewPin->PersistentGuid == OldPin->PersistentGuid) && OldPin->PersistentGuid.IsValid())
+	{
+		Result = ERedirectType_Name;
+	}
+
+	return Result;
+}
+
+void UK2Node_WriteDataChannel::PreloadRequiredAssets()
+{
+	Super::PreloadRequiredAssets();
+
+	if (DataChannel)
+	{
+		PreloadObject(DataChannel);
+		PreloadObject(DataChannel->Get());
+	}
+}
+
+bool UK2Node_WriteDataChannel::ShouldShowNodeProperties() const
+{
+	return Super::ShouldShowNodeProperties();
+}
+
+UEdGraphPin* UK2Node_WriteDataChannel::GetChannelSelectorPin() const
+{
+	return FindPinChecked(FName("Channel"), EGPD_Input);
+}
+
+UFunction* UK2Node_WriteDataChannel::GetWriteFunctionForType(const FNiagaraTypeDefinition& TypeDef)
+{
+	if (TypeDef == FNiagaraTypeHelper::GetDoubleDef() || TypeDef == FNiagaraTypeDefinition::GetFloatDef() || TypeDef == FNiagaraTypeDefinition::GetHalfDef())
+	{
+		return UNiagaraDataChannelWriter::StaticClass()->FindFunctionByName(GET_FUNCTION_NAME_CHECKED(UNiagaraDataChannelWriter, WriteFloat));
+	}
+	if (TypeDef == FNiagaraTypeHelper::GetVector2DDef() || TypeDef == FNiagaraTypeDefinition::GetVec2Def())
+	{
+		return UNiagaraDataChannelWriter::StaticClass()->FindFunctionByName(GET_FUNCTION_NAME_CHECKED(UNiagaraDataChannelWriter, WriteVector2D));
+	}
+	if (TypeDef == FNiagaraTypeDefinition::GetPositionDef())
+	{
+		return UNiagaraDataChannelWriter::StaticClass()->FindFunctionByName(GET_FUNCTION_NAME_CHECKED(UNiagaraDataChannelWriter, WritePosition));
+	}
+	if (TypeDef == FNiagaraTypeHelper::GetVectorDef() || TypeDef == FNiagaraTypeDefinition::GetVec3Def())
+	{
+		return UNiagaraDataChannelWriter::StaticClass()->FindFunctionByName(GET_FUNCTION_NAME_CHECKED(UNiagaraDataChannelWriter, WriteVector));
+	}
+	if (TypeDef == FNiagaraTypeHelper::GetVector4Def() || TypeDef == FNiagaraTypeDefinition::GetVec4Def())
+	{
+		return UNiagaraDataChannelWriter::StaticClass()->FindFunctionByName(GET_FUNCTION_NAME_CHECKED(UNiagaraDataChannelWriter, WriteVector4));
+	}
+	if (TypeDef == FNiagaraTypeHelper::GetQuatDef() || TypeDef == FNiagaraTypeDefinition::GetQuatDef())
+	{
+		return UNiagaraDataChannelWriter::StaticClass()->FindFunctionByName(GET_FUNCTION_NAME_CHECKED(UNiagaraDataChannelWriter, WriteQuat));
+	}
+	if (TypeDef == FNiagaraTypeDefinition::GetColorDef())
+	{
+		return UNiagaraDataChannelWriter::StaticClass()->FindFunctionByName(GET_FUNCTION_NAME_CHECKED(UNiagaraDataChannelWriter, WriteLinearColor));
+	}
+	if (TypeDef == FNiagaraTypeDefinition::GetIntDef())
+	{
+		return UNiagaraDataChannelWriter::StaticClass()->FindFunctionByName(GET_FUNCTION_NAME_CHECKED(UNiagaraDataChannelWriter, WriteInt));
+	}
+	if (TypeDef == FNiagaraTypeDefinition::GetBoolDef())
+	{
+		return UNiagaraDataChannelWriter::StaticClass()->FindFunctionByName(GET_FUNCTION_NAME_CHECKED(UNiagaraDataChannelWriter, WriteBool));
+	}
+	if (TypeDef.GetStruct() == FNiagaraSpawnInfo::StaticStruct())
+	{
+		return UNiagaraDataChannelWriter::StaticClass()->FindFunctionByName(GET_FUNCTION_NAME_CHECKED(UNiagaraDataChannelWriter, WriteSpawnInfo));
+	}
+	if (TypeDef.GetEnum())
+	{
+		return UNiagaraDataChannelWriter::StaticClass()->FindFunctionByName(GET_FUNCTION_NAME_CHECKED(UNiagaraDataChannelWriter, WriteEnum));
+	}
+	return nullptr;
+}
+
+#undef LOCTEXT_NAMESPACE
diff --git a/Engine/Plugins/FX/Niagara/Source/NiagaraBlueprintNodes/Private/K2Node_WriteDataChannel.h b/Engine/Plugins/FX/Niagara/Source/NiagaraBlueprintNodes/Private/K2Node_WriteDataChannel.h
new file mode 100644
--- /dev/null
+++ b/Engine/Plugins/FX/Niagara/Source/NiagaraBlueprintNodes/Private/K2Node_WriteDataChannel.h
@@ -0,0 +1,46 @@
+﻿// Copyright Epic Games, Inc. All Rights Reserved.
+
+#pragma once
+
+#include "CoreMinimal.h"
+#include "K2Node_CallFunction.h"
+#include "NiagaraDataChannel.h"
+#include "K2Node_WriteDataChannel.generated.h"
+
+
+UCLASS()
+class UK2Node_WriteDataChannel : public UK2Node_CallFunction
+{
+	GENERATED_BODY()
+
+	UK2Node_WriteDataChannel();
+
+	virtual void PostLoad() override;
+
+	// //~ Begin UEdGraphNode Interface.
+	virtual void AllocateDefaultPins() override;
+	virtual void PinDefaultValueChanged(UEdGraphPin* Pin) override;
+	virtual void PinConnectionListChanged(UEdGraphPin* Pin) override;
+	// //~ End UEdGraphNode Interface.
+
+	//~ Begin K2Node Interface
+	virtual void GetMenuActions(FBlueprintActionDatabaseRegistrar& InActionRegistrar) const override;
+	virtual FText GetMenuCategory() const override;
+	virtual void ExpandNode(class FKismetCompilerContext& CompilerContext, UEdGraph* SourceGraph) override;
+	virtual ERedirectType DoPinsMatchForReconstruction(const UEdGraphPin* NewPin, int32 NewPinIndex, const UEdGraphPin* OldPin, int32 OldPinIndex) const override;
+	virtual void PreloadRequiredAssets() override;
+	virtual bool ShouldShowNodeProperties() const override;
+	//~ End K2Node Interface
+
+private:
+	UEdGraphPin* GetChannelSelectorPin() const;
+	UFunction* GetWriteFunctionForType(const FNiagaraTypeDefinition& TypeDef);
+
+	UPROPERTY()
+	TObjectPtr<UNiagaraDataChannelAsset> DataChannel;
+
+#if WITH_EDITORONLY_DATA
+	UPROPERTY()
+	FGuid DataChannelVersion;
+#endif
+};
diff --git a/Engine/Plugins/FX/Niagara/Source/NiagaraBlueprintNodes/Private/NiagaraBlueprintNodesModule.cpp b/Engine/Plugins/FX/Niagara/Source/NiagaraBlueprintNodes/Private/NiagaraBlueprintNodesModule.cpp
new file mode 100644
--- /dev/null
+++ b/Engine/Plugins/FX/Niagara/Source/NiagaraBlueprintNodes/Private/NiagaraBlueprintNodesModule.cpp
@@ -0,0 +1,17 @@
+﻿// Copyright Epic Games, Inc. All Rights Reserved.
+
+#include "Modules/ModuleInterface.h"
+#include "Modules/ModuleManager.h"
+
+class FNiagaraBlueprintNodesModule : public IModuleInterface
+{
+public:
+
+	virtual void StartupModule() override
+	{}
+
+	virtual void ShutdownModule() override
+	{}
+};
+
+IMPLEMENT_MODULE(FNiagaraBlueprintNodesModule, NiagaraBlueprintNodes);
\ No newline at end of file
diff --git a/Engine/Plugins/FX/Niagara/Source/NiagaraBlueprintNodes/Private/NiagaraBlueprintUtil.cpp b/Engine/Plugins/FX/Niagara/Source/NiagaraBlueprintNodes/Private/NiagaraBlueprintUtil.cpp
new file mode 100644
--- /dev/null
+++ b/Engine/Plugins/FX/Niagara/Source/NiagaraBlueprintNodes/Private/NiagaraBlueprintUtil.cpp
@@ -0,0 +1,75 @@
+﻿// Copyright Epic Games, Inc. All Rights Reserved.
+
+#include "NiagaraBlueprintUtil.h"
+
+#include "EdGraphSchema_K2.h"
+
+
+FEdGraphPinType FNiagaraBlueprintUtil::TypeDefinitionToBlueprintType(const FNiagaraTypeDefinition& TypeDef)
+{
+	FName Category;
+	FName SubCategory;
+	UObject* SubCategoryObject = nullptr;
+	
+	if (TypeDef == FNiagaraTypeHelper::GetDoubleDef() || TypeDef == FNiagaraTypeDefinition::GetFloatDef() || TypeDef == FNiagaraTypeDefinition::GetHalfDef())
+	{
+		Category = UEdGraphSchema_K2::PC_Real;
+		SubCategory = UEdGraphSchema_K2::PC_Double;
+	}
+	else if (TypeDef == FNiagaraTypeDefinition::GetIntDef())
+	{
+		Category = UEdGraphSchema_K2::PC_Int;
+	}
+	else if (TypeDef == FNiagaraTypeDefinition::GetBoolDef())
+	{
+		Category = UEdGraphSchema_K2::PC_Boolean;
+	}
+	else if (TypeDef == FNiagaraTypeDefinition::GetVec3Def() || TypeDef == FNiagaraTypeDefinition::GetPositionDef())
+	{
+		Category = UEdGraphSchema_K2::PC_Struct;
+		SubCategoryObject = FNiagaraTypeHelper::GetVectorDef().GetStruct();
+	}
+	else if (TypeDef == FNiagaraTypeDefinition::GetVec2Def())
+	{
+		Category = UEdGraphSchema_K2::PC_Struct;
+		SubCategoryObject = FNiagaraTypeHelper::GetVector2DDef().GetStruct();
+	}
+	else if (TypeDef == FNiagaraTypeDefinition::GetVec4Def())
+	{
+		Category = UEdGraphSchema_K2::PC_Struct;
+		SubCategoryObject = FNiagaraTypeHelper::GetVector4Def().GetStruct();
+	}
+	else if (TypeDef == FNiagaraTypeDefinition::GetQuatDef())
+	{
+		Category = UEdGraphSchema_K2::PC_Struct;
+		SubCategoryObject = FNiagaraTypeHelper::GetQuatDef().GetStruct();
+	}
+	//TODO: add matrix def when supported
+	/*else if (TypeDef == FNiagaraTypeDefinition::GetMatrix4Def())
+	{
+		static UScriptStruct* MatrixStruct = FindObjectChecked<UScriptStruct>(CoreUObjectPkg, TEXT("Matrix"));
+		Category = UEdGraphSchema_K2::PC_Struct;
+		SubCategoryObject = MatrixStruct;
+	}*/
+	else if (TypeDef == FNiagaraTypeDefinition::GetColorDef())
+	{
+		Category = UEdGraphSchema_K2::PC_Struct;
+		SubCategoryObject = FNiagaraTypeDefinition::GetColorStruct();
+	}
+	else if (TypeDef.GetEnum())
+	{
+		Category = UEdGraphSchema_K2::PC_Byte;
+		SubCategoryObject = TypeDef.GetEnum();
+	}
+	else if (TypeDef.GetStruct())
+	{
+		Category = UEdGraphSchema_K2::PC_Struct;
+		SubCategoryObject = TypeDef.GetStruct();
+	}
+	else if (TypeDef.GetClass())
+	{
+		Category = UEdGraphSchema_K2::PC_Class;
+		SubCategoryObject = TypeDef.GetClass();
+	}
+	return FEdGraphPinType(Category, SubCategory, SubCategoryObject, EPinContainerType::None, false, FEdGraphTerminalType());
+}
diff --git a/Engine/Plugins/FX/Niagara/Source/NiagaraBlueprintNodes/Private/NiagaraBlueprintUtil.h b/Engine/Plugins/FX/Niagara/Source/NiagaraBlueprintNodes/Private/NiagaraBlueprintUtil.h
new file mode 100644
--- /dev/null
+++ b/Engine/Plugins/FX/Niagara/Source/NiagaraBlueprintNodes/Private/NiagaraBlueprintUtil.h
@@ -0,0 +1,15 @@
+﻿// Copyright Epic Games, Inc. All Rights Reserved.
+
+#pragma once
+
+#include "NiagaraTypes.h"
+#include "EdGraph/EdGraphPin.h"
+
+class FNiagaraBlueprintUtil
+{
+public:
+	/**
+	Attempts to convert the existing Niagara type to a matching type for BP pins
+	*/
+	static FEdGraphPinType TypeDefinitionToBlueprintType(const FNiagaraTypeDefinition& TypeDef);
+};
diff --git a/Engine/Plugins/FX/Niagara/Source/NiagaraEditor/Private/Customizations/NiagaraTypeCustomizations.cpp b/Engine/Plugins/FX/Niagara/Source/NiagaraEditor/Private/Customizations/NiagaraTypeCustomizations.cpp
--- a/Engine/Plugins/FX/Niagara/Source/NiagaraEditor/Private/Customizations/NiagaraTypeCustomizations.cpp
+++ b/Engine/Plugins/FX/Niagara/Source/NiagaraEditor/Private/Customizations/NiagaraTypeCustomizations.cpp
@@ -27,6 +27,7 @@
 #include "Widgets/Text/STextBlock.h"
 
 #include "NiagaraConstants.h"
+#include "NiagaraDataChannelPublic.h"
 #include "NiagaraDataInterfaceRW.h"
 #include "NiagaraEditorModule.h"
 #include "NiagaraEditorUtilities.h"
@@ -2749,5 +2750,262 @@ TSharedRef<SWidget> FNiagaraVariableDetailsCustomization::GetTypeMenu(TSharedPtr
 	.VarToModify(Var);
 }
 
+void FNiagaraDataChannelVariableDetailsCustomization::CustomizeHeader(TSharedRef<IPropertyHandle> PropertyHandle, FDetailWidgetRow& HeaderRow, IPropertyTypeCustomizationUtils& CustomizationUtils)
+{
+}
+
+void FNiagaraDataChannelVariableDetailsCustomization::CustomizeChildren(TSharedRef<IPropertyHandle> PropertyHandle, IDetailChildrenBuilder& ChildBuilder, IPropertyTypeCustomizationUtils& CustomizationUtils)
+{
+	TSharedPtr<IPropertyHandle> NameHandle = PropertyHandle->GetChildHandle(TEXT("Name"));
+	TSharedPtr<IPropertyHandle> TypeDefHandleHandle = PropertyHandle->GetChildHandle(TEXT("TypeDefHandle"));
+
+	FDetailWidgetRow& Row = ChildBuilder.AddCustomRow(FText::GetEmpty());
+
+	TArray<UObject*> Objects;
+	PropertyHandle->GetOuterObjects(Objects);
+	if (Objects.Num() > 1)
+	{
+		Row.NameContent()
+			[
+				PropertyHandle->CreatePropertyNameWidget()
+			];
+		Row.ValueContent()
+			[
+				PropertyHandle->CreatePropertyValueWidget()
+			];
+	}
+	else
+	{
+		void* VarPtr = nullptr;
+		if (PropertyHandle->GetValueData(VarPtr) == FPropertyAccess::Success)
+		{
+			bool bEnforceUniqueName = PropertyHandle->HasMetaData(TEXT("EnforceUniqueNames"));
+			FNiagaraDataChannelVariable* Variable = static_cast<FNiagaraDataChannelVariable*>(VarPtr);
+
+			Row.NameContent()
+			[
+				SNew(SHorizontalBox)
+				+ SHorizontalBox::Slot()
+				.AutoWidth()
+				.VAlign(VAlign_Center)
+				[
+					SNew(STextBlock)
+					.Font(IDetailLayoutBuilder::GetDetailFont())
+					.Text(LOCTEXT("DataChannelVarNameText", "Name: "))
+				]
+				+ SHorizontalBox::Slot()
+				.AutoWidth()
+				.VAlign(VAlign_Center)
+				.Padding(4.0f, 4.0f, 4.0f, 4.0f)
+				[
+					SNew(SNiagaraConstrainedBox)
+					.MinWidth(150.0f)
+					[
+						SNew( SEditableTextBox )
+						.Text_Lambda([NameHandle]() -> FText { FText OutText; NameHandle->GetValueAsDisplayText(OutText); return OutText; })
+						.Font(IDetailLayoutBuilder::GetDetailFont())
+						.SelectAllTextWhenFocused(true)
+						.ClearKeyboardFocusOnCommit(false)
+						.OnTextCommitted_Lambda([NameHandle](const FText& NewName, ETextCommit::Type)
+						{
+							FScopedTransaction Transaction(LOCTEXT("ChangeVarName", "Change variable name"));
+							TArray<UObject*> OuterObjects;
+							NameHandle->GetOuterObjects(OuterObjects);
+							for (UObject* OuterObj : OuterObjects)
+							{
+								OuterObj->Modify();
+							}
+							NameHandle->SetValue(FName(NewName.ToString()));
+						} )
+						.OnVerifyTextChanged_Lambda([bEnforceUniqueName, PropertyHandle](const FText& InNewText, FText& OutErrorMessage) -> bool
+						{
+							// if necessary, validate that the entered name is unique among all entries of the variable array
+							if (!bEnforceUniqueName)
+							{
+								return true;
+							}
+							FName NewName = FName(InNewText.ToString());
+							if (TSharedPtr<IPropertyHandle> Handle = PropertyHandle->GetParentHandle())
+							{
+								if (TSharedPtr<IPropertyHandleArray> PropertyHandleArray = Handle->AsArray())
+								{
+									uint32 NumElements;
+									PropertyHandleArray->GetNumElements(NumElements);
+									bool bNewNameValid = true;
+									for (uint32 i = 0; i < NumElements; i++)
+									{
+										TSharedRef<IPropertyHandle> ElementHandle = PropertyHandleArray->GetElement(i);
+										void* VarPtr = nullptr;
+										if (!ElementHandle->IsSamePropertyNode(PropertyHandle) && ElementHandle->GetValueData(VarPtr) == FPropertyAccess::Success)
+										{
+											if (NewName == static_cast<FNiagaraVariable*>(VarPtr)->GetName())
+											{
+												bNewNameValid = false;
+												break;
+											}
+										}
+									}
+									if (!bNewNameValid)
+									{
+										OutErrorMessage = LOCTEXT("DuplicateNameError", "Variable name has to be unique!");
+										return false;
+									}
+								}
+							}
+							return true;
+						})
+						.SelectAllTextOnCommit( true )
+					]
+				]
+				+ SHorizontalBox::Slot()
+				.AutoWidth()
+				.VAlign(VAlign_Center)
+				.Padding(4.0f, 2.0f, 4.0f, 2.0f)
+				[
+					SNew(SNiagaraConstrainedBox)
+					.MinWidth(150.0f)
+					[
+						PropertyHandle->CreateDefaultPropertyButtonWidgets()
+					]
+				]
+			];
+			
+			Row.ValueContent()
+			[
+				SNew(SHorizontalBox)
+				+ SHorizontalBox::Slot()
+				.AutoWidth()
+				.VAlign(VAlign_Center)
+				.Padding(0, 0, 2.0f, 0)
+				[
+					SNew(STextBlock)
+					.Font(IDetailLayoutBuilder::GetDetailFont())
+					.Text(LOCTEXT("DataChannelVarTypeText", "Type: "))
+				]
+				+ SHorizontalBox::Slot()
+				.VAlign(VAlign_Center)
+				.Padding(2.0f)
+				[
+					SNew(SComboButton)
+					.HasDownArrow(true)
+					.ContentPadding(0)
+					.OnGetMenuContent(this, &FNiagaraDataChannelVariableDetailsCustomization::GetTypeMenu, TypeDefHandleHandle, Variable)
+					.ButtonContent()
+					[
+						SNew(SHorizontalBox)
+						+ SHorizontalBox::Slot()
+						.VAlign(VAlign_Center)
+						.HAlign(HAlign_Center)
+						.AutoWidth()
+						[
+							SNew(SImage)
+							.ColorAndOpacity_Lambda([Variable]() {return UEdGraphSchema_Niagara::GetTypeColor(Variable->GetType());})
+							.Image(FNiagaraEditorStyle::Get().GetBrush("NiagaraEditor.Module.TypeIconPill"))
+						]
+						+ SHorizontalBox::Slot()
+						.Padding(4, 2, 2, 2)
+						[
+							SNew(STextBlock)
+							.MinDesiredWidth(150)
+							.Text_Lambda([Variable] {return Variable ? Variable->GetType().GetNameText() : FText::GetEmpty(); })
+							.Font(IDetailLayoutBuilder::GetDetailFont())
+						]
+					]
+				]
+			];
+		}
+	}
+}
+
+class SNiagaraDataChannelTypeSelectMenu : public SNiagaraParameterMenu
+{
+public:
+	//DECLARE_DELEGATE_RetVal_OneParam(bool, FOnAllowMakeType, const FNiagaraTypeDefinition&);
+
+	SLATE_BEGIN_ARGS(SNiagaraDataChannelTypeSelectMenu)
+		: _AutoExpandMenu(false)
+	{}
+	//~ Begin Required Args
+	SLATE_ARGUMENT(TSharedPtr<IPropertyHandle>, PropertyHandle)
+	SLATE_ARGUMENT(FNiagaraDataChannelVariable*, VarToModify)
+	//~ End Required Args
+	SLATE_ARGUMENT(bool, AutoExpandMenu)
+	SLATE_END_ARGS();
+
+	void Construct(const FArguments& InArgs);
+
+protected:
+	virtual void CollectAllActions(FGraphActionListBuilderBase& OutAllActions) override;
+
+private:
+	TSharedPtr<IPropertyHandle> PropertyHandle;
+	FNiagaraDataChannelVariable* VarToModify = nullptr;
+};
+
+
+void SNiagaraDataChannelTypeSelectMenu::Construct(const FArguments& InArgs)
+{
+	checkf(InArgs._VarToModify != nullptr, TEXT("Tried to construct change var type menu without valid var ptr!"));
+	this->PropertyHandle = InArgs._PropertyHandle;
+	this->VarToModify = InArgs._VarToModify;
+
+	SNiagaraParameterMenu::FArguments SuperArgs;
+	SuperArgs._AutoExpandMenu = InArgs._AutoExpandMenu;
+	SNiagaraParameterMenu::Construct(SuperArgs);
+}
+
+void SNiagaraDataChannelTypeSelectMenu::CollectAllActions(FGraphActionListBuilderBase& OutAllActions)
+{
+	FNiagaraMenuActionCollector Collector;
+
+	TArray<FNiagaraTypeDefinition> Types;
+	FNiagaraEditorUtilities::GetAllowedUserVariableTypes(Types);
+	for (FNiagaraTypeDefinition& Type : Types)
+	{
+		if (!FNiagaraDataChannelVariable::IsAllowedType(Type))
+		{
+			continue;
+		}
+		Type = FNiagaraDataChannelVariable::ToDataChannelType(Type);
+
+		FText Category = FNiagaraEditorUtilities::GetTypeDefinitionCategory(Type);
+		const FText DisplayName = Type.GetNameText();
+		const FText Tooltip = Type.GetNameText();
+		TSharedPtr<FNiagaraMenuAction> Action(new FNiagaraMenuAction(
+			Category, DisplayName, Tooltip, 0, FText::GetEmpty(),
+			FNiagaraMenuAction::FOnExecuteStackAction::CreateLambda(
+			[Var=VarToModify, PropHandle=PropertyHandle, Type]()
+			{
+				FScopedTransaction Transaction(LOCTEXT("Set Raw Niagara Variable Type", "Set Variable Type"));
+				check(Var);
+				check(PropHandle.IsValid());
+				TArray<UObject*> Objects;
+				PropHandle->GetOuterObjects(Objects);
+				for (UObject* Obj : Objects)
+				{
+					Obj->Modify();
+				}
+
+				PropHandle->NotifyPreChange();
+				Var->SetType(Type);
+				Var->Version = FGuid::NewGuid();
+				PropHandle->NotifyPostChange(EPropertyChangeType::ValueSet);
+				PropHandle->NotifyFinishedChangingProperties();
+			}
+		)));
+
+		Collector.AddAction(Action, 0);
+	}
+
+	Collector.AddAllActionsTo(OutAllActions);
+}
+
+TSharedRef<SWidget> FNiagaraDataChannelVariableDetailsCustomization::GetTypeMenu(TSharedPtr<IPropertyHandle> InPropertyHandle, FNiagaraDataChannelVariable* Var)
+{
+	return SNew(SNiagaraDataChannelTypeSelectMenu)
+	.PropertyHandle(InPropertyHandle)
+	.VarToModify(Var);
+}
+
 
 #undef LOCTEXT_NAMESPACE
diff --git a/Engine/Plugins/FX/Niagara/Source/NiagaraEditor/Private/Customizations/NiagaraTypeCustomizations.h b/Engine/Plugins/FX/Niagara/Source/NiagaraEditor/Private/Customizations/NiagaraTypeCustomizations.h
--- a/Engine/Plugins/FX/Niagara/Source/NiagaraEditor/Private/Customizations/NiagaraTypeCustomizations.h
+++ b/Engine/Plugins/FX/Niagara/Source/NiagaraEditor/Private/Customizations/NiagaraTypeCustomizations.h
@@ -16,6 +16,7 @@ class IPropertyHandle;
 class IPropertyHandleArray;
 class UMaterialInterface;
 struct FMaterialParameterInfo;
+struct FNiagaraDataChannelVariable;
 class UNiagaraGraph;
 class UNiagaraParameterDefinitions;
 class UNiagaraRendererProperties;
@@ -458,3 +459,19 @@ private:
 	TSharedRef<SWidget> GetTypeMenu(TSharedPtr<IPropertyHandle> InPropertyHandle, FNiagaraVariable* Var);
 };
 
+//** Properties customization for FNiagaraDataChannelVariable. */
+class FNiagaraDataChannelVariableDetailsCustomization : public IPropertyTypeCustomization
+{
+public:
+	static TSharedRef<IPropertyTypeCustomization> MakeInstance()
+	{
+		return MakeShared<FNiagaraDataChannelVariableDetailsCustomization>();
+	}
+
+	virtual void CustomizeHeader(TSharedRef<IPropertyHandle> PropertyHandle, FDetailWidgetRow& HeaderRow, IPropertyTypeCustomizationUtils& CustomizationUtils) override;
+
+	virtual void CustomizeChildren(TSharedRef<IPropertyHandle> PropertyHandle, IDetailChildrenBuilder& ChildBuilder, IPropertyTypeCustomizationUtils& CustomizationUtils) override;
+
+private:
+	TSharedRef<SWidget> GetTypeMenu(TSharedPtr<IPropertyHandle> InPropertyHandle, FNiagaraDataChannelVariable* Var);
+};
diff --git a/Engine/Plugins/FX/Niagara/Source/NiagaraEditor/Private/EdGraphSchema_Niagara.cpp b/Engine/Plugins/FX/Niagara/Source/NiagaraEditor/Private/EdGraphSchema_Niagara.cpp
--- a/Engine/Plugins/FX/Niagara/Source/NiagaraEditor/Private/EdGraphSchema_Niagara.cpp
+++ b/Engine/Plugins/FX/Niagara/Source/NiagaraEditor/Private/EdGraphSchema_Niagara.cpp
@@ -1547,7 +1547,7 @@ FLinearColor UEdGraphSchema_Niagara::GetTypeColor(const FNiagaraTypeDefinition&
 {
 	const UGraphEditorSettings* Settings = GetDefault<UGraphEditorSettings>();
 	const UNiagaraSettings* NiagaraSettings = GetDefault<UNiagaraSettings>();
-	if (Type == FNiagaraTypeDefinition::GetFloatDef())
+	if (Type == FNiagaraTypeDefinition::GetFloatDef() || Type == FNiagaraTypeHelper::GetDoubleDef())
 	{
 		return Settings->FloatPinTypeColor;
 	}
@@ -1559,7 +1559,7 @@ FLinearColor UEdGraphSchema_Niagara::GetTypeColor(const FNiagaraTypeDefinition&
 	{
 		return Settings->BooleanPinTypeColor;
 	}
-	else if (Type == FNiagaraTypeDefinition::GetVec3Def())
+	else if (Type == FNiagaraTypeDefinition::GetVec3Def() || Type == FNiagaraTypeHelper::GetVectorDef())
 	{
 		return Settings->VectorPinTypeColor;
 	}
diff --git a/Engine/Plugins/FX/Niagara/Source/NiagaraEditor/Private/NiagaraActions.cpp b/Engine/Plugins/FX/Niagara/Source/NiagaraEditor/Private/NiagaraActions.cpp
--- a/Engine/Plugins/FX/Niagara/Source/NiagaraEditor/Private/NiagaraActions.cpp
+++ b/Engine/Plugins/FX/Niagara/Source/NiagaraEditor/Private/NiagaraActions.cpp
@@ -766,8 +766,8 @@ void FNiagaraDataInterfaceNodeActionProvider_DataChannelWrite::GetNodeContextMen
 					if(Channel && Node)
 					{
 						Node->RemoveAllDynamicPins();
-						TConstArrayView<FNiagaraVariable> ChannelVars = Channel->GetVariables();
-						for (const FNiagaraVariable& Var : ChannelVars)
+						TConstArrayView<FNiagaraDataChannelVariable> ChannelVars = Channel->GetVariables();
+						for (const FNiagaraDataChannelVariable& Var : ChannelVars)
 						{
 							FNiagaraTypeDefinition Type = Var.GetType();
 							if (Type.IsEnum() == false)
@@ -857,8 +857,8 @@ void FNiagaraDataInterfaceNodeActionProvider_DataChannelRead::CollectAddPinActio
 
 	auto GatherAddPinsForChannel = [&](UNiagaraDataChannel* Channel)
 	{
-		TConstArrayView<FNiagaraVariable> ChannelVars = Channel->GetVariables();
-		for (const FNiagaraVariable& Var : ChannelVars)
+		TConstArrayView<FNiagaraDataChannelVariable> ChannelVars = Channel->GetVariables();
+		for (const FNiagaraDataChannelVariable& Var : ChannelVars)
 		{
 			FNiagaraTypeDefinition Type = Var.GetType();
 			if (Type.IsEnum() == false)
@@ -926,8 +926,8 @@ void FNiagaraDataInterfaceNodeActionProvider_DataChannelRead::AddDataChannelInit
 				if (Channel && Node)
 				{
 					Node->RemoveAllDynamicPins();
-					TConstArrayView<FNiagaraVariable> ChannelVars = Channel->GetVariables();
-					for (const FNiagaraVariable& Var : ChannelVars)
+					TConstArrayView<FNiagaraDataChannelVariable> ChannelVars = Channel->GetVariables();
+					for (const FNiagaraDataChannelVariable& Var : ChannelVars)
 					{
 						FNiagaraTypeDefinition Type = Var.GetType();
 						if (Type.IsEnum() == false)
diff --git a/Engine/Plugins/FX/Niagara/Source/NiagaraEditor/Private/NiagaraEditorModule.cpp b/Engine/Plugins/FX/Niagara/Source/NiagaraEditor/Private/NiagaraEditorModule.cpp
--- a/Engine/Plugins/FX/Niagara/Source/NiagaraEditor/Private/NiagaraEditorModule.cpp
+++ b/Engine/Plugins/FX/Niagara/Source/NiagaraEditor/Private/NiagaraEditorModule.cpp
@@ -1133,6 +1133,10 @@ void FNiagaraEditorModule::StartupModule()
 		FNiagaraVariable::StaticStruct()->GetFName(),
 		FOnGetPropertyTypeCustomizationInstance::CreateStatic(&FNiagaraVariableDetailsCustomization::MakeInstance));
 
+	PropertyModule.RegisterCustomPropertyTypeLayout(
+		FNiagaraDataChannelVariable::StaticStruct()->GetFName(),
+		FOnGetPropertyTypeCustomizationInstance::CreateStatic(&FNiagaraDataChannelVariableDetailsCustomization::MakeInstance));
+
 #if WITH_NIAGARA_DEBUGGER
 	PropertyModule.RegisterCustomPropertyTypeLayout(
 		FNiagaraDebugHUDVariable::StaticStruct()->GetFName(),
