Updated data channel tutorial link

diff --git a/Engine/Plugins/FX/Niagara/Source/NiagaraEditor/Private/Customizations/NiagaraDataChannelDetails.cpp b/Engine/Plugins/FX/Niagara/Source/NiagaraEditor/Private/Customizations/NiagaraDataChannelDetails.cpp
--- a/Engine/Plugins/FX/Niagara/Source/NiagaraEditor/Private/Customizations/NiagaraDataChannelDetails.cpp
+++ b/Engine/Plugins/FX/Niagara/Source/NiagaraEditor/Private/Customizations/NiagaraDataChannelDetails.cpp
@@ -74,8 +74,7 @@ void FNiagaraDataChannelAssetDetails::CustomizeDetails(IDetailLayoutBuilder& Det
 			.ToolTipText(LOCTEXT("DataChannelTypePageTooltip", "Opens a docs page that explains data channel types"))
 			.OnNavigate_Lambda([]()
 			{
-				//TODO (mga) set correct url here
-				FPlatformProcess::LaunchURL(TEXT("https://dev.epicgames.com/community/unreal-engine/learning?types=tutorial"), nullptr, nullptr);
+				FPlatformProcess::LaunchURL(TEXT("https://dev.epicgames.com/community/learning/tutorials/RJbm/unreal-engine-niagara-data-channels-intro"), nullptr, nullptr);
 			})
 		]
 	];
