Ensure lwc types are always initialized

diff --git a/Engine/Plugins/FX/Niagara/Source/Niagara/Private/NiagaraModule.cpp b/Engine/Plugins/FX/Niagara/Source/Niagara/Private/NiagaraModule.cpp
--- a/Engine/Plugins/FX/Niagara/Source/Niagara/Private/NiagaraModule.cpp
+++ b/Engine/Plugins/FX/Niagara/Source/Niagara/Private/NiagaraModule.cpp
@@ -279,11 +279,11 @@ FNiagaraVariable INiagaraModule::Translator_CallID;
 void INiagaraModule::StartupModule()
 {
 	VectorVM::Init();
-	FNiagaraTypeHelper::InitStaticTypes();
+	FNiagaraTypeHelper::InitLWCTypes();
 
 	LLM_SCOPE(ELLMTag::Niagara);
 	FNiagaraTypeDefinition::Init();
-	FNiagaraTypeHelper::RegisterStaticTypes();
+	FNiagaraTypeHelper::RegisterLWCTypes();
 
 	if (IsRunningCookCommandlet())
 	{
diff --git a/Engine/Plugins/FX/Niagara/Source/Niagara/Private/NiagaraTypes.cpp b/Engine/Plugins/FX/Niagara/Source/Niagara/Private/NiagaraTypes.cpp
--- a/Engine/Plugins/FX/Niagara/Source/Niagara/Private/NiagaraTypes.cpp
+++ b/Engine/Plugins/FX/Niagara/Source/Niagara/Private/NiagaraTypes.cpp
@@ -661,6 +661,7 @@ FNiagaraTypeDefinition FNiagaraTypeHelper::GetLWCType(const FNiagaraTypeDefiniti
 		return InType;
 	}
 
+	InitLWCTypes();
 	if(InType == FNiagaraTypeDefinition::GetFloatDef())
 	{
 		return DoubleDef;
@@ -695,7 +696,7 @@ FNiagaraTypeDefinition FNiagaraTypeHelper::GetLWCType(const FNiagaraTypeDefiniti
 	return InType;
 }
 
-void FNiagaraTypeHelper::InitStaticTypes()
+void FNiagaraTypeHelper::InitLWCTypes()
 {
 	if (Vector2DDef.IsValid() == false)
 	{
@@ -708,7 +709,7 @@ void FNiagaraTypeHelper::InitStaticTypes()
 	}
 }
 
-void FNiagaraTypeHelper::RegisterStaticTypes()
+void FNiagaraTypeHelper::RegisterLWCTypes()
 {
 	FNiagaraTypeRegistry::Register(Vector2DDef, ENiagaraTypeRegistryFlags::None);
 	FNiagaraTypeRegistry::Register(VectorDef, ENiagaraTypeRegistryFlags::None);
@@ -719,26 +720,31 @@ void FNiagaraTypeHelper::RegisterStaticTypes()
 
 FNiagaraTypeDefinition FNiagaraTypeHelper::GetVector2DDef()
 {
+	InitLWCTypes();
 	return Vector2DDef;
 }
 
 FNiagaraTypeDefinition FNiagaraTypeHelper::GetVectorDef()
 {
+	InitLWCTypes();
 	return VectorDef;
 }
 
 FNiagaraTypeDefinition FNiagaraTypeHelper::GetVector4Def()
 {
+	InitLWCTypes();
 	return Vector4Def;
 }
 
 FNiagaraTypeDefinition FNiagaraTypeHelper::GetQuatDef()
 {
+	InitLWCTypes();
 	return QuatDef;
 }
 
 FNiagaraTypeDefinition FNiagaraTypeHelper::GetDoubleDef()
 {
+	InitLWCTypes();
 	return DoubleDef;
 }
 
diff --git a/Engine/Plugins/FX/Niagara/Source/Niagara/Public/NiagaraTypes.h b/Engine/Plugins/FX/Niagara/Source/Niagara/Public/NiagaraTypes.h
--- a/Engine/Plugins/FX/Niagara/Source/Niagara/Public/NiagaraTypes.h
+++ b/Engine/Plugins/FX/Niagara/Source/Niagara/Public/NiagaraTypes.h
@@ -437,8 +437,8 @@ public:
 	static NIAGARA_API FNiagaraTypeDefinition GetQuatDef();
 	static NIAGARA_API FNiagaraTypeDefinition GetDoubleDef();
 
-	static NIAGARA_API void InitStaticTypes();
-	static NIAGARA_API void RegisterStaticTypes();
+	static NIAGARA_API void InitLWCTypes();
+	static NIAGARA_API void RegisterLWCTypes();
 private:
 	struct FRemapEntry
 	{
