Fix Niagara type registry crash in packaged builds

diff --git a/Engine/Plugins/FX/Niagara/Source/Niagara/Private/NiagaraModule.cpp b/Engine/Plugins/FX/Niagara/Source/Niagara/Private/NiagaraModule.cpp
--- a/Engine/Plugins/FX/Niagara/Source/Niagara/Private/NiagaraModule.cpp
+++ b/Engine/Plugins/FX/Niagara/Source/Niagara/Private/NiagaraModule.cpp
@@ -283,6 +283,7 @@ void INiagaraModule::StartupModule()
 
 	LLM_SCOPE(ELLMTag::Niagara);
 	FNiagaraTypeDefinition::Init();
+	FNiagaraTypeRegistry::ProcessRegistryQueue();
 	FNiagaraTypeHelper::RegisterLWCTypes();
 
 	if (IsRunningCookCommandlet())
diff --git a/Engine/Plugins/FX/Niagara/Source/Niagara/Public/NiagaraTypes.h b/Engine/Plugins/FX/Niagara/Source/Niagara/Public/NiagaraTypes.h
--- a/Engine/Plugins/FX/Niagara/Source/Niagara/Public/NiagaraTypes.h
+++ b/Engine/Plugins/FX/Niagara/Source/Niagara/Public/NiagaraTypes.h
@@ -6,6 +6,7 @@
 #include "Misc/SecureHash.h"
 #include "UObject/GCObject.h"
 #include "UObject/UnrealType.h"
+#include "Containers/Queue.h"
 #include "NiagaraCore.h"
 
 #include "NiagaraTypes.generated.h"
@@ -1267,9 +1268,33 @@ public:
 		Register(NewType, Flags);
 	}
 
+	static void ProcessRegistryQueue()
+	{
+		FNiagaraTypeRegistry& Registry = Get();
+		{
+			FRWScopeLock Lock(Registry.RegisteredTypesLock, SLT_Write);
+			Registry.bModuleInitialized = true;
+		}
+		FQueuedRegistryEntry Entry;
+		while (Registry.RegistryQueue.Dequeue(Entry))
+		{
+			Register(Entry.NewType, Entry.Flags);
+		}
+	}
+
 	static void Register(const FNiagaraTypeDefinition &NewType, ENiagaraTypeRegistryFlags Flags)
 	{
 		FNiagaraTypeRegistry& Registry = Get();
+		{
+			FReadScopeLock Lock(Registry.RegisteredTypesLock);
+			if (!Registry.bModuleInitialized)
+			{
+				// In a packaged game it can happen that CDOs are created before the Niagara module had a chance to be initialized.
+				// This is problematic, as the swc struct builder tries to access other Niagara types, so we delay the registration until the module is properly initialized.
+				Registry.RegistryQueue.Enqueue({NewType, Flags});
+				return;
+			}
+		}
 
 		if (FNiagaraTypeHelper::IsLWCType(NewType))
 		{
@@ -1409,6 +1434,11 @@ public:
 	NIAGARA_API virtual FString GetReferencerName() const;
 
 private:
+	struct FQueuedRegistryEntry
+	{
+		FNiagaraTypeDefinition NewType;
+		ENiagaraTypeRegistryFlags Flags;
+	};
 	friend class FLazySingleton;
 
 	NIAGARA_API FNiagaraTypeRegistry();
@@ -1430,6 +1460,9 @@ private:
 	TMap<TWeakObjectPtr<UScriptStruct>, TWeakObjectPtr<UScriptStruct>> LWCRegisteredStructRemapping;
 	TMap<uint32, FNiagaraLwcStructConverter> RegisteredStructConversionMap;
 	FRWLock RegisteredTypesLock;
+
+	bool bModuleInitialized = false;
+	TQueue<FQueuedRegistryEntry, EQueueMode::Mpsc> RegistryQueue;
 };
 
 USTRUCT()
