Improved test error message

diff --git a/Engine/Plugins/FX/Niagara/Source/Niagara/Private/NiagaraSimCache.cpp b/Engine/Plugins/FX/Niagara/Source/Niagara/Private/NiagaraSimCache.cpp
--- a/Engine/Plugins/FX/Niagara/Source/Niagara/Private/NiagaraSimCache.cpp
+++ b/Engine/Plugins/FX/Niagara/Source/Niagara/Private/NiagaraSimCache.cpp
@@ -1190,7 +1190,7 @@ bool UNiagaraSimCache::IsDataEqual(const UNiagaraSimCache& OtherCache, float Err
 	const TArray<FNiagaraSimCacheFrame>& OtherFrames = OtherCache.CacheFrames;
 	if (GetNumFrames() != OtherCache.GetNumFrames())
 	{
-		OutDifference = TEXT("Cache Frame count different");
+		OutDifference = FText::Format(FText::FromString("Cache Frame count different. Expected: {0}, Actual: {1}"), FText::AsNumber(GetNumFrames()), FText::AsNumber(OtherCache.GetNumFrames())).ToString();
 		return false;
 	}
 
