- Fix text 'headsers' -> 'headers'

diff --git a/Engine/Plugins/FX/Niagara/Source/NiagaraEditor/Private/NiagaraEditorCommands.cpp b/Engine/Plugins/FX/Niagara/Source/NiagaraEditor/Private/NiagaraEditorCommands.cpp
--- a/Engine/Plugins/FX/Niagara/Source/NiagaraEditor/Private/NiagaraEditorCommands.cpp
+++ b/Engine/Plugins/FX/Niagara/Source/NiagaraEditor/Private/NiagaraEditorCommands.cpp
@@ -46,7 +46,7 @@ void FNiagaraEditorCommands::RegisterCommands()
 	UI_COMMAND(ToggleResetSimulationOnChange, "Reset on change", "Toggles whether or not the simulation is reset whenever a change is made in the asset editor.", EUserInterfaceActionType::ToggleButton, FInputChord());
 	UI_COMMAND(ToggleResimulateOnChangeWhilePaused, "Resimulate when paused", "Toggles whether or not the simulation is rerun to the current time when making changes while paused.", EUserInterfaceActionType::ToggleButton, FInputChord());
 	UI_COMMAND(ToggleResetDependentSystems, "Reset Dependent Systems", "Toggles whether or not to reset all systems that include this emitter when it is reset by the user.", EUserInterfaceActionType::ToggleButton, FInputChord());
-	UI_COMMAND(CollapseStackToHeaders, "Collapse to Headers", "Expands all headsers and collapse all items.", EUserInterfaceActionType::Button, FInputChord(EModifierKey::Control, EKeys::O));
+	UI_COMMAND(CollapseStackToHeaders, "Collapse to Headers", "Expands all headers and collapse all items.", EUserInterfaceActionType::Button, FInputChord(EModifierKey::Control, EKeys::O));
 
 	UI_COMMAND(IsolateSelectedEmitters, "Isolate selected emitters", "Isolate all currently selected emitters.", EUserInterfaceActionType::Button, FInputChord(EKeys::S));
 	UI_COMMAND(DisableSelectedEmitters, "Disable selected emitters", "Disables all currently selected emitters and recompiles the system.", EUserInterfaceActionType::Button, FInputChord(EKeys::D));
