- Fixed an issue with relevance in Niagara where the default material relevance forced us to always render in the depth pass

diff --git a/Engine/Plugins/FX/Niagara/Source/Niagara/Private/NiagaraComponent.cpp b/Engine/Plugins/FX/Niagara/Source/Niagara/Private/NiagaraComponent.cpp
--- a/Engine/Plugins/FX/Niagara/Source/Niagara/Private/NiagaraComponent.cpp
+++ b/Engine/Plugins/FX/Niagara/Source/Niagara/Private/NiagaraComponent.cpp
@@ -271,23 +271,21 @@ void FNiagaraSceneProxy::OnTransformChanged(FRHICommandListBase& RHICmdList)
 FPrimitiveViewRelevance FNiagaraSceneProxy::GetViewRelevance(const FSceneView* View) const
 {
 	FPrimitiveViewRelevance Relevance;
-
 	if (!GetRenderingEnabled() || !FNiagaraUtilities::SupportsNiagaraRendering(View->GetFeatureLevel()))
 	{
 		return Relevance;
 	}
-	Relevance.bDynamicRelevance = true;
-
-	Relevance.bRenderCustomDepth = ShouldRenderCustomDepth();
-	Relevance.bDrawRelevance = IsShown(View) && View->Family->EngineShowFlags.Particles && View->Family->EngineShowFlags.Niagara;
-	Relevance.bShadowRelevance = IsShadowCast(View);
-	Relevance.bRenderInMainPass = ShouldRenderInMainPass();
-	Relevance.bUsesLightingChannels = GetLightingChannelMask() != GetDefaultLightingChannelMask();
-	Relevance.bTranslucentSelfShadow = bCastVolumetricTranslucentShadow;
-
-	Relevance |= RenderData->GetViewRelevance(*View, *this);
 
-	Relevance.bVelocityRelevance = DrawsVelocity() && Relevance.bOpaque && Relevance.bRenderInMainPass;
+	// Read Relevance from the material / renderer first as FPrimitiveViewRelevance() will turn on certain flags which we want to control (i.e. bRenderInMainPass)
+	Relevance							= RenderData->GetViewRelevance(*View, *this);
+	Relevance.bDynamicRelevance			= true;
+	Relevance.bDrawRelevance			= IsShown(View) && View->Family->EngineShowFlags.Particles && View->Family->EngineShowFlags.Niagara;
+	Relevance.bRenderCustomDepth		= ShouldRenderCustomDepth();
+	Relevance.bShadowRelevance			= IsShadowCast(View);
+	Relevance.bRenderInMainPass			= ShouldRenderInMainPass();
+	Relevance.bUsesLightingChannels		= GetLightingChannelMask() != GetDefaultLightingChannelMask();
+	Relevance.bTranslucentSelfShadow	= bCastVolumetricTranslucentShadow;
+	Relevance.bVelocityRelevance		= DrawsVelocity() && Relevance.bOpaque && Relevance.bRenderInMainPass;
 
 	return Relevance;
 }
diff --git a/Engine/Plugins/FX/Niagara/Source/Niagara/Private/NiagaraRenderer.cpp b/Engine/Plugins/FX/Niagara/Source/Niagara/Private/NiagaraRenderer.cpp
--- a/Engine/Plugins/FX/Niagara/Source/Niagara/Private/NiagaraRenderer.cpp
+++ b/Engine/Plugins/FX/Niagara/Source/Niagara/Private/NiagaraRenderer.cpp
@@ -384,18 +384,12 @@ FNiagaraRenderer::~FNiagaraRenderer()
 FPrimitiveViewRelevance FNiagaraRenderer::GetViewRelevance(const FSceneView* View, const FNiagaraSceneProxy *SceneProxy)const
 {
 	FPrimitiveViewRelevance Result;
-	bool bHasDynamicData = HasDynamicData();
-
-	//Always draw so our LastRenderTime is updated. We may not have dynamic data if we're disabled from visibility culling.
-	Result.bDrawRelevance =/* bHasDynamicData && */SceneProxy->IsShown(View) && View->Family->EngineShowFlags.Particles && View->Family->EngineShowFlags.Niagara;
-	Result.bShadowRelevance = bHasDynamicData && SceneProxy->IsShadowCast(View);
-	Result.bDynamicRelevance = bHasDynamicData;
-	if (bHasDynamicData)
+	if (HasDynamicData())
 	{
 		Result.bOpaque = View->Family->EngineShowFlags.Bounds;
+		Result.bRenderInSecondStageDepthPass = bRendersInSecondaryDepthPass;
 		DynamicDataRender->GetMaterialRelevance().SetPrimitiveViewRelevance(Result);
 	}
-	Result.bRenderInSecondStageDepthPass = bRendersInSecondaryDepthPass;
 
 	return Result;
 }
