- Allow newly spawned systems to avoid adding to concurrently running simulations
- Situations can occur where we execute the simulations in a TG and then attempt in a subsequent TG or an async task on the GT to add to the concurrently running simulation
- The new path will push the instance into an existing simulation that is not running concurrently or a new system simulation
- This also fixes a potential issue with TG transfers where we might of added to the wrong simulation
- The CVar fx.Niagara.WorldManager.AvoidAddingToExecutingSims can be used to flip between old and new behavior

diff --git a/Engine/Plugins/FX/Niagara/Source/Niagara/Private/NiagaraSystemInstance.cpp b/Engine/Plugins/FX/Niagara/Source/Niagara/Private/NiagaraSystemInstance.cpp
--- a/Engine/Plugins/FX/Niagara/Source/Niagara/Private/NiagaraSystemInstance.cpp
+++ b/Engine/Plugins/FX/Niagara/Source/Niagara/Private/NiagaraSystemInstance.cpp
@@ -460,7 +460,7 @@ void FNiagaraSystemInstance::SetSolo(bool bInSolo)
 		else
 		{
 			const ETickingGroup TickGroup = CalculateTickGroup();
-			TSharedPtr<FNiagaraSystemSimulation, ESPMode::ThreadSafe> NewSim = GetWorldManager()->GetSystemSimulation(TickGroup, System);
+			TSharedPtr<FNiagaraSystemSimulation, ESPMode::ThreadSafe> NewSim = GetWorldManager()->GetSystemSimulation(TickGroup, System, true);
 
 			NewSim->TransferInstance(this);
 		}
@@ -1044,7 +1044,7 @@ void FNiagaraSystemInstance::ReInitInternal()
 	else
 	{
 		const ETickingGroup TickGroup = CalculateTickGroup();
-		SystemSimulation = GetWorldManager()->GetSystemSimulation(TickGroup, System);
+		SystemSimulation = GetWorldManager()->GetSystemSimulation(TickGroup, System, true);
 	}
 
 	// Make sure that we've gotten propagated instance parameters before calling InitEmitters, as they might bind to them.
diff --git a/Engine/Plugins/FX/Niagara/Source/Niagara/Private/NiagaraSystemSimulation.cpp b/Engine/Plugins/FX/Niagara/Source/Niagara/Private/NiagaraSystemSimulation.cpp
--- a/Engine/Plugins/FX/Niagara/Source/Niagara/Private/NiagaraSystemSimulation.cpp
+++ b/Engine/Plugins/FX/Niagara/Source/Niagara/Private/NiagaraSystemSimulation.cpp
@@ -1086,13 +1086,15 @@ void FNiagaraSystemSimulation::Tick_GameThread_Internal(float DeltaSeconds, cons
 	}
 
 	check(IsInGameThread());
-	check(bInSpawnPhase == false);
 
 	FNiagaraCrashReporterScope CRScope(this);
 
 	// Work may not be complete if we back to back tick the GameThread without sending EOF updates
 	WaitForInstancesTickComplete();
 
+	// Ensure we waited correctly, this needs to be below the wait because back to back frames might not be flushed
+	check(bInSpawnPhase == false);
+
 	SCOPE_CYCLE_COUNTER(STAT_NiagaraOverview_GT);
 	SCOPE_CYCLE_COUNTER(STAT_NiagaraSystemSim_TickGT);
 	CSV_SCOPED_TIMING_STAT_EXCLUSIVE(Effects);
@@ -1152,7 +1154,7 @@ void FNiagaraSystemSimulation::Tick_GameThread_Internal(float DeltaSeconds, cons
 				// Tick demotion we need to do this now to ensure we complete in the correct group
 				if (DesiredTickGroup > SystemTickGroup)
 				{
-					TSharedPtr<FNiagaraSystemSimulation, ESPMode::ThreadSafe> NewSim = WorldManager->GetSystemSimulation(DesiredTickGroup, System);
+					TSharedPtr<FNiagaraSystemSimulation, ESPMode::ThreadSafe> NewSim = WorldManager->GetSystemSimulation(DesiredTickGroup, System, false);
 					NewSim->TransferInstance(Instance);
 					continue;
 				}
@@ -1186,7 +1188,7 @@ void FNiagaraSystemSimulation::Tick_GameThread_Internal(float DeltaSeconds, cons
 						if (UFXSystemComponent::RequiresLWCTileRecache(Instance->GetLWCTile(), SceneComponent->GetComponentLocation()))
 						{
 							//-OPT: For safety we reset everything, but if everything is local space we may not need to, or we could rebase.
-							if (Instance->GetTickCount() > 1)
+							if (Instance->GetTickCount() > 1 && FNiagaraUtilities::LogVerboseWarnings())
 							{
 								UE_LOG(LogNiagara, Warning, TEXT("NiagaraComponent(%s - %s) required LWC tile recache and was reset."), *GetFullNameSafe(SceneComponent), *GetFullNameSafe(System));
 							}
@@ -1223,7 +1225,7 @@ void FNiagaraSystemSimulation::Tick_GameThread_Internal(float DeltaSeconds, cons
 				const ETickingGroup DesiredTickGroup = Instance->CalculateTickGroup();
 				if (DesiredTickGroup != SystemTickGroup)
 				{
-					TSharedPtr<FNiagaraSystemSimulation, ESPMode::ThreadSafe> DestSim = WorldManager->GetSystemSimulation(DesiredTickGroup, System);
+					TSharedPtr<FNiagaraSystemSimulation, ESPMode::ThreadSafe> DestSim = WorldManager->GetSystemSimulation(DesiredTickGroup, System, false);
 					DestSim->TransferInstance(Instance);
 					continue;
 				}
@@ -1328,7 +1330,7 @@ void FNiagaraSystemSimulation::UpdateTickGroups_GameThread()
 		const ETickingGroup TickGroup = Instance->CalculateTickGroup();
 		if (TickGroup != SystemTickGroup)
 		{
-			TSharedPtr<FNiagaraSystemSimulation, ESPMode::ThreadSafe> NewSim = WorldManager->GetSystemSimulation(TickGroup, System);
+			TSharedPtr<FNiagaraSystemSimulation, ESPMode::ThreadSafe> NewSim = WorldManager->GetSystemSimulation(TickGroup, System, false);
 			NewSim->TransferInstance(Instance);
 		}
 	}
@@ -1343,7 +1345,7 @@ void FNiagaraSystemSimulation::UpdateTickGroups_GameThread()
 		const ETickingGroup DesiredTickGroup = Instance->CalculateTickGroup();
 		if (DesiredTickGroup != SystemTickGroup)
 		{
-			TSharedPtr<FNiagaraSystemSimulation, ESPMode::ThreadSafe> DestSim = WorldManager->GetSystemSimulation(DesiredTickGroup, System);
+			TSharedPtr<FNiagaraSystemSimulation, ESPMode::ThreadSafe> DestSim = WorldManager->GetSystemSimulation(DesiredTickGroup, System, false);
 			DestSim->TransferInstance(Instance);
 			continue;
 		}
@@ -1418,7 +1420,7 @@ void FNiagaraSystemSimulation::Spawn_GameThread(float DeltaSeconds, bool bPostAc
 		const ETickingGroup DesiredTickGroup = Instance->CalculateTickGroup();
 		if (DesiredTickGroup != SystemTickGroup)
 		{
-			TSharedPtr<FNiagaraSystemSimulation, ESPMode::ThreadSafe> DestSim = WorldManager->GetSystemSimulation(DesiredTickGroup, System);
+			TSharedPtr<FNiagaraSystemSimulation, ESPMode::ThreadSafe> DestSim = WorldManager->GetSystemSimulation(DesiredTickGroup, System, false);
 			DestSim->TransferInstance(Instance);
 			continue;
 		}
@@ -1581,6 +1583,11 @@ void FNiagaraSystemSimulation::DumpStalledInfo()
 	UE_LOG(LogNiagara, Fatal, TEXT("NiagaraSystemSimulation(%s) is stalled.\n%s"), *GetNameSafe(GetSystem()), Builder.ToString());
 }
 
+bool FNiagaraSystemSimulation::IsConcurrentRunning() const
+{
+	return ConcurrentTickGraphEvent.IsValid() && !ConcurrentTickGraphEvent->IsComplete();
+}
+
 void FNiagaraSystemSimulation::WaitForConcurrentTickComplete(bool bEnsureComplete)
 {
 	CSV_SCOPED_SET_WAIT_STAT(Effects);
diff --git a/Engine/Plugins/FX/Niagara/Source/Niagara/Private/NiagaraWorldManager.cpp b/Engine/Plugins/FX/Niagara/Source/Niagara/Private/NiagaraWorldManager.cpp
--- a/Engine/Plugins/FX/Niagara/Source/Niagara/Private/NiagaraWorldManager.cpp
+++ b/Engine/Plugins/FX/Niagara/Source/Niagara/Private/NiagaraWorldManager.cpp
@@ -96,6 +96,14 @@ static FAutoConsoleVariableRef CVarAllowVisibilityCullingForDynamicBounds(
 	ECVF_Default
 );
 
+static int GNiagaraWorldManagerAvoidAddingToExecutingSims = 1;
+static FAutoConsoleVariableRef CVarNiagaraWorldManagerAvoidAddingToExecutingSims(
+	TEXT("fx.Niagara.WorldManager.AvoidAddingToExecutingSims"),
+	GNiagaraWorldManagerAvoidAddingToExecutingSims,
+	TEXT("If system simulation is running concurrently already we will add new instances into a different sim to avoid blocking the GT."),
+	ECVF_Default
+);
+
 FAutoConsoleCommandWithWorld DumpNiagaraWorldManagerCommand(
 	TEXT("DumpNiagaraWorldManager"),
 	TEXT("Dump Information About the Niagara World Manager Contents"),
@@ -583,23 +591,56 @@ void FNiagaraWorldManager::CleanupParameterCollections()
 	ParameterCollections.Empty();
 }
 
-FNiagaraSystemSimulationRef FNiagaraWorldManager::GetSystemSimulation(ETickingGroup TickGroup, UNiagaraSystem* System)
+FNiagaraSystemSimulationRef FNiagaraWorldManager::GetSystemSimulation(ETickingGroup TickGroup, UNiagaraSystem* System, bool bAllowTickGroupChanges)
 {
 	LLM_SCOPE(ELLMTag::Niagara);
 
 	int32 ActualTickGroup = NiagaraWorldManagerInternal::GetNiagaraTickGroup(TickGroup);
-	if (ActiveNiagaraTickGroup == ActualTickGroup)
+
+	FNiagaraSystemSimulationRef* SimPtr = nullptr;
+	if (bAllowTickGroupChanges)
 	{
-		int32 DemotedTickGroup = NiagaraWorldManagerInternal::GetNiagaraTickGroup((ETickingGroup)(TickGroup + 1));
-		ActualTickGroup = DemotedTickGroup == ActualTickGroup ? 0 : DemotedTickGroup;
+		// Should we avoid adding a system simulation that might be in flights?
+		// This can occur in situations where something spawns a new instance after we have ticked the simulation.
+		// In which case we are better to add to an existing simulation that is not running concurrently or adding a new one
+		if ( GNiagaraWorldManagerAvoidAddingToExecutingSims )
+		{
+			SimPtr = SystemSimulations[ActualTickGroup].Find(System);
+			if ( SimPtr && (*SimPtr)->IsConcurrentRunning() )
+			{
+				for (int32 i=1; i < NiagaraNumTickGroups; ++i)
+				{
+					const int32 TestTickGroup = (ActualTickGroup + i) % NiagaraNumTickGroups;
+					SimPtr = SystemSimulations[TestTickGroup].Find(System);
+					if (SimPtr && (*SimPtr)->IsConcurrentRunning() == false)
+					{
+						return *SimPtr;
+					}
+				}
+
+				// We need to make a new system simulation
+				// Note we could probably find 'first free' slot but the likelyhood of it being last TG is low so not worth the complexity
+				const int32 DemotedTickGroup = NiagaraWorldManagerInternal::GetNiagaraTickGroup((ETickingGroup)(ActiveNiagaraTickGroup + 1));
+				ActualTickGroup = DemotedTickGroup == ActualTickGroup ? 0 : DemotedTickGroup;
+			}
+		}
+		else
+		{
+			if (ActiveNiagaraTickGroup == ActualTickGroup)
+			{
+				int32 DemotedTickGroup = NiagaraWorldManagerInternal::GetNiagaraTickGroup((ETickingGroup)(TickGroup + 1));
+				ActualTickGroup = DemotedTickGroup == ActualTickGroup ? 0 : DemotedTickGroup;
+			}
+		}
 	}
 
-	FNiagaraSystemSimulationRef* SimPtr = SystemSimulations[ActualTickGroup].Find(System);
+	SimPtr = SystemSimulations[ActualTickGroup].Find(System);
 	if (SimPtr != nullptr)
 	{
 		return *SimPtr;
 	}
-	
+
+	// Make a new simulation in the tick group
 	FNiagaraSystemSimulationRef Sim = MakeShared<FNiagaraSystemSimulation, ESPMode::ThreadSafe>();
 	SystemSimulations[ActualTickGroup].Add(System, Sim);
 	Sim->Init(System, World, false, TickGroup);
@@ -998,6 +1039,7 @@ void FNiagaraWorldManager::PostActorTick(float DeltaSeconds)
 	NiagaraWorldManagerInternal::ExecuteGlobalDeferredCallbacks();
 
 	DeltaSeconds *= DebugPlaybackRate;
+	ActiveNiagaraTickGroup = -1;
 
 	// Update any systems with post actor work
 	// - Instances that need to move to a higher tick group
@@ -1379,7 +1421,10 @@ void FNiagaraWorldManager::Tick(ETickingGroup TickGroup, float DeltaSeconds, ELe
 		TickFunctions[ActualTickGroup].EndTickGroup = TG_LastDemotable;
 	}
 
-	ActiveNiagaraTickGroup = -1;
+	if (TickGroup == NiagaraLastTickGroup)
+	{
+		ActiveNiagaraTickGroup = -1;
+	}
 
 	//We update scalability managers here so that any new systems can be culled or setup with other scalability based parameters correctly for their spawn.
 	UpdateScalabilityManagers(DeltaSeconds, true);
diff --git a/Engine/Plugins/FX/Niagara/Source/Niagara/Public/NiagaraSystemSimulation.h b/Engine/Plugins/FX/Niagara/Source/Niagara/Public/NiagaraSystemSimulation.h
--- a/Engine/Plugins/FX/Niagara/Source/Niagara/Public/NiagaraSystemSimulation.h
+++ b/Engine/Plugins/FX/Niagara/Source/Niagara/Public/NiagaraSystemSimulation.h
@@ -274,6 +274,8 @@ public:
 	/** Called after the sim cache has been read. */
 	void SimCachePostTick_Concurrent(float DeltaSeconds, const FGraphEventRef& MyCompletionGraphEvent);
 
+	bool IsConcurrentRunning() const;
+
 	/** Wait for system simulation concurrent tick to complete.  If bEnsureComplete is true we will trigger an ensure if it is not complete. */
 	void WaitForConcurrentTickComplete(bool bEnsureComplete = false);
 	/** Wait for system instances concurrent tick to complete.  If bEnsureComplete is true we will trigger an ensure if it is not complete. */
diff --git a/Engine/Plugins/FX/Niagara/Source/Niagara/Public/NiagaraWorldManager.h b/Engine/Plugins/FX/Niagara/Source/Niagara/Public/NiagaraWorldManager.h
--- a/Engine/Plugins/FX/Niagara/Source/Niagara/Public/NiagaraWorldManager.h
+++ b/Engine/Plugins/FX/Niagara/Source/Niagara/Public/NiagaraWorldManager.h
@@ -111,9 +111,9 @@ public:
 	UNiagaraParameterCollectionInstance* GetParameterCollection(UNiagaraParameterCollection* Collection);
 	void CleanupParameterCollections();
 
-	FNiagaraSystemSimulationRef GetSystemSimulation(ETickingGroup TickGroup, UNiagaraSystem* System);
+	FNiagaraSystemSimulationRef GetSystemSimulation(ETickingGroup TickGroup, UNiagaraSystem* System, bool bAllowTickGroupChanges);
 	void DestroySystemSimulation(UNiagaraSystem* System);
-	void DestroySystemInstance(FNiagaraSystemInstancePtr& InPtr);	
+	void DestroySystemInstance(FNiagaraSystemInstancePtr& InPtr);
 
 #if WITH_EDITOR
 	void OnSystemPostChange(UNiagaraSystem* System);
