- Disable parallel scene proxy creates in Editor
- Various other UObjects we might query are unable to run wide during EOF updates because they might invalidate render state

diff --git a/Engine/Plugins/FX/Niagara/Source/Niagara/Private/NiagaraComponent.cpp b/Engine/Plugins/FX/Niagara/Source/Niagara/Private/NiagaraComponent.cpp
--- a/Engine/Plugins/FX/Niagara/Source/Niagara/Private/NiagaraComponent.cpp
+++ b/Engine/Plugins/FX/Niagara/Source/Niagara/Private/NiagaraComponent.cpp
@@ -790,6 +790,17 @@ FName UNiagaraComponent::GetFNameForStatID() const
 	return Asset ? Asset->GetFNameForStatID() : Super::GetFNameForStatID();
 }
 
+bool UNiagaraComponent::RequiresGameThreadEndOfFrameRecreate() const
+{
+#if WITH_EDITORONLY_DATA
+	// Niagara is capable of running async but in the editor.  However various UObjects we might use during a create are not.
+	// For example, resolving a UStaticMesh will query the mesh compiler which will propagate changes and mark render state dirty which is not allowed during EOF updates.
+	return true;
+#else
+	return false;
+#endif
+}
+
 void UNiagaraComponent::TickComponent(float DeltaSeconds, enum ELevelTick TickType, FActorComponentTickFunction* ThisTickFunction)
 {
 	LLM_SCOPE(ELLMTag::Niagara);
diff --git a/Engine/Plugins/FX/Niagara/Source/Niagara/Public/NiagaraComponent.h b/Engine/Plugins/FX/Niagara/Source/Niagara/Public/NiagaraComponent.h
--- a/Engine/Plugins/FX/Niagara/Source/Niagara/Public/NiagaraComponent.h
+++ b/Engine/Plugins/FX/Niagara/Source/Niagara/Public/NiagaraComponent.h
@@ -85,7 +85,7 @@ public:
 
 	NIAGARA_API virtual FName GetFNameForStatID() const override;
 
-	virtual bool RequiresGameThreadEndOfFrameRecreate() const override { return false; }
+	NIAGARA_API virtual bool RequiresGameThreadEndOfFrameRecreate() const override;
 
 private:
 	UPROPERTY(EditAnywhere, Category="Niagara", meta = (DisplayName = "Niagara System Asset"))
