Fix for island type data channel systems being deactivated despite still being used

diff --git a/Engine/Plugins/FX/Niagara/Source/Niagara/Internal/NiagaraDataChannelCommon.h b/Engine/Plugins/FX/Niagara/Source/Niagara/Internal/NiagaraDataChannelCommon.h
--- a/Engine/Plugins/FX/Niagara/Source/Niagara/Internal/NiagaraDataChannelCommon.h
+++ b/Engine/Plugins/FX/Niagara/Source/Niagara/Internal/NiagaraDataChannelCommon.h
@@ -71,8 +71,7 @@ struct FNiagaraDataChannelData final
 
 	NIAGARA_API void BeginFrame(UNiagaraDataChannelHandler* Owner);
 	NIAGARA_API void EndFrame(UNiagaraDataChannelHandler* Owner);
-	NIAGARA_API void Tick(UNiagaraDataChannelHandler* Owner);
-	NIAGARA_API void ConsumePublishRequests(UNiagaraDataChannelHandler* Owner);
+	NIAGARA_API int32 ConsumePublishRequests(UNiagaraDataChannelHandler* Owner);
 
 	NIAGARA_API FNiagaraDataChannelGameData* GetGameData();
 	NIAGARA_API FNiagaraDataBufferRef GetCPUData(bool bPreviousFrame);
diff --git a/Engine/Plugins/FX/Niagara/Source/Niagara/Private/NiagaraDataChannel.cpp b/Engine/Plugins/FX/Niagara/Source/Niagara/Private/NiagaraDataChannel.cpp
--- a/Engine/Plugins/FX/Niagara/Source/Niagara/Private/NiagaraDataChannel.cpp
+++ b/Engine/Plugins/FX/Niagara/Source/Niagara/Private/NiagaraDataChannel.cpp
@@ -176,6 +176,8 @@ void FNiagaraDataChannelGameData::WriteToDataSet(FNiagaraDataBuffer* DestBuffer,
 		TFunction<void(UScriptStruct*, UScriptStruct*, uint8*)> WriteData;
 		WriteData = [&](UScriptStruct* SrcStruct, UScriptStruct* DestStruct, uint8* SrcPropertyBase)
 		{
+			//TODO: why doesn't this just use FNiagaraLwcStructConverter?
+			
 			//Write all data from the data channel into the data set. Converting into from LWC into the local LWC Tile space as we go.
 
 			uint8* SrcData = SrcPropertyBase;
@@ -579,7 +581,7 @@ void FNiagaraDataChannelData::EndFrame(UNiagaraDataChannelHandler* Owner)
 	BuffersForGPU.Reset();
 }
 
-void FNiagaraDataChannelData::ConsumePublishRequests(UNiagaraDataChannelHandler* Owner)
+int32 FNiagaraDataChannelData::ConsumePublishRequests(UNiagaraDataChannelHandler* Owner)
 {
 	check(IsValid(Owner));
 
@@ -589,7 +591,7 @@ void FNiagaraDataChannelData::ConsumePublishRequests(UNiagaraDataChannelHandler*
 	const UNiagaraDataChannel* DataChannel = Owner->GetDataChannel();
 	if(PublishRequests.Num() == 0 || DataChannel == nullptr)
 	{
-		return;
+		return 0;
 	}
 	
 	if(NDCCVars::bEmitWarningsOnLateNDCWrites && DataChannel->ShouldEnforceTickGroupReadWriteOrder())
@@ -619,7 +621,8 @@ void FNiagaraDataChannelData::ConsumePublishRequests(UNiagaraDataChannelHandler*
 	//Do a pass to gather the new total size for our DataChannel data.
 
 	//Each DI that generates DataChannel can control whether it's pushed to Game/CPU/GPU.
-	BuffersForGPU.Reserve(BuffersForGPU.Num() + PublishRequests.Num());
+	int32 RequestCount = PublishRequests.Num();
+	BuffersForGPU.Reserve(BuffersForGPU.Num() + RequestCount);
 
 	for (FNiagaraDataChannelPublishRequest& PublishRequest : PublishRequests)
 	{
@@ -728,6 +731,7 @@ void FNiagaraDataChannelData::ConsumePublishRequests(UNiagaraDataChannelHandler*
 #endif
 
 	PublishRequests.Reset();
+	return RequestCount;
 }
 
 FNiagaraDataBufferRef FNiagaraDataChannelData::GetCPUData(bool bPreviousFrame)
diff --git a/Engine/Plugins/FX/Niagara/Source/Niagara/Private/NiagaraDataChannel_Islands.cpp b/Engine/Plugins/FX/Niagara/Source/Niagara/Private/NiagaraDataChannel_Islands.cpp
--- a/Engine/Plugins/FX/Niagara/Source/Niagara/Private/NiagaraDataChannel_Islands.cpp
+++ b/Engine/Plugins/FX/Niagara/Source/Niagara/Private/NiagaraDataChannel_Islands.cpp
@@ -401,7 +401,17 @@ void FNDCIsland::EndFrame()
 
 void FNDCIsland::Tick()
 {
-	Data->ConsumePublishRequests(Owner);
+	int32 AddedData = Data->ConsumePublishRequests(Owner);
+	if (IsBeingUsed() && AddedData > 0)
+	{
+		for (UNiagaraComponent* Comp : NiagaraSystems)
+		{
+			if (Comp->IsComplete())
+			{
+				Comp->Activate();
+			}
+		}
+	}
 }
 
 bool FNDCIsland::Contains(FVector Point)
