- Temporarily path to avoid issues

diff --git a/Engine/Plugins/FX/Niagara/Source/Niagara/Private/NiagaraWorldManager.cpp b/Engine/Plugins/FX/Niagara/Source/Niagara/Private/NiagaraWorldManager.cpp
--- a/Engine/Plugins/FX/Niagara/Source/Niagara/Private/NiagaraWorldManager.cpp
+++ b/Engine/Plugins/FX/Niagara/Source/Niagara/Private/NiagaraWorldManager.cpp
@@ -96,7 +96,7 @@ static FAutoConsoleVariableRef CVarAllowVisibilityCullingForDynamicBounds(
 	ECVF_Default
 );
 
-static int GNiagaraWorldManagerAvoidAddingToExecutingSims = 1;
+static int GNiagaraWorldManagerAvoidAddingToExecutingSims = 0;
 static FAutoConsoleVariableRef CVarNiagaraWorldManagerAvoidAddingToExecutingSims(
 	TEXT("fx.Niagara.WorldManager.AvoidAddingToExecutingSims"),
 	GNiagaraWorldManagerAvoidAddingToExecutingSims,
