- Add missing tooltip when we have nothing to paste

diff --git a/Engine/Plugins/FX/Niagara/Source/NiagaraEditorWidgets/Private/NiagaraStackCommandContext.cpp b/Engine/Plugins/FX/Niagara/Source/NiagaraEditorWidgets/Private/NiagaraStackCommandContext.cpp
--- a/Engine/Plugins/FX/Niagara/Source/NiagaraEditorWidgets/Private/NiagaraStackCommandContext.cpp
+++ b/Engine/Plugins/FX/Niagara/Source/NiagaraEditorWidgets/Private/NiagaraStackCommandContext.cpp
@@ -191,6 +191,10 @@ FText FNiagaraStackCommandContext::GetCanPasteSelectedEntriesToolTip() const
 {
 	FText CanPasteMessage;
 	FNiagaraStackClipboardUtilities::TestCanPasteSelectionWithMessage(SelectedEntries, CanPasteMessage);
+	if (CanPasteMessage.IsEmptyOrWhitespace())
+	{
+		CanPasteMessage = LOCTEXT("NoValidDataToPaste", "Incompatible or no data to paste.");
+	}
 	return CanPasteMessage;
 }
 
