- Fix low latency meshes not rendering in shader complexity view modes

diff --git a/Engine/Plugins/FX/Niagara/Source/Niagara/Private/NiagaraRendererMeshes.cpp b/Engine/Plugins/FX/Niagara/Source/Niagara/Private/NiagaraRendererMeshes.cpp
--- a/Engine/Plugins/FX/Niagara/Source/Niagara/Private/NiagaraRendererMeshes.cpp
+++ b/Engine/Plugins/FX/Niagara/Source/Niagara/Private/NiagaraRendererMeshes.cpp
@@ -342,11 +342,12 @@ void FNiagaraRendererMeshes::PrepareParticleMeshRenderData(FParticleMeshRenderDa
 	// If these conditions change please update the DebugHUD display also to reflect it
 	// Note: SceneCaptures will use latent data as GpuReadyTickStage < CurrentParticleData->GetGPUDataReadyStage()
 	//       For main pass scene captures we exclude the batches if they are translucent
-	const bool bIsWireframe = AllowDebugViewmodes() && ViewFamily.EngineShowFlags.Wireframe;
+	const bool bIsExcludedViewMode = AllowDebugViewmodes() &&
+		(ViewFamily.EngineShowFlags.Wireframe || ViewFamily.EngineShowFlags.ShaderComplexity || ViewFamily.EngineShowFlags.ShaderComplexityWithQuadOverdraw);
 	ParticleMeshRenderData.bIsGpuLowLatencyTranslucency =
 		bGpuLowLatencyTranslucency &&
 		GpuReadyTickStage >= CurrentParticleData->GetGPUDataReadyStage() &&
-		!bIsWireframe &&
+		!bIsExcludedViewMode &&
 		!SceneProxy->CastsVolumetricTranslucentShadow() &&
 		ParticleMeshRenderData.DynamicDataMesh->Materials.Num() > 0 &&
 		ParticleMeshRenderData.DynamicDataMesh->IsGpuLowLatencyTranslucencyEnabled() &&
