- Reset OverrideFormat on return to the pool as some will use RenderTargetFormat vs OverrideFormat

diff --git a/Engine/Plugins/FX/Niagara/Source/Niagara/Private/NiagaraDataInterfaceRenderTargetCommon.h b/Engine/Plugins/FX/Niagara/Source/Niagara/Private/NiagaraDataInterfaceRenderTargetCommon.h
--- a/Engine/Plugins/FX/Niagara/Source/Niagara/Private/NiagaraDataInterfaceRenderTargetCommon.h
+++ b/Engine/Plugins/FX/Niagara/Source/Niagara/Private/NiagaraDataInterfaceRenderTargetCommon.h
@@ -38,6 +38,8 @@ namespace NiagaraDataInterfaceRenderTargetCommon
 		{
 			if (FNiagaraWorldManager* WorldManager = SystemInstance->GetWorldManager())
 			{
+				InstanceData->TargetTexture->OverrideFormat = PF_Unknown;
+
 				WorldManager->ObjectPoolReturn(InstanceData->TargetTexture);
 				WorldManager->RemoveReferencedObject(InstanceData->TargetTexture);
 				InstanceData->TargetTexture = nullptr;
