Expose ID setter for niagara data channels

diff --git a/Engine/Plugins/FX/Niagara/Source/Niagara/Private/NiagaraDataChannelAccessor.cpp b/Engine/Plugins/FX/Niagara/Source/Niagara/Private/NiagaraDataChannelAccessor.cpp
--- a/Engine/Plugins/FX/Niagara/Source/Niagara/Private/NiagaraDataChannelAccessor.cpp
+++ b/Engine/Plugins/FX/Niagara/Source/Niagara/Private/NiagaraDataChannelAccessor.cpp
@@ -221,6 +221,11 @@ void UNiagaraDataChannelWriter::WritePosition(FName VarName, int32 Index, FVecto
 	WriteData(FNiagaraVariableBase(FNiagaraTypeDefinition::GetPositionDef(), VarName), Index, InData);
 }
 
+void UNiagaraDataChannelWriter::WriteID(FName VarName, int32 Index, FNiagaraID InData)
+{
+	WriteData(FNiagaraVariableBase(FNiagaraTypeDefinition::GetIDDef(), VarName), Index, InData);
+}
+
 
 //////////////////////////////////////////////////////////////////////////
 
diff --git a/Engine/Plugins/FX/Niagara/Source/Niagara/Public/NiagaraDataChannelAccessor.h b/Engine/Plugins/FX/Niagara/Source/Niagara/Public/NiagaraDataChannelAccessor.h
--- a/Engine/Plugins/FX/Niagara/Source/Niagara/Public/NiagaraDataChannelAccessor.h
+++ b/Engine/Plugins/FX/Niagara/Source/Niagara/Public/NiagaraDataChannelAccessor.h
@@ -127,6 +127,11 @@ public:
 
 	UFUNCTION(BlueprintCallable, Category = NiagaraDataChannel, meta = (Keywords = "niagara DataChannel"))
 	NIAGARA_API void WritePosition(FName VarName, int32 Index, FVector InData);
+
+	UFUNCTION(BlueprintCallable, Category = NiagaraDataChannel, meta = (Keywords = "niagara DataChannel"))
+	NIAGARA_API void WriteID(FName VarName, int32 Index, FNiagaraID InData);
+
+
 };
 
 
