[Backout] - CL29280970 - Removing code, either we move to the new TG or we attempt to fix it in the future, but for the moment removing the change as it doesn't offer any value

diff --git a/Engine/Plugins/FX/Niagara/Source/Niagara/Private/NiagaraSystemInstance.cpp b/Engine/Plugins/FX/Niagara/Source/Niagara/Private/NiagaraSystemInstance.cpp
--- a/Engine/Plugins/FX/Niagara/Source/Niagara/Private/NiagaraSystemInstance.cpp
+++ b/Engine/Plugins/FX/Niagara/Source/Niagara/Private/NiagaraSystemInstance.cpp
@@ -460,7 +460,7 @@ void FNiagaraSystemInstance::SetSolo(bool bInSolo)
 		else
 		{
 			const ETickingGroup TickGroup = CalculateTickGroup();
-			TSharedPtr<FNiagaraSystemSimulation, ESPMode::ThreadSafe> NewSim = GetWorldManager()->GetSystemSimulation(TickGroup, System, true);
+			TSharedPtr<FNiagaraSystemSimulation, ESPMode::ThreadSafe> NewSim = GetWorldManager()->GetSystemSimulation(TickGroup, System);
 
 			NewSim->TransferInstance(this);
 		}
@@ -1044,7 +1044,7 @@ void FNiagaraSystemInstance::ReInitInternal()
 	else
 	{
 		const ETickingGroup TickGroup = CalculateTickGroup();
-		SystemSimulation = GetWorldManager()->GetSystemSimulation(TickGroup, System, true);
+		SystemSimulation = GetWorldManager()->GetSystemSimulation(TickGroup, System);
 	}
 
 	// Make sure that we've gotten propagated instance parameters before calling InitEmitters, as they might bind to them.
diff --git a/Engine/Plugins/FX/Niagara/Source/Niagara/Private/NiagaraSystemSimulation.cpp b/Engine/Plugins/FX/Niagara/Source/Niagara/Private/NiagaraSystemSimulation.cpp
--- a/Engine/Plugins/FX/Niagara/Source/Niagara/Private/NiagaraSystemSimulation.cpp
+++ b/Engine/Plugins/FX/Niagara/Source/Niagara/Private/NiagaraSystemSimulation.cpp
@@ -1086,15 +1086,13 @@ void FNiagaraSystemSimulation::Tick_GameThread_Internal(float DeltaSeconds, cons
 	}
 
 	check(IsInGameThread());
+	check(bInSpawnPhase == false);
 
 	FNiagaraCrashReporterScope CRScope(this);
 
 	// Work may not be complete if we back to back tick the GameThread without sending EOF updates
 	WaitForInstancesTickComplete();
 
-	// Ensure we waited correctly, this needs to be below the wait because back to back frames might not be flushed
-	check(bInSpawnPhase == false);
-
 	SCOPE_CYCLE_COUNTER(STAT_NiagaraOverview_GT);
 	SCOPE_CYCLE_COUNTER(STAT_NiagaraSystemSim_TickGT);
 	CSV_SCOPED_TIMING_STAT_EXCLUSIVE(Effects);
@@ -1154,7 +1152,7 @@ void FNiagaraSystemSimulation::Tick_GameThread_Internal(float DeltaSeconds, cons
 				// Tick demotion we need to do this now to ensure we complete in the correct group
 				if (DesiredTickGroup > SystemTickGroup)
 				{
-					TSharedPtr<FNiagaraSystemSimulation, ESPMode::ThreadSafe> NewSim = WorldManager->GetSystemSimulation(DesiredTickGroup, System, false);
+					TSharedPtr<FNiagaraSystemSimulation, ESPMode::ThreadSafe> NewSim = WorldManager->GetSystemSimulation(DesiredTickGroup, System);
 					NewSim->TransferInstance(Instance);
 					continue;
 				}
@@ -1188,7 +1186,7 @@ void FNiagaraSystemSimulation::Tick_GameThread_Internal(float DeltaSeconds, cons
 						if (UFXSystemComponent::RequiresLWCTileRecache(Instance->GetLWCTile(), SceneComponent->GetComponentLocation()))
 						{
 							//-OPT: For safety we reset everything, but if everything is local space we may not need to, or we could rebase.
-							if (Instance->GetTickCount() > 1 && FNiagaraUtilities::LogVerboseWarnings())
+							if (Instance->GetTickCount() > 1)
 							{
 								UE_LOG(LogNiagara, Warning, TEXT("NiagaraComponent(%s - %s) required LWC tile recache and was reset."), *GetFullNameSafe(SceneComponent), *GetFullNameSafe(System));
 							}
@@ -1225,7 +1223,7 @@ void FNiagaraSystemSimulation::Tick_GameThread_Internal(float DeltaSeconds, cons
 				const ETickingGroup DesiredTickGroup = Instance->CalculateTickGroup();
 				if (DesiredTickGroup != SystemTickGroup)
 				{
-					TSharedPtr<FNiagaraSystemSimulation, ESPMode::ThreadSafe> DestSim = WorldManager->GetSystemSimulation(DesiredTickGroup, System, false);
+					TSharedPtr<FNiagaraSystemSimulation, ESPMode::ThreadSafe> DestSim = WorldManager->GetSystemSimulation(DesiredTickGroup, System);
 					DestSim->TransferInstance(Instance);
 					continue;
 				}
@@ -1330,7 +1328,7 @@ void FNiagaraSystemSimulation::UpdateTickGroups_GameThread()
 		const ETickingGroup TickGroup = Instance->CalculateTickGroup();
 		if (TickGroup != SystemTickGroup)
 		{
-			TSharedPtr<FNiagaraSystemSimulation, ESPMode::ThreadSafe> NewSim = WorldManager->GetSystemSimulation(TickGroup, System, false);
+			TSharedPtr<FNiagaraSystemSimulation, ESPMode::ThreadSafe> NewSim = WorldManager->GetSystemSimulation(TickGroup, System);
 			NewSim->TransferInstance(Instance);
 		}
 	}
@@ -1345,7 +1343,7 @@ void FNiagaraSystemSimulation::UpdateTickGroups_GameThread()
 		const ETickingGroup DesiredTickGroup = Instance->CalculateTickGroup();
 		if (DesiredTickGroup != SystemTickGroup)
 		{
-			TSharedPtr<FNiagaraSystemSimulation, ESPMode::ThreadSafe> DestSim = WorldManager->GetSystemSimulation(DesiredTickGroup, System, false);
+			TSharedPtr<FNiagaraSystemSimulation, ESPMode::ThreadSafe> DestSim = WorldManager->GetSystemSimulation(DesiredTickGroup, System);
 			DestSim->TransferInstance(Instance);
 			continue;
 		}
@@ -1420,7 +1418,7 @@ void FNiagaraSystemSimulation::Spawn_GameThread(float DeltaSeconds, bool bPostAc
 		const ETickingGroup DesiredTickGroup = Instance->CalculateTickGroup();
 		if (DesiredTickGroup != SystemTickGroup)
 		{
-			TSharedPtr<FNiagaraSystemSimulation, ESPMode::ThreadSafe> DestSim = WorldManager->GetSystemSimulation(DesiredTickGroup, System, false);
+			TSharedPtr<FNiagaraSystemSimulation, ESPMode::ThreadSafe> DestSim = WorldManager->GetSystemSimulation(DesiredTickGroup, System);
 			DestSim->TransferInstance(Instance);
 			continue;
 		}
@@ -1583,11 +1581,6 @@ void FNiagaraSystemSimulation::DumpStalledInfo()
 	UE_LOG(LogNiagara, Fatal, TEXT("NiagaraSystemSimulation(%s) is stalled.\n%s"), *GetNameSafe(GetSystem()), Builder.ToString());
 }
 
-bool FNiagaraSystemSimulation::IsConcurrentRunning() const
-{
-	return ConcurrentTickGraphEvent.IsValid() && !ConcurrentTickGraphEvent->IsComplete();
-}
-
 void FNiagaraSystemSimulation::WaitForConcurrentTickComplete(bool bEnsureComplete)
 {
 	CSV_SCOPED_SET_WAIT_STAT(Effects);
diff --git a/Engine/Plugins/FX/Niagara/Source/Niagara/Private/NiagaraWorldManager.cpp b/Engine/Plugins/FX/Niagara/Source/Niagara/Private/NiagaraWorldManager.cpp
--- a/Engine/Plugins/FX/Niagara/Source/Niagara/Private/NiagaraWorldManager.cpp
+++ b/Engine/Plugins/FX/Niagara/Source/Niagara/Private/NiagaraWorldManager.cpp
@@ -96,14 +96,6 @@ static FAutoConsoleVariableRef CVarAllowVisibilityCullingForDynamicBounds(
 	ECVF_Default
 );
 
-static int GNiagaraWorldManagerAvoidAddingToExecutingSims = 0;
-static FAutoConsoleVariableRef CVarNiagaraWorldManagerAvoidAddingToExecutingSims(
-	TEXT("fx.Niagara.WorldManager.AvoidAddingToExecutingSims"),
-	GNiagaraWorldManagerAvoidAddingToExecutingSims,
-	TEXT("If system simulation is running concurrently already we will add new instances into a different sim to avoid blocking the GT."),
-	ECVF_Default
-);
-
 FAutoConsoleCommandWithWorld DumpNiagaraWorldManagerCommand(
 	TEXT("DumpNiagaraWorldManager"),
 	TEXT("Dump Information About the Niagara World Manager Contents"),
@@ -591,56 +583,23 @@ void FNiagaraWorldManager::CleanupParameterCollections()
 	ParameterCollections.Empty();
 }
 
-FNiagaraSystemSimulationRef FNiagaraWorldManager::GetSystemSimulation(ETickingGroup TickGroup, UNiagaraSystem* System, bool bAllowTickGroupChanges)
+FNiagaraSystemSimulationRef FNiagaraWorldManager::GetSystemSimulation(ETickingGroup TickGroup, UNiagaraSystem* System)
 {
 	LLM_SCOPE(ELLMTag::Niagara);
 
 	int32 ActualTickGroup = NiagaraWorldManagerInternal::GetNiagaraTickGroup(TickGroup);
-
-	FNiagaraSystemSimulationRef* SimPtr = nullptr;
-	if (bAllowTickGroupChanges)
+	if (ActiveNiagaraTickGroup == ActualTickGroup)
 	{
-		// Should we avoid adding a system simulation that might be in flights?
-		// This can occur in situations where something spawns a new instance after we have ticked the simulation.
-		// In which case we are better to add to an existing simulation that is not running concurrently or adding a new one
-		if ( GNiagaraWorldManagerAvoidAddingToExecutingSims )
-		{
-			SimPtr = SystemSimulations[ActualTickGroup].Find(System);
-			if ( SimPtr && (*SimPtr)->IsConcurrentRunning() )
-			{
-				for (int32 i=1; i < NiagaraNumTickGroups; ++i)
-				{
-					const int32 TestTickGroup = (ActualTickGroup + i) % NiagaraNumTickGroups;
-					SimPtr = SystemSimulations[TestTickGroup].Find(System);
-					if (SimPtr && (*SimPtr)->IsConcurrentRunning() == false)
-					{
-						return *SimPtr;
-					}
-				}
-
-				// We need to make a new system simulation
-				// Note we could probably find 'first free' slot but the likelyhood of it being last TG is low so not worth the complexity
-				const int32 DemotedTickGroup = NiagaraWorldManagerInternal::GetNiagaraTickGroup((ETickingGroup)(ActiveNiagaraTickGroup + 1));
-				ActualTickGroup = DemotedTickGroup == ActualTickGroup ? 0 : DemotedTickGroup;
-			}
-		}
-		else
-		{
-			if (ActiveNiagaraTickGroup == ActualTickGroup)
-			{
-				int32 DemotedTickGroup = NiagaraWorldManagerInternal::GetNiagaraTickGroup((ETickingGroup)(TickGroup + 1));
-				ActualTickGroup = DemotedTickGroup == ActualTickGroup ? 0 : DemotedTickGroup;
-			}
-		}
+		int32 DemotedTickGroup = NiagaraWorldManagerInternal::GetNiagaraTickGroup((ETickingGroup)(TickGroup + 1));
+		ActualTickGroup = DemotedTickGroup == ActualTickGroup ? 0 : DemotedTickGroup;
 	}
 
-	SimPtr = SystemSimulations[ActualTickGroup].Find(System);
+	FNiagaraSystemSimulationRef* SimPtr = SystemSimulations[ActualTickGroup].Find(System);
 	if (SimPtr != nullptr)
 	{
 		return *SimPtr;
 	}
-
-	// Make a new simulation in the tick group
+	
 	FNiagaraSystemSimulationRef Sim = MakeShared<FNiagaraSystemSimulation, ESPMode::ThreadSafe>();
 	SystemSimulations[ActualTickGroup].Add(System, Sim);
 	Sim->Init(System, World, false, TickGroup);
@@ -1039,7 +998,6 @@ void FNiagaraWorldManager::PostActorTick(float DeltaSeconds)
 	NiagaraWorldManagerInternal::ExecuteGlobalDeferredCallbacks();
 
 	DeltaSeconds *= DebugPlaybackRate;
-	ActiveNiagaraTickGroup = -1;
 
 	// Update any systems with post actor work
 	// - Instances that need to move to a higher tick group
@@ -1421,10 +1379,7 @@ void FNiagaraWorldManager::Tick(ETickingGroup TickGroup, float DeltaSeconds, ELe
 		TickFunctions[ActualTickGroup].EndTickGroup = TG_LastDemotable;
 	}
 
-	if (TickGroup == NiagaraLastTickGroup)
-	{
-		ActiveNiagaraTickGroup = -1;
-	}
+	ActiveNiagaraTickGroup = -1;
 
 	//We update scalability managers here so that any new systems can be culled or setup with other scalability based parameters correctly for their spawn.
 	UpdateScalabilityManagers(DeltaSeconds, true);
diff --git a/Engine/Plugins/FX/Niagara/Source/Niagara/Public/NiagaraSystemSimulation.h b/Engine/Plugins/FX/Niagara/Source/Niagara/Public/NiagaraSystemSimulation.h
--- a/Engine/Plugins/FX/Niagara/Source/Niagara/Public/NiagaraSystemSimulation.h
+++ b/Engine/Plugins/FX/Niagara/Source/Niagara/Public/NiagaraSystemSimulation.h
@@ -274,8 +274,6 @@ public:
 	/** Called after the sim cache has been read. */
 	void SimCachePostTick_Concurrent(float DeltaSeconds, const FGraphEventRef& MyCompletionGraphEvent);
 
-	bool IsConcurrentRunning() const;
-
 	/** Wait for system simulation concurrent tick to complete.  If bEnsureComplete is true we will trigger an ensure if it is not complete. */
 	void WaitForConcurrentTickComplete(bool bEnsureComplete = false);
 	/** Wait for system instances concurrent tick to complete.  If bEnsureComplete is true we will trigger an ensure if it is not complete. */
diff --git a/Engine/Plugins/FX/Niagara/Source/Niagara/Public/NiagaraWorldManager.h b/Engine/Plugins/FX/Niagara/Source/Niagara/Public/NiagaraWorldManager.h
--- a/Engine/Plugins/FX/Niagara/Source/Niagara/Public/NiagaraWorldManager.h
+++ b/Engine/Plugins/FX/Niagara/Source/Niagara/Public/NiagaraWorldManager.h
@@ -111,9 +111,9 @@ public:
 	UNiagaraParameterCollectionInstance* GetParameterCollection(UNiagaraParameterCollection* Collection);
 	void CleanupParameterCollections();
 
-	FNiagaraSystemSimulationRef GetSystemSimulation(ETickingGroup TickGroup, UNiagaraSystem* System, bool bAllowTickGroupChanges);
+	FNiagaraSystemSimulationRef GetSystemSimulation(ETickingGroup TickGroup, UNiagaraSystem* System);
 	void DestroySystemSimulation(UNiagaraSystem* System);
-	void DestroySystemInstance(FNiagaraSystemInstancePtr& InPtr);
+	void DestroySystemInstance(FNiagaraSystemInstancePtr& InPtr);	
 
 #if WITH_EDITOR
 	void OnSystemPostChange(UNiagaraSystem* System);
