Replacing c-style casts with explicit versions + small code cleanups

diff --git a/Engine/Plugins/FX/Niagara/Source/Niagara/Private/NiagaraDataInterfaceExport.cpp b/Engine/Plugins/FX/Niagara/Source/Niagara/Private/NiagaraDataInterfaceExport.cpp
--- a/Engine/Plugins/FX/Niagara/Source/Niagara/Private/NiagaraDataInterfaceExport.cpp
+++ b/Engine/Plugins/FX/Niagara/Source/Niagara/Private/NiagaraDataInterfaceExport.cpp
@@ -263,7 +263,7 @@ bool UNiagaraDataInterfaceExport::PerInstanceTick(void* PerInstanceData, FNiagar
 
 bool UNiagaraDataInterfaceExport::PerInstanceTickPostSimulate(void* PerInstanceData, FNiagaraSystemInstance* SystemInstance, float DeltaSeconds)
 {
-	FNDIExportInstanceData_GameThread* PIData = (FNDIExportInstanceData_GameThread*) PerInstanceData;
+	FNDIExportInstanceData_GameThread* PIData = static_cast<FNDIExportInstanceData_GameThread*>(PerInstanceData);
 	if ( !PIData->GatheredData.IsEmpty() )
 	{
 		UObject* CallbackHandler = PIData->CallbackHandler.Get();
diff --git a/Engine/Plugins/FX/Niagara/Source/Niagara/Private/NiagaraTypes.cpp b/Engine/Plugins/FX/Niagara/Source/Niagara/Private/NiagaraTypes.cpp
--- a/Engine/Plugins/FX/Niagara/Source/Niagara/Private/NiagaraTypes.cpp
+++ b/Engine/Plugins/FX/Niagara/Source/Niagara/Private/NiagaraTypes.cpp
@@ -306,21 +306,21 @@ FString FNiagaraTypeHelper::ToString(const uint8* ValueData, const UObject* Stru
 	FString Ret;
 	if (const UEnum* Enum = Cast<const UEnum>(StructOrEnum))
 	{
-		Ret = Enum->GetNameStringByValue(*(int32*)ValueData);
+		Ret = Enum->GetNameStringByValue(*reinterpret_cast<const int32*>(ValueData));
 	}
 	else if (const UScriptStruct* Struct = Cast<const UScriptStruct>(StructOrEnum))
 	{
 		if (Struct == FNiagaraTypeDefinition::GetFloatStruct())
 		{
-			Ret += FString::Printf(TEXT("%g "), *(float*)ValueData);
+			Ret += FString::Printf(TEXT("%g "), *reinterpret_cast<const float*>(ValueData));
 		}
 		else if (Struct == FNiagaraTypeDefinition::GetIntStruct())
 		{
-			Ret += FString::Printf(TEXT("%d "), *(int32*)ValueData);
+			Ret += FString::Printf(TEXT("%d "), *reinterpret_cast<const int32*>(ValueData));
 		}
 		else if (Struct == FNiagaraTypeDefinition::GetBoolStruct())
 		{
-			int32 Val = *(int32*)ValueData;
+			int32 Val = *reinterpret_cast<const int32*>(ValueData);
 			Ret += Val == 0xFFFFFFFF ? (TEXT("True")) : (Val == 0x0 ? TEXT("False") : TEXT("Invalid"));
 		}
 		else
@@ -331,26 +331,26 @@ FString FNiagaraTypeHelper::ToString(const uint8* ValueData, const UObject* Stru
 				const uint8* PropPtr = ValueData + PropertyIt->GetOffset_ForInternal();
 				if (Property->IsA(FFloatProperty::StaticClass()))
 				{
-					Ret += FString::Printf(TEXT("%s: %g "), *Property->GetNameCPP(), *(float*)PropPtr);
+					Ret += FString::Printf(TEXT("%s: %g "), *Property->GetNameCPP(), *reinterpret_cast<const float*>(PropPtr));
 				}
 				else if (Property->IsA(FDoubleProperty::StaticClass()))
 				{
-					Ret += FString::Printf(TEXT("%s: %g "), *Property->GetNameCPP(), *(double*)PropPtr);
+					Ret += FString::Printf(TEXT("%s: %g "), *Property->GetNameCPP(), *reinterpret_cast<const double*>(PropPtr));
 				}
 				else if (Property->IsA(FUInt16Property::StaticClass()))
 				{
-					FFloat16 Val = *(FFloat16*)PropPtr;
+					FFloat16 Val = *reinterpret_cast<const FFloat16*>(PropPtr);
 					Ret += FString::Printf(TEXT("%s: %f "), *Property->GetNameCPP(), Val.GetFloat());
 				}
 				else if (Property->IsA(FIntProperty::StaticClass()))
 				{
-					Ret += FString::Printf(TEXT("%s: %d "), *Property->GetNameCPP(), *(int32*)PropPtr);
+					Ret += FString::Printf(TEXT("%s: %d "), *Property->GetNameCPP(), *reinterpret_cast<const int32*>(PropPtr));
 				}
 				else if (Property->IsA(FBoolProperty::StaticClass()))
 				{
-					int32 Val = *(int32*)ValueData;
+					int32 Val = *reinterpret_cast<const int32*>(ValueData);
 					FString BoolStr = Val == 0xFFFFFFFF ? (TEXT("True")) : (Val == 0x0 ? TEXT("False") : TEXT("Invalid"));
-					Ret += FString::Printf(TEXT("%s: %d "), *Property->GetNameCPP(), *BoolStr);
+					Ret += FString::Printf(TEXT("%s: %s "), *Property->GetNameCPP(), *BoolStr);
 				}
 				else if (const FStructProperty* StructProp = CastFieldChecked<const FStructProperty>(Property))
 				{
diff --git a/Engine/Plugins/FX/Niagara/Source/Niagara/Public/NiagaraTypes.h b/Engine/Plugins/FX/Niagara/Source/Niagara/Public/NiagaraTypes.h
--- a/Engine/Plugins/FX/Niagara/Source/Niagara/Public/NiagaraTypes.h
+++ b/Engine/Plugins/FX/Niagara/Source/Niagara/Public/NiagaraTypes.h
@@ -292,8 +292,8 @@ struct FNiagaraStructConversionStep
 	{
 		for (int32 i = 0; i < Count; ++i)
 		{
-			DstType* Dst = (DstType*)((DestinationData + i * DestStride) + DestOffset);
-			const SrcType* Src = (const SrcType*)((SourceData + i * SourceStride) + SourceOffset);			
+			DstType* Dst = reinterpret_cast<DstType*>((DestinationData + i * DestStride) + DestOffset);
+			const SrcType* Src = reinterpret_cast<const SrcType*>((SourceData + i * SourceStride) + SourceOffset);			
 			*Dst = DstType(*Src);
 		}
 	}
@@ -1431,7 +1431,7 @@ public:
 
 	/** FGCObject interface */
 	NIAGARA_API virtual void AddReferencedObjects(FReferenceCollector& Collector) override;
-	NIAGARA_API virtual FString GetReferencerName() const;
+	NIAGARA_API virtual FString GetReferencerName() const override;
 
 private:
 	struct FQueuedRegistryEntry
@@ -1854,9 +1854,9 @@ struct FNiagaraVariableMatch
 	{
 	}
 
-	bool operator()(const FNiagaraVariable& Other) const
+	bool operator()(const FNiagaraVariableBase& Other) const
 	{
-		return static_cast<const FNiagaraVariableBase&>(Other) == VariableBase;
+		return Other == VariableBase;
 	}
 
 private:
