- Fix fx.NiagaraReleaseBuffersOnReset to allow releasing memory on reset
- Currently only GPU buffers will be released on reset, not CPU
- Releasing CPU memory reduces component pool usage
- In one test run we went from ~4mb -> 1.4mb of particle buffer memory by enabling the early release.

diff --git a/Engine/Plugins/FX/Niagara/Source/Niagara/Private/NiagaraDataSet.cpp b/Engine/Plugins/FX/Niagara/Source/Niagara/Private/NiagaraDataSet.cpp
--- a/Engine/Plugins/FX/Niagara/Source/Niagara/Private/NiagaraDataSet.cpp
+++ b/Engine/Plugins/FX/Niagara/Source/Niagara/Private/NiagaraDataSet.cpp
@@ -175,15 +175,29 @@ void FNiagaraDataSet::ResetBuffers()
 	
 	if (GetSimTarget() == ENiagaraSimTarget::CPUSim)
 	{
+		ResetBuffersInternal();
+
 		if (GNiagaraReleaseBuffersOnReset)
 		{
-			for (FNiagaraDataBuffer* Buffer : Data)
+			for (auto it=Data.CreateIterator(); it; ++it)
 			{
-				Buffer->ReleaseCPU();
+				FNiagaraDataBuffer* Buffer = Data[it.GetIndex()];
+				if (Buffer->IsBeingRead())
+				{
+				
+					// In this path the buffer might still be in use on the render thread for rendering the last frame therefore we push it into the deferred destruction queue as we can not release the data as we would race.
+#if NIAGARA_MEMORY_TRACKING
+					BufferSizeBytes -= Buffer->GetFloatBuffer().GetAllocatedSize() + Buffer->GetInt32Buffer().GetAllocatedSize() + Buffer->GetHalfBuffer().GetAllocatedSize();
+#endif
+					Buffer->Destroy();
+					it.RemoveCurrentSwap();
+				}
+				else
+				{
+					Buffer->ReleaseCPU();
+				}
 			}
 		}
-
-		ResetBuffersInternal();
 	}
 	else
 	{
@@ -579,7 +593,7 @@ FNiagaraDataBuffer::FNiagaraDataBuffer(FNiagaraDataSet* InOwner)
 FNiagaraDataBuffer::~FNiagaraDataBuffer()
 {
 	check(!IsInUse());
-	DEC_MEMORY_STAT_BY(STAT_NiagaraParticleMemory, FloatData.GetAllocatedSize() + Int32Data.GetAllocatedSize());
+	DEC_MEMORY_STAT_BY(STAT_NiagaraParticleMemory, FloatData.GetAllocatedSize() + Int32Data.GetAllocatedSize() + HalfData.GetAllocatedSize());
 #if NIAGARA_MEMORY_TRACKING
 	DEC_MEMORY_STAT_BY(STAT_NiagaraGPUParticleMemory, AllocationSizeBytes);
 #endif
