Added details panel support to data channel BP node, so users can exclude pins from being written to

diff --git a/Engine/Plugins/FX/Niagara/Source/NiagaraBlueprintNodes/Internal/K2Node_WriteDataChannel.h b/Engine/Plugins/FX/Niagara/Source/NiagaraBlueprintNodes/Internal/K2Node_WriteDataChannel.h
new file mode 100644
--- /dev/null
+++ b/Engine/Plugins/FX/Niagara/Source/NiagaraBlueprintNodes/Internal/K2Node_WriteDataChannel.h
@@ -0,0 +1,52 @@
+﻿// Copyright Epic Games, Inc. All Rights Reserved.
+
+#pragma once
+
+#include "CoreMinimal.h"
+#include "K2Node_CallFunction.h"
+#include "NiagaraDataChannel.h"
+#include "K2Node_WriteDataChannel.generated.h"
+
+
+UCLASS(MinimalAPI)
+class UK2Node_WriteDataChannel : public UK2Node_CallFunction
+{
+	GENERATED_BODY()
+
+public:
+	UK2Node_WriteDataChannel();
+
+	virtual void PostLoad() override;
+
+	// //~ Begin UEdGraphNode Interface.
+	virtual void AllocateDefaultPins() override;
+	virtual void PinDefaultValueChanged(UEdGraphPin* Pin) override;
+	virtual void PinConnectionListChanged(UEdGraphPin* Pin) override;
+	// //~ End UEdGraphNode Interface.
+
+	//~ Begin K2Node Interface
+	virtual void GetMenuActions(FBlueprintActionDatabaseRegistrar& InActionRegistrar) const override;
+	virtual FText GetMenuCategory() const override;
+	virtual void ExpandNode(class FKismetCompilerContext& CompilerContext, UEdGraph* SourceGraph) override;
+	virtual ERedirectType DoPinsMatchForReconstruction(const UEdGraphPin* NewPin, int32 NewPinIndex, const UEdGraphPin* OldPin, int32 OldPinIndex) const override;
+	virtual void PreloadRequiredAssets() override;
+	virtual bool ShouldShowNodeProperties() const override;
+	//~ End K2Node Interface
+
+	NIAGARABLUEPRINTNODES_API UNiagaraDataChannel* GetDataChannel() const;
+
+	UPROPERTY()
+	TSet<FGuid> IgnoredVariables;
+	
+private:
+	UEdGraphPin* GetChannelSelectorPin() const;
+	UFunction* GetWriteFunctionForType(const FNiagaraTypeDefinition& TypeDef);
+
+	UPROPERTY()
+	TObjectPtr<UNiagaraDataChannelAsset> DataChannel;
+
+#if WITH_EDITORONLY_DATA
+	UPROPERTY()
+	FGuid DataChannelVersion;
+#endif
+};
diff --git a/Engine/Plugins/FX/Niagara/Source/NiagaraBlueprintNodes/Internal/NiagaraBlueprintUtil.h b/Engine/Plugins/FX/Niagara/Source/NiagaraBlueprintNodes/Internal/NiagaraBlueprintUtil.h
new file mode 100644
--- /dev/null
+++ b/Engine/Plugins/FX/Niagara/Source/NiagaraBlueprintNodes/Internal/NiagaraBlueprintUtil.h
@@ -0,0 +1,15 @@
+﻿// Copyright Epic Games, Inc. All Rights Reserved.
+
+#pragma once
+
+#include "NiagaraTypes.h"
+#include "EdGraph/EdGraphPin.h"
+
+class FNiagaraBlueprintUtil
+{
+public:
+	/**
+	Attempts to convert the existing Niagara type to a matching type for BP pins
+	*/
+	static NIAGARABLUEPRINTNODES_API FEdGraphPinType TypeDefinitionToBlueprintType(const FNiagaraTypeDefinition& TypeDef);
+};
diff --git a/Engine/Plugins/FX/Niagara/Source/NiagaraBlueprintNodes/Private/K2Node_WriteDataChannel.cpp b/Engine/Plugins/FX/Niagara/Source/NiagaraBlueprintNodes/Private/K2Node_WriteDataChannel.cpp
--- a/Engine/Plugins/FX/Niagara/Source/NiagaraBlueprintNodes/Private/K2Node_WriteDataChannel.cpp
+++ b/Engine/Plugins/FX/Niagara/Source/NiagaraBlueprintNodes/Private/K2Node_WriteDataChannel.cpp
@@ -36,6 +36,10 @@ void UK2Node_WriteDataChannel::AllocateDefaultPins()
 	{
 		for (const FNiagaraDataChannelVariable& InVar : DataChannel->Get()->GetVariables())
 		{
+			if (IgnoredVariables.Contains(InVar.Version))
+			{
+				continue;
+			}
 			UEdGraphPin* NewPin = CreatePin(EGPD_Input, FNiagaraBlueprintUtil::TypeDefinitionToBlueprintType(InVar.GetType()), InVar.GetName());
 
 #if WITH_EDITORONLY_DATA
@@ -129,8 +133,12 @@ void UK2Node_WriteDataChannel::ExpandNode(FKismetCompilerContext& CompilerContex
 	// create the write function nodes
 	UEdGraphPin* LastExecPin = CreateWriterNode->FindPinChecked(UEdGraphSchema_K2::PN_Then, EGPD_Output);
 	UEdGraphPin* WriterResultPin = CreateWriterNode->FindPinChecked(UEdGraphSchema_K2::PN_ReturnValue, EGPD_Output);
-	for (const FNiagaraVariableBase& InVar : DataChannel->Get()->GetVariables())
+	for (const FNiagaraDataChannelVariable& InVar : DataChannel->Get()->GetVariables())
 	{
+		if (IgnoredVariables.Contains(InVar.Version))
+		{
+			continue;
+		}
 		UEdGraphPin* VarInputPin = FindPinChecked(InVar.GetName(), EGPD_Input);
 		if (VarInputPin == nullptr)
 		{
@@ -195,7 +203,12 @@ void UK2Node_WriteDataChannel::PreloadRequiredAssets()
 
 bool UK2Node_WriteDataChannel::ShouldShowNodeProperties() const
 {
-	return Super::ShouldShowNodeProperties();
+	return true;
+}
+
+UNiagaraDataChannel* UK2Node_WriteDataChannel::GetDataChannel() const
+{
+	return DataChannel ? DataChannel->Get() : nullptr;
 }
 
 UEdGraphPin* UK2Node_WriteDataChannel::GetChannelSelectorPin() const
@@ -245,6 +258,10 @@ UFunction* UK2Node_WriteDataChannel::GetWriteFunctionForType(const FNiagaraTypeD
 	{
 		return UNiagaraDataChannelWriter::StaticClass()->FindFunctionByName(GET_FUNCTION_NAME_CHECKED(UNiagaraDataChannelWriter, WriteSpawnInfo));
 	}
+	if (TypeDef.GetStruct() == FNiagaraID::StaticStruct())
+	{
+		return UNiagaraDataChannelWriter::StaticClass()->FindFunctionByName(GET_FUNCTION_NAME_CHECKED(UNiagaraDataChannelWriter, WriteID));
+	}
 	if (TypeDef.GetEnum())
 	{
 		return UNiagaraDataChannelWriter::StaticClass()->FindFunctionByName(GET_FUNCTION_NAME_CHECKED(UNiagaraDataChannelWriter, WriteEnum));
diff --git a/Engine/Plugins/FX/Niagara/Source/NiagaraBlueprintNodes/Private/K2Node_WriteDataChannel.h b/Engine/Plugins/FX/Niagara/Source/NiagaraBlueprintNodes/Private/K2Node_WriteDataChannel.h
deleted file mode 100644
diff --git a/Engine/Plugins/FX/Niagara/Source/NiagaraBlueprintNodes/Private/NiagaraBlueprintUtil.h b/Engine/Plugins/FX/Niagara/Source/NiagaraBlueprintNodes/Private/NiagaraBlueprintUtil.h
deleted file mode 100644
diff --git a/Engine/Plugins/FX/Niagara/Source/NiagaraEditorWidgets/NiagaraEditorWidgets.Build.cs b/Engine/Plugins/FX/Niagara/Source/NiagaraEditorWidgets/NiagaraEditorWidgets.Build.cs
--- a/Engine/Plugins/FX/Niagara/Source/NiagaraEditorWidgets/NiagaraEditorWidgets.Build.cs
+++ b/Engine/Plugins/FX/Niagara/Source/NiagaraEditorWidgets/NiagaraEditorWidgets.Build.cs
@@ -14,6 +14,7 @@ public class NiagaraEditorWidgets : ModuleRules
             "NiagaraCore",
 			"Niagara",
 			"NiagaraEditor",
+			"NiagaraBlueprintNodes",
 			"Engine",
 			"Core",
 			"CoreUObject",
diff --git a/Engine/Plugins/FX/Niagara/Source/NiagaraEditorWidgets/Private/DetailCustomizations/NiagaraDataChannelDetails.cpp b/Engine/Plugins/FX/Niagara/Source/NiagaraEditorWidgets/Private/DetailCustomizations/NiagaraDataChannelDetails.cpp
--- a/Engine/Plugins/FX/Niagara/Source/NiagaraEditorWidgets/Private/DetailCustomizations/NiagaraDataChannelDetails.cpp
+++ b/Engine/Plugins/FX/Niagara/Source/NiagaraEditorWidgets/Private/DetailCustomizations/NiagaraDataChannelDetails.cpp
@@ -2,30 +2,26 @@
 
 #include "NiagaraDataChannelDetails.h"
 
-#include "IDetailChildrenBuilder.h"
 #include "DetailLayoutBuilder.h"
 #include "PropertyCustomizationHelpers.h"
+#include "ScopedTransaction.h"
 
 #include "Types/SlateEnums.h"
-#include "Widgets/SItemSelector.h"
 #include "Widgets/SCompoundWidget.h"
 #include "Widgets/DeclarativeSyntaxSupport.h"
 #include "Widgets/Input/SComboBox.h"
 #include "Widgets/Input/SEditableTextBox.h"
-#include "Widgets/Input/SSearchBox.h"
 #include "Widgets/Input/SComboButton.h"
+#include "Widgets/Input/SCheckBox.h"
 #include "SListViewSelectorDropdownMenu.h"
 
-#include "SNiagaraNamePropertySelector.h"
-#include "NiagaraDataChannel.h"
 #include "DataInterface/NiagaraDataInterfaceDataChannelRead.h"
 #include "NiagaraDataChannel.h"
-#include "NiagaraDataChannel_Global.h"
 #include "NiagaraDataChannel_Islands.h"
 
 #include "NiagaraSystem.h"
 #include "NiagaraEditorUtilities.h"
-#include "NiagaraEditorWidgetsStyle.h"
+#include "K2Node_WriteDataChannel.h"
 
 #include "NiagaraDataInterfaceDetails.h"
 
@@ -52,8 +48,6 @@ private:
 
 	void OnSelectionChanged(TSharedPtr<FNiagaraVariableBase> InNewSelection, ESelectInfo::Type);
 
-	TSharedRef<SWidget> MakeSelectionWidget(TSharedPtr<FNiagaraVariableBase> InItem);
-
 	TSharedRef<ITableRow> GenerateAddElementRow(TSharedPtr<FNiagaraVariableBase> Entry, const TSharedRef<STableViewBase>& OwnerTable) const;
 
 
@@ -255,7 +249,7 @@ void FNiagaraDataInterfaceDataChannelReadDetails::CustomizeDetails(IDetailLayout
 	CategoryBuilder.GetDefaultProperties(Properties, true, true);
 	for (const TSharedRef<IPropertyHandle>& PropertyHandle : Properties)
 	{
-		FProperty* Property = PropertyHandle->GetProperty();
+		PropertyHandle->GetProperty();
 		CategoryBuilder.AddProperty(PropertyHandle);		
 	}
 }
@@ -265,5 +259,80 @@ TSharedRef<IDetailCustomization> FNiagaraDataInterfaceDataChannelReadDetails::Ma
 	return MakeShared<FNiagaraDataInterfaceDataChannelReadDetails>();
 }
 
+void FNiagaraDataChannelBPNodeDetails::CustomizeDetails(IDetailLayoutBuilder& DetailBuilder)
+{
+	TArray<TWeakObjectPtr<UObject>> SelectedObjects;
+	DetailBuilder.GetObjectsBeingCustomized(SelectedObjects);
+	if (SelectedObjects.Num() != 1)
+	{
+		return;
+	}
+
+	UObject* Obj = SelectedObjects[0].Get();
+	UK2Node_WriteDataChannel* Node = Cast<UK2Node_WriteDataChannel>(Obj);
+	if(Node == nullptr || Node->GetDataChannel() == nullptr)
+	{
+		return;
+	}
+
+	UNiagaraDataChannel* DataChannel = Node->GetDataChannel();
+	IDetailCategoryBuilder& ChannelCategoryBuilder = DetailBuilder.EditCategory(TEXT("Data Channel"));
+	for (const FNiagaraDataChannelVariable& Var : DataChannel->GetVariables())
+	{
+		TWeakObjectPtr<UK2Node_WriteDataChannel> NodePtr(Node);
+		FGuid VarGuid = Var.Version;
+		auto CheckStateChanged = [NodePtr, VarGuid](const ECheckBoxState NewState)
+		{
+			if (UK2Node_WriteDataChannel* Node = NodePtr.Get())
+			{
+				FScopedTransaction Transaction(LOCTEXT("ChangeAttributeWrite", "Change data channel attribute write"));
+				Node->Modify();
+				if (NewState == ECheckBoxState::Checked)
+				{
+					Node->IgnoredVariables.Remove(VarGuid);
+					Node->ReconstructNode();
+				}
+				else if (NewState == ECheckBoxState::Unchecked)
+				{
+					Node->IgnoredVariables.Add(VarGuid);
+					Node->ReconstructNode();
+				}
+			}
+		};
+		auto SPCheckState = [NodePtr, VarGuid]()
+		{
+			if (UK2Node_WriteDataChannel* Node = NodePtr.Get())
+			{
+				return Node->IgnoredVariables.Contains(VarGuid) ? ECheckBoxState::Unchecked : ECheckBoxState::Checked;
+			}
+			return ECheckBoxState::Undetermined;
+		};
+		
+		FDetailWidgetRow& Row = ChannelCategoryBuilder.AddCustomRow(FText::FromName(Var.GetName()));
+		Row.NameWidget
+		[
+			SNew(STextBlock)
+			.Text(FText::FromName(Var.GetName()))
+			.Font(IDetailLayoutBuilder::GetDetailFont())
+		];
+		Row.ValueWidget
+		[
+			SNew(SCheckBox)
+			.OnCheckStateChanged_Lambda(CheckStateChanged)
+			.IsChecked_Lambda(SPCheckState)
+			[
+				SNew(STextBlock)
+				.Text(LOCTEXT("AsPinText", "(As pin)"))
+				.Font(IDetailLayoutBuilder::GetDetailFont())
+			]
+		];
+	}
+}
+
+TSharedRef<IDetailCustomization> FNiagaraDataChannelBPNodeDetails::MakeInstance()
+{
+	return MakeShared<FNiagaraDataChannelBPNodeDetails>();
+}
+
 #undef LOCTEXT_NAMESPACE
 
diff --git a/Engine/Plugins/FX/Niagara/Source/NiagaraEditorWidgets/Private/DetailCustomizations/NiagaraDataChannelDetails.h b/Engine/Plugins/FX/Niagara/Source/NiagaraEditorWidgets/Private/DetailCustomizations/NiagaraDataChannelDetails.h
--- a/Engine/Plugins/FX/Niagara/Source/NiagaraEditorWidgets/Private/DetailCustomizations/NiagaraDataChannelDetails.h
+++ b/Engine/Plugins/FX/Niagara/Source/NiagaraEditorWidgets/Private/DetailCustomizations/NiagaraDataChannelDetails.h
@@ -37,3 +37,12 @@ private:
 	TWeakObjectPtr<UNiagaraDataInterfaceDataChannelRead> ReadDataInterfaceWeak;
 	TWeakObjectPtr<UNiagaraSystem> NiagaraSystemWeak;
 };
+
+
+/** Details customization for Blueprint nodes. */
+class FNiagaraDataChannelBPNodeDetails : public IDetailCustomization
+{
+public:
+	virtual void CustomizeDetails(IDetailLayoutBuilder& DetailBuilder) override;
+	static TSharedRef<IDetailCustomization> MakeInstance();
+};
\ No newline at end of file
diff --git a/Engine/Plugins/FX/Niagara/Source/NiagaraEditorWidgets/Private/NiagaraEditorWidgetsModule.cpp b/Engine/Plugins/FX/Niagara/Source/NiagaraEditorWidgets/Private/NiagaraEditorWidgetsModule.cpp
--- a/Engine/Plugins/FX/Niagara/Source/NiagaraEditorWidgets/Private/NiagaraEditorWidgetsModule.cpp
+++ b/Engine/Plugins/FX/Niagara/Source/NiagaraEditorWidgets/Private/NiagaraEditorWidgetsModule.cpp
@@ -144,6 +144,7 @@ void FNiagaraEditorWidgetsModule::StartupModule()
 	PropertyModule.RegisterCustomClassLayout("NiagaraDecalRendererProperties", FOnGetDetailCustomizationInstance::CreateStatic(&FNiagaraDecalRendererDetails::MakeInstance));
 
 	PropertyModule.RegisterCustomClassLayout("NiagaraDataInterfaceDataChannelRead", FOnGetDetailCustomizationInstance::CreateStatic(&FNiagaraDataInterfaceDataChannelReadDetails::MakeInstance));
+	PropertyModule.RegisterCustomClassLayout("K2Node_WriteDataChannel", FOnGetDetailCustomizationInstance::CreateStatic(&FNiagaraDataChannelBPNodeDetails::MakeInstance));
 
 	PropertyModule.RegisterCustomClassLayout(
 		UNiagaraDataChannel_Islands::StaticClass()->GetFName(),
