Remove early out from FNiagaraSystemCompilingManager::ProcessAsyncTasks that was incorrectly preventing tasks from getting cleaned up if all that was left was tasks awaiting retrieval

diff --git a/Engine/Plugins/FX/Niagara/Source/NiagaraEditor/Private/NiagaraSystemCompilingManager.cpp b/Engine/Plugins/FX/Niagara/Source/NiagaraEditor/Private/NiagaraSystemCompilingManager.cpp
--- a/Engine/Plugins/FX/Niagara/Source/NiagaraEditor/Private/NiagaraSystemCompilingManager.cpp
+++ b/Engine/Plugins/FX/Niagara/Source/NiagaraEditor/Private/NiagaraSystemCompilingManager.cpp
@@ -83,11 +83,6 @@ void FNiagaraSystemCompilingManager::ProcessAsyncTasks(bool bLimitExecutionTime)
 
 	{
 		FReadScopeLock ReadScope(QueueLock);
-		if (ActiveTasks.IsEmpty())
-		{
-			return;
-		}
-
 		for (FNiagaraCompilationTaskHandle TaskHandle : ActiveTasks)
 		{
 			FTaskPtr TaskPtr = SystemRequestMap.FindRef(TaskHandle);
