- Fix reset to defaults on emitter binding no reflecting in the UI

diff --git a/Engine/Plugins/FX/Niagara/Source/NiagaraEditor/Private/Customizations/NiagaraDataInterfaceEmitterBindingCustomization.cpp b/Engine/Plugins/FX/Niagara/Source/NiagaraEditor/Private/Customizations/NiagaraDataInterfaceEmitterBindingCustomization.cpp
--- a/Engine/Plugins/FX/Niagara/Source/NiagaraEditor/Private/Customizations/NiagaraDataInterfaceEmitterBindingCustomization.cpp
+++ b/Engine/Plugins/FX/Niagara/Source/NiagaraEditor/Private/Customizations/NiagaraDataInterfaceEmitterBindingCustomization.cpp
@@ -9,12 +9,14 @@
 #include "DetailWidgetRow.h"
 #include "IDetailChildrenBuilder.h"
 #include "IPropertyUtilities.h"
+#include "ScopedTransaction.h"
 #include "Modules/ModuleManager.h"
 #include "Widgets/Input/SCheckBox.h"
 #include "Widgets/Input/SEditableTextBox.h"
 #include "Widgets/SBoxPanel.h"
 #include "Widgets/SNiagaraDebugger.h"
 #include "Widgets/Views/STreeView.h"
+#include "Widgets/Input/SButton.h"
 #include "Widgets/Input/SSearchBox.h"
 #include "Widgets/Input/SSuggestionTextBox.h"
 #include "Widgets/Layout/SScrollBorder.h"
@@ -22,7 +24,7 @@
 
 #define LOCTEXT_NAMESPACE "NiagaraDataInterfaceEmitterBindingCustomization"
 
-void FNiagaraDataInterfaceEmitterBindingCustomization::CustomizeHeader(TSharedRef<IPropertyHandle> StructPropertyHandle, FDetailWidgetRow& HeaderRow, IPropertyTypeCustomizationUtils& CustomizationUtils)
+void FNiagaraDataInterfaceEmitterBindingCustomization::CustomizeHeader(TSharedRef<IPropertyHandle> InStructPropertyHandle, FDetailWidgetRow& HeaderRow, IPropertyTypeCustomizationUtils& CustomizationUtils)
 {
 	// Get owner system (if any)
 	if ( TSharedPtr<IPropertyUtilities> PropertyUtilities = CustomizationUtils.GetPropertyUtilities() )
@@ -32,16 +34,20 @@ void FNiagaraDataInterfaceEmitterBindingCustomization::CustomizeHeader(TSharedRe
 		{
 			if ( UObject* OwnerObject = SelectedObjects[0].Get() )
 			{
+				OwnerObjectWeak = OwnerObject;
 				NiagaraSystemWeak = OwnerObject->GetTypedOuter<UNiagaraSystem>();
 			}
 		}
 	}
 
 	// Get properties
+	StructPropertyHandle = InStructPropertyHandle;
 	BindingModeHandle = StructPropertyHandle->GetChildHandle(GET_MEMBER_NAME_CHECKED(FNiagaraDataInterfaceEmitterBinding, BindingMode));
 	EmitterNameHandle = StructPropertyHandle->GetChildHandle(GET_MEMBER_NAME_CHECKED(FNiagaraDataInterfaceEmitterBinding, EmitterName));
 	check(BindingModeHandle.IsValid() && EmitterNameHandle.IsValid())
 
+	StructPropertyHandle->MarkResetToDefaultCustomized(true);
+
 	// Build Widget
 	HeaderRow
 		.NameContent()
@@ -59,7 +65,7 @@ void FNiagaraDataInterfaceEmitterBindingCustomization::CustomizeHeader(TSharedRe
 			]
 			+ SHorizontalBox::Slot()
 			[
-				SNew(SSuggestionTextBox)
+				SAssignNew(EmitterTextBox, SSuggestionTextBox)
 				.ForegroundColor(FSlateColor::UseForeground())
 				.Visibility(this, &FNiagaraDataInterfaceEmitterBindingCustomization::IsEmitterNameVisible)
 				.Font(IDetailLayoutBuilder::GetDetailFont())
@@ -68,6 +74,25 @@ void FNiagaraDataInterfaceEmitterBindingCustomization::CustomizeHeader(TSharedRe
 				.HintText(LOCTEXT("EmitterNameHint", "Enter Emitter Name..."))
 				.OnShowingSuggestions(this, &FNiagaraDataInterfaceEmitterBindingCustomization::EmitterName_GetSuggestions)
 			]
+			+ SHorizontalBox::Slot()
+			.AutoWidth()
+			.VAlign(VAlign_Center)
+			.Padding(2, 1)
+			[
+				SNew(SButton)
+				.IsFocusable(false)
+				.ToolTipText(LOCTEXT("ResetToDefaultToolTip", "Reset to Default"))
+				.ButtonStyle(FAppStyle::Get(), "NoBorder")
+				.ContentPadding(0)
+				.Visibility(this, &FNiagaraDataInterfaceEmitterBindingCustomization::IsResetToDefaultsVisible)
+				.OnClicked(this, &FNiagaraDataInterfaceEmitterBindingCustomization::OnResetToDefaultsClicked)
+				.Content()
+				[
+					SNew(SImage)
+					.Image(FAppStyle::GetBrush("PropertyWindow.DiffersFromDefault"))
+					.ColorAndOpacity(FSlateColor::UseForeground())
+				]
+			]
 		];
 }
 
@@ -105,4 +130,35 @@ void FNiagaraDataInterfaceEmitterBindingCustomization::EmitterName_GetSuggestion
 	}
 }
 
+FNiagaraDataInterfaceEmitterBinding* FNiagaraDataInterfaceEmitterBindingCustomization::GetEmitterBinding() const
+{
+	UObject* OwnerObject = OwnerObjectWeak.Get();
+	return OwnerObject ? reinterpret_cast<FNiagaraDataInterfaceEmitterBinding*>(StructPropertyHandle->GetValueBaseAddress(reinterpret_cast<uint8*>(OwnerObject))) : nullptr;
+}
+
+EVisibility FNiagaraDataInterfaceEmitterBindingCustomization::IsResetToDefaultsVisible() const
+{
+	if (FNiagaraDataInterfaceEmitterBinding* EmitterBinding = GetEmitterBinding())
+	{
+		FNiagaraDataInterfaceEmitterBinding DefaultBinding;
+		return *EmitterBinding == DefaultBinding ? EVisibility::Hidden : EVisibility::Visible;
+	}
+	return EVisibility::Hidden;
+}
+
+FReply FNiagaraDataInterfaceEmitterBindingCustomization::OnResetToDefaultsClicked()
+{
+	if (FNiagaraDataInterfaceEmitterBinding* EmitterBinding = GetEmitterBinding())
+	{
+		FScopedTransaction Transaction(LOCTEXT("ResetEmitterBindingToDefault", "Reset Emitter Binding"));
+		OwnerObjectWeak->Modify();
+		StructPropertyHandle->NotifyPreChange();
+		*EmitterBinding = FNiagaraDataInterfaceEmitterBinding();
+		StructPropertyHandle->NotifyPostChange(EPropertyChangeType::ValueSet);
+		StructPropertyHandle->NotifyFinishedChangingProperties();
+		EmitterTextBox->SetText(FText::FromName(EmitterBinding->EmitterName));
+	}
+	return FReply::Handled();
+}
+
 #undef LOCTEXT_NAMESPACE
diff --git a/Engine/Plugins/FX/Niagara/Source/NiagaraEditor/Private/Customizations/NiagaraDataInterfaceEmitterBindingCustomization.h b/Engine/Plugins/FX/Niagara/Source/NiagaraEditor/Private/Customizations/NiagaraDataInterfaceEmitterBindingCustomization.h
--- a/Engine/Plugins/FX/Niagara/Source/NiagaraEditor/Private/Customizations/NiagaraDataInterfaceEmitterBindingCustomization.h
+++ b/Engine/Plugins/FX/Niagara/Source/NiagaraEditor/Private/Customizations/NiagaraDataInterfaceEmitterBindingCustomization.h
@@ -4,6 +4,7 @@
 
 #include "CoreMinimal.h"
 #include "IPropertyTypeCustomization.h"
+#include "Input/Reply.h"
 #include "Layout/Visibility.h"
 #include "NiagaraDataInterfaceEmitterBinding.h"
 
@@ -11,6 +12,8 @@ class FDetailWidgetRow;
 class IPropertyHandle;
 class UNiagaraSystem;
 
+class SSuggestionTextBox;
+
 class FNiagaraDataInterfaceEmitterBindingCustomization : public IPropertyTypeCustomization
 {
 public:
@@ -30,8 +33,16 @@ public:
 	void EmitterName_SetText(const FText& NewText, ETextCommit::Type CommitInfo);
 	void EmitterName_GetSuggestions(const FString& CurrText, TArray<FString>& OutSuggestions);
 
+	FNiagaraDataInterfaceEmitterBinding* GetEmitterBinding() const;
+	EVisibility IsResetToDefaultsVisible() const;
+	FReply OnResetToDefaultsClicked();
+
+	TWeakObjectPtr<UObject> OwnerObjectWeak;
 	TWeakObjectPtr<UNiagaraSystem> NiagaraSystemWeak;
+	TSharedPtr<IPropertyHandle> StructPropertyHandle;
 	TSharedPtr<IPropertyHandle> BindingModeHandle;
 	TSharedPtr<IPropertyHandle> EmitterNameHandle;
+
+	TSharedPtr<SSuggestionTextBox> EmitterTextBox;
 };
 
