Set display name of Niagara double type to "float", as that's what is used internally for simulations and is also similar to the way Blueprint handles double types (which still display as float in the ui)

diff --git a/Engine/Plugins/FX/Niagara/Source/Niagara/Public/NiagaraTypes.h b/Engine/Plugins/FX/Niagara/Source/Niagara/Public/NiagaraTypes.h
--- a/Engine/Plugins/FX/Niagara/Source/Niagara/Public/NiagaraTypes.h
+++ b/Engine/Plugins/FX/Niagara/Source/Niagara/Public/NiagaraTypes.h
@@ -751,6 +751,10 @@ public:
 			{
 				NameText = FText::FromString("Quat");
 			}
+			else if (*this == FNiagaraTypeHelper::GetDoubleDef())
+			{
+				NameText = FText::FromString("float");
+			}
 			else
 			{
 				NameText = GetStruct()->GetDisplayNameText();
