- Track approx. component memory

diff --git a/Engine/Plugins/FX/Niagara/Source/Niagara/Private/NiagaraComponent.cpp b/Engine/Plugins/FX/Niagara/Source/Niagara/Private/NiagaraComponent.cpp
--- a/Engine/Plugins/FX/Niagara/Source/Niagara/Private/NiagaraComponent.cpp
+++ b/Engine/Plugins/FX/Niagara/Source/Niagara/Private/NiagaraComponent.cpp
@@ -761,7 +761,8 @@ void UNiagaraComponent::ReleaseToPool()
 
 uint32 UNiagaraComponent::GetApproxMemoryUsage() const
 {
-	uint32 MemoryBytes = 0;
+	uint32 MemoryBytes = sizeof(UNiagaraComponent);
+	MemoryBytes += OverrideParameters.GetResourceSize();
 	if (SystemInstanceController)
 	{
 		MemoryBytes += uint32(SystemInstanceController->GetTotalBytesUsed());
