- Enable quick release of particle memory

diff --git a/Engine/Plugins/FX/Niagara/Source/Niagara/Private/NiagaraDataSet.cpp b/Engine/Plugins/FX/Niagara/Source/Niagara/Private/NiagaraDataSet.cpp
--- a/Engine/Plugins/FX/Niagara/Source/Niagara/Private/NiagaraDataSet.cpp
+++ b/Engine/Plugins/FX/Niagara/Source/Niagara/Private/NiagaraDataSet.cpp
@@ -91,7 +91,7 @@ static FAutoConsoleVariableRef CVarNiagaraGPUDataWarningSize(
 	ECVF_Scalability
 );
 
-static bool GNiagaraReleaseBuffersOnReset = false;
+static bool GNiagaraReleaseBuffersOnReset = true;
 static FAutoConsoleVariableRef CVarNiagaraReleaseBuffersOnReset(
 	TEXT("fx.NiagaraReleaseBuffersOnReset"),
 	GNiagaraReleaseBuffersOnReset,
