Make sure that we only return NPC when we have an initialized world

diff --git a/Engine/Plugins/FX/Niagara/Source/Niagara/Private/NiagaraWorldManager.cpp b/Engine/Plugins/FX/Niagara/Source/Niagara/Private/NiagaraWorldManager.cpp
--- a/Engine/Plugins/FX/Niagara/Source/Niagara/Private/NiagaraWorldManager.cpp
+++ b/Engine/Plugins/FX/Niagara/Source/Niagara/Private/NiagaraWorldManager.cpp
@@ -342,6 +342,7 @@ void FNiagaraWorldManager::Init(UWorld* InWorld)
 	ActiveNiagaraTickGroup = -1;
 	bAppHasFocus = true;
 	bIsTearingDown = false;
+	bIsCleaningUp = false;
 	WorldLoopTime = 0.0f;
 	RequestedDebugPlaybackMode = ENiagaraDebugPlaybackMode::Play;
 	DebugPlaybackMode = ENiagaraDebugPlaybackMode::Play;
@@ -545,7 +546,7 @@ void FNiagaraWorldManager::TickParameterCollections()
 
 UNiagaraParameterCollectionInstance* FNiagaraWorldManager::GetParameterCollection(UNiagaraParameterCollection* Collection)
 {
-	if (!Collection || bIsTearingDown)
+	if (!Collection || !HasActiveWorld())
 	{
 		return nullptr;
 	}
@@ -680,6 +681,8 @@ void FNiagaraWorldManager::OnWorldBeginTearDown()
 
 void FNiagaraWorldManager::OnWorldCleanup(bool bSessionEnded, bool bCleanupResources)
 {
+	bIsCleaningUp = true;
+
 	DeferredMethods.ExecuteAndClear();
 	ComponentPool->Cleanup(World);
 
@@ -2020,6 +2023,11 @@ void FNiagaraWorldManager::InvalidateCachedSystemScalabilityData()
 	}
 }
 
+bool FNiagaraWorldManager::HasActiveWorld() const
+{
+	return World && World->IsInitialized() && !World->bIsTearingDown && !bIsTearingDown && !bIsCleaningUp;
+}
+
 void FNiagaraWorldManager::PrimePoolForAllWorlds(UNiagaraSystem* System)
 {
 	if (GNigaraAllowPrimedPools)
@@ -2036,7 +2044,7 @@ void FNiagaraWorldManager::PrimePoolForAllWorlds(UNiagaraSystem* System)
 
 void FNiagaraWorldManager::PrimePoolForAllSystems()
 {
-	if (GNigaraAllowPrimedPools && World && World->IsGameWorld() && !World->bIsTearingDown)
+	if (GNigaraAllowPrimedPools && HasActiveWorld() && World->IsGameWorld())
 	{
 		//Prime the pool for all currently loaded systems.
 		for (TObjectIterator<UNiagaraSystem> It; It; ++It)
@@ -2052,7 +2060,7 @@ void FNiagaraWorldManager::PrimePoolForAllSystems()
 
 void FNiagaraWorldManager::PrimePool(UNiagaraSystem* System)
 {
-	if (GNigaraAllowPrimedPools && World && World->IsGameWorld() && !World->bIsTearingDown)
+	if (GNigaraAllowPrimedPools && HasActiveWorld() && World->IsGameWorld())
 	{
 		ComponentPool->PrimePool(System, World);
 	}
diff --git a/Engine/Plugins/FX/Niagara/Source/Niagara/Public/NiagaraWorldManager.h b/Engine/Plugins/FX/Niagara/Source/Niagara/Public/NiagaraWorldManager.h
--- a/Engine/Plugins/FX/Niagara/Source/Niagara/Public/NiagaraWorldManager.h
+++ b/Engine/Plugins/FX/Niagara/Source/Niagara/Public/NiagaraWorldManager.h
@@ -214,6 +214,8 @@ private:
 	static void InvalidateCachedSystemScalabilityDataForAllWorlds();
 	void InvalidateCachedSystemScalabilityData();
 
+	bool HasActiveWorld() const;
+
 public:
 	void SetDebugPlaybackMode(ENiagaraDebugPlaybackMode Mode) { RequestedDebugPlaybackMode = Mode; }
 	ENiagaraDebugPlaybackMode GetDebugPlaybackMode() const { return DebugPlaybackMode; }
@@ -390,6 +392,9 @@ private:
 
 	bool bIsTearingDown = false;
 
+	/** True if OnWorldCleanup has been called, will get reset with Init() */
+	bool bIsCleaningUp = false;
+
 	float WorldLoopTime = 0.0f;
 	
 	ENiagaraDebugPlaybackMode RequestedDebugPlaybackMode = ENiagaraDebugPlaybackMode::Play;
