Add option to allow some early niagara tick functions to be normal priority - since they have the entire frame to finish their async work. This can minimise GameThread waits if you have other high priority tasks that need to kick off async work earlier.
Disabled by default, use cvar "fx.Niagara.Experimental.TickGroupToStartHighPriorityTickFunctions -1" to completely disable.

diff --git a/Engine/Plugins/FX/Niagara/Source/Niagara/Private/NiagaraWorldManager.cpp b/Engine/Plugins/FX/Niagara/Source/Niagara/Private/NiagaraWorldManager.cpp
--- a/Engine/Plugins/FX/Niagara/Source/Niagara/Private/NiagaraWorldManager.cpp
+++ b/Engine/Plugins/FX/Niagara/Source/Niagara/Private/NiagaraWorldManager.cpp
@@ -64,6 +64,38 @@ static FAutoConsoleVariableRef CVarNiagaraAllowAsyncWorkToEndOfFrame(
 	ECVF_Default
 ); 
 
+static void NiagaraApplyTickGroupPrioritiesBasedOnCVar(IConsoleVariable*);
+
+//EXPERIMENTAL
+//When fx.Niagara.AllowAsyncWorkToEndOfFrame is enabled, earlier tick functions don't necessarily need high priority since they plenty of time for their async work to finish.
+//Disabling high priority on these tick functions can minimise GameThread waits by allowing other high priority tick functions to start async work as soon as possible.
+static int32 GNiagaraTickGroupToStartHighPriorityTickFunctions = -1;
+static FAutoConsoleVariableRef CVarNiagaraTickGroupToStartHighPriorityTickFunctions_Experimental(
+	TEXT("fx.Niagara.Experimental.TickGroupToStartHighPriorityTickFunctions"),
+	GNiagaraTickGroupToStartHighPriorityTickFunctions,
+	TEXT("All tick functions below this value will be set to normal priority, and any equal or above will be set to high priority. Set to -1 to have all high priority."),
+	FConsoleVariableDelegate::CreateStatic(&NiagaraApplyTickGroupPrioritiesBasedOnCVar),
+	ECVF_Default
+); 
+
+static void NiagaraApplyTickGroupPrioritiesBasedOnCVar(IConsoleVariable*)
+{
+	check (GEngine);
+	for (const FWorldContext& WorldContext : GEngine->GetWorldContexts())
+	{
+		if (const UWorld* World = WorldContext.World())
+		{
+			if (FNiagaraWorldManager* NiagaraWorldManager = FNiagaraWorldManager::Get(World))
+			{
+				for (int32 NiagaraTickGroupIndex = 0; NiagaraTickGroupIndex < NiagaraNumTickGroups; ++NiagaraTickGroupIndex)
+				{
+					NiagaraWorldManager->SetTickGroupPriority(NiagaraTickGroupIndex, NiagaraTickGroupIndex >= GNiagaraTickGroupToStartHighPriorityTickFunctions);
+				}
+			}
+		}
+	}
+}
+
 static int GNiagaraWaitOnPreGC = 1;
 static FAutoConsoleVariableRef CVarNiagaraWaitOnPreGC(
 	TEXT("fx.Niagara.WaitOnPreGC"),
@@ -357,7 +389,7 @@ void FNiagaraWorldManager::Init(UWorld* InWorld)
 		TickFunc.bCanEverTick = true;
 		TickFunc.bStartWithTickEnabled = true;
 		TickFunc.bAllowTickOnDedicatedServer = false;
-		TickFunc.bHighPriority = true;
+		TickFunc.bHighPriority = TickGroup >= GNiagaraTickGroupToStartHighPriorityTickFunctions;
 		TickFunc.Owner = this;
 		TickFunc.RegisterTickFunction(InWorld->PersistentLevel);
 	}
@@ -2142,6 +2174,12 @@ bool FNiagaraWorldManager::IsComponentLocalPlayerLinked(const USceneComponent* I
 	return false;
 }
 
+void FNiagaraWorldManager::SetTickGroupPriority(const int32 NiagaraTickGroupIndex, const bool bHighPriority)
+{
+	check ((NiagaraTickGroupIndex >= 0) && (NiagaraTickGroupIndex < NiagaraNumTickGroups));
+	TickFunctions[NiagaraTickGroupIndex].SetPriorityIncludingPrerequisites(bHighPriority);
+}
+
 #if DEBUG_SCALABILITY_STATE
 
 void FNiagaraWorldManager::DumpScalabilityState()
diff --git a/Engine/Plugins/FX/Niagara/Source/Niagara/Public/NiagaraWorldManager.h b/Engine/Plugins/FX/Niagara/Source/Niagara/Public/NiagaraWorldManager.h
--- a/Engine/Plugins/FX/Niagara/Source/Niagara/Public/NiagaraWorldManager.h
+++ b/Engine/Plugins/FX/Niagara/Source/Niagara/Public/NiagaraWorldManager.h
@@ -210,6 +210,9 @@ public:
 	void PrimePool(UNiagaraSystem* System);
 
 	static void RequestInvalidateCachedSystemScalabilityDataForAllWorlds();
+
+	void SetTickGroupPriority(const int32 NiagaraTickGroupIndex, const bool bHighPriority);
+
 private:
 	static void InvalidateCachedSystemScalabilityDataForAllWorlds();
 	void InvalidateCachedSystemScalabilityData();
