Niagara - Speculative fix to prevent load failure warnings when validating niagara content.

diff --git a/Engine/Plugins/FX/Niagara/Source/NiagaraEditor/Private/ViewModels/Stack/NiagaraStackFunctionInput.cpp b/Engine/Plugins/FX/Niagara/Source/NiagaraEditor/Private/ViewModels/Stack/NiagaraStackFunctionInput.cpp
--- a/Engine/Plugins/FX/Niagara/Source/NiagaraEditor/Private/ViewModels/Stack/NiagaraStackFunctionInput.cpp
+++ b/Engine/Plugins/FX/Niagara/Source/NiagaraEditor/Private/ViewModels/Stack/NiagaraStackFunctionInput.cpp
@@ -1736,7 +1736,7 @@ void UNiagaraStackFunctionInput::GetAvailableParameterHandles(TArray<FNiagaraPar
 
 	for (FAssetData& CollectionAsset : CollectionAssets)
 	{
-		if ( UNiagaraParameterCollection* Collection = Cast<UNiagaraParameterCollection>(CollectionAsset.GetAsset()) )
+		if ( UNiagaraParameterCollection* Collection = Cast<UNiagaraParameterCollection>(CollectionAsset.GetSoftObjectPath().TryLoad()) )
 		{
 			for (const FNiagaraVariable& CollectionParam : Collection->GetParameters())
 			{
