* Added Blueprint read node for data channels
* Fixed some data channel bugs

diff --git a/Engine/Plugins/FX/Niagara/Source/Niagara/Private/NiagaraDataChannel.cpp b/Engine/Plugins/FX/Niagara/Source/Niagara/Private/NiagaraDataChannel.cpp
--- a/Engine/Plugins/FX/Niagara/Source/Niagara/Private/NiagaraDataChannel.cpp
+++ b/Engine/Plugins/FX/Niagara/Source/Niagara/Private/NiagaraDataChannel.cpp
@@ -32,7 +32,7 @@ void FNiagaraDataChannelGameDataLayout::Init(const TArray<FNiagaraDataChannelVar
 	for (const FNiagaraDataChannelVariable& Var : Variables)
 	{
 		//Sigh.
-		//We must convert from the variable stored var in the data channels definition as we currently cannot serialize/store actual LWC types in FNiagawraTypeDefinitions.
+		//We must convert from the variable stored var in the data channels definition as we currently cannot serialize/store actual LWC types in FNiagaraTypeDefinitions.
 		FNiagaraTypeDefinition LWCType = FNiagaraTypeHelper::GetLWCType(Var.GetType());
 		FNiagaraVariableBase LWCVar(LWCType, Var.GetName());
 
@@ -185,13 +185,13 @@ void FNiagaraDataChannelGameData::WriteToDataSet(FNiagaraDataBuffer* DestBuffer,
 			//Positions are a special case that are stored as FVectors in game data but converted to an LWCTile local FVector3f in simulation data.
 			if (DestStruct == FNiagaraTypeDefinition::GetPositionStruct())
 			{	
-				float* DestX = (float*)DestBuffer->GetInstancePtrFloat(FloatCompIdx++, DestStartIdx);
-				float* DestY = (float*)DestBuffer->GetInstancePtrFloat(FloatCompIdx++, DestStartIdx);
-				float* DestZ = (float*)DestBuffer->GetInstancePtrFloat(FloatCompIdx++, DestStartIdx);
+				float* DestX = DestBuffer->GetInstancePtrFloat(FloatCompIdx++, DestStartIdx);
+				float* DestY = DestBuffer->GetInstancePtrFloat(FloatCompIdx++, DestStartIdx);
+				float* DestZ = DestBuffer->GetInstancePtrFloat(FloatCompIdx++, DestStartIdx);
 								
 				for (int32 i = 0; i < NumInstances; ++i)
 				{
-					FVector* Src = (FVector*)((SrcData + i * SrcVarSize));
+					FVector* Src = reinterpret_cast<FVector*>((SrcData + i * SrcVarSize));
 
 					FVector3f SimLocalSWC = FVector3f((*Src) - FVector(SimulationLwcTile) * FLargeWorldRenderScalar::GetTileSize());
 					*DestX++ = SimLocalSWC.X;
@@ -207,51 +207,50 @@ void FNiagaraDataChannelGameData::WriteToDataSet(FNiagaraDataBuffer* DestBuffer,
 				{
 					FProperty* SrcProperty = *SrcPropertyIt;
 					FProperty* DestProperty = *DestPropertyIt;
-					int32 SrcOffset = SrcProperty->GetOffset_ForInternal();
 					SrcData = SrcPropertyBase + SrcProperty->GetOffset_ForInternal();
 
 					//Convert any LWC doubles to floats. //TODO: Insert LWCTile... probably need to explicitly check for vectors etc.
 					if (SrcProperty->IsA(FDoubleProperty::StaticClass()))
 					{
 						check(DestProperty->IsA(FFloatProperty::StaticClass()));
-						float* Dest = (float*)DestBuffer->GetInstancePtrFloat(FloatCompIdx++, DestStartIdx);
+						float* Dest = DestBuffer->GetInstancePtrFloat(FloatCompIdx++, DestStartIdx);
 						//Write all instances for this component.
 						for (int32 i = 0; i < NumInstances; ++i)
 						{
-							double* Src = (double*)((SrcData + i * SrcVarSize));
+							double* Src = reinterpret_cast<double*>((SrcData + i * SrcVarSize));
 							*Dest++ = *Src;
 						}
 					}
 					else if (SrcProperty->IsA(FFloatProperty::StaticClass()))
 					{
-						float* Dest = (float*)DestBuffer->GetInstancePtrFloat(FloatCompIdx++, DestStartIdx);
+						float* Dest = DestBuffer->GetInstancePtrFloat(FloatCompIdx++, DestStartIdx);
 
 						//Write all instances for this component.
 						for (int32 i = 0; i < NumInstances; ++i)
 						{
-							float* Src = (float*)((SrcData + i * SrcVarSize));
+							float* Src = reinterpret_cast<float*>((SrcData + i * SrcVarSize));
 							*Dest++ = *Src;
 						}
 					}
 					else if (SrcProperty->IsA(FUInt16Property::StaticClass()))
 					{
-						FFloat16* Dest = (FFloat16*)DestBuffer->GetInstancePtrHalf(HalfCompIdx++, DestStartIdx);
+						FFloat16* Dest = DestBuffer->GetInstancePtrHalf(HalfCompIdx++, DestStartIdx);
 
 						//Write all instances for this component.
 						for (int32 i = 0; i < NumInstances; ++i)
 						{
-							FFloat16* Src = (FFloat16*)((SrcData + i * SrcVarSize));
+							FFloat16* Src = reinterpret_cast<FFloat16*>((SrcData + i * SrcVarSize));
 							*Dest++ = *Src;
 						}
 					}
 					else if (SrcProperty->IsA(FIntProperty::StaticClass()) || SrcProperty->IsA(FBoolProperty::StaticClass()))
 					{
-						int32* Dest = (int32*)DestBuffer->GetInstancePtrInt32(IntCompIdx++, DestStartIdx);
+						int32* Dest = DestBuffer->GetInstancePtrInt32(IntCompIdx++, DestStartIdx);
 
 						//Write all instances for this component.
 						for (int32 i = 0; i < NumInstances; ++i)
 						{
-							int32* Src = (int32*)((SrcData + i * SrcVarSize));
+							int32* Src = reinterpret_cast<int32*>((SrcData + i * SrcVarSize));
 							*Dest++ = *Src;
 						}
 					}
@@ -292,7 +291,6 @@ void FNiagaraDataChannelGameData::AppendFromDataSet(const FNiagaraDataBuffer* Sr
 	static TArray<uint8> LWCConversionBuffer;
 
 	int32 NumInstances = SrcBuffer->GetNumInstances();
-	int32 OrigNumDataChannel = NumElements;
 	NumElements += NumInstances;
 
 	const FNiagaraDataChannelGameDataLayout& Layout = DataChannel->GetGameDataLayout();
@@ -339,13 +337,13 @@ void FNiagaraDataChannelGameData::AppendFromDataSet(const FNiagaraDataBuffer* Sr
 			//Special case for writing to Niagara Positions. We must offset the data into the simulation local LWC space.
 			if (SrcStruct == FNiagaraTypeDefinition::GetPositionStruct())
 			{
-				float* SrcX = (float*)SrcBuffer->GetComponentPtrFloat(FloatCompIdx++);
-				float* SrcY = (float*)SrcBuffer->GetComponentPtrFloat(FloatCompIdx++);
-				float* SrcZ = (float*)SrcBuffer->GetComponentPtrFloat(FloatCompIdx++);
+				const float* SrcX = reinterpret_cast<const float*>(SrcBuffer->GetComponentPtrFloat(FloatCompIdx++));
+				const float* SrcY = reinterpret_cast<const float*>(SrcBuffer->GetComponentPtrFloat(FloatCompIdx++));
+				const float* SrcZ = reinterpret_cast<const float*>(SrcBuffer->GetComponentPtrFloat(FloatCompIdx++));
 
 				for (int32 i = 0; i < NumInstances; ++i)
 				{
-					FVector* Dest = (FVector*)(DestData + VarSize * i);
+					FVector* Dest = reinterpret_cast<FVector*>(DestData + VarSize * i);
 					*Dest = FVector(*SrcX++, *SrcY++, *SrcZ++) + FVector(SimulationLwcTile) * FLargeWorldRenderScalar::GetTileSize();
 				}
 			}
@@ -360,49 +358,49 @@ void FNiagaraDataChannelGameData::AppendFromDataSet(const FNiagaraDataBuffer* Sr
 					DestData += DestProperty->GetOffset_ForInternal();
 					if (DestPropertyIt->IsA(FDoubleProperty::StaticClass()))
 					{
-						double* Dest = (double*)(DestData);
-						float* Src = (float*)SrcBuffer->GetComponentPtrFloat(FloatCompIdx++);
+						double* Dest = reinterpret_cast<double*>(DestData);
+						const float* Src = reinterpret_cast<const float*>(SrcBuffer->GetComponentPtrFloat(FloatCompIdx++));
 
 						//Write all instances for this component.
 						for (int32 i = 0; i < NumInstances; ++i)
 						{
-							Dest = (double*)(DestData + VarSize * i);
+							Dest = reinterpret_cast<double*>(DestData + VarSize * i);
 							*Dest = *Src++;
 						}
 					}
 					else if (DestPropertyIt->IsA(FFloatProperty::StaticClass()))
 					{
-						float* Dest = (float*)(DestData);
-						float* Src = (float*)SrcBuffer->GetComponentPtrFloat(FloatCompIdx++);
+						float* Dest = reinterpret_cast<float*>(DestData);
+						const float* Src = reinterpret_cast<const float*>(SrcBuffer->GetComponentPtrFloat(FloatCompIdx++));
 
 						//Write all instances for this component.
 						for (int32 i = 0; i < NumInstances; ++i)
 						{
-							Dest = (float*)(DestData + VarSize * i);
+							Dest = reinterpret_cast<float*>(DestData + VarSize * i);
 							*Dest = *Src++;
 						}
 					}
 					else if (DestPropertyIt->IsA(FUInt16Property::StaticClass()))
 					{
-						FFloat16* Dest = (FFloat16*)(DestData);
-						FFloat16* Src = (FFloat16*)SrcBuffer->GetComponentPtrHalf(HalfCompIdx++);
+						FFloat16* Dest = reinterpret_cast<FFloat16*>(DestData);
+						const FFloat16* Src = reinterpret_cast<const FFloat16*>(SrcBuffer->GetComponentPtrHalf(HalfCompIdx++));
 
 						//Write all instances for this component.
 						for (int32 i = 0; i < NumInstances; ++i)
 						{
-							Dest = (FFloat16*)(DestData + VarSize * i);
+							Dest = reinterpret_cast<FFloat16*>(DestData + VarSize * i);
 							*Dest = *Src++;
 						}
 					}
 					else if (DestPropertyIt->IsA(FIntProperty::StaticClass()) || DestPropertyIt->IsA(FBoolProperty::StaticClass()))
 					{
-						int32* Dest = (int32*)(DestData);
-						int32* Src = (int32*)SrcBuffer->GetComponentPtrInt32(IntCompIdx++);
+						int32* Dest = reinterpret_cast<int32*>(DestData);
+						const int32* Src = reinterpret_cast<const int32*>(SrcBuffer->GetComponentPtrInt32(IntCompIdx++));
 
 						//Write all instances for this component.
 						for (int32 i = 0; i < NumInstances; ++i)
 						{
-							Dest = (int32*)(DestData + VarSize * i);
+							Dest = reinterpret_cast<int32*>(DestData + VarSize * i);
 							*Dest = *Src++;
 						}
 					}
@@ -596,7 +594,7 @@ int32 FNiagaraDataChannelData::ConsumePublishRequests(UNiagaraDataChannelHandler
 	
 	if(NDCCVars::bEmitWarningsOnLateNDCWrites && DataChannel->ShouldEnforceTickGroupReadWriteOrder())
 	{
-		ETickingGroup PublishSourceTG = (ETickingGroup)(FMath::Clamp((int32)Owner->GetCurrentTickGroup() - 1, 0, (int32)ETickingGroup::TG_MAX-1));//We're consuming from the previous TG.
+		ETickingGroup PublishSourceTG = static_cast<ETickingGroup>(FMath::Clamp(static_cast<int32>(Owner->GetCurrentTickGroup()) - 1, 0, static_cast<int32>(ETickingGroup::TG_MAX) - 1));//We're consuming from the previous TG.
 		ETickingGroup FinalWriteTG = DataChannel->GetFinalWriteTickGroup();
 
 		//TODO: Possibly allow late writes to be deferred to the next frame?
@@ -667,7 +665,7 @@ int32 FNiagaraDataChannelData::ConsumePublishRequests(UNiagaraDataChannelHandler
 
 	//Now do the actual data collection.
 
-	GameData->SetNum(NewGameDataChannel);
+	GameData->Reserve(NewGameDataChannel);
 
 	for (auto It = PublishRequests.CreateIterator(); It ; ++It)
 	{
@@ -1000,6 +998,20 @@ UNiagaraDataChannelReader* UNiagaraDataChannelLibrary::ReadFromNiagaraDataChanne
 	return CreateDataChannelReader(WorldContextObject, Channel->Get(), SearchParams, bReadPreviousFrame);
 }
 
+int32 UNiagaraDataChannelLibrary::GetDataChannelElementCount(const UObject* WorldContextObject, const UNiagaraDataChannelAsset* Channel, FNiagaraDataChannelSearchParameters SearchParams, bool bReadPreviousFrame)
+{
+	if (UNiagaraDataChannelReader* Reader = CreateDataChannelReader(WorldContextObject, Channel->Get(), SearchParams, bReadPreviousFrame))
+	{
+		return Reader->Num();
+	}
+	return 0;
+}
+
+void UNiagaraDataChannelLibrary::ReadFromNiagaraDataChannelSingle(const UObject*, const UNiagaraDataChannelAsset*, int32, FNiagaraDataChannelSearchParameters, bool, ENiagartaDataChannelReadResult&)
+{
+	// this function is just a placeholder and calls into CreateDataChannelReader and its individual read functions from the BP node
+}
+
 void UNiagaraDataChannelLibrary::WriteToNiagaraDataChannelSingle(const UObject*, const UNiagaraDataChannelAsset*, FNiagaraDataChannelSearchParameters, bool, bool, bool)
 {
 	// this function is just a placeholder and calls into CreateDataChannelWriter and its individual write functions from the BP node
diff --git a/Engine/Plugins/FX/Niagara/Source/Niagara/Private/NiagaraDataChannelAccessor.cpp b/Engine/Plugins/FX/Niagara/Source/Niagara/Private/NiagaraDataChannelAccessor.cpp
--- a/Engine/Plugins/FX/Niagara/Source/Niagara/Private/NiagaraDataChannelAccessor.cpp
+++ b/Engine/Plugins/FX/Niagara/Source/Niagara/Private/NiagaraDataChannelAccessor.cpp
@@ -28,7 +28,7 @@ int32 UNiagaraDataChannelReader::Num()const
 {
 	if (Data.IsValid())
 	{
-		return Data->GetGameData()->Num();
+		return bReadingPreviousFrame ? Data->GetGameData()->PrevNum() : Data->GetGameData()->Num();
 	}
 	return 0;
 }
@@ -46,71 +46,85 @@ bool UNiagaraDataChannelReader::ReadData(const FNiagaraVariableBase& Var, int32
 	return false;
 }
 
-double UNiagaraDataChannelReader::ReadFloat(FName VarName, int32 Index)const
+double UNiagaraDataChannelReader::ReadFloat(FName VarName, int32 Index, bool& IsValid)const
 {
 	double RetVal = 0.0f;
-	ReadData<double>(FNiagaraVariableBase(FNiagaraTypeHelper::GetDoubleDef(), VarName), Index, RetVal);
+	IsValid = ReadData<double>(FNiagaraVariableBase(FNiagaraTypeHelper::GetDoubleDef(), VarName), Index, RetVal);
 	return RetVal;
 }
 
-FVector2D UNiagaraDataChannelReader::ReadVector2D(FName VarName, int32 Index)const
+FVector2D UNiagaraDataChannelReader::ReadVector2D(FName VarName, int32 Index, bool& IsValid)const
 {
 	FVector2D RetVal = FVector2D::ZeroVector;
-	ReadData<FVector2D>(FNiagaraVariableBase(FNiagaraTypeHelper::GetVector2DDef(), VarName), Index, RetVal);
+	IsValid = ReadData<FVector2D>(FNiagaraVariableBase(FNiagaraTypeHelper::GetVector2DDef(), VarName), Index, RetVal);
 	return RetVal;
 }
 
-FVector UNiagaraDataChannelReader::ReadVector(FName VarName, int32 Index)const
+FVector UNiagaraDataChannelReader::ReadVector(FName VarName, int32 Index, bool& IsValid)const
 {
 	FVector RetVal = FVector::ZeroVector;
-	ReadData<FVector>(FNiagaraVariableBase(FNiagaraTypeHelper::GetVectorDef(), VarName), Index, RetVal);
+	IsValid = ReadData<FVector>(FNiagaraVariableBase(FNiagaraTypeHelper::GetVectorDef(), VarName), Index, RetVal);
 	return RetVal;
 }
 
-FVector4 UNiagaraDataChannelReader::ReadVector4(FName VarName, int32 Index)const
+FVector4 UNiagaraDataChannelReader::ReadVector4(FName VarName, int32 Index, bool& IsValid)const
 {
 	FVector4 RetVal = FVector4(0.0f);
-	ReadData<FVector4>(FNiagaraVariableBase(FNiagaraTypeHelper::GetVector4Def(), VarName), Index, RetVal);
+	IsValid = ReadData<FVector4>(FNiagaraVariableBase(FNiagaraTypeHelper::GetVector4Def(), VarName), Index, RetVal);
 	return RetVal;
 }
 
-FQuat UNiagaraDataChannelReader::ReadQuat(FName VarName, int32 Index)const
+FQuat UNiagaraDataChannelReader::ReadQuat(FName VarName, int32 Index, bool& IsValid)const
 {
 	FQuat RetVal = FQuat::Identity;
-	ReadData<FQuat>(FNiagaraVariableBase(FNiagaraTypeHelper::GetQuatDef(), VarName), Index, RetVal);
+	IsValid = ReadData<FQuat>(FNiagaraVariableBase(FNiagaraTypeHelper::GetQuatDef(), VarName), Index, RetVal);
 	return RetVal;
 }
 
-FLinearColor UNiagaraDataChannelReader::ReadLinearColor(FName VarName, int32 Index)const
+FLinearColor UNiagaraDataChannelReader::ReadLinearColor(FName VarName, int32 Index, bool& IsValid)const
 {
 	FLinearColor RetVal = FLinearColor::White;
-	ReadData<FLinearColor>(FNiagaraVariableBase(FNiagaraTypeDefinition::GetColorDef(), VarName), Index, RetVal);
+	IsValid = ReadData<FLinearColor>(FNiagaraVariableBase(FNiagaraTypeDefinition::GetColorDef(), VarName), Index, RetVal);
 	return RetVal;
 }
 
-int32 UNiagaraDataChannelReader::ReadInt(FName VarName, int32 Index)const
+int32 UNiagaraDataChannelReader::ReadInt(FName VarName, int32 Index, bool& IsValid)const
 {
 	int32 RetVal = 0;
-	ReadData<int32>(FNiagaraVariableBase(FNiagaraTypeDefinition::GetIntDef(), VarName), Index, RetVal);
+	IsValid = ReadData<int32>(FNiagaraVariableBase(FNiagaraTypeDefinition::GetIntDef(), VarName), Index, RetVal);
 	return RetVal;
 }
 
-uint8 UNiagaraDataChannelReader::ReadEnum(FName VarName, int32 Index) const
+uint8 UNiagaraDataChannelReader::ReadEnum(FName VarName, int32 Index, bool& IsValid) const
 {
-	return static_cast<uint8>(ReadInt(VarName, Index));
+	return static_cast<uint8>(ReadInt(VarName, Index, IsValid));
 }
 
-bool UNiagaraDataChannelReader::ReadBool(FName VarName, int32 Index)const
+bool UNiagaraDataChannelReader::ReadBool(FName VarName, int32 Index, bool& IsValid)const
 {
-	bool RetVal = false;
-	ReadData<bool>(FNiagaraVariableBase(FNiagaraTypeDefinition::GetBoolDef(), VarName), Index, RetVal);
+	FNiagaraBool RetVal(false);
+	IsValid = ReadData<FNiagaraBool>(FNiagaraVariableBase(FNiagaraTypeDefinition::GetBoolDef(), VarName), Index, RetVal);
 	return RetVal;
 }
 
-FVector UNiagaraDataChannelReader::ReadPosition(FName VarName, int32 Index)const
+FVector UNiagaraDataChannelReader::ReadPosition(FName VarName, int32 Index, bool& IsValid)const
 {
 	FVector RetVal = FVector::ZeroVector;
-	ReadData<FVector>(FNiagaraVariableBase(FNiagaraTypeDefinition::GetPositionDef(), VarName), Index, RetVal);
+	IsValid = ReadData<FVector>(FNiagaraVariableBase(FNiagaraTypeDefinition::GetPositionDef(), VarName), Index, RetVal);
+	return RetVal;
+}
+
+FNiagaraID UNiagaraDataChannelReader::ReadID(FName VarName, int32 Index, bool& IsValid) const
+{
+	FNiagaraID RetVal;
+	IsValid = ReadData<FNiagaraID>(FNiagaraVariableBase(FNiagaraTypeDefinition::GetIDDef(), VarName), Index, RetVal);
+	return RetVal;
+}
+
+FNiagaraSpawnInfo UNiagaraDataChannelReader::ReadSpawnInfo(FName VarName, int32 Index, bool& IsValid) const
+{
+	FNiagaraSpawnInfo RetVal;
+	IsValid = ReadData<FNiagaraSpawnInfo>(FNiagaraVariableBase(FNiagaraTypeDefinition(FNiagaraSpawnInfo::StaticStruct()), VarName), Index, RetVal);
 	return RetVal;
 }
 
diff --git a/Engine/Plugins/FX/Niagara/Source/Niagara/Public/NiagaraDataChannel.h b/Engine/Plugins/FX/Niagara/Source/Niagara/Public/NiagaraDataChannel.h
--- a/Engine/Plugins/FX/Niagara/Source/Niagara/Public/NiagaraDataChannel.h
+++ b/Engine/Plugins/FX/Niagara/Source/Niagara/Public/NiagaraDataChannel.h
@@ -190,6 +190,13 @@ void UNiagaraDataChannel::ForEachDataChannel(TAction Func)
 	}
 }
 
+UENUM(BlueprintType)
+enum class ENiagartaDataChannelReadResult : uint8
+{
+	Success,
+	Failure
+};
+
 /**
 * A C++ and Blueprint accessible library of utility functions for accessing Niagara DataChannel
 */
@@ -198,7 +205,7 @@ class NIAGARA_API UNiagaraDataChannelLibrary : public UBlueprintFunctionLibrary
 {
 	GENERATED_UCLASS_BODY()
 
-	UFUNCTION(BlueprintCallable, Category = NiagaraDataChannel, meta = (Keywords = "niagara DataChannel", WorldContext = "WorldContextObject", UnsafeDuringActorConstruction = "true"))
+	UFUNCTION(BlueprintInternalUseOnly, Category = NiagaraDataChannel, meta = (Keywords = "niagara DataChannel", WorldContext = "WorldContextObject", UnsafeDuringActorConstruction = "true"))
 	static UNiagaraDataChannelHandler* GetNiagaraDataChannel(const UObject* WorldContextObject, const UNiagaraDataChannelAsset* Channel);
 
 	/**
@@ -225,11 +232,39 @@ class NIAGARA_API UNiagaraDataChannelLibrary : public UBlueprintFunctionLibrary
 	 *								Reading the current frame allows for zero latency reads, but any data elements that are generated after this reader is used are missed.
 	 *								Reading the previous frame's data introduces a frame of latency but ensures we never miss any data as we have access to the whole frame.
 	 */
-	UFUNCTION(BlueprintCallable, Category = NiagaraDataChannel, meta = (Keywords = "niagara DataChannel", WorldContext = "WorldContextObject", UnsafeDuringActorConstruction = "true"))
+	UFUNCTION(BlueprintCallable, Category = NiagaraDataChannel, DisplayName="Read From Niagara Data Channel (Batch)", meta = (AdvancedDisplay = "SearchParams", Keywords = "niagara DataChannel", WorldContext = "WorldContextObject", UnsafeDuringActorConstruction = "true"))
 	static UNiagaraDataChannelReader* ReadFromNiagaraDataChannel(const UObject* WorldContextObject, const UNiagaraDataChannelAsset* Channel, FNiagaraDataChannelSearchParameters SearchParams, bool bReadPreviousFrame);
 
 	/**
-	 * Writes a single element to a Niagara Data Channel
+	 * Returns the number of readable elements in the given data channel
+	 *
+	 * @param WorldContextObject	World to execute in
+	 * @param Channel				The channel to read from
+	 * @param SearchParams			Parameters used when retrieving a specific set of Data Channel Data to read or write like the islands data channel type.
+	 * @param bReadPreviousFrame	True if this reader will read the previous frame's data. If false, we read the current frame.
+	 *								Reading the current frame allows for zero latency reads, but any data elements that are generated after this reader is used are missed.
+	 *								Reading the previous frame's data introduces a frame of latency but ensures we never miss any data as we have access to the whole frame.
+	 */
+	UFUNCTION(BlueprintCallable, Category = NiagaraDataChannel, meta = (bReadPreviousFrame="true", AdvancedDisplay = "SearchParams, bReadPreviousFrame", Keywords = "niagara DataChannel num size", WorldContext = "WorldContextObject", UnsafeDuringActorConstruction = "true"))
+	static int32 GetDataChannelElementCount(const UObject* WorldContextObject, const UNiagaraDataChannelAsset* Channel, FNiagaraDataChannelSearchParameters SearchParams, bool bReadPreviousFrame);
+
+	/**
+	 * Reads a single entry from the given data channel, if possible.
+	 *
+	 * @param WorldContextObject	World to execute in
+	 * @param Channel				The channel to read from
+	 * @param Index					The data index to read from
+	 * @param SearchParams			Parameters used when retrieving a specific set of Data Channel Data to read or write like the islands data channel type.
+	 * @param bReadPreviousFrame	True if this reader will read the previous frame's data. If false, we read the current frame.
+	 *								Reading the current frame allows for zero latency reads, but any data elements that are generated after this reader is used are missed.
+	 *								Reading the previous frame's data introduces a frame of latency but ensures we never miss any data as we have access to the whole frame.
+	 * @param ReadResult			Used by Blueprint for the return value
+	 */
+	UFUNCTION(BlueprintInternalUseOnly, Category = NiagaraDataChannel, DisplayName="Read From Niagara Data Channel", meta = (bReadPreviousFrame="true", AdvancedDisplay = "SearchParams, bReadPreviousFrame", ExpandEnumAsExecs="ReadResult", Keywords = "niagara DataChannel", WorldContext = "WorldContextObject", UnsafeDuringActorConstruction = "true"))
+	static void ReadFromNiagaraDataChannelSingle(const UObject* WorldContextObject, const UNiagaraDataChannelAsset* Channel, int32 Index, FNiagaraDataChannelSearchParameters SearchParams, bool bReadPreviousFrame, ENiagartaDataChannelReadResult& ReadResult);
+
+	/**
+	 * Writes a single element to a Niagara Data Channel. The element won't be immediately visible to readers, as it needs to be processed first. The earliest point it can be read is in the next tick group.
 	 *
 	 * @param WorldContextObject	World to execute in
 	 * @param Channel				The channel to write to
diff --git a/Engine/Plugins/FX/Niagara/Source/Niagara/Public/NiagaraDataChannelAccessor.h b/Engine/Plugins/FX/Niagara/Source/Niagara/Public/NiagaraDataChannelAccessor.h
--- a/Engine/Plugins/FX/Niagara/Source/Niagara/Public/NiagaraDataChannelAccessor.h
+++ b/Engine/Plugins/FX/Niagara/Source/Niagara/Public/NiagaraDataChannelAccessor.h
@@ -41,34 +41,40 @@ public:
 	NIAGARA_API int32 Num()const;
 
 	UFUNCTION(BlueprintCallable, Category = NiagaraDataChannel, meta = (Keywords = "Niagara DataChannel"))
-	NIAGARA_API double ReadFloat(FName VarName, int32 Index)const;
+	NIAGARA_API double ReadFloat(FName VarName, int32 Index, bool& IsValid)const;
 
 	UFUNCTION(BlueprintCallable, Category = NiagaraDataChannel, meta = (Keywords = "Niagara DataChannel"))
-	NIAGARA_API FVector2D ReadVector2D(FName VarName, int32 Index)const;
+	NIAGARA_API FVector2D ReadVector2D(FName VarName, int32 Index, bool& IsValid)const;
 
 	UFUNCTION(BlueprintCallable, Category = NiagaraDataChannel, meta = (Keywords = "Niagara DataChannel"))
-	NIAGARA_API FVector ReadVector(FName VarName, int32 Index)const;
+	NIAGARA_API FVector ReadVector(FName VarName, int32 Index, bool& IsValid)const;
 
 	UFUNCTION(BlueprintCallable, Category = NiagaraDataChannel, meta = (Keywords = "Niagara DataChannel"))
-	NIAGARA_API FVector4 ReadVector4(FName VarName, int32 Index)const;
+	NIAGARA_API FVector4 ReadVector4(FName VarName, int32 Index, bool& IsValid)const;
 
 	UFUNCTION(BlueprintCallable, Category = NiagaraDataChannel, meta = (Keywords = "Niagara DataChannel"))
-	NIAGARA_API FQuat ReadQuat(FName VarName, int32 Index)const;
+	NIAGARA_API FQuat ReadQuat(FName VarName, int32 Index, bool& IsValid)const;
 
 	UFUNCTION(BlueprintCallable, Category = NiagaraDataChannel, meta = (Keywords = "Niagara DataChannel"))
-	NIAGARA_API FLinearColor ReadLinearColor(FName VarName, int32 Index)const;
+	NIAGARA_API FLinearColor ReadLinearColor(FName VarName, int32 Index, bool& IsValid)const;
 
 	UFUNCTION(BlueprintCallable, Category = NiagaraDataChannel, meta = (Keywords = "Niagara DataChannel"))
-	NIAGARA_API int32 ReadInt(FName VarName, int32 Index)const;
+	NIAGARA_API int32 ReadInt(FName VarName, int32 Index, bool& IsValid)const;
 
 	UFUNCTION(BlueprintCallable, Category = NiagaraDataChannel, meta = (Keywords = "Niagara DataChannel"))
-	NIAGARA_API uint8 ReadEnum(FName VarName, int32 Index)const;
+	NIAGARA_API uint8 ReadEnum(FName VarName, int32 Index, bool& IsValid)const;
 
 	UFUNCTION(BlueprintCallable, Category = NiagaraDataChannel, meta = (Keywords = "Niagara DataChannel"))
-	NIAGARA_API bool ReadBool(FName VarName, int32 Index)const;
+	NIAGARA_API bool ReadBool(FName VarName, int32 Index, bool& IsValid)const;
 
 	UFUNCTION(BlueprintCallable, Category = NiagaraDataChannel, meta = (Keywords = "Niagara DataChannel"))
-	NIAGARA_API FVector ReadPosition(FName VarName, int32 Index)const;
+	NIAGARA_API FVector ReadPosition(FName VarName, int32 Index, bool& IsValid)const;
+
+	UFUNCTION(BlueprintCallable, Category = NiagaraDataChannel, meta = (Keywords = "Niagara DataChannel"))
+	NIAGARA_API FNiagaraID ReadID(FName VarName, int32 Index, bool& IsValid)const;
+
+	UFUNCTION(BlueprintCallable, Category = NiagaraDataChannel, meta = (Keywords = "Niagara DataChannel"))
+	NIAGARA_API FNiagaraSpawnInfo ReadSpawnInfo(FName VarName, int32 Index, bool& IsValid)const;
 };
 
 UCLASS(Experimental, BlueprintType, MinimalAPI)
diff --git a/Engine/Plugins/FX/Niagara/Source/Niagara/Public/NiagaraDataChannelPublic.h b/Engine/Plugins/FX/Niagara/Source/Niagara/Public/NiagaraDataChannelPublic.h
--- a/Engine/Plugins/FX/Niagara/Source/Niagara/Public/NiagaraDataChannelPublic.h
+++ b/Engine/Plugins/FX/Niagara/Source/Niagara/Public/NiagaraDataChannelPublic.h
@@ -205,7 +205,7 @@ struct FNiagaraDataChannelVariableBuffer
 		check(sizeof(T) == Size);
 		if (Index >= 0 && Index < Num())
 		{
-			T* Dest = ((T*)Data.GetData()) + Index;
+			T* Dest = reinterpret_cast<T*>(Data.GetData()) + Index;
 			*Dest = InData;
 
 			return true;
@@ -225,7 +225,7 @@ struct FNiagaraDataChannelVariableBuffer
 		if (Index >= 0 && Index < NumElems)
 		{
 			const uint8* DataPtr = bPreviousFrameData ? PrevData.GetData() : Data.GetData();
-			T* Src = ((T*)DataPtr) + Index;
+			const T* Src = reinterpret_cast<const T*>(DataPtr) + Index;
 			OutData = *Src;
 
 			return true;
diff --git a/Engine/Plugins/FX/Niagara/Source/NiagaraBlueprintNodes/Internal/K2Node_DataChannelBase.h b/Engine/Plugins/FX/Niagara/Source/NiagaraBlueprintNodes/Internal/K2Node_DataChannelBase.h
new file mode 100644
--- /dev/null
+++ b/Engine/Plugins/FX/Niagara/Source/NiagaraBlueprintNodes/Internal/K2Node_DataChannelBase.h
@@ -0,0 +1,50 @@
+﻿// Copyright Epic Games, Inc. All Rights Reserved.
+
+#pragma once
+
+#include "CoreMinimal.h"
+#include "NiagaraDataChannel.h"
+#include "K2Node_CallFunction.h"
+#include "K2Node_DataChannelBase.generated.h"
+
+
+UCLASS(MinimalAPI, Abstract, NotBlueprintable)
+class UK2Node_DataChannelBase : public UK2Node_CallFunction
+{
+	GENERATED_BODY()
+public:
+
+	NIAGARABLUEPRINTNODES_API UNiagaraDataChannel* GetDataChannel() const;
+
+	virtual void PostLoad() override;
+
+	// //~ Begin UEdGraphNode Interface.
+	virtual void AllocateDefaultPins() override;
+	virtual void PinDefaultValueChanged(UEdGraphPin* Pin) override;
+	virtual void PinConnectionListChanged(UEdGraphPin* Pin) override;
+	// //~ End UEdGraphNode Interface.
+
+	//~ Begin K2Node Interface
+	//virtual void GetMenuActions(FBlueprintActionDatabaseRegistrar& InActionRegistrar) const override;
+	virtual FText GetMenuCategory() const override;
+	virtual ERedirectType DoPinsMatchForReconstruction(const UEdGraphPin* NewPin, int32 NewPinIndex, const UEdGraphPin* OldPin, int32 OldPinIndex) const override;
+	virtual void PreloadRequiredAssets() override;
+	virtual bool ShouldShowNodeProperties() const override;
+	//~ End K2Node Interface
+
+	UPROPERTY()
+	TSet<FGuid> IgnoredVariables;
+
+protected:
+
+	UPROPERTY()
+	TObjectPtr<UNiagaraDataChannelAsset> DataChannel;
+
+	UEdGraphPin* GetChannelSelectorPin() const;
+
+#if WITH_EDITORONLY_DATA
+	UPROPERTY()
+	FGuid DataChannelVersion;
+#endif
+};
+
diff --git a/Engine/Plugins/FX/Niagara/Source/NiagaraBlueprintNodes/Internal/K2Node_ReadDataChannel.h b/Engine/Plugins/FX/Niagara/Source/NiagaraBlueprintNodes/Internal/K2Node_ReadDataChannel.h
new file mode 100644
--- /dev/null
+++ b/Engine/Plugins/FX/Niagara/Source/NiagaraBlueprintNodes/Internal/K2Node_ReadDataChannel.h
@@ -0,0 +1,30 @@
+﻿// Copyright Epic Games, Inc. All Rights Reserved.
+
+#pragma once
+
+#include "CoreMinimal.h"
+#include "K2Node_DataChannelBase.h"
+#include "K2Node_ReadDataChannel.generated.h"
+
+
+UCLASS(MinimalAPI)
+class UK2Node_ReadDataChannel : public UK2Node_DataChannelBase
+{
+	GENERATED_BODY()
+
+public:
+	UK2Node_ReadDataChannel();
+
+	// //~ Begin UEdGraphNode Interface.
+	virtual void AllocateDefaultPins() override;
+	// //~ End UEdGraphNode Interface.
+
+	//~ Begin K2Node Interface
+	virtual void GetMenuActions(FBlueprintActionDatabaseRegistrar& InActionRegistrar) const override;
+	virtual void ExpandNode(class FKismetCompilerContext& CompilerContext, UEdGraph* SourceGraph) override;
+	//~ End K2Node Interface
+
+private:
+	UFunction* GetReadFunctionForType(const FNiagaraTypeDefinition& TypeDef);
+	UEdGraphPin* GetVarPin(const FName& Name) const;
+};
diff --git a/Engine/Plugins/FX/Niagara/Source/NiagaraBlueprintNodes/Internal/K2Node_WriteDataChannel.h b/Engine/Plugins/FX/Niagara/Source/NiagaraBlueprintNodes/Internal/K2Node_WriteDataChannel.h
--- a/Engine/Plugins/FX/Niagara/Source/NiagaraBlueprintNodes/Internal/K2Node_WriteDataChannel.h
+++ b/Engine/Plugins/FX/Niagara/Source/NiagaraBlueprintNodes/Internal/K2Node_WriteDataChannel.h
@@ -3,50 +3,27 @@
 #pragma once
 
 #include "CoreMinimal.h"
-#include "K2Node_CallFunction.h"
-#include "NiagaraDataChannel.h"
+#include "K2Node_DataChannelBase.h"
 #include "K2Node_WriteDataChannel.generated.h"
 
 
 UCLASS(MinimalAPI)
-class UK2Node_WriteDataChannel : public UK2Node_CallFunction
+class UK2Node_WriteDataChannel : public UK2Node_DataChannelBase
 {
 	GENERATED_BODY()
 
 public:
 	UK2Node_WriteDataChannel();
 
-	virtual void PostLoad() override;
-
 	// //~ Begin UEdGraphNode Interface.
 	virtual void AllocateDefaultPins() override;
-	virtual void PinDefaultValueChanged(UEdGraphPin* Pin) override;
-	virtual void PinConnectionListChanged(UEdGraphPin* Pin) override;
 	// //~ End UEdGraphNode Interface.
 
 	//~ Begin K2Node Interface
 	virtual void GetMenuActions(FBlueprintActionDatabaseRegistrar& InActionRegistrar) const override;
-	virtual FText GetMenuCategory() const override;
 	virtual void ExpandNode(class FKismetCompilerContext& CompilerContext, UEdGraph* SourceGraph) override;
-	virtual ERedirectType DoPinsMatchForReconstruction(const UEdGraphPin* NewPin, int32 NewPinIndex, const UEdGraphPin* OldPin, int32 OldPinIndex) const override;
-	virtual void PreloadRequiredAssets() override;
-	virtual bool ShouldShowNodeProperties() const override;
 	//~ End K2Node Interface
 
-	NIAGARABLUEPRINTNODES_API UNiagaraDataChannel* GetDataChannel() const;
-
-	UPROPERTY()
-	TSet<FGuid> IgnoredVariables;
-	
 private:
-	UEdGraphPin* GetChannelSelectorPin() const;
 	UFunction* GetWriteFunctionForType(const FNiagaraTypeDefinition& TypeDef);
-
-	UPROPERTY()
-	TObjectPtr<UNiagaraDataChannelAsset> DataChannel;
-
-#if WITH_EDITORONLY_DATA
-	UPROPERTY()
-	FGuid DataChannelVersion;
-#endif
 };
diff --git a/Engine/Plugins/FX/Niagara/Source/NiagaraBlueprintNodes/Private/K2Node_DataChannelBase.cpp b/Engine/Plugins/FX/Niagara/Source/NiagaraBlueprintNodes/Private/K2Node_DataChannelBase.cpp
new file mode 100644
--- /dev/null
+++ b/Engine/Plugins/FX/Niagara/Source/NiagaraBlueprintNodes/Private/K2Node_DataChannelBase.cpp
@@ -0,0 +1,112 @@
+﻿// Copyright Epic Games, Inc. All Rights Reserved.
+
+
+#include "K2Node_DataChannelBase.h"
+#include "EdGraphSchema_K2.h"
+
+#include "NiagaraDataChannelPublic.h"
+
+#define LOCTEXT_NAMESPACE "K2Node_DataChannelBase"
+
+UNiagaraDataChannel* UK2Node_DataChannelBase::GetDataChannel() const
+{
+	return DataChannel ? DataChannel->Get() : nullptr;
+}
+
+void UK2Node_DataChannelBase::PostLoad()
+{
+	Super::PostLoad();
+
+#if WITH_EDITORONLY_DATA
+	if (DataChannel && DataChannel->Get()->GetVersion() != DataChannelVersion && HasValidBlueprint())
+	{
+		ReconstructNode();
+	}
+#endif
+}
+
+void UK2Node_DataChannelBase::AllocateDefaultPins()
+{
+	Super::AllocateDefaultPins();
+
+#if WITH_EDITORONLY_DATA
+	if (DataChannel)
+	{
+		DataChannelVersion = DataChannel->Get()->GetVersion();
+	}
+#endif
+}
+
+void UK2Node_DataChannelBase::PinDefaultValueChanged(UEdGraphPin* Pin)
+{
+	Super::PinDefaultValueChanged(Pin);
+
+#if WITH_EDITORONLY_DATA
+	if (Pin == GetChannelSelectorPin())
+	{
+		if (UNiagaraDataChannelAsset* ChannelAsset = Cast<UNiagaraDataChannelAsset>(Pin->DefaultObject))
+		{
+			DataChannel = ChannelAsset;
+			ReconstructNode();
+		}
+	}
+#endif
+}
+
+void UK2Node_DataChannelBase::PinConnectionListChanged(UEdGraphPin* Pin)
+{
+	Super::PinConnectionListChanged(Pin);
+
+#if WITH_EDITORONLY_DATA
+	if (Pin == GetChannelSelectorPin())
+	{
+		if (UNiagaraDataChannelAsset* ChannelAsset = Cast<UNiagaraDataChannelAsset>(Pin->DefaultObject))
+		{
+			DataChannel = Pin->LinkedTo.Num() == 0 ? ChannelAsset : nullptr;
+			ReconstructNode();
+		}
+	}
+#endif
+}
+
+FText UK2Node_DataChannelBase::GetMenuCategory() const
+{
+	static FText MenuCategory = LOCTEXT("MenuCategory", "Niagara Data Channel");
+	return MenuCategory;
+}
+
+UK2Node::ERedirectType UK2Node_DataChannelBase::DoPinsMatchForReconstruction(const UEdGraphPin* NewPin, int32 NewPinIndex, const UEdGraphPin* OldPin, int32 OldPinIndex) const
+{
+	const UEdGraphSchema_K2* K2Schema = GetDefault<UEdGraphSchema_K2>();
+	
+	ERedirectType Result = UK2Node::DoPinsMatchForReconstruction(NewPin, NewPinIndex, OldPin, OldPinIndex);
+	if (ERedirectType_None == Result && K2Schema && K2Schema->ArePinTypesCompatible(NewPin->PinType, OldPin->PinType) && (NewPin->PersistentGuid == OldPin->PersistentGuid) && OldPin->PersistentGuid.IsValid())
+	{
+		Result = ERedirectType_Name;
+	}
+
+	return Result;
+}
+
+void UK2Node_DataChannelBase::PreloadRequiredAssets()
+{
+	Super::PreloadRequiredAssets();
+
+	if (DataChannel)
+	{
+		PreloadObject(DataChannel);
+		PreloadObject(DataChannel->Get());
+	}
+}
+
+bool UK2Node_DataChannelBase::ShouldShowNodeProperties() const
+{
+	return true;
+}
+
+UEdGraphPin* UK2Node_DataChannelBase::GetChannelSelectorPin() const
+{
+	return FindPinChecked(FName("Channel"), EGPD_Input);
+}
+
+#undef LOCTEXT_NAMESPACE
diff --git a/Engine/Plugins/FX/Niagara/Source/NiagaraBlueprintNodes/Private/K2Node_ReadDataChannel.cpp b/Engine/Plugins/FX/Niagara/Source/NiagaraBlueprintNodes/Private/K2Node_ReadDataChannel.cpp
new file mode 100644
--- /dev/null
+++ b/Engine/Plugins/FX/Niagara/Source/NiagaraBlueprintNodes/Private/K2Node_ReadDataChannel.cpp
@@ -0,0 +1,241 @@
+﻿// Copyright Epic Games, Inc. All Rights Reserved.
+
+
+#include "K2Node_ReadDataChannel.h"
+
+#include "BlueprintActionDatabaseRegistrar.h"
+#include "BlueprintNodeSpawner.h"
+#include "K2Node_ExecutionSequence.h"
+#include "K2Node_IfThenElse.h"
+#include "KismetCompiler.h"
+#include "NiagaraBlueprintUtil.h"
+#include "NiagaraDataChannelAccessor.h"
+#include "Kismet/KismetMathLibrary.h"
+#include "Kismet/KismetSystemLibrary.h"
+#include "EdGraphSchema_K2.h"
+
+#define LOCTEXT_NAMESPACE "K2Node_ReadDataChannel"
+
+UK2Node_ReadDataChannel::UK2Node_ReadDataChannel()
+{
+	FunctionReference.SetExternalMember(GET_FUNCTION_NAME_CHECKED(UNiagaraDataChannelLibrary, ReadFromNiagaraDataChannelSingle), UNiagaraDataChannelLibrary::StaticClass());
+}
+
+void UK2Node_ReadDataChannel::AllocateDefaultPins()
+{
+	Super::AllocateDefaultPins();
+
+	if (DataChannel)
+	{
+		for (const FNiagaraDataChannelVariable& InVar : DataChannel->Get()->GetVariables())
+		{
+			if (IgnoredVariables.Contains(InVar.Version))
+			{
+				continue;
+			}
+			UEdGraphPin* NewPin = CreatePin(EGPD_Output, FNiagaraBlueprintUtil::TypeDefinitionToBlueprintType(InVar.GetType()), InVar.GetName());
+
+#if WITH_EDITORONLY_DATA
+			NewPin->PersistentGuid = InVar.Version;
+#endif
+		}
+	}
+}
+
+void UK2Node_ReadDataChannel::GetMenuActions(FBlueprintActionDatabaseRegistrar& InActionRegistrar) const
+{
+	const UClass* ActionKey = GetClass();
+	if (InActionRegistrar.IsOpenForRegistration(ActionKey))
+	{
+		UBlueprintNodeSpawner* NodeSpawner = UBlueprintNodeSpawner::Create(GetClass());
+		InActionRegistrar.AddBlueprintAction(ActionKey, NodeSpawner);
+	}
+}
+
+void UK2Node_ReadDataChannel::ExpandNode(FKismetCompilerContext& CompilerContext, UEdGraph* SourceGraph)
+{
+	const UEdGraphSchema_K2* Schema = CompilerContext.GetSchema();
+	UEdGraphPin* SuccessPin = FindPinChecked(FName("Success"), EGPD_Output);
+	UEdGraphPin* FailurePin = FindPinChecked(FName("Failure"), EGPD_Output);
+	if (!DataChannel)
+	{
+		// replace function with a noop if we don't have a data channel
+		UK2Node_ExecutionSequence* NoopNode = CompilerContext.SpawnIntermediateNode<UK2Node_ExecutionSequence>(this, SourceGraph);
+		NoopNode->AllocateDefaultPins();
+		
+		CompilerContext.MovePinLinksToIntermediate(*GetExecPin(), *NoopNode->GetExecPin());
+		CompilerContext.MovePinLinksToIntermediate(*FailurePin, *NoopNode->GetThenPinGivenIndex(0));
+		Schema->BreakPinLinks(*SuccessPin, false);
+		return;
+	}
+	ExpandSplitPins(CompilerContext, SourceGraph);
+	
+	// create function call node to init the writer object 
+	UK2Node_CallFunction* CreateReaderNode = CompilerContext.SpawnIntermediateNode<UK2Node_CallFunction>(this, SourceGraph);
+	CreateReaderNode->SetFromFunction(UNiagaraDataChannelLibrary::StaticClass()->FindFunctionByName(GET_FUNCTION_NAME_CHECKED(UNiagaraDataChannelLibrary, ReadFromNiagaraDataChannel)));
+	CreateReaderNode->AllocateDefaultPins();
+
+	// transfer the input pins over
+	static TArray<TPair<FName, FName>> PinsToTransfer = { {"Channel", "Channel"}, {"SearchParams", "SearchParams"}, {"bReadPreviousFrame", "bReadPreviousFrame"}};
+	
+	for (TPair<FName, FName> Pair : PinsToTransfer)
+	{
+		UEdGraphPin* OrgInputPin = FindPinChecked(Pair.Key, EGPD_Input);
+		UEdGraphPin* NewInputPin = CreateReaderNode->FindPinChecked(Pair.Value, EGPD_Input);
+		CompilerContext.MovePinLinksToIntermediate(*OrgInputPin, *NewInputPin);
+	}
+	CompilerContext.MovePinLinksToIntermediate(*GetExecPin(), *CreateReaderNode->GetExecPin());
+
+	// add a validity check for the return value
+	UK2Node_CallFunction* IsValidNode = CompilerContext.SpawnIntermediateNode<UK2Node_CallFunction>(this, SourceGraph);
+	IsValidNode->SetFromFunction(UKismetSystemLibrary::StaticClass()->FindFunctionByName(GET_FUNCTION_NAME_CHECKED(UKismetSystemLibrary, IsValid)));
+	IsValidNode->AllocateDefaultPins();
+	UK2Node_IfThenElse* IfElseNode = CompilerContext.SpawnIntermediateNode<UK2Node_IfThenElse>(this, SourceGraph);
+	IfElseNode->AllocateDefaultPins();
+
+	UEdGraphPin* IsValidInputPin = IsValidNode->FindPinChecked(FName("Object"), EGPD_Input);
+	UEdGraphPin* ReaderResultPin = CreateReaderNode->GetReturnValuePin();
+	Schema->TryCreateConnection(IsValidInputPin, ReaderResultPin);
+	Schema->TryCreateConnection(CreateReaderNode->GetThenPin(), IfElseNode->GetExecPin());
+	Schema->TryCreateConnection(IsValidNode->GetReturnValuePin(), IfElseNode->GetConditionPin());
+
+	CompilerContext.CopyPinLinksToIntermediate(*FailurePin, *IfElseNode->GetElsePin());
+
+	// create the read function nodes
+	UEdGraphPin* LastIsValidPin = nullptr;
+	UEdGraphPin* IndexPin = FindPinChecked(FName("Index"), EGPD_Input);
+	for (const FNiagaraDataChannelVariable& InVar : DataChannel->Get()->GetVariables())
+	{
+		if (IgnoredVariables.Contains(InVar.Version))
+		{
+			continue;
+		}
+		UEdGraphPin* VarOutputPin = GetVarPin(InVar.GetName());
+		if (VarOutputPin == nullptr)
+		{
+			CompilerContext.MessageLog.Error(*FText::Format(LOCTEXT("NoOutPinFound", "Missing output pin for variable '{0}' - @@"), FText::FromName(InVar.GetName())).ToString(), this);
+			continue;
+		}
+		
+		UFunction* ReadFunc = GetReadFunctionForType(InVar.GetType());
+		if (ReadFunc == nullptr)
+		{
+			CompilerContext.MessageLog.Error(*FText::Format(LOCTEXT("NoReadFuncFound", "Unable to find a read function for data channel variable '{0}' type {1}, looks like the type is not yet supported by UNiagaraDataChannelReader. (Source Pin @@)"), FText::FromName(InVar.GetName()), FText::FromString(InVar.GetType().GetName())).ToString(), VarOutputPin);
+			continue;
+		}
+		
+		UK2Node_CallFunction* ReadDataNode = CompilerContext.SpawnIntermediateNode<UK2Node_CallFunction>(this, SourceGraph);
+		ReadDataNode->SetFromFunction(ReadFunc);
+		ReadDataNode->AllocateDefaultPins();
+
+		// connect pins of the read function
+		if (!Schema->TryCreateConnection(ReaderResultPin, ReadDataNode->FindPinChecked(UEdGraphSchema_K2::PN_Self)))
+		{
+			CompilerContext.MessageLog.Error(*LOCTEXT("NoReaderConnection", "Unable to connect reader object result (UNiagaraDataChannelReader) to read function pins. @@").ToString(), this);
+			continue;
+		}
+		ReadDataNode->FindPinChecked(FName("VarName"), EGPD_Input)->DefaultValue = InVar.GetName().ToString();
+		CompilerContext.CopyPinLinksToIntermediate(*IndexPin, *ReadDataNode->FindPinChecked(FName("Index"), EGPD_Input));
+		CompilerContext.MovePinLinksToIntermediate(*VarOutputPin, *ReadDataNode->GetReturnValuePin());
+
+		// do a boolean AND of all the IsValid return values
+		UEdGraphPin* ReadValidPin = ReadDataNode->FindPinChecked(FName("IsValid"), EGPD_Output);
+		if (LastIsValidPin)
+		{
+			UK2Node_CallFunction* BoolAndNode = CompilerContext.SpawnIntermediateNode<UK2Node_CallFunction>(this, SourceGraph);
+			BoolAndNode->SetFromFunction(UKismetMathLibrary::StaticClass()->FindFunctionByName(GET_FUNCTION_NAME_CHECKED(UKismetMathLibrary, BooleanAND)));
+			BoolAndNode->AllocateDefaultPins();
+			Schema->TryCreateConnection(BoolAndNode->FindPinChecked(FName("A"), EGPD_Input), ReadValidPin);
+			Schema->TryCreateConnection(BoolAndNode->FindPinChecked(FName("B"), EGPD_Input), LastIsValidPin);
+
+			LastIsValidPin = BoolAndNode->GetReturnValuePin();
+		}
+		else
+		{
+			LastIsValidPin = ReadValidPin; 
+		}
+	}
+
+	// add a last branch based on all the read results
+	UK2Node_IfThenElse* IsValidBranchNode = CompilerContext.SpawnIntermediateNode<UK2Node_IfThenElse>(this, SourceGraph);
+	IsValidBranchNode->AllocateDefaultPins();
+	Schema->TryCreateConnection(IfElseNode->GetThenPin(), IsValidBranchNode->GetExecPin());
+	Schema->TryCreateConnection(LastIsValidPin, IsValidBranchNode->GetConditionPin());
+
+	// connect the last exec pins
+	CompilerContext.MovePinLinksToIntermediate(*FailurePin, *IsValidBranchNode->GetElsePin());
+	CompilerContext.MovePinLinksToIntermediate(*SuccessPin, *IsValidBranchNode->GetThenPin());
+}
+
+UFunction* UK2Node_ReadDataChannel::GetReadFunctionForType(const FNiagaraTypeDefinition& TypeDef)
+{
+	if (TypeDef == FNiagaraTypeHelper::GetDoubleDef() || TypeDef == FNiagaraTypeDefinition::GetFloatDef() || TypeDef == FNiagaraTypeDefinition::GetHalfDef())
+	{
+		return UNiagaraDataChannelReader::StaticClass()->FindFunctionByName(GET_FUNCTION_NAME_CHECKED(UNiagaraDataChannelReader, ReadFloat));
+	}
+	if (TypeDef == FNiagaraTypeHelper::GetVector2DDef() || TypeDef == FNiagaraTypeDefinition::GetVec2Def())
+	{
+		return UNiagaraDataChannelReader::StaticClass()->FindFunctionByName(GET_FUNCTION_NAME_CHECKED(UNiagaraDataChannelReader, ReadVector2D));
+	}
+	if (TypeDef == FNiagaraTypeDefinition::GetPositionDef())
+	{
+		return UNiagaraDataChannelReader::StaticClass()->FindFunctionByName(GET_FUNCTION_NAME_CHECKED(UNiagaraDataChannelReader, ReadPosition));
+	}
+	if (TypeDef == FNiagaraTypeHelper::GetVectorDef() || TypeDef == FNiagaraTypeDefinition::GetVec3Def())
+	{
+		return UNiagaraDataChannelReader::StaticClass()->FindFunctionByName(GET_FUNCTION_NAME_CHECKED(UNiagaraDataChannelReader, ReadVector));
+	}
+	if (TypeDef == FNiagaraTypeHelper::GetVector4Def() || TypeDef == FNiagaraTypeDefinition::GetVec4Def())
+	{
+		return UNiagaraDataChannelReader::StaticClass()->FindFunctionByName(GET_FUNCTION_NAME_CHECKED(UNiagaraDataChannelReader, ReadVector4));
+	}
+	if (TypeDef == FNiagaraTypeHelper::GetQuatDef() || TypeDef == FNiagaraTypeDefinition::GetQuatDef())
+	{
+		return UNiagaraDataChannelReader::StaticClass()->FindFunctionByName(GET_FUNCTION_NAME_CHECKED(UNiagaraDataChannelReader, ReadQuat));
+	}
+	if (TypeDef == FNiagaraTypeDefinition::GetColorDef())
+	{
+		return UNiagaraDataChannelReader::StaticClass()->FindFunctionByName(GET_FUNCTION_NAME_CHECKED(UNiagaraDataChannelReader, ReadLinearColor));
+	}
+	if (TypeDef == FNiagaraTypeDefinition::GetIntDef())
+	{
+		return UNiagaraDataChannelReader::StaticClass()->FindFunctionByName(GET_FUNCTION_NAME_CHECKED(UNiagaraDataChannelReader, ReadInt));
+	}
+	if (TypeDef == FNiagaraTypeDefinition::GetBoolDef())
+	{
+		return UNiagaraDataChannelReader::StaticClass()->FindFunctionByName(GET_FUNCTION_NAME_CHECKED(UNiagaraDataChannelReader, ReadBool));
+	}
+	if (TypeDef.GetStruct() == FNiagaraSpawnInfo::StaticStruct())
+	{
+		return UNiagaraDataChannelReader::StaticClass()->FindFunctionByName(GET_FUNCTION_NAME_CHECKED(UNiagaraDataChannelReader, ReadSpawnInfo));
+	}
+	if (TypeDef.GetStruct() == FNiagaraID::StaticStruct())
+	{
+		return UNiagaraDataChannelReader::StaticClass()->FindFunctionByName(GET_FUNCTION_NAME_CHECKED(UNiagaraDataChannelReader, ReadID));
+	}
+	if (TypeDef.GetEnum())
+	{
+		return UNiagaraDataChannelReader::StaticClass()->FindFunctionByName(GET_FUNCTION_NAME_CHECKED(UNiagaraDataChannelReader, ReadEnum));
+	}
+	return nullptr;
+}
+
+UEdGraphPin* UK2Node_ReadDataChannel::GetVarPin(const FName& Name) const
+{
+	FEdGraphPinType ExecType = GetExecPin()->PinType;
+	for (UEdGraphPin* Pin : Pins)
+	{
+		if (Pin->PinType == ExecType)
+		{
+			continue;
+		}
+		if ((EGPD_Output == Pin->Direction) && Pin->PinName == Name)
+		{
+			return Pin;
+		}
+	}
+	checkf(false, TEXT("Unable to find pin for variable %s"), *Name.ToString());
+	return nullptr;
+}
+
+#undef LOCTEXT_NAMESPACE
diff --git a/Engine/Plugins/FX/Niagara/Source/NiagaraBlueprintNodes/Private/K2Node_WriteDataChannel.cpp b/Engine/Plugins/FX/Niagara/Source/NiagaraBlueprintNodes/Private/K2Node_WriteDataChannel.cpp
--- a/Engine/Plugins/FX/Niagara/Source/NiagaraBlueprintNodes/Private/K2Node_WriteDataChannel.cpp
+++ b/Engine/Plugins/FX/Niagara/Source/NiagaraBlueprintNodes/Private/K2Node_WriteDataChannel.cpp
@@ -4,6 +4,7 @@
 
 #include "BlueprintActionDatabaseRegistrar.h"
 #include "BlueprintNodeSpawner.h"
+#include "K2Node_ExecutionSequence.h"
 #include "KismetCompiler.h"
 #include "NiagaraBlueprintUtil.h"
 #include "NiagaraDataChannel.h"
@@ -16,18 +17,6 @@ UK2Node_WriteDataChannel::UK2Node_WriteDataChannel()
 	FunctionReference.SetExternalMember(GET_FUNCTION_NAME_CHECKED(UNiagaraDataChannelLibrary, WriteToNiagaraDataChannelSingle), UNiagaraDataChannelLibrary::StaticClass());
 }
 
-void UK2Node_WriteDataChannel::PostLoad()
-{
-	Super::PostLoad();
-
-#if WITH_EDITORONLY_DATA
-	if (DataChannel && DataChannel->Get()->GetVersion() != DataChannelVersion && HasValidBlueprint())
-	{
-		ReconstructNode();
-	}
-#endif
-}
-
 void UK2Node_WriteDataChannel::AllocateDefaultPins()
 {
 	Super::AllocateDefaultPins();
@@ -46,45 +35,9 @@ void UK2Node_WriteDataChannel::AllocateDefaultPins()
 			NewPin->PersistentGuid = InVar.Version;
 #endif
 		}
-
-#if WITH_EDITORONLY_DATA
-		DataChannelVersion = DataChannel->Get()->GetVersion();
-#endif
 	}
 }
 
-void UK2Node_WriteDataChannel::PinDefaultValueChanged(UEdGraphPin* Pin)
-{
-	Super::PinDefaultValueChanged(Pin);
-
-#if WITH_EDITORONLY_DATA
-	if (Pin == GetChannelSelectorPin())
-	{
-		if (UNiagaraDataChannelAsset* ChannelAsset = Cast<UNiagaraDataChannelAsset>(Pin->DefaultObject))
-		{
-			DataChannel = ChannelAsset;
-			ReconstructNode();
-		}
-	}
-#endif
-}
-
-void UK2Node_WriteDataChannel::PinConnectionListChanged(UEdGraphPin* Pin)
-{
-	Super::PinConnectionListChanged(Pin);
-
-#if WITH_EDITORONLY_DATA
-	if (Pin == GetChannelSelectorPin())
-	{
-		if (UNiagaraDataChannelAsset* ChannelAsset = Cast<UNiagaraDataChannelAsset>(Pin->DefaultObject))
-		{
-			DataChannel = Pin->LinkedTo.Num() == 0 ? ChannelAsset : nullptr;
-			ReconstructNode();
-		}
-	}
-#endif
-}
-
 void UK2Node_WriteDataChannel::GetMenuActions(FBlueprintActionDatabaseRegistrar& InActionRegistrar) const
 {
 	const UClass* ActionKey = GetClass();
@@ -95,16 +48,15 @@ void UK2Node_WriteDataChannel::GetMenuActions(FBlueprintActionDatabaseRegistrar&
 	}
 }
 
-FText UK2Node_WriteDataChannel::GetMenuCategory() const
-{
-	static FText MenuCategory = LOCTEXT("MenuCategory", "Niagara Data Channel");
-	return MenuCategory;
-}
-
 void UK2Node_WriteDataChannel::ExpandNode(FKismetCompilerContext& CompilerContext, UEdGraph* SourceGraph)
 {
 	if (!DataChannel)
 	{
+		// replace function with a noop if we don't have a data channel
+		UK2Node_ExecutionSequence* NoopNode = CompilerContext.SpawnIntermediateNode<UK2Node_ExecutionSequence>(this, SourceGraph);
+		NoopNode->AllocateDefaultPins();
+		CompilerContext.MovePinLinksToIntermediate(*GetExecPin(), *NoopNode->GetExecPin());
+		CompilerContext.MovePinLinksToIntermediate(*GetThenPin(), *NoopNode->GetThenPinGivenIndex(0));
 		return;
 	}
 	ExpandSplitPins(CompilerContext, SourceGraph);
@@ -177,45 +129,6 @@ void UK2Node_WriteDataChannel::ExpandNode(FKismetCompilerContext& CompilerContex
 	CompilerContext.MovePinLinksToIntermediate(*OldThenPin, *LastExecPin);
 }
 
-UK2Node::ERedirectType UK2Node_WriteDataChannel::DoPinsMatchForReconstruction(const UEdGraphPin* NewPin, int32 NewPinIndex, const UEdGraphPin* OldPin, int32 OldPinIndex) const
-{
-	const UEdGraphSchema_K2* K2Schema = GetDefault<UEdGraphSchema_K2>();
-	
-	ERedirectType Result = UK2Node::DoPinsMatchForReconstruction(NewPin, NewPinIndex, OldPin, OldPinIndex);
-	if (ERedirectType_None == Result && K2Schema && K2Schema->ArePinTypesCompatible(NewPin->PinType, OldPin->PinType) && (NewPin->PersistentGuid == OldPin->PersistentGuid) && OldPin->PersistentGuid.IsValid())
-	{
-		Result = ERedirectType_Name;
-	}
-
-	return Result;
-}
-
-void UK2Node_WriteDataChannel::PreloadRequiredAssets()
-{
-	Super::PreloadRequiredAssets();
-
-	if (DataChannel)
-	{
-		PreloadObject(DataChannel);
-		PreloadObject(DataChannel->Get());
-	}
-}
-
-bool UK2Node_WriteDataChannel::ShouldShowNodeProperties() const
-{
-	return true;
-}
-
-UNiagaraDataChannel* UK2Node_WriteDataChannel::GetDataChannel() const
-{
-	return DataChannel ? DataChannel->Get() : nullptr;
-}
-
-UEdGraphPin* UK2Node_WriteDataChannel::GetChannelSelectorPin() const
-{
-	return FindPinChecked(FName("Channel"), EGPD_Input);
-}
-
 UFunction* UK2Node_WriteDataChannel::GetWriteFunctionForType(const FNiagaraTypeDefinition& TypeDef)
 {
 	if (TypeDef == FNiagaraTypeHelper::GetDoubleDef() || TypeDef == FNiagaraTypeDefinition::GetFloatDef() || TypeDef == FNiagaraTypeDefinition::GetHalfDef())
diff --git a/Engine/Plugins/FX/Niagara/Source/NiagaraEditorWidgets/Private/DetailCustomizations/NiagaraDataChannelDetails.cpp b/Engine/Plugins/FX/Niagara/Source/NiagaraEditorWidgets/Private/DetailCustomizations/NiagaraDataChannelDetails.cpp
--- a/Engine/Plugins/FX/Niagara/Source/NiagaraEditorWidgets/Private/DetailCustomizations/NiagaraDataChannelDetails.cpp
+++ b/Engine/Plugins/FX/Niagara/Source/NiagaraEditorWidgets/Private/DetailCustomizations/NiagaraDataChannelDetails.cpp
@@ -269,7 +269,7 @@ void FNiagaraDataChannelBPNodeDetails::CustomizeDetails(IDetailLayoutBuilder& De
 	}
 
 	UObject* Obj = SelectedObjects[0].Get();
-	UK2Node_WriteDataChannel* Node = Cast<UK2Node_WriteDataChannel>(Obj);
+	UK2Node_DataChannelBase* Node = Cast<UK2Node_DataChannelBase>(Obj);
 	if(Node == nullptr || Node->GetDataChannel() == nullptr)
 	{
 		return;
@@ -279,13 +279,13 @@ void FNiagaraDataChannelBPNodeDetails::CustomizeDetails(IDetailLayoutBuilder& De
 	IDetailCategoryBuilder& ChannelCategoryBuilder = DetailBuilder.EditCategory(TEXT("Data Channel"));
 	for (const FNiagaraDataChannelVariable& Var : DataChannel->GetVariables())
 	{
-		TWeakObjectPtr<UK2Node_WriteDataChannel> NodePtr(Node);
+		TWeakObjectPtr<UK2Node_DataChannelBase> NodePtr(Node);
 		FGuid VarGuid = Var.Version;
 		auto CheckStateChanged = [NodePtr, VarGuid](const ECheckBoxState NewState)
 		{
-			if (UK2Node_WriteDataChannel* Node = NodePtr.Get())
+			if (UK2Node_DataChannelBase* Node = NodePtr.Get())
 			{
-				FScopedTransaction Transaction(LOCTEXT("ChangeAttributeWrite", "Change data channel attribute write"));
+				FScopedTransaction Transaction(LOCTEXT("ChangeAttributeAccess", "Change data channel attribute access"));
 				Node->Modify();
 				if (NewState == ECheckBoxState::Checked)
 				{
@@ -301,7 +301,7 @@ void FNiagaraDataChannelBPNodeDetails::CustomizeDetails(IDetailLayoutBuilder& De
 		};
 		auto SPCheckState = [NodePtr, VarGuid]()
 		{
-			if (UK2Node_WriteDataChannel* Node = NodePtr.Get())
+			if (UK2Node_DataChannelBase* Node = NodePtr.Get())
 			{
 				return Node->IgnoredVariables.Contains(VarGuid) ? ECheckBoxState::Unchecked : ECheckBoxState::Checked;
 			}
diff --git a/Engine/Plugins/FX/Niagara/Source/NiagaraEditorWidgets/Private/NiagaraEditorWidgetsModule.cpp b/Engine/Plugins/FX/Niagara/Source/NiagaraEditorWidgets/Private/NiagaraEditorWidgetsModule.cpp
--- a/Engine/Plugins/FX/Niagara/Source/NiagaraEditorWidgets/Private/NiagaraEditorWidgetsModule.cpp
+++ b/Engine/Plugins/FX/Niagara/Source/NiagaraEditorWidgets/Private/NiagaraEditorWidgetsModule.cpp
@@ -145,6 +145,7 @@ void FNiagaraEditorWidgetsModule::StartupModule()
 
 	PropertyModule.RegisterCustomClassLayout("NiagaraDataInterfaceDataChannelRead", FOnGetDetailCustomizationInstance::CreateStatic(&FNiagaraDataInterfaceDataChannelReadDetails::MakeInstance));
 	PropertyModule.RegisterCustomClassLayout("K2Node_WriteDataChannel", FOnGetDetailCustomizationInstance::CreateStatic(&FNiagaraDataChannelBPNodeDetails::MakeInstance));
+	PropertyModule.RegisterCustomClassLayout("K2Node_ReadDataChannel", FOnGetDetailCustomizationInstance::CreateStatic(&FNiagaraDataChannelBPNodeDetails::MakeInstance));
 
 	PropertyModule.RegisterCustomClassLayout(
 		UNiagaraDataChannel_Islands::StaticClass()->GetFName(),
