Fixed a crash when editing a data channel asset and compiling a BP with dynamic read/write nodes in it

diff --git a/Engine/Plugins/FX/Niagara/Source/NiagaraBlueprintNodes/Private/K2Node_ReadDataChannel.cpp b/Engine/Plugins/FX/Niagara/Source/NiagaraBlueprintNodes/Private/K2Node_ReadDataChannel.cpp
--- a/Engine/Plugins/FX/Niagara/Source/NiagaraBlueprintNodes/Private/K2Node_ReadDataChannel.cpp
+++ b/Engine/Plugins/FX/Niagara/Source/NiagaraBlueprintNodes/Private/K2Node_ReadDataChannel.cpp
@@ -68,6 +68,11 @@ void UK2Node_ReadDataChannel::ExpandNode(FKismetCompilerContext& CompilerContext
 		Schema->BreakPinLinks(*SuccessPin, false);
 		return;
 	}
+	if (DataChannelVersion != DataChannel->Get()->GetVersion())
+	{
+		CompilerContext.MessageLog.Error(*LOCTEXT("StaleNode", "Node is out of sync with the data channel asset, please refresh node to fix up the pins - @@").ToString(), this);
+		return;
+	}
 	ExpandSplitPins(CompilerContext, SourceGraph);
 	
 	// create function call node to init the writer object 
diff --git a/Engine/Plugins/FX/Niagara/Source/NiagaraBlueprintNodes/Private/K2Node_WriteDataChannel.cpp b/Engine/Plugins/FX/Niagara/Source/NiagaraBlueprintNodes/Private/K2Node_WriteDataChannel.cpp
--- a/Engine/Plugins/FX/Niagara/Source/NiagaraBlueprintNodes/Private/K2Node_WriteDataChannel.cpp
+++ b/Engine/Plugins/FX/Niagara/Source/NiagaraBlueprintNodes/Private/K2Node_WriteDataChannel.cpp
@@ -59,6 +59,11 @@ void UK2Node_WriteDataChannel::ExpandNode(FKismetCompilerContext& CompilerContex
 		CompilerContext.MovePinLinksToIntermediate(*GetThenPin(), *NoopNode->GetThenPinGivenIndex(0));
 		return;
 	}
+	if (DataChannelVersion != DataChannel->Get()->GetVersion())
+	{
+		CompilerContext.MessageLog.Error(*LOCTEXT("StaleNode", "Node is out of sync with the data channel asset, please refresh node to fix up the pins - @@").ToString(), this);
+		return;
+	}
 	ExpandSplitPins(CompilerContext, SourceGraph);
 	
 	// create function call node to init the writer object 
