Automatically capture focus in the NDC type change menu

diff --git a/Engine/Plugins/FX/Niagara/Source/NiagaraEditor/Private/Customizations/NiagaraTypeCustomizations.cpp b/Engine/Plugins/FX/Niagara/Source/NiagaraEditor/Private/Customizations/NiagaraTypeCustomizations.cpp
--- a/Engine/Plugins/FX/Niagara/Source/NiagaraEditor/Private/Customizations/NiagaraTypeCustomizations.cpp
+++ b/Engine/Plugins/FX/Niagara/Source/NiagaraEditor/Private/Customizations/NiagaraTypeCustomizations.cpp
@@ -2886,7 +2886,7 @@ void FNiagaraDataChannelVariableDetailsCustomization::CustomizeChildren(TSharedR
 				.VAlign(VAlign_Center)
 				.Padding(2.0f)
 				[
-					SNew(SComboButton)
+					SAssignNew(ChangeTypeButton, SComboButton)
 					.HasDownArrow(true)
 					.ContentPadding(0)
 					.OnGetMenuContent(this, &FNiagaraDataChannelVariableDetailsCustomization::GetTypeMenu, TypeDefHandleHandle, Variable)
@@ -3002,9 +3002,13 @@ void SNiagaraDataChannelTypeSelectMenu::CollectAllActions(FGraphActionListBuilde
 
 TSharedRef<SWidget> FNiagaraDataChannelVariableDetailsCustomization::GetTypeMenu(TSharedPtr<IPropertyHandle> InPropertyHandle, FNiagaraDataChannelVariable* Var)
 {
-	return SNew(SNiagaraDataChannelTypeSelectMenu)
+	TSharedRef<SNiagaraDataChannelTypeSelectMenu> TypeSelectMenu = SNew(SNiagaraDataChannelTypeSelectMenu)
 	.PropertyHandle(InPropertyHandle)
 	.VarToModify(Var);
+
+	ChangeTypeButton->SetMenuContentWidgetToFocus(TypeSelectMenu->GetSearchBox());
+	
+	return TypeSelectMenu;
 }
 
 
diff --git a/Engine/Plugins/FX/Niagara/Source/NiagaraEditor/Private/Customizations/NiagaraTypeCustomizations.h b/Engine/Plugins/FX/Niagara/Source/NiagaraEditor/Private/Customizations/NiagaraTypeCustomizations.h
--- a/Engine/Plugins/FX/Niagara/Source/NiagaraEditor/Private/Customizations/NiagaraTypeCustomizations.h
+++ b/Engine/Plugins/FX/Niagara/Source/NiagaraEditor/Private/Customizations/NiagaraTypeCustomizations.h
@@ -474,4 +474,5 @@ public:
 
 private:
 	TSharedRef<SWidget> GetTypeMenu(TSharedPtr<IPropertyHandle> InPropertyHandle, FNiagaraDataChannelVariable* Var);
+	TSharedPtr<SComboButton> ChangeTypeButton;
 };
