- Don't post load ResolvedDataInterfaces in cooked builds to avoid PostLoad issues
- In cooked builds we are guaranteed PostLoad vs in uncooked builds, and if we post load something outered to an object other than our own it will be post loaded which can result in emitter being disabled due to GPU scripts not being ready

diff --git a/Engine/Plugins/FX/Niagara/Source/Niagara/Private/NiagaraScript.cpp b/Engine/Plugins/FX/Niagara/Source/Niagara/Private/NiagaraScript.cpp
--- a/Engine/Plugins/FX/Niagara/Source/Niagara/Private/NiagaraScript.cpp
+++ b/Engine/Plugins/FX/Niagara/Source/Niagara/Private/NiagaraScript.cpp
@@ -2157,6 +2157,11 @@ void UNiagaraScript::PostLoad()
 		GenerateDefaultFunctionBindings();
 	}
 
+#if WITH_EDITORONLY_DATA
+	// Only post load resolved data interfaces with editor data to avoid issues around postloading a data interace who's outer is to a different object.
+	// I.e. if we were to post load a data interface during UNiagaraEmitter postload phase that was outered to the UNiagaraSystem then the system would be post loaded during the emitters postload
+	// this can result in emitters not running because the shader has not been processed yet.
+	// When we run with editor data we defer this process to OnAssetLoaded but this doesn't exist in cooked builds
 	for (FNiagaraScriptResolvedDataInterfaceInfo& ResolvedDataInterface : ResolvedDataInterfaces)
 	{
 		if (ResolvedDataInterface.ResolvedDataInterface != nullptr)
@@ -2173,7 +2178,6 @@ void UNiagaraScript::PostLoad()
 		}
 	}
 
-#if WITH_EDITORONLY_DATA
 	// Because we might be using these cached data interfaces, we need to make sure that they are properly postloaded.
 	for (FNiagaraScriptDataInterfaceInfo& Info : CachedDefaultDataInterfaces)
 	{
