Improved data channel debug capabilities

* added instigator info for data channel writes
* added logging support for writes

diff --git a/Engine/Plugins/FX/Niagara/Source/Niagara/Internal/NiagaraDataChannelCommon.h b/Engine/Plugins/FX/Niagara/Source/Niagara/Internal/NiagaraDataChannelCommon.h
--- a/Engine/Plugins/FX/Niagara/Source/Niagara/Internal/NiagaraDataChannelCommon.h
+++ b/Engine/Plugins/FX/Niagara/Source/Niagara/Internal/NiagaraDataChannelCommon.h
@@ -42,6 +42,11 @@ struct FNiagaraDataChannelPublishRequest
 	*/
 	FVector3f LwcTile = FVector3f::ZeroVector;
 
+#if WITH_NIAGARA_DEBUGGER
+	/** Instigator of this write, used for debug tracking */
+	FString DebugSource;
+#endif
+
 	FNiagaraDataChannelPublishRequest() = default;
 	explicit FNiagaraDataChannelPublishRequest(FNiagaraDataBuffer* InData)
 		: Data(InData)
@@ -71,7 +76,7 @@ struct FNiagaraDataChannelData final
 
 	NIAGARA_API void BeginFrame(UNiagaraDataChannelHandler* Owner);
 	NIAGARA_API void EndFrame(UNiagaraDataChannelHandler* Owner);
-	NIAGARA_API int32 ConsumePublishRequests(UNiagaraDataChannelHandler* Owner);
+	NIAGARA_API int32 ConsumePublishRequests(UNiagaraDataChannelHandler* Owner, const ETickingGroup& TickGroup);
 
 	NIAGARA_API FNiagaraDataChannelGameData* GetGameData();
 	NIAGARA_API FNiagaraDataBufferRef GetCPUData(bool bPreviousFrame);
diff --git a/Engine/Plugins/FX/Niagara/Source/Niagara/Private/DataInterface/NiagaraDataInterfaceDataChannelWrite.cpp b/Engine/Plugins/FX/Niagara/Source/Niagara/Private/DataInterface/NiagaraDataInterfaceDataChannelWrite.cpp
--- a/Engine/Plugins/FX/Niagara/Source/Niagara/Private/DataInterface/NiagaraDataInterfaceDataChannelWrite.cpp
+++ b/Engine/Plugins/FX/Niagara/Source/Niagara/Private/DataInterface/NiagaraDataInterfaceDataChannelWrite.cpp
@@ -288,6 +288,9 @@ struct FNDIDataChannelWriteInstanceData
 				PublishRequest.bVisibleToGPUSims = Interface->bPublishToGPU;
 				PublishRequest.Data = Data->GetCurrentData();
 				PublishRequest.LwcTile = Instance->GetLWCTile();
+#if !UE_BUILD_SHIPPING
+				PublishRequest.DebugSource = FString::Format(TEXT("{0} ({1})"), {Instance->GetSystem()->GetName(), GetPathNameSafe(Interface)});
+#endif
 				DataChannelData->Publish(PublishRequest);
 			}
 		}
diff --git a/Engine/Plugins/FX/Niagara/Source/Niagara/Private/NiagaraDataChannel.cpp b/Engine/Plugins/FX/Niagara/Source/Niagara/Private/NiagaraDataChannel.cpp
--- a/Engine/Plugins/FX/Niagara/Source/Niagara/Private/NiagaraDataChannel.cpp
+++ b/Engine/Plugins/FX/Niagara/Source/Niagara/Private/NiagaraDataChannel.cpp
@@ -8,10 +8,15 @@
 #include "NiagaraDataChannelManager.h"
 #include "Engine/Engine.h"
 #include "UObject/UObjectGlobals.h"
-#include "UObject/Package.h"
 #include "NiagaraDataChannelAccessor.h"
 #include "NiagaraGpuComputeDispatchInterface.h"
-#include "Engine/Engine.h"
+#include "Misc/LazySingleton.h"
+
+#if WITH_NIAGARA_DEBUGGER
+#include "Serialization/JsonWriter.h"
+#include "Serialization/JsonSerializer.h"
+#include "Dom/JsonObject.h"
+#endif
 
 #include UE_INLINE_GENERATED_CPP_BY_NAME(NiagaraDataChannel)
 
@@ -21,6 +26,20 @@ namespace NDCCVars
 {
 	bool bEmitWarningsOnLateNDCWrites = !UE_BUILD_SHIPPING;
 	static FAutoConsoleVariableRef CVarEmitWarningsOnLateNDCWrites(TEXT("fx.Niagara.DataChannels.WarnOnLateWrites"), bEmitWarningsOnLateNDCWrites, TEXT("If true, late writes to NDCs will generate warnings. Late meaning after their final allowed tick group."), ECVF_Default);
+
+	int32 LogWritesToOutputLog = 0;
+	static FAutoConsoleVariableRef CVarLogWritesToOutputLog(TEXT("fx.Niagara.DataChannels.LogWritesToOutputLog"), LogWritesToOutputLog, TEXT("0=Disabled, 1=Log write summary, 2=Also write data; If >0, the NDC debugger will print all data channel writes to the output log."), ECVF_Default);
+
+	int32 FrameDataToCapture = 0;
+	static FAutoConsoleVariableRef CVarFrameDataToCapture(TEXT("fx.Niagara.DataChannels.FrameDataToCapture"), FrameDataToCapture, TEXT("The number of frames the debugger will capture for write requests."), ECVF_Default);
+
+#if WITH_NIAGARA_DEBUGGER
+	FAutoConsoleCommand CmdDumpWriteLog(
+		TEXT("fx.Niagara.DataChannels.DumpWriteLog"),
+		TEXT("Dump all the currently stored writes to the log (see fx.Niagara.DataChannels.FrameDataToCapture on how many frames are captured)"),
+		FConsoleCommandDelegate::CreateStatic(FNiagaraDataChannelDebugUtilities::DumpAllWritesToLog)
+	);
+#endif
 };
 
 //////////////////////////////////////////////////////////////////////////
@@ -565,7 +584,7 @@ void FNiagaraDataChannelData::BeginFrame(UNiagaraDataChannelHandler* Owner)
 void FNiagaraDataChannelData::EndFrame(UNiagaraDataChannelHandler* Owner)
 {
 	//We must do one final tick to process any final items generated by the last things to tick this frame.
-	ConsumePublishRequests(Owner);
+	ConsumePublishRequests(Owner, TG_LastDemotable);
 
 	check(RTProxy);
 	if (FNiagaraGpuComputeDispatchInterface* DispathInterface = FNiagaraGpuComputeDispatchInterface::Get(Owner->GetWorld()))
@@ -579,7 +598,7 @@ void FNiagaraDataChannelData::EndFrame(UNiagaraDataChannelHandler* Owner)
 	BuffersForGPU.Reset();
 }
 
-int32 FNiagaraDataChannelData::ConsumePublishRequests(UNiagaraDataChannelHandler* Owner)
+int32 FNiagaraDataChannelData::ConsumePublishRequests(UNiagaraDataChannelHandler* Owner, const ETickingGroup& TickGroup)
 {
 	check(IsValid(Owner));
 
@@ -708,13 +727,16 @@ int32 FNiagaraDataChannelData::ConsumePublishRequests(UNiagaraDataChannelHandler
 				check(0);
 			}
 		}
+#if WITH_NIAGARA_DEBUGGER
+		FNiagaraDataChannelDebugUtilities::LogWrite(PublishRequest, Owner->GetDataChannel(), TickGroup);
+#endif
 		
 		It.RemoveCurrentSwap();
 	}
 
 	CPUSimData->EndSimulate();
 
-#if !UE_BUILD_SHIPPING
+#if WITH_NIAGARA_DEBUGGER
 	if(Owner->GetDataChannel()->GetVerboseLogging())
 	{
 		FString Label = FString::Printf(TEXT("Data Channel %s - CURR"), *Owner->GetDataChannel()->GetName());
@@ -985,16 +1007,28 @@ UNiagaraDataChannelLibrary::UNiagaraDataChannelLibrary(const FObjectInitializer&
 
 UNiagaraDataChannelHandler* UNiagaraDataChannelLibrary::GetNiagaraDataChannel(const UObject* WorldContextObject, const UNiagaraDataChannelAsset* Channel)
 {
+	if (!Channel)
+	{
+		return nullptr;
+	}
 	return FindDataChannelHandler(WorldContextObject, Channel->Get());
 }
 
-UNiagaraDataChannelWriter* UNiagaraDataChannelLibrary::WriteToNiagaraDataChannel(const UObject* WorldContextObject, const UNiagaraDataChannelAsset* Channel, FNiagaraDataChannelSearchParameters SearchParams, int32 Count, bool bVisibleToGame, bool bVisibleToCPU, bool bVisibleToGPU)
+UNiagaraDataChannelWriter* UNiagaraDataChannelLibrary::WriteToNiagaraDataChannel(const UObject* WorldContextObject, const UNiagaraDataChannelAsset* Channel, FNiagaraDataChannelSearchParameters SearchParams, int32 Count, bool bVisibleToGame, bool bVisibleToCPU, bool bVisibleToGPU, const FString& DebugSource)
 {
-	return CreateDataChannelWriter(WorldContextObject, Channel->Get(), SearchParams, Count, bVisibleToGame, bVisibleToCPU, bVisibleToGPU);
+	if (!Channel)
+	{
+		return nullptr;
+	}
+	return CreateDataChannelWriter(WorldContextObject, Channel->Get(), SearchParams, Count, bVisibleToGame, bVisibleToCPU, bVisibleToGPU, DebugSource);
 }
 
 UNiagaraDataChannelReader* UNiagaraDataChannelLibrary::ReadFromNiagaraDataChannel(const UObject* WorldContextObject, const UNiagaraDataChannelAsset* Channel, FNiagaraDataChannelSearchParameters SearchParams, bool bReadPreviousFrame)
 {
+	if (!Channel)
+	{
+		return nullptr;
+	}
 	return CreateDataChannelReader(WorldContextObject, Channel->Get(), SearchParams, bReadPreviousFrame);
 }
 
@@ -1030,7 +1064,7 @@ UNiagaraDataChannelHandler* UNiagaraDataChannelLibrary::FindDataChannelHandler(c
 	return nullptr;
 }
 
-UNiagaraDataChannelWriter* UNiagaraDataChannelLibrary::CreateDataChannelWriter(const UObject* WorldContextObject, const UNiagaraDataChannel* Channel, FNiagaraDataChannelSearchParameters SearchParams, int32 Count, bool bVisibleToGame, bool bVisibleToCPU, bool bVisibleToGPU)
+UNiagaraDataChannelWriter* UNiagaraDataChannelLibrary::CreateDataChannelWriter(const UObject* WorldContextObject, const UNiagaraDataChannel* Channel, FNiagaraDataChannelSearchParameters SearchParams, int32 Count, bool bVisibleToGame, bool bVisibleToCPU, bool bVisibleToGPU, const FString& DebugSource)
 {
 	check(IsInGameThread());
 	UWorld* World = (WorldContextObject != nullptr) ? GEngine->GetWorldFromContextObject(WorldContextObject, EGetWorldErrorMode::LogAndReturnNull) : nullptr;
@@ -1042,7 +1076,7 @@ UNiagaraDataChannelWriter* UNiagaraDataChannelLibrary::CreateDataChannelWriter(c
 			{
 				if(UNiagaraDataChannelWriter* Writer = Handler->GetDataChannelWriter())
 				{
-					if(Writer->InitWrite(SearchParams, Count, bVisibleToGame, bVisibleToCPU, bVisibleToGPU))
+					if(Writer->InitWrite(SearchParams, Count, bVisibleToGame, bVisibleToCPU, bVisibleToGPU, DebugSource))
 					{
 						return Writer;
 					}
@@ -1148,7 +1182,7 @@ FNiagaraTypeDefinition FNiagaraDataChannelVariable::ToDataChannelType(const FNia
 
 //////////////////////////////////////////////////////////////////////////
 
-#if !UE_BUILD_SHIPPING
+#if WITH_NIAGARA_DEBUGGER
 void FNiagaraDataChannelDebugUtilities::BeginFrame(FNiagaraWorldManager* WorldMan, float DeltaSeconds)
 {
 	check(WorldMan);
@@ -1172,7 +1206,225 @@ UNiagaraDataChannelHandler* FNiagaraDataChannelDebugUtilities::FindDataChannelHa
 	check(WorldMan);
 	return WorldMan->GetDataChannelManager().FindDataChannelHandler(DataChannel);
 }
-#endif //!UE_BUILD_SHIPPING
+
+void FNiagaraDataChannelDebugUtilities::LogWrite(const FNiagaraDataChannelPublishRequest& WriteRequest, const UNiagaraDataChannel* DataChannel, const ETickingGroup& TickGroup)
+{
+	FNiagaraDataChannelDebugUtilities& Debugger = Get();
+	if (NDCCVars::FrameDataToCapture > 0)
+	{
+		if (Debugger.FrameData.Num() > NDCCVars::FrameDataToCapture)
+		{
+			Debugger.FrameData.SetNum(NDCCVars::FrameDataToCapture);
+		}
+		FFrameDebugData* Data;
+		if (Debugger.FrameData.Num() && Debugger.FrameData.Last().FrameNumber == GFrameCounter)
+		{
+			Data = &Debugger.FrameData.Last();
+		} 
+		else
+		{
+			if (Debugger.FrameData.Num() == NDCCVars::FrameDataToCapture)
+			{
+				Debugger.FrameData.RemoveAt(0);
+			}
+			Data = &Debugger.FrameData.AddDefaulted_GetRef();
+			Data->FrameNumber = GFrameCounter;
+		}
+		FChannelWriteRequest DebugData;
+		DebugData.DebugSource = WriteRequest.DebugSource;
+		DebugData.bVisibleToGame = WriteRequest.bVisibleToGame;
+		DebugData.bVisibleToCPUSims = WriteRequest.bVisibleToCPUSims;
+		DebugData.bVisibleToGPUSims = WriteRequest.bVisibleToGPUSims;
+		DebugData.TickGroup = TickGroup;
+		if (WriteRequest.GameData.IsValid())
+		{
+			DebugData.Data = WriteRequest.GameData;
+		}
+		else if (ensure(WriteRequest.Data))
+		{
+			DebugData.Data = MakeShared<FNiagaraDataChannelGameData>();
+			DebugData.Data->Init(DataChannel);
+			DebugData.Data->AppendFromDataSet(WriteRequest.Data, WriteRequest.LwcTile);
+		}
+		Data->WriteRequests.Add(DebugData);
+	}
+	else
+	{
+		if (Debugger.FrameData.Num())
+		{
+			Debugger.FrameData.Empty();
+		}
+	}
+	
+	if (NDCCVars::LogWritesToOutputLog > 0)
+	{
+		FString DataString;
+		uint32 NumInsts = 0;
+		if (FNiagaraDataChannelGameData* RequestGameData = WriteRequest.GameData.Get())
+		{
+			NumInsts = RequestGameData->Num();
+			if (NDCCVars::LogWritesToOutputLog > 1)
+			{
+				DataString = ToJson(RequestGameData);
+			}
+		}
+		else if (ensure(WriteRequest.Data))
+		{
+			NumInsts = WriteRequest.Data->GetNumInstances();
+			if (NDCCVars::LogWritesToOutputLog > 1)
+			{
+				FNiagaraDataChannelGameData TempData;
+				TempData.Init(DataChannel);
+				TempData.AppendFromDataSet(WriteRequest.Data, WriteRequest.LwcTile);
+				DataString = ToJson(&TempData);
+			}
+		}
+		
+		UE_LOG(LogNiagara, Log, TEXT("Frame %llu, TG %s, NDC write by %s (BP[%s]/CPU[%s]/GPU[%s]): %i entries to data channel %s %s%s")
+			, GFrameCounter
+			, *TickGroupToString(TickGroup)
+			, *WriteRequest.DebugSource
+			, WriteRequest.bVisibleToGame ? TEXT("X") : TEXT(" ")
+			, WriteRequest.bVisibleToCPUSims ? TEXT("X") : TEXT(" ")
+			, WriteRequest.bVisibleToGPUSims ? TEXT("X") : TEXT(" ")
+			, NumInsts
+			, *GetPathNameSafe(DataChannel)
+			, DataString.IsEmpty() ? TEXT("") : TEXT("\n")
+			, *DataString);
+	}
+}
+
+void FNiagaraDataChannelDebugUtilities::DumpAllWritesToLog()
+{
+	FNiagaraDataChannelDebugUtilities& Debugger = Get();
+	if (Debugger.FrameData.Num() == 0)
+	{
+		UE_LOG(LogNiagara, Log, TEXT("No writes are currently stored in the log. fx.Niagara.DataChannels.FrameDataToCapture = %i"), NDCCVars::FrameDataToCapture);
+		return;
+	}
+	UE_LOG(LogNiagara, Log, TEXT("Current Frame is %llu, logging data from oldest to newest:"), GFrameCounter);
+	for (const FFrameDebugData& FrameData : Debugger.FrameData)
+	{
+		UE_LOG(LogNiagara, Log, TEXT("Frame %llu: %i entries"), FrameData.FrameNumber, FrameData.WriteRequests.Num());
+		for (const FChannelWriteRequest& Request : FrameData.WriteRequests)
+		{
+			FString DataString = ToJson(Request.Data.Get());
+			UE_LOG(LogNiagara, Log, TEXT("Write by %s (BP[%s]/CPU[%s]/GPU[%s], TG %s): %i entries to data channel %s \n%s")
+			, *Request.DebugSource
+			, Request.bVisibleToGame ? TEXT("X") : TEXT(" ")
+			, Request.bVisibleToCPUSims ? TEXT("X") : TEXT(" ")
+			, Request.bVisibleToGPUSims ? TEXT("X") : TEXT(" ")
+			, *TickGroupToString(Request.TickGroup)
+			, Request.Data->Num()
+			, *GetPathNameSafe(Request.Data->GetDataChannel())
+			, *DataString);
+		}
+		UE_LOG(LogNiagara, Log, TEXT("----------------------------------------------"));
+	} 
+}
+
+FNiagaraDataChannelDebugUtilities& FNiagaraDataChannelDebugUtilities::Get()
+{
+	return TLazySingleton<FNiagaraDataChannelDebugUtilities>::Get();
+}
+
+void FNiagaraDataChannelDebugUtilities::TearDown()
+{
+	return TLazySingleton<FNiagaraDataChannelDebugUtilities>::TearDown();
+}
+
+FString FNiagaraDataChannelDebugUtilities::ToJson(FNiagaraDataChannelGameData* Data)
+{
+	FString JsonString;
+	TSharedRef<TJsonWriter<>> JsonWriter = TJsonWriterFactory<>::Create(&JsonString);
+	JsonWriter->WriteArrayStart();
+
+	TConstArrayView<FNiagaraDataChannelVariable> ChannelVariables = Data->GetDataChannel()->GetVariables();
+	for (int32 i = 0; i < Data->Num(); i++)
+	{
+		TSharedPtr<FJsonObject> JsonObject = MakeShared<FJsonObject>();
+		for (const FNiagaraDataChannelVariable& Var : ChannelVariables)
+		{
+			FString VarName = Var.GetName().ToString();
+			FNiagaraDataChannelVariableBuffer* Buffer = Data->FindVariableBuffer(Var);
+			if (Var.GetType() == FNiagaraTypeHelper::GetDoubleDef() && Buffer)
+			{
+				double Value;
+				Buffer->Read<double>(i, Value, false);
+				JsonObject->SetNumberField(VarName, Value);
+			}
+			else if (Var.GetType() == FNiagaraTypeDefinition::GetBoolDef() && Buffer)
+			{
+				FNiagaraBool Value;
+				Buffer->Read<FNiagaraBool>(i, Value, false);
+				JsonObject->SetBoolField(VarName, Value);
+			}
+			else if (Var.GetType() == FNiagaraTypeDefinition::GetIntDef() && Buffer)
+			{
+				int32 Value;
+				Buffer->Read<int32>(i, Value, false);
+				JsonObject->SetNumberField(VarName, Value);
+			}
+			else if ((Var.GetType() == FNiagaraTypeHelper::GetVectorDef() || Var.GetType() == FNiagaraTypeDefinition::GetPositionDef()) && Buffer)
+			{
+				FVector Value;
+				Buffer->Read<FVector>(i, Value, false);
+				JsonObject->SetStringField(VarName, Value.ToString());
+			}
+			else if (Var.GetType() == FNiagaraTypeHelper::GetVector2DDef() && Buffer)
+			{
+				FVector2D Value;
+				Buffer->Read<FVector2D>(i, Value, false);
+				JsonObject->SetStringField(VarName, Value.ToString());
+			}
+			else if (Var.GetType() == FNiagaraTypeHelper::GetVector4Def() && Buffer)
+			{
+				FVector4 Value;
+				Buffer->Read<FVector4>(i, Value, false);
+				JsonObject->SetStringField(VarName, Value.ToString());
+			}
+			else if (Var.GetType() == FNiagaraTypeHelper::GetQuatDef() && Buffer)
+			{
+				FQuat Value;
+				Buffer->Read<FQuat>(i, Value, false);
+				JsonObject->SetStringField(VarName, Value.ToString());
+			}
+			else if (Var.GetType() == FNiagaraTypeDefinition::GetColorDef() && Buffer)
+			{
+				FLinearColor Value;
+				Buffer->Read<FLinearColor>(i, Value, false);
+				JsonObject->SetStringField(VarName, Value.ToString());
+			}
+			else if (Var.GetType() == FNiagaraTypeDefinition::GetIDDef() && Buffer)
+			{
+				FNiagaraID Value;
+				Buffer->Read<FNiagaraID>(i, Value, false);
+				JsonObject->SetStringField(VarName, FString::FromInt(Value.Index) + "/" + FString::FromInt(Value.AcquireTag));
+			}
+			else if (Var.GetType().IsEnum() && Buffer)
+			{
+				int32 Value;
+				Buffer->Read<int32>(i, Value, false);
+				JsonObject->SetStringField(VarName, Var.GetType().GetEnum()->GetNameByValue(Value).ToString());
+			}
+			else
+			{
+				JsonObject->SetStringField(VarName, "???");
+			}
+		}
+		FJsonSerializer::Serialize(JsonObject.ToSharedRef(), JsonWriter);
+	}
+	JsonWriter->WriteArrayEnd();
+	JsonWriter->Close();
+	return JsonString;
+}
+
+FString FNiagaraDataChannelDebugUtilities::TickGroupToString(const ETickingGroup& TickGroup)
+{
+	static UEnum* TGEnum = StaticEnum<ETickingGroup>();
+	return TGEnum->GetDisplayNameTextByValue(TickGroup).ToString();
+}
+#endif //WITH_NIAGARA_DEBUGGER
 
 #undef LOCTEXT_NAMESPACE
 
diff --git a/Engine/Plugins/FX/Niagara/Source/Niagara/Private/NiagaraDataChannelAccessor.cpp b/Engine/Plugins/FX/Niagara/Source/Niagara/Private/NiagaraDataChannelAccessor.cpp
--- a/Engine/Plugins/FX/Niagara/Source/Niagara/Private/NiagaraDataChannelAccessor.cpp
+++ b/Engine/Plugins/FX/Niagara/Source/Niagara/Private/NiagaraDataChannelAccessor.cpp
@@ -142,7 +142,7 @@ void UNiagaraDataChannelWriter::WriteData(const FNiagaraVariableBase& Var, int32
 	}
 }
 
-bool UNiagaraDataChannelWriter::InitWrite(FNiagaraDataChannelSearchParameters SearchParams, int32 Count, bool bVisibleToGame, bool bVisibleToCPU, bool bVisibleToGPU)
+bool UNiagaraDataChannelWriter::InitWrite(FNiagaraDataChannelSearchParameters SearchParams, int32 Count, bool bVisibleToGame, bool bVisibleToCPU, bool bVisibleToGPU, const FString& DebugSource)
 {
 	if (Count == 0)
 	{
@@ -152,8 +152,7 @@ bool UNiagaraDataChannelWriter::InitWrite(FNiagaraDataChannelSearchParameters Se
 
 	check(Owner);
 
-	FNiagaraDataChannelDataPtr DestData = Owner->FindData(SearchParams, ENiagaraResourceAccess::WriteOnly);
-	if(DestData)
+	if(FNiagaraDataChannelDataPtr DestData = Owner->FindData(SearchParams, ENiagaraResourceAccess::WriteOnly))
 	{
 		//TODO- Dont create a whole new game Data here. Rather grab a pre-made staging/input One from the Data PTR With the relevant publish Flags.
 		Data = Owner->GetDataChannel()->CreateGameData();
@@ -164,6 +163,9 @@ bool UNiagaraDataChannelWriter::InitWrite(FNiagaraDataChannelSearchParameters Se
 		PublishRequest.bVisibleToCPUSims = bVisibleToCPU;
 		PublishRequest.bVisibleToGPUSims = bVisibleToGPU;
 		PublishRequest.GameData = Data;
+#if !UE_BUILD_SHIPPING
+		PublishRequest.DebugSource = DebugSource;
+#endif
 		DestData->Publish(PublishRequest);
 		return true;
 	}
diff --git a/Engine/Plugins/FX/Niagara/Source/Niagara/Private/NiagaraDataChannel_Global.cpp b/Engine/Plugins/FX/Niagara/Source/Niagara/Private/NiagaraDataChannel_Global.cpp
--- a/Engine/Plugins/FX/Niagara/Source/Niagara/Private/NiagaraDataChannel_Global.cpp
+++ b/Engine/Plugins/FX/Niagara/Source/Niagara/Private/NiagaraDataChannel_Global.cpp
@@ -66,7 +66,7 @@ void UNiagaraDataChannelHandler_Global::Tick(float DeltaSeconds, ETickingGroup T
 	SCOPE_CYCLE_COUNTER(STAT_DataChannelHandler_Global_Tick);
 	Super::Tick(DeltaSeconds, TickGroup, OwningWorld);
 
-	Data->ConsumePublishRequests(this);
+	Data->ConsumePublishRequests(this, TickGroup);
 }
 
 FNiagaraDataChannelDataPtr UNiagaraDataChannelHandler_Global::FindData(FNiagaraDataChannelSearchParameters SearchParams, ENiagaraResourceAccess AccessType)
diff --git a/Engine/Plugins/FX/Niagara/Source/Niagara/Private/NiagaraDataChannel_Islands.cpp b/Engine/Plugins/FX/Niagara/Source/Niagara/Private/NiagaraDataChannel_Islands.cpp
--- a/Engine/Plugins/FX/Niagara/Source/Niagara/Private/NiagaraDataChannel_Islands.cpp
+++ b/Engine/Plugins/FX/Niagara/Source/Niagara/Private/NiagaraDataChannel_Islands.cpp
@@ -150,7 +150,7 @@ void UNiagaraDataChannelHandler_Islands::Tick(float DeltaTime, ETickingGroup Tic
 	for (auto It = ActiveIslands.CreateIterator(); It; ++It)
 	{
 		FNDCIsland& Island = IslandPool[*It];
-		Island.Tick();
+		Island.Tick(TickGroup);
 
 		if(IslandChannel->GetDebugDrawSettings().ShowBounds())
 		{
@@ -399,9 +399,9 @@ void FNDCIsland::EndFrame()
 	Data->EndFrame(Owner);
 }
 
-void FNDCIsland::Tick()
+void FNDCIsland::Tick(const ETickingGroup& TickGroup)
 {
-	int32 AddedData = Data->ConsumePublishRequests(Owner);
+	int32 AddedData = Data->ConsumePublishRequests(Owner, TickGroup);
 	if (IsBeingUsed() && AddedData > 0)
 	{
 		for (UNiagaraComponent* Comp : NiagaraSystems)
diff --git a/Engine/Plugins/FX/Niagara/Source/Niagara/Private/NiagaraModule.cpp b/Engine/Plugins/FX/Niagara/Source/Niagara/Private/NiagaraModule.cpp
--- a/Engine/Plugins/FX/Niagara/Source/Niagara/Private/NiagaraModule.cpp
+++ b/Engine/Plugins/FX/Niagara/Source/Niagara/Private/NiagaraModule.cpp
@@ -32,6 +32,7 @@
 #include "NiagaraComponentRendererProperties.h"
 #include "NiagaraVolumeRendererProperties.h"
 #include "NiagaraCustomVersion.h"
+#include "NiagaraDataChannel.h"
 #include "NiagaraShaderModule.h"
 #include "UObject/CoreRedirects.h"
 #include "NiagaraGpuComputeDispatch.h"
@@ -614,6 +615,7 @@ void INiagaraModule::ShutdownModule()
 
 #if WITH_NIAGARA_DEBUGGER
 	DebuggerClient.Reset();
+	FNiagaraDataChannelDebugUtilities::TearDown();
 #endif
 
 #if NIAGARA_PERF_BASELINES
diff --git a/Engine/Plugins/FX/Niagara/Source/Niagara/Public/NiagaraDataChannel.h b/Engine/Plugins/FX/Niagara/Source/Niagara/Public/NiagaraDataChannel.h
--- a/Engine/Plugins/FX/Niagara/Source/Niagara/Public/NiagaraDataChannel.h
+++ b/Engine/Plugins/FX/Niagara/Source/Niagara/Public/NiagaraDataChannel.h
@@ -173,7 +173,7 @@ private:
 	/** Layout information for any data stored at the "Game" level. i.e. From game code/BP. AoS layout and LWC types. */
 	FNiagaraDataChannelGameDataLayout GameDataLayout;
 	
-	#if !UE_BUILD_SHIPPING
+	#if WITH_NIAGARA_DEBUGGER
 	mutable bool bVerboseLogging = false;
 	#endif
 };
@@ -218,9 +218,10 @@ class NIAGARA_API UNiagaraDataChannelLibrary : public UBlueprintFunctionLibrary
 	 * @param bVisibleToGame	If true, the data written to this data channel is visible to Blueprint and C++ logic reading from it
 	 * @param bVisibleToCPU	If true, the data written to this data channel is visible to Niagara CPU emitters
 	 * @param bVisibleToGPU	If true, the data written to this data channel is visible to Niagara GPU emitters
+	 * @param DebugSource	Instigator for this write, used in the debug hud to track writes to the data channel from different sources
 	 */
-	UFUNCTION(BlueprintCallable, Category = NiagaraDataChannel, DisplayName="Write To Niagara Data Channel (Batch)", meta = (AdvancedDisplay = "SearchParams", Keywords = "niagara DataChannel", WorldContext = "WorldContextObject", UnsafeDuringActorConstruction = "true"))
-	static UNiagaraDataChannelWriter* WriteToNiagaraDataChannel(const UObject* WorldContextObject, const UNiagaraDataChannelAsset* Channel, FNiagaraDataChannelSearchParameters SearchParams, int32 Count, UPARAM(DisplayName = "Visible to Blueprint") bool bVisibleToGame, UPARAM(DisplayName = "Visible to Niagara CPU") bool bVisibleToCPU, UPARAM(DisplayName = "Visible to Niagara GPU") bool bVisibleToGPU);
+	UFUNCTION(BlueprintCallable, Category = NiagaraDataChannel, DisplayName="Write To Niagara Data Channel (Batch)", meta = (AdvancedDisplay = "SearchParams, DebugSource", Keywords = "niagara DataChannel", WorldContext = "WorldContextObject", UnsafeDuringActorConstruction = "true", AutoCreateRefTerm="DebugSource"))
+	static UNiagaraDataChannelWriter* WriteToNiagaraDataChannel(const UObject* WorldContextObject, const UNiagaraDataChannelAsset* Channel, FNiagaraDataChannelSearchParameters SearchParams, int32 Count, UPARAM(DisplayName = "Visible to Blueprint") bool bVisibleToGame, UPARAM(DisplayName = "Visible to Niagara CPU") bool bVisibleToCPU, UPARAM(DisplayName = "Visible to Niagara GPU") bool bVisibleToGPU, const FString& DebugSource);
 
 	/**
 	 * Initializes and returns the Niagara Data Channel reader for the given data channel.
@@ -277,6 +278,6 @@ class NIAGARA_API UNiagaraDataChannelLibrary : public UBlueprintFunctionLibrary
 	static void WriteToNiagaraDataChannelSingle(const UObject* WorldContextObject, const UNiagaraDataChannelAsset* Channel, FNiagaraDataChannelSearchParameters SearchParams, bool bVisibleToBlueprint, bool bVisibleToNiagaraCPU, bool bVisibleToNiagaraGPU);
 
 	static UNiagaraDataChannelHandler* FindDataChannelHandler(const UObject* WorldContextObject, const UNiagaraDataChannel* Channel);
-	static UNiagaraDataChannelWriter* CreateDataChannelWriter(const UObject* WorldContextObject, const UNiagaraDataChannel* Channel, FNiagaraDataChannelSearchParameters SearchParams, int32 Count, bool bVisibleToGame, bool bVisibleToCPU, bool bVisibleToGPU);
+	static UNiagaraDataChannelWriter* CreateDataChannelWriter(const UObject* WorldContextObject, const UNiagaraDataChannel* Channel, FNiagaraDataChannelSearchParameters SearchParams, int32 Count, bool bVisibleToGame, bool bVisibleToCPU, bool bVisibleToGPU, const FString& DebugSource);
 	static UNiagaraDataChannelReader* CreateDataChannelReader(const UObject* WorldContextObject, const UNiagaraDataChannel* Channel, FNiagaraDataChannelSearchParameters SearchParams, bool bReadPreviousFrame);
 };
diff --git a/Engine/Plugins/FX/Niagara/Source/Niagara/Public/NiagaraDataChannelAccessor.h b/Engine/Plugins/FX/Niagara/Source/Niagara/Public/NiagaraDataChannelAccessor.h
--- a/Engine/Plugins/FX/Niagara/Source/Niagara/Public/NiagaraDataChannelAccessor.h
+++ b/Engine/Plugins/FX/Niagara/Source/Niagara/Public/NiagaraDataChannelAccessor.h
@@ -95,8 +95,8 @@ public:
 	TObjectPtr<UNiagaraDataChannelHandler> Owner;
 	
 	/** Call before each batch of writes to allocate the data we'll be writing to. */
-	UFUNCTION(BlueprintCallable, Category = NiagaraDataChannel, meta = (Keywords = "niagara DataChannel"))
-	NIAGARA_API bool InitWrite(FNiagaraDataChannelSearchParameters SearchParams, int32 Count, bool bVisibleToGame=true, bool bVisibleToCPU=true, bool bVisibleToGPU=true);
+	UFUNCTION(BlueprintCallable, Category = NiagaraDataChannel, meta = (bVisibleToGame="true", bVisibleToCPU="true", bVisibleToGPU="true", Keywords = "niagara DataChannel", AdvancedDisplay = "DebugSource", AutoCreateRefTerm="DebugSource"))
+	NIAGARA_API bool InitWrite(FNiagaraDataChannelSearchParameters SearchParams, int32 Count, bool bVisibleToGame, bool bVisibleToCPU, bool bVisibleToGPU, const FString& DebugSource);
 
 	UFUNCTION(BlueprintCallable, Category = NiagaraDataChannel, meta = (Keywords = "niagara DataChannel"))
 	NIAGARA_API int32 Num()const;
diff --git a/Engine/Plugins/FX/Niagara/Source/Niagara/Public/NiagaraDataChannelPublic.h b/Engine/Plugins/FX/Niagara/Source/Niagara/Public/NiagaraDataChannelPublic.h
--- a/Engine/Plugins/FX/Niagara/Source/Niagara/Public/NiagaraDataChannelPublic.h
+++ b/Engine/Plugins/FX/Niagara/Source/Niagara/Public/NiagaraDataChannelPublic.h
@@ -6,6 +6,7 @@
 #include "NiagaraCommon.h"
 #include "NiagaraDataChannelPublic.generated.h"
 
+struct FNiagaraDataChannelPublishRequest;
 class FNiagaraWorldManager;
 class UNiagaraDataChannelAsset;
 class UNiagaraDataChannel;
@@ -150,7 +151,7 @@ struct FNiagaraDataChannelGameDataLayout
 };
 
 
-#if !UE_BUILD_SHIPPING
+#if WITH_NIAGARA_DEBUGGER
 
 /** Hooks into internal NiagaraDataChannels code for debugging and testing purposes. */
 class FNiagaraDataChannelDebugUtilities
@@ -162,10 +163,36 @@ public:
 	static NIAGARA_API void Tick(FNiagaraWorldManager* WorldMan, float DeltaSeconds, ETickingGroup TickGroup);
 	
 	static NIAGARA_API UNiagaraDataChannelHandler* FindDataChannelHandler(FNiagaraWorldManager* WorldMan, UNiagaraDataChannel* DataChannel);
+
+	static void LogWrite(const FNiagaraDataChannelPublishRequest& WriteRequest, const UNiagaraDataChannel* DataChannel, const ETickingGroup& TickGroup);
+	static void DumpAllWritesToLog();
+
+	static FNiagaraDataChannelDebugUtilities& Get();
+	static void TearDown();
+
+private:
+	struct FChannelWriteRequest
+	{
+		TSharedPtr<FNiagaraDataChannelGameData> Data;
+		bool bVisibleToGame = false;
+		bool bVisibleToCPUSims = false;
+		bool bVisibleToGPUSims = false;
+		ETickingGroup TickGroup;
+		FString DebugSource;
+	};
+	struct FFrameDebugData
+	{
+		uint64 FrameNumber;
+		TArray<FChannelWriteRequest> WriteRequests;
+	};
+	
+	static FString ToJson(FNiagaraDataChannelGameData* Data);
+	static FString TickGroupToString(const ETickingGroup& TickGroup);
+	TArray<FFrameDebugData> FrameData;
 };
 
 
-#endif//UE_BUILD_SHIPPING
+#endif//WITH_NIAGARA_DEBUGGER
 
 
 /** Buffer containing a single FNiagaraVariable's data at the game level. AoS layout. LWC Types. */
diff --git a/Engine/Plugins/FX/Niagara/Source/Niagara/Public/NiagaraDataChannel_Islands.h b/Engine/Plugins/FX/Niagara/Source/Niagara/Public/NiagaraDataChannel_Islands.h
--- a/Engine/Plugins/FX/Niagara/Source/Niagara/Public/NiagaraDataChannel_Islands.h
+++ b/Engine/Plugins/FX/Niagara/Source/Niagara/Public/NiagaraDataChannel_Islands.h
@@ -37,7 +37,7 @@ public:
 	void Init(UNiagaraDataChannelHandler_Islands* Owner);
 	void BeginFrame();
 	void EndFrame();
-	void Tick();
+	void Tick(const ETickingGroup& TickGroup);
 
 	bool Contains(FVector Point);
 	double DistanceToPoint(FVector Point);
diff --git a/Engine/Plugins/FX/Niagara/Source/NiagaraBlueprintNodes/Private/K2Node_WriteDataChannel.cpp b/Engine/Plugins/FX/Niagara/Source/NiagaraBlueprintNodes/Private/K2Node_WriteDataChannel.cpp
--- a/Engine/Plugins/FX/Niagara/Source/NiagaraBlueprintNodes/Private/K2Node_WriteDataChannel.cpp
+++ b/Engine/Plugins/FX/Niagara/Source/NiagaraBlueprintNodes/Private/K2Node_WriteDataChannel.cpp
@@ -5,10 +5,12 @@
 #include "BlueprintActionDatabaseRegistrar.h"
 #include "BlueprintNodeSpawner.h"
 #include "K2Node_ExecutionSequence.h"
+#include "K2Node_Self.h"
 #include "KismetCompiler.h"
 #include "NiagaraBlueprintUtil.h"
 #include "NiagaraDataChannel.h"
 #include "NiagaraDataChannelAccessor.h"
+#include "Kismet/KismetSystemLibrary.h"
 
 #define LOCTEXT_NAMESPACE "K2Node_WriteDataChannel"
 
@@ -71,6 +73,18 @@ void UK2Node_WriteDataChannel::ExpandNode(FKismetCompilerContext& CompilerContex
 	CreateWriterNode->SetFromFunction(UNiagaraDataChannelLibrary::StaticClass()->FindFunctionByName(GET_FUNCTION_NAME_CHECKED(UNiagaraDataChannelLibrary, WriteToNiagaraDataChannel)));
 	CreateWriterNode->AllocateDefaultPins();
 
+#if WITH_NIAGARA_DEBUGGER
+	// add path info for current BP for debug purposes 
+	UK2Node_CallFunction* GetPathNode = CompilerContext.SpawnIntermediateNode<UK2Node_CallFunction>(this, SourceGraph);
+	GetPathNode->SetFromFunction(UKismetSystemLibrary::StaticClass()->FindFunctionByName(GET_FUNCTION_NAME_CHECKED(UKismetSystemLibrary, GetPathName)));
+	GetPathNode->AllocateDefaultPins();
+	UK2Node_Self* SelfNode = CompilerContext.SpawnIntermediateNode<UK2Node_Self>(this, SourceGraph);
+	SelfNode->AllocateDefaultPins();
+
+	CompilerContext.GetSchema()->TryCreateConnection(SelfNode->FindPinChecked(UEdGraphSchema_K2::PN_Self), GetPathNode->FindPinChecked(FName("Object")));
+	CompilerContext.GetSchema()->TryCreateConnection(CreateWriterNode->FindPinChecked(FName("DebugSource")), GetPathNode->GetReturnValuePin());
+#endif
+
 	// transfer the input pins over
 	static TArray<TPair<FName, FName>> PinsToTransfer = { {"Channel", "Channel"}, {"SearchParams", "SearchParams"}, {"bVisibleToBlueprint", "bVisibleToGame"},
 		{"bVisibleToNiagaraCPU", "bVisibleToCPU"}, {"bVisibleToNiagaraGPU", "bVisibleToGPU"}};
@@ -83,13 +97,13 @@ void UK2Node_WriteDataChannel::ExpandNode(FKismetCompilerContext& CompilerContex
 	}
 	CreateWriterNode->FindPinChecked(FName("Count"))->DefaultValue = FString::FromInt(1);
 	
-	UEdGraphPin* OldExecPin = FindPinChecked(UEdGraphSchema_K2::PN_Execute, EGPD_Input);
-	UEdGraphPin* NewExecPin = CreateWriterNode->FindPinChecked(UEdGraphSchema_K2::PN_Execute, EGPD_Input);
+	UEdGraphPin* OldExecPin = GetExecPin();
+	UEdGraphPin* NewExecPin = CreateWriterNode->GetExecPin();
 	CompilerContext.MovePinLinksToIntermediate(*OldExecPin, *NewExecPin);
 
 	// create the write function nodes
-	UEdGraphPin* LastExecPin = CreateWriterNode->FindPinChecked(UEdGraphSchema_K2::PN_Then, EGPD_Output);
-	UEdGraphPin* WriterResultPin = CreateWriterNode->FindPinChecked(UEdGraphSchema_K2::PN_ReturnValue, EGPD_Output);
+	UEdGraphPin* LastExecPin = CreateWriterNode->GetThenPin();
+	UEdGraphPin* WriterResultPin = CreateWriterNode->GetReturnValuePin();
 	for (const FNiagaraDataChannelVariable& InVar : DataChannel->Get()->GetVariables())
 	{
 		if (IgnoredVariables.Contains(InVar.Version))
@@ -124,14 +138,12 @@ void UK2Node_WriteDataChannel::ExpandNode(FKismetCompilerContext& CompilerContex
 		CompilerContext.MovePinLinksToIntermediate(*VarInputPin, *WriteDataNode->FindPinChecked(FName("InData"), EGPD_Input));
 
 		// connect exec pins
-		CompilerContext.GetSchema()->TryCreateConnection(LastExecPin, WriteDataNode->FindPinChecked(UEdGraphSchema_K2::PN_Execute, EGPD_Input));
-		LastExecPin = WriteDataNode->FindPinChecked(UEdGraphSchema_K2::PN_Then, EGPD_Output);
+		CompilerContext.GetSchema()->TryCreateConnection(LastExecPin, WriteDataNode->GetExecPin());
+		LastExecPin = WriteDataNode->GetThenPin();
 	}
 
-
 	// connect the last exec pin
-	UEdGraphPin* OldThenPin = FindPinChecked(UEdGraphSchema_K2::PN_Then, EGPD_Output);
-	CompilerContext.MovePinLinksToIntermediate(*OldThenPin, *LastExecPin);
+	CompilerContext.MovePinLinksToIntermediate(*GetThenPin(), *LastExecPin);
 }
 
 UFunction* UK2Node_WriteDataChannel::GetWriteFunctionForType(const FNiagaraTypeDefinition& TypeDef)
