Fixed issues with writes from NDC to game data.

diff --git a/Engine/Plugins/FX/Niagara/Source/Niagara/Private/NiagaraDataChannel.cpp b/Engine/Plugins/FX/Niagara/Source/Niagara/Private/NiagaraDataChannel.cpp
--- a/Engine/Plugins/FX/Niagara/Source/Niagara/Private/NiagaraDataChannel.cpp
+++ b/Engine/Plugins/FX/Niagara/Source/Niagara/Private/NiagaraDataChannel.cpp
@@ -310,6 +310,7 @@ void FNiagaraDataChannelGameData::AppendFromDataSet(const FNiagaraDataBuffer* Sr
 	static TArray<uint8> LWCConversionBuffer;
 
 	int32 NumInstances = SrcBuffer->GetNumInstances();
+	int32 StartIndex = NumElements;
 	NumElements += NumInstances;
 
 	const FNiagaraDataChannelGameDataLayout& Layout = DataChannel->GetGameDataLayout();
@@ -320,8 +321,6 @@ void FNiagaraDataChannelGameData::AppendFromDataSet(const FNiagaraDataBuffer* Sr
 		FNiagaraDataChannelVariableBuffer& VarBuffer = VariableData[VarIndex];
 
 		VarBuffer.SetNum(NumElements);
-
-		uint8* DestDataBase = VarBuffer.Data.GetData();
 		
 		FNiagaraVariableBase SimVar = Var;
 		if (FNiagaraTypeHelper::IsLWCType(Var.GetType()))
@@ -335,6 +334,8 @@ void FNiagaraDataChannelGameData::AppendFromDataSet(const FNiagaraDataBuffer* Sr
 			Var.SetType(FNiagaraTypeHelper::GetVectorDef());
 		}
 		int32 VarSize = Var.GetSizeInBytes();
+
+		uint8* DestDataBase = VarBuffer.Data.GetData() + StartIndex * VarSize;
 		
 		int32 SimVarIndex = CompiledData.Variables.IndexOfByKey(SimVar);
 		if (SimVarIndex == INDEX_NONE)
@@ -349,7 +350,7 @@ void FNiagaraDataChannelGameData::AppendFromDataSet(const FNiagaraDataBuffer* Sr
 		int32 HalfCompIdx = SimLayout.GetHalfComponentStart();
 		
 		TFunction<void(UScriptStruct*, UScriptStruct*, uint8*)> ReadData;
-		ReadData = [&](UScriptStruct* SrcStruct, UScriptStruct* DestStruct, uint8* DestData)
+		ReadData = [&](UScriptStruct* SrcStruct, UScriptStruct* DestStruct, uint8* DestDataBase)
 		{
 			//Read all data from the simulation and place in the game level buffers. Converting from LWC local tile space where needed.
 
@@ -362,7 +363,7 @@ void FNiagaraDataChannelGameData::AppendFromDataSet(const FNiagaraDataBuffer* Sr
 
 				for (int32 i = 0; i < NumInstances; ++i)
 				{
-					FVector* Dest = reinterpret_cast<FVector*>(DestData + VarSize * i);
+					FVector* Dest = reinterpret_cast<FVector*>(DestDataBase + VarSize * i);
 					*Dest = FVector(*SrcX++, *SrcY++, *SrcZ++) + FVector(SimulationLwcTile) * FLargeWorldRenderScalar::GetTileSize();
 				}
 			}
@@ -374,7 +375,8 @@ void FNiagaraDataChannelGameData::AppendFromDataSet(const FNiagaraDataBuffer* Sr
 				{
 					FProperty* SrcProperty = *SrcPropertyIt;
 					FProperty* DestProperty = *DestPropertyIt;
-					DestData += DestProperty->GetOffset_ForInternal();
+					int32 DestOffset = DestProperty->GetOffset_ForInternal();
+					uint8* DestData = DestDataBase + DestOffset;
 					if (DestPropertyIt->IsA(FDoubleProperty::StaticClass()))
 					{
 						double* Dest = reinterpret_cast<double*>(DestData);
