Added sim cache support for data channel writes (the writer data interface calls are persisted in the sim cache and played back when a frame is read).

diff --git a/Engine/Plugins/FX/Niagara/Source/Niagara/Internal/DataInterface/NiagaraDataInterfaceDataChannelCommon.h b/Engine/Plugins/FX/Niagara/Source/Niagara/Internal/DataInterface/NiagaraDataInterfaceDataChannelCommon.h
--- a/Engine/Plugins/FX/Niagara/Source/Niagara/Internal/DataInterface/NiagaraDataInterfaceDataChannelCommon.h
+++ b/Engine/Plugins/FX/Niagara/Source/Niagara/Internal/DataInterface/NiagaraDataInterfaceDataChannelCommon.h
@@ -4,8 +4,6 @@
 #include "NiagaraCommon.h"
 #include "NiagaraDataSetCompiledData.h"
 #include "NiagaraCompileHash.h"
-#include "Misc/LazySingleton.h"
-#include "NiagaraDataChannelCommon.h"
 #include "NiagaraDataInterface.h"
 #include "NiagaraDataInterfaceDataChannelCommon.generated.h"
 
diff --git a/Engine/Plugins/FX/Niagara/Source/Niagara/Internal/DataInterface/NiagaraDataInterfaceDataChannelRead.h b/Engine/Plugins/FX/Niagara/Source/Niagara/Internal/DataInterface/NiagaraDataInterfaceDataChannelRead.h
--- a/Engine/Plugins/FX/Niagara/Source/Niagara/Internal/DataInterface/NiagaraDataInterfaceDataChannelRead.h
+++ b/Engine/Plugins/FX/Niagara/Source/Niagara/Internal/DataInterface/NiagaraDataInterfaceDataChannelRead.h
@@ -21,10 +21,11 @@ There is working code for this already but the concept/API needs more fleshing o
 
 */
 
+#include "Containers/DynamicRHIResourceArray.h"
 
+#include "NiagaraDataChannelPublic.h"
 #include "NiagaraDataInterfaceDataChannelCommon.h"
 #include "NiagaraDataInterfaceRW.h"
-#include "NiagaraDataInterfaceEmitterBinding.h"
 #include "NiagaraDataSetAccessor.h"
 #include "NiagaraDataInterfaceDataChannelRead.generated.h"
 
diff --git a/Engine/Plugins/FX/Niagara/Source/Niagara/Internal/DataInterface/NiagaraDataInterfaceDataChannelWrite.h b/Engine/Plugins/FX/Niagara/Source/Niagara/Internal/DataInterface/NiagaraDataInterfaceDataChannelWrite.h
--- a/Engine/Plugins/FX/Niagara/Source/Niagara/Internal/DataInterface/NiagaraDataInterfaceDataChannelWrite.h
+++ b/Engine/Plugins/FX/Niagara/Source/Niagara/Internal/DataInterface/NiagaraDataInterfaceDataChannelWrite.h
@@ -18,7 +18,9 @@ This avoids cumbersome work in the graph to access data but requires special han
 
 
 #include "NiagaraDataInterfaceDataChannelCommon.h"
+#include "NiagaraDataChannelCommon.h"
 #include "NiagaraDataInterface.h"
+#include "NiagaraSimCacheCustomStorageInterface.h"
 #include "NiagaraDataInterfaceDataChannelWrite.generated.h"
 
 /** Additional compile time information used by the Write DI. */
@@ -34,8 +36,56 @@ struct FNDIDataChannelWriteCompiledData : public FNDIDataChannelCompiledData
 	bool Init(UNiagaraSystem* System, UNiagaraDataInterfaceDataChannelWrite* OwnerDI);
 };
 
+USTRUCT()
+struct FNDIDataChannelWriteSimCacheFrameBuffer
+{
+	GENERATED_BODY()
+
+	UPROPERTY()
+	TArray<uint8> Data;
+
+	UPROPERTY()
+	int32 Size = 0;
+
+	UPROPERTY()
+	FNiagaraVariableBase SourceVar;
+};
+
+USTRUCT()
+struct FNDIDataChannelWriteSimCacheFrame
+{
+	GENERATED_BODY()
+
+	UPROPERTY()
+	int32 NumElements = 0;
+
+	UPROPERTY()
+	TArray<FNDIDataChannelWriteSimCacheFrameBuffer> VariableData;
+
+	UPROPERTY()
+	bool bVisibleToGame = false;
+
+	UPROPERTY()
+	bool bVisibleToCPUSims = false;
+
+	UPROPERTY()
+	bool bVisibleToGPUSims = false;
+};
+
+UCLASS()
+class UNDIDataChannelWriteSimCacheData : public UObject
+{
+	GENERATED_BODY()
+
+public:
+
+	UPROPERTY()
+	TArray<FNDIDataChannelWriteSimCacheFrame> FrameData;
+};
+
+
 UCLASS(Experimental, EditInlineNew, Category = "Data Channels", CollapseCategories, meta = (DisplayName = "Data Channel Writer"), MinimalAPI)
-class UNiagaraDataInterfaceDataChannelWrite : public UNiagaraDataInterface
+class UNiagaraDataInterfaceDataChannelWrite : public UNiagaraDataInterface, public INiagaraSimCacheCustomStorageInterface
 {
 	GENERATED_UCLASS_BODY()
 
@@ -129,6 +179,11 @@ public:
 
 	bool ShouldPublish()const { return bPublishToGame || bPublishToCPU || bPublishToGPU; }
 
+	//sim cache functions
+	virtual UObject* SimCacheBeginWrite(UObject* SimCache, FNiagaraSystemInstance* NiagaraSystemInstance, const void* OptionalPerInstanceData, FNiagaraSimCacheFeedbackContext& FeedbackContext) const override;
+	virtual bool SimCacheWriteFrame(UObject* StorageObject, int FrameIndex, FNiagaraSystemInstance* SystemInstance, const void* OptionalPerInstanceData, FNiagaraSimCacheFeedbackContext& FeedbackContext) const override;
+	virtual bool SimCacheReadFrame(UObject* StorageObject, int FrameA, int FrameB, float Interp, FNiagaraSystemInstance* SystemInstance, void* OptionalPerInstanceData) override;
+	virtual void SimCachePostReadFrame(void* OptionalPerInstanceData, FNiagaraSystemInstance* SystemInstance) override;
 protected:
 	NIAGARA_API virtual bool CopyToInternal(UNiagaraDataInterface* Destination) const override;
 
diff --git a/Engine/Plugins/FX/Niagara/Source/Niagara/Private/DataInterface/NiagaraDataInterfaceDataChannelRead.cpp b/Engine/Plugins/FX/Niagara/Source/Niagara/Private/DataInterface/NiagaraDataInterfaceDataChannelRead.cpp
--- a/Engine/Plugins/FX/Niagara/Source/Niagara/Private/DataInterface/NiagaraDataInterfaceDataChannelRead.cpp
+++ b/Engine/Plugins/FX/Niagara/Source/Niagara/Private/DataInterface/NiagaraDataInterfaceDataChannelRead.cpp
@@ -12,6 +12,7 @@
 #include "NiagaraWorldManager.h"
 
 #include "NiagaraDataChannel.h"
+#include "NiagaraDataChannelCommon.h"
 #include "NiagaraDataChannelHandler.h"
 #include "NiagaraDataChannelManager.h"
 
diff --git a/Engine/Plugins/FX/Niagara/Source/Niagara/Private/DataInterface/NiagaraDataInterfaceDataChannelWrite.cpp b/Engine/Plugins/FX/Niagara/Source/Niagara/Private/DataInterface/NiagaraDataInterfaceDataChannelWrite.cpp
--- a/Engine/Plugins/FX/Niagara/Source/Niagara/Private/DataInterface/NiagaraDataInterfaceDataChannelWrite.cpp
+++ b/Engine/Plugins/FX/Niagara/Source/Niagara/Private/DataInterface/NiagaraDataInterfaceDataChannelWrite.cpp
@@ -19,7 +19,6 @@
 #if WITH_EDITOR
 #include "INiagaraEditorOnlyDataUtlities.h"
 #include "Modules/ModuleManager.h"
-#include "NiagaraModule.h"
 #endif
 
 #define LOCTEXT_NAMESPACE "NiagaraDataInterfaceDataChannelWrite"
@@ -286,9 +285,8 @@ struct FNDIDataChannelWriteInstanceData
 				PublishRequest.bVisibleToGame = Interface->bPublishToGame;
 				PublishRequest.bVisibleToCPUSims = Interface->bPublishToCPU;
 				PublishRequest.bVisibleToGPUSims = Interface->bPublishToGPU;
-				PublishRequest.Data = Data->GetCurrentData();
 				PublishRequest.LwcTile = Instance->GetLWCTile();
-#if !UE_BUILD_SHIPPING
+#if WITH_NIAGARA_DEBUGGER
 				PublishRequest.DebugSource = FString::Format(TEXT("{0} ({1})"), {Instance->GetSystem()->GetName(), GetPathNameSafe(Interface)});
 #endif
 				DataChannelData->Publish(PublishRequest);
@@ -585,6 +583,113 @@ bool UNiagaraDataInterfaceDataChannelWrite::Equals(const UNiagaraDataInterface*
 	return false;
 }
 
+UObject* UNiagaraDataInterfaceDataChannelWrite::SimCacheBeginWrite(UObject* SimCache, FNiagaraSystemInstance* NiagaraSystemInstance, const void* OptionalPerInstanceData, FNiagaraSimCacheFeedbackContext& FeedbackContext) const
+{
+	return NewObject<UNDIDataChannelWriteSimCacheData>(SimCache);
+}
+
+bool UNiagaraDataInterfaceDataChannelWrite::SimCacheWriteFrame(UObject* StorageObject, int FrameIndex, FNiagaraSystemInstance* SystemInstance, const void* OptionalPerInstanceData, FNiagaraSimCacheFeedbackContext& FeedbackContext) const
+{
+	if (!Channel || !Channel->Get())
+	{
+		FeedbackContext.Errors.Add(TEXT("Missing data channel asset for data channel writer DI"));
+		return false;
+	}
+	if (OptionalPerInstanceData == nullptr)
+	{
+		FeedbackContext.Errors.Add(TEXT("Missing per instance data for data channel writer DI"));
+		return false;
+	}
+	// put data from instance data into the storage object
+	const FNDIDataChannelWriteInstanceData* InstanceData = static_cast<const FNDIDataChannelWriteInstanceData*>(OptionalPerInstanceData);
+	if (UNDIDataChannelWriteSimCacheData* Storage = Cast<UNDIDataChannelWriteSimCacheData>(StorageObject))
+	{
+		ensure(Storage->FrameData.Num() == FrameIndex);
+		FNDIDataChannelWriteSimCacheFrame& FrameData = Storage->FrameData.AddDefaulted_GetRef();
+		
+		if (InstanceData->DataChannelData && ShouldPublish() && InstanceData->Data && InstanceData->Data->GetCurrentData() && InstanceData->Data->GetCurrentData()->GetNumInstances() > 0)
+		{
+			FNiagaraDataChannelGameData GameData;
+			GameData.Init(Channel->Get());
+			GameData.AppendFromDataSet(InstanceData->Data->GetCurrentData(), SystemInstance->GetLWCTile());
+			
+			FrameData.NumElements = GameData.Num();
+			for (const FNiagaraDataChannelVariableBuffer& VarBuffer : GameData.GetVariableBuffers())
+			{
+				FNDIDataChannelWriteSimCacheFrameBuffer& FrameBuffer = FrameData.VariableData.AddDefaulted_GetRef();
+				FrameBuffer.Size = VarBuffer.Size;
+				FrameBuffer.Data = VarBuffer.Data;
+			}
+			const FNiagaraDataChannelGameDataLayout& Layout = Channel->Get()->GetGameDataLayout();
+			for (const TPair<FNiagaraVariableBase, int32>& VarPair : Layout.VariableIndices)
+			{
+				FrameData.VariableData[VarPair.Value].SourceVar = VarPair.Key;
+			}
+			
+			FrameData.bVisibleToGame = bPublishToGame;
+			FrameData.bVisibleToCPUSims = bPublishToCPU;
+			FrameData.bVisibleToGPUSims = bPublishToGPU;
+		}
+		return true;
+	} 
+	return false;
+}
+
+bool UNiagaraDataInterfaceDataChannelWrite::SimCacheReadFrame(UObject* StorageObject, int FrameA, int FrameB, float Interp, FNiagaraSystemInstance* SystemInstance, void* OptionalPerInstanceData)
+{
+	if (Channel == nullptr || Channel->Get() == nullptr)
+	{
+		return false;
+	}
+
+	FNiagaraDataChannelDataPtr DataChannelData;
+	if (FNiagaraWorldManager* WorldMan = FNiagaraWorldManager::Get(SystemInstance->GetWorld()))
+	{
+		if (UNiagaraDataChannelHandler* Handler = WorldMan->GetDataChannelManager().FindDataChannelHandler(Channel->Get()))
+		{
+			FNiagaraDataChannelSearchParameters SearchParams(SystemInstance->GetAttachComponent());
+			DataChannelData = Handler->FindData(SearchParams, ENiagaraResourceAccess::WriteOnly);
+		}
+	}
+	if (!DataChannelData.IsValid())
+	{
+		return false;
+	}
+	
+	if (UNDIDataChannelWriteSimCacheData* Storage = Cast<UNDIDataChannelWriteSimCacheData>(StorageObject))
+	{
+		if (Storage->FrameData.IsValidIndex(FrameA))
+		{
+			FNDIDataChannelWriteSimCacheFrame& Frame = Storage->FrameData[FrameA];
+			FNiagaraDataChannelPublishRequest PublishRequest;
+			PublishRequest.bVisibleToGame = Frame.bVisibleToGame;
+			PublishRequest.bVisibleToCPUSims = Frame.bVisibleToCPUSims;
+			PublishRequest.bVisibleToGPUSims = Frame.bVisibleToGPUSims;
+#if WITH_NIAGARA_DEBUGGER
+			PublishRequest.DebugSource = FString::Format(TEXT("{0} (Sim cache {1})"), {SystemInstance->GetSystem()->GetName(), GetPathNameSafe(StorageObject->GetOuter())});
+#endif
+
+			PublishRequest.GameData = MakeShared<FNiagaraDataChannelGameData>();
+			PublishRequest.GameData->Init(Channel->Get());
+			PublishRequest.GameData->SetNum(Frame.NumElements);
+			for (int32 i = 0; i < Frame.VariableData.Num(); i++)
+			{
+				const FNDIDataChannelWriteSimCacheFrameBuffer& Buffer = Frame.VariableData[i];
+				PublishRequest.GameData->SetFromSimCache(Buffer.SourceVar, Buffer.Data, Buffer.Size);
+			}
+			
+			DataChannelData->Publish(PublishRequest);
+			return true;
+		}
+	}
+	return false;
+}
+
+void UNiagaraDataInterfaceDataChannelWrite::SimCachePostReadFrame(void* OptionalPerInstanceData, FNiagaraSystemInstance* SystemInstance)
+{
+	// send data to data channel
+}
+
 bool UNiagaraDataInterfaceDataChannelWrite::CopyToInternal(UNiagaraDataInterface* Destination)const
 {
 	if (!Super::CopyToInternal(Destination))
diff --git a/Engine/Plugins/FX/Niagara/Source/Niagara/Private/NiagaraDataChannel.cpp b/Engine/Plugins/FX/Niagara/Source/Niagara/Private/NiagaraDataChannel.cpp
--- a/Engine/Plugins/FX/Niagara/Source/Niagara/Private/NiagaraDataChannel.cpp
+++ b/Engine/Plugins/FX/Niagara/Source/Niagara/Private/NiagaraDataChannel.cpp
@@ -441,6 +441,18 @@ void FNiagaraDataChannelGameData::AppendFromDataSet(const FNiagaraDataBuffer* Sr
 	}
 }
 
+void FNiagaraDataChannelGameData::SetFromSimCache(const FNiagaraVariableBase& SourceVar, TConstArrayView<uint8> Data, int32 Size)
+{
+	const FNiagaraDataChannelGameDataLayout& Layout = DataChannel->GetGameDataLayout();
+	if (const int* Index = Layout.VariableIndices.Find(SourceVar))
+	{
+		if (VariableData[*Index].Size == Size)
+		{
+			VariableData[*Index].Data = Data;
+		}
+	}
+}
+
 //////////////////////////////////////////////////////////////////////////
 
 void FNiagaraDataChannelDataProxy::BeginFrame(bool bKeepPreviousFrameData)
diff --git a/Engine/Plugins/FX/Niagara/Source/Niagara/Public/NiagaraDataChannelPublic.h b/Engine/Plugins/FX/Niagara/Source/Niagara/Public/NiagaraDataChannelPublic.h
--- a/Engine/Plugins/FX/Niagara/Source/Niagara/Public/NiagaraDataChannelPublic.h
+++ b/Engine/Plugins/FX/Niagara/Source/Niagara/Public/NiagaraDataChannelPublic.h
@@ -319,4 +319,5 @@ public:
 	const UNiagaraDataChannel* GetDataChannel()const { return DataChannel.Get(); }
 
 	const TConstArrayView<FNiagaraDataChannelVariableBuffer> GetVariableBuffers()const { return VariableData; }
+	void SetFromSimCache(const FNiagaraVariableBase& SourceVar, TConstArrayView<uint8> Data, int32 Size);
 };
