- Filter variable bindings for simulation stages

diff --git a/Engine/Plugins/FX/Niagara/Source/NiagaraEditor/Private/Customizations/NiagaraTypeCustomizations.cpp b/Engine/Plugins/FX/Niagara/Source/NiagaraEditor/Private/Customizations/NiagaraTypeCustomizations.cpp
--- a/Engine/Plugins/FX/Niagara/Source/NiagaraEditor/Private/Customizations/NiagaraTypeCustomizations.cpp
+++ b/Engine/Plugins/FX/Niagara/Source/NiagaraEditor/Private/Customizations/NiagaraTypeCustomizations.cpp
@@ -342,10 +342,23 @@ TArray<FName> FNiagaraVariableAttributeBindingCustomization::GetNames(const FVer
 		else if ( SimulationStage )
 		{
 			// Unless we have an explicit "particles." binding we are not in the particle namespace
-			const bool IsParticleNamespace = !TargetVariableBinding->GetName().IsNone() && TargetVariableBinding->IsParticleBinding();
-			if ( IsParticleNamespace == Var.IsInNameSpace(FNiagaraConstants::ParticleAttributeNamespace) )
+			if (!TargetVariableBinding->GetName().IsNone() && TargetVariableBinding->IsParticleBinding())
 			{
-				Names.AddUnique(Var.GetName());
+				const bool bParticleAttribute = Var.IsInNameSpace(FNiagaraConstants::ParticleAttributeNamespaceString);
+				if (bParticleAttribute)
+				{
+					Names.AddUnique(Var.GetName());
+				}
+			}
+			else
+			{
+				const bool bUserAttribute = Var.IsInNameSpace(FNiagaraConstants::UserNamespaceString);
+				const bool bSystemAttribute = Var.IsInNameSpace(FNiagaraConstants::SystemNamespaceString);
+				const bool bEmitterAttribute = Var.IsInNameSpace(FNiagaraConstants::EmitterNamespaceString);
+				if (bUserAttribute || bSystemAttribute || bEmitterAttribute)
+				{
+					Names.AddUnique(Var.GetName());
+				}
 			}
 		}
 	}
