Niagara - Add a hack to parameter collection loading to prevent load file warnings when running on the server.

#localization none
#tests none

diff --git a/Engine/Plugins/FX/Niagara/Source/NiagaraEditor/Private/ViewModels/Stack/NiagaraStackFunctionInput.cpp b/Engine/Plugins/FX/Niagara/Source/NiagaraEditor/Private/ViewModels/Stack/NiagaraStackFunctionInput.cpp
--- a/Engine/Plugins/FX/Niagara/Source/NiagaraEditor/Private/ViewModels/Stack/NiagaraStackFunctionInput.cpp
+++ b/Engine/Plugins/FX/Niagara/Source/NiagaraEditor/Private/ViewModels/Stack/NiagaraStackFunctionInput.cpp
@@ -1734,9 +1734,16 @@ void UNiagaraStackFunctionInput::GetAvailableParameterHandles(TArray<FNiagaraPar
 	TArray<FAssetData> CollectionAssets;
 	AssetRegistryModule.Get().GetAssetsByClass(UNiagaraParameterCollection::StaticClass()->GetClassPathName(), CollectionAssets);
 
+	// This is a temporary HACK in order to prevent warnings when loading this content and the parameter collection is missing.
+	auto LoadQuiet = [](FAssetData& AssetData)
+	{
+		uint32 LoadFlags = LOAD_Quiet | LOAD_NoWarn;
+		return StaticLoadObject(UNiagaraParameterCollection::StaticClass(), nullptr, *AssetData.GetObjectPathString(), nullptr, LoadFlags, nullptr, true);
+	};
+
 	for (FAssetData& CollectionAsset : CollectionAssets)
 	{
-		if ( UNiagaraParameterCollection* Collection = Cast<UNiagaraParameterCollection>(CollectionAsset.GetSoftObjectPath().TryLoad()) )
+		if (UNiagaraParameterCollection* Collection = Cast<UNiagaraParameterCollection>(LoadQuiet(CollectionAsset)))
 		{
 			for (const FNiagaraVariable& CollectionParam : Collection->GetParameters())
 			{
