- Enable high pri cvar for some projects
- Clean up the high pri cvar related code

#tests PIE / Cooked

diff --git a/Engine/Plugins/FX/Niagara/Source/Niagara/Private/NiagaraWorldManager.cpp b/Engine/Plugins/FX/Niagara/Source/Niagara/Private/NiagaraWorldManager.cpp
--- a/Engine/Plugins/FX/Niagara/Source/Niagara/Private/NiagaraWorldManager.cpp
+++ b/Engine/Plugins/FX/Niagara/Source/Niagara/Private/NiagaraWorldManager.cpp
@@ -65,38 +65,6 @@ static FAutoConsoleVariableRef CVarNiagaraAllowAsyncWorkToEndOfFrame(
 	ECVF_Default
 ); 
 
-static void NiagaraApplyTickGroupPrioritiesBasedOnCVar(IConsoleVariable*);
-
-//EXPERIMENTAL
-//When fx.Niagara.AllowAsyncWorkToEndOfFrame is enabled, earlier tick functions don't necessarily need high priority since they plenty of time for their async work to finish.
-//Disabling high priority on these tick functions can minimise GameThread waits by allowing other high priority tick functions to start async work as soon as possible.
-static int32 GNiagaraTickGroupToStartHighPriorityTickFunctions = -1;
-static FAutoConsoleVariableRef CVarNiagaraTickGroupToStartHighPriorityTickFunctions_Experimental(
-	TEXT("fx.Niagara.Experimental.TickGroupToStartHighPriorityTickFunctions"),
-	GNiagaraTickGroupToStartHighPriorityTickFunctions,
-	TEXT("All tick functions below this value will be set to normal priority, and any equal or above will be set to high priority. Set to -1 to have all high priority."),
-	FConsoleVariableDelegate::CreateStatic(&NiagaraApplyTickGroupPrioritiesBasedOnCVar),
-	ECVF_Default
-); 
-
-static void NiagaraApplyTickGroupPrioritiesBasedOnCVar(IConsoleVariable*)
-{
-	check (GEngine);
-	for (const FWorldContext& WorldContext : GEngine->GetWorldContexts())
-	{
-		if (const UWorld* World = WorldContext.World())
-		{
-			if (FNiagaraWorldManager* NiagaraWorldManager = FNiagaraWorldManager::Get(World))
-			{
-				for (int32 NiagaraTickGroupIndex = 0; NiagaraTickGroupIndex < NiagaraNumTickGroups; ++NiagaraTickGroupIndex)
-				{
-					NiagaraWorldManager->SetTickGroupPriority(NiagaraTickGroupIndex, NiagaraTickGroupIndex >= GNiagaraTickGroupToStartHighPriorityTickFunctions);
-				}
-			}
-		}
-	}
-}
-
 static int GNiagaraWaitOnPreGC = 1;
 static FAutoConsoleVariableRef CVarNiagaraWaitOnPreGC(
 	TEXT("fx.Niagara.WaitOnPreGC"),
@@ -300,6 +268,7 @@ TMap<class UWorld*, class FNiagaraWorldManager*> FNiagaraWorldManager::WorldMana
 
 namespace NiagaraWorldManagerInternal
 {
+	static int32 GFirstHighPriTickGroup = 0;
 	static TQueue<TFunction<void()>, EQueueMode::Mpsc> GlobalDeferredCallbacks;
 
 	void ExecuteGlobalDeferredCallbacks()
@@ -316,6 +285,34 @@ namespace NiagaraWorldManagerInternal
 		const int ActualTickGroup = FMath::Clamp(TickGroup - NiagaraFirstTickGroup, 0, NiagaraNumTickGroups - 1);
 		return ActualTickGroup;
 	}
+
+	bool GetNiagaraTickGroupPriority(int32 NiagaraTickGroup)
+	{
+		return NiagaraTickGroup >= GFirstHighPriTickGroup;
+	}
+
+	static FAutoConsoleVariableRef CVarFirstHighPriTickGroup(
+		TEXT("fx.Niagara.WorldManager.FirstHighPriTickGroup"),
+		GFirstHighPriTickGroup,
+		TEXT("Defines which tick groups should be set to high priority for the world manager.\n")
+		TEXT("0 - (Default) all tick groups will run high priority.\n")
+		TEXT("1 - The first tick group will be normal priority, all others high, etc."),
+		FConsoleVariableDelegate::CreateLambda(
+			[](IConsoleVariable*)
+			{
+				FNiagaraWorldManager::ForAllWorldManagers(
+					[](FNiagaraWorldManager& WorldManager)
+					{
+						for (int32 i=0; i < NiagaraNumTickGroups; ++i)
+						{
+							WorldManager.SetTickGroupPriority(i, GetNiagaraTickGroupPriority(i));
+						}
+					}
+				);
+			}
+		),
+		ECVF_Default
+	); 
 }
 
 //////////////////////////////////////////////////////////////////////////
@@ -390,7 +387,7 @@ void FNiagaraWorldManager::Init(UWorld* InWorld)
 		TickFunc.bCanEverTick = true;
 		TickFunc.bStartWithTickEnabled = true;
 		TickFunc.bAllowTickOnDedicatedServer = false;
-		TickFunc.bHighPriority = TickGroup >= GNiagaraTickGroupToStartHighPriorityTickFunctions;
+		TickFunc.bHighPriority = NiagaraWorldManagerInternal::GetNiagaraTickGroupPriority(TickGroup);
 		TickFunc.Owner = this;
 		TickFunc.RegisterTickFunction(InWorld->PersistentLevel);
 	}
