Fixed enum read/write for data channel BP nodes

diff --git a/Engine/Plugins/FX/Niagara/Source/NiagaraBlueprintNodes/Private/K2Node_ReadDataChannel.cpp b/Engine/Plugins/FX/Niagara/Source/NiagaraBlueprintNodes/Private/K2Node_ReadDataChannel.cpp
--- a/Engine/Plugins/FX/Niagara/Source/NiagaraBlueprintNodes/Private/K2Node_ReadDataChannel.cpp
+++ b/Engine/Plugins/FX/Niagara/Source/NiagaraBlueprintNodes/Private/K2Node_ReadDataChannel.cpp
@@ -13,6 +13,7 @@
 #include "Kismet/KismetMathLibrary.h"
 #include "Kismet/KismetSystemLibrary.h"
 #include "EdGraphSchema_K2.h"
+#include "K2Node_CastByteToEnum.h"
 
 #define LOCTEXT_NAMESPACE "K2Node_ReadDataChannel"
 
@@ -90,6 +91,10 @@ void UK2Node_ReadDataChannel::ExpandNode(FKismetCompilerContext& CompilerContext
 		CompilerContext.MovePinLinksToIntermediate(*OrgInputPin, *NewInputPin);
 	}
 	CompilerContext.MovePinLinksToIntermediate(*GetExecPin(), *CreateReaderNode->GetExecPin());
+	if (FindPin(FName("WorldContextObject"), EGPD_Input) && CreateReaderNode->FindPin(FName("WorldContextObject")))
+	{
+		CompilerContext.MovePinLinksToIntermediate(*FindPin(FName("WorldContextObject"), EGPD_Input), *CreateReaderNode->FindPin(FName("WorldContextObject"), EGPD_Input));
+	}
 
 	// add a validity check for the return value
 	UK2Node_CallFunction* IsValidNode = CompilerContext.SpawnIntermediateNode<UK2Node_CallFunction>(this, SourceGraph);
@@ -141,7 +146,20 @@ void UK2Node_ReadDataChannel::ExpandNode(FKismetCompilerContext& CompilerContext
 		}
 		ReadDataNode->FindPinChecked(FName("VarName"), EGPD_Input)->DefaultValue = InVar.GetName().ToString();
 		CompilerContext.CopyPinLinksToIntermediate(*IndexPin, *ReadDataNode->FindPinChecked(FName("Index"), EGPD_Input));
-		CompilerContext.MovePinLinksToIntermediate(*VarOutputPin, *ReadDataNode->GetReturnValuePin());
+		if (InVar.GetType().IsEnum())
+		{
+			// the read function only returns a byte, so we need to convert it to the actual enum pin
+			UK2Node_CastByteToEnum* CastEnumNode = CompilerContext.SpawnIntermediateNode<UK2Node_CastByteToEnum>(this, SourceGraph);
+			CastEnumNode->bSafe = true; 
+			CastEnumNode->Enum = InVar.GetType().GetEnum(); 
+			CastEnumNode->AllocateDefaultPins();
+			Schema->TryCreateConnection(ReadDataNode->GetReturnValuePin(), CastEnumNode->FindPinChecked(FName("Byte"), EGPD_Input));
+			CompilerContext.MovePinLinksToIntermediate(*VarOutputPin, *CastEnumNode->FindPinChecked(UEdGraphSchema_K2::PN_ReturnValue));
+		}
+		else
+		{
+			CompilerContext.MovePinLinksToIntermediate(*VarOutputPin, *ReadDataNode->GetReturnValuePin());
+		}
 
 		// do a boolean AND of all the IsValid return values
 		UEdGraphPin* ReadValidPin = ReadDataNode->FindPinChecked(FName("IsValid"), EGPD_Output);
diff --git a/Engine/Plugins/FX/Niagara/Source/NiagaraBlueprintNodes/Private/K2Node_WriteDataChannel.cpp b/Engine/Plugins/FX/Niagara/Source/NiagaraBlueprintNodes/Private/K2Node_WriteDataChannel.cpp
--- a/Engine/Plugins/FX/Niagara/Source/NiagaraBlueprintNodes/Private/K2Node_WriteDataChannel.cpp
+++ b/Engine/Plugins/FX/Niagara/Source/NiagaraBlueprintNodes/Private/K2Node_WriteDataChannel.cpp
@@ -96,6 +96,10 @@ void UK2Node_WriteDataChannel::ExpandNode(FKismetCompilerContext& CompilerContex
 		CompilerContext.MovePinLinksToIntermediate(*OrgInputPin, *NewInputPin);
 	}
 	CreateWriterNode->FindPinChecked(FName("Count"))->DefaultValue = FString::FromInt(1);
+	if (FindPin(FName("WorldContextObject"), EGPD_Input) && CreateWriterNode->FindPin(FName("WorldContextObject")))
+	{
+		CompilerContext.MovePinLinksToIntermediate(*FindPin(FName("WorldContextObject"), EGPD_Input), *CreateWriterNode->FindPin(FName("WorldContextObject"), EGPD_Input));
+	}
 	
 	UEdGraphPin* OldExecPin = GetExecPin();
 	UEdGraphPin* NewExecPin = CreateWriterNode->GetExecPin();
@@ -135,6 +139,12 @@ void UK2Node_WriteDataChannel::ExpandNode(FKismetCompilerContext& CompilerContex
 			continue;
 		}
 		WriteDataNode->FindPinChecked(FName("VarName"), EGPD_Input)->DefaultValue = InVar.GetName().ToString();
+
+		if (InVar.GetType().IsEnum())
+		{
+			// bit of a dirty hack to just change the pin type of the function call node, but since the underlying type is the same it should be fine... 
+			WriteDataNode->FindPinChecked(FName("InData"), EGPD_Input)->PinType = VarInputPin->PinType;
+		}
 		CompilerContext.MovePinLinksToIntermediate(*VarInputPin, *WriteDataNode->FindPinChecked(FName("InData"), EGPD_Input));
 
 		// connect exec pins
