Niagara - Add a setting to load cached assets silently.

diff --git a/Engine/Plugins/FX/Niagara/Source/NiagaraEditor/Private/EdGraphSchema_Niagara.cpp b/Engine/Plugins/FX/Niagara/Source/NiagaraEditor/Private/EdGraphSchema_Niagara.cpp
--- a/Engine/Plugins/FX/Niagara/Source/NiagaraEditor/Private/EdGraphSchema_Niagara.cpp
+++ b/Engine/Plugins/FX/Niagara/Source/NiagaraEditor/Private/EdGraphSchema_Niagara.cpp
@@ -1924,40 +1924,6 @@ bool UEdGraphSchema_Niagara::IsSystemConstant(const FNiagaraVariable& Variable)c
 	return FNiagaraConstants::GetEngineConstants().Find(Variable) != INDEX_NONE;
 }
 
-static UNiagaraParameterCollection* EnsureCollectionLoaded(FAssetData& CollectionAsset)
-{
-	if (UNiagaraParameterCollection* Collection = CastChecked<UNiagaraParameterCollection>(CollectionAsset.GetAsset()))
-	{
-		// asset may not have been fully loaded so give it a chance to do it's PostLoad.  When this is triggered from
-		// within a load of an object (like if this is being triggered during a compile of a niagara script when it
-		// gets loaded), then the Collecction and it's DefaultInstance may not have been preloaded yet.  Keeping this
-		// code isolated here as we should get rid of it when we get rid of PostLoad triggering compilation.
-		if (Collection->HasAnyFlags(RF_NeedLoad))
-		{
-			if (FLinkerLoad* CollectionLinker = Collection->GetLinker())
-			{
-				CollectionLinker->Preload(Collection);
-			}
-		}
-		if (UNiagaraParameterCollectionInstance* CollectionInstance = Collection->GetDefaultInstance())
-		{
-			if (CollectionInstance->HasAnyFlags(RF_NeedLoad))
-			{
-				if (FLinkerLoad* CollectionInstanceLinker = CollectionInstance->GetLinker())
-				{
-					CollectionInstanceLinker->Preload(CollectionInstance);
-				}
-			}
-		}
-
-		Collection->ConditionalPostLoad();
-
-		return Collection;
-	}
-
-	return nullptr;
-}
-
 UNiagaraParameterCollection* UEdGraphSchema_Niagara::VariableIsFromParameterCollection(const FNiagaraVariable& Var)const
 {
 	FNameBuilder VarName(Var.GetName());
diff --git a/Engine/Plugins/FX/Niagara/Source/NiagaraEditor/Private/NiagaraEditorModule.cpp b/Engine/Plugins/FX/Niagara/Source/NiagaraEditor/Private/NiagaraEditorModule.cpp
--- a/Engine/Plugins/FX/Niagara/Source/NiagaraEditor/Private/NiagaraEditorModule.cpp
+++ b/Engine/Plugins/FX/Niagara/Source/NiagaraEditor/Private/NiagaraEditorModule.cpp
@@ -1610,6 +1610,12 @@ void FNiagaraEditorModule::OnPostEngineInit()
 	}
 
 	// ensure that all cached asset types are fully loaded.
+	const UNiagaraEditorSettings* NiagaraEditorSettings = GetDefault<UNiagaraEditorSettings>();
+	bool bForceSilentLoadingOfCachedAssets = NiagaraEditorSettings->GetForceSilentLoadingOfCachedAssets();
+
+	ParameterCollectionAssetCache.SetForceLoadSilent(bForceSilentLoadingOfCachedAssets);
+	ParameterDefinitionsAssetCache.SetForceLoadSilent(bForceSilentLoadingOfCachedAssets);
+
 	ParameterCollectionAssetCache.RefreshCache(true /*bAllowLoading*/);
 	ParameterDefinitionsAssetCache.RefreshCache(true /*bAllowLoading*/);
 
@@ -1909,6 +1915,11 @@ const TArray<TWeakObjectPtr<UNiagaraParameterDefinitions>>& FNiagaraEditorModule
 	return ParameterDefinitionsAssetCache.Get();
 }
 
+const TArray<TWeakObjectPtr<UNiagaraParameterCollection>>& FNiagaraEditorModule::GetCachedParameterCollectionAssets()
+{
+	return ParameterCollectionAssetCache.Get();
+}
+
 void FNiagaraEditorModule::GetTargetSystemAndEmitterForDataInterface(UNiagaraDataInterface* InDataInterface, UNiagaraSystem*& OutOwningSystem, FVersionedNiagaraEmitter& OutOwningEmitter)
 {
 	OutOwningSystem = InDataInterface->GetTypedOuter<UNiagaraSystem>();
diff --git a/Engine/Plugins/FX/Niagara/Source/NiagaraEditor/Private/NiagaraEditorSettings.cpp b/Engine/Plugins/FX/Niagara/Source/NiagaraEditor/Private/NiagaraEditorSettings.cpp
--- a/Engine/Plugins/FX/Niagara/Source/NiagaraEditor/Private/NiagaraEditorSettings.cpp
+++ b/Engine/Plugins/FX/Niagara/Source/NiagaraEditor/Private/NiagaraEditorSettings.cpp
@@ -66,6 +66,7 @@ UNiagaraEditorSettings::UNiagaraEditorSettings(const FObjectInitializer& ObjectI
 	bResetDependentSystemsWhenEditingEmitters = false;
 	SetupNamespaceMetadata();
 	TrackedUsageBaseClasses = { UNiagaraNode::StaticClass(), UNiagaraDataInterface::StaticClass(), UNiagaraRendererProperties::StaticClass() };
+	bForceSilentLoadingOfCachedAssets = false;
 }
 
 #define LOCTEXT_NAMESPACE "NamespaceMetadata"
diff --git a/Engine/Plugins/FX/Niagara/Source/NiagaraEditor/Private/NiagaraEditorUtilities.cpp b/Engine/Plugins/FX/Niagara/Source/NiagaraEditor/Private/NiagaraEditorUtilities.cpp
--- a/Engine/Plugins/FX/Niagara/Source/NiagaraEditor/Private/NiagaraEditorUtilities.cpp
+++ b/Engine/Plugins/FX/Niagara/Source/NiagaraEditor/Private/NiagaraEditorUtilities.cpp
@@ -3061,6 +3061,19 @@ bool FNiagaraEditorUtilities::GetAvailableParameterDefinitions(const TArray<FStr
 	return AssetRegistryModule.GetRegistry().GetAssets(ARFilter, OutParameterDefinitionsAssetData);
 }
 
+void FNiagaraEditorUtilities::GetAvailableParameterCollections(TArray<UNiagaraParameterCollection*>& OutParameterCollections)
+{
+	const TArray<TWeakObjectPtr<UNiagaraParameterCollection>>& CachedParameterCollectionAssets = FNiagaraEditorModule::Get().GetCachedParameterCollectionAssets();
+
+	for (const TWeakObjectPtr<UNiagaraParameterCollection>& CachedParameterCollectionAsset : CachedParameterCollectionAssets)
+	{
+		if (UNiagaraParameterCollection* ParameterCollection = CachedParameterCollectionAsset.Get())
+		{
+			OutParameterCollections.Add(ParameterCollection);
+		}
+	}
+}
+
 TSharedPtr<INiagaraParameterDefinitionsSubscriberViewModel> FNiagaraEditorUtilities::GetOwningLibrarySubscriberViewModelForGraph(const UNiagaraGraph* Graph)
 {
 	if (UNiagaraScript* BaseScript = Graph->GetTypedOuter<UNiagaraScript>())
diff --git a/Engine/Plugins/FX/Niagara/Source/NiagaraEditor/Private/ViewModels/Stack/NiagaraStackFunctionInput.cpp b/Engine/Plugins/FX/Niagara/Source/NiagaraEditor/Private/ViewModels/Stack/NiagaraStackFunctionInput.cpp
--- a/Engine/Plugins/FX/Niagara/Source/NiagaraEditor/Private/ViewModels/Stack/NiagaraStackFunctionInput.cpp
+++ b/Engine/Plugins/FX/Niagara/Source/NiagaraEditor/Private/ViewModels/Stack/NiagaraStackFunctionInput.cpp
@@ -1730,31 +1730,19 @@ void UNiagaraStackFunctionInput::GetAvailableParameterHandles(TArray<FNiagaraPar
 	}
 
 	//Parameter Collections
-	FAssetRegistryModule& AssetRegistryModule = FModuleManager::LoadModuleChecked<FAssetRegistryModule>("AssetRegistry");
-	TArray<FAssetData> CollectionAssets;
-	AssetRegistryModule.Get().GetAssetsByClass(UNiagaraParameterCollection::StaticClass()->GetClassPathName(), CollectionAssets);
-
-	// This is a temporary HACK in order to prevent warnings when loading this content and the parameter collection is missing.
-	auto LoadQuiet = [](FAssetData& AssetData)
+	TArray<UNiagaraParameterCollection*> AvailableParameterCollections;
+	FNiagaraEditorUtilities::GetAvailableParameterCollections(AvailableParameterCollections);
+	for (UNiagaraParameterCollection* Collection : AvailableParameterCollections)
 	{
-		uint32 LoadFlags = LOAD_Quiet | LOAD_NoWarn;
-		return StaticLoadObject(UNiagaraParameterCollection::StaticClass(), nullptr, *AssetData.GetObjectPathString(), nullptr, LoadFlags, nullptr, true);
-	};
-
-	for (FAssetData& CollectionAsset : CollectionAssets)
-	{
-		if (UNiagaraParameterCollection* Collection = Cast<UNiagaraParameterCollection>(LoadQuiet(CollectionAsset)))
+		for (const FNiagaraVariable& CollectionParam : Collection->GetParameters())
 		{
-			for (const FNiagaraVariable& CollectionParam : Collection->GetParameters())
+			if (FNiagaraEditorUtilities::AreTypesAssignable(CollectionParam.GetType(), InputType))
 			{
-				if (FNiagaraEditorUtilities::AreTypesAssignable(CollectionParam.GetType(), InputType))
-				{
-					AvailableParameterHandles.AddUnique(FNiagaraParameterHandle(CollectionParam.GetName()));
-				}
-				else if (UNiagaraScript* ConversionScript = FindConversionScript(CollectionParam.GetType(), ConversionScriptCache, bIncludeConversionScripts))
-				{
-					AvailableConversionHandles.Add(CollectionParam, ConversionScript);
-				}
+				AvailableParameterHandles.AddUnique(FNiagaraParameterHandle(CollectionParam.GetName()));
+			}
+			else if (UNiagaraScript* ConversionScript = FindConversionScript(CollectionParam.GetType(), ConversionScriptCache, bIncludeConversionScripts))
+			{
+				AvailableConversionHandles.Add(CollectionParam, ConversionScript);
 			}
 		}
 	}
diff --git a/Engine/Plugins/FX/Niagara/Source/NiagaraEditor/Public/NiagaraEditorModule.h b/Engine/Plugins/FX/Niagara/Source/NiagaraEditor/Public/NiagaraEditorModule.h
--- a/Engine/Plugins/FX/Niagara/Source/NiagaraEditor/Public/NiagaraEditorModule.h
+++ b/Engine/Plugins/FX/Niagara/Source/NiagaraEditor/Public/NiagaraEditorModule.h
@@ -241,6 +241,7 @@ public:
 #endif
 
 	const TArray<TWeakObjectPtr<UNiagaraParameterDefinitions>>& GetCachedParameterDefinitionsAssets();
+	const TArray<TWeakObjectPtr<UNiagaraParameterCollection>>& GetCachedParameterCollectionAssets();
 
 	NIAGARAEDITOR_API void GetTargetSystemAndEmitterForDataInterface(UNiagaraDataInterface* InDataInterface, UNiagaraSystem*& OutOwningSystem, FVersionedNiagaraEmitter& OutOwningEmitter);
 	NIAGARAEDITOR_API void GetDataInterfaceFeedbackSafe(UNiagaraDataInterface* InDataInterface, TArray<FNiagaraDataInterfaceError>& OutErrors, TArray<FNiagaraDataInterfaceFeedback>& Warnings, TArray<FNiagaraDataInterfaceFeedback>& Info);
@@ -298,7 +299,17 @@ private:
 			{
 				if (AssetDatum.IsAssetLoaded() || (bAllowLoading && FPackageName::GetPackageMountPoint(AssetDatum.PackageName.ToString()) != NAME_None))
 				{
-					if (AssetType* Asset = Cast<AssetType>(AssetDatum.GetAsset()))
+					AssetType* Asset = nullptr;
+					if (AssetDatum.IsAssetLoaded() == false && bForceLoadSilent)
+					{
+						Asset = Cast<AssetType>(LoadSilent(AssetDatum));
+					}
+					else
+					{
+						Asset = Cast<AssetType>(AssetDatum.GetAsset());
+					}
+
+					if (Asset != nullptr)
 					{
 						Asset->ConditionalPostLoad();
 						CachedAssets.Add(MakeWeakObjectPtr(Asset));
@@ -309,8 +320,18 @@ private:
 
 		const TArray<TWeakObjectPtr<AssetType>>& Get() const { return CachedAssets; };
 
+		void SetForceLoadSilent(bool bInForceLoadSilent) { bForceLoadSilent = bInForceLoadSilent; }
+
+	private:
+		static UObject* LoadSilent(const FAssetData& AssetData)
+		{
+			uint32 LoadFlags = LOAD_Quiet | LOAD_NoWarn;
+			return StaticLoadObject(AssetType::StaticClass(), nullptr, *AssetData.GetObjectPathString(), nullptr, LoadFlags, nullptr, true);
+		};
+
 	private:
 		TArray<TWeakObjectPtr<AssetType>> CachedAssets;
+		bool bForceLoadSilent = false;
 	};
 
 	void RegisterDefaultRendererFactories();
diff --git a/Engine/Plugins/FX/Niagara/Source/NiagaraEditor/Public/NiagaraEditorSettings.h b/Engine/Plugins/FX/Niagara/Source/NiagaraEditor/Public/NiagaraEditorSettings.h
--- a/Engine/Plugins/FX/Niagara/Source/NiagaraEditor/Public/NiagaraEditorSettings.h
+++ b/Engine/Plugins/FX/Niagara/Source/NiagaraEditor/Public/NiagaraEditorSettings.h
@@ -508,6 +508,8 @@ public:
 	
 	NIAGARAEDITOR_API float GetExposureValue() const;
 	NIAGARAEDITOR_API void SetAutoExposureValue(bool bInUseAutoExposure);
+
+	NIAGARAEDITOR_API bool GetForceSilentLoadingOfCachedAssets() const { return bForceSilentLoadingOfCachedAssets; }
 	
 private:
 	NIAGARAEDITOR_API bool IsAllowedObjectByClassUsageInternal(const UObject& InObject, TSet<const UObject*>& CheckedObjects) const;
@@ -591,6 +593,9 @@ private:
 	
 	UPROPERTY(config)
 	TArray<FNiagaraParameterPanelSectionStorage> SystemParameterPanelSectionData;
+
+	UPROPERTY(config)
+	bool bForceSilentLoadingOfCachedAssets;
 	
 
 	FOnIsClassAllowed OnIsClassAllowedDelegate;
diff --git a/Engine/Plugins/FX/Niagara/Source/NiagaraEditor/Public/NiagaraEditorUtilities.h b/Engine/Plugins/FX/Niagara/Source/NiagaraEditor/Public/NiagaraEditorUtilities.h
--- a/Engine/Plugins/FX/Niagara/Source/NiagaraEditor/Public/NiagaraEditorUtilities.h
+++ b/Engine/Plugins/FX/Niagara/Source/NiagaraEditor/Public/NiagaraEditorUtilities.h
@@ -351,6 +351,8 @@ namespace FNiagaraEditorUtilities
 
 	bool GetAvailableParameterDefinitions(const TArray<FString>& ExternalPackagePaths, TArray<FAssetData>& OutParameterDefinitionsAssetData);
 
+	void GetAvailableParameterCollections(TArray<UNiagaraParameterCollection*>& OutParameterCollections);
+
 	TSharedPtr<INiagaraParameterDefinitionsSubscriberViewModel> GetOwningLibrarySubscriberViewModelForGraph(const UNiagaraGraph* Graph);
 
 	TArray<UNiagaraParameterDefinitions*> DowncastParameterDefinitionsBaseArray(const TArray<UNiagaraParameterDefinitionsBase*> BaseArray);
diff --git a/Engine/Plugins/FX/Niagara/Source/NiagaraEditor/Public/Widgets/SNiagaraParameterMenu.cpp b/Engine/Plugins/FX/Niagara/Source/NiagaraEditor/Public/Widgets/SNiagaraParameterMenu.cpp
--- a/Engine/Plugins/FX/Niagara/Source/NiagaraEditor/Public/Widgets/SNiagaraParameterMenu.cpp
+++ b/Engine/Plugins/FX/Niagara/Source/NiagaraEditor/Public/Widgets/SNiagaraParameterMenu.cpp
@@ -277,18 +277,12 @@ void SNiagaraAddParameterFromPanelMenu::AddMakeNewGroup(FNiagaraMenuActionCollec
 void SNiagaraAddParameterFromPanelMenu::CollectParameterCollectionsActions(FNiagaraMenuActionCollector& Collector)
 {
 	//Create sub menus for parameter collections.
-	FAssetRegistryModule& AssetRegistryModule = FModuleManager::LoadModuleChecked<FAssetRegistryModule>("AssetRegistry");
-	TArray<FAssetData> CollectionAssets;
-	AssetRegistryModule.Get().GetAssetsByClass(UNiagaraParameterCollection::StaticClass()->GetClassPathName(), CollectionAssets);
-
 	const FText Category = GetNamespaceCategoryText(FNiagaraEditorGuids::ParameterCollectionNamespaceMetaDataGuid);
-	for (FAssetData& CollectionAsset : CollectionAssets)
+	TArray<UNiagaraParameterCollection*> AvailableParameterCollections;
+	FNiagaraEditorUtilities::GetAvailableParameterCollections(AvailableParameterCollections);
+	for (UNiagaraParameterCollection* Collection : AvailableParameterCollections)
 	{
-		UNiagaraParameterCollection* Collection = CastChecked<UNiagaraParameterCollection>(CollectionAsset.GetAsset());
-		if (Collection)
-		{
-			AddParameterGroup(Collector, Collection->GetParameters(), FNiagaraEditorGuids::ParameterCollectionNamespaceMetaDataGuid, Category, 10, FString()); //TODO
-		}
+		AddParameterGroup(Collector, Collection->GetParameters(), FNiagaraEditorGuids::ParameterCollectionNamespaceMetaDataGuid, Category, 10, FString());
 	}
 }
 
diff --git a/Engine/Plugins/FX/Niagara/Source/NiagaraEditorWidgets/Private/Stack/SNiagaraStackFunctionInputValue.cpp b/Engine/Plugins/FX/Niagara/Source/NiagaraEditorWidgets/Private/Stack/SNiagaraStackFunctionInputValue.cpp
--- a/Engine/Plugins/FX/Niagara/Source/NiagaraEditorWidgets/Private/Stack/SNiagaraStackFunctionInputValue.cpp
+++ b/Engine/Plugins/FX/Niagara/Source/NiagaraEditorWidgets/Private/Stack/SNiagaraStackFunctionInputValue.cpp
@@ -770,14 +770,11 @@ FReply SNiagaraStackFunctionInputValue::OnLinkedInputDoubleClicked(const FGeomet
 	FString ParamName;
 	FunctionInput->GetLinkedValueHandle().GetName().ToString().Split(TEXT("."), &ParamCollection, &ParamName);
 
-	FAssetRegistryModule& AssetRegistryModule = FModuleManager::LoadModuleChecked<FAssetRegistryModule>("AssetRegistry");
-	TArray<FAssetData> CollectionAssets;
-	AssetRegistryModule.Get().GetAssetsByClass(UNiagaraParameterCollection::StaticClass()->GetClassPathName(), CollectionAssets);
-
-	for (FAssetData& CollectionAsset : CollectionAssets)
+	TArray<UNiagaraParameterCollection*> AvailableParameterCollections;
+	FNiagaraEditorUtilities::GetAvailableParameterCollections(AvailableParameterCollections);
+	for (UNiagaraParameterCollection* Collection : AvailableParameterCollections)
 	{
-		UNiagaraParameterCollection* Collection = CastChecked<UNiagaraParameterCollection>(CollectionAsset.GetAsset());
-		if (Collection && Collection->GetNamespace() == *ParamCollection)
+		if (Collection->GetNamespace() == *ParamCollection)
 		{
 			if (UNiagaraParameterCollectionInstance* NPCInst = FunctionInput->GetSystemViewModel()->GetSystem().GetParameterCollectionOverride(Collection))
 			{
