- Fix race between game / render threads when deciding if rendering was enabled
- The render thread was reading from the renderer GT array to determine if rendering was enabled, on proxy recreates this could lead to a race and the system would not render for a frame

diff --git a/Engine/Plugins/FX/Niagara/Source/Niagara/Private/NiagaraComponent.cpp b/Engine/Plugins/FX/Niagara/Source/Niagara/Private/NiagaraComponent.cpp
--- a/Engine/Plugins/FX/Niagara/Source/Niagara/Private/NiagaraComponent.cpp
+++ b/Engine/Plugins/FX/Niagara/Source/Niagara/Private/NiagaraComponent.cpp
@@ -271,7 +271,8 @@ void FNiagaraSceneProxy::OnTransformChanged(FRHICommandListBase& RHICmdList)
 FPrimitiveViewRelevance FNiagaraSceneProxy::GetViewRelevance(const FSceneView* View) const
 {
 	FPrimitiveViewRelevance Relevance;
-	if (!GetRenderingEnabled() || !FNiagaraUtilities::SupportsNiagaraRendering(View->GetFeatureLevel()))
+
+	if (!RenderData || !RenderData->IsRenderingEnabled_RT() || !FNiagaraUtilities::SupportsNiagaraRendering(View->GetFeatureLevel()))
 	{
 		return Relevance;
 	}
@@ -426,16 +427,11 @@ uint32 FNiagaraSceneProxy::GetAllocatedSize() const
 	return Size;
 }
 
-bool FNiagaraSceneProxy::GetRenderingEnabled() const
-{
-	return RenderData ? RenderData->IsRenderingEnabled() : false;
-}
-
-void FNiagaraSceneProxy::SetRenderingEnabled(bool bInRenderingEnabled)
+void FNiagaraSceneProxy::SetRenderingEnabled_GT(bool bInRenderingEnabled)
 {
 	if (RenderData)
 	{
-		RenderData->SetRenderingEnabled(bInRenderingEnabled);
+		RenderData->SetRenderingEnabled_GT(bInRenderingEnabled);
 	}
 }
 
@@ -980,7 +976,7 @@ void UNiagaraComponent::TickComponent(float DeltaSeconds, enum ELevelTick TickTy
 		if (SceneProxy != nullptr)
 		{
 			FNiagaraSceneProxy* NiagaraProxy = static_cast<FNiagaraSceneProxy*>(SceneProxy);
-			NiagaraProxy->SetRenderingEnabled(bRenderingEnabled && (bCanRenderWhileSeeking || bIsSeeking == false));
+			NiagaraProxy->SetRenderingEnabled_GT(bRenderingEnabled && (bCanRenderWhileSeeking || bIsSeeking == false));
 		}
 	}
 }
diff --git a/Engine/Plugins/FX/Niagara/Source/Niagara/Private/NiagaraSystemRenderData.cpp b/Engine/Plugins/FX/Niagara/Source/Niagara/Private/NiagaraSystemRenderData.cpp
--- a/Engine/Plugins/FX/Niagara/Source/Niagara/Private/NiagaraSystemRenderData.cpp
+++ b/Engine/Plugins/FX/Niagara/Source/Niagara/Private/NiagaraSystemRenderData.cpp
@@ -92,9 +92,9 @@ void FNiagaraSystemRenderData::Destroy_RenderThread()
 FPrimitiveViewRelevance FNiagaraSystemRenderData::GetViewRelevance(const FSceneView& View, const FNiagaraSceneProxy& SceneProxy) const
 {
 	FPrimitiveViewRelevance Relevance;
-	if (IsRenderingEnabled())
+	if (IsRenderingEnabled_RT())
 	{
-		for (const auto Renderer : EmitterRenderers_RT)
+		for (const FNiagaraRenderer* Renderer : EmitterRenderers_RT)
 		{
 			if (Renderer)
 			{
@@ -350,6 +350,24 @@ void FNiagaraSystemRenderData::PostTickRenderers(const FNiagaraSystemInstance& S
 	}
 }
 
+void FNiagaraSystemRenderData::SetRenderingEnabled_GT(bool bInEnabled)
+{
+	check(IsInGameThread());
+
+	if (bRenderingEnabled_GT != bInEnabled)
+	{
+		bRenderingEnabled_GT = bInEnabled;
+		
+		// Note we send immediately because there is no guarantee we would push new data this frame
+		ENQUEUE_RENDER_COMMAND(FNiagaraSystemRenderData_SetRenderingEnabled)(UE::RenderCommandPipe::NiagaraDynamicData,
+			[this]()
+			{
+				bRenderingEnabled_RT = bRenderingEnabled_GT;
+			}
+		);
+	}
+}
+
 void FNiagaraSystemRenderData::OnSystemComplete(const FNiagaraSystemInstance& SystemInstance)
 {
 	if (EmitterRenderers_GT.Num() == 0)
diff --git a/Engine/Plugins/FX/Niagara/Source/Niagara/Public/NiagaraSceneProxy.h b/Engine/Plugins/FX/Niagara/Source/Niagara/Public/NiagaraSceneProxy.h
--- a/Engine/Plugins/FX/Niagara/Source/Niagara/Public/NiagaraSceneProxy.h
+++ b/Engine/Plugins/FX/Niagara/Source/Niagara/Public/NiagaraSceneProxy.h
@@ -30,11 +30,8 @@ public:
 	/** Called to allow renderers to free render state */
 	NIAGARA_API void DestroyRenderState_Concurrent();
 
-	/** Gets whether or not this scene proxy should be rendered. */
-	NIAGARA_API bool GetRenderingEnabled() const;
-
 	/** Sets whether or not this scene proxy should be rendered. */
-	NIAGARA_API void SetRenderingEnabled(bool bInRenderingEnabled);
+	NIAGARA_API void SetRenderingEnabled_GT(bool bInRenderingEnabled);
 
 	FNiagaraGpuComputeDispatchInterface* GetComputeDispatchInterface() const { return ComputeDispatchInterface; }
 
diff --git a/Engine/Plugins/FX/Niagara/Source/Niagara/Public/NiagaraSystemRenderData.h b/Engine/Plugins/FX/Niagara/Source/Niagara/Public/NiagaraSystemRenderData.h
--- a/Engine/Plugins/FX/Niagara/Source/Niagara/Public/NiagaraSystemRenderData.h
+++ b/Engine/Plugins/FX/Niagara/Source/Niagara/Public/NiagaraSystemRenderData.h
@@ -76,8 +76,10 @@ public:
 	void OnSystemComplete(const FNiagaraSystemInstance& SystemInstance);
 	void RecacheRenderers(const FNiagaraSystemInstance& SystemInstance, const FNiagaraSystemInstanceController& Controller);
 
-	FORCEINLINE bool IsRenderingEnabled() const { return bRenderingEnabled && (IsInRenderingThread() ? EmitterRenderers_RT.Num() > 0 : EmitterRenderers_GT.Num() > 0); }
-	FORCEINLINE void SetRenderingEnabled(bool bInEnabled) { bRenderingEnabled = bInEnabled; }
+	void SetRenderingEnabled_GT(bool bInEnabled);
+	FORCEINLINE bool IsRenderingEnabled_GT() const { return bRenderingEnabled_GT && EmitterRenderers_GT.Num() > 0; }
+	FORCEINLINE bool IsRenderingEnabled_RT() const { return bRenderingEnabled_RT && EmitterRenderers_RT.Num() > 0; }
+
 	FORCEINLINE bool HasAnyMotionBlurEnabled() const { return bAnyMotionBlurEnabled; }
 	FORCEINLINE bool HasAnyHeterogeneousVolumesEnabled() const { return bAnyHeterogeneousVolumesEnabled; }
 
@@ -93,7 +95,8 @@ private:
 	/** Indices of renderers in the order they should be rendered. */
 	TArray<int32> RendererDrawOrder;
 
-	bool bRenderingEnabled = true;
+	bool bRenderingEnabled_GT = true;
+	bool bRenderingEnabled_RT = true;
 	bool bAnyMotionBlurEnabled = false;
 	bool bAnyHeterogeneousVolumesEnabled = false;
 
