- Add option to adjust the anim notify progress information based on rate scale, this is disabled for existing content and enabled by default for new content
- Add option to apply time dilation to spawned niagara components based on rate scale, disabled by default

diff --git a/Engine/Plugins/FX/Niagara/Source/NiagaraAnimNotifies/Private/AnimNotifyState_TimedNiagaraEffect.cpp b/Engine/Plugins/FX/Niagara/Source/NiagaraAnimNotifies/Private/AnimNotifyState_TimedNiagaraEffect.cpp
--- a/Engine/Plugins/FX/Niagara/Source/NiagaraAnimNotifies/Private/AnimNotifyState_TimedNiagaraEffect.cpp
+++ b/Engine/Plugins/FX/Niagara/Source/NiagaraAnimNotifies/Private/AnimNotifyState_TimedNiagaraEffect.cpp
@@ -3,6 +3,7 @@
 #include "AnimNotifyState_TimedNiagaraEffect.h"
 
 #include "Animation/AnimInstance.h"
+#include "Animation/AnimSequenceBase.h"
 #include "Components/SkeletalMeshComponent.h"
 
 #include "NiagaraComponent.h"
@@ -24,7 +25,14 @@ UFXSystemComponent* UAnimNotifyState_TimedNiagaraEffect::SpawnEffect(USkeletalMe
 	// Only spawn if we've got valid params
 	if (ValidateParameters(MeshComp))
 	{
-		return UNiagaraFunctionLibrary::SpawnSystemAttached(Template, MeshComp, SocketName, LocationOffset, RotationOffset, EAttachLocation::KeepRelativeOffset, !bDestroyAtEnd);
+		if (UNiagaraComponent* NewComponent = UNiagaraFunctionLibrary::SpawnSystemAttached(Template, MeshComp, SocketName, LocationOffset, RotationOffset, EAttachLocation::KeepRelativeOffset, !bDestroyAtEnd))
+		{
+			if (bApplyRateScaleAsTimeDilation)
+			{
+				NewComponent->SetCustomTimeDilation(Animation->RateScale);
+			}
+			return NewComponent;
+		}
 	}
 	return nullptr;
 }
@@ -129,6 +137,21 @@ UAnimNotifyState_TimedNiagaraEffectAdvanced::UAnimNotifyState_TimedNiagaraEffect
 	NotifyProgressUserParameter = FName("NormalizedNotifyProgress");
 }
 
+void UAnimNotifyState_TimedNiagaraEffectAdvanced::Serialize(FArchive& Ar)
+{
+	Ar.UsingCustomVersion(FFortniteMainBranchObjectVersion::GUID);
+
+	Super::Serialize(Ar);
+
+	if (Ar.IsLoading())
+	{
+		if (Ar.CustomVer(FFortniteMainBranchObjectVersion::GUID) < FFortniteMainBranchObjectVersion::AnimNotifyAddRateScale)
+		{
+			bApplyRateScaleToProgress = false;
+		}
+	}
+}
+
 void UAnimNotifyState_TimedNiagaraEffectAdvanced::NotifyBegin(USkeletalMeshComponent* MeshComp, class UAnimSequenceBase* Animation, float TotalDuration)
 {
 }
@@ -164,7 +187,8 @@ void UAnimNotifyState_TimedNiagaraEffectAdvanced::NotifyTick(USkeletalMeshCompon
 	//TODO: There must be some way to avoid this map and lookup. The information about the current elapsed time and a mapping onto the notify range should be available somewhere in the mesh comp and notify.
 	if (FInstanceProgressInfo* ProgressInfo = ProgressInfoMap.Find(MeshComp))
 	{
-		ProgressInfo->Elapsed += FrameDeltaTime;
+		const float RateScale = bApplyRateScaleToProgress ? Animation->RateScale : 1.0f;
+		ProgressInfo->Elapsed += FrameDeltaTime * RateScale;
 	}
 
 	if (UFXSystemComponent* FXComponent = GetSpawnedEffect(MeshComp))
diff --git a/Engine/Plugins/FX/Niagara/Source/NiagaraAnimNotifies/Public/AnimNotifyState_TimedNiagaraEffect.h b/Engine/Plugins/FX/Niagara/Source/NiagaraAnimNotifies/Public/AnimNotifyState_TimedNiagaraEffect.h
--- a/Engine/Plugins/FX/Niagara/Source/NiagaraAnimNotifies/Public/AnimNotifyState_TimedNiagaraEffect.h
+++ b/Engine/Plugins/FX/Niagara/Source/NiagaraAnimNotifies/Public/AnimNotifyState_TimedNiagaraEffect.h
@@ -37,6 +37,10 @@ class UAnimNotifyState_TimedNiagaraEffect : public UAnimNotifyState
 	UPROPERTY(EditAnywhere, BlueprintReadWrite, Category = NiagaraSystem, meta = (ToolTip = "Rotation offset from the socket or bone for the Niagara system"))
 	FRotator RotationOffset;
 
+	// Should we set the animation rate scale as time dilation on the spawn effect?
+	UPROPERTY(EditAnywhere, BlueprintReadWrite, Category = NiagaraSystem)
+	bool bApplyRateScaleAsTimeDilation = false;
+
 	// Whether or not we destroy the component at the end of the notify or instead just stop
 	// the emitters.
 	UPROPERTY(EditAnywhere, BlueprintReadWrite, Category = NiagaraSystem, meta = (DisplayName = "Destroy Immediately", ToolTip = "Whether the Niagara system should be immediately destroyed at the end of the notify state or be allowed to finish"))
@@ -87,6 +91,8 @@ class UAnimNotifyState_TimedNiagaraEffectAdvanced : public UAnimNotifyState_Time
 	GENERATED_UCLASS_BODY()
 
 public:
+	virtual void Serialize(class FArchive& Ar) override;
+
 	UE_DEPRECATED(5.0, "Please use the other NotifyBegin function instead")
 	NIAGARAANIMNOTIFIES_API virtual void NotifyBegin(class USkeletalMeshComponent* MeshComp, class UAnimSequenceBase* Animation, float TotalDuration) override;
 
@@ -107,6 +113,10 @@ public:
 	UPROPERTY(EditAnywhere, Category = NotifyProgress, meta = (DisplayName = "Enable Normalized Notify Progress", ToolTip = "This send a 0-1 value of the normalized progress to the FX Component to the float User Parameter."))
 	bool bEnableNormalizedNotifyProgress = true;
 
+	// Should we apply the animation rate scale when calculating the elasped time.
+	UPROPERTY(EditAnywhere, BlueprintReadWrite, Category = NiagaraSystem)
+	bool bApplyRateScaleToProgress = true;
+
 	UPROPERTY(EditAnywhere, Category = NotifyProgress, meta = (DisplayName = "User Parameter", ToolTip = "The name of your niagara user variable you would like to send the normalized notify progress to."))
 	FName NotifyProgressUserParameter;
 
diff --git a/Engine/Source/Runtime/Core/Public/UObject/FortniteMainBranchObjectVersion.h b/Engine/Source/Runtime/Core/Public/UObject/FortniteMainBranchObjectVersion.h
--- a/Engine/Source/Runtime/Core/Public/UObject/FortniteMainBranchObjectVersion.h
+++ b/Engine/Source/Runtime/Core/Public/UObject/FortniteMainBranchObjectVersion.h
@@ -390,6 +390,9 @@ struct FFortniteMainBranchObjectVersion
 		// Enabled ObjectPtr property serialization for Dataflow nodes
 		AddDataflowObjectSerialization,
 
+		// Add anim notify rate scaling, defaults to on for new content, off for old content
+		AnimNotifyAddRateScale,
+
 		// -----<new versions can be added above this line>-------------------------------------------------
 		VersionPlusOne,
 		LatestVersion = VersionPlusOne - 1
