- Add missing NIAGARAANIMNOTIFIES_API specifier

diff --git a/Engine/Plugins/FX/Niagara/Source/NiagaraAnimNotifies/Public/AnimNotifyState_TimedNiagaraEffect.h b/Engine/Plugins/FX/Niagara/Source/NiagaraAnimNotifies/Public/AnimNotifyState_TimedNiagaraEffect.h
--- a/Engine/Plugins/FX/Niagara/Source/NiagaraAnimNotifies/Public/AnimNotifyState_TimedNiagaraEffect.h
+++ b/Engine/Plugins/FX/Niagara/Source/NiagaraAnimNotifies/Public/AnimNotifyState_TimedNiagaraEffect.h
@@ -91,7 +91,7 @@ class UAnimNotifyState_TimedNiagaraEffectAdvanced : public UAnimNotifyState_Time
 	GENERATED_UCLASS_BODY()
 
 public:
-	virtual void Serialize(class FArchive& Ar) override;
+	NIAGARAANIMNOTIFIES_API virtual void Serialize(class FArchive& Ar) override;
 
 	UE_DEPRECATED(5.0, "Please use the other NotifyBegin function instead")
 	NIAGARAANIMNOTIFIES_API virtual void NotifyBegin(class USkeletalMeshComponent* MeshComp, class UAnimSequenceBase* Animation, float TotalDuration) override;
