- Skel Mesh DI handle the various combinations of 16 bit bones / 16 bit weight & different bone modes

diff --git a/Engine/Plugins/FX/Niagara/Shaders/Private/NiagaraDataInterfaceSkeletalMesh.ush b/Engine/Plugins/FX/Niagara/Shaders/Private/NiagaraDataInterfaceSkeletalMesh.ush
--- a/Engine/Plugins/FX/Niagara/Shaders/Private/NiagaraDataInterfaceSkeletalMesh.ush
+++ b/Engine/Plugins/FX/Niagara/Shaders/Private/NiagaraDataInterfaceSkeletalMesh.ush
@@ -51,16 +51,26 @@ float3 DISKelMesh_RandomBarycentricCoord(uint Seed1, uint Seed2, uint Seed3)
 	return float3(1.0f - sqrt0, sqrt0 * (1.0 - r1), r1 * sqrt0);
 }
 
-int4 DISkelMesh_UnpackIndices4(uint Packed)
+int4 DISkelMesh_UnpackIndices4x8(uint Packed)
 {
 	return int4(Packed & 0xff, Packed >> 8 & 0xff, Packed >> 16 & 0xff, Packed >> 24 & 0xff);
 }
 
-float4 DISkelMesh_UnpackWeights4(uint Packed)
+int4 DISkelMesh_UnpackIndices4x16(uint PackedA, uint PackedB)
+{
+	return int4(PackedA & 0xffff, PackedA >> 16 & 0xffff, PackedB & 0xffff, PackedB >> 16 & 0xffff);
+}
+
+float4 DISkelMesh_UnpackWeights4x8(uint Packed)
 {
 	return float4(Packed & 0xff, Packed >> 8 & 0xff, Packed >> 16 & 0xff, Packed >> 24 & 0xff) / 255.0f;
 }
 
+float4 DISkelMesh_UnpackWeights4x16(uint PackedA, uint PackedB)
+{
+	return float4(PackedA & 0xffff, PackedA >> 16 & 0xffff, PackedB & 0xffff, PackedB >> 16 & 0xffff) / 65535.0f;
+}
+
 // UvMappingBuffer organized as an array of:
 //struct FFrozenNode
 //{
diff --git a/Engine/Plugins/FX/Niagara/Shaders/Private/NiagaraDataInterfaceSkeletalMeshTemplate.ush b/Engine/Plugins/FX/Niagara/Shaders/Private/NiagaraDataInterfaceSkeletalMeshTemplate.ush
--- a/Engine/Plugins/FX/Niagara/Shaders/Private/NiagaraDataInterfaceSkeletalMeshTemplate.ush
+++ b/Engine/Plugins/FX/Niagara/Shaders/Private/NiagaraDataInterfaceSkeletalMeshTemplate.ush
@@ -23,10 +23,11 @@ Buffer<uint>	{ParameterName}_MeshSampleRegionsVertices;
 Buffer<uint>	{ParameterName}_MeshTriangleMatricesOffsetBuffer;
 uint			{ParameterName}_MeshTriangleCount;
 uint			{ParameterName}_MeshVertexCount;
-uint			{ParameterName}_MeshWeightStride;
-uint			{ParameterName}_MeshSkinWeightIndexSize;
 uint			{ParameterName}_MeshNumTexCoord;
 uint			{ParameterName}_MeshNumWeights;
+uint			{ParameterName}_MeshBoneWeightStride;
+uint			{ParameterName}_MeshBoneIndexSizeBytes;
+uint			{ParameterName}_MeshBoneWeightSizeBytes;
 int				{ParameterName}_NumBones;
 int				{ParameterName}_NumFilteredBones;
 int				{ParameterName}_NumUnfilteredBones;
@@ -57,6 +58,22 @@ bool UniformSamplingRegionEnabled_{ParameterName}() { return ({ParameterName}_En
 bool UnlimitedBoneInfluences_{ParameterName}() { return ({ParameterName}_EnabledFeatures & 0x4) == 4; }
 bool HasMeshColors_{ParameterName}() { return ({ParameterName}_EnabledFeatures & 0x8) == 8; }
 
+void GetIndexWeight_{ParameterName}(int BoneIndexOffset, int BoneWeightOffset, int Layer, out int4 BlendIndices, out float4 BlendWeights)
+{
+#if DISKELMESH_ALLOW_16BIT
+	BlendIndices = {ParameterName}_MeshBoneIndexSizeBytes == 1 ?
+		DISkelMesh_UnpackIndices4x8({ParameterName}_MeshSkinWeightBuffer[BoneIndexOffset + Layer]) :
+		DISkelMesh_UnpackIndices4x16({ParameterName}_MeshSkinWeightBuffer[BoneIndexOffset + (Layer << 1)], {ParameterName}_MeshSkinWeightBuffer[BoneIndexOffset + (Layer << 1) + 1]);
+
+	BlendWeights = {ParameterName}_MeshBoneWeightSizeBytes == 1 ?
+		DISkelMesh_UnpackWeights4x8({ParameterName}_MeshSkinWeightBuffer[BoneWeightOffset + Layer]) :
+		DISkelMesh_UnpackWeights4x16({ParameterName}_MeshSkinWeightBuffer[BoneWeightOffset + (Layer << 1)], {ParameterName}_MeshSkinWeightBuffer[BoneWeightOffset + (Layer << 1) + 1]);
+#else
+	BlendIndices = DISkelMesh_UnpackIndices4x8({ParameterName}_MeshSkinWeightBuffer[BoneIndexOffset + Layer]);
+	BlendWeights = DISkelMesh_UnpackWeights4x8({ParameterName}_MeshSkinWeightBuffer[BoneWeightOffset + Layer]);
+#endif
+}
+
 ////////////////////////////////////////////////////////////////////////////////////////////////////////////////////////////////////////////////////////////////
 // Misc functions
 void GetPreSkinnedLocalBounds_{ParameterName}(out float3 OutCenter, out float3 OutExtentsMin, out float3 OutExtentsMax, out float3 OutExtents, out float3 OutHalfExtents)
@@ -200,56 +217,57 @@ FDISkelMeshSkinnedVertex GetSkinnedVertex_{ParameterName}(uint VertexIndex)
 		BRANCH
 		if (UnlimitedBoneInfluences_{ParameterName}())
 		{
-			uint BlendOffsetCount = {ParameterName}_MeshSkinWeightLookupBuffer[VertexIndex];
-			int NumBoneInfluences = BlendOffsetCount & 0xff;
-			int StreamOffset = BlendOffsetCount >> 8;
-			int WeightsOffset = StreamOffset + ({ParameterName}_MeshSkinWeightIndexSize * NumBoneInfluences);
+			const uint BlendOffsetCount = {ParameterName}_MeshSkinWeightLookupBuffer[VertexIndex];
+			const int NumBoneInfluences = BlendOffsetCount & 0xff;
+			const int StreamOffset = BlendOffsetCount >> 8;
+			const int WeightsOffset = StreamOffset + ({ParameterName}_MeshBoneIndexSizeBytes * NumBoneInfluences);
+			const float WeightNorm = {ParameterName}_MeshBoneWeightSizeBytes == 1 ? 1.0f / 255.0f : 1.0f / 65535.0f;
 
 			PrevBoneMatrix = float3x4(float4(0,0,0,0), float4(0,0,0,0), float4(0,0,0,0));
 			CurrBoneMatrix = float3x4(float4(0,0,0,0), float4(0,0,0,0), float4(0,0,0,0));
 			for (int InfluenceIdx = 0; InfluenceIdx < NumBoneInfluences; InfluenceIdx++)
 			{
-				int BoneIndexOffset = StreamOffset + ({ParameterName}_MeshSkinWeightIndexSize * InfluenceIdx);
+				const int BoneIndexOffset = StreamOffset + ({ParameterName}_MeshBoneIndexSizeBytes * InfluenceIdx);
 				int BoneIndex = {ParameterName}_MeshSkinWeightBuffer[BoneIndexOffset];
-				if ({ParameterName}_MeshSkinWeightIndexSize > 1)
+				if ( {ParameterName}_MeshBoneIndexSizeBytes > 1 )
 				{
-					BoneIndex = {ParameterName}_MeshSkinWeightBuffer[BoneIndexOffset + 1] << 8 | BoneIndex;
-					//@todo-lh: Workaround to fix issue in SPIRVEmitter of DXC; this block must be inside the if branch
-					float BoneWeight = float({ParameterName}_MeshSkinWeightBuffer[WeightsOffset + InfluenceIdx]) / 255.0f;
-
-					PrevBoneMatrix += GetPrevBoneSkinningMatrix_{ParameterName}(SectionBoneOffset + BoneIndex) * BoneWeight;
-					CurrBoneMatrix += GetCurrBoneSkinningMatrix_{ParameterName}(SectionBoneOffset + BoneIndex) * BoneWeight;
+					BoneIndex |= {ParameterName}_MeshSkinWeightBuffer[BoneIndexOffset + 1] << 8;
 				}
-				else
-				{
-					//@todo-lh: Workaround to fix issue in SPIRVEmitter of DXC; this block must be inside the if branch
-					float BoneWeight = float({ParameterName}_MeshSkinWeightBuffer[WeightsOffset + InfluenceIdx]) / 255.0f;
 
-					PrevBoneMatrix += GetPrevBoneSkinningMatrix_{ParameterName}(SectionBoneOffset + BoneIndex) * BoneWeight;
-					CurrBoneMatrix += GetCurrBoneSkinningMatrix_{ParameterName}(SectionBoneOffset + BoneIndex) * BoneWeight;
+				const int BoneWeightOffset = WeightsOffset + ({ParameterName}_MeshBoneWeightSizeBytes * InfluenceIdx);
+				int iBoneWeight = {ParameterName}_MeshSkinWeightBuffer[BoneWeightOffset];
+				if ( {ParameterName}_MeshBoneWeightSizeBytes > 1 )
+				{
+					iBoneWeight |= {ParameterName}_MeshSkinWeightBuffer[BoneWeightOffset + 1] << 8;
 				}
+				float BoneWeight = float(iBoneWeight) * WeightNorm;
+				PrevBoneMatrix += GetPrevBoneSkinningMatrix_{ParameterName}(SectionBoneOffset + BoneIndex) * BoneWeight;
+				CurrBoneMatrix += GetCurrBoneSkinningMatrix_{ParameterName}(SectionBoneOffset + BoneIndex) * BoneWeight;
 			}	
 		}
 		else
 #endif
 		{
-			uint WeightOffset = {ParameterName}_MeshNumWeights == 8 ? 2 : 1;
-			int4 BlendIndices = DISkelMesh_UnpackIndices4({ParameterName}_MeshSkinWeightBuffer[VertexIndex * {ParameterName}_MeshWeightStride]);
-	        float4 BlendWeights = DISkelMesh_UnpackWeights4({ParameterName}_MeshSkinWeightBuffer[VertexIndex * {ParameterName}_MeshWeightStride + WeightOffset]);
-			CurrBoneMatrix = GetCurrSkinningMatrix_{ParameterName}(SectionBoneOffset, BlendIndices, BlendWeights);
-			PrevBoneMatrix = GetPrevSkinningMatrix_{ParameterName}(SectionBoneOffset, BlendIndices, BlendWeights);
+			const uint BoneIndexOffset = VertexIndex * {ParameterName}_MeshBoneWeightStride;
+			const uint BoneWeightOffset = BoneIndexOffset + (({ParameterName}_MeshNumWeights * {ParameterName}_MeshBoneIndexSizeBytes) >> 2);
+			int4 BlendIndices;
+			float4 BlendWeights;
+			{
+				GetIndexWeight_{ParameterName}(BoneIndexOffset, BoneWeightOffset, 0, BlendIndices, BlendWeights);
+				CurrBoneMatrix = GetCurrSkinningMatrix_{ParameterName}(SectionBoneOffset, BlendIndices, BlendWeights);
+				PrevBoneMatrix = GetPrevSkinningMatrix_{ParameterName}(SectionBoneOffset, BlendIndices, BlendWeights);
+			}
+
 #if DISKELMESH_BONE_INFLUENCES != DISKELMESH_BONE_INFLUENCES_MAX4
 			// We are in 8 or unlimited mode so need to check for the extended weights
 			if ({ParameterName}_MeshNumWeights == 8)
 			{
-			    BlendIndices = DISkelMesh_UnpackIndices4({ParameterName}_MeshSkinWeightBuffer[VertexIndex * {ParameterName}_MeshWeightStride + 1]);
-			    BlendWeights = DISkelMesh_UnpackWeights4({ParameterName}_MeshSkinWeightBuffer[VertexIndex * {ParameterName}_MeshWeightStride + WeightOffset + 1]);
+				GetIndexWeight_{ParameterName}(BoneIndexOffset, BoneWeightOffset, 1, BlendIndices, BlendWeights);
 				CurrBoneMatrix += GetCurrSkinningMatrix_{ParameterName}(SectionBoneOffset, BlendIndices, BlendWeights);
 				PrevBoneMatrix += GetPrevSkinningMatrix_{ParameterName}(SectionBoneOffset, BlendIndices, BlendWeights);
 			}
 #endif
 		}
-
 		CurrPosition = mul(CurrBoneMatrix, float4(CurrPosition, 1.0f)).xyz;
 		PrevPosition = mul(PrevBoneMatrix, float4(PrevPosition, 1.0f)).xyz;
 
diff --git a/Engine/Plugins/FX/Niagara/Source/Niagara/Classes/NiagaraDataInterfaceSkeletalMesh.h b/Engine/Plugins/FX/Niagara/Source/Niagara/Classes/NiagaraDataInterfaceSkeletalMesh.h
--- a/Engine/Plugins/FX/Niagara/Source/Niagara/Classes/NiagaraDataInterfaceSkeletalMesh.h
+++ b/Engine/Plugins/FX/Niagara/Source/Niagara/Classes/NiagaraDataInterfaceSkeletalMesh.h
@@ -630,8 +630,9 @@ struct FNDISkeletalMesh_InstanceData
 	uint32 bUnlimitedBoneInfluences : 1;
 	const FSkinWeightDataVertexBuffer* MeshSkinWeightBuffer;
 	const FSkinWeightLookupVertexBuffer* MeshSkinWeightLookupBuffer;
-	uint32 MeshWeightStrideByte;
-	uint32 MeshSkinWeightIndexSizeByte;
+	uint32 MeshBoneWeightStrideBytes;
+	uint32 MeshBoneIndexSizeBytes;
+	uint32 MeshBoneWeightSizeBytes;
 
 	/** Extra mesh data upload to GPU.*/
 	FSkeletalMeshGpuSpawnStaticBuffers* MeshGpuSpawnStaticBuffers;
@@ -1074,8 +1075,9 @@ struct FNiagaraDISkeletalMeshPassedDataToRT
 	bool bIsGpuUniformlyDistributedSampling = false;
 
 	bool bUnlimitedBoneInfluences = false;
-	uint32 MeshWeightStrideByte = 0;
-	uint32 MeshSkinWeightIndexSizeByte = 0;
+	uint32 MeshBoneWeightStrideBytes = 0;
+	uint32 MeshBoneIndexSizeBytes = 0;
+	uint32 MeshBoneWeightSizeBytes = 0;
 	FMatrix44f Transform = FMatrix44f::Identity;
 	FMatrix44f PrevTransform = FMatrix44f::Identity;
 	float DeltaSeconds = 0.0f;
diff --git a/Engine/Plugins/FX/Niagara/Source/Niagara/Private/NiagaraDataInterfaceSkeletalMesh.cpp b/Engine/Plugins/FX/Niagara/Source/Niagara/Private/NiagaraDataInterfaceSkeletalMesh.cpp
--- a/Engine/Plugins/FX/Niagara/Source/Niagara/Private/NiagaraDataInterfaceSkeletalMesh.cpp
+++ b/Engine/Plugins/FX/Niagara/Source/Niagara/Private/NiagaraDataInterfaceSkeletalMesh.cpp
@@ -84,10 +84,11 @@ namespace NDISkelMeshLocal
 		SHADER_PARAMETER_SRV(Buffer<uint>,		MeshTriangleMatricesOffsetBuffer)
 		SHADER_PARAMETER(uint32,				MeshTriangleCount)
 		SHADER_PARAMETER(uint32,				MeshVertexCount)
-		SHADER_PARAMETER(uint32,				MeshWeightStride)
-		SHADER_PARAMETER(uint32,				MeshSkinWeightIndexSize)
 		SHADER_PARAMETER(uint32,				MeshNumTexCoord)
 		SHADER_PARAMETER(uint32,				MeshNumWeights)
+		SHADER_PARAMETER(uint32,				MeshBoneWeightStride)
+		SHADER_PARAMETER(uint32,				MeshBoneIndexSizeBytes)
+		SHADER_PARAMETER(uint32,				MeshBoneWeightSizeBytes)
 		SHADER_PARAMETER(int,					NumBones)
 		SHADER_PARAMETER(int,					NumFilteredBones)
 		SHADER_PARAMETER(int,					NumUnfilteredBones)
@@ -1372,8 +1373,9 @@ void FNiagaraDataInterfaceProxySkeletalMesh::ConsumePerInstanceDataFromGameThrea
 	Data.bUnlimitedBoneInfluences = SourceData->bUnlimitedBoneInfluences;
 	Data.DeltaSeconds = SourceData->DeltaSeconds;
 	Data.DynamicBuffer = SourceData->DynamicBuffer;
-	Data.MeshWeightStrideByte = SourceData->MeshWeightStrideByte;
-	Data.MeshSkinWeightIndexSizeByte = SourceData->MeshSkinWeightIndexSizeByte;
+	Data.MeshBoneWeightStrideBytes = SourceData->MeshBoneWeightStrideBytes;
+	Data.MeshBoneIndexSizeBytes = SourceData->MeshBoneIndexSizeBytes;
+	Data.MeshBoneWeightSizeBytes = SourceData->MeshBoneWeightSizeBytes;
 	Data.PrevTransform = SourceData->PrevTransform;
 	Data.StaticBuffers = SourceData->StaticBuffers;
 	Data.Transform = SourceData->Transform;
@@ -1403,8 +1405,9 @@ void UNiagaraDataInterfaceSkeletalMesh::ProvidePerInstanceDataForRenderThread(vo
 	Data->bUnlimitedBoneInfluences = SourceData->bUnlimitedBoneInfluences;
 	Data->DeltaSeconds = SourceData->DeltaSeconds;
 	Data->DynamicBuffer = SourceData->MeshGpuSpawnDynamicBuffers;
-	Data->MeshWeightStrideByte = SourceData->MeshWeightStrideByte;
-	Data->MeshSkinWeightIndexSizeByte = SourceData->MeshSkinWeightIndexSizeByte;
+	Data->MeshBoneWeightStrideBytes = SourceData->MeshBoneWeightStrideBytes;
+	Data->MeshBoneIndexSizeBytes = SourceData->MeshBoneIndexSizeBytes;
+	Data->MeshBoneWeightSizeBytes = SourceData->MeshBoneWeightSizeBytes;
 	Data->PrevTransform = FMatrix44f(SourceData->PrevTransform);	// LWC_TODO: Precision loss
 	Data->StaticBuffers = SourceData->MeshGpuSpawnStaticBuffers;
 	Data->Transform = FMatrix44f(SourceData->Transform);			// LWC_TODO: Precision loss
@@ -1596,8 +1599,9 @@ bool FNDISkeletalMesh_InstanceData::Init(UNiagaraDataInterfaceSkeletalMesh* Inte
 	ChangeId = Interface->ChangeId;
 	bIsGpuUniformlyDistributedSampling = false;
 	bUnlimitedBoneInfluences = false;
-	MeshWeightStrideByte = 0;
-	MeshSkinWeightIndexSizeByte = 0;
+	MeshBoneWeightStrideBytes = 0;
+	MeshBoneIndexSizeBytes = 0;
+	MeshBoneWeightSizeBytes = 0;
 	MeshGpuSpawnStaticBuffers = nullptr;
 	MeshGpuSpawnDynamicBuffers = nullptr;
 	bAllowCPUMeshDataAccess = false;
@@ -1981,8 +1985,10 @@ bool FNDISkeletalMesh_InstanceData::Init(UNiagaraDataInterfaceSkeletalMesh* Inte
 		{
 			GPUSkinBoneInfluenceType BoneInfluenceType = SkinWeightBuffer->GetBoneInfluenceType();
 			bUnlimitedBoneInfluences = (BoneInfluenceType == GPUSkinBoneInfluenceType::UnlimitedBoneInfluence);
-			MeshWeightStrideByte = SkinWeightBuffer->GetConstantInfluencesVertexStride();
-			MeshSkinWeightIndexSizeByte = SkinWeightBuffer->GetBoneIndexByteSize();
+			MeshBoneWeightStrideBytes = SkinWeightBuffer->GetConstantInfluencesVertexStride();
+			MeshBoneIndexSizeBytes = SkinWeightBuffer->GetBoneIndexByteSize();
+			MeshBoneWeightSizeBytes = SkinWeightBuffer->GetBoneWeightByteSize();
+
 			MeshSkinWeightBuffer = SkinWeightBuffer->GetDataVertexBuffer();
 			//check(MeshSkinWeightBufferSrv->IsValid()); // not available in this stream
 			MeshSkinWeightLookupBuffer = SkinWeightBuffer->GetLookupVertexBuffer();
@@ -2796,6 +2802,7 @@ void UNiagaraDataInterfaceSkeletalMesh::ModifyCompilationEnvironment(EShaderPlat
 {
 	Super::ModifyCompilationEnvironment(ShaderPlatform, OutEnvironment);
 
+	OutEnvironment.SetDefine(TEXT("DISKELMESH_ALLOW_16BIT"), int(GetDefault<UNiagaraSettings>()->NDISkelMesh_Support16BitIndexWeight));
 	OutEnvironment.SetDefine(TEXT("DISKELMESH_BONE_INFLUENCES"), int(GetDefault<UNiagaraSettings>()->NDISkelMesh_GpuMaxInfluences));
 	OutEnvironment.SetDefine(TEXT("DISKELMESH_PROBALIAS_FORMAT"), int(GetDefault<UNiagaraSettings>()->NDISkelMesh_GpuUniformSamplingFormat));
 	OutEnvironment.SetDefine(TEXT("DISKELMESH_ADJ_INDEX_FORMAT"), int(GetDefault<UNiagaraSettings>()->NDISkelMesh_AdjacencyTriangleIndexFormat));
@@ -3223,9 +3230,6 @@ void UNiagaraDataInterfaceSkeletalMesh::SetShaderParameters(const FNiagaraDataIn
 		ShaderParameters->MeshSkinWeightBuffer = FNiagaraRenderer::GetSrvOrDefaultUInt(InstanceData->MeshSkinWeightBuffer->GetSRV());
 		ShaderParameters->MeshSkinWeightLookupBuffer = FNiagaraRenderer::GetSrvOrDefaultUInt(InstanceData->MeshSkinWeightLookupBuffer->GetSRV());
 
-		ShaderParameters->MeshWeightStride = InstanceData->MeshWeightStrideByte / 4;
-		ShaderParameters->MeshSkinWeightIndexSize = InstanceData->MeshSkinWeightIndexSizeByte;
-
 		uint32 EnabledFeaturesBits = 0;
 		EnabledFeaturesBits |= StaticBuffers->IsUseGpuUniformlyDistributedSampling() ? 1 : 0;
 		EnabledFeaturesBits |= StaticBuffers->IsSamplingRegionsAllAreaWeighted() ? 2 : 0;
@@ -3237,6 +3241,10 @@ void UNiagaraDataInterfaceSkeletalMesh::SetShaderParameters(const FNiagaraDataIn
 		if (DynamicBuffers->DoesBoneDataExist())
 		{
 			ShaderParameters->MeshNumWeights = StaticBuffers->GetNumWeights();
+			ShaderParameters->MeshBoneWeightStride = InstanceData->MeshBoneWeightStrideBytes / 4;
+			ShaderParameters->MeshBoneIndexSizeBytes = InstanceData->MeshBoneIndexSizeBytes;
+			ShaderParameters->MeshBoneWeightSizeBytes = InstanceData->MeshBoneWeightSizeBytes;
+
 			ShaderParameters->MeshCurrBonesBuffer = DynamicBuffers->GetRWBufferBone().SectionSRV;
 			ShaderParameters->MeshPrevBonesBuffer = DynamicBuffers->GetRWBufferPrevBone().SectionSRV;
 			ShaderParameters->MeshCurrSamplingBonesBuffer = DynamicBuffers->GetRWBufferBone().SamplingSRV;
@@ -3247,6 +3255,10 @@ void UNiagaraDataInterfaceSkeletalMesh::SetShaderParameters(const FNiagaraDataIn
 		else
 		{
 			ShaderParameters->MeshNumWeights = 0;
+			ShaderParameters->MeshBoneWeightStride = 0;
+			ShaderParameters->MeshBoneIndexSizeBytes = 0;
+			ShaderParameters->MeshBoneWeightSizeBytes = 0;
+
 			ShaderParameters->MeshCurrBonesBuffer = FNiagaraRenderer::GetDummyFloat4Buffer();
 			ShaderParameters->MeshPrevBonesBuffer = FNiagaraRenderer::GetDummyFloat4Buffer();
 			ShaderParameters->MeshCurrSamplingBonesBuffer = FNiagaraRenderer::GetDummyFloat4Buffer();
@@ -3326,9 +3338,10 @@ void UNiagaraDataInterfaceSkeletalMesh::SetShaderParameters(const FNiagaraDataIn
 		ShaderParameters->MeshSkinWeightBuffer = FNiagaraRenderer::GetDummyUIntBuffer();
 		ShaderParameters->MeshSkinWeightLookupBuffer = FNiagaraRenderer::GetDummyUIntBuffer();
 
-		ShaderParameters->MeshWeightStride = 0;
-		ShaderParameters->MeshSkinWeightIndexSize = 0;
 		ShaderParameters->MeshNumWeights = 0;
+		ShaderParameters->MeshBoneWeightStride = 0;
+		ShaderParameters->MeshBoneIndexSizeBytes = 0;
+		ShaderParameters->MeshBoneWeightSizeBytes = 0;
 
 		ShaderParameters->MeshCurrBonesBuffer = FNiagaraRenderer::GetDummyFloat4Buffer();
 		ShaderParameters->MeshPrevBonesBuffer = FNiagaraRenderer::GetDummyFloat4Buffer();
diff --git a/Engine/Plugins/FX/Niagara/Source/Niagara/Public/NiagaraSettings.h b/Engine/Plugins/FX/Niagara/Source/Niagara/Public/NiagaraSettings.h
--- a/Engine/Plugins/FX/Niagara/Source/Niagara/Public/NiagaraSettings.h
+++ b/Engine/Plugins/FX/Niagara/Source/Niagara/Public/NiagaraSettings.h
@@ -225,6 +225,9 @@ class UNiagaraSettings : public UDeveloperSettings
 	UPROPERTY(config, EditAnywhere, Category = LightRenderer)
 	float DefaultLightInverseExposureBlend = 0.0f;
 
+	UPROPERTY(config, EditAnywhere, Category = SkeletalMeshDI, meta = (DisplayName = "Support 16 bit index & weights", ToolTip = "Enabled support for 16 bit bone index & bone weight, optional to reduce shader complexity.  Changing this setting requires restarting the editor.", ConfigRestartRequired = true))
+	bool NDISkelMesh_Support16BitIndexWeight = true;
+
 	UPROPERTY(config, EditAnywhere, Category=SkeletalMeshDI, meta = ( DisplayName = "Gpu Max Bone Influences", ToolTip = "Controls the maximum number of influences we allow the Skeletal Mesh Data Interface to use on the GPU.  Changing this setting requires restarting the editor.", ConfigRestartRequired = true))
 	TEnumAsByte<ENDISkelMesh_GpuMaxInfluences::Type> NDISkelMesh_GpuMaxInfluences;
 
