Forward declare FRayTracingInstance in some headers.
- reduces the number of files recompiled when RayTracingInstance.h is modified from >1000 to ~100.

diff --git a/Engine/Plugins/FX/Niagara/Source/Niagara/Public/NiagaraSystemRenderData.h b/Engine/Plugins/FX/Niagara/Source/Niagara/Public/NiagaraSystemRenderData.h
--- a/Engine/Plugins/FX/Niagara/Source/Niagara/Public/NiagaraSystemRenderData.h
+++ b/Engine/Plugins/FX/Niagara/Source/Niagara/Public/NiagaraSystemRenderData.h
@@ -7,10 +7,6 @@
 #include "PrimitiveUniformShaderParameters.h"
 #include "PrimitiveViewRelevance.h"
 
-#if RHI_RAYTRACING
-#include "RayTracingInstance.h"
-#endif
-
 class FNiagaraSystemInstance;
 class FNiagaraSystemRenderData;
 class FMeshElementCollector;
@@ -18,6 +14,10 @@ struct FNiagaraDynamicDataBase;
 class FNiagaraSceneProxy;
 class FSceneView;
 
+#if RHI_RAYTRACING
+struct FRayTracingInstance;
+#endif
+
 /**
  * This class wraps all data and functionality needed by a scene render proxy to render a single Niagara System Instance.
  */
diff --git a/Engine/Plugins/Runtime/CableComponent/Source/CableComponent/Private/CableComponent.cpp b/Engine/Plugins/Runtime/CableComponent/Source/CableComponent/Private/CableComponent.cpp
--- a/Engine/Plugins/Runtime/CableComponent/Source/CableComponent/Private/CableComponent.cpp
+++ b/Engine/Plugins/Runtime/CableComponent/Source/CableComponent/Private/CableComponent.cpp
@@ -12,6 +12,7 @@
 #include "CableComponentStats.h"
 #include "DynamicMeshBuilder.h"
 #include "StaticMeshResources.h"
+#include "RayTracingInstance.h"
 #include "SceneInterface.h"
 
 DEFINE_RENDER_COMMAND_PIPE(Cable, ERenderCommandPipeFlags::None);
diff --git a/Engine/Source/Runtime/Engine/Public/StaticMeshResources.h b/Engine/Source/Runtime/Engine/Public/StaticMeshResources.h
--- a/Engine/Source/Runtime/Engine/Public/StaticMeshResources.h
+++ b/Engine/Source/Runtime/Engine/Public/StaticMeshResources.h
@@ -39,7 +39,6 @@
 #include "Serialization/BulkData.h"
 #include "WeightedRandomSampler.h"
 #include "PerPlatformProperties.h"
-#include "RayTracingInstance.h"
 #include "RayTracingGeometry.h"
 
 class FDistanceFieldVolumeData;
