Added a cvar to override VM bytecode striping.  Currently strips/disables EVVM to hopefully resolve memory stomps.

#localization none
#tests none

diff --git a/Engine/Plugins/FX/Niagara/Source/Niagara/Private/NiagaraScript.cpp b/Engine/Plugins/FX/Niagara/Source/Niagara/Private/NiagaraScript.cpp
--- a/Engine/Plugins/FX/Niagara/Source/Niagara/Private/NiagaraScript.cpp
+++ b/Engine/Plugins/FX/Niagara/Source/Niagara/Private/NiagaraScript.cpp
@@ -125,6 +125,17 @@ static FAutoConsoleVariableRef CVarNiagaraCompileHashAllDataInterfaces(
 	ECVF_Default
 );
 
+static int32 GNiagaraStripByteCodeOverride = -1; // disable EVVM for now
+static FAutoConsoleVariableRef CVarNiagaraStripByteCodeOverride(
+	TEXT("fx.Niagara.StripByteCodeOverride"),
+	GNiagaraStripByteCodeOverride,
+	TEXT("Overrides project settings to strip byte code from the script on load..\n")
+	TEXT("-1 = Enabled, strips Experimental VM.\n")
+	TEXT("0 = Disabled.\n")
+	TEXT("1 = Enabled, strips original VM.\n"),
+	ECVF_Default
+);
+
 namespace NiagaraScriptInternal
 {
 #if WITH_EDITORONLY_DATA
@@ -421,19 +432,23 @@ void FNiagaraVMExecutableData::PostSerialize(const FArchive& Ar)
 		// only worry about stripping out byte code if we have both sets loaded
 		if (ByteCode.GetLength() > 0 && !ExperimentalContextData.IsEmpty())
 		{
-			switch (GetDefault<UNiagaraSettings>()->ByteCodeStripOption)
+			ENiagaraStripScriptByteCodeOption StripOption = GetDefault<UNiagaraSettings>()->ByteCodeStripOption;
+
+			const bool bStripOriginal = GNiagaraStripByteCodeOverride
+				? (GNiagaraStripByteCodeOverride > 0)
+				: StripOption == ENiagaraStripScriptByteCodeOption::Strip_Original;
+
+			const bool bStripExperimental = GNiagaraStripByteCodeOverride
+				? (GNiagaraStripByteCodeOverride < 0)
+				: StripOption == ENiagaraStripScriptByteCodeOption::Strip_Experimental;
+
+			if (bStripOriginal)
 			{
-			case ENiagaraStripScriptByteCodeOption::Strip_Original:
 				ByteCode.Reset();
-				break;
-
-			case ENiagaraStripScriptByteCodeOption::Strip_Experimental:
+			}
+			else if (bStripExperimental)
+			{
 				ExperimentalContextData.Empty();
-				break;
-			
-			default:
-			case ENiagaraStripScriptByteCodeOption::Default:
-				break;
 			}
 		}
 	}
