high precision collision fixes

#virtualized

diff --git a/Engine/Plugins/FX/Niagara/Shaders/Private/NiagaraDataInterfaceRigidMeshCollisionQuery.ush b/Engine/Plugins/FX/Niagara/Shaders/Private/NiagaraDataInterfaceRigidMeshCollisionQuery.ush
--- a/Engine/Plugins/FX/Niagara/Shaders/Private/NiagaraDataInterfaceRigidMeshCollisionQuery.ush
+++ b/Engine/Plugins/FX/Niagara/Shaders/Private/NiagaraDataInterfaceRigidMeshCollisionQuery.ush
@@ -272,6 +272,56 @@ void GetElementPoint_{ParameterName}(in float3 WorldPosition, in float DeltaTime
 	}
 }
 
+// Given a world space position (WorldPosition) and an element index compute the static mesh closest point using distance fields (position,velocity)
+void GetElementPointMeshDistanceFieldNoNormal_{ParameterName}(in float3 WorldPosition, in float DeltaTime, in float TimeFraction, in float MaxDistance, in int ElementIndex, out float OutClosestDistance, out float3 OutClosestVelocity)
+{
+	OutClosestVelocity = float3(0,0,0);
+	OutClosestDistance = 0.0;
+
+	if (ElementIndex >= 0)
+	{
+		const uint ElementIndexUint = uint(ElementIndex);
+		const float3 LocalPosition = GetLocalPosition_{ParameterName}(WorldPosition,ElementIndex,TimeFraction);
+
+		float MinDistance = MAX_DISTANCE;
+		int OutputIndex = -1;
+		float3 CollisionPosition = float3(0,0,0);
+		float3 CollisionNormal = float3(0,0,0);
+
+		if ( ElementIndexUint >= {ParameterName}_ElementOffsets[SPHERE_INDEX] && ElementIndexUint < {ParameterName}_ElementOffsets[SPHERE_INDEX+1] )
+		{
+			GetSphereProjection_{ParameterName}(LocalPosition, float3(0,0,0), {ParameterName}_ElementExtentBuffer[ElementIndex].x, ElementIndex,
+					CollisionPosition, CollisionNormal, OutputIndex, MinDistance);
+		}
+		else if ( ElementIndexUint >= {ParameterName}_ElementOffsets[CAPSULE_INDEX] && ElementIndexUint < {ParameterName}_ElementOffsets[CAPSULE_INDEX+1] )
+		{
+			GetCapsuleProjection_{ParameterName}(LocalPosition, {ParameterName}_ElementExtentBuffer[ElementIndex].xy, ElementIndex,
+					CollisionPosition, CollisionNormal, OutputIndex, MinDistance);
+		}
+		else if ( ElementIndexUint >= {ParameterName}_ElementOffsets[BOX_INDEX] && ElementIndexUint < {ParameterName}_ElementOffsets[BOX_INDEX+1] )
+		{
+			GetBoxProjection_{ParameterName}(LocalPosition, {ParameterName}_ElementExtentBuffer[ElementIndex].xyz, ElementIndex,
+					CollisionPosition, CollisionNormal, OutputIndex, MinDistance);
+		}
+
+		if (ElementIndex != -1)
+		{
+			uint DFIndex = {ParameterName}_DFIndexBuffer[ElementIndex];
+
+			if (DFIndex >= 0 && DFIndex < NumSceneObjects)
+			{
+				MaxDistance = max(.001, MaxDistance);
+				const FLWCVector3 LWCWorldPosition = MakeLWCVector3({ParameterName}_SystemLWCTile, WorldPosition);
+				OutClosestDistance = DistanceToNearestSurfaceForObject(DFIndex, LWCWorldPosition, MaxDistance);			
+			}
+
+			const float3 PreviousPosition = mul(GetPreviousTransform_{ParameterName}(ElementIndex), float4(CollisionPosition,1.0)).xyz;
+			const float3 CurrentPosition = mul(GetCurrentTransform_{ParameterName}(ElementIndex), float4(CollisionPosition,1.0)).xyz;
+			OutClosestVelocity = ( CurrentPosition - PreviousPosition ) / DeltaTime;		
+		}
+	}
+}
+
 // Given a world space position (WorldPosition) and an element index compute the static mesh closest distance
 void GetElementDistance_{ParameterName}(in float3 WorldPosition, in float TimeFraction, in int ElementIndex, out float OutClosestDistance)
 {	
diff --git a/Engine/Plugins/FX/Niagara/Source/Niagara/Private/NiagaraDataInterfaceRigidMeshCollisionQuery.cpp b/Engine/Plugins/FX/Niagara/Source/Niagara/Private/NiagaraDataInterfaceRigidMeshCollisionQuery.cpp
--- a/Engine/Plugins/FX/Niagara/Source/Niagara/Private/NiagaraDataInterfaceRigidMeshCollisionQuery.cpp
+++ b/Engine/Plugins/FX/Niagara/Source/Niagara/Private/NiagaraDataInterfaceRigidMeshCollisionQuery.cpp
@@ -91,6 +91,7 @@ static const FName GetNumCapsulesName(TEXT("GetNumCapsules"));
 
 static const FName GetClosestElementName(TEXT("GetClosestElement"));
 static const FName GetElementPointName(TEXT("GetElementPoint"));
+static const FName GetElementPointMeshDistanceFieldNoNormalName(TEXT("GetElementPointMeshDistanceFieldNoNormal"));
 static const FName GetElementDistanceName(TEXT("GetElementDistance"));
 static const FName GetClosestPointName(TEXT("GetClosestPoint"));
 static const FName GetClosestDistanceName(TEXT("GetClosestDistance"));
@@ -1393,6 +1394,25 @@ void UNiagaraDataInterfaceRigidMeshCollisionQuery::GetFunctions(TArray<FNiagaraF
 
 		OutFunctions.Add(Sig);
 	}
+	{
+		FNiagaraFunctionSignature Sig;
+		Sig.Name = GetElementPointMeshDistanceFieldNoNormalName;
+		Sig.SetDescription(LOCTEXT("GetClosestElementPointMeshDistanceFieldNoNormalDescription", "Given a world space position and an element index, computes the static mesh's closest point. Also returns normal and velocity for that point."));
+		Sig.SetFunctionVersion(FNiagaraRigidMeshCollisionDIFunctionVersion::LatestVersion);
+		Sig.bSupportsGPU = true;
+		Sig.bSupportsCPU = false;
+		Sig.bMemberFunction = true;
+		Sig.Inputs.Add(FNiagaraVariable(FNiagaraTypeDefinition(GetClass()), TEXT("Collision DI")));
+		Sig.Inputs.Add(FNiagaraVariable(FNiagaraTypeDefinition::GetPositionDef(), TEXT("World Position")));
+		Sig.Inputs.Add(FNiagaraVariable(FNiagaraTypeDefinition::GetFloatDef(), TEXT("Delta Time")));
+		Sig.Inputs.Add(FNiagaraVariable(FNiagaraTypeDefinition::GetFloatDef(), TEXT("Time Fraction")));
+		Sig.Inputs.Add(FNiagaraVariable(FNiagaraTypeDefinition(FNiagaraTypeDefinition::GetFloatDef()), TEXT("MaxDistance")));
+		Sig.Inputs.Add(FNiagaraVariable(FNiagaraTypeDefinition::GetIntDef(), TEXT("Element Index")));
+		Sig.Outputs.Add(FNiagaraVariable(FNiagaraTypeDefinition::GetFloatDef(), TEXT("Closest Distance")));		
+		Sig.Outputs.Add(FNiagaraVariable(FNiagaraTypeDefinition::GetVec3Def(), TEXT("Closest Velocity")));
+
+		OutFunctions.Add(Sig);
+	}
 	{
 		FNiagaraFunctionSignature Sig;
 		Sig.Name = GetElementDistanceName;
@@ -1516,6 +1536,7 @@ bool UNiagaraDataInterfaceRigidMeshCollisionQuery::GetFunctionHLSL(const FNiagar
 		(FunctionInfo.DefinitionName == GetClosestPointName) ||
 		(FunctionInfo.DefinitionName == GetClosestElementName) ||
 		(FunctionInfo.DefinitionName == GetElementPointName) ||
+		(FunctionInfo.DefinitionName == GetElementPointMeshDistanceFieldNoNormalName) ||
 		(FunctionInfo.DefinitionName == GetElementDistanceName) ||
 		(FunctionInfo.DefinitionName == GetClosestDistanceName) ||
 		(FunctionInfo.DefinitionName == GetClosestPointMeshDistanceFieldName) ||
@@ -1554,6 +1575,12 @@ bool UNiagaraDataInterfaceRigidMeshCollisionQuery::UpgradeFunctionCall(FNiagaraF
 			FunctionSignature.Outputs[0].SetType(FNiagaraTypeDefinition::GetPositionDef());
 			bChanged = true;
 		}
+		if (FunctionSignature.Name == GetElementPointMeshDistanceFieldNoNormalName && ensure(FunctionSignature.Inputs.Num() == 5) && ensure(FunctionSignature.Outputs.Num() == 2))
+		{
+			FunctionSignature.Inputs[1].SetType(FNiagaraTypeDefinition::GetPositionDef());
+			FunctionSignature.Outputs[0].SetType(FNiagaraTypeDefinition::GetPositionDef());
+			bChanged = true;
+		}
 		if (FunctionSignature.Name == GetElementDistanceName && ensure(FunctionSignature.Inputs.Num() == 4))
 		{
 			FunctionSignature.Inputs[1].SetType(FNiagaraTypeDefinition::GetPositionDef());
