Fix typo in Niagara Stats: "Simulaton" to "Simulation"

diff --git a/Engine/Plugins/FX/Niagara/Source/Niagara/Private/NiagaraSystemSimulation.cpp b/Engine/Plugins/FX/Niagara/Source/Niagara/Private/NiagaraSystemSimulation.cpp
--- a/Engine/Plugins/FX/Niagara/Source/Niagara/Private/NiagaraSystemSimulation.cpp
+++ b/Engine/Plugins/FX/Niagara/Source/Niagara/Private/NiagaraSystemSimulation.cpp
@@ -37,9 +37,9 @@
 #define NIAGARA_SYSTEMSIMULATION_DEBUGGING 0
 
 //High level stats for system sim tick.
-DECLARE_CYCLE_STAT(TEXT("System Simulaton Tick [GT]"), STAT_NiagaraSystemSim_TickGT, STATGROUP_Niagara);
-DECLARE_CYCLE_STAT(TEXT("System Simulaton Tick [CNC]"), STAT_NiagaraSystemSim_TickCNC, STATGROUP_Niagara);
-DECLARE_CYCLE_STAT(TEXT("System Simulaton SpawnNew [GT]"), STAT_NiagaraSystemSim_SpawnNewGT, STATGROUP_Niagara);
+DECLARE_CYCLE_STAT(TEXT("System Simulation Tick [GT]"), STAT_NiagaraSystemSim_TickGT, STATGROUP_Niagara);
+DECLARE_CYCLE_STAT(TEXT("System Simulation Tick [CNC]"), STAT_NiagaraSystemSim_TickCNC, STATGROUP_Niagara);
+DECLARE_CYCLE_STAT(TEXT("System Simulation SpawnNew [GT]"), STAT_NiagaraSystemSim_SpawnNewGT, STATGROUP_Niagara);
 //Some more detailed stats for system sim tick
 DECLARE_CYCLE_STAT(TEXT("System Prepare For Simulate [CNC]"), STAT_NiagaraSystemSim_PrepareForSimulateCNC, STATGROUP_Niagara);
 DECLARE_CYCLE_STAT(TEXT("System Sim Update [CNC]"), STAT_NiagaraSystemSim_UpdateCNC, STATGROUP_Niagara);
