Added Spawn Count Scale UI to emitter nodes
Changed the scalability indicator to function as a toggle button to enter & exit scalability mode

diff --git a/Engine/Plugins/FX/Niagara/Source/NiagaraEditorWidgets/Private/SNiagaraOverviewStackNode.cpp b/Engine/Plugins/FX/Niagara/Source/NiagaraEditorWidgets/Private/SNiagaraOverviewStackNode.cpp
--- a/Engine/Plugins/FX/Niagara/Source/NiagaraEditorWidgets/Private/SNiagaraOverviewStackNode.cpp
+++ b/Engine/Plugins/FX/Niagara/Source/NiagaraEditorWidgets/Private/SNiagaraOverviewStackNode.cpp
@@ -196,19 +196,44 @@ TSharedRef<SWidget> SNiagaraOverviewStackNode::CreateTitleRightWidget()
 			]
 		]
 
-		// scalability indicator
+		// scalability controls
 		+ SHorizontalBox::Slot()
 		.AutoWidth()
+		.Padding(2.f)
 		[
-			SNew(SBox)
-			.WidthOverride(16.f)
-			.HeightOverride(16.f)
+			SNew(SVerticalBox)
+			+ SVerticalBox::Slot()
+			.AutoHeight()
+			.HAlign(HAlign_Center)
+			.VAlign(VAlign_Center)
 			[
-				SNew(SImage)
-				.Image(FNiagaraEditorStyle::Get().GetBrush("NiagaraEditor.Scalability"))
-				.Visibility(this, &SNiagaraOverviewStackNode::GetScalabilityIndicatorVisibility)
-				.ToolTipText(FText::FormatOrdered(LOCTEXT("ScalabilityIndicatorToolTip",
-					"This {0} has scalability set up. Inspecting and editing scalability is accessible by using scalability mode from the toolbar."), EmitterHandleViewModelWeak.IsValid() ? FText::FromString("emitter") : FText::FromString("system")))
+				// Toggle Button to enter & exit scalability mode
+				SNew(SCheckBox)
+				.Style(&FAppStyle::GetWidgetStyle<FCheckBoxStyle>("ToggleButtonCheckbox"))
+				.IsChecked(this, &SNiagaraOverviewStackNode::IsScalabilityModeActive)
+				.Padding(2.f)
+				.OnCheckStateChanged(this, &SNiagaraOverviewStackNode::OnScalabilityModeStateChanged)
+				[
+					SNew(SBox)
+					.WidthOverride(16.f)
+					.HeightOverride(16.f)
+					[
+						SNew(SImage)
+						.Image(FNiagaraEditorStyle::Get().GetBrush("NiagaraEditor.Scalability"))
+						.Visibility(this, &SNiagaraOverviewStackNode::GetScalabilityIndicatorVisibility)
+						.ToolTipText(FText::FormatOrdered(LOCTEXT("ScalabilityIndicatorToolTip",
+							"This {0} has scalability set up. Inspecting and editing scalability is accessible by entering Scalability Mode by clicking this or the button in the toolbar.."), EmitterHandleViewModelWeak.IsValid() ? FText::FromString("emitter") : FText::FromString("system")))
+					]
+				]
+			]
+			// Spawn Count Scale Info
+			+ SVerticalBox::Slot()
+			.AutoHeight()
+			[
+				SNew(STextBlock)
+				.Text(this, &SNiagaraOverviewStackNode::GetSpawnCountScaleText)
+				.ToolTipText(this, &SNiagaraOverviewStackNode::GetSpawnCountScaleTooltip)
+				.Visibility(this, &SNiagaraOverviewStackNode::GetSpawnCountScaleTextVisibility)
 			]
 		]
 	
@@ -225,6 +250,94 @@ TSharedRef<SWidget> SNiagaraOverviewStackNode::CreateTitleRightWidget()
 		];
 }
 
+FText SNiagaraOverviewStackNode::GetSpawnCountScaleText() const
+{
+	if(EmitterHandleViewModelWeak.IsValid())
+	{
+		if(FVersionedNiagaraEmitterData* EmitterData = EmitterHandleViewModelWeak.Pin()->GetEmitterHandle()->GetInstance().GetEmitterData())
+		{
+			if(EmitterData->GetScalabilitySettings().bScaleSpawnCount)
+			{
+				return FText::FromString(FString::SanitizeFloat(EmitterData->GetScalabilitySettings().SpawnCountScale));
+			}
+		}
+	}
+
+	return FText::GetEmpty();
+}
+
+FText SNiagaraOverviewStackNode::GetSpawnCountScaleTooltip() const
+{
+	if(EmitterHandleViewModelWeak.IsValid())
+	{
+		if(FVersionedNiagaraEmitterData* EmitterData = EmitterHandleViewModelWeak.Pin()->GetEmitterHandle()->GetInstance().GetEmitterData())
+		{
+			if(EmitterData->GetScalabilitySettings().bScaleSpawnCount)
+			{
+				return FText::FormatOrdered(LOCTEXT("EmitterSpawnCountScaleInfoTooltip", "This emitter currently uses a Spawn Count Scale of {0}.\nThis affects the number of spawned particles. Enter Scalability Mode to view & edit."),
+					FText::FromString(FString::SanitizeFloat(EmitterData->GetScalabilitySettings().SpawnCountScale)));
+			}
+		}
+	}
+
+	return FText::GetEmpty();
+}
+
+EVisibility SNiagaraOverviewStackNode::GetSpawnCountScaleTextVisibility() const
+{
+	if(EmitterHandleViewModelWeak.IsValid())
+	{
+		if(FVersionedNiagaraEmitterData* EmitterData = EmitterHandleViewModelWeak.Pin()->GetEmitterHandle()->GetInstance().GetEmitterData())
+		{
+			if(EmitterData->GetScalabilitySettings().bScaleSpawnCount && EmitterData->GetScalabilitySettings().SpawnCountScale != 1.f)
+			{
+				return EVisibility::Visible;
+			}
+		}
+	}
+
+	return EVisibility::Collapsed;
+}
+
+ECheckBoxState SNiagaraOverviewStackNode::IsScalabilityModeActive() const
+{
+	TSharedPtr<FNiagaraSystemViewModel> NiagaraSystemViewModel;
+	
+	if(ScalabilityViewModel.IsValid())
+	{
+		NiagaraSystemViewModel = ScalabilityViewModel->GetSystemViewModel().Pin();
+	}
+	
+	if(NiagaraSystemViewModel.IsValid())
+	{
+		return NiagaraSystemViewModel->GetWorkflowMode().IsEqual(FName("Scalability")) ? ECheckBoxState::Checked : ECheckBoxState::Unchecked;
+	}
+
+	return ECheckBoxState::Unchecked;
+}
+
+void SNiagaraOverviewStackNode::OnScalabilityModeStateChanged(ECheckBoxState CheckBoxState)
+{
+	TSharedPtr<FNiagaraSystemViewModel> NiagaraSystemViewModel;
+	
+	if(ScalabilityViewModel.IsValid())
+	{
+		NiagaraSystemViewModel = ScalabilityViewModel->GetSystemViewModel().Pin();
+	}
+
+	if(NiagaraSystemViewModel.IsValid())
+	{
+		if(NiagaraSystemViewModel->GetWorkflowMode().IsEqual(FName("Scalability")))
+		{
+			NiagaraSystemViewModel->SetWorkflowMode(FName("Default"));
+		}
+		else
+		{
+			NiagaraSystemViewModel->SetWorkflowMode(FName("Scalability"));
+		}
+	}
+}
+
 FReply SNiagaraOverviewStackNode::OnCycleThroughIssues()
 {
 	const TArray<UNiagaraStackEntry*>& ChildrenWithIssues = StackViewModel->GetRootEntry()->GetAllChildrenWithIssues();
diff --git a/Engine/Plugins/FX/Niagara/Source/NiagaraEditorWidgets/Private/SNiagaraOverviewStackNode.h b/Engine/Plugins/FX/Niagara/Source/NiagaraEditorWidgets/Private/SNiagaraOverviewStackNode.h
--- a/Engine/Plugins/FX/Niagara/Source/NiagaraEditorWidgets/Private/SNiagaraOverviewStackNode.h
+++ b/Engine/Plugins/FX/Niagara/Source/NiagaraEditorWidgets/Private/SNiagaraOverviewStackNode.h
@@ -71,7 +71,13 @@ private:
 	{
 		return bIsHoveringThumbnail;
 	}
-	
+
+	FText GetSpawnCountScaleText() const;
+	FText GetSpawnCountScaleTooltip() const;
+	EVisibility GetSpawnCountScaleTextVisibility() const;
+	ECheckBoxState IsScalabilityModeActive() const;
+	void OnScalabilityModeStateChanged(ECheckBoxState CheckBoxState);
+
 	FReply OnCycleThroughIssues();
 	FReply OpenParentEmitter();
 	FText OpenParentEmitterTooltip() const;
