- Fix for BP array SetValue calls not updating the GPU copy
- Fix various poiints in the array code where we might overwrite the GPU sync incorrectly

diff --git a/Engine/Plugins/FX/Niagara/Source/Niagara/Classes/NiagaraDataInterfaceArrayImpl.h b/Engine/Plugins/FX/Niagara/Source/Niagara/Classes/NiagaraDataInterfaceArrayImpl.h
--- a/Engine/Plugins/FX/Niagara/Source/Niagara/Classes/NiagaraDataInterfaceArrayImpl.h
+++ b/Engine/Plugins/FX/Niagara/Source/Niagara/Classes/NiagaraDataInterfaceArrayImpl.h
@@ -565,7 +565,7 @@ struct FNDIArrayProxyImpl : public INDIArrayProxyBase
 			FNDIArrayInstanceData_GameThread<TArrayType>* InstanceData = PerInstanceData_GameThread.CreateConstIterator().Value();
 			FWriteScopeLock	ScopeLock(InstanceData->ArrayRWGuard);
 			InstanceData->bIsModified = false;
-			InstanceData->bIsRenderDirty = bShouldSyncToGpu;
+			InstanceData->bIsRenderDirty |= bShouldSyncToGpu;
 			InstanceData->ArrayData.Empty();
 			Owner->GetArrayReference().SetNum(InArrayData.Num());
 			FNDIArrayImplHelper<TArrayType>::CopyCpuToCpuMemory(Owner->GetArrayReference().GetData(), InArrayData.GetData(), InArrayData.Num());
@@ -598,7 +598,7 @@ struct FNDIArrayProxyImpl : public INDIArrayProxyBase
 			}
 			ArrayRef.GetArray().AddDefaulted(Index + 1 - ArrayRef.GetArray().Num());
 		}
-
+		InstanceData->bIsRenderDirty |= bShouldSyncToGpu;
 		FNDIArrayImplHelper<TArrayType>::CopyCpuToCpuMemory(ArrayRef.GetArray().GetData() + Index, &Value, 1);
 	}
 
@@ -629,7 +629,7 @@ struct FNDIArrayProxyImpl : public INDIArrayProxyBase
 		{
 			FWriteArrayRef ArrayData(Owner, InstanceData);
 			ArrayData.GetArray() = InArrayData;
-			InstanceData->bIsRenderDirty = bShouldSyncToGpu;
+			InstanceData->bIsRenderDirty |= bShouldSyncToGpu;
 		}
 	}
 
@@ -659,7 +659,7 @@ struct FNDIArrayProxyImpl : public INDIArrayProxyBase
 			}
 
 			ArrayData.GetArray()[Index] = Value;
-			InstanceData->bIsRenderDirty = bShouldSyncToGpu;
+			InstanceData->bIsRenderDirty |= bShouldSyncToGpu;
 		}
 	}
 
@@ -1212,7 +1212,7 @@ struct FNDIArrayProxyImpl : public INDIArrayProxyBase
 		FWriteArrayRef ArrayData(Owner, InstanceData);
 		ArrayData.GetArray().Reset();
 
-		InstanceData->bIsRenderDirty = bShouldSyncToGpu;
+		InstanceData->bIsRenderDirty |= bShouldSyncToGpu;
 	}
 
 	void VMResize(FVectorVMExternalFunctionContext& Context)
@@ -1236,7 +1236,7 @@ struct FNDIArrayProxyImpl : public INDIArrayProxyBase
 			}
 		}
 
-		InstanceData->bIsRenderDirty = bShouldSyncToGpu;
+		InstanceData->bIsRenderDirty |= bShouldSyncToGpu;
 	}
 
 	void VMSetValue(FVectorVMExternalFunctionContext& Context)
@@ -1259,7 +1259,7 @@ struct FNDIArrayProxyImpl : public INDIArrayProxyBase
 			}
 		}
 
-		InstanceData->bIsRenderDirty = bShouldSyncToGpu;
+		InstanceData->bIsRenderDirty |= bShouldSyncToGpu;
 	}
 
 	void VMPushValue(FVectorVMExternalFunctionContext& Context)
@@ -1281,7 +1281,7 @@ struct FNDIArrayProxyImpl : public INDIArrayProxyBase
 			}
 		}
 
-		InstanceData->bIsRenderDirty = bShouldSyncToGpu;
+		InstanceData->bIsRenderDirty |= bShouldSyncToGpu;
 	}
 
 	void VMPopValue(FVectorVMExternalFunctionContext& Context)
@@ -1308,7 +1308,7 @@ struct FNDIArrayProxyImpl : public INDIArrayProxyBase
 			}
 		}
 
-		InstanceData->bIsRenderDirty = bShouldSyncToGpu;
+		InstanceData->bIsRenderDirty |= bShouldSyncToGpu;
 	}
 
 	template<typename T = FNDIArrayImplHelper<TArrayType>>
@@ -1350,7 +1350,7 @@ struct FNDIArrayProxyImpl : public INDIArrayProxyBase
 			}
 		}
 
-		InstanceData->bIsRenderDirty = bShouldSyncToGpu;
+		InstanceData->bIsRenderDirty |= bShouldSyncToGpu;
 	}
 
 	template<typename T = FNDIArrayImplHelper<TArrayType>>
@@ -1385,7 +1385,7 @@ struct FNDIArrayProxyImpl : public INDIArrayProxyBase
 			}
 		}
 
-		InstanceData->bIsRenderDirty = bShouldSyncToGpu;
+		InstanceData->bIsRenderDirty |= bShouldSyncToGpu;
 	}
 
 	template<typename T = FNDIArrayImplHelper<TArrayType>>
@@ -1420,7 +1420,7 @@ struct FNDIArrayProxyImpl : public INDIArrayProxyBase
 			}
 		}
 
-		InstanceData->bIsRenderDirty = bShouldSyncToGpu;
+		InstanceData->bIsRenderDirty |= bShouldSyncToGpu;
 	}
 
 private:
