- Add support for reading the Skin Cache to Niagara Skel Mesh DI
- Add a project setting (enabled by default) to control if this is enabled as some lower end platforms could have issues with the additional Buffer's

diff --git a/Engine/Plugins/FX/Niagara/Shaders/Private/NiagaraDataInterfaceSkeletalMeshTemplate.ush b/Engine/Plugins/FX/Niagara/Shaders/Private/NiagaraDataInterfaceSkeletalMeshTemplate.ush
--- a/Engine/Plugins/FX/Niagara/Shaders/Private/NiagaraDataInterfaceSkeletalMeshTemplate.ush
+++ b/Engine/Plugins/FX/Niagara/Shaders/Private/NiagaraDataInterfaceSkeletalMeshTemplate.ush
@@ -50,13 +50,23 @@ float			{ParameterName}_InstanceInvDeltaTime;
 float3			{ParameterName}_PreSkinnedLocalBoundsCenter;
 float3			{ParameterName}_PreSkinnedLocalBoundsExtents;
 uint			{ParameterName}_EnabledFeatures;
+#if DISKELMESH_ALLOW_DEFORMED
+Buffer<float>	{ParameterName}_DeformedCurrPositionBuffer;
+Buffer<float>	{ParameterName}_DeformedPrevPositionBuffer;
+Buffer<float4>	{ParameterName}_DeformedTangentBuffer;
+#endif
 
 ////////////////////////////////////////////////////////////////////////////////////////////////////////////////////////////////////////////////////////////////
 // Flags test
-bool UniformTriangleSamplingEnable_{ParameterName}() { return ({ParameterName}_EnabledFeatures & 0x1) == 1; }
-bool UniformSamplingRegionEnabled_{ParameterName}() { return ({ParameterName}_EnabledFeatures & 0x2) == 2; }
-bool UnlimitedBoneInfluences_{ParameterName}() { return ({ParameterName}_EnabledFeatures & 0x4) == 4; }
-bool HasMeshColors_{ParameterName}() { return ({ParameterName}_EnabledFeatures & 0x8) == 8; }
+bool UniformTriangleSamplingEnable_{ParameterName}()	{ return ({ParameterName}_EnabledFeatures & 0x01) == 0x01; }
+bool UniformSamplingRegionEnabled_{ParameterName}()		{ return ({ParameterName}_EnabledFeatures & 0x02) == 0x02; }
+bool UnlimitedBoneInfluences_{ParameterName}()			{ return ({ParameterName}_EnabledFeatures & 0x04) == 0x04; }
+bool HasMeshColors_{ParameterName}()					{ return ({ParameterName}_EnabledFeatures & 0x08) == 0x08; }
+
+#if DISKELMESH_ALLOW_DEFORMED
+bool UseDeformedPositions_{ParameterName}()				{ return ({ParameterName}_EnabledFeatures & 0x10) == 0x10; }
+bool UseDeformedTangents_{ParameterName}()				{ return ({ParameterName}_EnabledFeatures & 0x20) == 0x20; }
+#endif
 
 void GetIndexWeight_{ParameterName}(int BoneIndexOffset, int BoneWeightOffset, int Layer, out int4 BlendIndices, out float4 BlendWeights)
 {
@@ -198,14 +208,44 @@ FDISkelMeshSkinnedVertex GetSkinnedVertex_{ParameterName}(uint VertexIndex)
 
 	VertexIndex = clamp(VertexIndex, 0, {ParameterName}_MeshVertexCount - 1);
 
-	//-TODO: Use a R32G32B32 format
-	float3 CurrPosition = float3({ParameterName}_MeshVertexBuffer[VertexIndex * 3], {ParameterName}_MeshVertexBuffer[VertexIndex * 3 + 1], {ParameterName}_MeshVertexBuffer[VertexIndex * 3 + 2]);
-	float3 PrevPosition = CurrPosition;
+	// Read Positions For Skinning
+	bool bNeedsSkinPositions = true;
+	float3 CurrPosition = 0;
+	float3 PrevPosition = 0;
+#if DISKELMESH_ALLOW_DEFORMED
+	if (UseDeformedPositions_{ParameterName}())
+	{
+		bNeedsSkinPositions = false;
+		CurrPosition = float3({ParameterName}_DeformedCurrPositionBuffer[VertexIndex * 3], {ParameterName}_DeformedCurrPositionBuffer[VertexIndex * 3 + 1], {ParameterName}_DeformedCurrPositionBuffer[VertexIndex * 3 + 2]);
+		PrevPosition = float3({ParameterName}_DeformedPrevPositionBuffer[VertexIndex * 3], {ParameterName}_DeformedPrevPositionBuffer[VertexIndex * 3 + 1], {ParameterName}_DeformedPrevPositionBuffer[VertexIndex * 3 + 2]);
+	}
+	else
+#endif
+	{
+		CurrPosition = float3({ParameterName}_MeshVertexBuffer[VertexIndex * 3], {ParameterName}_MeshVertexBuffer[VertexIndex * 3 + 1], {ParameterName}_MeshVertexBuffer[VertexIndex * 3 + 2]);
+		PrevPosition = CurrPosition;
+	}
 
-	float3 TangentX = TangentBias({ParameterName}_MeshTangentBuffer[VertexIndex * 2    ].xyz);
-	float4 TangentZ = TangentBias({ParameterName}_MeshTangentBuffer[VertexIndex * 2 + 1].xyzw);
-	
-	if ( {ParameterName}_MeshNumWeights > 0 )
+	// Read Tangents For Skinning
+	bool bNeedsSkinTangents = true;
+	float3 TangentX = 0;
+	float4 TangentZ = 0;
+#if DISKELMESH_ALLOW_DEFORMED
+	if (UseDeformedTangents_{ParameterName}())
+	{
+		bNeedsSkinTangents = false;
+		TangentX = TangentBias({ParameterName}_DeformedTangentBuffer[VertexIndex * 2    ].xyz);
+		TangentZ = TangentBias({ParameterName}_DeformedTangentBuffer[VertexIndex * 2 + 1].xyzw);
+	}
+	else
+#endif
+	{
+		TangentX = TangentBias({ParameterName}_MeshTangentBuffer[VertexIndex * 2    ].xyz);
+		TangentZ = TangentBias({ParameterName}_MeshTangentBuffer[VertexIndex * 2 + 1].xyzw);
+	}
+
+	// Perform skinning if needed
+	if ( {ParameterName}_MeshNumWeights > 0 && (bNeedsSkinPositions || bNeedsSkinTangents))
 	{
 		float3x4 PrevBoneMatrix;
 		float3x4 CurrBoneMatrix;
@@ -268,12 +308,19 @@ FDISkelMeshSkinnedVertex GetSkinnedVertex_{ParameterName}(uint VertexIndex)
 			}
 #endif
 		}
-		CurrPosition = mul(CurrBoneMatrix, float4(CurrPosition, 1.0f)).xyz;
-		PrevPosition = mul(PrevBoneMatrix, float4(PrevPosition, 1.0f)).xyz;
 
-		// Not using InverseTranspose of matrices so assuming uniform scaling only (same as SkinCache)
-		TangentX.xyz = mul(CurrBoneMatrix, float4(TangentX.xyz, 0.0f)).xyz;
-		TangentZ.xyz = mul(CurrBoneMatrix, float4(TangentZ.xyz, 0.0f)).xyz;
+		if ( bNeedsSkinPositions )
+		{
+			CurrPosition = mul(CurrBoneMatrix, float4(CurrPosition, 1.0f)).xyz;
+			PrevPosition = mul(PrevBoneMatrix, float4(PrevPosition, 1.0f)).xyz;
+		}
+
+		if ( bNeedsSkinTangents )
+		{
+			// Not using InverseTranspose of matrices so assuming uniform scaling only (same as SkinCache)
+			TangentX.xyz = mul(CurrBoneMatrix, float4(TangentX.xyz, 0.0f)).xyz;
+			TangentZ.xyz = mul(CurrBoneMatrix, float4(TangentZ.xyz, 0.0f)).xyz;
+		}
 	}
 
 	float3 TangentY = cross(TangentZ.xyz, TangentX.xyz) * TangentZ.w;
diff --git a/Engine/Plugins/FX/Niagara/Source/Niagara/Classes/NiagaraDataInterfaceSkeletalMesh.h b/Engine/Plugins/FX/Niagara/Source/Niagara/Classes/NiagaraDataInterfaceSkeletalMesh.h
--- a/Engine/Plugins/FX/Niagara/Source/Niagara/Classes/NiagaraDataInterfaceSkeletalMesh.h
+++ b/Engine/Plugins/FX/Niagara/Source/Niagara/Classes/NiagaraDataInterfaceSkeletalMesh.h
@@ -392,16 +392,19 @@ protected:
 class FSkeletalMeshGpuSpawnStaticBuffers : public FRenderResource
 {
 public:
-
 	virtual ~FSkeletalMeshGpuSpawnStaticBuffers() override;
 
-	FORCEINLINE_DEBUGGABLE void Initialise(struct FNDISkeletalMesh_InstanceData* InstData, const FSkeletalMeshLODRenderData& SkeletalMeshLODRenderData,const FSkeletalMeshSamplingLODBuiltData* SkeletalMeshSamplingLODBuiltData, FNiagaraSystemInstance* SystemInstance);
+	FORCEINLINE_DEBUGGABLE void Initialise(struct FNDISkeletalMesh_InstanceData* InstData, int32 InLODIndex, const FSkeletalMeshLODRenderData& SkeletalMeshLODRenderData, const FSkeletalMeshSamplingLODBuiltData* SkeletalMeshSamplingLODBuiltData, FNiagaraSystemInstance* SystemInstance);
 
 	virtual void InitRHI(FRHICommandListBase& RHICmdList) override;
 	virtual void ReleaseRHI() override;
 
 	virtual FString GetFriendlyName() const override { return TEXT("FSkeletalMeshGpuSpawnStaticBuffers"); }
 
+	USceneComponent* GetSceneComponent() const { return SceneComponent.Get(); }
+
+	int32 GetLODIndex() const { return LODIndex; };
+
 	FRHIShaderResourceView* GetBufferTriangleUniformSamplerProbAliasSRV() const { return BufferTriangleUniformSamplerProbAliasSRV; }
 	FRHIShaderResourceView* GetBufferTriangleMatricesOffsetSRV() const { return BufferTriangleMatricesOffsetSRV; }
 	uint32 GetTriangleCount() const { return TriangleCount; }
@@ -436,6 +439,8 @@ public:
 	int32 GetFilteredSocketBoneOffset() const { return FilteredSocketBoneOffset; }
 
 protected:
+	TWeakObjectPtr<USceneComponent> SceneComponent = nullptr;
+
 	FBufferRHIRef BufferTriangleUniformSamplerProbAliasRHI = nullptr;
 	FShaderResourceViewRHIRef BufferTriangleUniformSamplerProbAliasSRV = nullptr;
 	FBufferRHIRef BufferTriangleMatricesOffsetRHI = nullptr;
@@ -478,6 +483,7 @@ protected:
 	uint32 NumWeights = 0;
 
 	// Cached data for resource creation on RenderThread
+	int32 LODIndex = INDEX_NONE;
 	const FSkeletalMeshLODRenderData* LODRenderData = nullptr;
 	const FSkeletalMeshSamplingLODBuiltData* SkeletalMeshSamplingLODBuiltData = nullptr;
 	uint32 TriangleCount = 0;
@@ -626,6 +632,8 @@ struct FNDISkeletalMesh_InstanceData
 	/** True if the mesh we're using allows area weighted sampling on GPU. */
 	uint32 bIsGpuUniformlyDistributedSampling : 1;
 
+	uint32 bReadDeformedGeometry : 1 = true;
+
 	/** True if the mesh we're using is to be rendered in unlimited bone influences mode. */
 	uint32 bUnlimitedBoneInfluences : 1;
 	const FSkinWeightDataVertexBuffer* MeshSkinWeightBuffer;
@@ -759,6 +767,13 @@ public:
 	UPROPERTY(EditAnywhere, Category = "Performance")
 	bool bRequireCurrentFrameData = true;
 
+	/**
+	Overrides the project setting and allows you to opt out of reading from deformed geometry.
+	These is not performance gain from doing this, the branches will still exist in the generated code.
+	*/
+	UPROPERTY(EditAnywhere, Category = "Performance")
+	bool bReadDeformedGeometry = true;
+
 	/** Cached change id off of the data interface.*/
 	uint32 ChangeId;
 
@@ -1073,7 +1088,7 @@ struct FNiagaraDISkeletalMeshPassedDataToRT
 	const FSkeletalMeshConnectivityProxy* ConnectivityBuffer = nullptr;
 
 	bool bIsGpuUniformlyDistributedSampling = false;
-
+	bool bReadDeformedGeometry = false;
 	bool bUnlimitedBoneInfluences = false;
 	uint32 MeshBoneWeightStrideBytes = 0;
 	uint32 MeshBoneIndexSizeBytes = 0;
diff --git a/Engine/Plugins/FX/Niagara/Source/Niagara/Private/NiagaraDataInterfaceSkeletalMesh.cpp b/Engine/Plugins/FX/Niagara/Source/Niagara/Private/NiagaraDataInterfaceSkeletalMesh.cpp
--- a/Engine/Plugins/FX/Niagara/Source/Niagara/Private/NiagaraDataInterfaceSkeletalMesh.cpp
+++ b/Engine/Plugins/FX/Niagara/Source/Niagara/Private/NiagaraDataInterfaceSkeletalMesh.cpp
@@ -27,6 +27,8 @@
 #include "Rendering/SkeletalMeshRenderData.h"
 #include "Templates/AlignmentTemplates.h"
 #include "ShaderCore.h"
+#include "CachedGeometry.h"
+#include "SkeletalRenderPublic.h"
 
 #include UE_INLINE_GENERATED_CPP_BY_NAME(NiagaraDataInterfaceSkeletalMesh)
 
@@ -111,6 +113,10 @@ namespace NDISkelMeshLocal
 		SHADER_PARAMETER(FVector3f,				PreSkinnedLocalBoundsCenter)
 		SHADER_PARAMETER(FVector3f,				PreSkinnedLocalBoundsExtents)
 		SHADER_PARAMETER(uint32,				EnabledFeatures)
+
+		SHADER_PARAMETER_SRV(Buffer<float>,		DeformedCurrPositionBuffer)
+		SHADER_PARAMETER_SRV(Buffer<float>,		DeformedPrevPositionBuffer)
+		SHADER_PARAMETER_SRV(Buffer<float4>,	DeformedTangentBuffer)
 	END_SHADER_PARAMETER_STRUCT()
 
 	int32 GetProbAliasDWORDSize(int32 TriangleCount)
@@ -790,11 +796,12 @@ FSkeletalMeshGpuSpawnStaticBuffers::~FSkeletalMeshGpuSpawnStaticBuffers()
 	//ValidSections.Empty();
 }
 
-void FSkeletalMeshGpuSpawnStaticBuffers::Initialise(FNDISkeletalMesh_InstanceData* InstData, const FSkeletalMeshLODRenderData& SkeletalMeshLODRenderData, const FSkeletalMeshSamplingLODBuiltData* MeshSamplingLODBuiltData, FNiagaraSystemInstance* SystemInstance)
+void FSkeletalMeshGpuSpawnStaticBuffers::Initialise(FNDISkeletalMesh_InstanceData* InstData, int32 InLODIndex, const FSkeletalMeshLODRenderData& SkeletalMeshLODRenderData, const FSkeletalMeshSamplingLODBuiltData* MeshSamplingLODBuiltData, FNiagaraSystemInstance* SystemInstance)
 {
 	SkeletalMeshSamplingLODBuiltData = nullptr;
 	bUseGpuUniformlyDistributedSampling = false;
 
+	LODIndex = InLODIndex;
 	LODRenderData = nullptr;
 	TriangleCount = 0;
 	VertexCount = 0;
@@ -807,6 +814,8 @@ void FSkeletalMeshGpuSpawnStaticBuffers::Initialise(FNDISkeletalMesh_InstanceDat
 
 	if (InstData)
 	{
+		SceneComponent = InstData->SceneComponent;
+
 		SkeletalMeshSamplingLODBuiltData = MeshSamplingLODBuiltData;
 		bUseGpuUniformlyDistributedSampling = InstData->bIsGpuUniformlyDistributedSampling;
 
@@ -1370,6 +1379,7 @@ void FNiagaraDataInterfaceProxySkeletalMesh::ConsumePerInstanceDataFromGameThrea
 	FNiagaraDataInterfaceProxySkeletalMeshData& Data = SystemInstancesToData.FindOrAdd(Instance);
 
 	Data.bIsGpuUniformlyDistributedSampling = SourceData->bIsGpuUniformlyDistributedSampling;
+	Data.bReadDeformedGeometry = SourceData->bReadDeformedGeometry;
 	Data.bUnlimitedBoneInfluences = SourceData->bUnlimitedBoneInfluences;
 	Data.DeltaSeconds = SourceData->DeltaSeconds;
 	Data.DynamicBuffer = SourceData->DynamicBuffer;
@@ -1402,6 +1412,7 @@ void UNiagaraDataInterfaceSkeletalMesh::ProvidePerInstanceDataForRenderThread(vo
 	FNDISkeletalMesh_InstanceData* SourceData = static_cast<FNDISkeletalMesh_InstanceData*>(PerInstanceData);
 
 	Data->bIsGpuUniformlyDistributedSampling = SourceData->bIsGpuUniformlyDistributedSampling;
+	Data->bReadDeformedGeometry = SourceData->bReadDeformedGeometry;
 	Data->bUnlimitedBoneInfluences = SourceData->bUnlimitedBoneInfluences;
 	Data->DeltaSeconds = SourceData->DeltaSeconds;
 	Data->DynamicBuffer = SourceData->MeshGpuSpawnDynamicBuffers;
@@ -1598,6 +1609,7 @@ bool FNDISkeletalMesh_InstanceData::Init(UNiagaraDataInterfaceSkeletalMesh* Inte
 	DeltaSeconds = SystemInstance->GetWorld()->GetDeltaSeconds();
 	ChangeId = Interface->ChangeId;
 	bIsGpuUniformlyDistributedSampling = false;
+	bReadDeformedGeometry = Interface->bReadDeformedGeometry && GetDefault<UNiagaraSettings>()->NDISkelMesh_SupportReadingDeformedGeometry;
 	bUnlimitedBoneInfluences = false;
 	MeshBoneWeightStrideBytes = 0;
 	MeshBoneIndexSizeBytes = 0;
@@ -2026,7 +2038,7 @@ bool FNDISkeletalMesh_InstanceData::Init(UNiagaraDataInterfaceSkeletalMesh* Inte
 			bIsGpuUniformlyDistributedSampling &= MeshSamplingBuiltData != nullptr;
 
 			MeshGpuSpawnStaticBuffers = new FSkeletalMeshGpuSpawnStaticBuffers();
-			MeshGpuSpawnStaticBuffers->Initialise(this, *CachedLODData, MeshSamplingBuiltData, SystemInstance);
+			MeshGpuSpawnStaticBuffers->Initialise(this, CachedLODIdx, *CachedLODData, MeshSamplingBuiltData, SystemInstance);
 			BeginInitResource(MeshGpuSpawnStaticBuffers);
 
 			MeshGpuSpawnDynamicBuffers = new FSkeletalMeshGpuDynamicBufferProxy();
@@ -2394,6 +2406,7 @@ bool UNiagaraDataInterfaceSkeletalMesh::CopyToInternal(UNiagaraDataInterface* De
 	OtherTyped->bExcludeBone = bExcludeBone;
 	OtherTyped->ExcludeBoneName = ExcludeBoneName;
 	OtherTyped->bRequireCurrentFrameData = bRequireCurrentFrameData;
+	OtherTyped->bReadDeformedGeometry = bReadDeformedGeometry;
 	OtherTyped->UvSetIndex = UvSetIndex;
 #if WITH_EDITORONLY_DATA
 	OtherTyped->PreviewMesh = PreviewMesh;
@@ -2428,7 +2441,8 @@ bool UNiagaraDataInterfaceSkeletalMesh::Equals(const UNiagaraDataInterface* Othe
 		OtherTyped->bExcludeBone == bExcludeBone &&
 		OtherTyped->ExcludeBoneName == ExcludeBoneName &&
 		OtherTyped->UvSetIndex == UvSetIndex &&
-		OtherTyped->bRequireCurrentFrameData == bRequireCurrentFrameData;
+		OtherTyped->bRequireCurrentFrameData == bRequireCurrentFrameData &&
+		OtherTyped->bReadDeformedGeometry == bReadDeformedGeometry;
 }
 
 bool UNiagaraDataInterfaceSkeletalMesh::InitPerInstanceData(void* PerInstanceData, FNiagaraSystemInstance* SystemInstance)
@@ -2802,6 +2816,7 @@ void UNiagaraDataInterfaceSkeletalMesh::ModifyCompilationEnvironment(EShaderPlat
 {
 	Super::ModifyCompilationEnvironment(ShaderPlatform, OutEnvironment);
 
+	OutEnvironment.SetDefine(TEXT("DISKELMESH_ALLOW_DEFORMED"), int(GetDefault<UNiagaraSettings>()->NDISkelMesh_SupportReadingDeformedGeometry));
 	OutEnvironment.SetDefine(TEXT("DISKELMESH_ALLOW_16BIT"), int(GetDefault<UNiagaraSettings>()->NDISkelMesh_Support16BitIndexWeight));
 	OutEnvironment.SetDefine(TEXT("DISKELMESH_BONE_INFLUENCES"), int(GetDefault<UNiagaraSettings>()->NDISkelMesh_GpuMaxInfluences));
 	OutEnvironment.SetDefine(TEXT("DISKELMESH_PROBALIAS_FORMAT"), int(GetDefault<UNiagaraSettings>()->NDISkelMesh_GpuUniformSamplingFormat));
@@ -3231,10 +3246,42 @@ void UNiagaraDataInterfaceSkeletalMesh::SetShaderParameters(const FNiagaraDataIn
 		ShaderParameters->MeshSkinWeightLookupBuffer = FNiagaraRenderer::GetSrvOrDefaultUInt(InstanceData->MeshSkinWeightLookupBuffer->GetSRV());
 
 		uint32 EnabledFeaturesBits = 0;
-		EnabledFeaturesBits |= StaticBuffers->IsUseGpuUniformlyDistributedSampling() ? 1 : 0;
-		EnabledFeaturesBits |= StaticBuffers->IsSamplingRegionsAllAreaWeighted() ? 2 : 0;
-		EnabledFeaturesBits |= InstanceData->bUnlimitedBoneInfluences ? 4 : 0;
-		EnabledFeaturesBits |= StaticBuffers->HasMeshColors() ? 8 : 0;
+		EnabledFeaturesBits |= StaticBuffers->IsUseGpuUniformlyDistributedSampling() ? 0x01 : 0;
+		EnabledFeaturesBits |= StaticBuffers->IsSamplingRegionsAllAreaWeighted() ? 0x02 : 0;
+		EnabledFeaturesBits |= InstanceData->bUnlimitedBoneInfluences ? 0x04 : 0;
+		EnabledFeaturesBits |= StaticBuffers->HasMeshColors() ? 0x08 : 0;
+
+		// If we are allowed to read from the skin cache then find the information
+		ShaderParameters->DeformedCurrPositionBuffer = FNiagaraRenderer::GetDummyFloatBuffer();
+		ShaderParameters->DeformedPrevPositionBuffer = FNiagaraRenderer::GetDummyFloatBuffer();
+		ShaderParameters->DeformedTangentBuffer = FNiagaraRenderer::GetDummyFloat4Buffer();
+
+		if (InstanceData->bReadDeformedGeometry)
+		{
+			USkeletalMeshComponent* SkelComp = Cast<USkeletalMeshComponent>(StaticBuffers->GetSceneComponent());
+			FSkeletalMeshObject* SkelMeshObject = SkelComp ? SkelComp->MeshObject : nullptr;
+			if (SkelMeshObject)
+			{
+				FCachedGeometry CacheGeometry;
+				bool bCacheValid = SkelMeshObject->GetCachedGeometry(CacheGeometry);
+				bCacheValid &= CacheGeometry.Sections.Num() > 0;
+				bCacheValid &= StaticBuffers->GetLODIndex() == CacheGeometry.LODIndex;
+				if (bCacheValid)
+				{
+					if (CacheGeometry.Sections[0].PositionBuffer && (CacheGeometry.Sections[0].PositionBuffer != ShaderParameters->MeshVertexBuffer))
+					{
+						ShaderParameters->DeformedCurrPositionBuffer = CacheGeometry.Sections[0].PositionBuffer;
+						ShaderParameters->DeformedPrevPositionBuffer = CacheGeometry.Sections[0].PreviousPositionBuffer ? CacheGeometry.Sections[0].PreviousPositionBuffer : CacheGeometry.Sections[0].PositionBuffer;
+						EnabledFeaturesBits |= 0x10;
+					}
+					if (CacheGeometry.Sections[0].TangentBuffer && (CacheGeometry.Sections[0].TangentBuffer != ShaderParameters->MeshTangentBuffer))
+					{
+						ShaderParameters->DeformedTangentBuffer = CacheGeometry.Sections[0].TangentBuffer;
+						EnabledFeaturesBits |= 0x20;
+					}
+				}
+			}
+		}
 
 		FSkeletalMeshGpuDynamicBufferProxy* DynamicBuffers = InstanceData->DynamicBuffer;
 		check(DynamicBuffers);
diff --git a/Engine/Plugins/FX/Niagara/Source/Niagara/Public/NiagaraSettings.h b/Engine/Plugins/FX/Niagara/Source/Niagara/Public/NiagaraSettings.h
--- a/Engine/Plugins/FX/Niagara/Source/Niagara/Public/NiagaraSettings.h
+++ b/Engine/Plugins/FX/Niagara/Source/Niagara/Public/NiagaraSettings.h
@@ -225,6 +225,15 @@ class UNiagaraSettings : public UDeveloperSettings
 	UPROPERTY(config, EditAnywhere, Category = LightRenderer)
 	float DefaultLightInverseExposureBlend = 0.0f;
 
+	/* 
+	When enabled we will read deformed geometry if available, i.e. data from the deformed graph / skin cache
+	When disable we will only read from the default vertex data which does not include morph targets, skin, etc.
+	Changing this setting requires restarting the editor.
+	Note: Enabling this does add additional branches to the skel mesh data reading.
+	*/
+	UPROPERTY(config, EditAnywhere, Category = SkeletalMeshDI, meta = (DisplayName = "Support Reading Deformed Geometry", ConfigRestartRequired = true))
+	bool NDISkelMesh_SupportReadingDeformedGeometry = true;
+
 	UPROPERTY(config, EditAnywhere, Category = SkeletalMeshDI, meta = (DisplayName = "Support 16 bit index & weights", ToolTip = "Enabled support for 16 bit bone index & bone weight, optional to reduce shader complexity.  Changing this setting requires restarting the editor.", ConfigRestartRequired = true))
 	bool NDISkelMesh_Support16BitIndexWeight = true;
 
diff --git a/Engine/Source/Runtime/Engine/Private/GPUSkinCache.cpp b/Engine/Source/Runtime/Engine/Private/GPUSkinCache.cpp
--- a/Engine/Source/Runtime/Engine/Private/GPUSkinCache.cpp
+++ b/Engine/Source/Runtime/Engine/Private/GPUSkinCache.cpp
@@ -2254,6 +2254,17 @@ FRWBuffer* FGPUSkinCache::GetPreviousPositionBuffer(FGPUSkinCacheEntry const* En
 	return nullptr;
 }
 
+FRWBuffer* FGPUSkinCache::GetTangentBuffer(FGPUSkinCacheEntry const* Entry, uint32 SectionIndex)
+{
+	if (Entry)
+	{
+		FGPUSkinCacheEntry::FSectionDispatchData const& DispatchData = Entry->GetDispatchData()[SectionIndex];
+		FSkinCacheRWBuffer* SkinCacheRWBuffer = DispatchData.TangentBuffer;
+		return SkinCacheRWBuffer != nullptr ? &SkinCacheRWBuffer->Buffer : nullptr;
+	}
+	return nullptr;
+}
+
 uint32 FGPUSkinCache::GetUpdatedFrame(FGPUSkinCacheEntry const* Entry, uint32 SectionIndex)
 {
 	return Entry != nullptr ? Entry->GetDispatchData()[SectionIndex].UpdatedFrameNumber : 0;
diff --git a/Engine/Source/Runtime/Engine/Private/SkeletalRenderGPUSkin.cpp b/Engine/Source/Runtime/Engine/Private/SkeletalRenderGPUSkin.cpp
--- a/Engine/Source/Runtime/Engine/Private/SkeletalRenderGPUSkin.cpp
+++ b/Engine/Source/Runtime/Engine/Private/SkeletalRenderGPUSkin.cpp
@@ -2308,6 +2308,9 @@ bool FSkeletalMeshObjectGPUSkin::GetCachedGeometry(FCachedGeometry& OutCachedGeo
 			// Get the cached geometry SRVs from the skin cache.
 			CachedSection.PositionBuffer = FGPUSkinCache::GetPositionBuffer(SkinCacheEntry, SectionIndex)->SRV;
 			CachedSection.PreviousPositionBuffer = FGPUSkinCache::GetPreviousPositionBuffer(SkinCacheEntry, SectionIndex)->SRV;
+
+			FRWBuffer* TangentBuffer = FGPUSkinCache::GetTangentBuffer(SkinCacheEntry, SectionIndex);
+			CachedSection.TangentBuffer = TangentBuffer ? TangentBuffer->SRV : nullptr;
 		}
 		else
 		{
diff --git a/Engine/Source/Runtime/Engine/Public/CachedGeometry.h b/Engine/Source/Runtime/Engine/Public/CachedGeometry.h
--- a/Engine/Source/Runtime/Engine/Public/CachedGeometry.h
+++ b/Engine/Source/Runtime/Engine/Public/CachedGeometry.h
@@ -21,6 +21,7 @@ struct FCachedGeometry
 		
 		FRHIShaderResourceView* PositionBuffer = nullptr;			// Valid when the input comes from the GPUSkinCache (since it is doesn't use RDG)
 		FRHIShaderResourceView* PreviousPositionBuffer = nullptr;	// Valid when the input comes from the GPUSkinCache (since it is doesn't use RDG)
+		FRHIShaderResourceView* TangentBuffer = nullptr;			// Valid when the input comes from the GPUSkinCache (since it is doesn't use RDG)
 		
 		FRHIShaderResourceView* UVsBuffer = nullptr;
 		FRHIShaderResourceView* IndexBuffer = nullptr;
diff --git a/Engine/Source/Runtime/Engine/Public/GPUSkinCache.h b/Engine/Source/Runtime/Engine/Public/GPUSkinCache.h
--- a/Engine/Source/Runtime/Engine/Public/GPUSkinCache.h
+++ b/Engine/Source/Runtime/Engine/Public/GPUSkinCache.h
@@ -407,6 +407,7 @@ public:
 	FGPUSkinCacheEntry const* GetSkinCacheEntry(uint32 ComponentId) const;
 	static FRWBuffer* GetPositionBuffer(FGPUSkinCacheEntry const* Entry, uint32 SectionIndex);
 	static FRWBuffer* GetPreviousPositionBuffer(FGPUSkinCacheEntry const* Entry, uint32 SectionIndex);
+	static FRWBuffer* GetTangentBuffer(FGPUSkinCacheEntry const* Entry, uint32 SectionIndex);
 	static uint32 GetUpdatedFrame(FGPUSkinCacheEntry const* Entry, uint32 SectionIndex);
 
 	// Deprecated function. Can remove include of CachedGeometry.h when this is removed.
