Don't bother waiting for compilations if we're exiting the commandlet.  Existing compilations will be aborted with the module's PreExit.

diff --git a/Engine/Plugins/FX/Niagara/Source/NiagaraEditor/Private/NiagaraEditorModule.cpp b/Engine/Plugins/FX/Niagara/Source/NiagaraEditor/Private/NiagaraEditorModule.cpp
--- a/Engine/Plugins/FX/Niagara/Source/NiagaraEditor/Private/NiagaraEditorModule.cpp
+++ b/Engine/Plugins/FX/Niagara/Source/NiagaraEditor/Private/NiagaraEditorModule.cpp
@@ -2079,7 +2079,7 @@ void FNiagaraEditorModule::AddReferencedObjects(FReferenceCollector& Collector)
 
 void FNiagaraEditorModule::OnPreGarbageCollection()
 {
-	if (IsRunningCommandlet())
+	if (IsRunningCommandlet() && !IsEngineExitRequested())
 	{
 		// For commandlets like GenerateDistillFileSetsCommandlet, they just load the package and do some hierarchy navigation within it 
 		// tracking sub-assets, then they garbage collect. Since nothing is holding onto the system at the root level, it will be summarily
