* Support custom visualizations for data interfaces stored in sim caches
* Added visualization for cached data channels

diff --git a/Engine/Plugins/FX/Niagara/Source/Niagara/Classes/NiagaraSimCache.h b/Engine/Plugins/FX/Niagara/Source/Niagara/Classes/NiagaraSimCache.h
--- a/Engine/Plugins/FX/Niagara/Source/Niagara/Classes/NiagaraSimCache.h
+++ b/Engine/Plugins/FX/Niagara/Source/Niagara/Classes/NiagaraSimCache.h
@@ -442,6 +442,12 @@ public:
 	UFUNCTION(BlueprintCallable, Category = NiagaraSimCache)
 	NIAGARA_API TArray<FName> GetEmitterNames() const;
 
+	/** Returns a list of data interfaces we have captured in the SimCache. */
+	NIAGARA_API TArray<FNiagaraVariableBase> GetStoredDataInterfaces() const;
+	
+	/** Returns the actual data we have captured in the SimCache for the given data interface. */
+	NIAGARA_API UObject* GetDataInterfaceStorageObject(const FNiagaraVariableBase& DataInterface) const;
+
 	/**
 	Get number of active instances for the emitter at the given frame.
 	An EmitterIndex or INDEX_NONE will return information about the system instance.
diff --git a/Engine/Plugins/FX/Niagara/Source/Niagara/Classes/NiagaraSimCacheCustomStorageInterface.h b/Engine/Plugins/FX/Niagara/Source/Niagara/Classes/NiagaraSimCacheCustomStorageInterface.h
--- a/Engine/Plugins/FX/Niagara/Source/Niagara/Classes/NiagaraSimCacheCustomStorageInterface.h
+++ b/Engine/Plugins/FX/Niagara/Source/Niagara/Classes/NiagaraSimCacheCustomStorageInterface.h
@@ -19,6 +19,8 @@ class UNiagaraSimCacheCustomStorageInterface : public UInterface
 /**
 The current API for storing data inside a simulation cache.
 This is highly experimental and the API will change as we split editor / runtime data storage.
+
+See INiagaraDataInterfaceSimCacheVisualizer to implement a custom visualizer widget for the stored data.
 */
 class INiagaraSimCacheCustomStorageInterface
 {
diff --git a/Engine/Plugins/FX/Niagara/Source/Niagara/Internal/DataInterface/NiagaraDataInterfaceDataChannelWrite.h b/Engine/Plugins/FX/Niagara/Source/Niagara/Internal/DataInterface/NiagaraDataInterfaceDataChannelWrite.h
--- a/Engine/Plugins/FX/Niagara/Source/Niagara/Internal/DataInterface/NiagaraDataInterfaceDataChannelWrite.h
+++ b/Engine/Plugins/FX/Niagara/Source/Niagara/Internal/DataInterface/NiagaraDataInterfaceDataChannelWrite.h
@@ -72,7 +72,7 @@ struct FNDIDataChannelWriteSimCacheFrame
 	bool bVisibleToGPUSims = false;
 };
 
-UCLASS()
+UCLASS(MinimalAPI)
 class UNDIDataChannelWriteSimCacheData : public UObject
 {
 	GENERATED_BODY()
@@ -81,6 +81,9 @@ public:
 
 	UPROPERTY()
 	TArray<FNDIDataChannelWriteSimCacheFrame> FrameData;
+
+	UPROPERTY()
+	FSoftObjectPath DataChannelReference;
 };
 
 
diff --git a/Engine/Plugins/FX/Niagara/Source/Niagara/Private/DataInterface/NiagaraDataInterfaceDataChannelWrite.cpp b/Engine/Plugins/FX/Niagara/Source/Niagara/Private/DataInterface/NiagaraDataInterfaceDataChannelWrite.cpp
--- a/Engine/Plugins/FX/Niagara/Source/Niagara/Private/DataInterface/NiagaraDataInterfaceDataChannelWrite.cpp
+++ b/Engine/Plugins/FX/Niagara/Source/Niagara/Private/DataInterface/NiagaraDataInterfaceDataChannelWrite.cpp
@@ -605,6 +605,7 @@ bool UNiagaraDataInterfaceDataChannelWrite::SimCacheWriteFrame(UObject* StorageO
 	if (UNDIDataChannelWriteSimCacheData* Storage = Cast<UNDIDataChannelWriteSimCacheData>(StorageObject))
 	{
 		ensure(Storage->FrameData.Num() == FrameIndex);
+		Storage->DataChannelReference = Channel.Get();
 		FNDIDataChannelWriteSimCacheFrame& FrameData = Storage->FrameData.AddDefaulted_GetRef();
 		
 		if (InstanceData->DataChannelData && ShouldPublish() && InstanceData->Data && InstanceData->Data->GetCurrentData() && InstanceData->Data->GetCurrentData()->GetNumInstances() > 0)
diff --git a/Engine/Plugins/FX/Niagara/Source/Niagara/Private/NiagaraSimCache.cpp b/Engine/Plugins/FX/Niagara/Source/Niagara/Private/NiagaraSimCache.cpp
--- a/Engine/Plugins/FX/Niagara/Source/Niagara/Private/NiagaraSimCache.cpp
+++ b/Engine/Plugins/FX/Niagara/Source/Niagara/Private/NiagaraSimCache.cpp
@@ -857,6 +857,22 @@ TArray<FName> UNiagaraSimCache::GetEmitterNames() const
 	return EmitterNames;
 }
 
+TArray<FNiagaraVariableBase> UNiagaraSimCache::GetStoredDataInterfaces() const
+{
+	TArray<FNiagaraVariableBase> DataInterfaces;
+	DataInterfaceStorage.GenerateKeyArray(DataInterfaces);
+	return DataInterfaces;
+}
+
+UObject* UNiagaraSimCache::GetDataInterfaceStorageObject(const FNiagaraVariableBase& DataInterface) const
+{
+	if (const TObjectPtr<UObject>* StoredObject = DataInterfaceStorage.Find(DataInterface))
+	{
+		return StoredObject->Get();
+	}
+	return nullptr;
+}
+
 int UNiagaraSimCache::GetEmitterNumInstances(int32 EmitterIndex, int32 FrameIndex) const
 {
 	if ( CacheFrames.IsValidIndex(FrameIndex) )
diff --git a/Engine/Plugins/FX/Niagara/Source/Niagara/Public/NiagaraTypes.h b/Engine/Plugins/FX/Niagara/Source/Niagara/Public/NiagaraTypes.h
--- a/Engine/Plugins/FX/Niagara/Source/Niagara/Public/NiagaraTypes.h
+++ b/Engine/Plugins/FX/Niagara/Source/Niagara/Public/NiagaraTypes.h
@@ -1731,7 +1731,6 @@ struct FNiagaraVariable : public FNiagaraVariableBase
 	template<typename T>
 	void SetValue(const T& Data)
 	{
-		static_assert(!TIsUECoreVariant<T, double>::Value, "Double core variant, please use SetDoubleValue.");
 		check(sizeof(T) == TypeDefHandle->GetSize());
 		AllocateData();
 		FMemory::Memcpy(VarData.GetData(), &Data, VarData.Num());
@@ -1740,7 +1739,6 @@ struct FNiagaraVariable : public FNiagaraVariableBase
 	template<typename T>
 	T GetValue() const
 	{
-		static_assert(!TIsUECoreVariant<T, double>::Value, "Double core variant, please use GetDoubleValue.");
 		check(sizeof(T) == TypeDefHandle->GetSize());
 		check(IsDataAllocated());
 		T Value;
diff --git a/Engine/Plugins/FX/Niagara/Source/NiagaraEditor/Private/Customizations/NiagaraDataInterfaceSimCacheVisualizer.cpp b/Engine/Plugins/FX/Niagara/Source/NiagaraEditor/Private/Customizations/NiagaraDataInterfaceSimCacheVisualizer.cpp
new file mode 100644
--- /dev/null
+++ b/Engine/Plugins/FX/Niagara/Source/NiagaraEditor/Private/Customizations/NiagaraDataInterfaceSimCacheVisualizer.cpp
@@ -0,0 +1,11 @@
+﻿// Copyright Epic Games, Inc. All Rights Reserved.
+
+
+#include "Customizations/NiagaraDataInterfaceSimCacheVisualizer.h"
+#include "ViewModels/NiagaraSimCacheViewModel.h"
+
+
+TSharedPtr<SWidget> INiagaraDataInterfaceSimCacheVisualizer::CreateWidgetFor(UObject*, TSharedPtr<FNiagaraSimCacheViewModel>)
+{
+	return SNullWidget::NullWidget;
+}
diff --git a/Engine/Plugins/FX/Niagara/Source/NiagaraEditor/Private/Customizations/SimCache/FNiagaraDataChannelSimCacheVisualizer.cpp b/Engine/Plugins/FX/Niagara/Source/NiagaraEditor/Private/Customizations/SimCache/FNiagaraDataChannelSimCacheVisualizer.cpp
new file mode 100644
--- /dev/null
+++ b/Engine/Plugins/FX/Niagara/Source/NiagaraEditor/Private/Customizations/SimCache/FNiagaraDataChannelSimCacheVisualizer.cpp
@@ -0,0 +1,312 @@
+﻿// Copyright Epic Games, Inc. All Rights Reserved.
+
+#include "FNiagaraDataChannelSimCacheVisualizer.h"
+
+#include "NiagaraEditorModule.h"
+#include "DataInterface/NiagaraDataInterfaceDataChannelWrite.h"
+#include "ViewModels/NiagaraSimCacheViewModel.h"
+#include "Widgets/Colors/SColorBlock.h"
+
+#define LOCTEXT_NAMESPACE "NiagaraDataChannelCacheView"
+
+namespace DataChannelVisualizer
+{
+	const FNDIDataChannelWriteSimCacheFrame& GetFrame(TSharedPtr<FNiagaraSimCacheViewModel> SimCacheViewModel, UNDIDataChannelWriteSimCacheData* DataChannelData)
+	{
+		int32 FrameIndex = FMath::Clamp(SimCacheViewModel->GetFrameIndex(), 0, DataChannelData->FrameData.Num() - 1);
+		return DataChannelData->FrameData[FrameIndex];
+	}
+
+	FText ChannelVariableToString(const FNiagaraVariable& Var)
+	{
+		if (!Var.IsDataAllocated())
+		{
+			return FText();
+		}
+		FNumberFormattingOptions NumberFormat = FNumberFormattingOptions()
+				.SetUseGrouping(false)
+				.SetMaximumFractionalDigits(4)
+				.SetMinimumFractionalDigits(0);
+		
+		if (Var.GetType() == FNiagaraTypeHelper::GetDoubleDef())
+		{
+			double Value = Var.GetValue<double>();
+			return FText::Format(FText::FromString("{0}"), FText::AsNumber(Value, &NumberFormat));
+		}
+		if (Var.GetType() == FNiagaraTypeDefinition::GetBoolDef())
+		{
+			bool Value = Var.GetValue<bool>();
+			return Value ? LOCTEXT("true", "true") : LOCTEXT("false", "false");
+		}
+		if (Var.GetType() == FNiagaraTypeDefinition::GetIntDef())
+		{
+			int32 Value = Var.GetValue<int32>();
+			return FText::AsNumber(Value);
+		}
+		if (Var.GetType() == FNiagaraTypeHelper::GetVectorDef())
+		{
+			FVector Value = Var.GetValue<FVector>();
+			return FText::Format(FText::FromString("X = {0}\nY = {1}\nZ = {2}"), FText::AsNumber(Value.X, &NumberFormat), FText::AsNumber(Value.Y, &NumberFormat), FText::AsNumber(Value.Z, &NumberFormat));
+		}
+		if (Var.GetType() == FNiagaraTypeHelper::GetVector2DDef())
+		{
+			FVector2D Value = Var.GetValue<FVector2D>();
+			return FText::Format(FText::FromString("X = {0}\nY = {1}"), FText::AsNumber(Value.X, &NumberFormat), FText::AsNumber(Value.Y, &NumberFormat));
+		}
+		if (Var.GetType() == FNiagaraTypeHelper::GetVector4Def())
+		{
+			FVector4 Value = Var.GetValue<FVector4>();
+			return FText::Format(FText::FromString("X = {0}, Y = {1}\nZ = {2}, W = {3}"), FText::AsNumber(Value.X, &NumberFormat), FText::AsNumber(Value.Y, &NumberFormat), FText::AsNumber(Value.Z, &NumberFormat), FText::AsNumber(Value.W, &NumberFormat));
+		}
+		if (Var.GetType() == FNiagaraTypeHelper::GetQuatDef())
+		{
+			FQuat Value = Var.GetValue<FQuat>();
+			return FText::Format(FText::FromString("X = {0}, Y = {1}\nZ = {2}, W = {3}"), FText::AsNumber(Value.X, &NumberFormat), FText::AsNumber(Value.Y, &NumberFormat), FText::AsNumber(Value.Z, &NumberFormat), FText::AsNumber(Value.W, &NumberFormat));
+		}
+		if (Var.GetType() == FNiagaraTypeDefinition::GetColorDef())
+		{
+			FLinearColor Value = Var.GetValue<FLinearColor>();
+			return FText::Format(FText::FromString("R = {0}, G = {1}\nB = {2}, A = {3}"), FText::AsNumber(Value.R, &NumberFormat), FText::AsNumber(Value.G, &NumberFormat), FText::AsNumber(Value.B, &NumberFormat), FText::AsNumber(Value.A, &NumberFormat));
+		}
+		if (Var.GetType() == FNiagaraTypeDefinition::GetIDDef())
+		{
+			FNiagaraID Value = Var.GetValue<FNiagaraID>();
+			return FText::Format(LOCTEXT("ID_Format", "{0}/{1}"), FText::AsNumber(Value.Index), FText::AsNumber(Value.AcquireTag));
+		}
+		if (Var.GetType().IsEnum())
+		{
+			int32 Value = Var.GetValue<int32>();
+			return Var.GetType().GetEnum()->GetDisplayNameTextByValue(Value);
+		}
+		return LOCTEXT("UnsupportedType", "Unsupported Type");
+	}
+}
+
+class SNiagaraDataChannelCacheView : public SCompoundWidget
+{
+public:
+	SLATE_BEGIN_ARGS(SNiagaraDataChannelCacheView)
+	{}
+	SLATE_END_ARGS()
+
+	virtual ~SNiagaraDataChannelCacheView() override;
+
+	void Construct(const FArguments& InArgs, TSharedPtr<FNiagaraSimCacheViewModel> InViewModel, UNDIDataChannelWriteSimCacheData* InData);
+
+	TSharedRef<ITableRow> MakeRowWidget(const TSharedPtr<int32> RowIndexPtr, const TSharedRef<STableViewBase>& OwnerTable) const;
+
+	const FNDIDataChannelWriteSimCacheFrame& GetFrame() const
+	{
+		return DataChannelVisualizer::GetFrame(SimCacheViewModel, DataChannelData.Get());
+	}
+
+private:
+	void UpdateRows(bool Refresh = true);
+	
+	TSharedPtr<FNiagaraSimCacheViewModel> SimCacheViewModel;
+	TStrongObjectPtr<UNDIDataChannelWriteSimCacheData> DataChannelData;
+
+	TSharedPtr<SListView<TSharedPtr<int32>>> ListViewWidget;
+	TArray<TSharedPtr<int32>> RowItems;	
+	TSharedPtr<SHeaderRow> HeaderRowWidget;
+	TMap<FName, int32> NameToIndexMapping;
+};
+
+class SDataChannelVisualizerRowWidget : public SMultiColumnTableRow<TSharedPtr<int32>>
+{
+public:
+SLATE_BEGIN_ARGS(SDataChannelVisualizerRowWidget) {}
+	SLATE_ARGUMENT(TSharedPtr<int32>, RowIndexPtr)
+	SLATE_ARGUMENT(TSharedPtr<FNiagaraSimCacheViewModel>, SimCacheViewModel)
+	SLATE_ARGUMENT(UNDIDataChannelWriteSimCacheData*, Data)
+SLATE_END_ARGS()
+
+	void Construct(const FArguments& InArgs, const TSharedRef<STableViewBase>& InOwnerTableView, const TMap<FName, int32>& NameMapping);
+	virtual TSharedRef<SWidget> GenerateWidgetForColumn(const FName& InColumnName) override;
+
+	TSharedPtr<int32> RowIndexPtr;
+	TSharedPtr<FNiagaraSimCacheViewModel> SimCacheViewModel;
+	TStrongObjectPtr<UNDIDataChannelWriteSimCacheData> DataChannelData;
+	TMap<FName, int32> NameToIndexMapping;
+};
+
+SNiagaraDataChannelCacheView::~SNiagaraDataChannelCacheView()
+{
+	SimCacheViewModel->OnViewDataChanged().RemoveAll(this);
+}
+
+void SNiagaraDataChannelCacheView::Construct(const FArguments&, TSharedPtr<FNiagaraSimCacheViewModel> InViewModel, UNDIDataChannelWriteSimCacheData* InData)
+{
+	SimCacheViewModel = InViewModel;
+	DataChannelData.Reset(InData);
+
+	HeaderRowWidget = SNew(SHeaderRow);
+	
+	// Generate fixed columns
+	HeaderRowWidget->AddColumn(
+		SHeaderRow::Column(FName("_ElementIndex"))
+		.DefaultLabel(LOCTEXT("ElementName", "Element Index"))
+		.HAlignHeader(HAlign_Center)
+		.VAlignHeader(VAlign_Fill)
+		.HAlignCell(HAlign_Center)
+		.VAlignCell(VAlign_Center)
+		.ManualWidth(100)
+		.SortMode(EColumnSortMode::None)
+	);
+	
+	// Generate a column for each variable
+	for (int i = 0; i < GetFrame().VariableData.Num(); i++)
+	{
+		const FNDIDataChannelWriteSimCacheFrameBuffer& Buffer = GetFrame().VariableData[i];
+		NameToIndexMapping.Add(Buffer.SourceVar.GetName(), i);
+		HeaderRowWidget->AddColumn(
+			SHeaderRow::Column(Buffer.SourceVar.GetName())
+			.DefaultLabel(FText::FromName(Buffer.SourceVar.GetName()))
+			.HAlignHeader(HAlign_Center)
+			.VAlignHeader(VAlign_Fill)
+			.HAlignCell(HAlign_Center)
+			.VAlignCell(VAlign_Center)
+			.FillWidth(1)
+			.ManualWidth(150)
+			.SortMode(EColumnSortMode::None)
+		);
+	}
+
+	SAssignNew(ListViewWidget, SListView<TSharedPtr<int32>>)
+		.ListItemsSource(&RowItems)
+		.OnGenerateRow(this, &SNiagaraDataChannelCacheView::MakeRowWidget)
+		.Visibility(EVisibility::Visible)
+		.SelectionMode(ESelectionMode::Single)
+		.ConsumeMouseWheel(EConsumeMouseWheel::Never)
+		.ScrollbarVisibility(EVisibility::Collapsed)
+		.HeaderRow(HeaderRowWidget);
+
+	ChildSlot
+	[
+		SNew(SVerticalBox)
+		+ SVerticalBox::Slot()
+		.AutoHeight()
+		[
+			SNew(STextBlock)
+			.Margin(10)
+			.Text_Lambda([this]()
+			{
+				const FNDIDataChannelWriteSimCacheFrame& FrameData = GetFrame();
+				if (FrameData.NumElements == 0)
+				{
+					return LOCTEXT("NoData", "This frame contains no data channel writes.");
+				}
+				return FText::Format(LOCTEXT("WriteSummary", "This frame writes {0} elements to the data channel {4}.\n[{1}] Visible to Blueprint\n[{2}] Visible to CPU systems\n[{3}] Visible to GPU systems"),
+					FText::AsNumber(FrameData.NumElements),
+					FText::FromString(FrameData.bVisibleToGame ? TEXT("X") : TEXT("  ")),
+					FText::FromString(FrameData.bVisibleToCPUSims ? TEXT("X") : TEXT("  ")),
+					FText::FromString(FrameData.bVisibleToGPUSims ? TEXT("X") : TEXT("  ")),
+					FText::FromString(DataChannelData->DataChannelReference.GetAssetPathString()));
+			})
+		]
+		+ SVerticalBox::Slot()
+		.AutoHeight()
+		[
+			ListViewWidget.ToSharedRef()
+		]
+	];
+
+	UpdateRows();
+	SimCacheViewModel->OnViewDataChanged().AddSP(this, &SNiagaraDataChannelCacheView::UpdateRows);
+}
+
+TSharedRef<ITableRow> SNiagaraDataChannelCacheView::MakeRowWidget(const TSharedPtr<int32> RowIndexPtr, const TSharedRef<STableViewBase>& OwnerTable) const
+{
+	return
+		SNew(SDataChannelVisualizerRowWidget, OwnerTable, NameToIndexMapping)
+		.RowIndexPtr(RowIndexPtr)
+		.SimCacheViewModel(SimCacheViewModel)
+		.Data(DataChannelData.Get());
+}
+
+void SNiagaraDataChannelCacheView::UpdateRows(bool)
+{
+	const FNDIDataChannelWriteSimCacheFrame& FrameData = GetFrame();
+	RowItems.Reset(FrameData.NumElements);
+	for (int32 i = 0; i < FrameData.NumElements; ++i)
+	{
+		RowItems.Emplace(MakeShared<int32>(i));
+	}
+	ListViewWidget->RequestListRefresh();
+}
+
+void SDataChannelVisualizerRowWidget::Construct(const FArguments& InArgs, const TSharedRef<STableViewBase>& InOwnerTableView, const TMap<FName, int32>& NameMapping)
+{
+	RowIndexPtr = InArgs._RowIndexPtr;
+	SimCacheViewModel = InArgs._SimCacheViewModel;
+	DataChannelData.Reset(InArgs._Data);
+	NameToIndexMapping = NameMapping;
+
+	SMultiColumnTableRow::Construct(
+		FSuperRowType::FArguments()
+		.Style(FAppStyle::Get(), "DataTableEditor.CellListViewRow"),
+		InOwnerTableView
+	);
+}
+
+TSharedRef<SWidget> SDataChannelVisualizerRowWidget::GenerateWidgetForColumn(const FName& InColumnName)
+{
+	if(!SimCacheViewModel.Get()->IsCacheValid())
+	{
+		return SNullWidget::NullWidget;
+	}
+
+	const int32 InstanceIndex = *RowIndexPtr;
+	if (InColumnName == FName("_ElementIndex"))
+	{
+		return SNew(STextBlock)
+			.Text(FText::AsNumber(InstanceIndex));
+	}
+
+	int32 Index = NameToIndexMapping.FindChecked(InColumnName);
+	const FNDIDataChannelWriteSimCacheFrameBuffer& Buffer = DataChannelVisualizer::GetFrame(SimCacheViewModel, DataChannelData.Get()).VariableData[Index];
+	FNiagaraVariable DataVar = Buffer.SourceVar;
+	if (DataVar.GetType() == FNiagaraTypeDefinition::GetPositionDef())
+	{
+		DataVar.SetType(FNiagaraTypeHelper::GetVectorDef());
+	}
+	else
+	{
+		DataVar.SetType(FNiagaraTypeHelper::GetLWCType(DataVar.GetType()));
+	}
+	DataVar.SetData(Buffer.Data.GetData() + (Buffer.Size * InstanceIndex));
+
+	FText DisplayText = DataChannelVisualizer::ChannelVariableToString(DataVar);
+	if (DataVar.GetType() == FNiagaraTypeDefinition::GetColorDef())
+	{
+		return SNew(SHorizontalBox)
+			+ SHorizontalBox::Slot()
+			.Padding(FMargin(10, 0))
+			.AutoWidth()
+			[
+				SNew(SColorBlock)
+				.Color(DataVar.GetValue<FLinearColor>())
+				.ShowBackgroundForAlpha(true)
+			]
+			+ SHorizontalBox::Slot()
+			[
+				SNew(STextBlock)
+				.Margin(FMargin(0, 0, 15, 0))
+				.Text(DisplayText)
+			];
+	}
+	return SNew(STextBlock)
+		.Margin(FMargin(15, 0))
+		.Text(DisplayText);
+}
+
+TSharedPtr<SWidget> FNiagaraDataChannelSimCacheVisualizer::CreateWidgetFor(UObject* CachedData, TSharedPtr<FNiagaraSimCacheViewModel> ViewModel)
+{
+	if (UNDIDataChannelWriteSimCacheData* DataChannelData = Cast<UNDIDataChannelWriteSimCacheData>(CachedData))
+	{
+		return SNew(SNiagaraDataChannelCacheView, ViewModel, DataChannelData);
+	}
+	return TSharedPtr<SWidget>();
+}
+
+#undef LOCTEXT_NAMESPACE
diff --git a/Engine/Plugins/FX/Niagara/Source/NiagaraEditor/Private/Customizations/SimCache/FNiagaraDataChannelSimCacheVisualizer.h b/Engine/Plugins/FX/Niagara/Source/NiagaraEditor/Private/Customizations/SimCache/FNiagaraDataChannelSimCacheVisualizer.h
new file mode 100644
--- /dev/null
+++ b/Engine/Plugins/FX/Niagara/Source/NiagaraEditor/Private/Customizations/SimCache/FNiagaraDataChannelSimCacheVisualizer.h
@@ -0,0 +1,16 @@
+﻿// Copyright Epic Games, Inc. All Rights Reserved.
+
+#pragma once
+
+#include "Customizations/NiagaraDataInterfaceSimCacheVisualizer.h"
+
+/*
+ * Provides a custom widget to show data channel writes stored in the sim cache in a table.
+ */
+class FNiagaraDataChannelSimCacheVisualizer : public INiagaraDataInterfaceSimCacheVisualizer
+{
+public:
+	virtual ~FNiagaraDataChannelSimCacheVisualizer() override = default;
+	
+	virtual TSharedPtr<SWidget> CreateWidgetFor(UObject* CachedData, TSharedPtr<FNiagaraSimCacheViewModel> ViewModel) override;
+};
diff --git a/Engine/Plugins/FX/Niagara/Source/NiagaraEditor/Private/NiagaraEditorModule.cpp b/Engine/Plugins/FX/Niagara/Source/NiagaraEditor/Private/NiagaraEditorModule.cpp
--- a/Engine/Plugins/FX/Niagara/Source/NiagaraEditor/Private/NiagaraEditorModule.cpp
+++ b/Engine/Plugins/FX/Niagara/Source/NiagaraEditor/Private/NiagaraEditorModule.cpp
@@ -161,8 +161,7 @@
 #include "Misc/ScopedSlowTask.h"
 
 #include "NiagaraActions.h"
-
-#include "NiagaraDataChannel.h"
+#include "Customizations/SimCache/FNiagaraDataChannelSimCacheVisualizer.h"
 
 #include "Engine/AssetManager.h"
 #include "ViewModels/HierarchyEditor/NiagaraHierarchyCommands.h"
@@ -1424,6 +1423,9 @@ void FNiagaraEditorModule::StartupModule()
 	MessageManager->RegisterMessageTopic(FNiagaraMessageTopics::CompilerTopicName);
 	MessageManager->RegisterMessageTopic(FNiagaraMessageTopics::ObjectTopicName);
 
+	// Register sim cache visualizers
+	RegisterDataInterfaceCacheVisualizer(UNiagaraDataInterfaceDataChannelWrite::StaticClass(), MakeShared<FNiagaraDataChannelSimCacheVisualizer>());
+
 #if NIAGARA_PERF_BASELINES
 	UNiagaraEffectType::OnGeneratePerfBaselines().BindRaw(this, &FNiagaraEditorModule::GeneratePerfBaselines);
 #endif
@@ -1720,6 +1722,28 @@ TSharedRef<INiagaraEditorWidgetProvider> FNiagaraEditorModule::GetWidgetProvider
 	return WidgetProvider.ToSharedRef();
 }
 
+void FNiagaraEditorModule::RegisterDataInterfaceCacheVisualizer(UClass* DataInterfaceClass, TSharedRef<INiagaraDataInterfaceSimCacheVisualizer> InCacheVisualizer)
+{
+	DataInterfaceVisualizers.FindOrAdd(DataInterfaceClass).AddUnique(InCacheVisualizer);
+}
+
+void FNiagaraEditorModule::UnregisterDataInterfaceCacheVisualizer(UClass* DataInterfaceClass, TSharedRef<INiagaraDataInterfaceSimCacheVisualizer> InCacheVisualizer)
+{
+	if (TArray<TSharedRef<INiagaraDataInterfaceSimCacheVisualizer>>* CacheVisualizers = DataInterfaceVisualizers.Find(DataInterfaceClass))
+	{
+		CacheVisualizers->Remove(InCacheVisualizer);
+	}
+}
+
+TArrayView<TSharedRef<INiagaraDataInterfaceSimCacheVisualizer>> FNiagaraEditorModule::FindDataInterfaceCacheVisualizer(UClass* DataInterfaceClass)
+{
+	if (TArray<TSharedRef<INiagaraDataInterfaceSimCacheVisualizer>>* CacheVisualizers = DataInterfaceVisualizers.Find(DataInterfaceClass))
+	{
+		return *CacheVisualizers;
+	}
+	return TArrayView<TSharedRef<INiagaraDataInterfaceSimCacheVisualizer>>();
+}
+
 TSharedRef<FNiagaraScriptMergeManager> FNiagaraEditorModule::GetScriptMergeManager() const
 {
 	return ScriptMergeManager.ToSharedRef();
diff --git a/Engine/Plugins/FX/Niagara/Source/NiagaraEditor/Private/NiagaraEditorStyle.cpp b/Engine/Plugins/FX/Niagara/Source/NiagaraEditor/Private/NiagaraEditorStyle.cpp
--- a/Engine/Plugins/FX/Niagara/Source/NiagaraEditor/Private/NiagaraEditorStyle.cpp
+++ b/Engine/Plugins/FX/Niagara/Source/NiagaraEditor/Private/NiagaraEditorStyle.cpp
@@ -819,7 +819,7 @@ void FNiagaraEditorStyle::InitSimCacheEditor()
 		.SetInactiveBrush(CORE_IMAGE_BRUSH("Common/Selection", Icon8x8, SelectionColor_Inactive))
 		.SetInactiveHoveredBrush(CORE_IMAGE_BRUSH("Common/Selection", Icon8x8, SelectionColor_Inactive)));
 
-	Set("NiagaraEditor.SimCache.EmitterItem", FTableRowStyle()
+	FTableRowStyle EmitterStyle = FTableRowStyle()
 		.SetEvenRowBackgroundBrush(CORE_IMAGE_BRUSH("Common/Selection", Icon8x8, SystemBackground))
 		.SetEvenRowBackgroundHoveredBrush(CORE_IMAGE_BRUSH("Common/Selection", Icon8x8, SelectionColor_Inactive))
 		.SetOddRowBackgroundBrush(CORE_IMAGE_BRUSH("Common/Selection", Icon8x8, SystemBackground))
@@ -828,7 +828,8 @@ void FNiagaraEditorStyle::InitSimCacheEditor()
 		.SetActiveBrush(CORE_IMAGE_BRUSH("Common/Selection", Icon8x8, SelectionColor))
 		.SetActiveHoveredBrush(CORE_IMAGE_BRUSH("Common/Selection", Icon8x8, SelectionColor))
 		.SetInactiveBrush(CORE_IMAGE_BRUSH("Common/Selection", Icon8x8, SelectionColor_Inactive))
-		.SetInactiveHoveredBrush(CORE_IMAGE_BRUSH("Common/Selection", Icon8x8, SelectionColor_Inactive)));
+		.SetInactiveHoveredBrush(CORE_IMAGE_BRUSH("Common/Selection", Icon8x8, SelectionColor_Inactive));
+	Set("NiagaraEditor.SimCache.EmitterItem", EmitterStyle);
 
 	Set("NiagaraEditor.SimCache.ComponentItem", FTableRowStyle()
 		.SetEvenRowBackgroundBrush(FSlateNoResource())
@@ -841,6 +842,9 @@ void FNiagaraEditorStyle::InitSimCacheEditor()
 		.SetInactiveBrush(CORE_IMAGE_BRUSH("Common/Selection", Icon8x8, SelectionColor_Inactive))
 		.SetInactiveHoveredBrush(CORE_IMAGE_BRUSH("Common/Selection", Icon8x8, SelectionColor_Inactive)));
 
+	FTableRowStyle DataInterfaceStyle = EmitterStyle;
+	Set("NiagaraEditor.SimCache.DataInterfaceItem", DataInterfaceStyle);
+
 	const FLinearColor NormalColor(0.15, 0.15, 0.15, 1);
 
 	Set ("NiagaraEditor.SimCache.FilterToggleStyle", FCheckBoxStyle()
diff --git a/Engine/Plugins/FX/Niagara/Source/NiagaraEditor/Private/ViewModels/NiagaraSimCacheViewModel.cpp b/Engine/Plugins/FX/Niagara/Source/NiagaraEditor/Private/ViewModels/NiagaraSimCacheViewModel.cpp
--- a/Engine/Plugins/FX/Niagara/Source/NiagaraEditor/Private/ViewModels/NiagaraSimCacheViewModel.cpp
+++ b/Engine/Plugins/FX/Niagara/Source/NiagaraEditor/Private/ViewModels/NiagaraSimCacheViewModel.cpp
@@ -7,7 +7,6 @@
 #include "NiagaraDebuggerCommon.h"
 #include "Serialization/MemoryReader.h"
 #include "Widgets/SNiagaraSimCacheTreeView.h"
-#include "Serialization/ObjectAndNameAsStringProxyArchive.h"
 #include "UObject/Package.h"
 
 #define LOCTEXT_NAMESPACE "NiagaraSimCacheViewModel"
@@ -45,26 +44,12 @@ void FNiagaraSimCacheViewModel::Initialize(TWeakObjectPtr<UNiagaraSimCache> InSi
 	OnViewDataChangedDelegate.Broadcast(true);
 }
 
-void FNiagaraSimCacheViewModel::UpdateSimCache(const FNiagaraSystemSimCacheCaptureReply& Reply)
+void FNiagaraSimCacheViewModel::SetComponentFilters(const TArray<FString>& NewComponentFilterArray)
 {
-	UNiagaraSimCache* TempSimCache = nullptr;
-	
-	if (Reply.SimCacheData.Num() > 0)
-	{
-		TempSimCache = NewObject<UNiagaraSimCache>();
-
-		FMemoryReader ArReader(Reply.SimCacheData);
-		FObjectAndNameAsStringProxyArchive ProxyArReader(ArReader, false);
-		TempSimCache->Serialize(ProxyArReader);
-		bComponentFilterActive = false;
-		UpdateComponentInfos();
-		
-	}
-	else
-	{
-		UE_LOG(LogNiagaraEditor, Warning, TEXT("Debug Spreadsheet received empty sim cache data."));
-	}
-	Initialize(TempSimCache);
+	bComponentFilterActive = true;
+	ComponentFilterArray.Empty();
+	ComponentFilterArray.Append(NewComponentFilterArray);
+	OnViewDataChangedDelegate.Broadcast(true);
 }
 
 void FNiagaraSimCacheViewModel::SetupPreviewComponentAndInstance()
@@ -171,9 +156,19 @@ void FNiagaraSimCacheViewModel::SetFrameIndex(const int32 InFrameIndex)
 	OnViewDataChangedDelegate.Broadcast(false);
 }
 
-void FNiagaraSimCacheViewModel::SetEmitterIndex(const int32 InEmitterIndex)
+UObject* FNiagaraSimCacheViewModel::GetActiveDataInterfaceStorage() const
+{
+	if (SimCache)
+	{
+		return SimCache->GetDataInterfaceStorageObject(GetActiveDataInterface());
+	}
+	return nullptr;
+}
+
+void FNiagaraSimCacheViewModel::SetEmitterIndex(const int32 InEmitterIndex, FNiagaraVariableBase InActiveDataInterface)
 {
 	EmitterIndex = InEmitterIndex;
+	ActiveDataInterface = InActiveDataInterface;
 	UpdateCachedFrame();
 	UpdateCurrentEntries();
 	bComponentFilterActive = false;
@@ -439,7 +434,7 @@ void FNiagaraSimCacheViewModel::BuildEntries(TWeakPtr<SNiagaraSimCacheTreeView>
 	}
 		
 	BuildTreeItemChildren(SharedSystemTreeItem, OwningTreeView);
-		
+	
 	for(int32 i = 0; i < GetNumEmitterLayouts(); i++)
 	{
 		const TSharedRef<FNiagaraSimCacheEmitterTreeItem> CurrentEmitterItem = MakeShared<FNiagaraSimCacheEmitterTreeItem>(OwningTreeView);
@@ -456,6 +451,32 @@ void FNiagaraSimCacheViewModel::BuildEntries(TWeakPtr<SNiagaraSimCacheTreeView>
 
 		BuildTreeItemChildren(CurrentEmitterItem, OwningTreeView);
 	}
+	if (SimCache)
+	{
+		for (const FNiagaraVariableBase& Var : SimCache->GetStoredDataInterfaces())
+		{
+			const TSharedRef<FNiagaraSimCacheDataInterfaceTreeItem> CurrentDataInterfaceItem = MakeShared<FNiagaraSimCacheDataInterfaceTreeItem>(OwningTreeView);
+			const TSharedRef<FNiagaraSimCacheOverviewDataInterfaceItem> CurrentDataInterfaceBufferItem = MakeShared<FNiagaraSimCacheOverviewDataInterfaceItem>();
+		
+			CurrentDataInterfaceItem->SetDisplayName(FText::FromName(Var.GetName()));
+			CurrentDataInterfaceItem->DataInterfaceReference = Var;
+			CurrentDataInterfaceBufferItem->SetDisplayName(FText::FromName(Var.GetName()));
+			CurrentDataInterfaceBufferItem->DataInterfaceReference = Var;
+
+			for(int32 i = 0; i < GetNumEmitterLayouts(); i++)
+			{
+				if (Var.GetName().ToString().StartsWith(GetEmitterLayoutName(i).ToString() + "."))
+				{
+					CurrentDataInterfaceItem->SetBufferIndex(i);
+					CurrentDataInterfaceBufferItem->SetBufferIndex(i);
+					break;
+				}
+			}
+		
+			RootEntries.Add(CurrentDataInterfaceItem);
+			BufferEntries.Add(CurrentDataInterfaceBufferItem);
+		}
+	}
 
 	UpdateCurrentEntries();
 }
diff --git a/Engine/Plugins/FX/Niagara/Source/NiagaraEditor/Private/ViewModels/NiagaraSimCacheViewModel.h b/Engine/Plugins/FX/Niagara/Source/NiagaraEditor/Private/ViewModels/NiagaraSimCacheViewModel.h
deleted file mode 100644
diff --git a/Engine/Plugins/FX/Niagara/Source/NiagaraEditor/Private/Widgets/SNiagaraSimCacheOverview.cpp b/Engine/Plugins/FX/Niagara/Source/NiagaraEditor/Private/Widgets/SNiagaraSimCacheOverview.cpp
--- a/Engine/Plugins/FX/Niagara/Source/NiagaraEditor/Private/Widgets/SNiagaraSimCacheOverview.cpp
+++ b/Engine/Plugins/FX/Niagara/Source/NiagaraEditor/Private/Widgets/SNiagaraSimCacheOverview.cpp
@@ -127,13 +127,14 @@ TSharedRef<ITableRow> SNiagaraSimCacheOverview::OnGenerateRowForItem(TSharedRef<
 	{
 		"NiagaraEditor.SimCache.SystemItem",
 		"NiagaraEditor.SimCache.EmitterItem",
-		"NiagaraEditor.SimCache.ComponentItem"
+		"NiagaraEditor.SimCache.ComponentItem",
+		"NiagaraEditor.SimCache.DataInterfaceItem"
 	};
 
 	ENiagaraSimCacheOverviewItemType StyleType = Item->GetType();
 	
 	return SNew(STableRow<TSharedRef<FNiagaraSimCacheOverviewItem>>, Owner)
-	.Style(FNiagaraEditorStyle::Get(), ItemStyles[(int32)StyleType])
+	.Style(FNiagaraEditorStyle::Get(), ItemStyles[static_cast<int32>(StyleType)])
 	.Padding(1.0f)
 	[
 		SNew(SNiagaraSimCacheBufferItem)
@@ -141,12 +142,11 @@ TSharedRef<ITableRow> SNiagaraSimCacheOverview::OnGenerateRowForItem(TSharedRef<
 	];
 }
 
-void SNiagaraSimCacheOverview::OnListSelectionChanged(TSharedPtr<FNiagaraSimCacheOverviewItem> Item,
-	ESelectInfo::Type SelectInfo)
+void SNiagaraSimCacheOverview::OnListSelectionChanged(TSharedPtr<FNiagaraSimCacheOverviewItem> Item, ESelectInfo::Type)
 {
-	if(Item.IsValid())
+	if (Item.IsValid())
 	{
-		ViewModel->SetEmitterIndex(Item->GetBufferIndex());
+		ViewModel->SetEmitterIndex(Item->GetBufferIndex(), Item->GetDataInterface());
 	}
 }
 
diff --git a/Engine/Plugins/FX/Niagara/Source/NiagaraEditor/Private/Widgets/SNiagaraSimCacheOverview.h b/Engine/Plugins/FX/Niagara/Source/NiagaraEditor/Private/Widgets/SNiagaraSimCacheOverview.h
--- a/Engine/Plugins/FX/Niagara/Source/NiagaraEditor/Private/Widgets/SNiagaraSimCacheOverview.h
+++ b/Engine/Plugins/FX/Niagara/Source/NiagaraEditor/Private/Widgets/SNiagaraSimCacheOverview.h
@@ -2,6 +2,7 @@
 
 #pragma once
 
+#include "NiagaraTypes.h"
 #include "Widgets/SCompoundWidget.h"
 #include "Widgets/SWidget.h"
 #include "Widgets/Text/STextBlock.h"
@@ -17,6 +18,7 @@ enum class ENiagaraSimCacheOverviewItemType
 	System,
 	Emitter,
 	Component,
+	DataInterface,
 	MAX,
 };
 
@@ -34,7 +36,7 @@ struct FNiagaraSimCacheOverviewItem : public TSharedFromThis<FNiagaraSimCacheOve
 		return DisplayName;
 	}
 
-	void SetDisplayName(FText NewName)
+	void SetDisplayName(const FText& NewName)
 	{
 		DisplayName = NewName;
 	}
@@ -49,6 +51,11 @@ struct FNiagaraSimCacheOverviewItem : public TSharedFromThis<FNiagaraSimCacheOve
 		BufferIndex = NewIndex;
 	}
 
+	virtual FNiagaraVariableBase GetDataInterface()
+	{
+		return FNiagaraVariableBase();
+	}
+
 	virtual ENiagaraSimCacheOverviewItemType GetType () { return ENiagaraSimCacheOverviewItemType::MAX; }
 
 	virtual TSharedRef<SWidget> GetRowWidget()
@@ -60,16 +67,24 @@ struct FNiagaraSimCacheOverviewItem : public TSharedFromThis<FNiagaraSimCacheOve
 	FText DisplayName;
 };
 
-struct FNiagaraSimCacheOverviewSystemItem : public FNiagaraSimCacheOverviewItem
+struct FNiagaraSimCacheOverviewSystemItem : FNiagaraSimCacheOverviewItem
 {
 	virtual ENiagaraSimCacheOverviewItemType GetType () override { return ENiagaraSimCacheOverviewItemType::System; }
 };
 
-struct FNiagaraSimCacheOverviewEmitterItem : public FNiagaraSimCacheOverviewItem
+struct FNiagaraSimCacheOverviewEmitterItem : FNiagaraSimCacheOverviewItem
 {
 	virtual ENiagaraSimCacheOverviewItemType GetType() override {return ENiagaraSimCacheOverviewItemType::Emitter; }
 };
 
+struct FNiagaraSimCacheOverviewDataInterfaceItem : FNiagaraSimCacheOverviewItem
+{
+	virtual ENiagaraSimCacheOverviewItemType GetType() override {return ENiagaraSimCacheOverviewItemType::DataInterface; }
+	virtual FNiagaraVariableBase GetDataInterface() override { return DataInterfaceReference; }
+
+	FNiagaraVariableBase DataInterfaceReference;
+};
+
 class SNiagaraSimCacheOverview : public SCompoundWidget
 {
 public:
diff --git a/Engine/Plugins/FX/Niagara/Source/NiagaraEditor/Private/Widgets/SNiagaraSimCacheTreeView.cpp b/Engine/Plugins/FX/Niagara/Source/NiagaraEditor/Private/Widgets/SNiagaraSimCacheTreeView.cpp
--- a/Engine/Plugins/FX/Niagara/Source/NiagaraEditor/Private/Widgets/SNiagaraSimCacheTreeView.cpp
+++ b/Engine/Plugins/FX/Niagara/Source/NiagaraEditor/Private/Widgets/SNiagaraSimCacheTreeView.cpp
@@ -209,13 +209,14 @@ TSharedRef<ITableRow> SNiagaraSimCacheTreeView::OnGenerateRow(TSharedRef<FNiagar
 	{
 		"NiagaraEditor.SimCache.SystemItem",
 		"NiagaraEditor.SimCache.EmitterItem",
-		"NiagaraEditor.SimCache.ComponentItem"
+		"NiagaraEditor.SimCache.ComponentItem",
+		"NiagaraEditor.SimCache.DataInterfaceItem"
 	};
 
 	ENiagaraSimCacheOverviewItemType StyleType = Item->GetType();
 	
 	return SNew(STableRow<TSharedRef<FNiagaraSimCacheTreeItem>>, OwnerTable)
-	.Style(FNiagaraEditorStyle::Get(), ItemStyles[(int32)StyleType])
+	.Style(FNiagaraEditorStyle::Get(), ItemStyles[static_cast<int32>(StyleType)])
 	[
 		SNew(SNiagaraSimCacheTreeItem)
 		.Item(Item)
@@ -334,4 +335,13 @@ void SNiagaraSimCacheTreeView::SelectAll()
 	UpdateStringFilters();
 }
 
+bool SNiagaraSimCacheTreeView::IsDataInterfaceViewActive() const
+{
+	if (ViewModel)
+	{
+		return ViewModel->GetActiveDataInterface().IsValid();
+	}
+	return false;
+}
+
 #undef LOCTEXT_NAMESPACE
diff --git a/Engine/Plugins/FX/Niagara/Source/NiagaraEditor/Private/Widgets/SNiagaraSimCacheTreeView.h b/Engine/Plugins/FX/Niagara/Source/NiagaraEditor/Private/Widgets/SNiagaraSimCacheTreeView.h
--- a/Engine/Plugins/FX/Niagara/Source/NiagaraEditor/Private/Widgets/SNiagaraSimCacheTreeView.h
+++ b/Engine/Plugins/FX/Niagara/Source/NiagaraEditor/Private/Widgets/SNiagaraSimCacheTreeView.h
@@ -73,6 +73,7 @@ SLATE_END_ARGS()
 	void ClearFilterSelection();
 	void UpdateStringFilters();
 	void SelectAll();
+	bool IsDataInterfaceViewActive() const;
 private:
 	TSharedRef<ITableRow> OnGenerateRow(TSharedRef<FNiagaraSimCacheTreeItem> Item, const TSharedRef<STableViewBase>& OwnerTable);
 
@@ -95,7 +96,7 @@ private:
 	TArray<TSharedRef<FNiagaraSimCacheTreeItem>> SelectionForFilter;
 };
 
-struct FNiagaraSimCacheTreeItem : public FNiagaraSimCacheOverviewItem
+struct FNiagaraSimCacheTreeItem : FNiagaraSimCacheOverviewItem
 {
 	FNiagaraSimCacheTreeItem(TWeakPtr<SNiagaraSimCacheTreeView> InOwner)
 	{
@@ -144,6 +145,14 @@ struct FNiagaraSimCacheTreeItem : public FNiagaraSimCacheOverviewItem
 			.VAlign(VAlign_Center)
 			[
 				SNew(SNiagaraSimCacheTreeViewFilterWidget, SharedThis(this), Owner)
+				.Visibility_Lambda([this]()
+				{
+					if (TSharedPtr<SNiagaraSimCacheTreeView> TreeView = Owner.Pin())
+					{
+						return TreeView->IsDataInterfaceViewActive() ? EVisibility::Collapsed : EVisibility::Visible; 
+					}
+					return EVisibility::Collapsed;
+				})
 			];
 	}
 
@@ -155,7 +164,7 @@ struct FNiagaraSimCacheTreeItem : public FNiagaraSimCacheOverviewItem
 	
 };
 
-struct FNiagaraSimCacheEmitterTreeItem : public FNiagaraSimCacheTreeItem
+struct FNiagaraSimCacheEmitterTreeItem : FNiagaraSimCacheTreeItem
 {
 	
 	FNiagaraSimCacheEmitterTreeItem(TWeakPtr<SNiagaraSimCacheTreeView> InOwner): FNiagaraSimCacheTreeItem(InOwner)
@@ -167,7 +176,20 @@ struct FNiagaraSimCacheEmitterTreeItem : public FNiagaraSimCacheTreeItem
 	virtual ENiagaraSimCacheOverviewItemType GetType () override { return ENiagaraSimCacheOverviewItemType::Emitter; }
 };
 
-struct FNiagaraSimCacheComponentTreeItem : public FNiagaraSimCacheTreeItem
+struct FNiagaraSimCacheDataInterfaceTreeItem : FNiagaraSimCacheTreeItem
+{
+	FNiagaraSimCacheDataInterfaceTreeItem(TWeakPtr<SNiagaraSimCacheTreeView> InOwner): FNiagaraSimCacheTreeItem(InOwner)
+	{
+	}
+
+	virtual ~FNiagaraSimCacheDataInterfaceTreeItem() override {}
+	
+	virtual ENiagaraSimCacheOverviewItemType GetType () override { return ENiagaraSimCacheOverviewItemType::DataInterface; }
+
+	FNiagaraVariableBase DataInterfaceReference;
+};
+
+struct FNiagaraSimCacheComponentTreeItem : FNiagaraSimCacheTreeItem
 {
 	
 	FNiagaraSimCacheComponentTreeItem(TWeakPtr<SNiagaraSimCacheTreeView> InOwner): FNiagaraSimCacheTreeItem(InOwner)
@@ -221,6 +243,14 @@ struct FNiagaraSimCacheComponentTreeItem : public FNiagaraSimCacheTreeItem
 		.AutoWidth()
 		[
 			SAssignNew(VisibilityWidget, SSimCacheTreeViewVisibilityWidget, SharedThis(this), Owner)
+			.Visibility_Lambda([this]()
+			{
+				if (TSharedPtr<SNiagaraSimCacheTreeView> TreeView = Owner.Pin())
+				{
+					return TreeView->IsDataInterfaceViewActive() ? EVisibility::Collapsed : EVisibility::Visible; 
+				}
+				return EVisibility::Collapsed;
+			})
 		];
 		
 		return Contents;
diff --git a/Engine/Plugins/FX/Niagara/Source/NiagaraEditor/Private/Widgets/SNiagaraSimCacheView.cpp b/Engine/Plugins/FX/Niagara/Source/NiagaraEditor/Private/Widgets/SNiagaraSimCacheView.cpp
--- a/Engine/Plugins/FX/Niagara/Source/NiagaraEditor/Private/Widgets/SNiagaraSimCacheView.cpp
+++ b/Engine/Plugins/FX/Niagara/Source/NiagaraEditor/Private/Widgets/SNiagaraSimCacheView.cpp
@@ -4,17 +4,14 @@
 #include "ViewModels/NiagaraSimCacheViewModel.h"
 
 #include "CoreMinimal.h"
+#include "NiagaraEditorModule.h"
 
-#include "Widgets/Input/SSpinBox.h"
 #include "Widgets/Layout/SScrollBox.h"
+#include "Widgets/Layout/SWidgetSwitcher.h"
 #include "Widgets/Views/SListView.h"
 #include "Widgets/Views/STableRow.h"
 #include "Widgets/Views/SHeaderRow.h"
 
-#include "IContentBrowserSingleton.h"
-#include "ContentBrowserModule.h"
-#include "SNiagaraSimCacheTreeView.h"
-
 #define LOCTEXT_NAMESPACE "NiagaraSimCacheView"
 
 static const FName NAME_Instance("Instance");
@@ -90,38 +87,65 @@ void SNiagaraSimCacheView::Construct(const FArguments& InArgs)
 		.AlwaysShowScrollbar(true)
 		.Thickness(12.0f)
 		.Orientation(Orient_Vertical);
-	
+	DataInterfaceScrollBar =
+		SNew(SScrollBar)
+		.AlwaysShowScrollbar(false)
+		.Thickness(12.0f)
+		.Orientation(Orient_Vertical);
+
+	//// Main Spreadsheet View
+	SAssignNew(ListViewWidget, SListView<TSharedPtr<int32>>)
+		.ListItemsSource(&RowItems)
+		.OnGenerateRow(this, &SNiagaraSimCacheView::MakeRowWidget)
+		.Visibility(EVisibility::Visible)
+		.SelectionMode(ESelectionMode::Single)
+		.ExternalScrollbar(VerticalScrollBar)
+		.ConsumeMouseWheel(EConsumeMouseWheel::Always)
+		.AllowOverscroll(EAllowOverscroll::No)
+		.HeaderRow(HeaderRowWidget);
+
 	// Widget
 	ChildSlot
 	[
 		SNew(SVerticalBox)
-		//// Main Spreadsheet View
+		
 		+ SVerticalBox::Slot()
 		.FillHeight(1.0f)
 		[
 			SNew(SHorizontalBox)
 			+ SHorizontalBox::Slot()
 			[
-				 SNew(SScrollBox)
-				 .Orientation(Orient_Horizontal)
-				 .ExternalScrollbar(HorizontalScrollBar)
-				 + SScrollBox::Slot()
+				SNew(SScrollBox)
+				.Orientation(Orient_Horizontal)
+				.ExternalScrollbar(HorizontalScrollBar)
+				+ SScrollBox::Slot()
 				[
-					SAssignNew(ListViewWidget, SListView<TSharedPtr<int32>>)
-					.ListItemsSource(&RowItems)
-					.OnGenerateRow(this, &SNiagaraSimCacheView::MakeRowWidget)
-					.Visibility(EVisibility::Visible)
-					.SelectionMode(ESelectionMode::Single)
-					.ExternalScrollbar(VerticalScrollBar)
-					.ConsumeMouseWheel(EConsumeMouseWheel::Always)
-					.AllowOverscroll(EAllowOverscroll::No)
-					.HeaderRow(HeaderRowWidget)
+					// switcher for spreadsheet / data interfaces
+					SAssignNew(SwitchWidget, SWidgetSwitcher)
+					.WidgetIndex_Lambda([this] { return SimCacheViewModel->GetActiveDataInterface().IsValid() ? 1 : 0; })
+					+SWidgetSwitcher::Slot()
+					[
+						ListViewWidget.ToSharedRef()
+					]
+					+SWidgetSwitcher::Slot()
+					[
+						SNullWidget::NullWidget
+					]
 				]
 			]
 			+ SHorizontalBox::Slot()
 			.AutoWidth()
 			[
-				VerticalScrollBar
+				SNew(SWidgetSwitcher)
+				.WidgetIndex_Lambda([this] { return SimCacheViewModel->GetActiveDataInterface().IsValid() ? 1 : 0; })
+				+SWidgetSwitcher::Slot()
+				[
+					VerticalScrollBar
+				]
+				+SWidgetSwitcher::Slot()
+				[
+					DataInterfaceScrollBar.ToSharedRef()
+				]
 			]
 		]
 		+ SVerticalBox::Slot()
@@ -182,6 +206,52 @@ void SNiagaraSimCacheView::GenerateColumns()
 
 }
 
+void SNiagaraSimCacheView::UpdateDIWidget()
+{
+	for (TSharedPtr<SWidget> Widget : DIVisualizerWidgets)
+	{
+		SwitchWidget->RemoveSlot(Widget.ToSharedRef());
+	}
+	DIVisualizerWidgets.Empty();
+	TSharedRef<SVerticalBox> WidgetBox = SNew(SVerticalBox);
+	if (SimCacheViewModel->GetActiveDataInterface().IsValid())
+	{
+		FNiagaraEditorModule& NiagaraEditorModule = FModuleManager::GetModuleChecked<FNiagaraEditorModule>("NiagaraEditor");
+		for (TSharedRef<INiagaraDataInterfaceSimCacheVisualizer> Visualizer : NiagaraEditorModule.FindDataInterfaceCacheVisualizer(SimCacheViewModel->GetActiveDataInterface().GetType().GetClass()))
+		{
+			if (UObject* DataObject = SimCacheViewModel->GetActiveDataInterfaceStorage())
+			{
+				TSharedPtr<SWidget> DIVisualizerWidget = Visualizer->CreateWidgetFor(DataObject, SimCacheViewModel);
+				DIVisualizerWidgets.Add(DIVisualizerWidget);
+				WidgetBox->AddSlot()
+					.AutoHeight()
+					.AttachWidget(DIVisualizerWidget.ToSharedRef());
+			}
+		}
+	}
+	if (DIVisualizerWidgets.Num() == 0)
+	{
+		TSharedPtr<SWidget> DIVisualizerWidget = SNew(SBox)
+		   .Padding(10)
+		   [
+			   SNew(STextBlock)
+			   .Text(LOCTEXT("NoDataInterface", "No valid visualizer found for data interface"))
+		   ];
+		DIVisualizerWidgets.Add(DIVisualizerWidget);
+		WidgetBox->AddSlot()
+			.AutoHeight()
+			.AttachWidget(DIVisualizerWidget.ToSharedRef());
+	}
+	SwitchWidget->AddSlot(1).AttachWidget(
+		SNew(SScrollBox)
+		.Orientation(Orient_Vertical)
+		.ExternalScrollbar(DataInterfaceScrollBar)
+		 + SScrollBox::Slot()
+		[
+			WidgetBox
+		]);
+}
+
 void SNiagaraSimCacheView::UpdateColumns(const bool bReset)
 {
 	if(bReset)
@@ -211,6 +281,7 @@ void SNiagaraSimCacheView::OnSimCacheChanged()
 {
 	UpdateRows(true);
 	UpdateColumns(true);
+	UpdateDIWidget();
 }
 
 void SNiagaraSimCacheView::OnViewDataChanged(const bool bFullRefresh)
@@ -226,6 +297,7 @@ void SNiagaraSimCacheView::OnBufferChanged()
 {
 	UpdateRows(true);
 	UpdateColumns(true);
+	UpdateDIWidget();
 }
 
 bool SNiagaraSimCacheView::GetShouldGenerateWidget(FName Name)
diff --git a/Engine/Plugins/FX/Niagara/Source/NiagaraEditor/Private/Widgets/SNiagaraSimCacheView.h b/Engine/Plugins/FX/Niagara/Source/NiagaraEditor/Private/Widgets/SNiagaraSimCacheView.h
--- a/Engine/Plugins/FX/Niagara/Source/NiagaraEditor/Private/Widgets/SNiagaraSimCacheView.h
+++ b/Engine/Plugins/FX/Niagara/Source/NiagaraEditor/Private/Widgets/SNiagaraSimCacheView.h
@@ -9,6 +9,8 @@ class FNiagaraSimCacheViewModel;
 class ITableRow;
 class STableViewBase;
 class SHeaderRow;
+class SWidgetSwitcher;
+class SScrollBar;
 
 class SNiagaraSimCacheView : public SCompoundWidget
 {
@@ -19,7 +21,6 @@ public:
 
 	using FBufferSelectionInfo = TPair<int32, FText>;
 
-	bool IsStringFilterEnabled() const;
 	void Construct(const FArguments& InArgs);
 
 	TSharedRef<ITableRow> MakeRowWidget(const TSharedPtr<int32> RowIndexPtr, const TSharedRef<STableViewBase>& OwnerTable) const;
@@ -38,10 +39,14 @@ public:
 private:
 	bool GetShouldGenerateWidget(FName Name);
 	void GenerateColumns();
+	void UpdateDIWidget();
 
 	TArray<TSharedPtr<int32>>					RowItems;
 	TSharedPtr<FNiagaraSimCacheViewModel>		SimCacheViewModel;
 	
 	TSharedPtr<SHeaderRow>						HeaderRowWidget;
 	TSharedPtr<SListView<TSharedPtr<int32>>>	ListViewWidget;
+	TSharedPtr<SWidgetSwitcher>					SwitchWidget;
+	TArray<TSharedPtr<SWidget>>					DIVisualizerWidgets;
+	TSharedPtr<SScrollBar>						DataInterfaceScrollBar;
 };
\ No newline at end of file
diff --git a/Engine/Plugins/FX/Niagara/Source/NiagaraEditor/Public/Customizations/NiagaraDataInterfaceSimCacheVisualizer.h b/Engine/Plugins/FX/Niagara/Source/NiagaraEditor/Public/Customizations/NiagaraDataInterfaceSimCacheVisualizer.h
new file mode 100644
--- /dev/null
+++ b/Engine/Plugins/FX/Niagara/Source/NiagaraEditor/Public/Customizations/NiagaraDataInterfaceSimCacheVisualizer.h
@@ -0,0 +1,22 @@
+﻿// Copyright Epic Games, Inc. All Rights Reserved.
+
+#pragma once
+
+#include "CoreMinimal.h"
+
+class FNiagaraSimCacheViewModel;
+class SWidget;
+
+/**
+ * Implementations can be registered with the niagara editor module to provide custom visualizations for data interfaces stored in sim caches.
+ * This is needed since data interfaces can store any custom uobject, so the sim cache editor has no knowledge how to display the data stored within.
+ *
+ * See FNiagaraDataChannelCacheVisualizer for an example implementation.
+ */
+class NIAGARAEDITOR_API INiagaraDataInterfaceSimCacheVisualizer
+{
+public:
+	virtual ~INiagaraDataInterfaceSimCacheVisualizer() = default;
+	
+	virtual TSharedPtr<SWidget> CreateWidgetFor(UObject* CachedData, TSharedPtr<FNiagaraSimCacheViewModel> ViewModel);
+};
diff --git a/Engine/Plugins/FX/Niagara/Source/NiagaraEditor/Public/NiagaraEditorModule.h b/Engine/Plugins/FX/Niagara/Source/NiagaraEditor/Public/NiagaraEditorModule.h
--- a/Engine/Plugins/FX/Niagara/Source/NiagaraEditor/Public/NiagaraEditorModule.h
+++ b/Engine/Plugins/FX/Niagara/Source/NiagaraEditor/Public/NiagaraEditorModule.h
@@ -12,6 +12,7 @@
 #include "NiagaraPerfBaseline.h"
 #include "NiagaraDebuggerCommon.h"
 #include "NiagaraRendererProperties.h"
+#include "Customizations/NiagaraDataInterfaceSimCacheVisualizer.h"
 #include "NiagaraEditorModule.generated.h"
 
 class IAssetTools;
@@ -78,7 +79,7 @@ struct FReservedParameter
 
 public:
 	FReservedParameter()
-		: Parameter(FNiagaraVariable())
+		: Parameter(FNiagaraVariableBase())
 		, ReservingDefinitionsAsset(nullptr)
 	{};
 
@@ -183,6 +184,11 @@ public:
 
 	TSharedRef<INiagaraEditorWidgetProvider> GetWidgetProvider() const;
 
+	
+	NIAGARAEDITOR_API void RegisterDataInterfaceCacheVisualizer(UClass* DataInterfaceClass, TSharedRef<INiagaraDataInterfaceSimCacheVisualizer> InCacheVisualizer);
+	NIAGARAEDITOR_API void UnregisterDataInterfaceCacheVisualizer(UClass* DataInterfaceClass, TSharedRef<INiagaraDataInterfaceSimCacheVisualizer> InCacheVisualizer);
+	TArrayView<TSharedRef<INiagaraDataInterfaceSimCacheVisualizer>> FindDataInterfaceCacheVisualizer(UClass* DataInterfaceClass);
+
 	TSharedRef<FNiagaraScriptMergeManager> GetScriptMergeManager() const;
 
 	// Object pooling methods used to prevent unnecessary object allocation during compiles
@@ -416,6 +422,8 @@ private:
 
 	TSharedPtr<INiagaraEditorWidgetProvider> WidgetProvider;
 
+	TMap<TObjectKey<UClass>, TArray<TSharedRef<INiagaraDataInterfaceSimCacheVisualizer>>> DataInterfaceVisualizers;
+
 	TSharedPtr<FNiagaraScriptMergeManager> ScriptMergeManager;
 
 	TSharedPtr<INiagaraEditorOnlyDataUtilities> EditorOnlyDataUtilities;
diff --git a/Engine/Plugins/FX/Niagara/Source/NiagaraEditor/Public/ViewModels/NiagaraSimCacheViewModel.h b/Engine/Plugins/FX/Niagara/Source/NiagaraEditor/Public/ViewModels/NiagaraSimCacheViewModel.h
new file mode 100644
--- /dev/null
+++ b/Engine/Plugins/FX/Niagara/Source/NiagaraEditor/Public/ViewModels/NiagaraSimCacheViewModel.h
@@ -0,0 +1,138 @@
+// Copyright Epic Games, Inc. All Rights Reserved.
+
+#pragma once
+
+#include "CoreMinimal.h"
+#include "NiagaraSimCache.h"
+#include "Widgets/SNiagaraSimCacheOverview.h"
+
+class SNiagaraSimCacheTreeView;
+struct FNiagaraSimCacheComponentTreeItem;
+struct FNiagaraSimCacheTreeItem;
+struct FNiagaraSimCacheOverviewItem;
+struct FNiagaraSystemSimCacheCaptureReply;
+struct FNiagaraSimCacheDataBuffersLayout;
+class UNiagaraComponent;
+
+class FNiagaraSimCacheViewModel : public TSharedFromThis<FNiagaraSimCacheViewModel>, public FGCObject
+{
+public:
+	struct FComponentInfo
+	{
+		FName Name = NAME_None;
+		uint32 ComponentOffset = INDEX_NONE;
+		bool bIsFloat = false;
+		bool bIsHalf = false;
+		bool bIsInt32 = false;
+		bool bShowAsBool = false;
+		UEnum* Enum = nullptr;
+	};
+
+	DECLARE_MULTICAST_DELEGATE_OneParam(FOnViewDataChanged, bool)
+	DECLARE_MULTICAST_DELEGATE(FOnSimCacheChanged)
+	DECLARE_MULTICAST_DELEGATE(FOnBufferChanged)
+	
+	NIAGARAEDITOR_API FNiagaraSimCacheViewModel();
+	NIAGARAEDITOR_API virtual ~FNiagaraSimCacheViewModel() override;
+	
+	NIAGARAEDITOR_API TConstArrayView<FComponentInfo> GetCurrentComponentInfos() const { return GetComponentInfos(EmitterIndex); }
+	NIAGARAEDITOR_API TConstArrayView<FComponentInfo> GetComponentInfos(int32 InEmitterIndex) const;
+
+	NIAGARAEDITOR_API int32 GetNumInstances() const { return NumInstances; }
+	NIAGARAEDITOR_API int32 GetNumFrames() const;
+	NIAGARAEDITOR_API int32 GetFrameIndex() const { return FrameIndex; }
+	NIAGARAEDITOR_API void SetFrameIndex(const int32 InFrameIndex);
+
+	NIAGARAEDITOR_API FNiagaraVariableBase GetActiveDataInterface() const { return ActiveDataInterface; };
+	NIAGARAEDITOR_API UObject* GetActiveDataInterfaceStorage() const;
+
+	NIAGARAEDITOR_API FOnViewDataChanged& OnViewDataChanged();
+	NIAGARAEDITOR_API FOnSimCacheChanged& OnSimCacheChanged();
+	NIAGARAEDITOR_API FOnBufferChanged& OnBufferChanged();
+
+	void Initialize(TWeakObjectPtr<UNiagaraSimCache> SimCache);
+
+	TConstArrayView<FString> GetComponentFilters() const { return ComponentFilterArray; }
+	void SetComponentFilters(const TArray<FString>& NewComponentFilterArray);
+	bool IsComponentFilterActive() const { return bComponentFilterActive; }
+	
+	int32 GetEmitterIndex() const { return EmitterIndex; };
+	void SetEmitterIndex(int32 InEmitterIndex, FNiagaraVariableBase ActiveDataInterface = FNiagaraVariableBase());
+
+	FText GetComponentText(FName ComponentName, int32 InstanceIndex) const;
+	bool IsCacheValid();
+	int32 GetNumEmitterLayouts();
+	FName GetEmitterLayoutName(int32 Index);
+
+	// Construct entries for the tree view.
+	void BuildEntries(TWeakPtr<SNiagaraSimCacheTreeView> OwningTreeView);
+	TArray<TSharedRef<FNiagaraSimCacheTreeItem>>* GetCurrentRootEntries();
+	TArray<TSharedRef<FNiagaraSimCacheOverviewItem>>* GetBufferEntries();
+
+	UNiagaraComponent* GetPreviewComponent() { return PreviewComponent; }
+
+	//~ FGCObject interface
+	virtual void AddReferencedObjects(FReferenceCollector& Collector) override;
+	virtual FString GetReferencerName() const override
+	{
+		return TEXT("FNiagaraSimCacheViewModel");
+	}
+
+private:
+	void SetupPreviewComponentAndInstance();
+	void BuildComponentInfos(const FName Name, const UScriptStruct* Struct, TArray<FComponentInfo>& ComponentInfos);
+	void OnCacheModified(UNiagaraSimCache* SimCache);
+	void UpdateCachedFrame();
+	void UpdateCurrentEntries();
+
+	void BuildTreeItemChildren(TSharedPtr<FNiagaraSimCacheTreeItem> InTreeItem, TWeakPtr<SNiagaraSimCacheTreeView> OwningTreeView);
+	void RecursiveBuildTreeItemChildren(FNiagaraSimCacheTreeItem* Root, TSharedRef<FNiagaraSimCacheComponentTreeItem> Parent, FNiagaraTypeDefinition TypeDefinition, TWeakPtr<SNiagaraSimCacheTreeView> OwningTreeView);
+
+	// Build out component infos for each buffer in this cache
+	void UpdateComponentInfos();
+	
+	// The sim cache being viewed
+	TObjectPtr<UNiagaraSimCache> SimCache = nullptr;
+
+	// Component for preview scene
+	TObjectPtr<UNiagaraComponent> PreviewComponent = nullptr;
+
+	// Which frame of the cached sim is being viewed
+	int32 FrameIndex = 0;
+	
+	// Which Emitter of the cached sim is being viewed
+	int32 EmitterIndex = INDEX_NONE;
+
+	// Which data interface of the cached sim is being viewed
+	FNiagaraVariableBase ActiveDataInterface;
+
+	// Number of particles in the given frame
+	int32 NumInstances = 0;
+
+	// Cached Component infos for this system
+	TArray<FComponentInfo>				SystemComponentInfos;
+	TArray<TArray<FComponentInfo>>		EmitterComponentInfos;
+	// TODO: Component infos for Data Interfaces?
+
+	TArray<TSharedRef<FNiagaraSimCacheTreeItem>> RootEntries;
+	TArray<TSharedRef<FNiagaraSimCacheTreeItem>> CurrentRootEntries;
+
+	TArray<TSharedRef<FNiagaraSimCacheOverviewItem>> BufferEntries;
+
+	// String filters gathered from the overview tree
+	TArray<FString>						ComponentFilterArray;
+	bool								bComponentFilterActive = false;
+	
+	TArray<float>						FloatComponents;
+	TArray<FFloat16>					HalfComponents;
+	TArray<int32>						Int32Components;
+
+	int32								FoundFloatComponents = 0;
+	int32								FoundHalfComponents = 0;
+	int32								FoundInt32Components = 0;
+
+	bool								bDelegatesAdded = false;
+	FOnViewDataChanged					OnViewDataChangedDelegate;
+	FOnSimCacheChanged					OnSimCacheChangedDelegate;
+	FOnBufferChanged					OnBufferChangedDelegate;
+};
