- Check for InstanceData being nullptr

diff --git a/Engine/Plugins/FX/Niagara/Source/Niagara/Classes/NiagaraDataInterfaceArrayImpl.h b/Engine/Plugins/FX/Niagara/Source/Niagara/Classes/NiagaraDataInterfaceArrayImpl.h
--- a/Engine/Plugins/FX/Niagara/Source/Niagara/Classes/NiagaraDataInterfaceArrayImpl.h
+++ b/Engine/Plugins/FX/Niagara/Source/Niagara/Classes/NiagaraDataInterfaceArrayImpl.h
@@ -598,7 +598,10 @@ struct FNDIArrayProxyImpl : public INDIArrayProxyBase
 			}
 			ArrayRef.GetArray().AddDefaulted(Index + 1 - ArrayRef.GetArray().Num());
 		}
-		InstanceData->bIsRenderDirty |= bShouldSyncToGpu;
+		if (InstanceData)
+		{
+			InstanceData->bIsRenderDirty |= bShouldSyncToGpu;
+		}
 		FNDIArrayImplHelper<TArrayType>::CopyCpuToCpuMemory(ArrayRef.GetArray().GetData() + Index, &Value, 1);
 	}
 
