- Default deformed buffers when no data is present

diff --git a/Engine/Plugins/FX/Niagara/Source/Niagara/Private/NiagaraDataInterfaceSkeletalMesh.cpp b/Engine/Plugins/FX/Niagara/Source/Niagara/Private/NiagaraDataInterfaceSkeletalMesh.cpp
--- a/Engine/Plugins/FX/Niagara/Source/Niagara/Private/NiagaraDataInterfaceSkeletalMesh.cpp
+++ b/Engine/Plugins/FX/Niagara/Source/Niagara/Private/NiagaraDataInterfaceSkeletalMesh.cpp
@@ -3423,6 +3423,10 @@ void UNiagaraDataInterfaceSkeletalMesh::SetShaderParameters(const FNiagaraDataIn
 		ShaderParameters->PreSkinnedLocalBoundsExtents = FVector3f::ZeroVector;
 
 		ShaderParameters->EnabledFeatures = 0;
+
+		ShaderParameters->DeformedCurrPositionBuffer = FNiagaraRenderer::GetDummyFloatBuffer();
+		ShaderParameters->DeformedPrevPositionBuffer = FNiagaraRenderer::GetDummyFloatBuffer();
+		ShaderParameters->DeformedTangentBuffer = FNiagaraRenderer::GetDummyFloat4Buffer();
 	}
 }
 
