fix accurate collisions when no mesh df exists

diff --git a/Engine/Plugins/FX/Niagara/Shaders/Private/NiagaraDataInterfaceRigidMeshCollisionQuery.ush b/Engine/Plugins/FX/Niagara/Shaders/Private/NiagaraDataInterfaceRigidMeshCollisionQuery.ush
--- a/Engine/Plugins/FX/Niagara/Shaders/Private/NiagaraDataInterfaceRigidMeshCollisionQuery.ush
+++ b/Engine/Plugins/FX/Niagara/Shaders/Private/NiagaraDataInterfaceRigidMeshCollisionQuery.ush
@@ -276,14 +276,13 @@ void GetElementPoint_{ParameterName}(in float3 WorldPosition, in float DeltaTime
 void GetElementPointMeshDistanceFieldNoNormal_{ParameterName}(in float3 WorldPosition, in float DeltaTime, in float TimeFraction, in float MaxDistance, in int ElementIndex, out float OutClosestDistance, out float3 OutClosestVelocity)
 {
 	OutClosestVelocity = float3(0,0,0);
-	OutClosestDistance = 0.0;
+	OutClosestDistance = MAX_DISTANCE;
 
 	if (ElementIndex >= 0)
 	{
 		const uint ElementIndexUint = uint(ElementIndex);
 		const float3 LocalPosition = GetLocalPosition_{ParameterName}(WorldPosition,ElementIndex,TimeFraction);
-
-		float MinDistance = MAX_DISTANCE;
+		
 		int OutputIndex = -1;
 		float3 CollisionPosition = float3(0,0,0);
 		float3 CollisionNormal = float3(0,0,0);
@@ -291,34 +290,31 @@ void GetElementPointMeshDistanceFieldNoNormal_{ParameterName}(in float3 WorldPos
 		if ( ElementIndexUint >= {ParameterName}_ElementOffsets[SPHERE_INDEX] && ElementIndexUint < {ParameterName}_ElementOffsets[SPHERE_INDEX+1] )
 		{
 			GetSphereProjection_{ParameterName}(LocalPosition, float3(0,0,0), {ParameterName}_ElementExtentBuffer[ElementIndex].x, ElementIndex,
-					CollisionPosition, CollisionNormal, OutputIndex, MinDistance);
+					CollisionPosition, CollisionNormal, OutputIndex, OutClosestDistance);
 		}
 		else if ( ElementIndexUint >= {ParameterName}_ElementOffsets[CAPSULE_INDEX] && ElementIndexUint < {ParameterName}_ElementOffsets[CAPSULE_INDEX+1] )
 		{
 			GetCapsuleProjection_{ParameterName}(LocalPosition, {ParameterName}_ElementExtentBuffer[ElementIndex].xy, ElementIndex,
-					CollisionPosition, CollisionNormal, OutputIndex, MinDistance);
+					CollisionPosition, CollisionNormal, OutputIndex, OutClosestDistance);
 		}
 		else if ( ElementIndexUint >= {ParameterName}_ElementOffsets[BOX_INDEX] && ElementIndexUint < {ParameterName}_ElementOffsets[BOX_INDEX+1] )
 		{
 			GetBoxProjection_{ParameterName}(LocalPosition, {ParameterName}_ElementExtentBuffer[ElementIndex].xyz, ElementIndex,
-					CollisionPosition, CollisionNormal, OutputIndex, MinDistance);
+					CollisionPosition, CollisionNormal, OutputIndex, OutClosestDistance);
 		}
 
-		if (ElementIndex != -1)
-		{
-			uint DFIndex = {ParameterName}_DFIndexBuffer[ElementIndex];
-
-			if (DFIndex >= 0 && DFIndex < NumSceneObjects)
-			{
-				MaxDistance = max(.001, MaxDistance);
-				const FLWCVector3 LWCWorldPosition = MakeLWCVector3({ParameterName}_SystemLWCTile, WorldPosition);
-				OutClosestDistance = DistanceToNearestSurfaceForObject(DFIndex, LWCWorldPosition, MaxDistance);			
-			}
+		uint DFIndex = {ParameterName}_DFIndexBuffer[ElementIndex];
 
-			const float3 PreviousPosition = mul(GetPreviousTransform_{ParameterName}(ElementIndex), float4(CollisionPosition,1.0)).xyz;
-			const float3 CurrentPosition = mul(GetCurrentTransform_{ParameterName}(ElementIndex), float4(CollisionPosition,1.0)).xyz;
-			OutClosestVelocity = ( CurrentPosition - PreviousPosition ) / DeltaTime;		
+		if (DFIndex >= 0 && DFIndex < NumSceneObjects)
+		{
+			MaxDistance = max(.001, MaxDistance);
+			const FLWCVector3 LWCWorldPosition = MakeLWCVector3({ParameterName}_SystemLWCTile, WorldPosition);
+			OutClosestDistance = DistanceToNearestSurfaceForObject(DFIndex, LWCWorldPosition, MaxDistance);			
 		}
+
+		const float3 PreviousPosition = mul(GetPreviousTransform_{ParameterName}(ElementIndex), float4(CollisionPosition,1.0)).xyz;
+		const float3 CurrentPosition = mul(GetCurrentTransform_{ParameterName}(ElementIndex), float4(CollisionPosition,1.0)).xyz;
+		OutClosestVelocity = ( CurrentPosition - PreviousPosition ) / DeltaTime;		
 	}
 }
 
