* Small code cleanups to fix static analysis warnings
* Improved enum name

diff --git a/Engine/Plugins/FX/Niagara/Source/NiagaraEditor/Private/Customizations/SimCache/NiagaraRenderTargetVolumeSimCacheVisualizer.cpp b/Engine/Plugins/FX/Niagara/Source/NiagaraEditor/Private/Customizations/SimCache/NiagaraRenderTargetVolumeSimCacheVisualizer.cpp
--- a/Engine/Plugins/FX/Niagara/Source/NiagaraEditor/Private/Customizations/SimCache/NiagaraRenderTargetVolumeSimCacheVisualizer.cpp
+++ b/Engine/Plugins/FX/Niagara/Source/NiagaraEditor/Private/Customizations/SimCache/NiagaraRenderTargetVolumeSimCacheVisualizer.cpp
@@ -49,8 +49,8 @@ public:
 
 	TStrongObjectPtr<UHeterogeneousVolumeComponent> PreviewComponent;
 
-	void SetTemperatureMask(ENiagaraRenderTargetVolumeMask NewMask);
-	ENiagaraRenderTargetVolumeMask TemparatureMask = ENiagaraRenderTargetVolumeMask::G;
+	void SetTemperatureMask(ENiagaraRenderTargetVolumeVisualizerMask NewMask);
+	ENiagaraRenderTargetVolumeVisualizerMask TemparatureMask = ENiagaraRenderTargetVolumeVisualizerMask::G;
 
 	void SetDensityScale(float NewValue);
 	float DensityScale = 0.1f;
@@ -75,7 +75,7 @@ TSharedPtr<SWidget> FNiagaraRenderTargetVolumeSimCacheVisualizer::CreateWidgetFo
 {
 	if (UAnimatedSparseVolumeTexture* VolumeTexture = Cast<UAnimatedSparseVolumeTexture>(CachedData))
 	{
-		static UEnum* MaskEnum = StaticEnum<ENiagaraRenderTargetVolumeMask>();
+		static UEnum* MaskEnum = StaticEnum<ENiagaraRenderTargetVolumeVisualizerMask>();
 		TSharedPtr<SNiagaraVolumeTextureViewport> Viewport = SNew(SNiagaraVolumeTextureViewport, ViewModel, VolumeTexture);
 		
 		return SNew(SVerticalBox)
@@ -114,7 +114,7 @@ TSharedPtr<SWidget> FNiagaraRenderTargetVolumeSimCacheVisualizer::CreateWidgetFo
 					})
 					.OnEnumSelectionChanged_Lambda([Viewport](int32 NewValue, ESelectInfo::Type)
 					{
-						ENiagaraRenderTargetVolumeMask Value = static_cast<ENiagaraRenderTargetVolumeMask>(NewValue);
+						ENiagaraRenderTargetVolumeVisualizerMask Value = static_cast<ENiagaraRenderTargetVolumeVisualizerMask>(NewValue);
 						Viewport->SetTemperatureMask(Value);
 					})
 				]
@@ -321,7 +321,7 @@ void SNiagaraVolumeTextureViewport::OnFloatingButtonClicked()
 {
 }
 
-void SNiagaraVolumeTextureViewport::SetTemperatureMask(ENiagaraRenderTargetVolumeMask NewMask)
+void SNiagaraVolumeTextureViewport::SetTemperatureMask(ENiagaraRenderTargetVolumeVisualizerMask NewMask)
 {
 	if (PreviewComponent)
 	{
@@ -330,10 +330,10 @@ void SNiagaraVolumeTextureViewport::SetTemperatureMask(ENiagaraRenderTargetVolum
 			UMaterial *DuplicateMat = DuplicateObject<UMaterial>(CurrentMat, PreviewComponent.Get());
 			FGuid ExprGuid;
 			DuplicateMat->SetStaticComponentMaskParameterValueEditorOnly("Temperature Mask",
-				NewMask == ENiagaraRenderTargetVolumeMask::R,
-				NewMask == ENiagaraRenderTargetVolumeMask::G,
-				NewMask == ENiagaraRenderTargetVolumeMask::B,
-				NewMask == ENiagaraRenderTargetVolumeMask::A,
+				NewMask == ENiagaraRenderTargetVolumeVisualizerMask::R,
+				NewMask == ENiagaraRenderTargetVolumeVisualizerMask::G,
+				NewMask == ENiagaraRenderTargetVolumeVisualizerMask::B,
+				NewMask == ENiagaraRenderTargetVolumeVisualizerMask::A,
 				ExprGuid);
 			PreviewComponent->OverrideMaterials[0] = DuplicateMat;
 			PreviewComponent->PostLoad();
diff --git a/Engine/Plugins/FX/Niagara/Source/NiagaraEditor/Private/Customizations/SimCache/NiagaraRenderTargetVolumeSimCacheVisualizer.h b/Engine/Plugins/FX/Niagara/Source/NiagaraEditor/Private/Customizations/SimCache/NiagaraRenderTargetVolumeSimCacheVisualizer.h
--- a/Engine/Plugins/FX/Niagara/Source/NiagaraEditor/Private/Customizations/SimCache/NiagaraRenderTargetVolumeSimCacheVisualizer.h
+++ b/Engine/Plugins/FX/Niagara/Source/NiagaraEditor/Private/Customizations/SimCache/NiagaraRenderTargetVolumeSimCacheVisualizer.h
@@ -6,8 +6,9 @@
 
 #include "NiagaraRenderTargetVolumeSimCacheVisualizer.generated.h"
 
+// Used in the visualizer panel to populate the temperature mask dropdown
 UENUM()
-enum class ENiagaraRenderTargetVolumeMask
+enum class ENiagaraRenderTargetVolumeVisualizerMask
 {
 	R, G, B, A
 };
diff --git a/Engine/Plugins/FX/Niagara/Source/NiagaraEditor/Private/NiagaraBakerRenderer.cpp b/Engine/Plugins/FX/Niagara/Source/NiagaraEditor/Private/NiagaraBakerRenderer.cpp
--- a/Engine/Plugins/FX/Niagara/Source/NiagaraEditor/Private/NiagaraBakerRenderer.cpp
+++ b/Engine/Plugins/FX/Niagara/Source/NiagaraEditor/Private/NiagaraBakerRenderer.cpp
@@ -12,10 +12,8 @@
 #include "NiagaraEditorCommon.h"
 
 #include "NiagaraDataInterfaceGrid2DCollection.h"
-#include "NiagaraDataInterfaceRenderTarget2D.h"
 
 #include "Components/SceneCaptureComponent2D.h"
-#include "Engine/Canvas.h"
 #include "Engine/TextureRenderTarget2D.h"
 #include "Misc/FileHelper.h"
 #include "Misc/Paths.h"
@@ -30,7 +28,6 @@
 #include "ImageWrapperHelper.h"
 #include "LegacyScreenPercentageDriver.h"
 #include "UObject/Package.h"
-#include "VolumeCache.h"
 #include "TextureResource.h"
 #include "SceneInterface.h"
 
@@ -580,11 +577,11 @@ void FNiagaraBakerRenderer::RenderParticleAttribute(UTextureRenderTarget2D* Rend
 		}
 		const FNiagaraVariableLayoutInfo& VariableInfo = ParticleDataSet.GetCompiledData().VariableLayouts[VariableIndex];
 	
-		float* FloatChannels[4];
-		FloatChannels[0] = (float*)ParticleDataBuffer->GetComponentPtrFloat(VariableInfo.GetFloatComponentStart());
-		FloatChannels[1] = VariableInfo.GetNumFloatComponents() > 1 ? (float*)ParticleDataBuffer->GetComponentPtrFloat(VariableInfo.GetFloatComponentStart() + 1) : nullptr;
-		FloatChannels[2] = VariableInfo.GetNumFloatComponents() > 2 ? (float*)ParticleDataBuffer->GetComponentPtrFloat(VariableInfo.GetFloatComponentStart() + 2) : nullptr;
-		FloatChannels[3] = VariableInfo.GetNumFloatComponents() > 3 ? (float*)ParticleDataBuffer->GetComponentPtrFloat(VariableInfo.GetFloatComponentStart() + 3) : nullptr;
+		const float* FloatChannels[4];
+		FloatChannels[0] = reinterpret_cast<const float*>(ParticleDataBuffer->GetComponentPtrFloat(VariableInfo.GetFloatComponentStart()));
+		FloatChannels[1] = VariableInfo.GetNumFloatComponents() > 1 ? reinterpret_cast<const float*>(ParticleDataBuffer->GetComponentPtrFloat(VariableInfo.GetFloatComponentStart() + 1)) : nullptr;
+		FloatChannels[2] = VariableInfo.GetNumFloatComponents() > 2 ? reinterpret_cast<const float*>(ParticleDataBuffer->GetComponentPtrFloat(VariableInfo.GetFloatComponentStart() + 2)) : nullptr;
+		FloatChannels[3] = VariableInfo.GetNumFloatComponents() > 3 ? reinterpret_cast<const float*>(ParticleDataBuffer->GetComponentPtrFloat(VariableInfo.GetFloatComponentStart() + 3)) : nullptr;
 	
 		const FIntPoint RenderTargetSize(RenderTarget->GetSurfaceWidth(), RenderTarget->GetSurfaceHeight());
 		const int32 ParticleBufferStore = RenderTargetSize.X * RenderTargetSize.Y;
@@ -649,13 +646,12 @@ void FNiagaraBakerRenderer::RenderSparseVolumeTexture(UTextureRenderTarget2D* Re
 		return;
 	}
 
-	bool CreatePreviewComponent = SVTPreviewComponent == nullptr;
-	if (CreatePreviewComponent)
+	if (SVTPreviewComponent == nullptr)
 	{
 		SVTPreviewComponent = NewObject<UHeterogeneousVolumeComponent>(GetTransientPackage(), NAME_None, RF_Transient);
 				
 		// create HV component and wire all the things
-		UMaterialInterface* MaterialInterface = LoadObject<UMaterialInterface>(NULL, TEXT("/Engine/EngineMaterials/SparseVolumeMaterial"), NULL, LOAD_None, NULL);
+		UMaterialInterface* MaterialInterface = LoadObject<UMaterialInterface>(nullptr, TEXT("/Engine/EngineMaterials/SparseVolumeMaterial"));
 		
 		UMaterial *Mat = MaterialInterface->GetMaterial();
 
@@ -864,7 +860,7 @@ bool FVolumeDataInterfaceHelper::Initialize(const TArray<FString>& InputDataInte
 	{
 		VolumeRenderTargetDataInterface = CastChecked<UNiagaraDataInterfaceRenderTargetVolume>(DataInterface);
 		VolumeRenderTargetProxy = static_cast<FNiagaraDataInterfaceProxyRenderTargetVolumeProxy*>(VolumeRenderTargetDataInterface->GetProxy());
-		VolumeRenderTargetInstanceData_GameThread = reinterpret_cast<FRenderTargetVolumeRWInstanceData_GameThread*>(SystemInstance->FindDataInterfaceInstanceData(VolumeRenderTargetDataInterface));
+		VolumeRenderTargetInstanceData_GameThread = static_cast<FRenderTargetVolumeRWInstanceData_GameThread*>(SystemInstance->FindDataInterfaceInstanceData(VolumeRenderTargetDataInterface));
 		if (VolumeRenderTargetInstanceData_GameThread == nullptr)
 		{
 			return false;
@@ -875,7 +871,7 @@ bool FVolumeDataInterfaceHelper::Initialize(const TArray<FString>& InputDataInte
 	{
 		Grid3DDataInterface = CastChecked<UNiagaraDataInterfaceGrid3DCollection>(DataInterface);
 		Grid3DProxy = static_cast<FNiagaraDataInterfaceProxyGrid3DCollectionProxy*>(Grid3DDataInterface->GetProxy());
-		Grid3DInstanceData_GameThread = reinterpret_cast<FGrid3DCollectionRWInstanceData_GameThread*>(SystemInstance->FindDataInterfaceInstanceData(Grid3DDataInterface));
+		Grid3DInstanceData_GameThread = static_cast<FGrid3DCollectionRWInstanceData_GameThread*>(SystemInstance->FindDataInterfaceInstanceData(Grid3DDataInterface));
 		if (Grid3DInstanceData_GameThread == nullptr)
 		{
 			return false;
diff --git a/Engine/Plugins/FX/Niagara/Source/NiagaraEditor/Private/NiagaraBakerRenderer.h b/Engine/Plugins/FX/Niagara/Source/NiagaraEditor/Private/NiagaraBakerRenderer.h
--- a/Engine/Plugins/FX/Niagara/Source/NiagaraEditor/Private/NiagaraBakerRenderer.h
+++ b/Engine/Plugins/FX/Niagara/Source/NiagaraEditor/Private/NiagaraBakerRenderer.h
@@ -101,7 +101,7 @@ class FNiagaraBakerRenderer : FGCObject
 {
 public:
 	FNiagaraBakerRenderer(UNiagaraSystem* NiagaraSystem);
-	virtual ~FNiagaraBakerRenderer();
+	virtual ~FNiagaraBakerRenderer() override;
 
 	void SetAbsoluteTime(float AbsoluteTime, bool bShouldTickComponent = true);
 
diff --git a/Engine/Plugins/FX/Niagara/Source/NiagaraEditor/Private/Widgets/SNiagaraSystemViewport.cpp b/Engine/Plugins/FX/Niagara/Source/NiagaraEditor/Private/Widgets/SNiagaraSystemViewport.cpp
--- a/Engine/Plugins/FX/Niagara/Source/NiagaraEditor/Private/Widgets/SNiagaraSystemViewport.cpp
+++ b/Engine/Plugins/FX/Niagara/Source/NiagaraEditor/Private/Widgets/SNiagaraSystemViewport.cpp
@@ -29,7 +29,6 @@
 #include "Engine/Canvas.h"
 #include "Engine/Font.h"
 #include "Engine/TextureCube.h"
-#include "Serialization/ArchiveCountMem.h"
 #include "Slate/SceneViewport.h"
 #include "ThumbnailRendering/ThumbnailManager.h"
 #include "Widgets/Docking/SDockTab.h"
@@ -824,7 +823,6 @@ void SNiagaraSystemViewport::Tick( const FGeometry& AllottedGeometry, const doub
 	{
 		const float MotionTime = PreviewComponent->GetDesiredAge();
 
-		UWorld* World = GetWorld();
 		FVector Location;
 		Location.X = MotionRadius * FMath::Sin(FMath::DegreesToRadians(MotionRate) * MotionTime);
 		Location.Y = 0.0f;
@@ -1256,11 +1254,6 @@ SNiagaraBaselineViewport::~SNiagaraBaselineViewport()
 	}
 }
 
-void SNiagaraBaselineViewport::AddReferencedObjects( FReferenceCollector& Collector )
-{
-
-}
-
 void SNiagaraBaselineViewport::RefreshViewport()
 {
 	SceneViewport->InvalidateDisplay();
diff --git a/Engine/Plugins/FX/Niagara/Source/NiagaraEditor/Private/Widgets/SNiagaraSystemViewport.h b/Engine/Plugins/FX/Niagara/Source/NiagaraEditor/Private/Widgets/SNiagaraSystemViewport.h
--- a/Engine/Plugins/FX/Niagara/Source/NiagaraEditor/Private/Widgets/SNiagaraSystemViewport.h
+++ b/Engine/Plugins/FX/Niagara/Source/NiagaraEditor/Private/Widgets/SNiagaraSystemViewport.h
@@ -13,7 +13,6 @@
 #include "ISequencerModule.h"
 #include "SEditorViewport.h"
 #include "SCommonEditorViewportToolbarBase.h"
-#include "Particles/ParticlePerfStatsManager.h"
 #include "NiagaraPerfBaseline.h"
 
 class FNiagaraSystemViewModel;
@@ -38,7 +37,7 @@ public:
 	SLATE_END_ARGS()
 	
 	void Construct(const FArguments& InArgs, TSharedRef<FNiagaraSystemViewModel> InSystemViewModel);
-	~SNiagaraSystemViewport();
+	virtual ~SNiagaraSystemViewport() override;
 	
 	virtual void AddReferencedObjects( FReferenceCollector& Collector ) override;
 	virtual FString GetReferencerName() const override
@@ -57,7 +56,7 @@ public:
 	/** If true, render background object in the preview scene. */
 	bool bShowBackground;
 
-	TSharedRef<class FAdvancedPreviewScene> GetPreviewScene() { return AdvancedPreviewScene.ToSharedRef(); }
+	TSharedRef<class FAdvancedPreviewScene> GetPreviewScene() const { return AdvancedPreviewScene.ToSharedRef(); }
 	TWeakPtr<FNiagaraSystemViewModel> GetSystemViewModel() { return SystemViewModel; }
 	
 	/** The material editor has been added to a tab */
@@ -115,10 +114,9 @@ protected:
 	FText GetViewportCompileStatusText() const;
 
 private:
-	bool IsVisible() const override;
+	virtual bool IsVisible() const override;
 	void OnScreenShotCaptured(UTexture2D* ScreenShot);
 
-private:
 	/** The parent tab where this viewport resides */
 	TWeakPtr<SDockTab> ParentTab;
 
@@ -139,7 +137,7 @@ private:
 	/** Level viewport client */
 	TSharedPtr<class FNiagaraSystemViewportClient> SystemViewportClient;
 
-	uint32 DrawFlags;
+	uint32 DrawFlags = 0;
 
 	FOnThumbnailCaptured OnThumbnailCaptured;
 
@@ -160,27 +158,19 @@ private:
 #if NIAGARA_PERF_BASELINES
 
 /** Niagara Baseline Display Viewport */
-class SNiagaraBaselineViewport : public SEditorViewport, public FGCObject
+class SNiagaraBaselineViewport : public SEditorViewport
 {
 public:
 	SLATE_BEGIN_ARGS(SNiagaraBaselineViewport) {}
 	SLATE_END_ARGS()
 
 	void Construct(const FArguments& InArgs);
-	virtual ~SNiagaraBaselineViewport();
-
-	virtual void AddReferencedObjects(FReferenceCollector& Collector) override;
-	virtual FString GetReferencerName() const override
-	{
-		return TEXT("SNiagaraBaselineViewport");
-	}
+	virtual ~SNiagaraBaselineViewport() override;
 
 	virtual void Tick(const FGeometry& AllottedGeometry, const double InCurrentTime, const float InDeltaTime) override;
 
 	void RefreshViewport();
 
-	TSharedRef<class FAdvancedPreviewScene> GetPreviewScene() { return AdvancedPreviewScene.ToSharedRef(); }
-
 	void Init(TSharedPtr<SWindow>& InOwnerWindow);
 
 	bool AddBaseline(UNiagaraEffectType* EffectType);
@@ -193,10 +183,9 @@ protected:
 	virtual void BindCommands() override;
 	virtual void OnFocusViewportToSelection() override;
 	virtual void PopulateViewportOverlays(TSharedRef<class SOverlay> Overlay) override;
-	EVisibility OnGetViewportCompileTextVisibility() const;
 
 private:
-	bool IsVisible() const override;
+	virtual bool IsVisible() const override;
 
 	/** Preview Scene - uses advanced preview settings */
 	TSharedPtr<class FAdvancedPreviewScene> AdvancedPreviewScene;
