- Add FX type exclude list for tick dependency rule
- This allows certain FX types to opt out of the check

diff --git a/Engine/Plugins/FX/Niagara/Source/NiagaraEditor/Private/NiagaraValidationRules.cpp b/Engine/Plugins/FX/Niagara/Source/NiagaraEditor/Private/NiagaraValidationRules.cpp
--- a/Engine/Plugins/FX/Niagara/Source/NiagaraEditor/Private/NiagaraValidationRules.cpp
+++ b/Engine/Plugins/FX/Niagara/Source/NiagaraEditor/Private/NiagaraValidationRules.cpp
@@ -1007,6 +1007,11 @@ void UNiagaraValidationRule_TickDependencyCheck::CheckValidity(const FNiagaraVal
 		return;
 	}
 
+	if (EffectTypesToExclude.Contains(TSoftObjectPtr<UNiagaraEffectType>(NiagaraSystem->GetEffectType())))
+	{
+		return;
+	}
+
 	TSet<UNiagaraDataInterface*> VisitedDIs;
 	NiagaraSystem->ForEachScript(
 		[&](UNiagaraScript* NiagaraScript)
diff --git a/Engine/Plugins/FX/Niagara/Source/NiagaraEditor/Public/NiagaraValidationRules.h b/Engine/Plugins/FX/Niagara/Source/NiagaraEditor/Public/NiagaraValidationRules.h
--- a/Engine/Plugins/FX/Niagara/Source/NiagaraEditor/Public/NiagaraValidationRules.h
+++ b/Engine/Plugins/FX/Niagara/Source/NiagaraEditor/Public/NiagaraValidationRules.h
@@ -8,6 +8,7 @@
 #include "NiagaraRendererProperties.h"
 #include "NiagaraValidationRules.generated.h"
 
+class UNiagaraEffectType;
 class UNiagaraScript;
 
 namespace NiagaraValidation
@@ -225,6 +226,10 @@ public:
 	/** Check that the skeletal mesh interface isn't adding a tick dependency on the CPU. */
 	UPROPERTY(EditAnywhere, Category = Validation)
 	bool bCheckSkeletalMeshInterface = true;
+
+	/** If the system uses one of these effect types the rule will not be run. */
+	UPROPERTY(EditAnywhere, Category = Validation)
+	TArray<TSoftObjectPtr<UNiagaraEffectType>> EffectTypesToExclude;
 };
 
 /** This validation rule checks to see if you have exposed user data interfaces. */
