- Disable GPU emitters that use FP16 if we do not support typed UAV store / load
- fx.Niagara.GpuEmitterCheckFloat16Support can be used to disable this check
- Also ensure we flush any pending setting updates each frame

diff --git a/Engine/Plugins/FX/Niagara/Source/Niagara/Private/NiagaraComponentSettings.cpp b/Engine/Plugins/FX/Niagara/Source/Niagara/Private/NiagaraComponentSettings.cpp
--- a/Engine/Plugins/FX/Niagara/Source/Niagara/Private/NiagaraComponentSettings.cpp
+++ b/Engine/Plugins/FX/Niagara/Source/Niagara/Private/NiagaraComponentSettings.cpp
@@ -37,9 +37,12 @@ namespace FNiagaraComponentSettings
 		return HashCombine(GetTypeHash(Var.SystemName), GetTypeHash(Var.EmitterName));
 	}
 
+	bool									bNeedsSettingsUpdate = true;
+
 	bool									bUseSystemDenyList = false;
 	bool									bUseEmitterDenyList = false;
 	bool									bAllowGpuEmitters = true;
+	bool									bAllowGpuComputeFloat16 = true;
 
 	TSet<FName>								SystemDenyList;
 	TSet<FNiagaraEmitterNameSettingsRef>	EmitterDenyList;
@@ -86,19 +89,24 @@ namespace FNiagaraComponentSettings
 		return ExistingSet.Difference(OutSet).Num() > 0;
 	}
 
+	static void UpdateScalabilityForSystems()
+	{
+		for (TObjectIterator<UNiagaraSystem> It; It; ++It)
+		{
+			if (UNiagaraSystem* System = *It)
+			{
+				System->UpdateScalability();
+			}
+		}
+	}
+
 	static void UpdateSystemDenyList(IConsoleVariable*)
 	{
 		const bool bWasChanged = ParseIntoSet(SystemDenyListString, SystemDenyList);
 		bUseSystemDenyList = SystemDenyList.Num() > 0;
 		if (bWasChanged)
 		{
-			for (TObjectIterator<UNiagaraSystem> It; It; ++It)
-			{
-				if (UNiagaraSystem* System = *It)
-				{
-					System->UpdateScalability();
-				}
-			}
+			UpdateScalabilityForSystems();
 		}
 	}
 
@@ -110,16 +118,15 @@ namespace FNiagaraComponentSettings
 		bUseEmitterDenyList = EmitterDenyList.Num() > 0 || GpuEmitterDenyList.Num() > 0 || GpuDataInterfaceDenyList.Num() > 0;
 		if (bWasChanged)
 		{
-			for (TObjectIterator<UNiagaraSystem> It; It; ++It)
-			{
-				if (UNiagaraSystem* System = *It)
-				{
-					System->UpdateScalability();
-				}
-			}
+			UpdateScalabilityForSystems();
 		}
 	}
 
+	void RequestUpdateSettings(IConsoleVariable*)
+	{
+		bNeedsSettingsUpdate = true;
+	}
+
 	static FAutoConsoleVariableRef CVarNiagaraSetSystemDenyList(
 		TEXT("fx.Niagara.SetSystemDenyList"),
 		SystemDenyListString,
@@ -152,6 +159,7 @@ namespace FNiagaraComponentSettings
 		TEXT("fx.Niagara.SetGpuRHIDenyList"),
 		GpuRHIDenyListString,
 		TEXT("Set Gpu RHI deny list to use, comma separated and uses wildcards, i.e. (*MyRHI*) would exclude anything that contains MyRHI"),
+		FConsoleVariableDelegate::CreateStatic(RequestUpdateSettings),
 		ECVF_Scalability | ECVF_Default
 	);
 
@@ -159,6 +167,7 @@ namespace FNiagaraComponentSettings
 		TEXT("fx.Niagara.SetGpuRHIAdapterDenyList"),
 		GpuRHIAdapterDenyListString,
 		TEXT("Set Gpu RHI Adapter deny list to use, comma separated and uses wildcards, i.e. (*MyGpu*) would exclude anything that contains MyGpu"),
+		FConsoleVariableDelegate::CreateStatic(RequestUpdateSettings),
 		ECVF_Scalability | ECVF_Default
 	);
 
@@ -168,15 +177,27 @@ namespace FNiagaraComponentSettings
 		TEXT("Set Gpu deny list to use, more targetted than to allow comparing OS,OSVersion,CPU,GPU.\n")
 		TEXT("Format is OSLabel,OSVersion,CPU,GPU| blank entries are assumed to auto pass matching.\n")
 		TEXT("For example, =\",,MyCpu,MyGpu+MyOS,,,\" would match MyCpu & MyGpu or MyOS."),
+		FConsoleVariableDelegate::CreateStatic(RequestUpdateSettings),
+		ECVF_Scalability | ECVF_Default
+	);
+
+	//-TODO: We can remove this once we confirm all RHIs report the correct values
+	static bool GGpuEmitterCheckFloat16Support = true;
+	static FAutoConsoleVariableRef CVarNiagaraGpuEmitterCheckFloat16Support(
+		TEXT("fx.Niagara.GpuEmitterCheckFloat16Support"),
+		GGpuEmitterCheckFloat16Support,
+		TEXT("When enabled we check to see if the RHI has support for Float16 UAV read / write, if it doesn't GPU emitters that use Float16 are banned from running."),
+		FConsoleVariableDelegate::CreateStatic(RequestUpdateSettings),
 		ECVF_Scalability | ECVF_Default
 	);
 
 	void UpdateSettings()
 	{
-		if (GDynamicRHI == nullptr)
+		if (!bNeedsSettingsUpdate || GDynamicRHI == nullptr)
 		{
 			return;
 		}
+		bNeedsSettingsUpdate = false;
 
 		bool bShouldAllowGpuEmitters = true;
 		if (GpuRHIDenyListString.Len() > 0)
@@ -236,8 +257,10 @@ namespace FNiagaraComponentSettings
 			}
 		}
 
+		bool bNeedsScalabilityUpdate = false;
 		if (bAllowGpuEmitters != bShouldAllowGpuEmitters)
 		{
+			bNeedsScalabilityUpdate = true;
 			bAllowGpuEmitters = bShouldAllowGpuEmitters;
 			if (bAllowGpuEmitters == false)
 			{
@@ -252,6 +275,22 @@ namespace FNiagaraComponentSettings
 				UE_LOG(LogNiagara, Log, TEXT("GPU emitters are enabled for RHI(%s) Adapter(%s)."), GDynamicRHI->GetName(), *GRHIAdapterName);
 			}
 		}
+
+		const bool Float16UAVSupported = GGpuEmitterCheckFloat16Support ? UE::PixelFormat::HasCapabilities(EPixelFormat::PF_R16F, EPixelFormatCapabilities::TypedUAVLoad | EPixelFormatCapabilities::TypedUAVStore) : true;
+		if (bAllowGpuComputeFloat16 != Float16UAVSupported)
+		{
+			bNeedsScalabilityUpdate = true;
+			bAllowGpuComputeFloat16 = Float16UAVSupported;
+			if (bAllowGpuComputeFloat16 == false)
+			{
+				UE_LOG(LogNiagara, Log, TEXT("UAV Float16 is not supported, compressed GPU emitters will be disabled."));
+			}
+		}
+
+		if (bNeedsScalabilityUpdate)
+		{
+			UpdateScalabilityForSystems();
+		}
 	}
 
 	bool IsSystemAllowedToRun(const UNiagaraSystem* System)
@@ -297,10 +336,36 @@ namespace FNiagaraComponentSettings
 				}
 			}
 		}
+
 		if (bAllowGpuEmitters == false)
 		{
 			return EmitterData.SimTarget != ENiagaraSimTarget::GPUComputeSim;
 		}
+
+		if (bAllowGpuComputeFloat16 == false && (EmitterData.SimTarget == ENiagaraSimTarget::GPUComputeSim))
+		{
+			// We have to lookup in the system to find if we used compressed attributes on this emitter or not
+			if (UNiagaraSystem* OwnerSystem = NiagaraEmitter.GetTypedOuter<UNiagaraSystem>())
+			{
+				const TArray<FNiagaraEmitterHandle>& EmitterHandles = OwnerSystem->GetEmitterHandles();
+				for (int32 iEmitter=0; iEmitter < EmitterHandles.Num(); ++iEmitter)
+				{
+					if (EmitterHandles[iEmitter].GetInstance().Emitter != &NiagaraEmitter)
+					{
+						continue;
+					}
+
+					const TArray<TSharedRef<const FNiagaraEmitterCompiledData>>& EmitterCompiledData = OwnerSystem->GetEmitterCompiledData();
+					if (!ensure(EmitterCompiledData.IsValidIndex(iEmitter)))
+					{
+						continue;
+					}
+
+					return EmitterCompiledData[iEmitter]->DataSetCompiledData.TotalHalfComponents == 0;
+				}
+			}
+		}
+
 		return true;
 	}
 }
diff --git a/Engine/Plugins/FX/Niagara/Source/Niagara/Private/NiagaraModule.cpp b/Engine/Plugins/FX/Niagara/Source/Niagara/Private/NiagaraModule.cpp
--- a/Engine/Plugins/FX/Niagara/Source/Niagara/Private/NiagaraModule.cpp
+++ b/Engine/Plugins/FX/Niagara/Source/Niagara/Private/NiagaraModule.cpp
@@ -573,6 +573,7 @@ void INiagaraModule::OnPostGarbageCollect()
 
 void INiagaraModule::OnBeginFrame()
 {
+	FNiagaraComponentSettings::UpdateSettings();
 	FNiagaraPlatformSet::RefreshScalability();
 }
 
