- Remove ensure as it can trip when we have not compelted compliation, once we complete compliation it will be recached

diff --git a/Engine/Plugins/FX/Niagara/Source/Niagara/Private/NiagaraComponentSettings.cpp b/Engine/Plugins/FX/Niagara/Source/Niagara/Private/NiagaraComponentSettings.cpp
--- a/Engine/Plugins/FX/Niagara/Source/Niagara/Private/NiagaraComponentSettings.cpp
+++ b/Engine/Plugins/FX/Niagara/Source/Niagara/Private/NiagaraComponentSettings.cpp
@@ -356,7 +356,7 @@ namespace FNiagaraComponentSettings
 					}
 
 					const TArray<TSharedRef<const FNiagaraEmitterCompiledData>>& EmitterCompiledData = OwnerSystem->GetEmitterCompiledData();
-					if (!ensure(EmitterCompiledData.IsValidIndex(iEmitter)))
+					if (!EmitterCompiledData.IsValidIndex(iEmitter))
 					{
 						continue;
 					}
