- Disable deformed reads on low end platforms due to shared samplers

diff --git a/Engine/Plugins/FX/Niagara/Shaders/Private/NiagaraDataInterfaceSkeletalMesh.ush b/Engine/Plugins/FX/Niagara/Shaders/Private/NiagaraDataInterfaceSkeletalMesh.ush
--- a/Engine/Plugins/FX/Niagara/Shaders/Private/NiagaraDataInterfaceSkeletalMesh.ush
+++ b/Engine/Plugins/FX/Niagara/Shaders/Private/NiagaraDataInterfaceSkeletalMesh.ush
@@ -27,6 +27,11 @@
 #define DISKELMESH_ADJ_INDEX_FORMAT DISKELMESH_ADJ_INDEX_FORMAT_FULL
 #endif
 
+#if FEATURE_LEVEL <= FEATURE_LEVEL_ES3_1
+#undef DISKELMESH_ALLOW_DEFORMED
+#define DISKELMESH_ALLOW_DEFORMED	0
+#endif
+
 //-TODO: We need to consolodate the random interface inside NiagaraEmitterInstanceShader.usf it's a bit messy at the moment but should follow what FNDIRandomHelper does
 float DISKelMesh_Random(uint S1, uint S2, uint S3)
 {
