Added wait implementation to niagara compilation manager + added wait call to engine tests

#virtualized

diff --git a/Engine/Plugins/FX/Niagara/Source/NiagaraEditor/Private/NiagaraSystemCompilingManager.cpp b/Engine/Plugins/FX/Niagara/Source/NiagaraEditor/Private/NiagaraSystemCompilingManager.cpp
--- a/Engine/Plugins/FX/Niagara/Source/NiagaraEditor/Private/NiagaraSystemCompilingManager.cpp
+++ b/Engine/Plugins/FX/Niagara/Source/NiagaraEditor/Private/NiagaraSystemCompilingManager.cpp
@@ -2,10 +2,9 @@
 
 #include "NiagaraSystemCompilingManager.h"
 
-#include "Algo/RemoveIf.h"
 #include "NiagaraCompilationTasks.h"
 #include "NiagaraEditorModule.h"
-
+#include "UObject/UObjectIterator.h"
 
 #include "Misc/ScopeRWLock.h"
 
@@ -54,10 +53,31 @@ int32 FNiagaraSystemCompilingManager::GetNumRemainingAssets() const
 
 void FNiagaraSystemCompilingManager::FinishCompilationForObjects(TArrayView<UObject* const> InObjects)
 {
+	if (InObjects.Num() == 0)
+	{
+		return;
+	}
+	for (UObject* Iter : InObjects)
+	{
+		if (UNiagaraSystem* Asset = Cast<UNiagaraSystem>(Iter))
+		{
+			Asset->WaitForCompilationComplete(true);
+		}
+	}
 }
 
 void FNiagaraSystemCompilingManager::FinishAllCompilation()
 {
+	TArray<UObject*> SystemsToFinish;
+	for (TObjectIterator<UNiagaraSystem> SystemIterator; SystemIterator; ++SystemIterator)
+	{
+		UNiagaraSystem* Asset = *SystemIterator;
+		if (Asset->HasOutstandingCompilationRequests(true))
+		{
+			SystemsToFinish.Add(Asset);
+		}
+	}
+	FinishCompilationForObjects(SystemsToFinish);
 }
 
 void FNiagaraSystemCompilingManager::Shutdown()
