Incorporate shaders for GPU simulations into the new AsyncTask compilation mode for Niagara systems.  Results in results being available to the user much more quickly in the editor (before issuing the shader compiles would only happen after the translation/cpu compilation results were completed and applied) and further cleans up the amount Niagara will block the cook process.  Currently only active with the AsyncTask compilation mode and gated by the cvar fx.Niagara.GpuScriptsCompiledBySystem.

diff --git a/Engine/Plugins/FX/Niagara/Source/Niagara/Classes/NiagaraScript.h b/Engine/Plugins/FX/Niagara/Source/Niagara/Classes/NiagaraScript.h
--- a/Engine/Plugins/FX/Niagara/Source/Niagara/Classes/NiagaraScript.h
+++ b/Engine/Plugins/FX/Niagara/Source/Niagara/Classes/NiagaraScript.h
@@ -245,7 +245,7 @@ public:
 	UPROPERTY()
 	TArray<FNiagaraVariableBase> AdditionalVariables;
 
-	NIAGARA_API TArray<FString> GetAdditionalVariableStrings();
+	NIAGARA_API TArray<FString> GetAdditionalVariableStrings() const;
 #endif
 
 	/** Whether or not we need to bake Rapid Iteration params. True to keep params, false to bake.*/
@@ -1030,6 +1030,7 @@ public:
 	bool IsDynamicInputScript() const { return IsDynamicInputScript(Usage); }
 	bool IsParticleEventScript() const { return IsParticleEventScript(Usage); }
 	bool IsParticleScript() const {	return IsParticleScript(Usage);}
+	bool IsGPUScript() const { return IsGPUScript(Usage); }
 
 	bool IsNonParticleScript() const { return IsNonParticleScript(Usage); }
 	
@@ -1082,6 +1083,7 @@ public:
 	NIAGARA_API bool IsReadyToRun(ENiagaraSimTarget SimTarget) const;
 #if WITH_EDITORONLY_DATA
 	NIAGARA_API static void SetPreviewFeatureLevel(ERHIFeatureLevel::Type PreviewFeatureLevel);
+	NIAGARA_API static ERHIFeatureLevel::Type GetPreviewFeatureLevel();
 #endif
 
 #if WITH_EDITORONLY_DATA
@@ -1127,6 +1129,10 @@ public:
 
 	// Infrastructure for GPU compute Shaders
 #if WITH_EDITOR
+	// utility function answering whether it's the NiagaraSystem or the NiagaraScript that will be responsible
+	// for the generation of the FNiagaraShaderMap and general compilation for GPU scripts
+	NIAGARA_API static bool AreGpuScriptsCompiledBySystem();
+
 	NIAGARA_API void CacheResourceShadersForCooking(EShaderPlatform ShaderPlatform, TArray<TUniquePtr<FNiagaraShaderScript>>& InOutCachedResources, const ITargetPlatform* TargetPlatform = nullptr);
 
 	NIAGARA_API void CacheResourceShadersForRendering(bool bRegenerateId, bool bForceRecompile=false);
@@ -1258,6 +1264,15 @@ public:
 	/** External call used to identify the values for a successful VM script compilation. OnVMScriptCompiled will be issued in this case.*/
 	void SetVMCompilationResults(const FNiagaraVMExecutableDataId& InCompileId, FNiagaraVMExecutableData& InScriptVM, const FString& EmitterUniqueName, const TMap<FName, UNiagaraDataInterface*>& ObjectNameMap, bool ApplyRapidIterationParameters);
 
+	void SetComputeCompilationResults(
+		const ITargetPlatform* TargetPlatform,
+		EShaderPlatform ShaderPlatform,
+		ERHIFeatureLevel::Type FeatureLevel,
+		const FNiagaraShaderScriptParametersMetadata& SharedParameters,
+		const FNiagaraShaderMapRef& ShaderMap);
+
+	NIAGARA_API bool IsShaderMapCached(const ITargetPlatform* TargetPlatform, const FNiagaraShaderMapId& ShaderMapId) const;
+
 	/** Updates the RI parameter store based on the provided variables (missing entries will be
 	    added and stale entries will be removed.  Returns true if the parameter store was modified. */
 	NIAGARA_API bool ApplyRapidIterationParameters(TConstArrayView<FNiagaraVariable> InParameters, bool bAllowRemoval);
@@ -1395,9 +1410,6 @@ private:
 	TSharedPtr<FNiagaraShaderScript> ScriptResourcesByFeatureLevel[ERHIFeatureLevel::Num];
 #endif
 
-	/** Compute shader compiled for this script */
-	FComputeShaderRHIRef ScriptShader;
-
 	/** Runtime stat IDs generated from StatScopes. */
 #if STATS
 	TArray<TStatId> StatScopesIDs;
diff --git a/Engine/Plugins/FX/Niagara/Source/Niagara/Classes/NiagaraSystem.h b/Engine/Plugins/FX/Niagara/Source/Niagara/Classes/NiagaraSystem.h
--- a/Engine/Plugins/FX/Niagara/Source/Niagara/Classes/NiagaraSystem.h
+++ b/Engine/Plugins/FX/Niagara/Source/Niagara/Classes/NiagaraSystem.h
@@ -48,6 +48,7 @@ struct FNiagaraGraphCachedDataBase;
 #endif
 
 class UNiagaraDataChannel;
+class ITargetPlatform;
 
 USTRUCT()
 struct FNiagaraEmitterCompiledData
@@ -402,7 +403,7 @@ public:
 	NIAGARA_API void RemoveSystemParametersForEmitter(const FNiagaraEmitterHandle& EmitterHandle);
 
 	/** Request that any dirty scripts referenced by this system be compiled.*/
-	NIAGARA_API bool RequestCompile(bool bForce, FNiagaraSystemUpdateContext* OptionalUpdateContext = nullptr);
+	NIAGARA_API bool RequestCompile(bool bForce, FNiagaraSystemUpdateContext* OptionalUpdateContext = nullptr, const ITargetPlatform* TargetPlatform = nullptr);
 
 	/** If we have a pending compile request, is it done with yet? */
 	NIAGARA_API bool PollForCompilationComplete(bool bFlushRequestCompile = true);
diff --git a/Engine/Plugins/FX/Niagara/Source/Niagara/Private/NiagaraCompilationTypes.cpp b/Engine/Plugins/FX/Niagara/Source/Niagara/Private/NiagaraCompilationTypes.cpp
--- a/Engine/Plugins/FX/Niagara/Source/Niagara/Private/NiagaraCompilationTypes.cpp
+++ b/Engine/Plugins/FX/Niagara/Source/Niagara/Private/NiagaraCompilationTypes.cpp
@@ -186,7 +186,7 @@ public:
 		FNiagaraSystemAsyncCompileResults NewCompileRequest;
 		INiagaraModule& NiagaraModule = FModuleManager::GetModuleChecked<INiagaraModule>("Niagara");
 
-		CompileRequestHandle = NiagaraModule.RequestCompileSystem(Options.System, bForced);
+		CompileRequestHandle = NiagaraModule.RequestCompileSystem(Options.System, bForced, Options.TargetPlatform);
 
 		return CompileRequestHandle != INDEX_NONE;
 	}
@@ -313,10 +313,17 @@ public:
 			ResultIt;
 			++ResultIt)
 		{
+			UNiagaraScript* TargetScript = ResultIt->Key;
 			const FNiagaraScriptAsyncCompileData& ScriptCompileData = ResultIt->Value;
-			if (ensure(ScriptCompileData.ExeData.IsValid()))
+
+			// because our compilation process includes the generation of rapid iteration parameters and static
+			// variables we need to generate the ExecutableDataId
+			// if we dirtied any RI parameters then we need to regenerate our CompilationId
+			FNiagaraVMExecutableDataId UpdatedCompileId;
+			TargetScript->ComputeVMCompilationId(UpdatedCompileId, FGuid());
+
+			if (ScriptCompileData.ExeData.IsValid())
 			{
-				UNiagaraScript* TargetScript = ResultIt->Key;
 				TMap<FName, UNiagaraDataInterface*> ObjectNameMap;
 
 				if (ScriptCompileData.bFromDerivedDataCache)
@@ -334,12 +341,6 @@ public:
 					});
 				}
 
-				// because our compilation process includes the generation of rapid iteration parameters and static
-				// variables we need to generate the ExecutableDataId
-				// if we dirtied any RI parameters then we need to regenerate our CompilationId
-				FNiagaraVMExecutableDataId UpdatedCompileId;
-				TargetScript->ComputeVMCompilationId(UpdatedCompileId, FGuid());
-
 				constexpr bool bApplyRapidIterationParameters = false;
 				TargetScript->SetVMCompilationResults(
 					UpdatedCompileId,
@@ -347,17 +348,33 @@ public:
 					ScriptCompileData.UniqueEmitterName,
 					ObjectNameMap,
 					bApplyRapidIterationParameters);
+
+				if (!ScriptCompileData.CompiledShaders.IsEmpty())
+				{
+					for (const FNiagaraCompiledShaderInfo& ShaderMapInfo : ScriptCompileData.CompiledShaders)
+					{
+						TargetScript->SetComputeCompilationResults(
+							ShaderMapInfo.TargetPlatform,
+							ShaderMapInfo.ShaderPlatform,
+							ShaderMapInfo.FeatureLevel,
+							ScriptCompileData.ExeData->ShaderScriptParametersMetadata,
+							ShaderMapInfo.CompiledShader);
+					}
+
+				}
 			}
 		}
 
 		// Synchronize the variables that we actually encountered during precompile so that we can expose them to the end user.
-		FNiagaraUserRedirectionParameterStore& ExposedParameters = Options.System->GetExposedParameters();
-
-		for (const FNiagaraVariable& ExposedVariable : CompileResults.ExposedVariables)
 		{
-			if (!ExposedParameters.FindParameterOffset(ExposedVariable))
+			FNiagaraUserRedirectionParameterStore& ExposedParameters = Options.System->GetExposedParameters();
+
+			for (const FNiagaraVariable& ExposedVariable : CompileResults.ExposedVariables)
 			{
-				ExposedParameters.AddParameter(ExposedVariable, true, false);
+				if (!ExposedParameters.FindParameterOffset(ExposedVariable))
+				{
+					ExposedParameters.AddParameter(ExposedVariable, true, false);
+				}
 			}
 		}
 	}
@@ -748,13 +765,13 @@ private:
 TUniquePtr<FNiagaraActiveCompilation> FNiagaraActiveCompilation::CreateCompilation()
 {
 #if WITH_EDITORONLY_DATA
-	const ENiagaraCompilationMode Compilationmode = GetDefault<UNiagaraSettings>()->CompilationMode;
+	const ENiagaraCompilationMode CompilationMode = GetDefault<UNiagaraSettings>()->CompilationMode;
 
-	if (Compilationmode == ENiagaraCompilationMode::AsyncTasks)
+	if (CompilationMode == ENiagaraCompilationMode::AsyncTasks)
 	{
 		return MakeUnique<FNiagaraActiveCompilationAsyncTask>();
 	}
-	else if (Compilationmode == ENiagaraCompilationMode::Verify)
+	else if (CompilationMode == ENiagaraCompilationMode::Verify)
 	{
 		return MakeUnique<FNiagaraActiveCompilationVerify>();
 	}
diff --git a/Engine/Plugins/FX/Niagara/Source/Niagara/Private/NiagaraEmitter.cpp b/Engine/Plugins/FX/Niagara/Source/Niagara/Private/NiagaraEmitter.cpp
--- a/Engine/Plugins/FX/Niagara/Source/Niagara/Private/NiagaraEmitter.cpp
+++ b/Engine/Plugins/FX/Niagara/Source/Niagara/Private/NiagaraEmitter.cpp
@@ -1631,14 +1631,17 @@ void FVersionedNiagaraEmitterData::CacheFromCompiledData(const FNiagaraDataSetCo
 
 	// Cache shaders for all GPU scripts
 #if WITH_EDITORONLY_DATA
-	if (AreAllScriptAndSourcesSynchronized())
+	if (!UNiagaraScript::AreGpuScriptsCompiledBySystem())
 	{
-		ForEachScript(
-			[](UNiagaraScript* Script)
-			{
-				Script->CacheResourceShadersForRendering(false, false);
-			}
-		);
+		if (AreAllScriptAndSourcesSynchronized())
+		{
+			ForEachScript(
+				[](UNiagaraScript* Script)
+				{
+					Script->CacheResourceShadersForRendering(false, false);
+				}
+			);
+		}
 	}
 #endif
 
diff --git a/Engine/Plugins/FX/Niagara/Source/Niagara/Private/NiagaraModule.cpp b/Engine/Plugins/FX/Niagara/Source/Niagara/Private/NiagaraModule.cpp
--- a/Engine/Plugins/FX/Niagara/Source/Niagara/Private/NiagaraModule.cpp
+++ b/Engine/Plugins/FX/Niagara/Source/Niagara/Private/NiagaraModule.cpp
@@ -772,10 +772,10 @@ void INiagaraModule::UnregisterGraphTraversalCacher(FDelegateHandle DelegateHand
 	GraphTraversalCacheDelegate.Unbind();
 }
 
-FNiagaraCompilationTaskHandle INiagaraModule::RequestCompileSystem(UNiagaraSystem* System, bool bForce)
+FNiagaraCompilationTaskHandle INiagaraModule::RequestCompileSystem(UNiagaraSystem* System, bool bForce, const ITargetPlatform* TargetPlatform)
 {
 	checkf(RequestCompileSystemDelegate.IsBound(), TEXT("RequestCompileSystemDelegate delegate not bound."));
-	return RequestCompileSystemDelegate.Execute(System, bForce);
+	return RequestCompileSystemDelegate.Execute(System, bForce, TargetPlatform);
 }
 
 FDelegateHandle INiagaraModule::RegisterRequestCompileSystem(FOnRequestCompileSystem RequestCompileSystemCallback)
diff --git a/Engine/Plugins/FX/Niagara/Source/Niagara/Private/NiagaraScript.cpp b/Engine/Plugins/FX/Niagara/Source/Niagara/Private/NiagaraScript.cpp
--- a/Engine/Plugins/FX/Niagara/Source/Niagara/Private/NiagaraScript.cpp
+++ b/Engine/Plugins/FX/Niagara/Source/Niagara/Private/NiagaraScript.cpp
@@ -136,6 +136,14 @@ static FAutoConsoleVariableRef CVarNiagaraStripByteCodeOverride(
 	ECVF_Default
 );
 
+static bool GNiagaraGpuScriptsCompiledBySystem = true;
+static FAutoConsoleVariableRef CVarNiagaraGpuScriptsCompiledBySystem(
+	TEXT("fx.Niagara.GpuScriptsCompiledBySystem"),
+	GNiagaraGpuScriptsCompiledBySystem,
+	TEXT("If true GPU shaders for Niagara scripts will be compiled along with the NiagaraSystem (when using AsyncTask compilation mode). \n"),
+	ECVF_Default
+);
+
 namespace NiagaraScriptInternal
 {
 #if WITH_EDITORONLY_DATA
@@ -757,7 +765,7 @@ void FNiagaraVMExecutableDataId::Invalidate()
 
 #if WITH_EDITORONLY_DATA
 
-TArray<FString> FNiagaraVMExecutableDataId::GetAdditionalVariableStrings()
+TArray<FString> FNiagaraVMExecutableDataId::GetAdditionalVariableStrings() const
 {
 	TArray<FString> Vars;
 	for (const FNiagaraVariableBase& Var : AdditionalVariables)
@@ -2373,24 +2381,46 @@ void UNiagaraScript::SetPreviewFeatureLevel(ERHIFeatureLevel::Type InPreviewFeat
 	}
 
 	// Force scripts to recache that can run on the GPU and already have a script resource
-	for (TObjectIterator<UNiagaraScript> It; It; ++It)
+	if (!AreGpuScriptsCompiledBySystem())
 	{
-		UNiagaraScript* Script = *It;
-		if ( Script == nullptr || !Script->OwnerCanBeRunOnGpu() || Script->ScriptResource == nullptr )
+		for (TObjectIterator<UNiagaraScript> It; It; ++It)
 		{
-			continue;
-		}
+			UNiagaraScript* Script = *It;
+			if ( Script == nullptr || !Script->OwnerCanBeRunOnGpu() || Script->ScriptResource == nullptr )
+			{
+				continue;
+			}
 
-		for (TSharedPtr<FNiagaraShaderScript>& ExistingScript : Script->ScriptResourcesByFeatureLevel )
+			for (TSharedPtr<FNiagaraShaderScript>& ExistingScript : Script->ScriptResourcesByFeatureLevel )
+			{
+				if (ExistingScript == Script->ScriptResource)
+				{
+					Script->CacheResourceShadersForRendering(true, false);
+					break;
+				}
+			}
+		}
+	}
+	else
+	{
+		for (TObjectIterator<UNiagaraSystem> It; It; ++It)
 		{
-			if (ExistingScript == Script->ScriptResource)
+			if (UNiagaraSystem* System = *It)
 			{
-				Script->CacheResourceShadersForRendering(true, false);
-				break;
+				if (System->bFullyLoaded && System->HasAnyGPUEmitters())
+				{
+					System->RequestCompile(false);
+				}
 			}
 		}
 	}
 }
+
+ERHIFeatureLevel::Type UNiagaraScript::GetPreviewFeatureLevel()
+{
+	return NiagaraScriptInternal::PreviewFeatureLevel.Get(GMaxRHIFeatureLevel);
+}
+
 #endif
 
 void UNiagaraScript::GenerateStatIDs()
@@ -2440,7 +2470,11 @@ void UNiagaraScript::PostEditChangeProperty(FPropertyChangedEvent& PropertyChang
 {
 	Super::PostEditChangeProperty(PropertyChangedEvent);
 
-	CacheResourceShadersForRendering(true);
+	if (!AreGpuScriptsCompiledBySystem())
+	{
+		CacheResourceShadersForRendering(true);
+	}
+
 	CustomAssetRegistryTagCache.Reset();
 	OnPropertyChangedDelegate.Broadcast(PropertyChangedEvent);
 }
@@ -2953,17 +2987,20 @@ void UNiagaraScript::SetVMCompilationResults(const FNiagaraVMExecutableDataId& I
 	GenerateStatIDs();
 
 	// Now go ahead and trigger the GPU script compile now that we have a compiled GPU hlsl script.
-	if (Usage == ENiagaraScriptUsage::ParticleGPUComputeScript)
+	if (!AreGpuScriptsCompiledBySystem())
 	{
-		if (CachedScriptVMId.CompilerVersionID.IsValid() && CachedScriptVMId.BaseScriptCompileHash.IsValid())
-		{
-			CacheResourceShadersForRendering(false, true);
-		}
-		else
+		if (Usage == ENiagaraScriptUsage::ParticleGPUComputeScript)
 		{
-			UE_LOG(LogNiagara, Warning,
-				TEXT("Could not cache resource shaders for rendering for script %s because it had an invalid cached script id. This should be fixed by force recompiling the owning asset using the 'Full Rebuild' option and then saving the asset."),
-				*GetPathName());
+			if (CachedScriptVMId.CompilerVersionID.IsValid() && CachedScriptVMId.BaseScriptCompileHash.IsValid())
+			{
+				CacheResourceShadersForRendering(false, true);
+			}
+			else
+			{
+				UE_LOG(LogNiagara, Warning,
+					TEXT("Could not cache resource shaders for rendering for script %s because it had an invalid cached script id. This should be fixed by force recompiling the owning asset using the 'Full Rebuild' option and then saving the asset."),
+					*GetPathName());
+			}
 		}
 	}
 
@@ -2978,6 +3015,98 @@ void UNiagaraScript::SetVMCompilationResults(const FNiagaraVMExecutableDataId& I
 	OnVMScriptCompiled().Broadcast(this, InCompileId.ScriptVersionID);
 }
 
+void UNiagaraScript::SetComputeCompilationResults(
+	const ITargetPlatform* TargetPlatform,
+	EShaderPlatform ShaderPlatform,
+	ERHIFeatureLevel::Type FeatureLevel,
+	const FNiagaraShaderScriptParametersMetadata& ShaderParameters,
+	const FNiagaraShaderMapRef& ShaderMap)
+{
+	// for now we will use the CachedScriptResourcesForCooking for storing all compilation results and will update ScriptResource
+	// if it matches with the supplied parameters
+	auto ScriptMatchesFeatureLevel = [FeatureLevel](const TUniquePtr<FNiagaraShaderScript>& ShaderScript) -> bool
+	{
+		return ShaderScript->GetFeatureLevel() == FeatureLevel;
+	};
+
+	TArray<TUniquePtr<FNiagaraShaderScript>>& CachedScripts = CachedScriptResourcesForCooking.FindOrAdd(TargetPlatform);
+	FNiagaraShaderScript* TargetShaderScript = nullptr;
+
+	if (TUniquePtr<FNiagaraShaderScript>* ExistingScript = CachedScripts.FindByPredicate(ScriptMatchesFeatureLevel))
+	{
+		TargetShaderScript = ExistingScript->Get();
+	}
+	else
+	{
+		TargetShaderScript = CachedScripts.Emplace_GetRef(MakeUnique<FNiagaraShaderScript>()).Get();
+	}
+
+	auto UpdateShaderScript = [this, FeatureLevel, ShaderPlatform, &ShaderParameters, ShaderMap](FNiagaraShaderScript* ScriptToUpdate) -> void
+	{
+		ScriptToUpdate->BuildScriptParametersMetadata(ShaderParameters);
+		ScriptToUpdate->SetShaderMap(ShaderMap);
+		ScriptToUpdate->SetScript(this, FeatureLevel, ShaderPlatform, CachedScriptVMId.CompilerVersionID, CachedScriptVMId.AdditionalDefines,
+			CachedScriptVMId.GetAdditionalVariableStrings(),
+			CachedScriptVMId.BaseScriptCompileHash, CachedScriptVMId.ReferencedCompileHashes,
+			CachedScriptVMId.bUsesRapidIterationParams, GetFriendlyName());
+	};
+
+	if (ensure(TargetShaderScript))
+	{
+		UpdateShaderScript(TargetShaderScript);
+
+		// next see if we need to update the currently active ShaderScript (and update it's renderthread data)
+		{
+			const ERHIFeatureLevel::Type ActiveFeatureLevel = NiagaraScriptInternal::PreviewFeatureLevel.Get(GMaxRHIFeatureLevel);
+			const EShaderPlatform ActiveShaderPlatform = GShaderPlatformForFeatureLevel[ActiveFeatureLevel];
+
+			if (ActiveFeatureLevel == FeatureLevel && ActiveShaderPlatform == ShaderPlatform)
+			{
+				if (!ScriptResource.IsValid())
+				{
+					ScriptResource = MakeShared<FNiagaraShaderScript>(*TargetShaderScript);
+				}
+				else
+				{
+					UpdateShaderScript(ScriptResource.Get());
+				}
+
+				ENQUEUE_RENDER_COMMAND(FSetShaderMapOnScriptResources)(
+					[RT_Script = ScriptResource, RT_ShaderMap = ShaderMap](FRHICommandListImmediate& RHICmdList)
+				{
+					RT_Script->SetRenderingThreadShaderMap(RT_ShaderMap);
+				});
+			}
+		}
+	}
+}
+
+bool UNiagaraScript::IsShaderMapCached(const ITargetPlatform* TargetPlatform, const FNiagaraShaderMapId& ShaderMapId) const
+{
+	if (TargetPlatform)
+	{
+		if (const TArray<TUniquePtr<FNiagaraShaderScript>>* CachedShaders = CachedScriptResourcesForCooking.Find(TargetPlatform))
+		{
+			for (const TUniquePtr<FNiagaraShaderScript>& CachedShader : *CachedShaders)
+			{
+				if (CachedShader.IsValid() && CachedShader->IsSame(ShaderMapId))
+				{
+					return CachedShader->IsShaderMapComplete();
+				}
+			}
+		}
+	}
+	else
+	{
+		if (ScriptResource.IsValid())
+		{
+			return ScriptResource->IsSame(ShaderMapId) && ScriptResource->IsShaderMapComplete();
+		}
+	}
+
+	return false;
+}
+
 bool UNiagaraScript::ApplyRapidIterationParameters(TConstArrayView<FNiagaraVariable> InParameters, bool bAllowRemoval)
 {
 	bool bStoreDirty = false;
@@ -3416,6 +3545,12 @@ TArray<FName> UNiagaraScript::FindShaderFormatsForCooking(const ITargetPlatform*
 
 void UNiagaraScript::BeginCacheForCookedPlatformData(const ITargetPlatform *TargetPlatform)
 {
+	if (AreGpuScriptsCompiledBySystem())
+	{
+		// if the system is responsible for generating the GPU scripts, then we don't need to do anything here
+		return;
+	}
+
 	TArray<FName> ShaderFormatsForCooking = FindShaderFormatsForCooking(TargetPlatform);
 
 	if (!ShaderFormatsForCooking.IsEmpty())
@@ -3455,6 +3590,19 @@ bool UNiagaraScript::IsCachedCookedPlatformDataLoaded(const ITargetPlatform* Tar
 		return true;
 	}
 
+	if (AreGpuScriptsCompiledBySystem())
+	{
+		if (UNiagaraSystem* SystemOwner = FindRootSystem())
+		{
+			if (SystemOwner->HasOutstandingCompilationRequests(true))
+			{
+				return false;
+			}
+		}
+
+		return true;
+	}
+
 	const TArray<FName> ShaderFormats = FindShaderFormatsForCooking(TargetPlatform);
 
 	if (!ShaderFormats.IsEmpty())
@@ -3494,6 +3642,14 @@ bool UNiagaraScript::IsCachedCookedPlatformDataLoaded(const ITargetPlatform* Tar
 	return true;
 }
 
+// utility function answering whether it's the NiagaraSystme or the NiagaraScript that will be responsible
+// for the generation of the FNiagaraShaderMap and general compilation for GPU scripts
+bool UNiagaraScript::AreGpuScriptsCompiledBySystem()
+{
+	const ENiagaraCompilationMode CompilationMode = GetDefault<UNiagaraSettings>()->CompilationMode;
+	return GNiagaraGpuScriptsCompiledBySystem && CompilationMode == ENiagaraCompilationMode::AsyncTasks;
+}
+
 void UNiagaraScript::CacheResourceShadersForCooking(EShaderPlatform ShaderPlatform, TArray<TUniquePtr<FNiagaraShaderScript>>& InOutCachedResources, const ITargetPlatform* TargetPlatform)
 {
 	if (CanBeRunOnGpu())
@@ -3834,7 +3990,6 @@ bool UNiagaraScript::SynchronizeExecutablesWithCompilation(const UNiagaraScript*
 	FNiagaraVMExecutableDataId Id;
 	ComputeVMCompilationId(Id, FGuid());
 
-#if 1 // TODO Shaun... turn this on...
 	if (Id == Script->GetVMExecutableDataCompilationId())
 	{
 		CachedScriptVM.Reset();
@@ -3860,15 +4015,17 @@ bool UNiagaraScript::SynchronizeExecutablesWithCompilation(const UNiagaraScript*
 		//SyncAliases(RenameMap);
 
 		// Now go ahead and trigger the GPU script compile now that we have a compiled GPU hlsl script.
-		if (Usage == ENiagaraScriptUsage::ParticleGPUComputeScript)
+		if (!AreGpuScriptsCompiledBySystem())
 		{
-			CacheResourceShadersForRendering(false, true);
+			if (Usage == ENiagaraScriptUsage::ParticleGPUComputeScript)
+			{
+				CacheResourceShadersForRendering(false, true);
+			}
 		}
 
 		OnVMScriptCompiled().Broadcast(this, FGuid()); // TODO what version id to use here?
 		return true;
 	}
-#endif
 
 	return false;
 }
diff --git a/Engine/Plugins/FX/Niagara/Source/Niagara/Private/NiagaraSystem.cpp b/Engine/Plugins/FX/Niagara/Source/Niagara/Private/NiagaraSystem.cpp
--- a/Engine/Plugins/FX/Niagara/Source/Niagara/Private/NiagaraSystem.cpp
+++ b/Engine/Plugins/FX/Niagara/Source/Niagara/Private/NiagaraSystem.cpp
@@ -281,7 +281,7 @@ void UNiagaraSystem::BeginCacheForCookedPlatformData(const ITargetPlatform *Targ
 
 	if (bNeedsRequestCompile)
 	{
-		RequestCompile(false);
+		RequestCompile(false, nullptr, TargetPlatform);
 	}
 
 	// check if any of the active compilations requires waiting - note that none should but the original
@@ -3091,7 +3091,7 @@ void UNiagaraSystem::InvalidateCachedData()
 	CachedTraversalData.Reset();
 }
 
-bool UNiagaraSystem::RequestCompile(bool bForce, FNiagaraSystemUpdateContext* OptionalUpdateContext)
+bool UNiagaraSystem::RequestCompile(bool bForce, FNiagaraSystemUpdateContext* OptionalUpdateContext, const ITargetPlatform* TargetPlatform)
 {
 	check(IsInGameThread());
 	TRACE_CPUPROFILER_EVENT_SCOPE(UNiagaraSystem::RequestCompile)
@@ -3149,6 +3149,7 @@ bool UNiagaraSystem::RequestCompile(bool bForce, FNiagaraSystemUpdateContext* Op
 
 		FNiagaraCompilationOptions Options;
 		Options.System = this;
+		Options.TargetPlatform = TargetPlatform;
 		Options.bForced = bForce;
 
 		bLaunchedCompilations = CurrentCompilation->Launch(Options);
diff --git a/Engine/Plugins/FX/Niagara/Source/Niagara/Public/NiagaraCommon.h b/Engine/Plugins/FX/Niagara/Source/Niagara/Public/NiagaraCommon.h
--- a/Engine/Plugins/FX/Niagara/Source/Niagara/Public/NiagaraCommon.h
+++ b/Engine/Plugins/FX/Niagara/Source/Niagara/Public/NiagaraCommon.h
@@ -46,8 +46,10 @@ enum EShaderPlatform : uint16;
 enum EPixelFormat : uint8;
 
 class FNiagaraDataBuffer;
+class FNiagaraShaderMap;
 
 using FNiagaraDataBufferRef = TRefCountPtr<FNiagaraDataBuffer>;
+using FNiagaraShaderMapRef = TRefCountPtr<FNiagaraShaderMap>;
 
 #define NIAGARA_MEMORY_TRACKING	!UE_BUILD_SHIPPING
 
diff --git a/Engine/Plugins/FX/Niagara/Source/Niagara/Public/NiagaraCompilationTypes.h b/Engine/Plugins/FX/Niagara/Source/Niagara/Public/NiagaraCompilationTypes.h
--- a/Engine/Plugins/FX/Niagara/Source/Niagara/Public/NiagaraCompilationTypes.h
+++ b/Engine/Plugins/FX/Niagara/Source/Niagara/Public/NiagaraCompilationTypes.h
@@ -10,6 +10,7 @@
 #include "NiagaraCompilationTypes.generated.h"
 
 class UNiagaraSystem;
+class ITargetPlatform;
 
 using FNiagaraCompilationTaskHandle = int32;
 
@@ -32,6 +33,14 @@ struct FNiagaraSystemCompileMetrics
 	TMap<TObjectKey<UNiagaraScript>, FNiagaraScriptCompileMetrics> ScriptMetrics;
 };
 
+struct FNiagaraCompiledShaderInfo
+{
+	FNiagaraShaderMapRef CompiledShader;
+	const ITargetPlatform* TargetPlatform = nullptr;
+	EShaderPlatform ShaderPlatform = SP_NumPlatforms;
+	ERHIFeatureLevel::Type FeatureLevel = ERHIFeatureLevel::Num;
+};
+
 USTRUCT()
 struct FNiagaraScriptAsyncCompileData
 {
@@ -39,6 +48,7 @@ struct FNiagaraScriptAsyncCompileData
 
 	FNiagaraVMExecutableDataId CompileId;
 	TSharedPtr<struct FNiagaraVMExecutableData> ExeData;
+	TArray<FNiagaraCompiledShaderInfo> CompiledShaders;
 	FString UniqueEmitterName;
 	FNiagaraEmitterID EmitterID = INDEX_NONE;
 	FNiagaraScriptCompileMetrics CompileMetrics;
@@ -77,6 +87,7 @@ struct FNiagaraSystemAsyncCompileResults
 struct FNiagaraCompilationOptions
 {
 	UNiagaraSystem* System = nullptr;
+	const ITargetPlatform* TargetPlatform = nullptr;
 	bool bForced = false;
 };
 
diff --git a/Engine/Plugins/FX/Niagara/Source/Niagara/Public/NiagaraModule.h b/Engine/Plugins/FX/Niagara/Source/Niagara/Public/NiagaraModule.h
--- a/Engine/Plugins/FX/Niagara/Source/Niagara/Public/NiagaraModule.h
+++ b/Engine/Plugins/FX/Niagara/Source/Niagara/Public/NiagaraModule.h
@@ -45,7 +45,7 @@ public:
 	DECLARE_DELEGATE_RetVal_FiveParams(CompileRequestDuplicatePtr, FOnPrecompileDuplicate, const FNiagaraCompileRequestDataBase* /*OwningSystemRequestData*/, UNiagaraSystem* /*OwningSystem*/, UNiagaraEmitter* /*OwningEmitter*/, UNiagaraScript* /*TargetScript*/, FGuid /*Version*/);
 	DECLARE_DELEGATE_RetVal_TwoParams(GraphCachedDataPtr, FOnCacheGraphTraversal, const UObject*, FGuid);
 
-	DECLARE_DELEGATE_RetVal_TwoParams(FNiagaraCompilationTaskHandle, FOnRequestCompileSystem, UNiagaraSystem*, bool);
+	DECLARE_DELEGATE_RetVal_ThreeParams(FNiagaraCompilationTaskHandle, FOnRequestCompileSystem, UNiagaraSystem*, bool, const ITargetPlatform*);
 	DECLARE_DELEGATE_RetVal_FourParams(bool, FOnPollSystemCompile, FNiagaraCompilationTaskHandle, FNiagaraSystemAsyncCompileResults&, bool /*bWait*/, bool /*bPeek*/);
 	DECLARE_DELEGATE_OneParam(FOnAbortSystemCompile, FNiagaraCompilationTaskHandle);
 
@@ -111,7 +111,7 @@ public:
 	NIAGARA_API FDelegateHandle RegisterGraphTraversalCacher(FOnCacheGraphTraversal PreCompiler);
 	NIAGARA_API void UnregisterGraphTraversalCacher(FDelegateHandle DelegateHandle);
 
-	NIAGARA_API FNiagaraCompilationTaskHandle RequestCompileSystem(UNiagaraSystem* System, bool bForce);
+	NIAGARA_API FNiagaraCompilationTaskHandle RequestCompileSystem(UNiagaraSystem* System, bool bForce, const ITargetPlatform* TargetPlatform);
 	NIAGARA_API FDelegateHandle RegisterRequestCompileSystem(FOnRequestCompileSystem RequestCompileSystemCallback);
 	NIAGARA_API void UnregisterRequestCompileSystem(FDelegateHandle DelegateHandle);
 
diff --git a/Engine/Plugins/FX/Niagara/Source/NiagaraEditor/Private/NiagaraCompilationTasks.cpp b/Engine/Plugins/FX/Niagara/Source/NiagaraEditor/Private/NiagaraCompilationTasks.cpp
--- a/Engine/Plugins/FX/Niagara/Source/NiagaraEditor/Private/NiagaraCompilationTasks.cpp
+++ b/Engine/Plugins/FX/Niagara/Source/NiagaraEditor/Private/NiagaraCompilationTasks.cpp
@@ -3,6 +3,8 @@
 #include "NiagaraCompilationTasks.h"
 
 #include "Algo/RemoveIf.h"
+#include "DataDrivenShaderPlatformInfo.h"
+#include "Interfaces/ITargetPlatformManagerModule.h"
 #include "Misc/PathViews.h"
 #include "NiagaraCompilationTypes.h"
 #include "NiagaraDigestDatabase.h"
@@ -11,17 +13,49 @@
 #include "NiagaraCompilationBridge.h"
 #include "NiagaraScriptSource.h"
 #include "NiagaraShader.h"
+#include "NiagaraShaderParametersBuilder.h"
 #include "NiagaraSimulationStageBase.h"
 #include "NiagaraSystem.h"
 #include "NiagaraSystemCompilingManager.h"
+#include "ProfilingDebugging/CookStats.h"
 #include "Serialization/ObjectAndNameAsStringProxyArchive.h"
+#include "ShaderParameterMetadataBuilder.h"
 
 // needed for AsType...pretty ugly
 #include "NiagaraNodeOutput.h"
 #include "NiagaraNodeFunctionCall.h"
 
-namespace NiagaraCompilationCopyImpl
+#if ENABLE_COOK_STATS
+namespace NiagaraSystemCookStats
 {
+	extern FCookStats::FDDCResourceUsageStats UsageStats;
+
+	static int32 ScriptTranslationCount = 0;
+	static int32 CpuScriptCompileCount = 0;
+	static int32 CpuScriptDdcCacheHitCount = 0;
+	static int32 CpuScriptDdcCacheMissCount = 0;
+	static int32 GpuScriptCompileCount = 0;
+	static int32 GpuScriptDdcCacheHitCount = 0;
+	static int32 GpuScriptDdcCacheMissCount = 0;
+
+	static FCookStatsManager::FAutoRegisterCallback RegisterNiagaraCompilationCookStats([](FCookStatsManager::AddStatFuncRef AddStat)
+	{
+		AddStat(TEXT("Niagara"), FCookStatsManager::CreateKeyValueArray(TEXT("NiagaraScriptTranslationCount"), ScriptTranslationCount));
+		AddStat(TEXT("Niagara"), FCookStatsManager::CreateKeyValueArray(TEXT("NiagaraCpuScriptCompileCount"), CpuScriptCompileCount));
+		AddStat(TEXT("Niagara"), FCookStatsManager::CreateKeyValueArray(TEXT("NiagaraCpuScriptDdcCacheHitCount"), CpuScriptDdcCacheHitCount));
+		AddStat(TEXT("Niagara"), FCookStatsManager::CreateKeyValueArray(TEXT("NiagaraCpuScriptDdcCacheMissCount"), CpuScriptDdcCacheMissCount));
+		AddStat(TEXT("Niagara"), FCookStatsManager::CreateKeyValueArray(TEXT("NiagaraGpuScriptCompileCount"), GpuScriptCompileCount));
+		AddStat(TEXT("Niagara"), FCookStatsManager::CreateKeyValueArray(TEXT("NiagaraGpuScriptDdcCacheHitCount"), GpuScriptDdcCacheHitCount));
+		AddStat(TEXT("Niagara"), FCookStatsManager::CreateKeyValueArray(TEXT("NiagaraGpuScriptDdcCacheMissCount"), GpuScriptDdcCacheMissCount));
+	});
+}
+#endif
+
+
+namespace NiagaraCompilationTasksImpl
+{
+	static UE::DerivedData::FCacheBucket NiagaraDDCBucket("NiagaraScript");
+
 	void GetUsagesToDuplicate(ENiagaraScriptUsage TargetUsage, TArray<ENiagaraScriptUsage>& DuplicateUsages)
 	{
 		// For now we need to include both spawn and update for each target usage, otherwise attribute lists in the precompiled histories aren't generated correctly.
@@ -113,8 +147,26 @@ namespace NiagaraCompilationCopyImpl
 
 		CompileId.AppendKeyString(KeyString);
 
-		static UE::DerivedData::FCacheBucket Bucket("NiagaraScript");
-		return { Bucket, FIoHash::HashBuffer(MakeMemoryView(FTCHARToUTF8(KeyString))) };
+		return { NiagaraDDCBucket, FIoHash::HashBuffer(MakeMemoryView(FTCHARToUTF8(KeyString))) };
+	}
+
+	static UE::DerivedData::FCacheKey BuildNiagaraComputeDDCCacheKey(const FNiagaraShaderMapId& ShaderMapId, EShaderPlatform ShaderPlatform)
+	{
+		static const FString NIAGARASHADERMAP_DERIVEDDATA_VER = FDevSystemGuids::GetSystemGuid(FDevSystemGuids::Get().NIAGARASHADERMAP_DERIVEDDATA_VER).ToString(EGuidFormats::DigitsWithHyphens);
+
+		FName Format = LegacyShaderPlatformToShaderFormat(ShaderPlatform);
+
+		FString KeyString;
+		KeyString.Reserve(1024);
+
+		Format.ToString(KeyString);
+		KeyString.Appendf(TEXT("_%d_"), GetTargetPlatformManagerRef().ShaderFormatVersion(Format));
+		ShaderMapAppendKeyString(ShaderPlatform, KeyString);
+		ShaderMapId.AppendKeyString(KeyString);
+		KeyString.AppendChar(TCHAR('_'));
+		KeyString.Append(NIAGARASHADERMAP_DERIVEDDATA_VER);
+
+		return { NiagaraDDCBucket, FIoHash::HashBuffer(MakeMemoryView(FTCHARToUTF8(KeyString))) };
 	}
 
 	static bool ValidateExecData(const UNiagaraScript* Script, const FNiagaraVMExecutableData& ExecData, FString& ErrorString)
@@ -195,6 +247,44 @@ namespace NiagaraCompilationCopyImpl
 		return FSharedBuffer();
 	}
 
+	static FSharedBuffer ShaderMapToBinaryData(FNiagaraShaderMap* ShaderMap)
+	{
+		if (!UNiagaraScript::ShouldUseDDC())
+		{
+			return FSharedBuffer();
+		}
+
+		TArray<uint8> BinaryData;
+		FMemoryWriter Ar(BinaryData, true);
+		ShaderMap->Serialize(Ar);
+		
+		if (!BinaryData.IsEmpty() && !Ar.IsError())
+		{
+			return MakeSharedBufferFromArray(MoveTemp(BinaryData));
+		}
+
+		return FSharedBuffer();
+	}
+
+	static FNiagaraShaderMapRef BinaryDataToShaderMap(FSharedBuffer SharedBuffer)
+	{
+		if (!UNiagaraScript::ShouldUseDDC() || SharedBuffer.GetSize() == 0)
+		{
+			return nullptr;
+		}
+
+		FMemoryReaderView Ar(SharedBuffer.GetView(), true);
+		FNiagaraShaderMapRef ShaderMap = new FNiagaraShaderMap(FNiagaraShaderMap::WorkerThread);
+		ShaderMap->Serialize(Ar);
+
+		if (Ar.IsError())
+		{
+			return nullptr;
+		}
+
+		return ShaderMap;
+	}
+
 	static TArray<FNiagaraVariable> ExtractInputVariablesFromHistory(const TNiagaraParameterMapHistory<FNiagaraCompilationDigestBridge>& History, const FNiagaraCompilationGraph* FunctionGraph)
 	{
 		TArray<FNiagaraVariable> InputVariables;
@@ -330,7 +420,58 @@ namespace NiagaraCompilationCopyImpl
 		}
 		return nullptr;
 	}
-} // NiagaraCompilationCopyImpl
+
+	TSharedPtr<FNiagaraShaderScriptParametersMetadata> BuildScriptParametersMetadata(const FNiagaraCompilationCopyData* CompilationCopyData, const FNiagaraShaderScriptParametersMetadata& InScriptParametersMetadata)
+	{
+		TSharedPtr<FNiagaraShaderScriptParametersMetadata> ScriptParametersMetadata = MakeShared<FNiagaraShaderScriptParametersMetadata>();
+		ScriptParametersMetadata->DataInterfaceParamInfo = InScriptParametersMetadata.DataInterfaceParamInfo;
+
+		FShaderParametersMetadataBuilder ShaderMetadataBuilder(TShaderParameterStructTypeInfo<FNiagaraShader::FParameters>::GetStructMetadata());
+
+		auto FindDataInterfaceByName = [CompilationCopyData](const FString& DIName) -> const UNiagaraDataInterfaceBase*
+		{
+			for (const auto& DataInterfacePair : CompilationCopyData->AggregatedDataInterfaceCDODuplicates)
+			{
+				if (DataInterfacePair.Key->GetName() == DIName)
+				{
+					return DataInterfacePair.Value;
+				}
+			}
+
+			return nullptr;
+		};
+
+		// Build meta data for each data interface
+		for (FNiagaraDataInterfaceGPUParamInfo& DataInterfaceParamInfo : ScriptParametersMetadata->DataInterfaceParamInfo)
+		{
+			const UNiagaraDataInterfaceBase* DataInterface = FindDataInterfaceByName(DataInterfaceParamInfo.DIClassName);
+			if (ensure(DataInterface))
+			{
+				const uint32 NextMemberOffset = ShaderMetadataBuilder.GetNextMemberOffset();
+				FNiagaraShaderParametersBuilder ShaderParametersBuilder(DataInterfaceParamInfo, ScriptParametersMetadata->LooseMetadataNames, ScriptParametersMetadata->StructIncludeInfos, ShaderMetadataBuilder);
+				DataInterface->BuildShaderParameters(ShaderParametersBuilder);
+				DataInterfaceParamInfo.ShaderParametersOffset = NextMemberOffset;
+			}
+		}
+
+		ScriptParametersMetadata->ShaderParametersMetadata = MakeShareable<FShaderParametersMetadata>(ShaderMetadataBuilder.Build(FShaderParametersMetadata::EUseCase::ShaderParameterStruct, TEXT("FNiagaraShaderScript")));
+
+		return ScriptParametersMetadata;
+	}
+
+	void BuildDebugGroupName(const FNiagaraPrecompileData& PrecompileData, const FNiagaraCompileOptions& CompileOptions, FStringBuilderBase& DebugGroupName)
+	{
+		FPathViews::Append(DebugGroupName, PrecompileData.SourceName);
+		FPathViews::Append(DebugGroupName, PrecompileData.EmitterUniqueName);
+		FPathViews::Append(DebugGroupName, PrecompileData.ENiagaraScriptUsageEnum->GetNameStringByValue((int64)CompileOptions.TargetUsage));
+
+		if (CompileOptions.TargetUsageId.IsValid())
+		{
+			DebugGroupName << TEXT("_");
+			CompileOptions.TargetUsageId.AppendString(DebugGroupName, EGuidFormats::Digits);
+		}
+	}
+} // NiagaraCompilationTasksImpl
 
 FNiagaraSystemCompilationTask::FNiagaraSystemCompilationTask(FNiagaraCompilationTaskHandle InTaskHandle, UNiagaraSystem* InSystem)
 	: TaskHandle(InTaskHandle)
@@ -362,7 +503,9 @@ bool FNiagaraSystemCompilationTask::FCompileGroupInfo::HasOutstandingCompileTask
 {
 	for (int32 CompileTaskIndex : CompileTaskIndices)
 	{
-		if (ParentTask.CompileTasks[CompileTaskIndex].IsOutstanding())
+		const FCompileTaskInfo& CompileTask = ParentTask.CompileTasks[CompileTaskIndex];
+
+		if (CompileTask.IsOutstanding(ParentTask))
 		{
 			return true;
 		}
@@ -460,11 +603,28 @@ void FNiagaraSystemCompilationTask::FCompileGroupInfo::InstantiateCompileGraph(c
 	}
 }
 
-bool FNiagaraSystemCompilationTask::FCompileTaskInfo::IsOutstanding() const
+bool FNiagaraSystemCompilationTask::FCompileComputeShaderTaskInfo::IsOutstanding() const
 {
-	return !ExeData.IsValid();
+	return !ShaderMap.IsValid();
 }
 
+bool FNiagaraSystemCompilationTask::FCompileTaskInfo::IsOutstanding(const FNiagaraSystemCompilationTask& ParentTask) const
+{
+	if (!ExeData.IsValid())
+	{
+		return true;
+	}
+	
+	for (int32 TaskIndex : ComputeShaderTaskIndices)
+	{
+		if (ParentTask.CompileComputeShaderTasks[TaskIndex].IsOutstanding())
+		{
+			return true;
+		}
+	}
+
+	return false;
+}
 
 FNiagaraSystemCompilationTask::FCompileGroupInfo* FNiagaraSystemCompilationTask::GetCompileGroupInfo(int32 EmitterIndex)
 {
@@ -557,8 +717,10 @@ void FNiagaraSystemCompilationTask::FCompileTaskInfo::CollectNamedDataInterfaces
 	}
 }
 
-void FNiagaraSystemCompilationTask::FCompileTaskInfo::TranslateAndIssueCompile(FNiagaraSystemCompilationTask* SystemCompileTask, const FCompileGroupInfo& GroupInfo)
+void FNiagaraSystemCompilationTask::FCompileTaskInfo::Translate(FNiagaraSystemCompilationTask* SystemCompileTask, const FCompileGroupInfo& GroupInfo)
 {
+	COOK_STAT(NiagaraSystemCookStats::ScriptTranslationCount++);
+
 	const FNiagaraPrecompileData* PrecompileData = nullptr;
 	const FNiagaraCompilationCopyData* CompilationCopyData = nullptr;
 
@@ -604,24 +766,95 @@ void FNiagaraSystemCompilationTask::FCompileTaskInfo::TranslateAndIssueCompile(F
 
 		TranslationTime = (float) (FPlatformTime::Seconds() - TranslationStartTime);
 	}
+}
+
+void FNiagaraSystemCompilationTask::FCompileTaskInfo::IssueCompileVm(FNiagaraSystemCompilationTask* SystemCompileTask, const FCompileGroupInfo& GroupInfo)
+{
+	COOK_STAT(NiagaraSystemCookStats::CpuScriptCompileCount++);
+
+	const FNiagaraPrecompileData* PrecompileData = nullptr;
+
+	if (const FEmitterInfo* EmitterInfo = SystemCompileTask->SystemInfo.EmitterInfoBySourceEmitter(GroupInfo.SourceEmitterIndex))
+	{
+		PrecompileData = static_cast<const FNiagaraPrecompileData*>(SystemCompileTask->SystemPrecompileData->GetDependentRequest(EmitterInfo->DigestedEmitterIndex).Get());
+	}
+	else
+	{
+		PrecompileData = SystemCompileTask->SystemPrecompileData.Get();
+	}
 
 	// COMPILATION
+	TStringBuilder<512> DebugGroupName;
+	NiagaraCompilationTasksImpl::BuildDebugGroupName(*PrecompileData, CompileOptions, DebugGroupName);
+
+	ScriptCompiler = MakeUnique<FHlslNiagaraCompiler>();
+	ScriptCompilationJobId = ScriptCompiler->CompileScriptVM(DebugGroupName, CompileOptions, TranslateResults, TranslateOutput, TranslatedHlsl);
+}
+
+void FNiagaraSystemCompilationTask::FCompileTaskInfo::IssueTranslateGpu(FNiagaraSystemCompilationTask* SystemCompileTask, const FCompileGroupInfo& GroupInfo)
+{
+	const FNiagaraPrecompileData* PrecompileData = nullptr;
+
+	if (const FEmitterInfo* EmitterInfo = SystemCompileTask->SystemInfo.EmitterInfoBySourceEmitter(GroupInfo.SourceEmitterIndex))
+	{
+		PrecompileData = static_cast<const FNiagaraPrecompileData*>(SystemCompileTask->SystemPrecompileData->GetDependentRequest(EmitterInfo->DigestedEmitterIndex).Get());
+	}
+	else
 	{
-		Compiler = MakeShared<FHlslNiagaraCompiler>();
+		PrecompileData = SystemCompileTask->SystemPrecompileData.Get();
+	}
 
-		TStringBuilder<512> DebugGroupName;
-		FPathViews::Append(DebugGroupName, PrecompileData->SourceName);
-		FPathViews::Append(DebugGroupName, PrecompileData->EmitterUniqueName);
-		FPathViews::Append(DebugGroupName, PrecompileData->ENiagaraScriptUsageEnum->GetNameStringByValue((int64)CompileOptions.TargetUsage));
+	// COMPILATION
+	TStringBuilder<512> DebugGroupName;
+	NiagaraCompilationTasksImpl::BuildDebugGroupName(*PrecompileData, CompileOptions, DebugGroupName);
 
-		if (CompileOptions.TargetUsageId.IsValid())
-		{
-			DebugGroupName << TEXT("_");
-			CompileOptions.TargetUsageId.AppendString(DebugGroupName, EGuidFormats::Digits);
-		}
+	ScriptCompiler = MakeUnique<FHlslNiagaraCompiler>();
+	ScriptCompilationJobId = ScriptCompiler->CreateShaderIntermediateData(DebugGroupName, CompileOptions, TranslateResults, TranslateOutput, TranslatedHlsl);
+}
+
+void FNiagaraSystemCompilationTask::FCompileTaskInfo::IssueCompileGpu(FNiagaraSystemCompilationTask* SystemCompileTask, const FCompileGroupInfo& GroupInfo)
+{
+	COOK_STAT(NiagaraSystemCookStats::GpuScriptCompileCount++);
+
+	// we are going to populate a FNiagaraShaderMap.  This involves it's own little DDC hop as we store
+	// the shader in a separate bucket
+
+	const FNiagaraPrecompileData* PrecompileData = nullptr;
+	const FNiagaraCompilationCopyData* CompilationCopyData = nullptr;
+
+	if (const FEmitterInfo* EmitterInfo = SystemCompileTask->SystemInfo.EmitterInfoBySourceEmitter(GroupInfo.SourceEmitterIndex))
+	{
+		PrecompileData = static_cast<const FNiagaraPrecompileData*>(SystemCompileTask->SystemPrecompileData->GetDependentRequest(EmitterInfo->DigestedEmitterIndex).Get());
+		CompilationCopyData = static_cast<const FNiagaraCompilationCopyData*>(GroupInfo.CompilationCopy->GetDependentRequest(EmitterInfo->DigestedEmitterIndex).Get());
+	}
+	else
+	{
+		PrecompileData = SystemCompileTask->SystemPrecompileData.Get();
+		CompilationCopyData = GroupInfo.CompilationCopy.Get();
+	}
+
+	check(!ComputeShaderTaskIndices.IsEmpty());
 
-		CompilationJobId = Compiler->CompileScript(DebugGroupName, CompileOptions, TranslateResults, TranslateOutput, TranslatedHlsl);
+	TStringBuilder<512> DebugGroupName;
+	NiagaraCompilationTasksImpl::BuildDebugGroupName(*PrecompileData, CompileOptions, DebugGroupName);
+
+	TSharedPtr<FNiagaraShaderScriptParametersMetadata> ShaderParameters =
+		NiagaraCompilationTasksImpl::BuildScriptParametersMetadata(CompilationCopyData, TranslateOutput.ScriptData.ShaderScriptParametersMetadata);
+
+	ShaderMapCompiler = MakeUnique<FNiagaraShaderMapCompiler>(SystemCompileTask->NiagaraShaderType, ShaderParameters);
+
+	for (int32 ShaderTaskIndex : ComputeShaderTaskIndices)
+	{
+		const FCompileComputeShaderTaskInfo& ComputeInfo = SystemCompileTask->CompileComputeShaderTasks[ShaderTaskIndex];
+		ShaderMapCompiler->AddShaderPlatform(ComputeInfo.ShaderMapId, ComputeInfo.ShaderPlatform);
 	}
+
+	ShaderMapCompiler->CompileScript(CompileId, PrecompileData->SourceName, DebugGroupName, CompileOptions, TranslateResults, TranslateOutput, TranslatedHlsl);
+}
+
+TOptional<FNiagaraCompileResults> FNiagaraSystemCompilationTask::FCompileTaskInfo::HandleDeprecatedGpuScriptResults() const
+{
+	return TOptional<FNiagaraCompileResults>();
 }
 
 /** Returns true if the task has valid results */
@@ -631,11 +864,11 @@ bool FNiagaraSystemCompilationTask::FCompileTaskInfo::RetrieveCompilationResult(
 
 	if (CompileResultsReadyEvent->IsCompleted())
 	{
-		check(ExeData.IsValid());
 		return true;
 	}
-
 	check(!ExeData.IsValid());
+	check(ScriptCompiler.IsValid());
+	check(ScriptCompilationJobId != INDEX_NONE);
 
 	// in cases where asynchronous shader compiling isn't allowed we'll simply block execution here till the task
 	// is complete
@@ -644,8 +877,7 @@ bool FNiagaraSystemCompilationTask::FCompileTaskInfo::RetrieveCompilationResult(
 		bWait = true;
 	}
 
-	TOptional<FNiagaraCompileResults> CompileResult;
-	CompileResult = Compiler->GetCompileResult(CompilationJobId, bWait);
+	TOptional<FNiagaraCompileResults> CompileResult = ScriptCompiler->GetCompileResult(ScriptCompilationJobId, bWait);
 	if (!CompileResult)
 	{
 		return false;
@@ -653,8 +885,32 @@ bool FNiagaraSystemCompilationTask::FCompileTaskInfo::RetrieveCompilationResult(
 
 	FNiagaraCompileResults& Results = CompileResult.GetValue();
 
-	FString OutGraphLevelErrorMessages;
-	for (const FNiagaraCompileEvent& Message : Results.CompileEvents)
+	ExeData = Results.Data;
+
+	RetrieveTranslateResult();
+
+	ExeData->LastCompileStatus = (FNiagaraCompileResults::CompileResultsToSummary(&Results));
+	if (ExeData->LastCompileStatus != ENiagaraScriptCompileStatus::NCS_Error)
+	{
+		// When there are no errors the compile events get emptied, so add them back here.
+		ExeData->LastCompileEvents.Append(Results.CompileEvents);
+	}
+
+	CompilerWallTime = Results.CompilerWallTime;
+	CompilerPreprocessTime = Results.CompilerPreprocessTime;
+	CompilerWorkerTime = Results.CompilerWorkerTime;
+
+	return true;
+}
+
+void FNiagaraSystemCompilationTask::FCompileTaskInfo::RetrieveTranslateResult()
+{
+	if (!ExeData.IsValid())
+	{
+		ExeData = MakeShared<FNiagaraVMExecutableData>(TranslateOutput.ScriptData);
+	}
+
+	for (const FNiagaraCompileEvent& Message : TranslateResults.CompileEvents)
 	{
 #if defined(NIAGARA_SCRIPT_COMPILE_LOGGING_MEDIUM)
 		UE_LOG(LogNiagaraEditor, Log, TEXT("%s"), *Message.Message);
@@ -662,44 +918,56 @@ bool FNiagaraSystemCompilationTask::FCompileTaskInfo::RetrieveCompilationResult(
 		if (Message.Severity == FNiagaraCompileEventSeverity::Error)
 		{
 			// Write the error messages to the string as well so that they can be echoed up the chain.
-			if (OutGraphLevelErrorMessages.Len() > 0)
+			if (ExeData->ErrorMsg.Len() > 0)
 			{
-				OutGraphLevelErrorMessages += "\n";
+				ExeData->ErrorMsg += "\n";
 			}
-			OutGraphLevelErrorMessages += Message.Message;
+			ExeData->ErrorMsg += Message.Message;
 		}
 	}
-	Results.Data->ErrorMsg = OutGraphLevelErrorMessages;
-
-	Results.Data->LastCompileStatus = (FNiagaraCompileResults::CompileResultsToSummary(&Results));
-	if (Results.Data->LastCompileStatus != ENiagaraScriptCompileStatus::NCS_Error)
-	{
-		// When there are no errors the compile events get emptied, so add them back here.
-		Results.Data->LastCompileEvents.Append(Results.CompileEvents);
-	}
 
-	ExeData = CompileResult->Data;
+	ExeData->LastCompileEvents.Append(TranslateResults.CompileEvents);
+	ExeData->ExternalDependencies = TranslateResults.CompileDependencies;
+	ExeData->CompileTags = TranslateResults.CompileTags;
+	ExeData->CompileTagsEditorOnly = TranslateResults.CompileTagsEditorOnly;
 
 	// we also need to include the information about the rapid iteration parameters that were used
 	// to generate the ExeData
 	ExeData->BakedRapidIterationParameters = BakedRapidIterationParameters;
 
-	CompileResultsReadyEvent->Trigger();
-
-	CompilerWallTime = CompileResult->CompilerWallTime;
-	CompilerPreprocessTime = CompileResult->CompilerPreprocessTime;
-	CompilerWorkerTime = CompileResult->CompilerWorkerTime;
+	const bool bHasTranslationErrors = TranslateResults.NumErrors > 0;
+	const bool bHasTranslationWarnings = TranslateResults.NumWarnings > 0;
 
-	return true;
+	ExeData->LastCompileStatus = bHasTranslationErrors
+		? ENiagaraScriptCompileStatus::NCS_Error
+		: bHasTranslationWarnings
+			? ENiagaraScriptCompileStatus::NCS_UpToDateWithWarnings
+			: ENiagaraScriptCompileStatus::NCS_UpToDate;
 }
 
-void FNiagaraSystemCompilationTask::FCompileTaskInfo::ProcessCompilation(FNiagaraSystemCompilationTask* SystemCompileTask, const FCompileGroupInfo& GroupInfo)
+bool FNiagaraSystemCompilationTask::FCompileTaskInfo::RetrieveShaderMap(bool bWait)
 {
-	ON_SCOPE_EXIT
+	check(ShaderMapCompiler.IsValid());
+	check(CompileResultsReadyEvent.IsValid() && CompileResultsReadyEvent->IsValid());
+
+	if (CompileResultsReadyEvent->IsCompleted())
 	{
-		CompileResultsProcessedEvent->Trigger();
-	};
+		return true;
+	}
+
+	ExeData = ShaderMapCompiler->ReadScriptMetaData();
 
+	if (ShaderMapCompiler->ProcessCompileResults(bWait))
+	{
+		ExeData->LastCompileStatus = ENiagaraScriptCompileStatus::NCS_UpToDate;
+		return true;
+	}
+
+	return false;
+}
+
+void FNiagaraSystemCompilationTask::FCompileTaskInfo::GenerateOptimizedVMByteCode(FNiagaraSystemCompilationTask* SystemCompileTask, const FCompileGroupInfo& GroupInfo)
+{
 #if VECTORVM_SUPPORTS_LEGACY
 	const bool bExperimentalVMDisabled = CompileOptions.AdditionalDefines.Contains(FNiagaraCompileOptions::ExperimentalVMDisabled);
 
@@ -726,7 +994,7 @@ void FNiagaraSystemCompilationTask::FCompileTaskInfo::ProcessCompilation(FNiagar
 		}
 
 		FVectorVMOptimizeContext OptimizeContext = { };
-		uint64 AssetPathHash = CityHash64((char *)AssetPath.GetCharArray().GetData(), AssetPath.GetCharArray().Num());
+		uint64 AssetPathHash = CityHash64((char*)AssetPath.GetCharArray().GetData(), AssetPath.GetCharArray().Num());
 		OptimizeVectorVMScript(ByteCode, ExeData->ByteCode.GetLength(), ExtFnTable.GetData(), ExtFnTable.Num(), &OptimizeContext, AssetPathHash, VVMFlag_OptOmitStats | VVMFlag_OptSaveIntermediateState);
 
 		// extract a human readable version of the script
@@ -737,10 +1005,50 @@ void FNiagaraSystemCompilationTask::FCompileTaskInfo::ProcessCompilation(FNiagar
 
 		FreeVectorVMOptimizeContext(&OptimizeContext);
 
-		ByteCodeOptimizeTime = (float) (OptimizeStartTime - FPlatformTime::Seconds());
+		ByteCodeOptimizeTime = (float)(OptimizeStartTime - FPlatformTime::Seconds());
 	}
 #endif // VECTORVM_SUPPORTS_EXPERIMENTAL
+}
+
+void FNiagaraSystemCompilationTask::FCompileTaskInfo::ProcessCompilation(FNiagaraSystemCompilationTask* SystemCompileTask, const FCompileGroupInfo& GroupInfo)
+{
+	ON_SCOPE_EXIT
+	{
+		CompileResultsProcessedEvent->Trigger();
+	};
+
+	if (ScriptCompileType == EScriptCompileType::CompileForVm)
+	{
+		GenerateOptimizedVMByteCode(SystemCompileTask, GroupInfo);
+	}
+
+	// we need to update our ComputeShaderTasks with the generated shader map so that we can put the data to the DDC
+	if (ScriptCompileType == EScriptCompileType::CompileForGpu)
+	{
+		for (int32 ComputeShaderTaskIndex : ComputeShaderTaskIndices)
+		{
+			FCompileComputeShaderTaskInfo& ShaderTaskInfo = SystemCompileTask->CompileComputeShaderTasks[ComputeShaderTaskIndex];
+
+			ShaderTaskInfo.ShaderMap = ShaderMapCompiler->GetShaderMap(ShaderTaskInfo.ShaderMapId);
+		}
+	}
+}
+
+bool FNiagaraSystemCompilationTask::FDDCTaskInfo::ResolveGet(bool bSuccess)
+{
+	if (DataCacheGetKey != UE::DerivedData::FCacheKey::Empty)
+	{
+		if (bSuccess)
+		{
+			bFromDerivedDataCache = true;
+		}
+		else
+		{
+			DataCachePutKeys.AddUnique(DataCacheGetKey);
+		}
+	}
 
+	return bSuccess;
 }
 
 void FNiagaraSystemCompilationTask::FDispatchAndProcessDataCacheGetRequests::Launch(FNiagaraSystemCompilationTask* SystemCompileTask)
@@ -749,25 +1057,70 @@ void FNiagaraSystemCompilationTask::FDispatchAndProcessDataCacheGetRequests::Lau
 	// and there's no actual change to the compileid...in that case we shouldn't worry about trying to grab
 	// stuff from the ddc again
 
+	union FDDCUserData
+	{
+		uint64 UserData;
+		struct 
+		{
+			int32 CompileTaskIndex;
+			bool bShaderCompileTask;
+		};
+	};
 
 	using namespace UE::DerivedData;
 
 	TArray<FCacheGetValueRequest> GetRequests;
+
 	const int32 CompileTaskCount = SystemCompileTask->CompileTasks.Num();
+	const int32 ShaderCompileTaskCount = SystemCompileTask->CompileComputeShaderTasks.Num();
 
-	GetRequests.Reserve(CompileTaskCount);
+	GetRequests.Reserve(CompileTaskCount + ShaderCompileTaskCount);
 
 	for (int32 CompileTaskIt = 0; CompileTaskIt < CompileTaskCount; ++CompileTaskIt)
 	{
 		FNiagaraSystemCompilationTask::FCompileTaskInfo& CompileTask = SystemCompileTask->CompileTasks[CompileTaskIt];
-		if (CompileTask.IsOutstanding())
+		if (!CompileTask.ExeData.IsValid())
 		{
-			FCacheGetValueRequest& GetValueRequest = GetRequests.AddDefaulted_GetRef();
-			GetValueRequest.Name = CompileTask.AssetPath;
-			GetValueRequest.Key = NiagaraCompilationCopyImpl::BuildNiagaraDDCCacheKey(CompileTask.CompileId, CompileTask.AssetPath);
-			GetValueRequest.UserData = (uint64)CompileTaskIt;
+			const UE::DerivedData::FCacheKey CacheKey = NiagaraCompilationTasksImpl::BuildNiagaraDDCCacheKey(CompileTask.CompileId, CompileTask.AssetPath);
+
+			if (CacheKey != CompileTask.DDCTaskInfo.DataCacheGetKey)
+			{
+				FDDCUserData Index;
+				Index.CompileTaskIndex = CompileTaskIt;
+				Index.bShaderCompileTask = false;
+
+				FCacheGetValueRequest& GetValueRequest = GetRequests.AddDefaulted_GetRef();
+				GetValueRequest.Name = CompileTask.AssetPath;
+				GetValueRequest.Key = CacheKey;
+				GetValueRequest.UserData = Index.UserData;
+
+				CompileTask.DDCTaskInfo.DataCacheGetKey = CacheKey;
+			}
+		}
+	}
+
+	for (int32 ShaderCompileTaskIt = 0; ShaderCompileTaskIt < ShaderCompileTaskCount; ++ShaderCompileTaskIt)
+	{
+		FNiagaraSystemCompilationTask::FCompileComputeShaderTaskInfo& ShaderCompileTask = SystemCompileTask->CompileComputeShaderTasks[ShaderCompileTaskIt];
+		if (ShaderCompileTask.IsOutstanding())
+		{
+			const UE::DerivedData::FCacheKey CacheKey = NiagaraCompilationTasksImpl::BuildNiagaraComputeDDCCacheKey(ShaderCompileTask.ShaderMapId, ShaderCompileTask.ShaderPlatform);
+
+			if (CacheKey != ShaderCompileTask.DDCTaskInfo.DataCacheGetKey)
+			{
+				FDDCUserData Index;
+				Index.CompileTaskIndex = ShaderCompileTaskIt;
+				Index.bShaderCompileTask = true;
+
+				FNiagaraSystemCompilationTask::FCompileTaskInfo& CompileTask = SystemCompileTask->CompileTasks[ShaderCompileTask.ParentCompileTaskIndex];
+
+				FCacheGetValueRequest& ShaderGetValueRequest = GetRequests.AddDefaulted_GetRef();
+				ShaderGetValueRequest.Name = CompileTask.AssetPath;
+				ShaderGetValueRequest.Key = CacheKey;
+				ShaderGetValueRequest.UserData = Index.UserData;
 
-			CompileTask.DataCacheGetKey = GetValueRequest.Key;
+				ShaderCompileTask.DDCTaskInfo.DataCacheGetKey = CacheKey;
+			}
 		}
 	}
 
@@ -778,15 +1131,30 @@ void FNiagaraSystemCompilationTask::FDispatchAndProcessDataCacheGetRequests::Lau
 		FRequestBarrier AsyncBarrier(SystemCompileTask->DDCRequestOwner);
 		GetCache().GetValue(GetRequests, SystemCompileTask->DDCRequestOwner, [this, SystemCompileTask](FCacheGetValueResponse&& Response)
 		{
-			if (ensure(SystemCompileTask->CompileTasks.IsValidIndex(Response.UserData)))
+			FDDCUserData Index;
+			Index.UserData = Response.UserData;
+
+			FDDCTaskInfo* DDCTaskInfo = nullptr;
+			if (Index.bShaderCompileTask)
 			{
-				FNiagaraSystemCompilationTask::FCompileTaskInfo& CompileTask = SystemCompileTask->CompileTasks[Response.UserData];
-				if (Response.Status == EStatus::Ok)
+				if (ensure(SystemCompileTask->CompileComputeShaderTasks.IsValidIndex(Index.CompileTaskIndex)))
 				{
-					CompileTask.PendingDDCData = Response.Value.GetData().Decompress().MoveToUnique();
+					DDCTaskInfo = &SystemCompileTask->CompileComputeShaderTasks[Index.CompileTaskIndex].DDCTaskInfo;
+				}
+			}
+			else
+			{
+				if (ensure(SystemCompileTask->CompileTasks.IsValidIndex(Index.CompileTaskIndex)))
+				{
+					DDCTaskInfo = &SystemCompileTask->CompileTasks[Index.CompileTaskIndex].DDCTaskInfo;
 				}
 			}
 
+			if (DDCTaskInfo && Response.Status == EStatus::Ok)
+			{
+				DDCTaskInfo->PendingDDCData = Response.Value.GetData().Decompress().MoveToUnique();
+			}
+
 			if (PendingGetRequestCount.fetch_sub(1, std::memory_order_relaxed) == 1)
 			{
 				// in order to make sure that processing the binary data from the DDC can properly serialize
@@ -797,23 +1165,41 @@ void FNiagaraSystemCompilationTask::FDispatchAndProcessDataCacheGetRequests::Lau
 					check(IsInGameThread());
 					for (FNiagaraSystemCompilationTask::FCompileTaskInfo& CompileTask : SystemCompileTask->CompileTasks)
 					{
-						if (!CompileTask.PendingDDCData.IsNull())
+						if (!CompileTask.DDCTaskInfo.PendingDDCData.IsNull())
 						{
 							FNiagaraVMExecutableData ExeData;
-							if (NiagaraCompilationCopyImpl::BinaryToExecData(CompileTask.SourceScript.Get(), CompileTask.PendingDDCData.MoveToShared(), ExeData))
+							if (NiagaraCompilationTasksImpl::BinaryToExecData(CompileTask.SourceScript.Get(), CompileTask.DDCTaskInfo.PendingDDCData.MoveToShared(), ExeData))
 							{
 								CompileTask.ExeData = MakeShared<FNiagaraVMExecutableData>(MoveTemp(ExeData));
 							}
+						}
 
-							if (!CompileTask.ExeData.IsValid())
-							{
-								CompileTask.DataCachePutKeys.Add(CompileTask.DataCacheGetKey);
-							}
-							else
-							{
-								CompileTask.bFromDerivedDataCache = true;
-							}
+						if (CompileTask.DDCTaskInfo.ResolveGet(CompileTask.ExeData.IsValid()))
+						{
+							COOK_STAT(NiagaraSystemCookStats::CpuScriptDdcCacheHitCount++);
+						}
+						else
+						{
+							COOK_STAT(NiagaraSystemCookStats::CpuScriptDdcCacheMissCount++);
+						}
+					}
+
+					for (FNiagaraSystemCompilationTask::FCompileComputeShaderTaskInfo& CompileTask : SystemCompileTask->CompileComputeShaderTasks)
+					{
+						if (!CompileTask.DDCTaskInfo.PendingDDCData.IsNull())
+						{
+							CompileTask.ShaderMap = NiagaraCompilationTasksImpl::BinaryDataToShaderMap(CompileTask.DDCTaskInfo.PendingDDCData.MoveToShared());
+						}
+
+						if (CompileTask.DDCTaskInfo.ResolveGet(CompileTask.ShaderMap.IsValid()))
+						{
+							COOK_STAT(NiagaraSystemCookStats::GpuScriptDdcCacheHitCount++);
 						}
+						else
+						{
+							COOK_STAT(NiagaraSystemCookStats::GpuScriptDdcCacheMissCount++);
+						}
+
 					}
 
 					CompletionEvent.Trigger();
@@ -832,16 +1218,20 @@ void FNiagaraSystemCompilationTask::FDispatchDataCachePutRequests::Launch(FNiaga
 	using namespace UE::DerivedData;
 
 	TArray<FCachePutValueRequest> PutRequests;
-	PutRequests.Reserve(SystemCompileTask->CompileTasks.Num());
+
+	const int32 CompileTaskCount = SystemCompileTask->CompileTasks.Num();
+	const int32 ShaderCompileTaskCount = SystemCompileTask->CompileComputeShaderTasks.Num();
+
+	PutRequests.Reserve(CompileTaskCount + ShaderCompileTaskCount);
 
 	for (const FNiagaraSystemCompilationTask::FCompileTaskInfo& CompileTask : SystemCompileTask->CompileTasks)
 	{
-		if (CompileTask.ExeData.IsValid() && !CompileTask.DataCachePutKeys.IsEmpty())
+		if (CompileTask.ExeData.IsValid() && !CompileTask.DDCTaskInfo.DataCachePutKeys.IsEmpty())
 		{
-			if (FSharedBuffer SharedBuffer = NiagaraCompilationCopyImpl::ExecToBinaryData(CompileTask.SourceScript.Get(), *CompileTask.ExeData))
+			if (FSharedBuffer SharedBuffer = NiagaraCompilationTasksImpl::ExecToBinaryData(CompileTask.SourceScript.Get(), *CompileTask.ExeData))
 			{
 				FValue PutValue = FValue::Compress(SharedBuffer);
-				for (const FCacheKey& CachePutKey : CompileTask.DataCachePutKeys)
+				for (const FCacheKey& CachePutKey : CompileTask.DDCTaskInfo.DataCachePutKeys)
 				{
 					FCachePutValueRequest& PutRequest = PutRequests.AddDefaulted_GetRef();
 					PutRequest.Name = CompileTask.AssetPath;
@@ -852,6 +1242,26 @@ void FNiagaraSystemCompilationTask::FDispatchDataCachePutRequests::Launch(FNiaga
 		}
 	}
 
+	for (const FNiagaraSystemCompilationTask::FCompileComputeShaderTaskInfo& CompileTask : SystemCompileTask->CompileComputeShaderTasks)
+	{
+		if (CompileTask.ShaderMap.IsValid() && !CompileTask.DDCTaskInfo.DataCachePutKeys.IsEmpty())
+		{
+			const FNiagaraSystemCompilationTask::FCompileTaskInfo& ParentCompileTask = SystemCompileTask->CompileTasks[CompileTask.ParentCompileTaskIndex];
+
+			if (FSharedBuffer SharedBuffer = NiagaraCompilationTasksImpl::ShaderMapToBinaryData(CompileTask.ShaderMap.GetReference()))
+			{
+				FValue PutValue = FValue::Compress(SharedBuffer);
+				for (const FCacheKey& CachePutKey : CompileTask.DDCTaskInfo.DataCachePutKeys)
+				{
+					FCachePutValueRequest& PutRequest = PutRequests.AddDefaulted_GetRef();
+					PutRequest.Name = ParentCompileTask.AssetPath;
+					PutRequest.Key = CachePutKey;
+					PutRequest.Value = PutValue;
+				}
+			}
+		}
+	}
+
 	PendingPutRequestCount = PutRequests.Num();
 
 	if (!PutRequests.IsEmpty())
@@ -873,7 +1283,7 @@ void FNiagaraSystemCompilationTask::FDispatchDataCachePutRequests::Launch(FNiaga
 
 void FNiagaraSystemCompilationTask::DigestSystemInfo()
 {
-	using namespace NiagaraCompilationCopyImpl;
+	using namespace NiagaraCompilationTasksImpl;
 
 	FNiagaraDigestDatabase& DigestDatabase = FNiagaraDigestDatabase::Get();
 	const UNiagaraScript* SystemSpawnScript = System_GT->GetSystemSpawnScript();
@@ -920,31 +1330,34 @@ void FNiagaraSystemCompilationTask::DigestSystemInfo()
 		FEmitterInfo& EmitterInfo = SystemInfo.EmitterInfo.AddDefaulted_GetRef();
 		const FVersionedNiagaraEmitter& HandleInstance = Handle.GetInstance();
 
-		const UNiagaraGraph* EmitterGraph = GetGraphFromScriptSource(Handle.GetEmitterData()->GraphSource);
-		ChangeIdBuilder.ParseReferencedGraphs(EmitterGraph);
+		if (FVersionedNiagaraEmitterData* EmitterData = Handle.GetEmitterData())
+		{
+			const UNiagaraGraph* EmitterGraph = GetGraphFromScriptSource(EmitterData->GraphSource);
+			ChangeIdBuilder.ParseReferencedGraphs(EmitterGraph);
 
-		EmitterInfo.UniqueEmitterName = HandleInstance.Emitter->GetUniqueEmitterName();
-		EmitterInfo.UniqueInstanceName = Handle.GetUniqueInstanceName();
-		EmitterInfo.Enabled = Handle.GetIsEnabled();
-		EmitterInfo.ConstantResolver = FNiagaraFixedConstantResolver(FCompileConstantResolver(HandleInstance, ENiagaraScriptUsage::EmitterSpawnScript));
-		EmitterInfo.SourceEmitterIndex = SourceEmitterIndex;
-		EmitterInfo.DigestedEmitterIndex = DigestedEmitterIndex;
-		EmitterInfo.SourceGraph = DigestDatabase.CreateGraphDigest(EmitterGraph, ChangeIdBuilder);
+			if (const UNiagaraEmitter* Emitter = HandleInstance.Emitter)
+			{
+				EmitterInfo.UniqueEmitterName = HandleInstance.Emitter->GetUniqueEmitterName();
+				EmitterInfo.UniqueInstanceName = Handle.GetUniqueInstanceName();
+				EmitterInfo.Enabled = Handle.GetIsEnabled();
+				EmitterInfo.ConstantResolver = FNiagaraFixedConstantResolver(FCompileConstantResolver(HandleInstance, ENiagaraScriptUsage::EmitterSpawnScript));
+				EmitterInfo.SourceEmitterIndex = SourceEmitterIndex;
+				EmitterInfo.DigestedEmitterIndex = DigestedEmitterIndex;
+				EmitterInfo.SourceGraph = DigestDatabase.CreateGraphDigest(EmitterGraph, ChangeIdBuilder);
+			}
 
-		{
-			constexpr bool bCompilableOnly = false;
-			constexpr bool bEnabledOnly = false;
-			TArray<UNiagaraScript*> EmitterScripts;
-			Handle.GetEmitterData()->GetScripts(EmitterScripts, bCompilableOnly, bEnabledOnly);
-			EmitterInfo.OwnedScriptKeys.Reserve(EmitterScripts.Num());
-			for (const UNiagaraScript* EmitterScript : EmitterScripts)
 			{
-				EmitterInfo.OwnedScriptKeys.Add(EmitterScript);
+				constexpr bool bCompilableOnly = false;
+				constexpr bool bEnabledOnly = false;
+				TArray<UNiagaraScript*> EmitterScripts;
+				EmitterData->GetScripts(EmitterScripts, bCompilableOnly, bEnabledOnly);
+				EmitterInfo.OwnedScriptKeys.Reserve(EmitterScripts.Num());
+				for (const UNiagaraScript* EmitterScript : EmitterScripts)
+				{
+					EmitterInfo.OwnedScriptKeys.Add(EmitterScript);
+				}
 			}
-		}
 
-		if (FVersionedNiagaraEmitterData* EmitterData = Handle.GetEmitterData())
-		{
 			const TArray<UNiagaraSimulationStageBase*> SimStages = EmitterData->GetSimulationStages();
 			EmitterInfo.SimStages.Reserve(SimStages.Num());
 			for (const UNiagaraSimulationStageBase* SourceSimStage : SimStages)
@@ -996,7 +1409,13 @@ void FNiagaraSystemCompilationTask::DigestParameterCollections(TConstArrayView<T
 	}
 }
 
-void FNiagaraSystemCompilationTask::AddScript(int32 EmitterIndex, UNiagaraScript* Script, bool bRequiresCompilation)
+void FNiagaraSystemCompilationTask::DigestShaderInfo(const ITargetPlatform* InTargetPlatform, FNiagaraShaderType* InShaderType)
+{
+	TargetPlatform = InTargetPlatform;
+	NiagaraShaderType = InShaderType;
+}
+
+void FNiagaraSystemCompilationTask::AddScript(int32 EmitterIndex, UNiagaraScript* Script, const FNiagaraVMExecutableDataId& CompileId, bool bRequiresCompilation, TConstArrayView<FShaderCompileRequest> ShaderRequests)
 {
 	if (!Script)
 	{
@@ -1005,8 +1424,6 @@ void FNiagaraSystemCompilationTask::AddScript(int32 EmitterIndex, UNiagaraScript
 
 	if (bRequiresCompilation)
 	{
-		FNiagaraVMExecutableDataId CompileId;
-		Script->ComputeVMCompilationId(CompileId, FGuid());
 		FVersionedNiagaraScriptData* ScriptData = Script->GetScriptData(CompileId.ScriptVersionID);
 		if (ensure(ScriptData))
 		{
@@ -1017,7 +1434,7 @@ void FNiagaraSystemCompilationTask::AddScript(int32 EmitterIndex, UNiagaraScript
 			}
 
 			TArray<ENiagaraScriptUsage> DuplicateUsages;
-			NiagaraCompilationCopyImpl::GetUsagesToDuplicate(Script->GetUsage(), DuplicateUsages);
+			NiagaraCompilationTasksImpl::GetUsagesToDuplicate(Script->GetUsage(), DuplicateUsages);
 			for (ENiagaraScriptUsage DuplicateUsage : DuplicateUsages)
 			{
 				GroupInfo->ValidUsages.AddUnique(DuplicateUsage);
@@ -1040,7 +1457,41 @@ void FNiagaraSystemCompilationTask::AddScript(int32 EmitterIndex, UNiagaraScript
 			TaskInfo.SourceScript = Script;
 			TaskInfo.ScriptKey = Script;
 			TaskInfo.AssetPath = Script->GetPathName();
-			TaskInfo.CompileId = MoveTemp(CompileId);
+			TaskInfo.CompileId = CompileId;
+
+			const bool bIsGpuScript = TaskInfo.CompileOptions.IsGpuScript() && UNiagaraScript::IsGPUScript(TaskInfo.CompileOptions.TargetUsage);
+
+			TaskInfo.ScriptCompileType = bIsGpuScript
+				? UNiagaraScript::AreGpuScriptsCompiledBySystem()
+					? EScriptCompileType::CompileForGpu
+					: EScriptCompileType::TranslateForGpu
+				: TaskInfo.CompileOptions.IsGpuScript()
+					? EScriptCompileType::DummyCompileForCpuStubs
+					: EScriptCompileType::CompileForVm;
+
+			if (TaskInfo.ScriptCompileType == EScriptCompileType::CompileForGpu)
+			{
+				check(SystemInfo.EmitterInfo.IsValidIndex(EmitterIndex));
+				const FEmitterInfo& EmitterInfo = SystemInfo.EmitterInfo[EmitterIndex];
+
+				if (EmitterInfo.Enabled)
+				{
+					TaskInfo.ComputeShaderTaskIndices.Reserve(ShaderRequests.Num());
+					for (const FShaderCompileRequest& ShaderRequest : ShaderRequests)
+					{
+						if (Script->ShouldCompile(ShaderRequest.ShaderPlatform))
+						{
+							TaskInfo.ComputeShaderTaskIndices.Add(CompileComputeShaderTasks.Num());
+							FCompileComputeShaderTaskInfo& ShaderTaskInfo = CompileComputeShaderTasks.AddDefaulted_GetRef();
+
+							ShaderTaskInfo.ParentCompileTaskIndex = CompileTaskIndex;
+							ShaderTaskInfo.ShaderMapId = ShaderRequest.ShaderMapId;
+							ShaderTaskInfo.ShaderPlatform = ShaderRequest.ShaderPlatform;
+						}
+					}
+				}
+			}
+
 			TaskInfo.TaskStartTime = FPlatformTime::Seconds();
 		}
 	}
@@ -1083,7 +1534,7 @@ void FNiagaraSystemCompilationTask::AddScript(int32 EmitterIndex, UNiagaraScript
 	Info.Usage = Script->GetUsage();
 	Info.UsageId = Script->GetUsageId();
 	Info.SourceEmitterIndex = EmitterIndex;
-	for (UNiagaraScript* DependentScript : NiagaraCompilationCopyImpl::FindDependentScripts(System_GT.Get(), EmitterData, Script))
+	for (UNiagaraScript* DependentScript : NiagaraCompilationTasksImpl::FindDependentScripts(System_GT.Get(), EmitterData, Script))
 	{
 		Info.DependentScripts.AddUnique(DependentScript);
 	}
@@ -1104,7 +1555,29 @@ void FNiagaraSystemCompilationTask::Tick()
 
 				if (HasPendingCompilationTask)
 				{
-					CompileTask.RetrieveCompilationResult(false);
+					bool bRetrieveCompleted = false;
+
+					switch (CompileTask.ScriptCompileType)
+					{
+						case EScriptCompileType::CompileForVm:
+						case EScriptCompileType::TranslateForGpu:
+							bRetrieveCompleted = CompileTask.RetrieveCompilationResult(false);
+						break;
+
+						case EScriptCompileType::DummyCompileForCpuStubs:
+							CompileTask.RetrieveTranslateResult();
+							bRetrieveCompleted = true;
+						break;
+
+						case EScriptCompileType::CompileForGpu:
+							bRetrieveCompleted = CompileTask.RetrieveShaderMap(false);
+						break;
+					}
+
+					if (bRetrieveCompleted)
+					{
+						CompileTask.CompileResultsReadyEvent->Trigger();
+					}
 				}
 			}
 		} break;
@@ -1139,9 +1612,21 @@ bool FNiagaraSystemCompilationTask::Poll(FNiagaraSystemAsyncCompileResults& Resu
 				{
 					FNiagaraScriptAsyncCompileData& CompileData = Results.CompileResultMap.Add(SourceScript);
 
-					CompileData.bFromDerivedDataCache = TaskInfo.bFromDerivedDataCache;
+					CompileData.bFromDerivedDataCache = TaskInfo.DDCTaskInfo.bFromDerivedDataCache;
 					CompileData.CompileId = TaskInfo.CompileId;
 					CompileData.ExeData = TaskInfo.ExeData;
+
+					for (int32 ShaderTaskIndex : TaskInfo.ComputeShaderTaskIndices)
+					{
+						const FCompileComputeShaderTaskInfo& ShaderTaskInfo = CompileComputeShaderTasks[ShaderTaskIndex];
+
+						FNiagaraCompiledShaderInfo& ResultsInfo = CompileData.CompiledShaders.AddDefaulted_GetRef();
+						ResultsInfo.TargetPlatform = TargetPlatform;
+						ResultsInfo.ShaderPlatform = ShaderTaskInfo.ShaderPlatform;
+						ResultsInfo.FeatureLevel = ShaderTaskInfo.ShaderMapId.FeatureLevel;
+						ResultsInfo.CompiledShader = ShaderTaskInfo.ShaderMap;
+					}
+
 					CompileData.NamedDataInterfaces.Reserve(TaskInfo.NamedDataInterfaces.Num());
 					Algo::Transform(TaskInfo.NamedDataInterfaces, CompileData.NamedDataInterfaces, [](const TTuple<FName, UNiagaraDataInterface*>& InElement)
 					{
@@ -1272,7 +1757,7 @@ void FNiagaraSystemCompilationTask::Precompile()
 	// Add the User and System variables that we did encounter to the list that emitters might also encounter.
 	TArray<FNiagaraVariable> SystemEncounteredUserVariables;
 	SystemPrecompileData->GatherPreCompiledVariables(FNiagaraConstants::UserNamespaceString, SystemEncounteredUserVariables);
-	NiagaraCompilationCopyImpl::AppendUnique(SystemEncounteredUserVariables, EncounterableVars);
+	NiagaraCompilationTasksImpl::AppendUnique(SystemEncounteredUserVariables, EncounterableVars);
 	SystemExposedVariables.Append(SystemEncounteredUserVariables);
 
 	SystemPrecompileData->GatherPreCompiledVariables(FNiagaraConstants::SystemNamespaceString, EncounterableVars);
@@ -1681,7 +2166,7 @@ struct FNiagaraSystemCompilationTask::FBuildRapidIterationTaskBuilder
 
 				if (Builder.Histories.Num() == 1)
 				{
-					BuilderTask->ModuleInputVariables = NiagaraCompilationCopyImpl::ExtractInputVariablesFromHistory(
+					BuilderTask->ModuleInputVariables = NiagaraCompilationTasksImpl::ExtractInputVariablesFromHistory(
 						Builder.Histories[0], BuilderTask->FunctionCallNode->CalledGraph.Get());
 
 					BuilderTask->FunctionCallNode->MultiFindParameterMapDefaultValues(BuilderTask->ScriptUsage, BuilderTask->ConstantResolver, BuilderTask->ModuleInputVariables);
@@ -1724,14 +2209,14 @@ struct FNiagaraSystemCompilationTask::FBuildRapidIterationTaskBuilder
 
 						ValidParameterNames.Add(RapidIterationParameter.GetName());
 
-						NiagaraCompilationCopyImpl::FixupRenamedParameter(
+						NiagaraCompilationTasksImpl::FixupRenamedParameter(
 							ModuleInputVariable,
 							RapidIterationParameter,
 							InputTask->FunctionCallNode,
 							GuidMapping,
 							MergeTaskOutput);
 
-						NiagaraCompilationCopyImpl::AddRapidIterationParameter(
+						NiagaraCompilationTasksImpl::AddRapidIterationParameter(
 							AliasedFunctionInputHandle,
 							ModuleInputVariable,
 							RapidIterationParameter,
@@ -1767,6 +2252,8 @@ UE::Tasks::FTask FNiagaraSystemCompilationTask::BuildRapidIterationParametersAsy
 {
 	using namespace UE::Tasks;
 
+	COOK_STAT(auto Timer = NiagaraSystemCookStats::UsageStats.TimeSyncWork(); Timer.TrackCyclesOnly(););
+
 	TMap<int32 /*EmitterIndex*/, FTask> CollectStaticVariableTasks;
 	
 	FCollectStaticVariablesTaskBuilder SystemTaskBuilder(*this);
@@ -1826,7 +2313,10 @@ UE::Tasks::FTask FNiagaraSystemCompilationTask::BuildRapidIterationParametersAsy
 void FNiagaraSystemCompilationTask::BuildAndApplyRapidIterationParameters()
 {
 	check(IsInGameThread());
-	BuildRapidIterationParametersAsync().Wait();
+	{
+		COOK_STAT(auto Timer = NiagaraSystemCookStats::UsageStats.TimeAsyncWait(); Timer.TrackCyclesOnly(););
+		BuildRapidIterationParametersAsync().Wait();
+	}
 
 	// see FNiagaraActiveCompilationAsyncTask::Apply() for why we don't allow removal
 	constexpr bool bAllowParameterRemoval = false;
@@ -1851,6 +2341,19 @@ void FNiagaraSystemCompilationTask::BuildAndApplyRapidIterationParameters()
 	}
 }
 
+void FNiagaraSystemCompilationTask::IssuePostResultsProcessedTasks()
+{
+	CompileCompletionEvent.Trigger();
+	CompilationState = EState::ResultsProcessed;
+	PutRequestHelper = MakeUnique<FDispatchDataCachePutRequests>();
+	PutRequestHelper->Launch(this);
+
+	Launch(UE_SOURCE_LOCATION, [this]
+	{
+		CompilationState = EState::Completed;
+	}, PutRequestHelper->CompletionEvent);
+}
+
 void FNiagaraSystemCompilationTask::IssueCompilationTasks()
 {
 	using namespace UE::Tasks;
@@ -1862,6 +2365,8 @@ void FNiagaraSystemCompilationTask::IssueCompilationTasks()
 	{
 		FTask PrecompileTask = Launch(UE_SOURCE_LOCATION, [this]
 		{
+			COOK_STAT(auto Timer = NiagaraSystemCookStats::UsageStats.TimeSyncWork(); Timer.TrackCyclesOnly(););
+
 			Precompile();
 		});
 
@@ -1872,25 +2377,57 @@ void FNiagaraSystemCompilationTask::IssueCompilationTasks()
 			{
 				FTask InstantiateGraphTask = Launch(UE_SOURCE_LOCATION, [this, &GroupInfo]
 				{
+					COOK_STAT(auto Timer = NiagaraSystemCookStats::UsageStats.TimeSyncWork(); Timer.TrackCyclesOnly(););
+
 					GroupInfo.InstantiateCompileGraph(*this);
 				}, PrecompileTask);
 
 				for (int32 CompileTaskIndex : GroupInfo.CompileTaskIndices)
 				{
 					FCompileTaskInfo& CompileTask = CompileTasks[CompileTaskIndex];
-					if (CompileTask.IsOutstanding())
+					if (CompileTask.IsOutstanding(*this))
 					{
 						CompileTask.CompileResultsReadyEvent = MakeUnique<FTaskEvent>(UE_SOURCE_LOCATION);
 						CompileTask.CompileResultsProcessedEvent = MakeUnique<FTaskEvent>(UE_SOURCE_LOCATION);
 
 						IssueCompilationTasks.Add(Launch(UE_SOURCE_LOCATION, [this, &GroupInfo, &CompileTask]
 						{
+							COOK_STAT(auto Timer = NiagaraSystemCookStats::UsageStats.TimeSyncWork(); Timer.TrackCyclesOnly(););
+
 							CompileTask.CollectNamedDataInterfaces(this, GroupInfo);
-							CompileTask.TranslateAndIssueCompile(this, GroupInfo);
+							switch (CompileTask.ScriptCompileType)
+							{
+								case EScriptCompileType::CompileForVm:
+									CompileTask.Translate(this, GroupInfo);
+									CompileTask.IssueCompileVm(this, GroupInfo);
+								break;
+
+								case EScriptCompileType::CompileForGpu:
+									CompileTask.Translate(this, GroupInfo);
+									CompileTask.IssueCompileGpu(this, GroupInfo);
+								break;
+
+								// path for handling the translation of GPU scripts but we're not handling the compilation of
+								// the compute shaders ourselves, rather it's using the old path of the FNiagaraShaderScript getting
+								// cached by the UNiagaraScript post CPU translation/compilation
+								case EScriptCompileType::TranslateForGpu:
+									CompileTask.Translate(this, GroupInfo);
+									CompileTask.IssueTranslateGpu(this, GroupInfo);
+								break;
+
+								// path for dealing with the CPU scripts (Particle spawn/update) when the emitter is GPU
+								// These scripts shouldn't need to be processed, but currently they are required for supplying
+								// RI parameters to the GPU execution context
+								case EScriptCompileType::DummyCompileForCpuStubs:
+									CompileTask.Translate(this, GroupInfo);
+								break;
+							}
 						}, InstantiateGraphTask));
 
 						Launch(UE_SOURCE_LOCATION, [this, &GroupInfo, &CompileTask]
 						{
+							COOK_STAT(auto Timer = NiagaraSystemCookStats::UsageStats.TimeSyncWork(); Timer.TrackCyclesOnly(););
+
 							CompileTask.ProcessCompilation(this, GroupInfo);
 						}, *CompileTask.CompileResultsReadyEvent);
 
@@ -1908,15 +2445,7 @@ void FNiagaraSystemCompilationTask::IssueCompilationTasks()
 
 	Launch(UE_SOURCE_LOCATION, [this]
 	{
-		CompileCompletionEvent.Trigger();
-		CompilationState = EState::ResultsProcessed;
-		PutRequestHelper = MakeUnique<FDispatchDataCachePutRequests>();
-		PutRequestHelper->Launch(this);
-
-		Launch(UE_SOURCE_LOCATION, [this]
-		{
-			CompilationState = EState::Completed;
-		}, PutRequestHelper->CompletionEvent);
+		IssuePostResultsProcessedTasks();
 	}, CompilationTasks);
 }
 
@@ -1924,7 +2453,7 @@ bool FNiagaraSystemCompilationTask::HasOutstandingCompileTasks() const
 {
 	for (const FCompileTaskInfo& CompileTask : CompileTasks)
 	{
-		if (!CompileTask.ExeData.IsValid())
+		if (CompileTask.IsOutstanding(*this))
 		{
 			return true;
 		}
@@ -1937,6 +2466,8 @@ UE::Tasks::FTask FNiagaraSystemCompilationTask::BeginTasks()
 {
 	using namespace UE::Tasks;
 
+	COOK_STAT(auto Timer = NiagaraSystemCookStats::UsageStats.TimeSyncWork(); Timer.TrackCyclesOnly(););
+
 	InitialGetRequestHelper.Launch(this);
 
 	TArray<FTask> SystemTaskPrerequisites;
@@ -1966,6 +2497,8 @@ UE::Tasks::FTask FNiagaraSystemCompilationTask::BeginTasks()
 
 	FTask SystemTask = Launch(UE_SOURCE_LOCATION, [this]
 	{
+		COOK_STAT(auto Timer = NiagaraSystemCookStats::UsageStats.TimeSyncWork(); Timer.TrackCyclesOnly(););
+
 		if (HasOutstandingCompileTasks())
 		{
 			TArray<FTask> CompilationTaskPrerequisites;
@@ -1990,8 +2523,7 @@ UE::Tasks::FTask FNiagaraSystemCompilationTask::BeginTasks()
 		}
 		else
 		{
-			CompilationState = EState::Completed;
-			CompileCompletionEvent.Trigger();
+			IssuePostResultsProcessedTasks();
 		}
 	}, SystemTaskPrerequisites);
 
diff --git a/Engine/Plugins/FX/Niagara/Source/NiagaraEditor/Private/NiagaraCompilationTasks.h b/Engine/Plugins/FX/Niagara/Source/NiagaraEditor/Private/NiagaraCompilationTasks.h
--- a/Engine/Plugins/FX/Niagara/Source/NiagaraEditor/Private/NiagaraCompilationTasks.h
+++ b/Engine/Plugins/FX/Niagara/Source/NiagaraEditor/Private/NiagaraCompilationTasks.h
@@ -10,6 +10,7 @@
 
 struct FNiagaraSystemCompilationTask
 {
+	// helper for script DDC get requests
 	struct FDispatchAndProcessDataCacheGetRequests
 	{
 		void Launch(FNiagaraSystemCompilationTask* SystemCompileTask);
@@ -18,6 +19,7 @@ struct FNiagaraSystemCompilationTask
 		std::atomic<int32> PendingGetRequestCount = 0;
 	};
 
+	// helepr for script DDC put requests
 	struct FDispatchDataCachePutRequests
 	{
 		void Launch(FNiagaraSystemCompilationTask* SystemCompileTask);
@@ -26,24 +28,33 @@ struct FNiagaraSystemCompilationTask
 		std::atomic<int32> PendingPutRequestCount = 0;
 	};
 
+	struct FShaderCompileRequest
+	{
+		FNiagaraShaderMapId ShaderMapId;
+		EShaderPlatform ShaderPlatform;
+	};
+
 	FNiagaraSystemCompilationTask(FNiagaraCompilationTaskHandle InTaskHandle, UNiagaraSystem* InSystem);
 
 	void Abort();
 
 	void DigestSystemInfo();
 	void DigestParameterCollections(TConstArrayView<TWeakObjectPtr<UNiagaraParameterCollection>> Collections);
-	void AddScript(int32 EmitterIndex, UNiagaraScript* ScriptToCompile, bool bRequiresCompilation);
+	void DigestShaderInfo(const ITargetPlatform* TargetPlatform, FNiagaraShaderType* ShaderType);
+	void AddScript(int32 EmitterIndex, UNiagaraScript* ScriptToCompile, const FNiagaraVMExecutableDataId& CompileId, bool bRequiresCompilation, TConstArrayView<FShaderCompileRequest> ShaderRequests);
 	UE::Tasks::FTask BeginTasks();
 
 	UE::Tasks::FTask BuildRapidIterationParametersAsync();
 	void BuildAndApplyRapidIterationParameters();
 	void IssueCompilationTasks();
+	void IssuePostResultsProcessedTasks();
 	bool HasOutstandingCompileTasks() const;
 
 	double PrepareStartTime = 0.0;
 	double QueueStartTime = 0.0;
 	double LaunchStartTime = 0.0;
 
+	// compilation was forced
 	bool bForced = false;
 
 	enum class EState : uint8
@@ -119,39 +130,86 @@ struct FNiagaraSystemCompilationTask
 		int32 SourceEmitterIndex = INDEX_NONE;
 	};
 
+	struct FDDCTaskInfo
+	{
+		bool ResolveGet(bool bSuccess);
+
+		UE::DerivedData::FCacheKey DataCacheGetKey;
+		TArray< UE::DerivedData::FCacheKey> DataCachePutKeys;
+		FUniqueBuffer PendingDDCData;
+		bool bFromDerivedDataCache = false;
+	};
+
+	struct FCompileComputeShaderTaskInfo
+	{
+		int32 ParentCompileTaskIndex = INDEX_NONE;
+
+		FNiagaraShaderMapId ShaderMapId;
+		EShaderPlatform ShaderPlatform;
+		FNiagaraShaderMapRef ShaderMap;
+
+		FDDCTaskInfo DDCTaskInfo;
+
+		bool IsOutstanding() const;
+	};
+
+	enum class EScriptCompileType
+	{
+		Invalid = 0,
+		CompileForVm,
+		TranslateForGpu,
+		CompileForGpu,
+		DummyCompileForCpuStubs,
+	};
+
 	struct FCompileTaskInfo
 	{
 		FNiagaraCompileOptions CompileOptions;
 		FNiagaraTranslateResults TranslateResults;
 		FNiagaraTranslatorOutput TranslateOutput;
 		FString TranslatedHlsl;
-		TSharedPtr<FHlslNiagaraCompiler> Compiler;
-		FUniqueBuffer PendingDDCData;
 		TSharedPtr<FNiagaraVMExecutableData> ExeData;
+
+		// complier for VM script compiles
+		TUniquePtr<FHlslNiagaraCompiler> ScriptCompiler;
+
+		// information for GPU script compiles
+		TUniquePtr<FNiagaraShaderMapCompiler> ShaderMapCompiler;
+		TArray<int32> ComputeShaderTaskIndices;
+
 		TWeakObjectPtr<UNiagaraScript> SourceScript;
 		TObjectKey<UNiagaraScript> ScriptKey;
 		FString AssetPath;
 		FNiagaraVMExecutableDataId CompileId;
-		int32 CompilationJobId = INDEX_NONE;
-		UE::DerivedData::FCacheKey DataCacheGetKey;
-		TArray<UE::DerivedData::FCacheKey> DataCachePutKeys;
+		uint32 ScriptCompilationJobId;
+
+		FDDCTaskInfo DDCTaskInfo;
+
 		TUniquePtr<UE::Tasks::FTaskEvent> CompileResultsReadyEvent;
 		TUniquePtr<UE::Tasks::FTaskEvent> CompileResultsProcessedEvent;
 		TMap<FName, UNiagaraDataInterface*> NamedDataInterfaces;
 		TArray<FNiagaraVariable> BakedRapidIterationParameters; // todo - remove and have the data pulled directly from the precompile data
-		bool bFromDerivedDataCache = false;
 		double TaskStartTime = 0;
 		float TranslationTime = 0.0f;
 		float CompilerWallTime = 0.0f;
 		float CompilerWorkerTime = 0.0f;
 		float CompilerPreprocessTime = 0.0f;
 		float ByteCodeOptimizeTime = 0.0f;
+		EScriptCompileType ScriptCompileType = EScriptCompileType::Invalid;
 
-		bool IsOutstanding() const;
+		bool IsOutstanding(const FNiagaraSystemCompilationTask& ParentTask) const;
 		void CollectNamedDataInterfaces(FNiagaraSystemCompilationTask* SystemCompileTask, const FCompileGroupInfo& GroupInfo);
-		void TranslateAndIssueCompile(FNiagaraSystemCompilationTask* SystemCompileTask, const FCompileGroupInfo& GroupInfo);
+		void Translate(FNiagaraSystemCompilationTask* SystemCompileTask, const FCompileGroupInfo& GroupInfo);
+		void IssueCompileVm(FNiagaraSystemCompilationTask* SystemCompileTask, const FCompileGroupInfo& GroupInfo);
+		void IssueTranslateGpu(FNiagaraSystemCompilationTask* SystemCompileTask, const FCompileGroupInfo& GroupInfo);
+		void IssueCompileGpu(FNiagaraSystemCompilationTask* SystemCompileTask, const FCompileGroupInfo& GroupInfo);
 		void ProcessCompilation(FNiagaraSystemCompilationTask* SystemCompileTask, const FCompileGroupInfo& GroupInfo);
+		void GenerateOptimizedVMByteCode(FNiagaraSystemCompilationTask* SystemCompileTask, const FCompileGroupInfo& GroupInfo);
+		void RetrieveTranslateResult();
 		bool RetrieveCompilationResult(bool bWait);
+		bool RetrieveShaderMap(bool bWait);
+
+		TOptional<FNiagaraCompileResults> HandleDeprecatedGpuScriptResults() const;
 	};
 
 	struct FCollectStaticVariablesTaskBuilder;
@@ -165,9 +223,13 @@ struct FNiagaraSystemCompilationTask
 
 	TArray<FCompileGroupInfo> CompileGroups;
 	TArray<FCompileTaskInfo> CompileTasks;
+	TArray<FCompileComputeShaderTaskInfo> CompileComputeShaderTasks;
 
 	const FNiagaraCompilationTaskHandle TaskHandle;
 
+	const ITargetPlatform* TargetPlatform = nullptr;
+	FNiagaraShaderType* NiagaraShaderType = nullptr;
+
 	int32 TasksAwaitingDDCGetResults = 0;
 
 	void Tick();
diff --git a/Engine/Plugins/FX/Niagara/Source/NiagaraEditor/Private/NiagaraCompiler.cpp b/Engine/Plugins/FX/Niagara/Source/NiagaraEditor/Private/NiagaraCompiler.cpp
--- a/Engine/Plugins/FX/Niagara/Source/NiagaraEditor/Private/NiagaraCompiler.cpp
+++ b/Engine/Plugins/FX/Niagara/Source/NiagaraEditor/Private/NiagaraCompiler.cpp
@@ -1,6 +1,8 @@
 // Copyright Epic Games, Inc. All Rights Reserved.
 
 #include "NiagaraCompiler.h"
+
+#include "DataDrivenShaderPlatformInfo.h"
 #include "EdGraphSchema_Niagara.h"
 #include "EdGraphUtilities.h"
 #include "INiagaraEditorTypeUtilities.h"
@@ -1912,6 +1914,169 @@ int32 FHlslNiagaraCompiler::CompileScript(const FStringView GroupName, const FNi
 	return JobID;
 }
 
+uint32 FHlslNiagaraCompiler::CompileScriptVM(const FStringView GroupName, const FNiagaraCompileOptions& InOptions, const FNiagaraTranslateResults& InTranslateResults, const FNiagaraTranslatorOutput& TranslatorOutput, const FString& TranslatedHLSL)
+{
+	check(!InOptions.IsGpuScript() || !UNiagaraScript::IsParticleScript(InOptions.TargetUsage));
+
+	CompileResults.Data = MakeShared<FNiagaraVMExecutableData>();
+
+	CompileResults.Data->LastHlslTranslation = TEXT("");
+
+	FShaderCompilerInput Input;
+	Input.Target = FShaderTarget(SF_Compute, SP_PCD3D_SM5);
+	Input.VirtualSourceFilePath = TEXT("/Plugin/FX/Niagara/Private/NiagaraEmitterInstanceShader.usf");
+	Input.EntryPointName = TEXT("SimulateMain");
+	Input.Environment.SetDefine(TEXT("VM_SIMULATION"), 1);
+	Input.Environment.SetDefine(TEXT("COMPUTESHADER"), 1);
+	Input.Environment.SetDefine(TEXT("PIXELSHADER"), 0);
+	Input.Environment.SetDefine(TEXT("DOMAINSHADER"), 0);
+	Input.Environment.SetDefine(TEXT("HULLSHADER"), 0);
+	Input.Environment.SetDefine(TEXT("VERTEXSHADER"), 0);
+	Input.Environment.SetDefine(TEXT("GEOMETRYSHADER"), 0);
+	Input.Environment.SetDefine(TEXT("MESHSHADER"), 0);
+	Input.Environment.SetDefine(TEXT("AMPLIFICATIONSHADER"), 0);
+	Input.Environment.IncludeVirtualPathToContentsMap.Add(TEXT("/Engine/Generated/NiagaraEmitterInstance.ush"), TranslatedHLSL);
+	Input.DebugInfoFlags = GShaderCompilingManager->GetDumpShaderDebugInfoFlags();
+	Input.DumpDebugInfoRootPath = GShaderCompilingManager->GetAbsoluteShaderDebugInfoDirectory() / TEXT("VM");
+	Input.DebugGroupName = GroupName;
+	Input.DebugExtension.Empty();
+	Input.DumpDebugInfoPath.Empty();
+
+	CompileResults.DumpDebugInfoPath = Input.DumpDebugInfoPath;
+
+	uint32 JobID = FShaderCommonCompileJob::GetNextJobId();
+	CompilationJob = MakeUnique<FNiagaraCompilerJob>();
+	CompilationJob->TranslatorOutput = TranslatorOutput;
+
+	CompileResults.bVMSucceeded = (CompilationJob->TranslatorOutput.Errors.Len() == 0) && (TranslatedHLSL.Len() > 0) && !InTranslateResults.NumErrors;
+
+	CompileResults.AppendCompileEvents(MakeArrayView(InTranslateResults.CompileEvents));
+	CompileResults.Data->LastCompileEvents.Append(InTranslateResults.CompileEvents);
+	CompileResults.Data->ExternalDependencies = InTranslateResults.CompileDependencies;
+	CompileResults.Data->CompileTags = InTranslateResults.CompileTags;
+	CompileResults.Data->CompileTagsEditorOnly = InTranslateResults.CompileTagsEditorOnly;
+	CompilationJob->TranslatorOutput.ScriptData.LastHlslTranslation = TranslatedHLSL;
+	CompilationJob->TranslatorOutput.ScriptData.ExternalDependencies = InTranslateResults.CompileDependencies;
+	CompilationJob->TranslatorOutput.ScriptData.CompileTags = InTranslateResults.CompileTags;
+	CompilationJob->TranslatorOutput.ScriptData.CompileTagsEditorOnly = InTranslateResults.CompileTagsEditorOnly;
+
+	bool bJobScheduled = false;
+	if (CompileResults.bVMSucceeded)
+	{
+		SCOPE_CYCLE_COUNTER(STAT_NiagaraEditor_HlslCompiler_CompileShader_VectorVM);
+		CompilationJob->StartTime = FPlatformTime::Seconds();
+
+		FShaderType* NiagaraShaderType = nullptr;
+		for (TLinkedList<FShaderType*>::TIterator ShaderTypeIt(FShaderType::GetTypeList()); ShaderTypeIt; ShaderTypeIt.Next())
+		{
+			if (FNiagaraShaderType* ShaderType = ShaderTypeIt->GetNiagaraShaderType())
+			{
+				NiagaraShaderType = ShaderType;
+				break;
+			}
+		}
+		if (NiagaraShaderType)
+		{
+			TRefCountPtr<FShaderCompileJob> Job = GShaderCompilingManager->PrepareShaderCompileJob(JobID, FShaderCompileJobKey(NiagaraShaderType), EShaderCompileJobPriority::Normal);
+			if (Job)
+			{
+				TArray<FShaderCommonCompileJobPtr> NewJobs;
+				CompilationJob->ShaderCompileJob = Job;
+				Input.ShaderFormat = FName(TEXT("VVM_1_0"));
+				if (GNiagaraSkipVectorVMBackendOptimizations != 0)
+				{
+					Input.Environment.CompilerFlags.Add(CFLAG_SkipOptimizations);
+				}
+				Job->Input = Input;
+				NewJobs.Add(FShaderCommonCompileJobPtr(Job));
+
+				GShaderCompilingManager->SubmitJobs(NewJobs, FString(), FString());
+			}
+			bJobScheduled = true;
+		}
+	}
+	CompileResults.Data->LastHlslTranslation = TranslatedHLSL;
+
+	if (!bJobScheduled)
+	{
+		CompileResults.Data->ByteCode.Reset();
+		CompileResults.Data->Attributes.Empty();
+		CompileResults.Data->Parameters.Empty();
+		CompileResults.Data->InternalParameters.Empty();
+		CompileResults.Data->DataInterfaceInfo.Empty();
+		CompileResults.Data->UObjectInfos.Empty();
+
+	}
+	CompilationJob->CompileResults = CompileResults;
+
+	return JobID;
+}
+
+int32 FHlslNiagaraCompiler::CreateShaderIntermediateData(const FStringView GroupName, const FNiagaraCompileOptions& InOptions, const FNiagaraTranslateResults& InTranslateResults, const FNiagaraTranslatorOutput& TranslatorOutput, const FString& TranslatedHLSL)
+{
+	check(InOptions.IsGpuScript() && UNiagaraScript::IsParticleScript(InOptions.TargetUsage));
+
+	CompileResults.Data = MakeShared<FNiagaraVMExecutableData>();
+
+	//TODO: This should probably be done via the same route that other shaders take through the shader compiler etc.
+	//But that adds the complexity of a new shader type, new shader class and a new shader map to contain them etc.
+	//Can do things simply for now.
+
+	CompileResults.Data->LastHlslTranslation = TEXT("");
+
+	FShaderCompilerInput Input;
+	Input.Target = FShaderTarget(SF_Compute, SP_PCD3D_SM5);
+	Input.VirtualSourceFilePath = TEXT("/Plugin/FX/Niagara/Private/NiagaraEmitterInstanceShader.usf");
+	Input.EntryPointName = TEXT("SimulateMain");
+	Input.Environment.SetDefine(TEXT("VM_SIMULATION"), 1);
+	Input.Environment.SetDefine(TEXT("COMPUTESHADER"), 1);
+	Input.Environment.SetDefine(TEXT("PIXELSHADER"), 0);
+	Input.Environment.SetDefine(TEXT("DOMAINSHADER"), 0);
+	Input.Environment.SetDefine(TEXT("HULLSHADER"), 0);
+	Input.Environment.SetDefine(TEXT("VERTEXSHADER"), 0);
+	Input.Environment.SetDefine(TEXT("GEOMETRYSHADER"), 0);
+	Input.Environment.SetDefine(TEXT("MESHSHADER"), 0);
+	Input.Environment.SetDefine(TEXT("AMPLIFICATIONSHADER"), 0);
+	Input.Environment.IncludeVirtualPathToContentsMap.Add(TEXT("/Engine/Generated/NiagaraEmitterInstance.ush"), TranslatedHLSL);
+	Input.DebugInfoFlags = GShaderCompilingManager->GetDumpShaderDebugInfoFlags();
+	Input.DumpDebugInfoRootPath = GShaderCompilingManager->GetAbsoluteShaderDebugInfoDirectory() / TEXT("VM");
+	Input.DebugGroupName = GroupName;
+	Input.DebugExtension.Empty();
+	Input.DumpDebugInfoPath.Empty();
+
+	CompileResults.DumpDebugInfoPath = Input.DumpDebugInfoPath;
+
+	uint32 JobID = FShaderCommonCompileJob::GetNextJobId();
+	CompilationJob = MakeUnique<FNiagaraCompilerJob>();
+	CompilationJob->TranslatorOutput = TranslatorOutput;
+
+	CompileResults.bVMSucceeded = (CompilationJob->TranslatorOutput.Errors.Len() == 0) && (TranslatedHLSL.Len() > 0) && !InTranslateResults.NumErrors;
+
+	// only issue jobs for VM compilation if we're going to be using the resulting byte code.  This excludes particle scripts when we're using
+	// a GPU simulation
+	CompileResults.bComputeSucceeded = false;
+	if (CompileResults.bVMSucceeded)
+	{
+		//Clear out current contents of compile results.
+		*(CompileResults.Data) = CompilationJob->TranslatorOutput.ScriptData;
+		CompileResults.Data->ByteCode.Reset();
+		CompileResults.bComputeSucceeded = true;
+	}
+
+	CompileResults.AppendCompileEvents(MakeArrayView(InTranslateResults.CompileEvents));
+	CompileResults.Data->LastCompileEvents.Append(InTranslateResults.CompileEvents);
+	CompileResults.Data->ExternalDependencies = InTranslateResults.CompileDependencies;
+	CompileResults.Data->CompileTags = InTranslateResults.CompileTags;
+	CompileResults.Data->CompileTagsEditorOnly = InTranslateResults.CompileTagsEditorOnly;
+
+	// Early out if compiling a GPU particle script as we do not need to submit a CPU compile request.
+	// This must be done after we add in the translator errors etc so tha they are passed to the compile job correctly.
+	CompileResults.Data->LastHlslTranslationGPU = TranslatedHLSL;
+	DumpDebugInfo(CompileResults, Input, true);
+	CompilationJob->CompileResults = CompileResults;
+	return JobID;
+}
+
 void FHlslNiagaraCompiler::FixupVMAssembly(FString& Asm)
 {
 	const TCHAR* OpTag = TEXT("__OP__");
@@ -2258,6 +2423,125 @@ void FHlslNiagaraCompiler::Warning(FText WarningText)
 	CompileResults.NumWarnings++;
 }
 
+FNiagaraShaderMapCompiler::FNiagaraShaderMapCompiler(
+	const FNiagaraShaderType* InShaderType,
+	TSharedPtr<FNiagaraShaderScriptParametersMetadata> InShaderParameters)
+	: ShaderType(InShaderType)
+	, ShaderParameters(InShaderParameters)
+{
+}
+
+void FNiagaraShaderMapCompiler::AddShaderPlatform(const FNiagaraShaderMapId& ShaderMapId, EShaderPlatform ShaderPlatform)
+{
+	FActiveCompilation& ActiveCompilation = ActiveCompilations.AddDefaulted_GetRef();
+	ActiveCompilation.ShaderMapId = ShaderMapId;
+	ActiveCompilation.ShaderPlatform = ShaderPlatform;
+	ActiveCompilation.ShaderMap = new FNiagaraShaderMap(FNiagaraShaderMap::WorkerThread);
+}
+
+void FNiagaraShaderMapCompiler::CompileScript(
+	const FNiagaraVMExecutableDataId& ScriptCompileId,
+	const FStringView SourceName,
+	const FStringView DebugGroupName,
+	const FNiagaraCompileOptions& CompileOptions,
+	const FNiagaraTranslateResults& TranslateResults,
+	const FNiagaraTranslatorOutput& TranslatorOutput,
+	const FString& TranslatedHLSL)
+{
+	TArray<TRefCountPtr<FShaderCommonCompileJob>> CompileJobs;
+
+	for (FActiveCompilation& ActiveCompilation : ActiveCompilations)
+	{
+		TRefCountPtr<FSharedShaderCompilerEnvironment> CompilationEnvironment = new FSharedShaderCompilerEnvironment();
+		CompilationEnvironment->SetDefine(TEXT("GPU_SIMULATION_SHADER"), TEXT("1"));
+		CompilationEnvironment->SetDefine(TEXT("NIAGARA_COMPRESSED_ATTRIBUTES_ENABLED"),
+			ScriptCompileId.AdditionalDefines.Contains(TEXT("CompressAttributes")) ? 1 : 0);
+
+		// Fast math breaks The ExecGrid layout script because floor(x/y) returns a bad value if x == y. Yay.
+		if (IsMetalPlatform(ActiveCompilation.ShaderPlatform))
+		{
+			CompilationEnvironment->CompilerFlags.Add(CFLAG_NoFastMath);
+		}
+
+		ActiveCompilation.ShaderMap->CreateCompileJobs(
+			ShaderType,
+			DebugGroupName,
+			ActiveCompilation.ShaderMapId,
+			TranslatedHLSL,
+			CompilationEnvironment,
+			TranslatorOutput.ScriptData.SimulationStageMetaData,
+			TranslatorOutput.ScriptData.SimulationStageMetaData.Num(),
+			ActiveCompilation.ShaderPlatform,
+			ShaderParameters,
+			ActiveCompilation.ShaderCompileJobs
+		);
+
+		CompileJobs.Append(ActiveCompilation.ShaderCompileJobs);
+	}
+
+	// we also need to populate the ExeData for the script based on the translator results.  This handles all the meta data of the script.
+	ScriptExeData = MakeShared<FNiagaraVMExecutableData>(TranslatorOutput.ScriptData);
+	ScriptExeData->LastCompileEvents.Append(TranslateResults.CompileEvents);
+	ScriptExeData->ExternalDependencies = TranslateResults.CompileDependencies;
+	ScriptExeData->CompileTags = TranslateResults.CompileTags;
+	ScriptExeData->CompileTagsEditorOnly = TranslateResults.CompileTagsEditorOnly;
+	ScriptExeData->LastHlslTranslationGPU = TranslatedHLSL;
+
+	GShaderCompilingManager->SubmitJobs(CompileJobs, FString(SourceName));
+}
+
+bool FNiagaraShaderMapCompiler::ProcessCompileResults(bool bWait)
+{
+	check(!bWait); // not currently implemented
+	check(IsInGameThread());
+
+	for (TArray<FActiveCompilation>::TIterator CompileIt = ActiveCompilations.CreateIterator(); CompileIt; ++CompileIt)
+	{
+		// make sure that all of the shader compile jobs have been released and finalized
+		const bool bReadyToProcess = !CompileIt->ShaderCompileJobs.ContainsByPredicate([](const FShaderCommonCompileJobPtr& CompileJob) -> bool
+		{
+			return !CompileJob->bReleased || !CompileJob->bFinalized;
+		});
+
+		if (!bReadyToProcess)
+		{
+			// todo - it might be worth keeping track of jobs that aren't getting handled because of the above
+			// condition.  Either it's taking a long time and so it could be worth reporting, or because the
+			// job is lost and we'll never complete.
+			continue;
+		}
+
+		FActiveCompilation& CurrentCompilation = *CompileIt;
+
+		// for now we'll process all shaders at once (need to measure the cost here)
+		for (FShaderCommonCompileJobPtr& ShaderCompileJob : CurrentCompilation.ShaderCompileJobs)
+		{
+			CurrentCompilation.ShaderMap->ProcessAndFinalizeShaderCompileJob(ShaderCompileJob);
+		}
+
+		// now that we've added all the results into the shader map we can move it over to CompletedCompilations
+		CompletedCompilations.Add(CurrentCompilation.ShaderMap);
+
+		// and remove it from the ActiveCompilations
+		CompileIt.RemoveCurrentSwap();
+	}
+
+	return ActiveCompilations.IsEmpty();
+}
+
+FNiagaraShaderMapRef FNiagaraShaderMapCompiler::GetShaderMap(const FNiagaraShaderMapId& ShaderMapId) const
+{
+	for (const FNiagaraShaderMapRef& ShaderMapRef : CompletedCompilations)
+	{
+		if (ShaderMapRef->GetShaderMapId() == ShaderMapId)
+		{
+			return ShaderMapRef;
+		}
+	}
+
+	return nullptr;
+}
+
 //////////////////////////////////////////////////////////////////////////
 
 #undef LOCTEXT_NAMESPACE
diff --git a/Engine/Plugins/FX/Niagara/Source/NiagaraEditor/Private/NiagaraCompiler.h b/Engine/Plugins/FX/Niagara/Source/NiagaraEditor/Private/NiagaraCompiler.h
--- a/Engine/Plugins/FX/Niagara/Source/NiagaraEditor/Private/NiagaraCompiler.h
+++ b/Engine/Plugins/FX/Niagara/Source/NiagaraEditor/Private/NiagaraCompiler.h
@@ -45,7 +45,9 @@ public:
 	NIAGARAEDITOR_API virtual int32 CompileScript(const class FNiagaraCompileRequestData* InCompileRequest, const FNiagaraCompileOptions& InOptions, const FNiagaraTranslateResults& InTranslateResults, FNiagaraTranslatorOutput* TranslatorOutput, FString& TranslatedHLSL) override;
 
 	NIAGARAEDITOR_API virtual int32 CompileScript(const FStringView GroupName, const FNiagaraCompileOptions& InOptions, const FNiagaraTranslateResults& InTranslateResults, const FNiagaraTranslatorOutput& TranslatorOutput, const FString& TranslatedHLSL) override;
-	NIAGARAEDITOR_API virtual TOptional<FNiagaraCompileResults> GetCompileResult(int32 JobID, bool bWait = false) override;
+	NIAGARAEDITOR_API virtual uint32 CompileScriptVM(const FStringView GroupName, const FNiagaraCompileOptions& InOptions, const FNiagaraTranslateResults& InTranslateResults, const FNiagaraTranslatorOutput& TranslatorOutput, const FString& TranslatedHLSL);
+	NIAGARAEDITOR_API virtual int32 CreateShaderIntermediateData(const FStringView GroupName, const FNiagaraCompileOptions& InOptions, const FNiagaraTranslateResults& InTranslateResults, const FNiagaraTranslatorOutput& TranslatorOutput, const FString& TranslatedHLSL);
+	virtual TOptional<FNiagaraCompileResults> GetCompileResult(int32 JobID, bool bWait = false) override;
 
 	NIAGARAEDITOR_API virtual void Error(FText ErrorText) override;
 	NIAGARAEDITOR_API virtual void Warning(FText WarningText) override;
@@ -58,3 +60,58 @@ private:
 	/** SCW doesn't have access to the VM op code names so we do a fixup pass to make these human readable after we get the data back from SCW. */
 	NIAGARAEDITOR_API void FixupVMAssembly(FString& Asm);
 };
+
+class FNiagaraShaderMapCompiler
+{
+public:
+	FNiagaraShaderMapCompiler(
+		const FNiagaraShaderType* InShaderType,
+		TSharedPtr<FNiagaraShaderScriptParametersMetadata> InShaderParameters);
+
+	FNiagaraShaderMapCompiler() = delete;
+
+	void CompileScript(
+		const FNiagaraVMExecutableDataId& ScriptCompileId,
+		FStringView SourceName,
+		FStringView GroupName,
+		const FNiagaraCompileOptions& InOptions,
+		const FNiagaraTranslateResults& InTranslateResults,
+		const FNiagaraTranslatorOutput& TranslatorOutput,
+		const FString& TranslatedHLSL);
+	void AddShaderPlatform(const FNiagaraShaderMapId& ShaderMapId, EShaderPlatform ShaderPlatform);
+
+	// returns true if there are no more results requiring processing
+	bool ProcessCompileResults(bool bWait = false);
+
+	TConstArrayView<FNiagaraShaderMapRef> ReadCompletedCompilations() const
+	{
+		return CompletedCompilations;
+	}
+
+	FNiagaraShaderMapRef GetShaderMap(const FNiagaraShaderMapId& ShaderMapId) const;
+
+	TSharedPtr<FNiagaraVMExecutableData> ReadScriptMetaData() const
+	{
+		return ScriptExeData;
+	}
+
+	TSharedPtr<FNiagaraShaderScriptParametersMetadata> GetShaderParameters() const
+	{
+		return ShaderParameters;
+	}
+
+private:
+	struct FActiveCompilation
+	{
+		FNiagaraShaderMapId ShaderMapId;
+		EShaderPlatform ShaderPlatform;
+		FNiagaraShaderMapRef ShaderMap;
+		TArray<FShaderCommonCompileJobPtr> ShaderCompileJobs;
+	};
+
+	const FNiagaraShaderType* ShaderType = nullptr;
+	TSharedPtr<FNiagaraShaderScriptParametersMetadata> ShaderParameters;
+	TArray<FActiveCompilation> ActiveCompilations;
+	TArray<FNiagaraShaderMapRef> CompletedCompilations;
+	TSharedPtr<FNiagaraVMExecutableData> ScriptExeData;
+};
diff --git a/Engine/Plugins/FX/Niagara/Source/NiagaraEditor/Private/NiagaraEditorModule.cpp b/Engine/Plugins/FX/Niagara/Source/NiagaraEditor/Private/NiagaraEditorModule.cpp
--- a/Engine/Plugins/FX/Niagara/Source/NiagaraEditor/Private/NiagaraEditorModule.cpp
+++ b/Engine/Plugins/FX/Niagara/Source/NiagaraEditor/Private/NiagaraEditorModule.cpp
@@ -1355,9 +1355,9 @@ void FNiagaraEditorModule::StartupModule()
 		return CacheGraphTraversal(InObj, Version);
 	}));
 
-	RequestCompileSystemHandle = NiagaraModule.RegisterRequestCompileSystem(INiagaraModule::FOnRequestCompileSystem::CreateLambda([this](UNiagaraSystem* System, bool bForced)
+	RequestCompileSystemHandle = NiagaraModule.RegisterRequestCompileSystem(INiagaraModule::FOnRequestCompileSystem::CreateLambda([this](UNiagaraSystem* System, bool bForced, const ITargetPlatform* TargetPlatform)
 	{
-		return RequestCompileSystem(System, bForced);
+		return RequestCompileSystem(System, bForced, TargetPlatform);
 	}));
 
 	PollSystemCompileHandle = NiagaraModule.RegisterPollSystemCompile(INiagaraModule::FOnPollSystemCompile::CreateLambda([this](FNiagaraCompilationTaskHandle TaskHandle, FNiagaraSystemAsyncCompileResults& Results, bool bWait, bool bPeek)
diff --git a/Engine/Plugins/FX/Niagara/Source/NiagaraEditor/Private/NiagaraSystemCompilingManager.cpp b/Engine/Plugins/FX/Niagara/Source/NiagaraEditor/Private/NiagaraSystemCompilingManager.cpp
--- a/Engine/Plugins/FX/Niagara/Source/NiagaraEditor/Private/NiagaraSystemCompilingManager.cpp
+++ b/Engine/Plugins/FX/Niagara/Source/NiagaraEditor/Private/NiagaraSystemCompilingManager.cpp
@@ -2,8 +2,11 @@
 
 #include "NiagaraSystemCompilingManager.h"
 
+#include "DataDrivenShaderPlatformInfo.h"
 #include "NiagaraCompilationTasks.h"
 #include "NiagaraEditorModule.h"
+#include "NiagaraShaderType.h"
+#include "ProfilingDebugging/CookStats.h"
 #include "UObject/UObjectIterator.h"
 
 #include "Misc/ScopeRWLock.h"
@@ -18,6 +21,80 @@ static FAutoConsoleVariableRef CVarNiagaraCompilationMaxActiveTaskCount(
 	ECVF_Default
 );
 
+#if ENABLE_COOK_STATS
+namespace NiagaraSystemCookStats
+{
+	FCookStats::FDDCResourceUsageStats UsageStats;
+	static FCookStatsManager::FAutoRegisterCallback RegisterCookStats([](FCookStatsManager::AddStatFuncRef AddStat)
+	{
+		UsageStats.LogStats(AddStat, TEXT("NiagaraSystem.Usage"), TEXT(""));
+	});
+}
+#endif
+
+namespace NiagaraSystemCompilingManagerImpl
+{
+
+FNiagaraShaderType* GetNiagaraShaderType()
+{
+	FNiagaraShaderType* FoundShaderType = nullptr;
+	for (TLinkedList<FShaderType*>::TIterator ShaderTypeIt(FShaderType::GetTypeList()); ShaderTypeIt; ShaderTypeIt.Next())
+	{
+		if (FNiagaraShaderType* ShaderType = ShaderTypeIt->GetNiagaraShaderType())
+		{
+			if (ensure(FoundShaderType == nullptr))
+			{
+				FoundShaderType = ShaderType;
+			}
+		}
+	}
+
+	return FoundShaderType;
+}
+
+FNiagaraShaderMapId BuildShaderMapId(FNiagaraShaderType* ShaderType, const ITargetPlatform* TargetPlatform, EShaderPlatform ShaderPlatform, ERHIFeatureLevel::Type FeatureLevel, const FNiagaraVMExecutableDataId& BaseScriptId)
+{
+	FNiagaraShaderMapId ShaderMapId;
+	ShaderMapId.FeatureLevel = FeatureLevel;
+	ShaderMapId.bUsesRapidIterationParams = BaseScriptId.bUsesRapidIterationParams;
+	BaseScriptId.BaseScriptCompileHash.ToSHAHash(ShaderMapId.BaseCompileHash);
+	ShaderMapId.CompilerVersionID = BaseScriptId.CompilerVersionID;
+
+	ShaderMapId.ReferencedCompileHashes.Reserve(BaseScriptId.ReferencedCompileHashes.Num());
+	for (const FNiagaraCompileHash& Hash : BaseScriptId.ReferencedCompileHashes)
+	{
+		Hash.ToSHAHash(ShaderMapId.ReferencedCompileHashes.AddDefaulted_GetRef());
+	}
+
+	ShaderMapId.AdditionalDefines.Empty(BaseScriptId.AdditionalDefines.Num());
+	for (const FString& Define : BaseScriptId.AdditionalDefines)
+	{
+		ShaderMapId.AdditionalDefines.Emplace(Define);
+	}
+
+	const TArray<FString> AdditionalVariableStrings = BaseScriptId.GetAdditionalVariableStrings();
+	ShaderMapId.AdditionalVariables.Empty(AdditionalVariableStrings.Num());
+	for (const FString& Variable : AdditionalVariableStrings)
+	{
+		ShaderMapId.AdditionalVariables.Emplace(Variable);
+	}
+
+	ShaderMapId.ShaderTypeDependencies.Emplace(ShaderType, ShaderPlatform);
+
+	if (TargetPlatform)
+	{
+		ShaderMapId.LayoutParams.InitializeForPlatform(TargetPlatform->IniPlatformName(), TargetPlatform->HasEditorOnlyData());
+	}
+	else
+	{
+		ShaderMapId.LayoutParams.InitializeForCurrent();
+	}
+
+	return ShaderMapId;
+}
+
+};
+
 FNiagaraSystemCompilingManager& FNiagaraSystemCompilingManager::Get()
 {
 	static FNiagaraSystemCompilingManager Singleton;
@@ -86,90 +163,94 @@ void FNiagaraSystemCompilingManager::Shutdown()
 
 void FNiagaraSystemCompilingManager::ProcessAsyncTasks(bool bLimitExecutionTime)
 {
-	// process any pending GameThreadTasks
 	{
-		TArray<FGameThreadFunction> PendingFunctions;
+		COOK_STAT(auto Timer = NiagaraSystemCookStats::UsageStats.TimeSyncWork(); Timer.TrackCyclesOnly(););
 
+		// process any pending GameThreadTasks
 		{
-			FWriteScopeLock Write(GameThreadFunctionLock);
-			PendingFunctions = MoveTemp(GameThreadFunctions);
-		}
+			TArray<FGameThreadFunction> PendingFunctions;
 
-		for (FGameThreadFunction& PendingFunction : PendingFunctions)
-		{
-			PendingFunction();
+			{
+				FWriteScopeLock Write(GameThreadFunctionLock);
+				PendingFunctions = MoveTemp(GameThreadFunctions);
+			}
+
+			for (FGameThreadFunction& PendingFunction : PendingFunctions)
+			{
+				PendingFunction();
+			}
 		}
-	}
 
-	{
-		FReadScopeLock ReadScope(QueueLock);
-		for (FNiagaraCompilationTaskHandle TaskHandle : ActiveTasks)
 		{
-			FTaskPtr TaskPtr = SystemRequestMap.FindRef(TaskHandle);
-			if (TaskPtr.IsValid())
+			FReadScopeLock ReadScope(QueueLock);
+			for (FNiagaraCompilationTaskHandle TaskHandle : ActiveTasks)
 			{
-				TaskPtr->Tick();
+				FTaskPtr TaskPtr = SystemRequestMap.FindRef(TaskHandle);
+				if (TaskPtr.IsValid())
+				{
+					TaskPtr->Tick();
+				}
 			}
 		}
-	}
 
-	{
-		FWriteScopeLock WriteScope(QueueLock);
-
-		// find the list of tasks that we can remove
-		TArray<FNiagaraCompilationTaskHandle> TasksToRemove;
-		TArray<FNiagaraCompilationTaskHandle> TasksToRetrieve;
-		for (TArray<FNiagaraCompilationTaskHandle>::TIterator TaskIt(ActiveTasks); TaskIt; ++TaskIt)
 		{
-			FTaskPtr TaskPtr = SystemRequestMap.FindRef(*TaskIt);
+			FWriteScopeLock WriteScope(QueueLock);
 
-			bool bRemoveCurrent = true;
-			if (TaskPtr.IsValid())
+			// find the list of tasks that we can remove
+			TArray<FNiagaraCompilationTaskHandle> TasksToRemove;
+			TArray<FNiagaraCompilationTaskHandle> TasksToRetrieve;
+			for (TArray<FNiagaraCompilationTaskHandle>::TIterator TaskIt(ActiveTasks); TaskIt; ++TaskIt)
 			{
-				if (TaskPtr->AreResultsPending())
+				FTaskPtr TaskPtr = SystemRequestMap.FindRef(*TaskIt);
+
+				bool bRemoveCurrent = true;
+				if (TaskPtr.IsValid())
 				{
-					TasksToRetrieve.Add(*TaskIt);
+					if (TaskPtr->AreResultsPending())
+					{
+						TasksToRetrieve.Add(*TaskIt);
+					}
+					else if (TaskPtr->CanRemove())
+					{
+						TasksToRemove.Add(*TaskIt);
+					}
+					else
+					{
+						bRemoveCurrent = false;
+					}
 				}
-				else if (TaskPtr->CanRemove())
+
+				if (bRemoveCurrent)
 				{
-					TasksToRemove.Add(*TaskIt);
+					TaskIt.RemoveCurrent();
 				}
-				else
+			}
+
+			// go through the entries that are awaiting retrieval and clean up any that have been retrieved
+			for (TArray<FNiagaraCompilationTaskHandle>::TIterator TaskIt(RequestsAwaitingRetrieval); TaskIt; ++TaskIt)
+			{
+				FTaskPtr TaskPtr = SystemRequestMap.FindRef(*TaskIt);
+				if (!TaskPtr.IsValid() || TaskPtr->CanRemove())
 				{
-					bRemoveCurrent = false;
+					TasksToRemove.Add(*TaskIt);
+					TaskIt.RemoveCurrent();
 				}
 			}
 
-			if (bRemoveCurrent)
+			// remove tasks that can be erased
+			for (FNiagaraCompilationTaskHandle TaskToRemove : TasksToRemove)
 			{
-				TaskIt.RemoveCurrent();
+				ensure(!QueuedRequests.Contains(TaskToRemove));
+				SystemRequestMap.Remove(TaskToRemove);
 			}
-		}
 
-		// go through the entries that are awaiting retrieval and clean up any that have been retrieved
-		for (TArray<FNiagaraCompilationTaskHandle>::TIterator TaskIt(RequestsAwaitingRetrieval); TaskIt; ++TaskIt)
-		{
-			FTaskPtr TaskPtr = SystemRequestMap.FindRef(*TaskIt);
-			if (!TaskPtr.IsValid() || TaskPtr->CanRemove())
+			// finally populate RequestsAwaitingRetrieval with any new entries
+			for (FNiagaraCompilationTaskHandle TaskToRetrieve : TasksToRetrieve)
 			{
-				TasksToRemove.Add(*TaskIt);
-				TaskIt.RemoveCurrent();
+				ensure(!QueuedRequests.Contains(TaskToRetrieve));
+				RequestsAwaitingRetrieval.Add(TaskToRetrieve);
 			}
 		}
-
-		// remove tasks that can be erased
-		for (FNiagaraCompilationTaskHandle TaskToRemove : TasksToRemove)
-		{
-			ensure(!QueuedRequests.Contains(TaskToRemove));
-			SystemRequestMap.Remove(TaskToRemove);
-		}
-
-		// finally populate RequestsAwaitingRetrieval with any new entries
-		for (FNiagaraCompilationTaskHandle TaskToRetrieve : TasksToRetrieve)
-		{
-			ensure(!QueuedRequests.Contains(TaskToRetrieve));
-			RequestsAwaitingRetrieval.Add(TaskToRetrieve);
-		}
 	}
 
 	// queue up as many tasks as we can
@@ -185,19 +266,71 @@ FNiagaraCompilationTaskHandle FNiagaraSystemCompilingManager::AddSystem(UNiagara
 
 	struct FCompilableScriptInfo
 	{
-		FCompilableScriptInfo(bool& bHasCompilation, UNiagaraScript* InScript, bool bForced, int32 InEmitterIndex = INDEX_NONE)
+		FCompilableScriptInfo(UNiagaraScript* InScript, bool bForced, int32 InEmitterIndex, bool InGpuEmitter, bool& bHasCompilation)
 			: Script(InScript)
 			, EmitterIndex(InEmitterIndex)
 		{
-			bRequiresCompilation = Script && Script->IsCompilable() && (bForced || !Script->AreScriptAndSourceSynchronized());
+			const bool bIsValidScriptTarget = Script
+				&& Script->IsCompilable();
+			// when we successfully get rid of the CPU side scripts (particle spawn/update) for GPU emitters we can reinstate this check
+			//	&& InGpuEmitter == Script->IsGPUScript();
+
+			if (!bIsValidScriptTarget)
+			{
+				return;
+			}
+
+			Script->ComputeVMCompilationId(CompileId, FGuid());
+			bRequiresCompilation = !CompileId.IsValid() || CompileId != Script->GetVMExecutableDataCompilationId();
+
+			bHasCompilation = bHasCompilation || bRequiresCompilation;
+		}
+
+		void UpdateComputeShaders(bool InGpuEmitter, FNiagaraShaderType* ShaderType, const ITargetPlatform* TargetPlatform, TConstArrayView<FPlatformFeatureLevelPair> FeatureLevels, bool& bHasCompilation)
+		{
+			// check if the GPU shaders need compilation
+			if (UNiagaraScript::AreGpuScriptsCompiledBySystem() && InGpuEmitter && Script->IsGPUScript())
+			{
+				ShaderRequests.Reserve(FeatureLevels.Num());
+
+				for (const FPlatformFeatureLevelPair& PlatformFeatureLevel : FeatureLevels)
+				{
+					const bool bScriptIsMissingOrDirty = true;
+
+					const FNiagaraShaderMapId ShaderMapId = NiagaraSystemCompilingManagerImpl::BuildShaderMapId(
+						ShaderType,
+						TargetPlatform,
+						PlatformFeatureLevel.Key,
+						PlatformFeatureLevel.Value,
+						CompileId);
+
+					if (bRequiresCompilation || !Script->IsShaderMapCached(TargetPlatform, ShaderMapId))
+					{
+						FNiagaraSystemCompilationTask::FShaderCompileRequest& Request = ShaderRequests.AddDefaulted_GetRef();
+						Request.ShaderMapId = ShaderMapId;
+						Request.ShaderPlatform = PlatformFeatureLevel.Key;
+					}
+				}
+
+				if (!ShaderRequests.IsEmpty())
+				{
+					bRequiresCompilation = true;
+				}
+			}
+
 			bHasCompilation = bHasCompilation || bRequiresCompilation;
 		}
 
+		FNiagaraVMExecutableDataId CompileId;
+		TArray<FNiagaraSystemCompilationTask::FShaderCompileRequest> ShaderRequests;
 		UNiagaraScript* Script;
 		int32 EmitterIndex = INDEX_NONE;
 		bool bRequiresCompilation = false;
 	};
 
+	TArray<FPlatformFeatureLevelPair> FeatureLevels;
+	FindOrAddFeatureLevels(CompileOptions, FeatureLevels);
+
 	TArray<UNiagaraScript*> AllScripts;
 	TArray<FCompilableScriptInfo> ScriptsToCompile;
 	bool bHasScriptToCompile = false;
@@ -206,7 +339,7 @@ FNiagaraCompilationTaskHandle FNiagaraSystemCompilingManager::AddSystem(UNiagara
 	for (UNiagaraScript* SystemScript : { System->GetSystemSpawnScript(), System->GetSystemUpdateScript() })
 	{
 		AllScripts.Add(SystemScript);
-		ScriptsToCompile.Emplace(bHasScriptToCompile, SystemScript, CompileOptions.bForced);
+		ScriptsToCompile.Emplace(SystemScript, CompileOptions.bForced, INDEX_NONE /*EmitterIndex*/, false /*InGpuEmitter*/, bHasScriptToCompile);
 	}
 
 	const TArray<FNiagaraEmitterHandle>& EmitterHandles = System->GetEmitterHandles();
@@ -219,6 +352,17 @@ FNiagaraCompilationTaskHandle FNiagaraSystemCompilingManager::AddSystem(UNiagara
 			continue;
 		}
 
+		if (CompileOptions.TargetPlatform)
+		{
+			if (const UNiagaraEmitter* Emitter = Handle.GetInstance().Emitter)
+			{
+				if (!Emitter->NeedsLoadForTargetPlatform(CompileOptions.TargetPlatform))
+				{
+					continue;
+				}
+			}
+		}
+
 		if (const FVersionedNiagaraEmitterData* EmitterData = Handle.GetEmitterData())
 		{
 			constexpr bool bCompilableOnly = false; // we want to include emitter scripts for parameter store processing
@@ -226,10 +370,13 @@ FNiagaraCompilationTaskHandle FNiagaraSystemCompilingManager::AddSystem(UNiagara
 
 			TArray<UNiagaraScript*> EmitterScripts;
 			EmitterData->GetScripts(EmitterScripts, bCompilableOnly, bEnabledOnly);
+
+			const bool bGpuEmitter = EmitterData->SimTarget == ENiagaraSimTarget::GPUComputeSim;
 			for (UNiagaraScript* EmitterScript : EmitterScripts)
 			{
 				AllScripts.Add(EmitterScript);
-				ScriptsToCompile.Emplace(bHasScriptToCompile, EmitterScript, CompileOptions.bForced, EmitterIt);
+				FCompilableScriptInfo& ScriptToCompile = ScriptsToCompile.Emplace_GetRef(EmitterScript, CompileOptions.bForced, EmitterIt, bGpuEmitter, bHasScriptToCompile);
+				ScriptToCompile.UpdateComputeShaders(bGpuEmitter, NiagaraShaderType, CompileOptions.TargetPlatform, FeatureLevels, bHasScriptToCompile);
 			}
 		}
 	}
@@ -241,37 +388,43 @@ FNiagaraCompilationTaskHandle FNiagaraSystemCompilingManager::AddSystem(UNiagara
 
 	FNiagaraCompilationTaskHandle RequestHandle = NextTaskHandle++;
 
-	// do we really need to care about wrapping?
-	if (RequestHandle == INDEX_NONE)
 	{
-		RequestHandle = NextTaskHandle++;
-	}
+		COOK_STAT(auto Timer = NiagaraSystemCookStats::UsageStats.TimeSyncWork(); Timer.TrackCyclesOnly(););
 
-	FTaskPtr CompilationTask;
 
-	{
-		FWriteScopeLock WriteLock(QueueLock);
+		// do we really need to care about wrapping?
+		if (RequestHandle == INDEX_NONE)
+		{
+			RequestHandle = NextTaskHandle++;
+		}
 
-		QueuedRequests.Add(RequestHandle);
-		CompilationTask = SystemRequestMap.Add(
-			RequestHandle,
-			MakeShared<FNiagaraSystemCompilationTask, ESPMode::ThreadSafe>(RequestHandle, System));
-	}
+		FTaskPtr CompilationTask;
 
-	CompilationTask->PrepareStartTime = FPlatformTime::Seconds();
+		{
+			FWriteScopeLock WriteLock(QueueLock);
 
-	// we're going to have to compile something so let's digest all the collections and build our compilation task
-	CompilationTask->DigestParameterCollections(CompileOptions.ParameterCollections);
-	CompilationTask->DigestSystemInfo();
+			QueuedRequests.Add(RequestHandle);
+			CompilationTask = SystemRequestMap.Add(
+				RequestHandle,
+				MakeShared<FNiagaraSystemCompilationTask, ESPMode::ThreadSafe>(RequestHandle, System));
+		}
 
-	CompilationTask->bForced = CompileOptions.bForced;
+		CompilationTask->PrepareStartTime = FPlatformTime::Seconds();
 
-	for (const FCompilableScriptInfo& ScriptToCompile : ScriptsToCompile)
-	{
-		CompilationTask->AddScript(ScriptToCompile.EmitterIndex, ScriptToCompile.Script, ScriptToCompile.bRequiresCompilation);
-	}
+		// we're going to have to compile something so let's digest all the collections and build our compilation task
+		CompilationTask->DigestParameterCollections(CompileOptions.ParameterCollections);
+		CompilationTask->DigestSystemInfo();
+		CompilationTask->DigestShaderInfo(CompileOptions.TargetPlatform, NiagaraShaderType);
+
+		CompilationTask->bForced = CompileOptions.bForced;
+
+		for (const FCompilableScriptInfo& ScriptToCompile : ScriptsToCompile)
+		{
+			CompilationTask->AddScript(ScriptToCompile.EmitterIndex, ScriptToCompile.Script, ScriptToCompile.CompileId, ScriptToCompile.bRequiresCompilation, ScriptToCompile.ShaderRequests);
+		}
 
-	CompilationTask->QueueStartTime = FPlatformTime::Seconds();
+		CompilationTask->QueueStartTime = FPlatformTime::Seconds();
+	}
 
 	ConditionalLaunchTask();
 
@@ -292,6 +445,8 @@ bool FNiagaraSystemCompilingManager::PollSystemCompile(FNiagaraCompilationTaskHa
 	{
 		if (bWait)
 		{
+			COOK_STAT(auto Timer = NiagaraSystemCookStats::UsageStats.TimeAsyncWait(); Timer.TrackCyclesOnly(););
+
 			TaskPtr->WaitTillCompileCompletion();
 		}
 
@@ -366,10 +521,53 @@ bool FNiagaraSystemCompilingManager::ConditionalLaunchTask()
 	return false;
 }
 
-FNiagaraCompilationTaskHandle FNiagaraEditorModule::RequestCompileSystem(UNiagaraSystem* System, bool bForced)
+void FNiagaraSystemCompilingManager::FindOrAddFeatureLevels(const FCompileOptions& CompileOptions, TArray<FPlatformFeatureLevelPair>& FeatureLevels)
+{
+	if (NiagaraShaderType == nullptr)
+	{
+		NiagaraShaderType = NiagaraSystemCompilingManagerImpl::GetNiagaraShaderType();
+	}
+
+	if (ensure(NiagaraShaderType))
+	{
+		if (CompileOptions.TargetPlatform)
+		{
+			TArray<FPlatformFeatureLevelPair>& CachedFeatureLevels = PlatformFeatureLevels.FindOrAdd(CompileOptions.TargetPlatform);
+
+			if (CachedFeatureLevels.IsEmpty())
+			{
+				TArray<FName> DesiredShaderFormats;
+				CompileOptions.TargetPlatform->GetAllTargetedShaderFormats(DesiredShaderFormats);
+				for (const FName& ShaderFormat : DesiredShaderFormats)
+				{
+					const EShaderPlatform ShaderPlatform = ShaderFormatToLegacyShaderPlatform(ShaderFormat);
+					ERHIFeatureLevel::Type TargetFeatureLevel = GetMaxSupportedFeatureLevel(ShaderPlatform);
+
+					if (NiagaraShaderType->ShouldCompilePermutation(FShaderPermutationParameters(ShaderPlatform)))
+					{
+						CachedFeatureLevels.AddUnique(MakeTuple(ShaderPlatform, TargetFeatureLevel));
+					}
+				}
+			}
+
+			FeatureLevels = CachedFeatureLevels;
+		}
+		else
+		{
+			// if no target platform has been supplied then we just use the current preview feature level
+			FeatureLevels.Add(MakeTuple(CompileOptions.PreviewShaderPlatform, CompileOptions.PreviewFeatureLevel));
+		}
+	}
+}
+
+FNiagaraCompilationTaskHandle FNiagaraEditorModule::RequestCompileSystem(UNiagaraSystem* System, bool bForced, const ITargetPlatform* TargetPlatform)
 {
 	FNiagaraSystemCompilingManager::FCompileOptions CompileOptions;
 	CompileOptions.bForced = bForced;
+	CompileOptions.TargetPlatform = TargetPlatform;
+
+	CompileOptions.PreviewFeatureLevel = UNiagaraScript::GetPreviewFeatureLevel();
+	CompileOptions.PreviewShaderPlatform = GShaderPlatformForFeatureLevel[CompileOptions.PreviewFeatureLevel];
 
 	ParameterCollectionAssetCache.RefreshCache(!FUObjectThreadContext::Get().IsRoutingPostLoad /*bAllowLoading*/);
 	const TArray<TWeakObjectPtr<UNiagaraParameterCollection>>& Collections = ParameterCollectionAssetCache.Get();
diff --git a/Engine/Plugins/FX/Niagara/Source/NiagaraEditor/Private/NiagaraSystemCompilingManager.h b/Engine/Plugins/FX/Niagara/Source/NiagaraEditor/Private/NiagaraSystemCompilingManager.h
--- a/Engine/Plugins/FX/Niagara/Source/NiagaraEditor/Private/NiagaraSystemCompilingManager.h
+++ b/Engine/Plugins/FX/Niagara/Source/NiagaraEditor/Private/NiagaraSystemCompilingManager.h
@@ -9,6 +9,7 @@
 #include "NiagaraSystem.h"
 
 struct FNiagaraSystemCompilationTask;
+class ITargetPlatform;
 
 class FNiagaraSystemCompilingManager : public IAssetCompilingManager
 {
@@ -16,6 +17,9 @@ public:
 	struct FCompileOptions
 	{
 		TArray<TWeakObjectPtr<UNiagaraParameterCollection>> ParameterCollections;
+		const ITargetPlatform* TargetPlatform = nullptr;
+		EShaderPlatform PreviewShaderPlatform = EShaderPlatform::SP_NumPlatforms;
+		ERHIFeatureLevel::Type PreviewFeatureLevel = ERHIFeatureLevel::Num;
 		bool bForced = false;
 	};
 
@@ -42,6 +46,12 @@ protected:
 
 	bool ConditionalLaunchTask();
 
+	FNiagaraShaderType* NiagaraShaderType = nullptr;
+
+	using FPlatformFeatureLevelPair = TPair<EShaderPlatform, ERHIFeatureLevel::Type>;
+	TMap<const ITargetPlatform*, TArray<FPlatformFeatureLevelPair>> PlatformFeatureLevels;
+	void FindOrAddFeatureLevels(const FCompileOptions& CompileOptions, TArray<FPlatformFeatureLevelPair>& FeatureLevels);
+
 	using FTaskPtr = TSharedPtr<FNiagaraSystemCompilationTask, ESPMode::ThreadSafe>;
 
 	mutable FRWLock QueueLock;
diff --git a/Engine/Plugins/FX/Niagara/Source/NiagaraEditor/Public/INiagaraCompiler.h b/Engine/Plugins/FX/Niagara/Source/NiagaraEditor/Public/INiagaraCompiler.h
--- a/Engine/Plugins/FX/Niagara/Source/NiagaraEditor/Public/INiagaraCompiler.h
+++ b/Engine/Plugins/FX/Niagara/Source/NiagaraEditor/Public/INiagaraCompiler.h
@@ -65,7 +65,7 @@ class INiagaraCompiler
 {
 public:
 	UE_DEPRECATED(5.4, "Please update to supply the GroupName directly as this will be removed in a future version.")
-	virtual int32 CompileScript(const class FNiagaraCompileRequestData* InCompileRequest, const FNiagaraCompileOptions& InOptions, const FNiagaraTranslateResults& InTranslateResults, FNiagaraTranslatorOutput* TranslatorOutput, FString& TranslatedHLSL) = 0;
+	virtual int32 CompileScript(const class FNiagaraCompileRequestData* InCompileRequest, const FNiagaraCompileOptions& InOptions, const FNiagaraTranslateResults& InTranslateResults, FNiagaraTranslatorOutput* TranslatorOutput, FString& TranslatedHLSL) { return INDEX_NONE; };
 
 	/** Starts the async compilation of a script and returns the job handle to retrieve the results */
 	virtual int32 CompileScript(const FStringView GroupName, const FNiagaraCompileOptions& InOptions, const FNiagaraTranslateResults& InTranslateResults, const FNiagaraTranslatorOutput& TranslatorOutput, const FString& TranslatedHLSL) = 0;
diff --git a/Engine/Plugins/FX/Niagara/Source/NiagaraEditor/Public/NiagaraEditorModule.h b/Engine/Plugins/FX/Niagara/Source/NiagaraEditor/Public/NiagaraEditorModule.h
--- a/Engine/Plugins/FX/Niagara/Source/NiagaraEditor/Public/NiagaraEditorModule.h
+++ b/Engine/Plugins/FX/Niagara/Source/NiagaraEditor/Public/NiagaraEditorModule.h
@@ -50,6 +50,7 @@ class UNiagaraReservedParametersManager;
 class FNiagaraGraphDataCache;
 class UNiagaraParameterCollection;
 struct FNiagaraSystemAsyncCompileResults;
+class ITargetPlatform;
 
 DECLARE_STATS_GROUP(TEXT("Niagara Editor"), STATGROUP_NiagaraEditor, STATCAT_Advanced);
 
@@ -161,7 +162,7 @@ public:
 		FGuid TargetVersion);
 	TSharedPtr<FNiagaraGraphCachedDataBase, ESPMode::ThreadSafe> CacheGraphTraversal(const UObject* Obj, FGuid Version);
 
-	FNiagaraCompilationTaskHandle RequestCompileSystem(UNiagaraSystem* System, bool bForce);
+	FNiagaraCompilationTaskHandle RequestCompileSystem(UNiagaraSystem* System, bool bForce, const ITargetPlatform* TargetPlatform);
 	bool PollSystemCompile(FNiagaraCompilationTaskHandle, FNiagaraSystemAsyncCompileResults&, bool /*bWait*/, bool /*bPeek*/);
 	void AbortSystemCompile(FNiagaraCompilationTaskHandle);
 
diff --git a/Engine/Plugins/FX/Niagara/Source/NiagaraShader/Private/NiagaraShader.cpp b/Engine/Plugins/FX/Niagara/Source/NiagaraShader/Private/NiagaraShader.cpp
--- a/Engine/Plugins/FX/Niagara/Source/NiagaraShader/Private/NiagaraShader.cpp
+++ b/Engine/Plugins/FX/Niagara/Source/NiagaraShader/Private/NiagaraShader.cpp
@@ -76,7 +76,7 @@ FCriticalSection FNiagaraShaderMap::AllNiagaraShaderMapsGuard;
 #endif
 
 #if WITH_EDITOR
-TMap<TRefCountPtr<FNiagaraShaderMap>, TArray<FNiagaraShaderScript*>> FNiagaraShaderMap::NiagaraShaderMapsBeingCompiled;
+TMap<FNiagaraShaderMapRef, TArray<FNiagaraShaderScript*>> FNiagaraShaderMap::NiagaraShaderMapsBeingCompiled;
 
 static inline bool ShouldCacheNiagaraShader(const FNiagaraShaderType* ShaderType, EShaderPlatform Platform, const FNiagaraShaderScript* Script)
 {
@@ -390,6 +390,76 @@ void FNiagaraShaderType::BeginCompileShader(
 	NewJobs.Add(FShaderCommonCompileJobPtr(NewJob));
 }
 
+
+/**
+ * Enqueues a compilation for a new shader of this type.
+ * @param script - The script to link the shader with.
+ */
+void FNiagaraShaderType::BeginCompileShaderFromSource(
+	FStringView FriendlyName,
+	uint32 ShaderMapId,
+	int32 PermutationId,
+	TSharedRef<FNiagaraShaderScriptParametersMetadata> ShaderParameters,
+	FSharedShaderCompilerEnvironment* CompilationEnvironment,
+	FStringView ScriptSource,
+	EShaderPlatform Platform,
+	const FSimulationStageMetaData& SimStageMetaData,
+	FShaderTarget Target,
+	TArray<TRefCountPtr<FShaderCommonCompileJob>>& NewJobs
+) const
+{
+	FShaderCompileJob* NewJob = GShaderCompilingManager->PrepareShaderCompileJob(ShaderMapId, FShaderCompileJobKey(this, nullptr, PermutationId), EShaderCompileJobPriority::Normal);
+	if (!NewJob)
+	{
+		return;
+	}
+	NewJob->ShaderParameters = MakeShared<const FParameters, ESPMode::ThreadSafe>(ShaderParameters);
+	NewJob->Input.SharedEnvironment = CompilationEnvironment;
+	NewJob->Input.Target = Target;
+	NewJob->Input.ShaderFormat = LegacyShaderPlatformToShaderFormat(Platform);
+	NewJob->Input.VirtualSourceFilePath = TEXT("/Plugin/FX/Niagara/Private/NiagaraEmitterInstanceShader.usf");
+	NewJob->Input.EntryPointName = TEXT("SimulateMainComputeCS");
+	NewJob->Input.Environment.SetDefine(TEXT("GPU_SIMULATION"), 1);
+	NewJob->Input.Environment.SetDefine(TEXT("NIAGARA_MAX_GPU_SPAWN_INFOS"), NIAGARA_MAX_GPU_SPAWN_INFOS);
+	NewJob->Input.Environment.IncludeVirtualPathToContentsMap.Emplace(TEXT("/Engine/Generated/NiagaraEmitterInstance.ush"), ScriptSource);
+
+	NewJob->Input.Environment.SetDefine(TEXT("SimulationStageIndex"), PermutationId);
+
+	if (SimStageMetaData.bWritesParticles && SimStageMetaData.bPartialParticleUpdate)
+	{
+		NewJob->Input.Environment.SetDefine(TEXT("NIAGARA_PARTICLE_PARTIAL_ENABLED"), 1);
+	}
+
+	const FIntVector ThreadGroupSize = SimStageMetaData.GpuDispatchNumThreads;
+	NewJob->Input.Environment.SetDefine(TEXT("NIAGARA_DISPATCH_TYPE"), int(SimStageMetaData.GpuDispatchType));
+	NewJob->Input.Environment.SetDefine(TEXT("NIAGARA_DISPATCH_INDIRECT"), SimStageMetaData.bGpuIndirectDispatch ? 1 : 0);
+	NewJob->Input.Environment.SetDefine(TEXT("THREADGROUP_SIZE"), ThreadGroupSize.X * ThreadGroupSize.Y * ThreadGroupSize.Z);
+	NewJob->Input.Environment.SetDefine(TEXT("THREADGROUP_SIZE_X"), ThreadGroupSize.X);
+	NewJob->Input.Environment.SetDefine(TEXT("THREADGROUP_SIZE_Y"), ThreadGroupSize.Y);
+	NewJob->Input.Environment.SetDefine(TEXT("THREADGROUP_SIZE_Z"), ThreadGroupSize.Z);
+
+	AddUniformBufferIncludesToEnvironment(NewJob->Input.Environment, Platform);
+
+	FShaderCompilerEnvironment& ShaderEnvironment = NewJob->Input.Environment;
+
+	// Allow the shader type to modify the compile environment.
+	ModifyCompilationEnvironment(FShaderPermutationParameters(Platform), ShaderEnvironment);
+
+	::GlobalBeginCompileShader(
+		FString(FriendlyName),
+		nullptr,
+		this,
+		nullptr,//ShaderPipeline,
+		PermutationId,
+		TEXT("/Plugin/FX/Niagara/Private/NiagaraEmitterInstanceShader.usf"),
+		TEXT("SimulateMainComputeCS"),
+		FShaderTarget(GetFrequency(), Platform),
+		NewJob->Input
+	);
+
+	NewJobs.Emplace(NewJob);
+}
+
 void FNiagaraShaderType::AddUniformBufferIncludesToEnvironment(FShaderCompilerEnvironment& OutEnvironment, EShaderPlatform Platform) const
 {
 	UE::ShaderParameters::AddUniformBufferIncludesToEnvironment(OutEnvironment, ReferencedUniformBufferNames);
@@ -457,7 +527,7 @@ static FString GetNiagaraShaderMapKeyString(const FNiagaraShaderMapId& ShaderMap
 	return FDerivedDataCacheInterface::BuildCacheKey(TEXT("NIAGARASM"), *NIAGARASHADERMAP_DERIVEDDATA_VER, *ShaderMapKeyString);
 }
 
-void FNiagaraShaderMap::LoadFromDerivedDataCache(const FNiagaraShaderScript* Script, const FNiagaraShaderMapId& ShaderMapId, EShaderPlatform Platform, TRefCountPtr<FNiagaraShaderMap>& InOutShaderMap)
+void FNiagaraShaderMap::LoadFromDerivedDataCache(const FNiagaraShaderScript* Script, const FNiagaraShaderMapId& ShaderMapId, EShaderPlatform Platform, FNiagaraShaderMapRef& InOutShaderMap)
 {
 	if (InOutShaderMap != NULL)
 	{
@@ -668,6 +738,70 @@ void FNiagaraShaderMap::Compile(
 	}
 }
 
+
+/**
+* Compiles the shaders for a script and caches them in this shader map.
+* @param script - The script to compile shaders for.
+* @param InShaderMapId - the script id and set of static parameters to compile for
+* @param Platform - The platform to compile to
+*/
+void FNiagaraShaderMap::CreateCompileJobs(
+	const FNiagaraShaderType* ShaderType,
+	FStringView FriendlyName,
+	const FNiagaraShaderMapId& InShaderMapId,
+	FStringView ScriptSource,
+	TRefCountPtr<FSharedShaderCompilerEnvironment> CompilationEnvironment,
+	TConstArrayView<FSimulationStageMetaData> StageMetaData,
+	int32 PermutationCount,
+	EShaderPlatform InPlatform,
+	TSharedPtr<FNiagaraShaderScriptParametersMetadata> ShaderParameters,
+	TArray<TRefCountPtr<FShaderCommonCompileJob>>& CompileJobs)
+{
+	if (FPlatformProperties::RequiresCookedData())
+	{
+		return;
+	}
+
+	// Make sure we are operating on a referenced shader map or the below Find will cause this shader map to be deleted,
+	// Since it creates a temporary ref counted pointer.
+	check(NumRefs > 0);
+
+	// Assign a unique identifier so that shaders from this shader map can be associated with it after a deferred compile
+	CompilingId = FShaderCommonCompileJob::GetNextJobId();
+
+	// Store the script name for debugging purposes.
+	FNiagaraShaderMapContent* NewContent = new FNiagaraShaderMapContent(InPlatform);
+	NewContent->FriendlyName = FString(FriendlyName);
+	NewContent->ShaderMapId = InShaderMapId;
+	AssignContent(NewContent);
+
+	// Compile this niagara shader .
+	for (int32 PermutationId = 0; PermutationId < PermutationCount; ++PermutationId)
+	{
+		ShaderType->BeginCompileShaderFromSource(
+			FriendlyName,
+			CompilingId,
+			PermutationId,
+			ShaderParameters.ToSharedRef(),
+			CompilationEnvironment,
+			ScriptSource,
+			InPlatform,
+			StageMetaData[PermutationId],
+			FShaderTarget(ShaderType->GetFrequency(), InPlatform),
+			CompileJobs
+		);
+	}
+
+	// Register this shader map in the global script->shadermap map
+	Register(InPlatform);
+
+	// Mark the shader map as not having been finalized with ProcessCompilationResults
+	bCompilationFinalized = false;
+
+	// Mark as not having been compiled
+	bCompiledSuccessfully = false;
+}
+
 FShader* FNiagaraShaderMap::ProcessCompilationResultsForSingleJob(const TRefCountPtr<class FShaderCommonCompileJob>& SingleJob, const FSHAHash& ShaderMapHash)
 {
 	auto CurrentJob = SingleJob->GetSingleShaderJob();
@@ -682,7 +816,7 @@ FShader* FNiagaraShaderMap::ProcessCompilationResultsForSingleJob(const TRefCoun
 	Shader = NiagaraShaderType->FinishCompileShader(ShaderMapHash, *CurrentJob, GetContent()->FriendlyName);
 	bCompiledSuccessfully = CurrentJob->bSucceeded;
 
-	FNiagaraShader *NiagaraShader = static_cast<FNiagaraShader*>(Shader);
+	FNiagaraShader* NiagaraShader = static_cast<FNiagaraShader*>(Shader);
 
 	// UE-67395 - we had a case where we polluted the DDC with a shader containing no bytecode.
 	check(Shader && Shader->GetCodeSize() > 0);
@@ -690,6 +824,19 @@ FShader* FNiagaraShaderMap::ProcessCompilationResultsForSingleJob(const TRefCoun
 	return GetMutableContent()->FindOrAddShader(NiagaraShaderType->GetHashedName(), CurrentJob->Key.PermutationId, Shader);
 }
 
+void FNiagaraShaderMap::ProcessAndFinalizeShaderCompileJob(const TRefCountPtr<FShaderCommonCompileJob>& SingleJob)
+{
+	FSHAHash ShaderMapHash;
+	GetContent()->ShaderMapId.GetScriptHash(ShaderMapHash);
+
+	ProcessCompilationResultsForSingleJob(SingleJob, ShaderMapHash);
+
+	FinalizeContent();
+
+	// The shader map can now be used on the rendering thread
+	bCompilationFinalized = true;
+}
+
 bool FNiagaraShaderMap::ProcessCompilationResults(const TArray<FShaderCommonCompileJobPtr>& InCompilationResults, int32& InOutJobIndex, float& TimeBudget)
 {
 	check(InOutJobIndex < InCompilationResults.Num());
@@ -914,6 +1061,23 @@ FNiagaraShaderMap::FNiagaraShaderMap() :
 #endif
 }
 
+#if WITH_EDITOR
+FNiagaraShaderMap::FNiagaraShaderMap(EWorkerThread)
+	: CompilingId(1)
+	, NumRefs(0)
+	, bDeletedThroughDeferredCleanup(false)
+	, bRegistered(false)
+	, bCompilationFinalized(true)
+	, bCompiledSuccessfully(true)
+	, bIsPersistent(true) 
+{
+#if ALLOW_SHADERMAP_DEBUG_DATA
+	FScopeLock AllSMAccess(&AllNiagaraShaderMapsGuard);
+	AllNiagaraShaderMaps.Add(this);
+#endif
+}
+#endif // WITH_EDITOR
+
 FNiagaraShaderMap::~FNiagaraShaderMap()
 { 
 	checkSlow(IsInGameThread() || IsAsyncLoading());
@@ -939,7 +1103,7 @@ bool FNiagaraShaderMap::RemovePendingScript(FNiagaraShaderScript* Script)
 	bool bRemoved = false;
 	//All access to NiagaraShaderMapsBeingCompiled must be done on the game thread!
 	check(IsInGameThread());
-	for (TMap<TRefCountPtr<FNiagaraShaderMap>, TArray<FNiagaraShaderScript*> >::TIterator It(NiagaraShaderMapsBeingCompiled); It; ++It)
+	for (TMap<FNiagaraShaderMapRef, TArray<FNiagaraShaderScript*> >::TIterator It(NiagaraShaderMapsBeingCompiled); It; ++It)
 	{
 		TArray<FNiagaraShaderScript*>& Scripts = It.Value();
 		int32 Result = Scripts.Remove(Script);
@@ -993,7 +1157,7 @@ const FNiagaraShaderMap* FNiagaraShaderMap::GetShaderMapBeingCompiled(const FNia
 	// Inefficient search, but only when compiling a lot of shaders
 	//All access to NiagaraShaderMapsBeingCompiled must be done on the game thread!
 	check(IsInGameThread());
-	for (TMap<TRefCountPtr<FNiagaraShaderMap>, TArray<FNiagaraShaderScript*> >::TIterator It(NiagaraShaderMapsBeingCompiled); It; ++It)
+	for (TMap<FNiagaraShaderMapRef, TArray<FNiagaraShaderScript*> >::TIterator It(NiagaraShaderMapsBeingCompiled); It; ++It)
 	{
 		TArray<FNiagaraShaderScript*>& Scripts = It.Value();
 		
diff --git a/Engine/Plugins/FX/Niagara/Source/NiagaraShader/Private/NiagaraShaderCompilationManager.cpp b/Engine/Plugins/FX/Niagara/Source/NiagaraShader/Private/NiagaraShaderCompilationManager.cpp
--- a/Engine/Plugins/FX/Niagara/Source/NiagaraShader/Private/NiagaraShaderCompilationManager.cpp
+++ b/Engine/Plugins/FX/Niagara/Source/NiagaraShader/Private/NiagaraShaderCompilationManager.cpp
@@ -112,17 +112,17 @@ void FNiagaraShaderCompilationManager::ProcessCompiledNiagaraShaderMaps(TMap<int
 	check(IsInGameThread());
 
 	// Keeps shader maps alive as they are passed from the shader compiler and applied to the owning Script
-	TArray<TRefCountPtr<FNiagaraShaderMap> > LocalShaderMapReferences;
+	TArray<FNiagaraShaderMapRef> LocalShaderMapReferences;
 	TMap<FNiagaraShaderScript*, FNiagaraShaderMap*> ScriptsToUpdate;
 
 	// Process compiled shader maps in FIFO order, in case a shader map has been enqueued multiple times,
 	// Which can happen if a script is edited while a background compile is going on
 	for (TMap<int32, FNiagaraShaderMapFinalizeResults>::TIterator ProcessIt(CompiledShaderMaps); ProcessIt; ++ProcessIt)
 	{
-		TRefCountPtr<FNiagaraShaderMap> ShaderMap = NULL;
+		FNiagaraShaderMapRef ShaderMap = NULL;
 		TArray<FNiagaraShaderScript*>* Scripts = NULL;
 
-		for (TMap<TRefCountPtr<FNiagaraShaderMap>, TArray<FNiagaraShaderScript*> >::TIterator ShaderMapIt(FNiagaraShaderMap::GetInFlightShaderMaps()); ShaderMapIt; ++ShaderMapIt)
+		for (TMap<FNiagaraShaderMapRef, TArray<FNiagaraShaderScript*> >::TIterator ShaderMapIt(FNiagaraShaderMap::GetInFlightShaderMaps()); ShaderMapIt; ++ShaderMapIt)
 		{
 			if (ShaderMapIt.Key()->GetCompilingId() == ProcessIt.Key())
 			{
diff --git a/Engine/Plugins/FX/Niagara/Source/NiagaraShader/Private/NiagaraShared.cpp b/Engine/Plugins/FX/Niagara/Source/NiagaraShader/Private/NiagaraShared.cpp
--- a/Engine/Plugins/FX/Niagara/Source/NiagaraShader/Private/NiagaraShared.cpp
+++ b/Engine/Plugins/FX/Niagara/Source/NiagaraShader/Private/NiagaraShared.cpp
@@ -369,7 +369,7 @@ void FNiagaraShaderScript::SerializeShaderMap(FArchive& Ar)
 
 			if (bValid)
 			{
-				TRefCountPtr<FNiagaraShaderMap> LoadedShaderMap = new FNiagaraShaderMap();
+				FNiagaraShaderMapRef LoadedShaderMap = new FNiagaraShaderMap();
 				bool bLoaded = LoadedShaderMap->Serialize(Ar, true, true);
 
 				// Toss the loaded shader data if this is a server only instance
@@ -785,7 +785,7 @@ void FNiagaraShaderScript::GetShaderMapIDsWithUnfinishedCompilation(TArray<int32
 */
 bool FNiagaraShaderScript::BeginCompileShaderMap(
 	const FNiagaraShaderMapId& ShaderMapId,
-	TRefCountPtr<FNiagaraShaderMap>& OutShaderMap,
+	FNiagaraShaderMapRef& OutShaderMap,
 	bool bApplyCompletedShaderMapForRendering,
 	bool bSynchronous)
 {
@@ -798,7 +798,7 @@ bool FNiagaraShaderScript::BeginCompileShaderMap(
 	SCOPE_SECONDS_COUNTER(NiagaraCompileTime);
 
 	// Queue hlsl generation and shader compilation - Unlike materials, we queue this here, and compilation happens from the editor module
-	TRefCountPtr<FNiagaraShaderMap> NewShaderMap = new FNiagaraShaderMap();
+	FNiagaraShaderMapRef NewShaderMap = new FNiagaraShaderMap();
 	OutstandingCompileShaderMapIds.AddUnique(NewShaderMap->GetCompilingId());		
 	UE_LOG(LogShaders, Log, TEXT("BeginCompileShaderMap AddUnique %p %d"), this, NewShaderMap->GetCompilingId());
 
diff --git a/Engine/Plugins/FX/Niagara/Source/NiagaraShader/Public/NiagaraShaderType.h b/Engine/Plugins/FX/Niagara/Source/NiagaraShader/Public/NiagaraShaderType.h
--- a/Engine/Plugins/FX/Niagara/Source/NiagaraShader/Public/NiagaraShaderType.h
+++ b/Engine/Plugins/FX/Niagara/Source/NiagaraShader/Public/NiagaraShaderType.h
@@ -23,7 +23,9 @@ struct FSharedShaderCompilerEnvironment;
 		);
 
 class FNiagaraShaderScript;
+class FShaderCommonCompileJob;
 class FShaderCompileJob;
+struct FSimulationStageMetaData;
 class FUniformExpressionSet;
 
 
@@ -126,10 +128,23 @@ public:
 			const FNiagaraShaderScript* Script,
 			FSharedShaderCompilerEnvironment* CompilationEnvironment,
 			EShaderPlatform Platform,
-			TArray<TRefCountPtr<class FShaderCommonCompileJob>>& NewJobs,
+			TArray<TRefCountPtr<FShaderCommonCompileJob>>& NewJobs,
 			FShaderTarget Target
 		);
 
+	void BeginCompileShaderFromSource(
+		FStringView FriendlyName,
+		uint32 ShaderMapId,
+		int32 PermutationId,
+		TSharedRef<FNiagaraShaderScriptParametersMetadata> ShaderParameters,
+		FSharedShaderCompilerEnvironment* CompilationEnvironment,
+		FStringView ScriptSource,
+		EShaderPlatform Platform,
+		const FSimulationStageMetaData& SimStageMetaData,
+		FShaderTarget Target,
+		TArray<TRefCountPtr<FShaderCommonCompileJob>>& NewJobs
+	) const;
+
 	/**
 	 * Either creates a new instance of this type or returns an equivalent existing shader.
 	 * @param script - The script to link the shader with.
diff --git a/Engine/Plugins/FX/Niagara/Source/NiagaraShader/Public/NiagaraShared.h b/Engine/Plugins/FX/Niagara/Source/NiagaraShader/Public/NiagaraShared.h
--- a/Engine/Plugins/FX/Niagara/Source/NiagaraShader/Public/NiagaraShared.h
+++ b/Engine/Plugins/FX/Niagara/Source/NiagaraShader/Public/NiagaraShared.h
@@ -20,6 +20,7 @@
 #include "SceneTypes.h"
 #include "StaticParameterSet.h"
 #include "Misc/Optional.h"
+#include "NiagaraCommon.h"
 #include "NiagaraCompileHash.h"
 #include "NiagaraDataInterfaceBase.h"
 #include "NiagaraCore.h"
@@ -33,6 +34,7 @@ class FThreadSafeBool;
 class UNiagaraScriptBase;
 struct FNiagaraVMExecutableDataId;
 struct FSharedShaderCompilerEnvironment;
+struct FSimulationStageMetaData;
 
 /** Defines the compile event types for translation/compilation.*/
 UENUM()
@@ -387,7 +389,7 @@ public:
 	* @param ReferenceSet	The set to compare against
 	* @return				true if the sets are equal
 	*/
-	bool operator==(const FNiagaraShaderMapId& ReferenceSet) const;
+	NIAGARASHADER_API bool operator==(const FNiagaraShaderMapId& ReferenceSet) const;
 
 	bool operator!=(const FNiagaraShaderMapId& ReferenceSet) const
 	{
@@ -396,7 +398,7 @@ public:
 
 #if WITH_EDITOR
 	/** Appends string representations of this Id to a key string. */
-	void AppendKeyString(FString& KeyString) const;
+	NIAGARASHADER_API void AppendKeyString(FString& KeyString) const;
 #endif
 
 	/** Returns true if the requested shader type is a dependency of this shader map Id. */
@@ -418,7 +420,7 @@ public:
 	struct NiagaraCompilationQueueItem
 	{
 		FNiagaraShaderScript* Script;
-		TRefCountPtr<FNiagaraShaderMap>ShaderMap;
+		FNiagaraShaderMapRef ShaderMap;
 		FNiagaraShaderMapId ShaderMapId;
 		EShaderPlatform Platform;
 		bool bApply;
@@ -438,7 +440,7 @@ public:
 		return CompilationQueue;
 	}
 
-	void Queue(FNiagaraShaderScript *InScript, TRefCountPtr<FNiagaraShaderMap>InShaderMap, const FNiagaraShaderMapId &MapId, EShaderPlatform InPlatform, bool InApply)
+	void Queue(FNiagaraShaderScript *InScript, FNiagaraShaderMapRef InShaderMap, const FNiagaraShaderMapId &MapId, EShaderPlatform InPlatform, bool InApply)
 	{
 		check(IsInGameThread());
 		NiagaraCompilationQueueItem NewQueueItem;
@@ -498,6 +500,8 @@ class FNiagaraShaderMap : public TShaderMap<FNiagaraShaderMapContent, FNiagaraSh
 public:
 	using Super = TShaderMap<FNiagaraShaderMapContent, FNiagaraShaderMapPointerTable>;
 
+	enum EWorkerThread { WorkerThread };
+
 	/**
 	* Finds the shader map for a script.
 	* @param Platform - The platform to lookup for
@@ -516,13 +520,16 @@ public:
 	* Attempts to load the shader map for the given script from the Derived Data Cache.
 	* If InOutShaderMap is valid, attempts to load the individual missing shaders instead.
 	*/
-	static void LoadFromDerivedDataCache(const FNiagaraShaderScript* Script, const FNiagaraShaderMapId& ShaderMapId, EShaderPlatform Platform, TRefCountPtr<FNiagaraShaderMap>& InOutShaderMap);
+	static void LoadFromDerivedDataCache(const FNiagaraShaderScript* Script, const FNiagaraShaderMapId& ShaderMapId, EShaderPlatform Platform, FNiagaraShaderMapRef& InOutShaderMap);
 #endif
 
-	FNiagaraShaderMap();
+	NIAGARASHADER_API FNiagaraShaderMap();
+#if WITH_EDITOR
+	NIAGARASHADER_API FNiagaraShaderMap(EWorkerThread);
+#endif
 
 	// Destructor.
-	~FNiagaraShaderMap();
+	NIAGARASHADER_API ~FNiagaraShaderMap();
 
 #if WITH_EDITOR
 	/**
@@ -541,6 +548,26 @@ public:
 		bool bApplyCompletedShaderMapForRendering
 		);
 
+	/**
+	* Compiles the shaders for a script and caches them in this shader map.
+	* @param script - The script to compile shaders for.
+	* @param ShaderMapId - the set of static parameters to compile for
+	* @param Platform - The platform to compile to
+	*/
+	NIAGARASHADER_API void CreateCompileJobs(
+		const FNiagaraShaderType* ShaderType,
+		FStringView FriendlyName,
+		const FNiagaraShaderMapId& ShaderMapId,
+		FStringView SourceSource,
+		TRefCountPtr<FSharedShaderCompilerEnvironment> CompilationEnvironment,
+		TConstArrayView<FSimulationStageMetaData> StageMetaData,
+		int32 PermutationCount,
+		EShaderPlatform Platform,
+		TSharedPtr<FNiagaraShaderScriptParametersMetadata> ShaderParameters,
+		TArray<TRefCountPtr<FShaderCommonCompileJob>>& CompileJobs);
+
+	NIAGARASHADER_API void ProcessAndFinalizeShaderCompileJob(const TRefCountPtr<FShaderCommonCompileJob>& SingleJob);
+
 	/** Sorts the incoming compiled jobs into the appropriate mesh shader maps, and finalizes this shader map so that it can be used for rendering. */
 	bool ProcessCompilationResults(const TArray<TRefCountPtr<class FShaderCommonCompileJob>>& InCompilationResults, int32& ResultIndex, float& TimeBudget);
 
@@ -588,7 +615,7 @@ public:
 #endif // WITH_EDITOR
 
 	/** Serializes the shader map. */
-	bool Serialize(FArchive& Ar, bool bInlineShaderResources = true, bool bLoadedByCookedMaterial = false);
+	NIAGARASHADER_API bool Serialize(FArchive& Ar, bool bInlineShaderResources = true, bool bLoadedByCookedMaterial = false);
 
 #if WITH_EDITOR
 	/** Saves this shader map to the derived data cache. */
@@ -621,7 +648,7 @@ public:
 	uint32 GetCompilingId()  { return CompilingId; }
 
 #if WITH_EDITOR
-	static TMap<TRefCountPtr<FNiagaraShaderMap>, TArray<FNiagaraShaderScript*>>& GetInFlightShaderMaps()
+	static TMap<FNiagaraShaderMapRef, TArray<FNiagaraShaderScript*>>& GetInFlightShaderMaps()
 	{
 		//All access to NiagaraShaderMapsBeingCompiled must be done on the game thread!
 		check(IsInGameThread());
@@ -653,7 +680,7 @@ private:
 
 #if WITH_EDITOR
 	/** Tracks resources and their shader maps that need to be compiled but whose compilation is being deferred. */
-	static TMap<TRefCountPtr<FNiagaraShaderMap>, TArray<FNiagaraShaderScript*> > NiagaraShaderMapsBeingCompiled;
+	static TMap<FNiagaraShaderMapRef, TArray<FNiagaraShaderScript*> > NiagaraShaderMapsBeingCompiled;
 #endif
 
 	/** Uniquely identifies this shader map during compilation, needed for deferred compilation where shaders from multiple shader maps are compiled together. */
@@ -902,7 +929,7 @@ public:
 		return OnCompilationCompleteDelegate;
 	}
 
-	bool IsSame(const FNiagaraShaderMapId& InId) const;
+	NIAGARASHADER_API bool IsSame(const FNiagaraShaderMapId& InId) const;
 
 	int32 GetNumPermutations() const { return NumPermutations; }
 
@@ -947,7 +974,7 @@ private:
 	 * Code that sets this is responsible for updating RenderingThreadShaderMap in a thread safe way.
 	 * During an async compile, this will be NULL and will not contain the actual shader map until compilation is complete.
 	 */
-	TRefCountPtr<FNiagaraShaderMap> GameThreadShaderMap;
+	FNiagaraShaderMapRef GameThreadShaderMap;
 
 	/** 
 	 * Shader map for this FNiagaraShaderScript which is accessible by the rendering thread. 
diff --git a/Engine/Source/Runtime/Core/Private/UObject/FortniteMainBranchObjectVersion.cpp b/Engine/Source/Runtime/Core/Private/UObject/FortniteMainBranchObjectVersion.cpp
--- a/Engine/Source/Runtime/Core/Private/UObject/FortniteMainBranchObjectVersion.cpp
+++ b/Engine/Source/Runtime/Core/Private/UObject/FortniteMainBranchObjectVersion.cpp
@@ -11,7 +11,7 @@ TMap<FGuid, FGuid> FFortniteMainBranchObjectVersion::GetSystemGuids()
 	SystemGuids.Add(DevGuids.MATERIALSHADERMAP_DERIVEDDATA_VER, FGuid("CE7776E63E9E43FBB8C103EE1AB40B5D"));
 	SystemGuids.Add(DevGuids.NANITE_DERIVEDDATA_VER, FGuid("D28C003F825F4C178F698938FB0E47F7"));
 	SystemGuids.Add(DevGuids.NIAGARASHADERMAP_DERIVEDDATA_VER, FGuid("D0DF49D408F34E38A98FA097E0280EC7"));
-	SystemGuids.Add(DevGuids.Niagara_LatestScriptCompileVersion, FGuid("1EED8B67ED534586B369FA9BEC820760"));
+	SystemGuids.Add(DevGuids.Niagara_LatestScriptCompileVersion, FGuid("C4B4DCCEA3ED44FF93A7E76871460578"));
 	SystemGuids.Add(DevGuids.SkeletalMeshDerivedDataVersion, FGuid("518D7C51D31D9BE4A9F4B81542AC8119"));
 	SystemGuids.Add(DevGuids.STATICMESH_DERIVEDDATA_VER, FGuid("543843B103794E0E9BA4BE80FB602F79"));
 	return SystemGuids;
