Fix for ambiguously typed clamp in GetSkinnedTriangleVertexDataHelper

diff --git a/Engine/Plugins/FX/Niagara/Shaders/Private/NiagaraDataInterfaceSkeletalMeshTemplate.ush b/Engine/Plugins/FX/Niagara/Shaders/Private/NiagaraDataInterfaceSkeletalMeshTemplate.ush
--- a/Engine/Plugins/FX/Niagara/Shaders/Private/NiagaraDataInterfaceSkeletalMeshTemplate.ush
+++ b/Engine/Plugins/FX/Niagara/Shaders/Private/NiagaraDataInterfaceSkeletalMeshTemplate.ush
@@ -494,7 +494,7 @@ void GetSkinnedTriangleVertexDataHelper_{ParameterName}(
     int VertexIndex2 = 0;
 	if ( {ParameterName}_MeshTriangleCount != 0 )
 	{
-		TriangleIndex = clamp(TriangleIndex, 0, {ParameterName}_MeshTriangleCount - 1);
+		TriangleIndex = clamp(TriangleIndex, 0, int({ParameterName}_MeshTriangleCount) - 1);
 		const uint IndexBufferOffset = TriangleIndex * 3;
 		VertexIndex0 = {ParameterName}_MeshIndexBuffer[IndexBufferOffset];
 		VertexIndex1 = {ParameterName}_MeshIndexBuffer[IndexBufferOffset + 1];
