Removes ensures from getting triggered when typical graph errors suffice for pushing the informaiton to the user (rather than a potentially long hitch processing the ensure)

diff --git a/Engine/Plugins/FX/Niagara/Source/NiagaraEditor/Private/NiagaraCompilationCopy.cpp b/Engine/Plugins/FX/Niagara/Source/NiagaraEditor/Private/NiagaraCompilationCopy.cpp
--- a/Engine/Plugins/FX/Niagara/Source/NiagaraEditor/Private/NiagaraCompilationCopy.cpp
+++ b/Engine/Plugins/FX/Niagara/Source/NiagaraEditor/Private/NiagaraCompilationCopy.cpp
@@ -264,8 +264,6 @@ void FNiagaraPrecompileData::FinishPrecompile(
 			Builder.BuildParameterMaps(FoundOutputNode, true);
 			Builder.EndUsage();
 
-			ensure(Builder.Histories.Num() <= 1);
-
 			int HistoryIdx = 0;
 			for (FParameterMapHistory& History : Builder.Histories)
 			{
@@ -553,8 +551,6 @@ void FNiagaraCompilationCopyData::CreateParameterMapHistory(const FNiagaraSystem
 			Builder.BuildParameterMaps(FoundOutputNode, true);
 			Builder.EndUsage();
 
-			ensure(Builder.Histories.Num() <= 1);
-
 			for (FParameterMapHistory& History : Builder.Histories)
 			{
 				History.OriginatingScriptUsage = FoundOutputNode->Usage;
diff --git a/Engine/Plugins/FX/Niagara/Source/NiagaraEditor/Private/NiagaraCompiler.cpp b/Engine/Plugins/FX/Niagara/Source/NiagaraEditor/Private/NiagaraCompiler.cpp
--- a/Engine/Plugins/FX/Niagara/Source/NiagaraEditor/Private/NiagaraCompiler.cpp
+++ b/Engine/Plugins/FX/Niagara/Source/NiagaraEditor/Private/NiagaraCompiler.cpp
@@ -658,8 +658,6 @@ void FNiagaraCompileRequestData::FinishPrecompile(const TArray<FNiagaraVariable>
 				Builder.BuildParameterMaps(FoundOutputNode, true);
 				Builder.EndUsage();
 
-				ensure(Builder.Histories.Num() <= 1);
-
 				int HistoryIdx = 0;
 				for (FNiagaraParameterMapHistory& History : Builder.Histories)
 				{
@@ -898,8 +896,6 @@ void FNiagaraCompileRequestDuplicateData::FinishPrecompileDuplicate(const TArray
 			Builder.BuildParameterMaps(FoundOutputNode, true);
 			Builder.EndUsage();
 
-			ensure(Builder.Histories.Num() <= 1);
-
 			for (FNiagaraParameterMapHistory& History : Builder.Histories)
 			{
 				History.OriginatingScriptUsage = FoundOutputNode->GetUsage();
@@ -2515,9 +2511,9 @@ bool FNiagaraShaderMapCompiler::ProcessCompileResults(bool bWait)
 
 		// for now we'll process all shaders at once (need to measure the cost here)
 		for (FShaderCommonCompileJobPtr& ShaderCompileJob : CurrentCompilation.ShaderCompileJobs)
-		{
-			CurrentCompilation.ShaderMap->ProcessAndFinalizeShaderCompileJob(ShaderCompileJob);
-		}
+			{
+				CurrentCompilation.ShaderMap->ProcessAndFinalizeShaderCompileJob(ShaderCompileJob);
+			}
 
 		// now that we've added all the results into the shader map we can move it over to CompletedCompilations
 		CompletedCompilations.Add(CurrentCompilation.ShaderMap);
diff --git a/Engine/Plugins/FX/Niagara/Source/NiagaraEditor/Private/ViewModels/Stack/NiagaraStackGraphUtilities.cpp b/Engine/Plugins/FX/Niagara/Source/NiagaraEditor/Private/ViewModels/Stack/NiagaraStackGraphUtilities.cpp
--- a/Engine/Plugins/FX/Niagara/Source/NiagaraEditor/Private/ViewModels/Stack/NiagaraStackGraphUtilities.cpp
+++ b/Engine/Plugins/FX/Niagara/Source/NiagaraEditor/Private/ViewModels/Stack/NiagaraStackGraphUtilities.cpp
@@ -1859,7 +1859,7 @@ void FNiagaraStackGraphUtilities::GetStackFunctionOutputVariables(UNiagaraNodeFu
 
 	FunctionCallNode.BuildParameterMapHistory(Builder, false);
 
-	if (ensureMsgf(Builder.Histories.Num() == 1, TEXT("Invalid Stack Graph - Function call node has invalid history count!")))
+	if (Builder.Histories.Num() == 1)
 	{
 		for (int32 i = 0; i < Builder.Histories[0].Variables.Num(); i++)
 		{
@@ -1882,6 +1882,10 @@ void FNiagaraStackGraphUtilities::GetStackFunctionOutputVariables(UNiagaraNodeFu
 			}
 		}
 	}
+	else
+	{
+		UE_LOG(LogNiagaraEditor, Log, TEXT("Invalid Stack Graph - Function call node has invalid history count!"));
+	}
 }
 
 bool FNiagaraStackGraphUtilities::GetStackFunctionInputAndOutputVariables(UNiagaraNodeFunctionCall& FunctionCallNode, FCompileConstantResolver ConstantResolver, TArray<FNiagaraVariable>& OutVariables, TArray<FNiagaraVariable>& OutVariablesWithOriginalAliasesIntact)
diff --git a/Engine/Plugins/FX/Niagara/Source/NiagaraEditor/Private/ViewModels/Stack/NiagaraStackModuleItemLinkedInputCollection.cpp b/Engine/Plugins/FX/Niagara/Source/NiagaraEditor/Private/ViewModels/Stack/NiagaraStackModuleItemLinkedInputCollection.cpp
--- a/Engine/Plugins/FX/Niagara/Source/NiagaraEditor/Private/ViewModels/Stack/NiagaraStackModuleItemLinkedInputCollection.cpp
+++ b/Engine/Plugins/FX/Niagara/Source/NiagaraEditor/Private/ViewModels/Stack/NiagaraStackModuleItemLinkedInputCollection.cpp
@@ -82,7 +82,7 @@ void UNiagaraStackModuleItemLinkedInputCollection::RefreshChildrenInternal(const
 			: FCompileConstantResolver();
 		FunctionCallNode->BuildParameterMapHistory(Builder, false);
 
-		if (ensureMsgf(Builder.Histories.Num() == 1, TEXT("Invalid Stack Graph - Function call node has invalid history count!")))
+		if (Builder.Histories.Num() == 1)
 		{
 			for (int32 i = 0; i < Builder.Histories[0].Variables.Num(); i++)
 			{
@@ -114,6 +114,10 @@ void UNiagaraStackModuleItemLinkedInputCollection::RefreshChildrenInternal(const
 				}
 			}
 		}
+		else
+		{
+			UE_LOG(LogNiagaraEditor, Log, TEXT("Invalid Stack Graph - Function call node has invalid history count!"));
+		}
 
 		auto SortNewChildrenPred = [](const UNiagaraStackEntry& EntryA, const UNiagaraStackEntry& EntryB) {
 			const UNiagaraStackModuleItemOutput* ModuleItemA = static_cast<const UNiagaraStackModuleItemOutput*>(&EntryA);
