Fix handling for shader compiler errors with new AsyncTask GPU simulation shaders

diff --git a/Engine/Plugins/FX/Niagara/Source/Niagara/Classes/NiagaraScript.h b/Engine/Plugins/FX/Niagara/Source/Niagara/Classes/NiagaraScript.h
--- a/Engine/Plugins/FX/Niagara/Source/Niagara/Classes/NiagaraScript.h
+++ b/Engine/Plugins/FX/Niagara/Source/Niagara/Classes/NiagaraScript.h
@@ -1269,7 +1269,8 @@ public:
 		EShaderPlatform ShaderPlatform,
 		ERHIFeatureLevel::Type FeatureLevel,
 		const FNiagaraShaderScriptParametersMetadata& SharedParameters,
-		const FNiagaraShaderMapRef& ShaderMap);
+		const FNiagaraShaderMapRef& ShaderMap,
+		TConstArrayView<FShaderCompilerError> CompilationErrors);
 
 	NIAGARA_API bool IsShaderMapCached(const ITargetPlatform* TargetPlatform, const FNiagaraShaderMapId& ShaderMapId) const;
 
@@ -1349,6 +1350,8 @@ private:
 	bool ShouldOptimizeByteCode() const;
 	bool ShouldFreeUnoptimizedByteCode() const;
 
+	void ProcessCompilationErrors(bool bCompileSucceeded, TConstArrayView<FShaderCompilerError> CompilationErrors, TArray<FString>& ProcessedErrors);
+
 public:
 	/** Kicks off an async job to convert the ByteCode into an optimized version for the platform we are running on. */
 	FGraphEventRef HandleByteCodeOptimization(bool bShouldForceNow = false);
diff --git a/Engine/Plugins/FX/Niagara/Source/Niagara/Private/NiagaraCompilationTypes.cpp b/Engine/Plugins/FX/Niagara/Source/Niagara/Private/NiagaraCompilationTypes.cpp
--- a/Engine/Plugins/FX/Niagara/Source/Niagara/Private/NiagaraCompilationTypes.cpp
+++ b/Engine/Plugins/FX/Niagara/Source/Niagara/Private/NiagaraCompilationTypes.cpp
@@ -358,7 +358,8 @@ public:
 							ShaderMapInfo.ShaderPlatform,
 							ShaderMapInfo.FeatureLevel,
 							ScriptCompileData.ExeData->ShaderScriptParametersMetadata,
-							ShaderMapInfo.CompiledShader);
+							ShaderMapInfo.CompiledShader,
+							ShaderMapInfo.CompilationErrors);
 					}
 
 				}
diff --git a/Engine/Plugins/FX/Niagara/Source/Niagara/Private/NiagaraScript.cpp b/Engine/Plugins/FX/Niagara/Source/Niagara/Private/NiagaraScript.cpp
--- a/Engine/Plugins/FX/Niagara/Source/Niagara/Private/NiagaraScript.cpp
+++ b/Engine/Plugins/FX/Niagara/Source/Niagara/Private/NiagaraScript.cpp
@@ -3015,13 +3015,60 @@ void UNiagaraScript::SetVMCompilationResults(const FNiagaraVMExecutableDataId& I
 	OnVMScriptCompiled().Broadcast(this, InCompileId.ScriptVersionID);
 }
 
+void UNiagaraScript::ProcessCompilationErrors(bool bCompileSucceeded, TConstArrayView<FShaderCompilerError> CompilationErrors, TArray<FString>& ProcessedErrors)
+{
+	ProcessedErrors.Reserve(CompilationErrors.Num());
+
+	if (CompilationErrors.IsEmpty())
+	{
+		return;
+	}
+
+	for (const FShaderCompilerError& CompilationError : CompilationErrors)
+	{
+		ProcessedErrors.Add(CompilationError.GetErrorString().Replace(TEXT("Error"), TEXT("Err0r")));
+	}
+
+	static const IConsoleVariable* ShowShaderWarningsVariable = IConsoleManager::Get().FindConsoleVariable(TEXT("niagara.ShowShaderCompilerWarnings"));
+
+	if (!bCompileSucceeded || (ShowShaderWarningsVariable && ShowShaderWarningsVariable->GetBool()))
+	{
+		UE_LOG(LogNiagara, Display, TEXT("There were issues for compilation of script \"%s\""), *GetFriendlyName());
+
+		if (!LogNiagara.IsSuppressed(ELogVerbosity::Verbose))
+		{
+			TArray<FString> OutputByLines;
+			CachedScriptVM.LastHlslTranslationGPU.ParseIntoArrayLines(OutputByLines, false);
+
+			UE_LOG(LogNiagara, Verbose, TEXT("Compile output as text:"));
+			UE_LOG(LogNiagara, Verbose, TEXT("==================================================================================="));
+			for (int32 i = 0; i < OutputByLines.Num(); i++)
+			{
+				UE_LOG(LogNiagara, Verbose, TEXT("/*%04d*/\t\t%s"), i + 1, *OutputByLines[i]);
+			}
+			UE_LOG(LogNiagara, Verbose, TEXT("==================================================================================="));
+		}
+
+		for (const FShaderCompilerError& Error : CompilationErrors)
+		{
+			UE_LOG(LogNiagara, Warning, TEXT("%s"), *Error.GetErrorString())
+		}
+	}
+}
+
 void UNiagaraScript::SetComputeCompilationResults(
 	const ITargetPlatform* TargetPlatform,
 	EShaderPlatform ShaderPlatform,
 	ERHIFeatureLevel::Type FeatureLevel,
 	const FNiagaraShaderScriptParametersMetadata& ShaderParameters,
-	const FNiagaraShaderMapRef& ShaderMap)
+	const FNiagaraShaderMapRef& ShaderMap,
+	TConstArrayView<FShaderCompilerError> CompilationErrors)
 {
+	const bool bCompilationSuccessful = ShaderMap.IsValid() && ShaderMap->IsValid();
+
+	TArray<FString> ProcessedErrors;
+	ProcessCompilationErrors(bCompilationSuccessful, CompilationErrors, ProcessedErrors);
+
 	// for now we will use the CachedScriptResourcesForCooking for storing all compilation results and will update ScriptResource
 	// if it matches with the supplied parameters
 	auto ScriptMatchesFeatureLevel = [FeatureLevel](const TUniquePtr<FNiagaraShaderScript>& ShaderScript) -> bool
@@ -3041,10 +3088,11 @@ void UNiagaraScript::SetComputeCompilationResults(
 		TargetShaderScript = CachedScripts.Emplace_GetRef(MakeUnique<FNiagaraShaderScript>()).Get();
 	}
 
-	auto UpdateShaderScript = [this, FeatureLevel, ShaderPlatform, &ShaderParameters, ShaderMap](FNiagaraShaderScript* ScriptToUpdate) -> void
+	auto UpdateShaderScript = [this, &ProcessedErrors, &ShaderParameters, FeatureLevel, ShaderPlatform](FNiagaraShaderScript* ScriptToUpdate, FNiagaraShaderMap* InShaderMap) -> void
 	{
+		ScriptToUpdate->SetCompileErrors(ProcessedErrors);
 		ScriptToUpdate->BuildScriptParametersMetadata(ShaderParameters);
-		ScriptToUpdate->SetShaderMap(ShaderMap);
+		ScriptToUpdate->SetGameThreadShaderMap(InShaderMap);
 		ScriptToUpdate->SetScript(this, FeatureLevel, ShaderPlatform, CachedScriptVMId.CompilerVersionID, CachedScriptVMId.AdditionalDefines,
 			CachedScriptVMId.GetAdditionalVariableStrings(),
 			CachedScriptVMId.BaseScriptCompileHash, CachedScriptVMId.ReferencedCompileHashes,
@@ -3053,7 +3101,9 @@ void UNiagaraScript::SetComputeCompilationResults(
 
 	if (ensure(TargetShaderScript))
 	{
-		UpdateShaderScript(TargetShaderScript);
+		FNiagaraShaderMap* ShaderMapToApply = bCompilationSuccessful ? ShaderMap.GetReference() : nullptr;
+
+		UpdateShaderScript(TargetShaderScript, ShaderMapToApply);
 
 		// next see if we need to update the currently active ShaderScript (and update it's renderthread data)
 		{
@@ -3068,11 +3118,11 @@ void UNiagaraScript::SetComputeCompilationResults(
 				}
 				else
 				{
-					UpdateShaderScript(ScriptResource.Get());
+					UpdateShaderScript(ScriptResource.Get(), ShaderMapToApply);
 				}
 
 				ENQUEUE_RENDER_COMMAND(FSetShaderMapOnScriptResources)(
-					[RT_Script = ScriptResource, RT_ShaderMap = ShaderMap](FRHICommandListImmediate& RHICmdList)
+					[RT_Script = ScriptResource, RT_ShaderMap = ShaderMapToApply](FRHICommandListImmediate& RHICmdList)
 				{
 					RT_Script->SetRenderingThreadShaderMap(RT_ShaderMap);
 				});
diff --git a/Engine/Plugins/FX/Niagara/Source/Niagara/Public/NiagaraCompilationTypes.h b/Engine/Plugins/FX/Niagara/Source/Niagara/Public/NiagaraCompilationTypes.h
--- a/Engine/Plugins/FX/Niagara/Source/Niagara/Public/NiagaraCompilationTypes.h
+++ b/Engine/Plugins/FX/Niagara/Source/Niagara/Public/NiagaraCompilationTypes.h
@@ -36,6 +36,7 @@ struct FNiagaraSystemCompileMetrics
 struct FNiagaraCompiledShaderInfo
 {
 	FNiagaraShaderMapRef CompiledShader;
+	TArray<FShaderCompilerError> CompilationErrors;
 	const ITargetPlatform* TargetPlatform = nullptr;
 	EShaderPlatform ShaderPlatform = SP_NumPlatforms;
 	ERHIFeatureLevel::Type FeatureLevel = ERHIFeatureLevel::Num;
diff --git a/Engine/Plugins/FX/Niagara/Source/NiagaraEditor/Private/NiagaraCompilationTasks.cpp b/Engine/Plugins/FX/Niagara/Source/NiagaraEditor/Private/NiagaraCompilationTasks.cpp
--- a/Engine/Plugins/FX/Niagara/Source/NiagaraEditor/Private/NiagaraCompilationTasks.cpp
+++ b/Engine/Plugins/FX/Niagara/Source/NiagaraEditor/Private/NiagaraCompilationTasks.cpp
@@ -1028,8 +1028,7 @@ void FNiagaraSystemCompilationTask::FCompileTaskInfo::ProcessCompilation(FNiagar
 		for (int32 ComputeShaderTaskIndex : ComputeShaderTaskIndices)
 		{
 			FCompileComputeShaderTaskInfo& ShaderTaskInfo = SystemCompileTask->CompileComputeShaderTasks[ComputeShaderTaskIndex];
-
-			ShaderTaskInfo.ShaderMap = ShaderMapCompiler->GetShaderMap(ShaderTaskInfo.ShaderMapId);
+			ShaderMapCompiler->GetShaderMap(ShaderTaskInfo.ShaderMapId, ShaderTaskInfo.ShaderMap, ShaderTaskInfo.CompilationErrors);
 		}
 	}
 }
@@ -1244,7 +1243,7 @@ void FNiagaraSystemCompilationTask::FDispatchDataCachePutRequests::Launch(FNiaga
 
 	for (const FNiagaraSystemCompilationTask::FCompileComputeShaderTaskInfo& CompileTask : SystemCompileTask->CompileComputeShaderTasks)
 	{
-		if (CompileTask.ShaderMap.IsValid() && !CompileTask.DDCTaskInfo.DataCachePutKeys.IsEmpty())
+		if (CompileTask.ShaderMap.IsValid() && CompileTask.ShaderMap->IsValid() && !CompileTask.DDCTaskInfo.DataCachePutKeys.IsEmpty())
 		{
 			const FNiagaraSystemCompilationTask::FCompileTaskInfo& ParentCompileTask = SystemCompileTask->CompileTasks[CompileTask.ParentCompileTaskIndex];
 
@@ -1625,6 +1624,7 @@ bool FNiagaraSystemCompilationTask::Poll(FNiagaraSystemAsyncCompileResults& Resu
 						ResultsInfo.ShaderPlatform = ShaderTaskInfo.ShaderPlatform;
 						ResultsInfo.FeatureLevel = ShaderTaskInfo.ShaderMapId.FeatureLevel;
 						ResultsInfo.CompiledShader = ShaderTaskInfo.ShaderMap;
+						ResultsInfo.CompilationErrors = ShaderTaskInfo.CompilationErrors;
 					}
 
 					CompileData.NamedDataInterfaces.Reserve(TaskInfo.NamedDataInterfaces.Num());
diff --git a/Engine/Plugins/FX/Niagara/Source/NiagaraEditor/Private/NiagaraCompilationTasks.h b/Engine/Plugins/FX/Niagara/Source/NiagaraEditor/Private/NiagaraCompilationTasks.h
--- a/Engine/Plugins/FX/Niagara/Source/NiagaraEditor/Private/NiagaraCompilationTasks.h
+++ b/Engine/Plugins/FX/Niagara/Source/NiagaraEditor/Private/NiagaraCompilationTasks.h
@@ -147,7 +147,7 @@ struct FNiagaraSystemCompilationTask
 		FNiagaraShaderMapId ShaderMapId;
 		EShaderPlatform ShaderPlatform;
 		FNiagaraShaderMapRef ShaderMap;
-
+		TArray<FShaderCompilerError> CompilationErrors;
 		FDDCTaskInfo DDCTaskInfo;
 
 		bool IsOutstanding() const;
diff --git a/Engine/Plugins/FX/Niagara/Source/NiagaraEditor/Private/NiagaraCompiler.cpp b/Engine/Plugins/FX/Niagara/Source/NiagaraEditor/Private/NiagaraCompiler.cpp
--- a/Engine/Plugins/FX/Niagara/Source/NiagaraEditor/Private/NiagaraCompiler.cpp
+++ b/Engine/Plugins/FX/Niagara/Source/NiagaraEditor/Private/NiagaraCompiler.cpp
@@ -2508,15 +2508,29 @@ bool FNiagaraShaderMapCompiler::ProcessCompileResults(bool bWait)
 		}
 
 		FActiveCompilation& CurrentCompilation = *CompileIt;
+		FCompletedCompilation& CompletedCompilation = CompletedCompilations.AddDefaulted_GetRef();
 
 		// for now we'll process all shaders at once (need to measure the cost here)
-		for (FShaderCommonCompileJobPtr& ShaderCompileJob : CurrentCompilation.ShaderCompileJobs)
+		for (const FShaderCommonCompileJobPtr& ShaderCompileJob : CurrentCompilation.ShaderCompileJobs)
+		{
+			if (ShaderCompileJob.IsValid() && ShaderCompileJob->bSucceeded)
 			{
 				CurrentCompilation.ShaderMap->ProcessAndFinalizeShaderCompileJob(ShaderCompileJob);
 			}
+			else
+			{
+				CurrentCompilation.ShaderMap->SetCompiledSuccessfully(false);
+			}
+
+			// pass on error/warning info
+			if (const FShaderCompileJob* SingleShaderJob = ShaderCompileJob->GetSingleShaderJob())
+			{
+				CompletedCompilation.CompilationErrors.Append(SingleShaderJob->Output.Errors);
+			}
+		}
 
 		// now that we've added all the results into the shader map we can move it over to CompletedCompilations
-		CompletedCompilations.Add(CurrentCompilation.ShaderMap);
+		CompletedCompilation.ShaderMap = CurrentCompilation.ShaderMap;
 
 		// and remove it from the ActiveCompilations
 		CompileIt.RemoveCurrentSwap();
@@ -2525,17 +2539,19 @@ bool FNiagaraShaderMapCompiler::ProcessCompileResults(bool bWait)
 	return ActiveCompilations.IsEmpty();
 }
 
-FNiagaraShaderMapRef FNiagaraShaderMapCompiler::GetShaderMap(const FNiagaraShaderMapId& ShaderMapId) const
+bool FNiagaraShaderMapCompiler::GetShaderMap(const FNiagaraShaderMapId& ShaderMapId, FNiagaraShaderMapRef& OutShaderMap, TArray<FShaderCompilerError>& OutCompilationErrors) const
 {
-	for (const FNiagaraShaderMapRef& ShaderMapRef : CompletedCompilations)
+	for (const FCompletedCompilation& CompletedCompilation : CompletedCompilations)
 	{
-		if (ShaderMapRef->GetShaderMapId() == ShaderMapId)
+		if (CompletedCompilation.ShaderMap->GetShaderMapId() == ShaderMapId)
 		{
-			return ShaderMapRef;
+			OutShaderMap = CompletedCompilation.ShaderMap;
+			OutCompilationErrors = CompletedCompilation.CompilationErrors;
+			return true;
 		}
 	}
 
-	return nullptr;
+	return false;
 }
 
 //////////////////////////////////////////////////////////////////////////
diff --git a/Engine/Plugins/FX/Niagara/Source/NiagaraEditor/Private/NiagaraCompiler.h b/Engine/Plugins/FX/Niagara/Source/NiagaraEditor/Private/NiagaraCompiler.h
--- a/Engine/Plugins/FX/Niagara/Source/NiagaraEditor/Private/NiagaraCompiler.h
+++ b/Engine/Plugins/FX/Niagara/Source/NiagaraEditor/Private/NiagaraCompiler.h
@@ -83,12 +83,18 @@ public:
 	// returns true if there are no more results requiring processing
 	bool ProcessCompileResults(bool bWait = false);
 
-	TConstArrayView<FNiagaraShaderMapRef> ReadCompletedCompilations() const
+	struct FCompletedCompilation
+	{
+		FNiagaraShaderMapRef ShaderMap;
+		TArray<FShaderCompilerError> CompilationErrors;
+	};
+
+	TConstArrayView<FCompletedCompilation> ReadCompletedCompilations() const
 	{
 		return CompletedCompilations;
 	}
 
-	FNiagaraShaderMapRef GetShaderMap(const FNiagaraShaderMapId& ShaderMapId) const;
+	bool GetShaderMap(const FNiagaraShaderMapId& ShaderMapId, FNiagaraShaderMapRef& OutShaderMap, TArray<FShaderCompilerError>& OutCompilationErrors) const;
 
 	TSharedPtr<FNiagaraVMExecutableData> ReadScriptMetaData() const
 	{
@@ -112,6 +118,6 @@ private:
 	const FNiagaraShaderType* ShaderType = nullptr;
 	TSharedPtr<FNiagaraShaderScriptParametersMetadata> ShaderParameters;
 	TArray<FActiveCompilation> ActiveCompilations;
-	TArray<FNiagaraShaderMapRef> CompletedCompilations;
+	TArray<FCompletedCompilation> CompletedCompilations;
 	TSharedPtr<FNiagaraVMExecutableData> ScriptExeData;
 };
diff --git a/Engine/Plugins/FX/Niagara/Source/NiagaraShader/Public/NiagaraShared.h b/Engine/Plugins/FX/Niagara/Source/NiagaraShader/Public/NiagaraShared.h
--- a/Engine/Plugins/FX/Niagara/Source/NiagaraShader/Public/NiagaraShared.h
+++ b/Engine/Plugins/FX/Niagara/Source/NiagaraShader/Public/NiagaraShared.h
@@ -906,10 +906,6 @@ public:
 
 #if WITH_EDITOR
 	NIAGARASHADER_API bool MatchesScript(ERHIFeatureLevel::Type InFeatureLevel, EShaderPlatform InShaderPlatform, const FNiagaraVMExecutableDataId& ScriptId) const;
-	void SetCompileErrors(TArray<FString>& InErrors)
-	{
-		CompileErrors = InErrors;
-	}
 #endif
 
 	UNiagaraScriptBase* GetBaseVMScript() { return BaseVMScript; }
