* Fixed delete icon in Niagara parameter collection editor
* Fixed deleting an entry in a parameter collection via hotkey (instead of using the button)
* Fixed top level area expansion state

diff --git a/Engine/Plugins/FX/Niagara/Source/Niagara/Private/NiagaraParameterCollection.cpp b/Engine/Plugins/FX/Niagara/Source/Niagara/Private/NiagaraParameterCollection.cpp
--- a/Engine/Plugins/FX/Niagara/Source/Niagara/Private/NiagaraParameterCollection.cpp
+++ b/Engine/Plugins/FX/Niagara/Source/Niagara/Private/NiagaraParameterCollection.cpp
@@ -540,6 +540,10 @@ FString UNiagaraParameterCollection::FriendlyNameFromParameterName(FString Param
 
 FString UNiagaraParameterCollection::ParameterNameFromFriendlyName(const FString& FriendlyName)const
 {
+	if (FriendlyName.StartsWith(GetFullNamespace()))
+	{
+		return FriendlyName;
+	}
 	return FString::Printf(TEXT("%s%s"), *GetFullNamespace(), *FriendlyName);
 }
 
diff --git a/Engine/Plugins/FX/Niagara/Source/NiagaraEditor/Private/Widgets/SNiagaraParameterCollection.cpp b/Engine/Plugins/FX/Niagara/Source/NiagaraEditor/Private/Widgets/SNiagaraParameterCollection.cpp
--- a/Engine/Plugins/FX/Niagara/Source/NiagaraEditor/Private/Widgets/SNiagaraParameterCollection.cpp
+++ b/Engine/Plugins/FX/Niagara/Source/NiagaraEditor/Private/Widgets/SNiagaraParameterCollection.cpp
@@ -254,7 +254,7 @@ void SNiagaraParameterCollection::Construct(const FArguments& InArgs, TSharedRef
 	ChildSlot
 	[
 		SAssignNew(ExpandableArea, SExpandableArea)
-		.InitiallyCollapsed(Collection->GetIsExpanded() == false)
+		.InitiallyCollapsed(false)
 		.BorderImage(FAppStyle::GetBrush("ToolPanel.DarkGroupBorder"))
 		.OnAreaExpansionChanged(this, &SNiagaraParameterCollection::AreaExpandedChanged)
 		.Padding(0)
@@ -359,7 +359,8 @@ void SNiagaraParameterCollection::ViewModelSelectionChanged()
 
 void SNiagaraParameterCollection::ViewModelIsExpandedChanged()
 {
-	ExpandableArea->SetExpanded(Collection->GetIsExpanded());
+	// we ignore the expansion state for the root entry, otherwise the user wouldn't be able to see anything in a collapsed state
+	ExpandableArea->SetExpanded(true);
 }
 
 void SNiagaraParameterCollection::AreaExpandedChanged(bool bIsExpanded)
@@ -620,12 +621,14 @@ TSharedRef<ITableRow> SNiagaraParameterCollection::OnGenerateRowForParameter(TSh
 					SNew(SHorizontalBox)
 					+ SHorizontalBox::Slot()
 					.AutoWidth()
+					.VAlign(VAlign_Center)
 					.Padding(FMargin(3, 0, 0, 0))
 					[
 						CustomValueEditor.ToSharedRef()
 					]
 					+ SHorizontalBox::Slot()
 					.HAlign(HAlign_Right)
+					.VAlign(VAlign_Center)
 					.Padding(4.f)
 					[
 						ParameterDeleteButton
