Fixed missing comment toggle on new compact style Niagara nodes

diff --git a/Engine/Plugins/FX/Niagara/Source/NiagaraEditor/Private/Widgets/SNiagaraGraphNode.cpp b/Engine/Plugins/FX/Niagara/Source/NiagaraEditor/Private/Widgets/SNiagaraGraphNode.cpp
--- a/Engine/Plugins/FX/Niagara/Source/NiagaraEditor/Private/Widgets/SNiagaraGraphNode.cpp
+++ b/Engine/Plugins/FX/Niagara/Source/NiagaraEditor/Private/Widgets/SNiagaraGraphNode.cpp
@@ -331,7 +331,8 @@ void SNiagaraGraphNode::UpdateGraphNodeCompact()
 	.AllowPinning( true )
 	.EnableTitleBarBubble( true )
 	.EnableBubbleCtrls( true )
-	.GraphLOD( this, &SGraphNode::GetCurrentLOD );
+	.GraphLOD( this, &SGraphNode::GetCurrentLOD )
+	.IsGraphNodeHovered(this, &SGraphNode::IsHovered);
 
 	GetOrAddSlot( ENodeZone::TopCenter )
 	.SlotOffset( TAttribute<FVector2D>( CommentBubble.Get(), &SCommentBubble::GetOffset ))
