- Modify sim cache buffer data to be packed into a single array
- Add optional BulkData serialization, controlled with bAllowSerializeLargeCache
- Large caches (BulkData) can grow beyond 2GB and are considerably faster to serialize in the editor
- Small caches will use TArray serialization and are not optimal.  These are useful when sending data as text across the network for debugging for example as BulkData serialization is not supported by all FArchieve types
- Clean up sim cache includes to reduce compiles
- Update existing SimCache tests

#virtualized

diff --git a/Engine/Plugins/FX/Niagara/Source/Niagara/Classes/NiagaraSimCache.h b/Engine/Plugins/FX/Niagara/Source/Niagara/Classes/NiagaraSimCache.h
--- a/Engine/Plugins/FX/Niagara/Source/Niagara/Classes/NiagaraSimCache.h
+++ b/Engine/Plugins/FX/Niagara/Source/Niagara/Classes/NiagaraSimCache.h
@@ -4,9 +4,14 @@
 
 #include "NiagaraCommon.h"
 #include "NiagaraScript.h"
+
+#include "Serialization/BulkData.h"
+#include "Serialization/BulkDataBuffer.h"
+
 #include "NiagaraSimCache.generated.h"
 
 class UNiagaraComponent;
+struct FNiagaraSimCacheDataBuffersLayout;
 
 UENUM(BlueprintType)
 enum class ENiagaraSimCacheAttributeCaptureMode : uint8
@@ -38,14 +43,6 @@ struct FNiagaraSimCacheCreateParameters
 {
 	GENERATED_BODY()
 
-	FNiagaraSimCacheCreateParameters()
-		: bAllowRebasing(true)
-		, bAllowDataInterfaceCaching(true)
-		, bAllowInterpolation(false)
-		, bAllowVelocityExtrapolation(false)
-	{
-	}
-
 	/**
 	How do we want to capture attributes for the simulation cache.
 	The mode selected depends on what situations the cache can be used in.
@@ -58,14 +55,14 @@ struct FNiagaraSimCacheCreateParameters
 	i.e. World space emitters can be moved to the new component's location
 	*/
 	UPROPERTY(EditAnywhere, BlueprintReadWrite, Category="SimCache")
-	uint32 bAllowRebasing : 1;
+	uint32 bAllowRebasing : 1 = true;
 
 	/**
 	When enabled Data Interface data will be stored in the SimCache.
 	This can result in a large increase to the cache size, depending on what Data Interfaces are used
 	*/
 	UPROPERTY(EditAnywhere, BlueprintReadWrite, Category="SimCache")
-	uint32 bAllowDataInterfaceCaching : 1;
+	uint32 bAllowDataInterfaceCaching : 1 = true;
 
 	/**
 	When enabled we allow the cache to be generated for interpolation.
@@ -73,14 +70,20 @@ struct FNiagaraSimCacheCreateParameters
 	By default we will capture and interpolate all Position & Quat types, you can adjust this using the include / exclude list.
 	*/
 	UPROPERTY(EditAnywhere, BlueprintReadWrite, Category="SimCache")
-	uint32 bAllowInterpolation : 1;
+	uint32 bAllowInterpolation : 1 = false;
 
 	/**
 	When enabled we allow the cache to be generated for extrapolation.
 	This will force the velocity attribute to be maintained.
 	*/
 	UPROPERTY(EditAnywhere, BlueprintReadWrite, Category="SimCache")
-	uint32 bAllowVelocityExtrapolation : 1;
+	uint32 bAllowVelocityExtrapolation : 1 = true;
+
+	/**
+	When enabled the cache will support serializing large amounts of cache data.
+	*/
+	UPROPERTY(EditAnywhere, BlueprintReadWrite, Category = "SimCache")
+	uint32 bAllowSerializeLargeCache : 1 = true;
 
 	/**
 	List of Attributes to force include in the SimCache rebase, they should be the full path to the attribute
@@ -127,23 +130,26 @@ struct FNiagaraSimCacheDataBuffers
 	uint32 NumInstances = 0;
 
 	UPROPERTY()
-	TArray<uint8> FloatData;
+	uint32 IDAcquireTag = 0;
 
 	UPROPERTY()
-	TArray<uint8> HalfData;
+	uint32 IDToIndexTableElements = 0;
 
-	UPROPERTY()
-	TArray<uint8> Int32Data;
+	TArrayView<uint8> FloatData;
+	TArrayView<uint8> HalfData;
+	TArrayView<uint8> Int32Data;
+	TArrayView<int32> IDToIndexTable;
+	TArrayView<uint32> InterpMapping;
 
-	UPROPERTY()
-	TArray<int32> IDToIndexTable;
+	FBulkDataBuffer<uint8> DataBuffer;
+	FByteBulkData BulkData;
 
-	UPROPERTY()
-	uint32 IDAcquireTag = 0;
+	void SetupForWrite(const FNiagaraSimCacheDataBuffersLayout& CacheLayout);
+	bool SetupForRead(const FNiagaraSimCacheDataBuffersLayout& CacheLayout);
+
+	bool SerializeAsArray(FArchive& Ar, UObject* OwnerObject, const FNiagaraSimCacheDataBuffersLayout& CacheLayout);
+	bool SerializeAsBulkData(FArchive& Ar, UObject* OwnerObject, const FNiagaraSimCacheDataBuffersLayout& CacheLayout);
 
-	UPROPERTY()
-	TArray<uint32> InterpMapping;
-	
 	bool operator==(const FNiagaraSimCacheDataBuffers& Other) const;
 	bool operator!=(const FNiagaraSimCacheDataBuffers& Other) const;
 };
@@ -378,6 +384,7 @@ public:
 	static constexpr float CacheAgeResolution = 10000.0f;
 
 	// UObject Interface
+	NIAGARA_API virtual void Serialize(FArchive& Ar) override;
 	NIAGARA_API virtual bool IsReadyForFinishDestroy() override;
 	// UObject Interface
 
diff --git a/Engine/Plugins/FX/Niagara/Source/Niagara/Private/NiagaraDataInterfaceRenderTargetVolume.cpp b/Engine/Plugins/FX/Niagara/Source/Niagara/Private/NiagaraDataInterfaceRenderTargetVolume.cpp
--- a/Engine/Plugins/FX/Niagara/Source/Niagara/Private/NiagaraDataInterfaceRenderTargetVolume.cpp
+++ b/Engine/Plugins/FX/Niagara/Source/Niagara/Private/NiagaraDataInterfaceRenderTargetVolume.cpp
@@ -12,6 +12,7 @@
 #include "NiagaraDataInterfaceRenderTargetCommon.h"
 #include "NiagaraGpuComputeDispatchInterface.h"
 #include "NiagaraSettings.h"
+#include "NiagaraSimCache.h"
 #include "NiagaraSVTShaders.h"
 #include "NiagaraSystemInstance.h"
 #include "NiagaraStats.h"
diff --git a/Engine/Plugins/FX/Niagara/Source/Niagara/Private/NiagaraDebuggerClient.cpp b/Engine/Plugins/FX/Niagara/Source/Niagara/Private/NiagaraDebuggerClient.cpp
--- a/Engine/Plugins/FX/Niagara/Source/Niagara/Private/NiagaraDebuggerClient.cpp
+++ b/Engine/Plugins/FX/Niagara/Source/Niagara/Private/NiagaraDebuggerClient.cpp
@@ -521,8 +521,11 @@ bool FNiagaraSimCacheCaptureInfo::Process()
 		//First Frame of Capture, init the sim cache.
 		if (ProcessedFrames == Request.CaptureDelayFrames)
 		{	
+			FNiagaraSimCacheCreateParameters CacheCreateParameters;
+			CacheCreateParameters.bAllowSerializeLargeCache = false;
+
 			SimCache.Reset(NewObject<UNiagaraSimCache>(GetTransientPackage()));
-			SimCache->BeginWrite(FNiagaraSimCacheCreateParameters(), Comp);
+			SimCache->BeginWrite(CacheCreateParameters, Comp);
 		}
 
 		SimCache->WriteFrame(Comp);
diff --git a/Engine/Plugins/FX/Niagara/Source/Niagara/Private/NiagaraSimCache.cpp b/Engine/Plugins/FX/Niagara/Source/Niagara/Private/NiagaraSimCache.cpp
--- a/Engine/Plugins/FX/Niagara/Source/Niagara/Private/NiagaraSimCache.cpp
+++ b/Engine/Plugins/FX/Niagara/Source/Niagara/Private/NiagaraSimCache.cpp
@@ -6,6 +6,7 @@
 #include "Misc/LargeWorldRenderPosition.h"
 #include "NiagaraConstants.h"
 #include "NiagaraComponent.h"
+#include "NiagaraCustomVersion.h"
 #include "NiagaraDataInterfaceUtilities.h"
 #include "NiagaraDataSetReadback.h"
 #include "NiagaraEmitterInstance.h"
@@ -21,6 +22,8 @@
 UNiagaraSimCache::FOnCacheBeginWrite	UNiagaraSimCache::OnCacheBeginWrite;
 UNiagaraSimCache::FOnCacheEndWrite		UNiagaraSimCache::OnCacheEndWrite;
 
+////////////////////////////////////////////////////////////////////////////////////////////////////////////////////////////////////////////////////
+
 FNiagaraSimCacheFeedbackContext::~FNiagaraSimCacheFeedbackContext()
 {
 	if (bAutoLogIssues)
@@ -37,6 +40,9 @@ FNiagaraSimCacheFeedbackContext::~FNiagaraSimCacheFeedbackContext()
 	}
 }
 
+
+////////////////////////////////////////////////////////////////////////////////////////////////////////////////////////////////////////////////////
+
 int32 FNiagaraSimCacheDataBuffersLayout::IndexOfCacheVariable(const FNiagaraVariableBase& InVariable) const
 {
 	return Variables.IndexOfByPredicate(
@@ -53,10 +59,215 @@ const FNiagaraSimCacheVariable* FNiagaraSimCacheDataBuffersLayout::FindCacheVari
 	return Index != INDEX_NONE ? &Variables[Index] : nullptr;
 }
 
+////////////////////////////////////////////////////////////////////////////////////////////////////////////////////////////////////////////////////
+
+struct FNiagaraSimCacheBuffersSetup
+{
+	explicit FNiagaraSimCacheBuffersSetup(const FNiagaraSimCacheDataBuffersLayout& CacheLayout, const FNiagaraSimCacheDataBuffers& DataBuffer)
+	{
+		FloatDataNum			= CacheLayout.FloatCount * DataBuffer.NumInstances;
+		HalfDataNum				= CacheLayout.HalfCount * DataBuffer.NumInstances;
+		Int32DataNum			= CacheLayout.Int32Count * DataBuffer.NumInstances;
+		IDToIndexTableNum		= DataBuffer.IDToIndexTableElements;
+		InterpMappingNum		= CacheLayout.bAllowInterpolation ? DataBuffer.NumInstances : 0;
+
+		FloatDataOffset			= 0;
+		HalfDataOffset			= Align(FloatDataOffset + (FloatDataNum * sizeof(float)), sizeof(FFloat16));
+		Int32DataOffset			= Align(HalfDataOffset + (HalfDataNum * sizeof(FFloat16)), sizeof(int32));
+		IDToIndexTableOffset	= Align(Int32DataOffset + (Int32DataNum * sizeof(int32)), sizeof(int32));
+		InterpMappingOffset		= Align(IDToIndexTableOffset + (IDToIndexTableNum * sizeof(int32)), sizeof(uint32));
+		BufferSize				= InterpMappingOffset + (InterpMappingNum * sizeof(int32));
+	}
+
+	bool SetArrayViews(FNiagaraSimCacheDataBuffers& DataBuffer) const
+	{
+		const TArrayView64<uint8> DataBufferView = DataBuffer.DataBuffer.GetView();
+		if (DataBufferView.Num() != BufferSize)
+		{
+			DataBuffer.NumInstances		= 0;
+			DataBuffer.FloatData		= TArrayView<uint8>();
+			DataBuffer.HalfData			= TArrayView<uint8>();
+			DataBuffer.Int32Data		= TArrayView<uint8>();
+			DataBuffer.IDToIndexTable	= TArrayView<int32>();
+			DataBuffer.InterpMapping	= TArrayView<uint32>();
+			return false;
+		}
+		else
+		{
+			uint8* BaseData				= DataBufferView.GetData();
+			DataBuffer.FloatData		= TArrayView<uint8>(BaseData + FloatDataOffset, FloatDataNum * sizeof(float));
+			DataBuffer.HalfData			= TArrayView<uint8>(BaseData + HalfDataOffset, HalfDataNum * sizeof(FFloat16));
+			DataBuffer.Int32Data		= TArrayView<uint8>(BaseData + Int32DataOffset, Int32DataNum * sizeof(int32));
+			DataBuffer.IDToIndexTable	= TArrayView<int32>(reinterpret_cast<int32*>(BaseData + IDToIndexTableOffset), IDToIndexTableNum);
+			DataBuffer.InterpMapping	= TArrayView<uint32>(reinterpret_cast<uint32*>(BaseData + InterpMappingOffset), InterpMappingNum);
+			return true;
+		}
+	}
+
+	int64 FloatDataNum = 0;
+	int64 HalfDataNum = 0;
+	int64 Int32DataNum = 0;
+	int64 IDToIndexTableNum = 0;
+	int64 InterpMappingNum = 0;
+
+	int64 FloatDataOffset = 0;
+	int64 HalfDataOffset = 0;
+	int64 Int32DataOffset = 0;
+	int64 IDToIndexTableOffset = 0;
+	int64 InterpMappingOffset = 0;
+	int64 BufferSize = 0;
+};
+
+void FNiagaraSimCacheDataBuffers::SetupForWrite(const FNiagaraSimCacheDataBuffersLayout& CacheLayout)
+{
+	const FNiagaraSimCacheBuffersSetup BufferSetup(CacheLayout, *this);
+
+	BulkData.RemoveBulkData();
+	DataBuffer.Reset(static_cast<uint8*>(FMemory::Malloc(BufferSetup.BufferSize)), BufferSetup.BufferSize);
+	const bool bSuccess = BufferSetup.SetArrayViews(*this);
+	checkf(bSuccess == true, TEXT("Cache Write setup can not fail."));
+}
+
+bool FNiagaraSimCacheDataBuffers::SetupForRead(const FNiagaraSimCacheDataBuffersLayout& CacheLayout)
+{
+	const FNiagaraSimCacheBuffersSetup BufferSetup(CacheLayout, *this);
+
+	if (DataBuffer.GetView().Num() != BufferSetup.BufferSize)
+	{
+		DataBuffer = BulkData.GetCopyAsBuffer(0, true);
+		BulkData.UnloadBulkData();
+	}
+
+	return BufferSetup.SetArrayViews(*this);
+}
+
+bool FNiagaraSimCacheDataBuffers::SerializeAsArray(FArchive& Ar, UObject* OwnerObject, const FNiagaraSimCacheDataBuffersLayout& CacheLayout)
+{
+	TArray<uint8> TransientArray;
+	TransientArray = DataBuffer.GetView();
+
+	Ar << TransientArray;
+
+	if (Ar.IsLoading())
+	{
+		uint8* DataBufferMemory = nullptr;
+		if (TransientArray.Num() > 0)
+		{
+			DataBufferMemory = static_cast<uint8*>(FMemory::Malloc(TransientArray.Num()));
+			FMemory::Memcpy(DataBufferMemory, TransientArray.GetData(), TransientArray.Num());
+		}
+		DataBuffer.Reset(DataBufferMemory, TransientArray.Num());
+	}
+
+	return SetupForRead(CacheLayout);
+}
+
+bool FNiagaraSimCacheDataBuffers::SerializeAsBulkData(FArchive& Ar, UObject* OwnerObject, const FNiagaraSimCacheDataBuffersLayout& CacheLayout)
+{
+	if (Ar.IsSaving())
+	{
+		TArrayView64<uint8> DataBufferView = DataBuffer.GetView();
+		if (DataBufferView.Num() > 0)
+		{
+			BulkData.Lock(LOCK_READ_WRITE);
+			{
+				uint8* NewBulkData = BulkData.Realloc(DataBufferView.Num());
+				FMemory::Memcpy(NewBulkData, DataBufferView.GetData(), DataBufferView.Num());
+			}
+			BulkData.Unlock();
+		}
+		else
+		{
+			BulkData.RemoveBulkData();
+		}
+	}
+
+	BulkData.Serialize(Ar, OwnerObject);
+
+	return SetupForRead(CacheLayout);
+}
+
+bool FNiagaraSimCacheDataBuffers::operator==(const FNiagaraSimCacheDataBuffers& Other) const
+{
+	bool bMatches;
+	bMatches = NumInstances == Other.NumInstances;
+	bMatches &= IDAcquireTag == Other.IDAcquireTag;
+	bMatches &= FloatData.Num() == Other.FloatData.Num();
+	bMatches &= HalfData.Num() == Other.HalfData.Num();
+	bMatches &= Int32Data.Num() == Other.Int32Data.Num();
+	bMatches &= IDToIndexTable.Num() == Other.IDToIndexTable.Num();
+	bMatches &= InterpMapping.Num() == Other.InterpMapping.Num();
+
+	TConstArrayView<uint8> ThisDataBuffer = DataBuffer.GetView();
+	TConstArrayView<uint8> OtherDataBuffer = Other.DataBuffer.GetView();
+	bMatches &= ThisDataBuffer.Num() == OtherDataBuffer.Num();
+	if (bMatches)
+	{
+		bMatches &= FMemory::Memcmp(ThisDataBuffer.GetData(), OtherDataBuffer.GetData(), ThisDataBuffer.Num()) != 0;
+	}
+
+	return bMatches;
+}
+
+bool FNiagaraSimCacheDataBuffers::operator!=(const FNiagaraSimCacheDataBuffers& Other) const
+{
+	return !(*this == Other);
+}
+
+////////////////////////////////////////////////////////////////////////////////////////////////////////////////////////////////////////////////////
+
 UNiagaraSimCache::UNiagaraSimCache(const FObjectInitializer& ObjectInitializer)
 {
 }
 
+void UNiagaraSimCache::Serialize(FArchive& Ar)
+{
+	Ar.UsingCustomVersion(FNiagaraCustomVersion::GUID);
+
+	Super::Serialize(Ar);
+
+	// Early out if we aren't using bulk data, these caches will no longer be supported
+	const int32 NiagaraVersion = Ar.CustomVer(FNiagaraCustomVersion::GUID);
+	if (NiagaraVersion < FNiagaraCustomVersion::SimCache_BulkDataVersion1)
+	{
+		UE_LOG(LogNiagara, Warning, TEXT("SimCache needs to be regenerated as this version is no longer supported."));
+		SoftNiagaraSystem.Reset();
+		return;
+	}
+
+	bool bIsCacheValid = true;
+	if ( CreateParameters.bAllowSerializeLargeCache )
+	{
+		for (FNiagaraSimCacheFrame& CacheFrame : CacheFrames)
+		{
+			bIsCacheValid &= CacheFrame.SystemData.SystemDataBuffers.SerializeAsBulkData(Ar, this, CacheLayout.SystemLayout);
+			for (int32 iEmitter = 0; iEmitter < CacheLayout.EmitterLayouts.Num(); ++iEmitter)
+			{
+				FNiagaraSimCacheEmitterFrame& EmitterData = CacheFrame.EmitterData[iEmitter];
+				bIsCacheValid &= EmitterData.ParticleDataBuffers.SerializeAsBulkData(Ar, this, CacheLayout.EmitterLayouts[iEmitter]);
+			}
+		}
+	}
+	else
+	{
+		for (FNiagaraSimCacheFrame& CacheFrame : CacheFrames)
+		{
+			bIsCacheValid &= CacheFrame.SystemData.SystemDataBuffers.SerializeAsArray(Ar, this, CacheLayout.SystemLayout);
+			for (int32 iEmitter = 0; iEmitter < CacheLayout.EmitterLayouts.Num(); ++iEmitter)
+			{
+				FNiagaraSimCacheEmitterFrame& EmitterData = CacheFrame.EmitterData[iEmitter];
+				bIsCacheValid &= EmitterData.ParticleDataBuffers.SerializeAsArray(Ar, this, CacheLayout.EmitterLayouts[iEmitter]);
+			}
+		}
+	}
+
+	if (!bIsCacheValid)
+	{
+		UE_LOG(LogNiagara, Warning, TEXT("SimCache buffer serialization failed, likely due to bulk data not being supported, cache is invalid."));
+		SoftNiagaraSystem.Reset();
+	}
+}
+
 bool UNiagaraSimCache::IsReadyForFinishDestroy()
 {
 	return PendingCommandsInFlight == 0;
@@ -738,7 +949,6 @@ bool UNiagaraSimCache::ReadFrame(int32 FrameIndex, float FrameFraction, FNiagara
 		RebaseTransform = RebaseTransform * CacheFrame.LocalToWorld.Inverse();
 	}
 
-
 	const int32 NextFrameIndex = FMath::Min(FrameIndex + 1, CacheFrames.Num() - 1);
 	const float FrameDeltaSeconds = CacheFrames[NextFrameIndex].SimulationAge - CacheFrame.SimulationAge;
 	const float SimDeltaSeconds = Helper.SystemInstance->CachedDeltaSeconds;
@@ -1111,23 +1321,6 @@ void UNiagaraSimCache::ReadQuatAttributeWithRebase(TArray<FQuat>& OutValues, FQu
 	}
 }
 
-bool FNiagaraSimCacheDataBuffers::operator==(const FNiagaraSimCacheDataBuffers& Other) const
-{
-	bool InstanceCountSame = NumInstances == Other.NumInstances;
-	bool SameAquireTag = IDAcquireTag == Other.IDAcquireTag;
-	bool FloatDataSame = FloatData == Other.FloatData;
-	bool HalfDataSame = HalfData == Other.HalfData;
-	bool IntDataSame = Int32Data == Other.Int32Data;
-	bool IndexTableSame = IDToIndexTable == Other.IDToIndexTable;
-	bool MappingSame = InterpMapping == Other.InterpMapping;
-	return InstanceCountSame && SameAquireTag && FloatDataSame && HalfDataSame
-		&& IntDataSame && IndexTableSame && MappingSame;
-}
-bool FNiagaraSimCacheDataBuffers::operator!=(const FNiagaraSimCacheDataBuffers& Other) const
-{
-	return !(*this == Other);
-}
-
 namespace CacheCompare
 {
 	constexpr int32 ErrorStrLen = 8000;
@@ -1236,16 +1429,35 @@ bool UNiagaraSimCache::IsDataEqual(const UNiagaraSimCache& OtherCache, float Err
 			CacheCompare::AddError(Errors, TEXT("SystemDataBuffers.NumInstances different - frame ") + Frame + TEXT(": expected ") + FString::FromInt(ExpectedFrame.SystemData.SystemDataBuffers.NumInstances) + TEXT(", got ") + FString::FromInt(OtherFrame.SystemData.SystemDataBuffers.NumInstances) + '\n');
 			bEqual = false;
 		}
-		if (OtherFrame.SystemData.SystemDataBuffers.Int32Data != ExpectedFrame.SystemData.SystemDataBuffers.Int32Data)
+
+		if (OtherFrame.SystemData.SystemDataBuffers.Int32Data.Num() == ExpectedFrame.SystemData.SystemDataBuffers.Int32Data.Num())
+		{
+			if (FMemory::Memcmp(OtherFrame.SystemData.SystemDataBuffers.Int32Data.GetData(), ExpectedFrame.SystemData.SystemDataBuffers.Int32Data.GetData(), ExpectedFrame.SystemData.SystemDataBuffers.Int32Data.Num()) != 0)
+			{
+				CacheCompare::AddError(Errors, TEXT("SystemDataBuffers.Int32Data different - frame ") + Frame + '\n');
+				bEqual = false;
+			}
+		}
+		else
 		{
-			CacheCompare::AddError(Errors, TEXT("SystemDataBuffers.Int32Data different - frame ") + Frame + '\n');
+			CacheCompare::AddError(Errors, TEXT("SystemDataBuffers.Int32Data count different - frame ") + Frame + '\n');
 			bEqual = false;
 		}
-		if (OtherFrame.SystemData.SystemDataBuffers.IDToIndexTable != ExpectedFrame.SystemData.SystemDataBuffers.IDToIndexTable)
+
+		if (OtherFrame.SystemData.SystemDataBuffers.IDToIndexTable.Num() == ExpectedFrame.SystemData.SystemDataBuffers.IDToIndexTable.Num())
 		{
-			CacheCompare::AddError(Errors, TEXT("SystemDataBuffers.IDToIndexTable different - frame ") + Frame + '\n');
+			if (FMemory::Memcmp(OtherFrame.SystemData.SystemDataBuffers.IDToIndexTable.GetData(), ExpectedFrame.SystemData.SystemDataBuffers.IDToIndexTable.GetData(), ExpectedFrame.SystemData.SystemDataBuffers.IDToIndexTable.Num()) != 0)
+			{
+				CacheCompare::AddError(Errors, TEXT("SystemDataBuffers.IDToIndexTable different - frame ") + Frame + '\n');
+				bEqual = false;
+			}
+		}
+		else
+		{
+			CacheCompare::AddError(Errors, TEXT("SystemDataBuffers.IDToIndexTable count different - frame ") + Frame + '\n');
 			bEqual = false;
 		}
+
 		if (OtherFrame.SystemData.SystemDataBuffers.FloatData.Num() == ExpectedFrame.SystemData.SystemDataBuffers.FloatData.Num())
 		{
 			for (int i = 0; i < OtherFrame.SystemData.SystemDataBuffers.FloatData.Num(); i += sizeof(float))
@@ -1265,6 +1477,7 @@ bool UNiagaraSimCache::IsDataEqual(const UNiagaraSimCache& OtherCache, float Err
 			CacheCompare::AddError(Errors, TEXT("SystemDataBuffers.FloatData count different - frame ") + Frame + '\n');
 			bEqual = false;
 		}
+
 		if (OtherFrame.SystemData.SystemDataBuffers.HalfData.Num() == ExpectedFrame.SystemData.SystemDataBuffers.HalfData.Num())
 		{
 			for (int i = 0; i < OtherFrame.SystemData.SystemDataBuffers.HalfData.Num(); i += sizeof(FFloat16))
diff --git a/Engine/Plugins/FX/Niagara/Source/Niagara/Private/NiagaraSimCacheHelper.h b/Engine/Plugins/FX/Niagara/Source/Niagara/Private/NiagaraSimCacheHelper.h
--- a/Engine/Plugins/FX/Niagara/Source/Niagara/Private/NiagaraSimCacheHelper.h
+++ b/Engine/Plugins/FX/Niagara/Source/Niagara/Private/NiagaraSimCacheHelper.h
@@ -624,53 +624,52 @@ struct FNiagaraSimCacheHelper
 		}
 
 		CacheBuffer.NumInstances = NumInstances;
+		CacheBuffer.IDToIndexTableElements = DataBuffer.GetIDTable().Num();
+		CacheBuffer.SetupForWrite(CacheLayout);
 
 		int32 iComponent = 0;
 
 		// Copy Float
-		CacheBuffer.FloatData.AddDefaulted(CacheLayout.FloatCount * NumInstances * sizeof(float));
 		for ( uint32 i=0; i < CacheLayout.FloatCount; ++i )
 		{
 			const uint32 Component = CacheLayout.CacheBufferWriteInfo.ComponentMappingsFromDataBuffer[iComponent++];
 			const uint8* Source = DataBuffer.GetComponentPtrFloat(Component) + (FirstInstance * sizeof(float));
 			uint8* Dest = CacheBuffer.FloatData.GetData() + (i * NumInstances * sizeof(float));
 			CheckedMemcpy(CacheBuffer.FloatData, Dest, DataBuffer.GetFloatBuffer(), Source, sizeof(float) * NumInstances);
-			//FMemory::Memcpy(Dest, Source, sizeof(float) * NumInstances);
 		}
 		
 		// Copy Half
-		CacheBuffer.HalfData.AddDefaulted(CacheLayout.HalfCount * NumInstances * sizeof(FFloat16));
 		for (uint32 i = 0; i < CacheLayout.HalfCount; ++i)
 		{
 			const uint32 Component = CacheLayout.CacheBufferWriteInfo.ComponentMappingsFromDataBuffer[iComponent++];
 			const uint8* Source = DataBuffer.GetComponentPtrHalf(Component) + (FirstInstance * sizeof(FFloat16));
 			uint8* Dest = CacheBuffer.HalfData.GetData() + (i * NumInstances * sizeof(FFloat16));
 			CheckedMemcpy(CacheBuffer.HalfData, Dest, DataBuffer.GetHalfBuffer(), Source, sizeof(FFloat16) * NumInstances);
-			//FMemory::Memcpy(Dest, Source, sizeof(FFloat16) * NumInstances);
 		}
 
 		// Copy Int32
-		CacheBuffer.Int32Data.AddDefaulted(CacheLayout.Int32Count * NumInstances * sizeof(int32));
 		for (uint32 i = 0; i < CacheLayout.Int32Count; ++i)
 		{
 			const uint32 Component = CacheLayout.CacheBufferWriteInfo.ComponentMappingsFromDataBuffer[iComponent++];
 			const uint8* Source = DataBuffer.GetComponentPtrInt32(Component) + (FirstInstance * sizeof(int32));
 			uint8* Dest = CacheBuffer.Int32Data.GetData() + (i * NumInstances * sizeof(int32));
 			CheckedMemcpy(CacheBuffer.Int32Data, Dest, DataBuffer.GetInt32Buffer(), Source, sizeof(int32) * NumInstances);
-			//FMemory::Memcpy(Dest, Source, sizeof(int32) * NumInstances);
 		}
 
 		// Copy ID to Index Table
-		CacheBuffer.IDToIndexTable = DataBuffer.GetIDTable();
+		if (CacheBuffer.IDToIndexTable.Num() > 0)
+		{
+			check(CacheBuffer.IDToIndexTable.Num() == DataBuffer.GetIDTable().Num());
+			FMemory::Memcpy(CacheBuffer.IDToIndexTable.GetData(), DataBuffer.GetIDTable().GetData(), CacheBuffer.IDToIndexTable.Num() * sizeof(int32));
+		}
 		CacheBuffer.IDAcquireTag = DataBuffer.GetIDAcquireTag();
 
 		// Generate a interp mapping (if we have enabled it)
 		if (CacheLayout.bAllowInterpolation)
 		{
-			//-TODO: Persistent ID mapping
-			CacheBuffer.InterpMapping.SetNumUninitialized(DataBuffer.GetNumInstances());
+			check(NumInstances == DataBuffer.GetNumInstances());
 			const uint8* UniqueIDs = DataBuffer.GetComponentPtrInt32(CacheLayout.CacheBufferWriteInfo.ComponentUniqueID);
-			FMemory::Memcpy(CacheBuffer.InterpMapping.GetData(), UniqueIDs, DataBuffer.GetNumInstances() * sizeof(int32));
+			FMemory::Memcpy(CacheBuffer.InterpMapping.GetData(), UniqueIDs, NumInstances * sizeof(int32));
 		}
 	}
 
diff --git a/Engine/Plugins/FX/Niagara/Source/Niagara/Public/NiagaraDebuggerCommon.h b/Engine/Plugins/FX/Niagara/Source/Niagara/Public/NiagaraDebuggerCommon.h
--- a/Engine/Plugins/FX/Niagara/Source/Niagara/Public/NiagaraDebuggerCommon.h
+++ b/Engine/Plugins/FX/Niagara/Source/Niagara/Public/NiagaraDebuggerCommon.h
@@ -10,7 +10,6 @@ All common code shared between the editor side debugger and debugger clients run
 #include "Misc/NotifyHook.h"
 #include "NiagaraTypes.h"
 #include "NiagaraCommon.h"
-#include "NiagaraSimCache.h"
 #include "NiagaraDebuggerCommon.generated.h"
 
 //////////////////////////////////////////////////////////////////////////
diff --git a/Engine/Plugins/FX/Niagara/Source/NiagaraCore/Public/NiagaraCustomVersion.h b/Engine/Plugins/FX/Niagara/Source/NiagaraCore/Public/NiagaraCustomVersion.h
--- a/Engine/Plugins/FX/Niagara/Source/NiagaraCore/Public/NiagaraCustomVersion.h
+++ b/Engine/Plugins/FX/Niagara/Source/NiagaraCore/Public/NiagaraCustomVersion.h
@@ -197,6 +197,8 @@ struct FNiagaraCustomVersion
 
 		ParameterBindingWithValueRenameFixup, // FNiagaraParameterBindingWithValue wouldn't necessarily have the appropriate ResolvedParameter namespace when it comes to emitter merging
 
+		SimCache_BulkDataVersion1, // Sim Cache moved to bulk data by default
+
 		// DO NOT ADD A NEW VERSION UNLESS YOU HAVE TALKED TO THE NIAGARA LEAD. Mismanagement of these versions can lead to data loss if it is adjusted in multiple streams simultaneously.
 		// -----<new versions can be added above this line>  -------------------------------------------------
 		VersionPlusOne,
