- Add missing include

diff --git a/Engine/Plugins/FX/Niagara/Source/Niagara/Private/DataInterface/NiagaraDataInterfaceDataChannelWrite.cpp b/Engine/Plugins/FX/Niagara/Source/Niagara/Private/DataInterface/NiagaraDataInterfaceDataChannelWrite.cpp
--- a/Engine/Plugins/FX/Niagara/Source/Niagara/Private/DataInterface/NiagaraDataInterfaceDataChannelWrite.cpp
+++ b/Engine/Plugins/FX/Niagara/Source/Niagara/Private/DataInterface/NiagaraDataInterfaceDataChannelWrite.cpp
@@ -7,11 +7,11 @@
 #include "NiagaraCommon.h"
 #include "NiagaraShared.h"
 
+#include "NiagaraSimCache.h"
 #include "NiagaraSystem.h"
 #include "NiagaraWorldManager.h"
 #include "NiagaraSystemInstance.h"
 
-
 #include "NiagaraDataChannel.h"
 #include "NiagaraDataChannelHandler.h"
 #include "NiagaraDataChannelManager.h"
