- Disable Velocity Extrapolation by default

diff --git a/Engine/Plugins/FX/Niagara/Source/Niagara/Classes/NiagaraSimCache.h b/Engine/Plugins/FX/Niagara/Source/Niagara/Classes/NiagaraSimCache.h
--- a/Engine/Plugins/FX/Niagara/Source/Niagara/Classes/NiagaraSimCache.h
+++ b/Engine/Plugins/FX/Niagara/Source/Niagara/Classes/NiagaraSimCache.h
@@ -77,7 +77,7 @@ struct FNiagaraSimCacheCreateParameters
 	This will force the velocity attribute to be maintained.
 	*/
 	UPROPERTY(EditAnywhere, BlueprintReadWrite, Category="SimCache")
-	uint32 bAllowVelocityExtrapolation : 1 = true;
+	uint32 bAllowVelocityExtrapolation : 1 = false;
 
 	/**
 	When enabled the cache will support serializing large amounts of cache data.
