Partial rollback of 29965558 to investigate warnings in Niagara
Original CL Desc
-------------------------------------------------------------------------------
[FScriptSetHelper and FScriptMapHelper iterators]
fixes for potential issues

diff --git a/Engine/Plugins/FX/Niagara/Source/NiagaraEditor/Private/NiagaraEditorUtilities.cpp b/Engine/Plugins/FX/Niagara/Source/NiagaraEditor/Private/NiagaraEditorUtilities.cpp
--- a/Engine/Plugins/FX/Niagara/Source/NiagaraEditor/Private/NiagaraEditorUtilities.cpp
+++ b/Engine/Plugins/FX/Niagara/Source/NiagaraEditor/Private/NiagaraEditorUtilities.cpp
@@ -509,22 +509,23 @@ bool FNiagaraEditorUtilities::NestedPropertiesAppendCompileHash(const void* Cont
 		{
 			continue;
 		}
-
-		if (Property->IsA(FStructProperty::StaticClass()))
+		else if (Property->IsA(FStructProperty::StaticClass()))
 		{
 			FStructProperty* StructProp = CastFieldChecked<FStructProperty>(Property);
 			const void* StructContainer = Property->ContainerPtrToValuePtr<uint8>(Container);
 			NestedPropertiesAppendCompileHash(StructContainer, StructProp->Struct, EFieldIteratorFlags::IncludeSuper, *PropertyName, InVisitor);
+			continue;
 		}
 		else if (Property->IsA(FEnumProperty::StaticClass()))
 		{
 			FEnumProperty* CastProp = CastFieldChecked<FEnumProperty>(Property);
 			const void* EnumContainer = Property->ContainerPtrToValuePtr<uint8>(Container);
-			if (!PODPropertyAppendCompileHash(EnumContainer, CastProp->GetUnderlyingProperty(), *PropertyName, InVisitor))
+			if (PODPropertyAppendCompileHash(EnumContainer, CastProp->GetUnderlyingProperty(), *PropertyName, InVisitor))
 			{
-				check(false);
-				return false;
+				continue;
 			}
+			check(false);
+			return false;
 		}
 		else if (Property->IsA(FObjectProperty::StaticClass()))
 		{
@@ -551,6 +552,7 @@ bool FNiagaraEditorUtilities::NestedPropertiesAppendCompileHash(const void* Cont
 					InVisitor->UpdateString(*PropertyName, TEXT("nullptr"));
 				}
 			}
+			continue;
 		}
 		else if (Property->IsA(FMapProperty::StaticClass()))
 		{
@@ -571,11 +573,19 @@ bool FNiagaraEditorUtilities::NestedPropertiesAppendCompileHash(const void* Cont
 				{
 					// To be safe, let's gather up all the keys and sort them lexicographically so that this is stable across application runs.
 					TArray<FName> Names;
-					Names.Reserve(MapHelper.Num());
-					for (FScriptMapHelper::FIterator It(MapHelper); It; ++It)
+					Names.AddUninitialized(MapHelper.Num());
+					for (int32 i = 0; i < MapHelper.Num(); i++)
 					{
-						const FName* KeyPtr = (FName*)MapHelper.GetKeyPtr(It);
-						Names.Add(*KeyPtr);
+						FName* KeyPtr = (FName*)(MapHelper.GetKeyPtr(i));
+						if (KeyPtr)
+						{
+							Names[i] = *KeyPtr;
+						}
+						else
+						{
+							Names[i] = FName();
+							UE_LOG(LogNiagaraEditor, Warning, TEXT("Bad key in %s at %d"), *Property->GetName(), i);
+						}
 					}
 					// Sort stably over runs
 					Names.Sort(FNameLexicalLess());
@@ -584,37 +594,41 @@ bool FNiagaraEditorUtilities::NestedPropertiesAppendCompileHash(const void* Cont
 					// We support map values of POD types or map values of structs with POD types internally. Anything else we should generate a warning on.
 					if (MapHelper.GetValueProperty()->IsA(FStructProperty::StaticClass()))
 					{
-						bool bErrorReportedOnce = false;
+						bool bPassed = true;
 						FStructProperty* StructProp = CastFieldChecked<FStructProperty>(MapHelper.GetValueProperty());
 
-						for (FScriptMapHelper::FIterator It(MapHelper); It; ++It)
+						for (int32 ArrayIdx = 0; ArrayIdx < MapHelper.Num(); ArrayIdx++)
 						{
-							InVisitor->UpdateString(*FString::Printf(TEXT("Key[%d]"), It.GetLogicalIndex()), Names[It.GetLogicalIndex()].ToString());
-							if (!NestedPropertiesAppendCompileHash(MapHelper.GetValuePtr(It), StructProp->Struct, EFieldIteratorFlags::IncludeSuper, FString::Printf(TEXT("Value[%d]"), It.GetLogicalIndex()), InVisitor))
+							InVisitor->UpdateString(*FString::Printf(TEXT("Key[%d]"), ArrayIdx), Names[ArrayIdx].ToString());
+							if (!NestedPropertiesAppendCompileHash(MapHelper.GetValuePtr(ArrayIdx), StructProp->Struct, EFieldIteratorFlags::IncludeSuper, FString::Printf(TEXT("Value[%d]"), ArrayIdx), InVisitor))
 							{
-								if (!bErrorReportedOnce)
-								{
-									UE_LOG(LogNiagaraEditor, Warning, TEXT("Skipping %s because it is an map value property of unsupported underlying type, please add \"meta = (SkipForCompileHash=\"true\")\" to avoid this warning in the future or handle it yourself in NestedPropertiesAppendCompileHash!"), *Property->GetName());
-									bErrorReportedOnce = true;
-								}
+								UE_LOG(LogNiagaraEditor, Warning, TEXT("Skipping %s because it is an map value property of unsupported underlying type, please add \"meta = (SkipForCompileHash=\"true\")\" to avoid this warning in the future or handle it yourself in NestedPropertiesAppendCompileHash!"), *Property->GetName());
+								bPassed = false;
+								continue;
 							}
 						}
+						if (bPassed)
+						{
+							continue;
+						}
 					}
 					else
 					{
-						bool bErrorReportedOnce = false;
-						for (FScriptMapHelper::FIterator It(MapHelper); It; ++It)
+						bool bPassed = true;
+						for (int32 ArrayIdx = 0; ArrayIdx < MapHelper.Num(); ArrayIdx++)
 						{
-							InVisitor->UpdateString(*FString::Printf(TEXT("Key[%d]"), It.GetLogicalIndex()), Names[It.GetLogicalIndex()].ToString());
-							if (!PODPropertyAppendCompileHash(MapHelper.GetPairPtr(It), MapHelper.GetValueProperty(), FString::Printf(TEXT("Value[%d]"), It.GetLogicalIndex()), InVisitor))
+							InVisitor->UpdateString(*FString::Printf(TEXT("Key[%d]"), ArrayIdx), Names[ArrayIdx].ToString());
+							if (!PODPropertyAppendCompileHash(MapHelper.GetPairPtr(ArrayIdx), MapHelper.GetValueProperty(), FString::Printf(TEXT("Value[%d]"), ArrayIdx), InVisitor))
 							{
-								if (!bErrorReportedOnce)
-								{
-									UE_LOG(LogNiagaraEditor, Warning, TEXT("Skipping %s because it is an map value property of unsupported underlying type, please add \"meta = (SkipForCompileHash=\"true\")\" to avoid this warning in the future or handle it yourself in PODPropertyAppendCompileHash!"), *Property->GetName());
-									bErrorReportedOnce = true;
-								}
+								UE_LOG(LogNiagaraEditor, Warning, TEXT("Skipping %s because it is an map value property of unsupported underlying type, please add \"meta = (SkipForCompileHash=\"true\")\" to avoid this warning in the future or handle it yourself in PODPropertyAppendCompileHash!"), *Property->GetName());
+								bPassed = false;
+								continue;
 							}
 						}
+						if (bPassed)
+						{
+							continue;
+						}
 					}
 				}
 				else
@@ -622,6 +636,7 @@ bool FNiagaraEditorUtilities::NestedPropertiesAppendCompileHash(const void* Cont
 					UE_LOG(LogNiagaraEditor, Warning, TEXT("Skipping %s because it is a map property, please add \"meta = (SkipForCompileHash=\"true\")\" to avoid this warning in the future or handle it yourself in NestedPropertiesAppendCompileHash!"), *Property->GetName());
 				}
 			}
+			continue;
 		}
 		else if (Property->IsA(FArrayProperty::StaticClass()))
 		{
@@ -636,41 +651,50 @@ bool FNiagaraEditorUtilities::NestedPropertiesAppendCompileHash(const void* Cont
 			// We support arrays of POD types or arrays of structs with POD types internally. Anything else we should generate a warning on.
 			if (CastProp->Inner->IsA(FStructProperty::StaticClass()))
 			{
-				bool bErrorReportedOnce = false;
+				bool bPassed = true;
 				FStructProperty* StructProp = CastFieldChecked<FStructProperty>(CastProp->Inner);
 
 				for (int32 ArrayIdx = 0; ArrayIdx < ArrayHelper.Num(); ArrayIdx++)
 				{
 					if (!NestedPropertiesAppendCompileHash(ArrayHelper.GetRawPtr(ArrayIdx), StructProp->Struct, EFieldIteratorFlags::IncludeSuper, *PropertyName, InVisitor))
 					{
-						if (!bErrorReportedOnce)
-						{
-							UE_LOG(LogNiagaraEditor, Warning, TEXT("Skipping %s because it is an array property of unsupported underlying type, please add \"meta = (SkipForCompileHash=\"true\")\" to avoid this warning in the future or handle it yourself in NestedPropertiesAppendCompileHash!"), *Property->GetName());
-							bErrorReportedOnce = true;
-						}
+						UE_LOG(LogNiagaraEditor, Warning, TEXT("Skipping %s because it is an array property of unsupported underlying type, please add \"meta = (SkipForCompileHash=\"true\")\" to avoid this warning in the future or handle it yourself in NestedPropertiesAppendCompileHash!"), *Property->GetName());
+						bPassed = false;
+						continue;
 					}
 				}
+				if (bPassed)
+				{
+					continue;
+				}
 			}
 			else
 			{
-				bool bErrorReportedOnce = false;
+				bool bPassed = true;
 				for (int32 ArrayIdx = 0; ArrayIdx < ArrayHelper.Num(); ArrayIdx++)
 				{
 					if (!PODPropertyAppendCompileHash(ArrayHelper.GetRawPtr(ArrayIdx), CastProp->Inner, *PropertyName, InVisitor))
 					{
-						if (!bErrorReportedOnce)
+						if (bPassed)
 						{
 							UE_LOG(LogNiagaraEditor, Warning, TEXT("Skipping %s because it is an array property of unsupported underlying type, please add \"meta = (SkipForCompileHash=\"true\")\" to avoid this warning in the future or handle it yourself in NestedPropertiesAppendCompileHash!"), *Property->GetName());
-							bErrorReportedOnce = true;
 						}
+						bPassed = false;
+						continue;
 					}
 				}
+				if (bPassed)
+				{
+					continue;
+				}
 			}
 
 			UE_LOG(LogNiagaraEditor, Warning, TEXT("Skipping %s because it is an array property, please add \"meta = (SkipForCompileHash=\"true\")\" to avoid this warning in the future or handle it yourself in NestedPropertiesAppendCompileHash!"), *Property->GetName());
+			continue;
 		}
 		else if (Property->IsA(FTextProperty::StaticClass()))
 		{
+			FTextProperty* CastProp = CastFieldChecked<FTextProperty>(Property);
 			UE_LOG(LogNiagaraEditor, Warning, TEXT("Skipping %s because it is a UText property, please add \"meta = (SkipForCompileHash=\"true\")\" to avoid this warning in the future or handle it yourself in NestedPropertiesAppendCompileHash!"), *Property->GetName());
 			return true;
 		}
