Attribute trimming fix for impure functions not finding all variables that need to be preserved because of Module namespace shenanigans

diff --git a/Engine/Plugins/FX/Niagara/Source/NiagaraEditor/Private/NiagaraAttributeTrimmer.cpp b/Engine/Plugins/FX/Niagara/Source/NiagaraEditor/Private/NiagaraAttributeTrimmer.cpp
--- a/Engine/Plugins/FX/Niagara/Source/NiagaraEditor/Private/NiagaraAttributeTrimmer.cpp
+++ b/Engine/Plugins/FX/Niagara/Source/NiagaraEditor/Private/NiagaraAttributeTrimmer.cpp
@@ -443,6 +443,69 @@ FName FNiagaraAttributeTrimmerHelper<GraphBridge>::FindAttributeForRead(const FP
 	return NAME_None;
 }
 
+// for a given pin search through the variables of the parameter map and try to resolve any potential ambiguity with module namespaces
+template<typename GraphBridge>
+FName FNiagaraAttributeTrimmerHelper<GraphBridge>::FindParameterMapVariable(const FParamMapHistory& ParamMap, const FModuleScopedPin& Pin)
+{
+	auto FindVariableName = [&ParamMap](FName NameToEvaluate, FName& OutVariableName) -> bool
+	{
+		if (ParamMap.FindVariableByName(NameToEvaluate) != INDEX_NONE)
+		{
+			OutVariableName = NameToEvaluate;
+			return true;
+		}
+
+		// we also need to check to see if the attribute has resolved namespaces (i.e. Particles.Module.MyAttribute or StackContext.MyAttribute)
+		// for which we'll just search through the parameter maps variables directly
+		const int32 AliasedIndex = ParamMap.VariablesWithOriginalAliasesIntact.IndexOfByPredicate([&](const FNiagaraVariable& Variable)
+		{
+			return Variable.GetName() == NameToEvaluate;
+		});
+
+		if (ParamMap.Variables.IsValidIndex(AliasedIndex))
+		{
+			OutVariableName = ParamMap.Variables[AliasedIndex].GetName();
+			return true;
+		}
+
+		return false;
+	};
+
+	FName FoundVariableName;
+
+	// try the pin name
+	if (FindVariableName(Pin.Pin->PinName, FoundVariableName))
+	{
+		return FoundVariableName;
+	}
+	else if (Pin.ModuleName != NAME_None)
+	{
+		FNameBuilder ModuleNameString(Pin.ModuleName);
+
+		// check if we need to replace an explicit module name with 'Module'
+		FString GenericPinName = Pin.Pin->PinName.ToString();
+		if (GenericPinName.ReplaceInline(*ModuleNameString, *FNiagaraConstants::ModuleNamespaceString))
+		{
+			if (FindVariableName(*GenericPinName, FoundVariableName))
+			{
+				return FoundVariableName;
+			}
+		}
+
+		// try replacing 'Module' with the explicit ModuleName
+		FString ExplicitPinName = Pin.Pin->PinName.ToString();
+		if (ExplicitPinName.ReplaceInline(*FNiagaraConstants::ModuleNamespaceString, *ModuleNameString))
+		{
+			if (FindVariableName(*ExplicitPinName, FoundVariableName))
+			{
+				return FoundVariableName;
+			}
+		}
+	}
+
+	return NAME_None;
+}
+
 // given the set of expressions (as defined in FindDependencies above) we resolve the named attribute aggregating the dependent reads and custom nodes
 template<typename GraphBridge>
 void FNiagaraAttributeTrimmerHelper<GraphBridge>::ResolveDependencyChain(const FParamMapHistory& ParamMap, const FDependencyMap& DependencyData, const FName& AttributeName, FDependencyChain& ResolvedDependencies)
@@ -651,30 +714,12 @@ void FNiagaraAttributeTrimmerHelper<GraphBridge>::TrimAttributes_Aggressive(cons
 			{
 				// see if this variable corresponds to something in our parameter map's list of variables
 				// if it does then we need to mark it as something that needs to be preserved
-				// Note that we need to resolve the Pin name's module namespace before we actually search for it
-				// as a variable
-				FString VariableNameString = DependentPin.Pin->PinName.ToString();
-				VariableNameString.ReplaceInline(*DependentPin.ModuleName.ToString(), *FNiagaraConstants::ModuleNamespace.ToString());
-
-				const FName VariableName = *VariableNameString;
-
-				if (ParamMap->FindVariableByName(VariableName) != INDEX_NONE)
-				{
-					AttributesToPreserve.Add(VariableName);
-				}
-				else
+				// Note that the parameter map can have the variable represented with either an explicit namespace
+				// or the generic Module, so we need to try both.
+				const FName MatchingVariableName = FindParameterMapVariable(*ParamMap, DependentPin);
+				if (MatchingVariableName != NAME_None)
 				{
-					// we also need to check to see if the attribute has resolved namespaces (i.e. Particles.Module.MyAttribute or StackContext.MyAttribute)
-					// for which we'll just search through the parameter maps variables directly
-					const int32 AliasedIndex = ParamMap->VariablesWithOriginalAliasesIntact.IndexOfByPredicate([&](const FNiagaraVariable& Variable)
-						{
-							return Variable.GetName() == VariableName;
-						});
-
-					if (ParamMap->Variables.IsValidIndex(AliasedIndex))
-					{
-						AttributesToPreserve.Add(ParamMap->Variables[AliasedIndex].GetName());
-					}
+					AttributesToPreserve.Add(MatchingVariableName);
 				}
 			}
 		}
diff --git a/Engine/Plugins/FX/Niagara/Source/NiagaraEditor/Private/NiagaraAttributeTrimmer.h b/Engine/Plugins/FX/Niagara/Source/NiagaraEditor/Private/NiagaraAttributeTrimmer.h
--- a/Engine/Plugins/FX/Niagara/Source/NiagaraEditor/Private/NiagaraAttributeTrimmer.h
+++ b/Engine/Plugins/FX/Niagara/Source/NiagaraEditor/Private/NiagaraAttributeTrimmer.h
@@ -56,6 +56,9 @@ class FNiagaraAttributeTrimmerHelper : public GraphBridge
 	// for a specific read of a variable finds the actual name of the attribute being read
 	static FName FindAttributeForRead(const FParamMapHistory& ParamMap, const FModuleScopedPin& ReadPin);
 
+	// for a given pin search through the variables of the parameter map and try to resolve any potential ambiguity with module namespaces
+	static FName FindParameterMapVariable(const FParamMapHistory& ParamMap, const FModuleScopedPin& Pin);
+
 	// given the set of expressions (as defined in FindDependencies above) we resolve the named attribute aggregating the dependent reads and custom nodes
 	static void ResolveDependencyChain(const FParamMapHistory& ParamMap, const FDependencyMap& DependencyData, const FName& AttributeName, FDependencyChain& ResolvedDependencies);
 
diff --git a/Engine/Source/Runtime/Core/Private/UObject/FortniteMainBranchObjectVersion.cpp b/Engine/Source/Runtime/Core/Private/UObject/FortniteMainBranchObjectVersion.cpp
--- a/Engine/Source/Runtime/Core/Private/UObject/FortniteMainBranchObjectVersion.cpp
+++ b/Engine/Source/Runtime/Core/Private/UObject/FortniteMainBranchObjectVersion.cpp
@@ -11,7 +11,7 @@ TMap<FGuid, FGuid> FFortniteMainBranchObjectVersion::GetSystemGuids()
 	SystemGuids.Add(DevGuids.MATERIALSHADERMAP_DERIVEDDATA_VER, FGuid("CE7776E63E9E43FBB8C103EE1AB40B5D"));
 	SystemGuids.Add(DevGuids.NANITE_DERIVEDDATA_VER, FGuid("D28C003F825F4C178F698938FB0E47F7"));
 	SystemGuids.Add(DevGuids.NIAGARASHADERMAP_DERIVEDDATA_VER, FGuid("D0DF49D408F34E38A98FA097E0280EC7"));
-	SystemGuids.Add(DevGuids.Niagara_LatestScriptCompileVersion, FGuid("4D737458B300435D86BAF0F8F3150C68"));
+	SystemGuids.Add(DevGuids.Niagara_LatestScriptCompileVersion, FGuid("9042E07ACA4C4375B7D063C44FA21114"));
 	SystemGuids.Add(DevGuids.SkeletalMeshDerivedDataVersion, FGuid("518D7C51D31D9BE4A9F4B81542AC8119"));
 	SystemGuids.Add(DevGuids.STATICMESH_DERIVEDDATA_VER, FGuid("543843B103794E0E9BA4BE80FB602F79"));
 	return SystemGuids;
