Fixed some class comments and spelling errors

diff --git a/Engine/Plugins/FX/Niagara/Source/NiagaraEditor/Public/NiagaraValidationRules.h b/Engine/Plugins/FX/Niagara/Source/NiagaraEditor/Public/NiagaraValidationRules.h
--- a/Engine/Plugins/FX/Niagara/Source/NiagaraEditor/Public/NiagaraValidationRules.h
+++ b/Engine/Plugins/FX/Niagara/Source/NiagaraEditor/Public/NiagaraValidationRules.h
@@ -17,7 +17,7 @@ namespace NiagaraValidation
 	NIAGARAEDITOR_API void ValidateAllRulesInSystem(TSharedPtr<FNiagaraSystemViewModel> ViewModel, TFunction<void(const FNiagaraValidationResult& Result)> ResultCallback);
 }
 
-/** This validation rule ensures that no Systems have a warm up time set. */
+/** This validation rule ensures that systems don't have a warmup time set. */
 UCLASS(Category = "Validation", DisplayName = "No Warmup Time")
 class UNiagaraValidationRule_NoWarmupTime : public UNiagaraValidationRule
 {
@@ -71,7 +71,7 @@ public:
 	virtual void CheckValidity(const FNiagaraValidationContext& Context, TArray<FNiagaraValidationResult>& OutResults) const override;
 };
 
-/** This validation rule can ban the use of certain datainterfaces on all or a subset of platforms. */
+/** This validation rule can ban the use of certain data interfaces on all or a subset of platforms. */
 UCLASS(Category = "Validation", DisplayName = "Banned DataInterfaces")
 class UNiagaraValidationRule_BannedDataInterfaces : public UNiagaraValidationRule
 {
@@ -109,7 +109,7 @@ public:
 	virtual void CheckValidity(const FNiagaraValidationContext& Context, TArray<FNiagaraValidationResult>& OutResults) const override;
 };
 
-/** This validation rule that can be used to inform or ban GPU usage on the provided platforms. */
+/** This validation rule that can be used to ban GPU usage on the provided platforms or warn that GPU emitters might now work correctly. */
 UCLASS(Category = "Validation", DisplayName = "Gpu Usage")
 class UNiagaraValidationRule_GpuUsage : public UNiagaraValidationRule
 {
@@ -125,7 +125,7 @@ public:
 	virtual void CheckValidity(const FNiagaraValidationContext& Context, TArray<FNiagaraValidationResult>& OutResults) const override;
 };
 
-/** This validation rule can marks this effect type as invalid and so must be changed. Forces a choice of correct Effect Type for an System rather than. Leaving as the default. */
+/** This validation rule always fails and can be used to mark a default/test effect type as stand-in that must be changed. Effectively forces the user to choose a correct effect type for a system. */
 UCLASS(Category = "Validation", DisplayName = "Invalid Effect Type")
 class UNiagaraValidationRule_InvalidEffectType : public UNiagaraValidationRule
 {
@@ -134,7 +134,7 @@ public:
 	virtual void CheckValidity(const FNiagaraValidationContext& Context, TArray<FNiagaraValidationResult>& OutResults) const override;
 };
 
-/** This validation rule checks for various issue with Large World Coordinates. */
+/** This validation rule checks for various common issue with Large World Coordinates like mixing vector and position types. */
 UCLASS(Category = "Validation", DisplayName = "Large World Coordinates")
 class UNiagaraValidationRule_LWC : public UNiagaraValidationRule
 {
@@ -202,7 +202,7 @@ public:
 	int32 MaxTotalIterations = 1;
 };
 
-/** Validation rule to check that we don't have a tick dependency we don't want.  */
+/** Validation rule to check for unwanted tick dependencies.  */
 UCLASS(Category = "Validation", DisplayName = "Tick Dependency Check")
 class UNiagaraValidationRule_TickDependencyCheck : public UNiagaraValidationRule
 {
