Added path info to error log messages

diff --git a/Engine/Plugins/FX/Niagara/Source/Niagara/Private/NiagaraSimCache.cpp b/Engine/Plugins/FX/Niagara/Source/Niagara/Private/NiagaraSimCache.cpp
--- a/Engine/Plugins/FX/Niagara/Source/Niagara/Private/NiagaraSimCache.cpp
+++ b/Engine/Plugins/FX/Niagara/Source/Niagara/Private/NiagaraSimCache.cpp
@@ -230,7 +230,7 @@ void UNiagaraSimCache::Serialize(FArchive& Ar)
 	const int32 NiagaraVersion = Ar.CustomVer(FNiagaraCustomVersion::GUID);
 	if (NiagaraVersion < FNiagaraCustomVersion::SimCache_BulkDataVersion1)
 	{
-		UE_LOG(LogNiagara, Warning, TEXT("SimCache needs to be regenerated as this version is no longer supported."));
+		UE_LOG(LogNiagara, Warning, TEXT("SimCache needs to be regenerated as this version is no longer supported. %s (system %s)"), *GetPathName(), *GetPathNameSafe(SoftNiagaraSystem.Get()));
 		SoftNiagaraSystem.Reset();
 		return;
 	}
@@ -263,7 +263,7 @@ void UNiagaraSimCache::Serialize(FArchive& Ar)
 
 	if (!bIsCacheValid)
 	{
-		UE_LOG(LogNiagara, Warning, TEXT("SimCache buffer serialization failed, likely due to bulk data not being supported, cache is invalid."));
+		UE_LOG(LogNiagara, Warning, TEXT("SimCache buffer serialization failed, likely due to bulk data not being supported, cache is invalid. %s"), *GetPathName());
 		SoftNiagaraSystem.Reset();
 	}
 }
