Tweaked Hierarchy Editor
- Visuals are a bit different - more consistent background brushes with the rest of the engine
- Changed the layout from a horizontalbox to a splitter. Allows resizing of the the three columns

diff --git a/Engine/Plugins/FX/Niagara/Source/NiagaraEditor/Private/Widgets/SNiagaraHierarchy.cpp b/Engine/Plugins/FX/Niagara/Source/NiagaraEditor/Private/Widgets/SNiagaraHierarchy.cpp
--- a/Engine/Plugins/FX/Niagara/Source/NiagaraEditor/Private/Widgets/SNiagaraHierarchy.cpp
+++ b/Engine/Plugins/FX/Niagara/Source/NiagaraEditor/Private/Widgets/SNiagaraHierarchy.cpp
@@ -541,10 +541,7 @@ void SNiagaraHierarchy::Construct(const FArguments& InArgs, TObjectPtr<UNiagaraH
 	OnGenerateRowContentWidget = InArgs._OnGenerateRowContentWidget;
 	OnGenerateCustomDetailsPanelNameWidget = InArgs._OnGenerateCustomDetailsPanelNameWidget;
 	
-	SHorizontalBox::FSlot* DetailsPanelSlot = nullptr;
-
-	LightBackgroundBrush = FSlateColorBrush(FStyleColors::Panel);
-	RecessedBackgroundBrush = FSlateColorBrush(FStyleColors::Recessed);
+	SSplitter::FSlot* DetailsPanelSlot = nullptr;
 
 	TSharedRef<SWidget> AddSectionButton = SNew(SHorizontalBox)
 		+ SHorizontalBox::Slot()
@@ -619,10 +616,15 @@ void SNiagaraHierarchy::Construct(const FArguments& InArgs, TObjectPtr<UNiagaraH
 	ChildSlot
 	[
 		SNew(SBorder)
-		.BorderImage(&RecessedBackgroundBrush)
+		.Padding(0.f)
+		.BorderImage(FAppStyle::Get().GetBrush("Brushes.Panel"))
 		[
-			SNew(SHorizontalBox)
-			+ SHorizontalBox::Slot()
+			SNew(SSplitter)
+			.Orientation(Orient_Horizontal)
+			.PhysicalSplitterHandleSize(2.f)
+			+ SSplitter::Slot()
+			.Value(0.3f)
+			.MinSize(0.1f)
 			[
 				SNew(SVerticalBox)
 				+ SVerticalBox::Slot()
@@ -644,7 +646,7 @@ void SNiagaraHierarchy::Construct(const FArguments& InArgs, TObjectPtr<UNiagaraH
 					.UseAllottedWidth(true)
 				]
 				+ SVerticalBox::Slot()
-				.Padding(6.f, 2.f)
+				.Padding(1.f, 2.f)
 				[
 					SAssignNew(SourceTreeView, STreeView<TSharedPtr<FNiagaraHierarchyItemViewModelBase>>)
 					.TreeItemsSource(&GetSourceItems())
@@ -655,12 +657,9 @@ void SNiagaraHierarchy::Construct(const FArguments& InArgs, TObjectPtr<UNiagaraH
 					.OnContextMenuOpening(this, &SNiagaraHierarchy::SummonContextMenuForSelectedRows, false)
 				]
 			]
-			+ SHorizontalBox::Slot()
-			.AutoWidth()
-			[
-				SNew(SSeparator).Orientation(EOrientation::Orient_Vertical)
-			]
-			+ SHorizontalBox::Slot()
+			+ SSplitter::Slot()
+			.Value(0.4f)
+			.MinSize(0.1f)
 			[
 				SNew(SVerticalBox)
 				+ SVerticalBox::Slot()
@@ -672,7 +671,7 @@ void SNiagaraHierarchy::Construct(const FArguments& InArgs, TObjectPtr<UNiagaraH
 					[
 						SNew(SBorder)
 						.Padding(0.f)
-						.BorderImage(&LightBackgroundBrush)
+						.BorderImage(FAppStyle::Get().GetBrush("Brushes.Header"))
 						[
 							AddSectionButton
 						]
@@ -692,16 +691,15 @@ void SNiagaraHierarchy::Construct(const FArguments& InArgs, TObjectPtr<UNiagaraH
 					.AutoHeight()
 					[
 						SNew(SBorder)
-						.BorderImage(&LightBackgroundBrush)
 						.Padding(0.f)
+						.BorderImage(FAppStyle::Get().GetBrush("Brushes.Header"))
 						[
 							AddCategoryButton
 						]
 					]
-					// we reserve space at the bottom for root drops
 					+ SVerticalBox::Slot()
 					.FillHeight(0.1f)
-					.Padding(2.f, 4.f)
+					.Padding(1.f, 4.f, 1.f, 0.f)
 					[
 						SNew(SDropTarget)
 						.OnDropped(this, &SNiagaraHierarchy::HandleHierarchyRootDrop)
@@ -710,22 +708,28 @@ void SNiagaraHierarchy::Construct(const FArguments& InArgs, TObjectPtr<UNiagaraH
 						.OnDragLeave(this, &SNiagaraHierarchy::OnRootDragLeave)
 						[
 							SNew(SBorder)
-							.BorderImage(FAppStyle::GetBrush("DashedBorder"))
-							.BorderBackgroundColor(FLinearColor(0.2f, 0.2f, 0.2f, 0.5f))
+							.Padding(0.f)
+							.BorderImage(FAppStyle::Get().GetBrush("Brushes.Recessed"))
 							[
-								SNew(SBox)
-								.HAlign(HAlign_Center)
-								.VAlign(VAlign_Center)
+								SNew(SBorder)
+								.Padding(1.f)
+								.BorderImage(FAppStyle::GetBrush("DashedBorder"))
+								.BorderBackgroundColor(FLinearColor(0.2f, 0.2f, 0.2f, 0.5f))
 								[
-									SNew(SImage)
-									.Image(FNiagaraEditorStyle::Get().GetBrush("NiagaraEditor.HierarchyEditor.RootDropIcon"))
-									.ColorAndOpacity(this, &SNiagaraHierarchy::GetRootIconColor)
+									SNew(SBox)
+									.HAlign(HAlign_Center)
+									.VAlign(VAlign_Center)
+									[
+										SNew(SImage)
+										.Image(FNiagaraEditorStyle::Get().GetBrush("NiagaraEditor.HierarchyEditor.RootDropIcon"))
+										.ColorAndOpacity(this, &SNiagaraHierarchy::GetRootIconColor)
+									]
 								]
 							]
 						]
 					]
 					+ SVerticalBox::Slot()
-					.Padding(6.f, 2.f)
+					.Padding(1.f, 0.f)
 					[
 						SAssignNew(HierarchyTreeView, STreeView<TSharedPtr<FNiagaraHierarchyItemViewModelBase>>)
 						.TreeItemsSource(&InHierarchyViewModel->GetHierarchyItems())
@@ -737,7 +741,10 @@ void SNiagaraHierarchy::Construct(const FArguments& InArgs, TObjectPtr<UNiagaraH
 					]
 				]
 			]
-			+ SHorizontalBox::Slot().Expose(DetailsPanelSlot)		
+			+ SSplitter::Slot()
+			.Value(0.3f)
+			.MinSize(0.1f)
+			.Expose(DetailsPanelSlot)		
 		]
 	];
 
diff --git a/Engine/Plugins/FX/Niagara/Source/NiagaraEditor/Public/Widgets/SNiagaraHierarchy.h b/Engine/Plugins/FX/Niagara/Source/NiagaraEditor/Public/Widgets/SNiagaraHierarchy.h
--- a/Engine/Plugins/FX/Niagara/Source/NiagaraEditor/Public/Widgets/SNiagaraHierarchy.h
+++ b/Engine/Plugins/FX/Niagara/Source/NiagaraEditor/Public/Widgets/SNiagaraHierarchy.h
@@ -206,9 +206,7 @@ private:
 	TOptional<FSearchItem> FocusedSearchResult;
 
 	mutable TWeakPtr<FNiagaraHierarchyItemViewModelBase> SelectedDetailsPanelItemViewModel;
-	
-	FSlateBrush LightBackgroundBrush;
-	FSlateBrush RecessedBackgroundBrush;
+
 private:
 	TObjectPtr<UNiagaraHierarchyRoot> SourceRoot = nullptr;
 	TSharedPtr<FNiagaraHierarchyRootViewModel> SourceRootViewModel;
