- Add FlushComputeAndDeferredQueues to world manager
- This wraps some duplicated functionality when manually ticking worlds, like in the baker, etc.
- It will flush the compute disatpch queue, optionally wait for GPU to complete, then execute any deferred callbacks

diff --git a/Engine/Plugins/FX/Niagara/Source/Niagara/Private/NiagaraWorldManager.cpp b/Engine/Plugins/FX/Niagara/Source/Niagara/Private/NiagaraWorldManager.cpp
--- a/Engine/Plugins/FX/Niagara/Source/Niagara/Private/NiagaraWorldManager.cpp
+++ b/Engine/Plugins/FX/Niagara/Source/Niagara/Private/NiagaraWorldManager.cpp
@@ -2096,6 +2096,27 @@ void FNiagaraWorldManager::PrimePool(UNiagaraSystem* System)
 	}
 }
 
+void FNiagaraWorldManager::FlushComputeAndDeferredQueues(bool bWaitForGPU)
+{
+	check(IsInGameThread());
+
+	// Flush the compute dispatch queue
+	if ( FNiagaraGpuComputeDispatchInterface* ComputeDispatchInterface = FNiagaraGpuComputeDispatchInterface::Get(World) )
+	{
+		if (bWaitForGPU)
+		{
+			ComputeDispatchInterface->FlushAndWait_GameThread();
+		}
+		else
+		{
+			ComputeDispatchInterface->FlushPendingTicks_GameThread();
+		}
+	}
+
+	// Flush any deferred callbacks
+	NiagaraWorldManagerInternal::ExecuteGlobalDeferredCallbacks();
+}
+
 void FNiagaraWorldManager::EnqueueGlobalDeferredCallback(TFunction<void()>&& Callback)
 {
 	using namespace NiagaraWorldManagerInternal;
diff --git a/Engine/Plugins/FX/Niagara/Source/Niagara/Public/NiagaraWorldManager.h b/Engine/Plugins/FX/Niagara/Source/Niagara/Public/NiagaraWorldManager.h
--- a/Engine/Plugins/FX/Niagara/Source/Niagara/Public/NiagaraWorldManager.h
+++ b/Engine/Plugins/FX/Niagara/Source/Niagara/Public/NiagaraWorldManager.h
@@ -232,6 +232,11 @@ public:
 
 	class FNiagaraDeferredMethodQueue& GetDeferredMethodQueue() { return DeferredMethods; }
 
+	/**
+	Flush the compute simulation queue and any deferred actions.
+	*/
+	NIAGARA_API void FlushComputeAndDeferredQueues(bool bWaitForGPU);
+
 	/**
 	This is a threadsafe queue which will execute on the game thread.
 	The queue is flushed at the start of each tick group & pre / post actor tick.
diff --git a/Engine/Plugins/FX/Niagara/Source/NiagaraEditor/Private/NiagaraBakerRenderer.cpp b/Engine/Plugins/FX/Niagara/Source/NiagaraEditor/Private/NiagaraBakerRenderer.cpp
--- a/Engine/Plugins/FX/Niagara/Source/NiagaraEditor/Private/NiagaraBakerRenderer.cpp
+++ b/Engine/Plugins/FX/Niagara/Source/NiagaraEditor/Private/NiagaraBakerRenderer.cpp
@@ -326,11 +326,9 @@ void FNiagaraBakerRenderer::SetAbsoluteTime(float AbsoluteTime, bool bShouldTick
 			// Send EOF updates before we flush our pending ticks to ensure everything is ready for Niagara
 			World->SendAllEndOfFrameUpdates();
 
-			// Since captures, etc, don't flush GPU updates so we need to force flush them
-			FNiagaraGpuComputeDispatchInterface* ComputeDispatchInterface = FNiagaraGpuComputeDispatchInterface::Get(World);
-			if ( ensureMsgf(ComputeDispatchInterface, TEXT("The batcher was not valid on the world this may result in incorrect baking")) )
+			if (FNiagaraWorldManager* WorldManager = FNiagaraWorldManager::Get(World))
 			{
-				ComputeDispatchInterface->FlushPendingTicks_GameThread();
+				WorldManager->FlushComputeAndDeferredQueues(false);
 			}
 		}
 	}
diff --git a/Engine/Plugins/FX/Niagara/Source/NiagaraEditor/Private/ViewModels/NiagaraBakerViewModel.cpp b/Engine/Plugins/FX/Niagara/Source/NiagaraEditor/Private/ViewModels/NiagaraBakerViewModel.cpp
--- a/Engine/Plugins/FX/Niagara/Source/NiagaraEditor/Private/ViewModels/NiagaraBakerViewModel.cpp
+++ b/Engine/Plugins/FX/Niagara/Source/NiagaraEditor/Private/ViewModels/NiagaraBakerViewModel.cpp
@@ -4,11 +4,11 @@
 #include "ViewModels/NiagaraSystemViewModel.h"
 #include "Widgets/SNiagaraBakerWidget.h"
 #include "NiagaraBakerRenderer.h"
-#include "NiagaraGpuComputeDispatchInterface.h"
 #include "NiagaraComponent.h"
 #include "NiagaraPlatformSet.h"
 #include "NiagaraSettings.h"
 #include "NiagaraSystem.h"
+#include "NiagaraWorldManager.h"
 
 #include "NiagaraDataInterfaceRenderTarget2D.h"
 
@@ -688,9 +688,9 @@ FNiagaraBakerFeedbackContext FNiagaraBakerViewModel::RenderBaker()
 	// Ensure we flush everything from the render thread
 	if ( UWorld* World = BakerRenderer->GetWorld() )
 	{
-		if ( FNiagaraGpuComputeDispatchInterface* DispatchInterface = FNiagaraGpuComputeDispatchInterface::Get(World) )
+		if (FNiagaraWorldManager* WorldManager = FNiagaraWorldManager::Get(World))
 		{
-			DispatchInterface->FlushAndWait_GameThread();
+			WorldManager->FlushComputeAndDeferredQueues(true);
 		}
 	}
 
