Fix potential nullptr access in data channel BP node

diff --git a/Engine/Plugins/FX/Niagara/Source/Niagara/Private/NiagaraDataChannel.cpp b/Engine/Plugins/FX/Niagara/Source/Niagara/Private/NiagaraDataChannel.cpp
--- a/Engine/Plugins/FX/Niagara/Source/Niagara/Private/NiagaraDataChannel.cpp
+++ b/Engine/Plugins/FX/Niagara/Source/Niagara/Private/NiagaraDataChannel.cpp
@@ -1048,9 +1048,12 @@ UNiagaraDataChannelReader* UNiagaraDataChannelLibrary::ReadFromNiagaraDataChanne
 
 int32 UNiagaraDataChannelLibrary::GetDataChannelElementCount(const UObject* WorldContextObject, const UNiagaraDataChannelAsset* Channel, FNiagaraDataChannelSearchParameters SearchParams, bool bReadPreviousFrame)
 {
-	if (UNiagaraDataChannelReader* Reader = CreateDataChannelReader(WorldContextObject, Channel->Get(), SearchParams, bReadPreviousFrame))
+	if (Channel && Channel->Get())
 	{
-		return Reader->Num();
+		if (UNiagaraDataChannelReader* Reader = CreateDataChannelReader(WorldContextObject, Channel->Get(), SearchParams, bReadPreviousFrame))
+		{
+			return Reader->Num();
+		}
 	}
 	return 0;
 }
