- Add GPU usage information to ribbon renderer output

diff --git a/Engine/Plugins/FX/Niagara/Source/NiagaraEditor/Private/Commandlets/NiagaraSystemAuditCommandlet.cpp b/Engine/Plugins/FX/Niagara/Source/NiagaraEditor/Private/Commandlets/NiagaraSystemAuditCommandlet.cpp
--- a/Engine/Plugins/FX/Niagara/Source/NiagaraEditor/Private/Commandlets/NiagaraSystemAuditCommandlet.cpp
+++ b/Engine/Plugins/FX/Niagara/Source/NiagaraEditor/Private/Commandlets/NiagaraSystemAuditCommandlet.cpp
@@ -425,6 +425,15 @@ bool UNiagaraSystemAuditCommandlet::ProcessNiagaraSystems()
 						RendererBuilder.Append(EmitterHandle.GetInstance().Emitter->GetPathName());
 						RendererBuilder.Append(TEXT(","));
 						RendererBuilder.Append(NiagaraRibbonTessellationModeEnum->GetValueAsString(RibbonRendererProperties->TessellationMode));
+						RendererBuilder.Append(TEXT(","));
+						if (RibbonRendererProperties->bUseGPUInit)
+						{
+							RendererBuilder.Append(TEXT("GPUInit=true "));
+						}
+						if (EmitterData->SimTarget == ENiagaraSimTarget::GPUComputeSim)
+						{
+							RendererBuilder.Append(TEXT("GPUSim=true "));
+						}
 						NiagaraRibbonRenderers.Add(RendererBuilder.ToString());
 					}
 				}
@@ -569,7 +578,7 @@ void UNiagaraSystemAuditCommandlet::DumpResults()
 
 	if (bRendererDetailed)
 	{
-		DumpSimpleSet(NiagaraRibbonRenderers, TEXT("NiagaraRibbonRenderers"), TEXT("Name,TessellationMode"));
+		DumpSimpleSet(NiagaraRibbonRenderers, TEXT("NiagaraRibbonRenderers"), TEXT("Name,TessellationMode,GPU"));
 	}
 }
 
