- Validation rule for ribbon renderer
- Allows validation information for using GPU / GPU init as it may not be valid for all platforms

diff --git a/Engine/Plugins/FX/Niagara/Source/NiagaraEditor/Private/NiagaraValidationRules.cpp b/Engine/Plugins/FX/Niagara/Source/NiagaraEditor/Private/NiagaraValidationRules.cpp
--- a/Engine/Plugins/FX/Niagara/Source/NiagaraEditor/Private/NiagaraValidationRules.cpp
+++ b/Engine/Plugins/FX/Niagara/Source/NiagaraEditor/Private/NiagaraValidationRules.cpp
@@ -10,6 +10,7 @@
 #include "NiagaraEditorSettings.h"
 #include "NiagaraNodeFunctionCall.h"
 #include "NiagaraMeshRendererProperties.h"
+#include "NiagaraRibbonRendererProperties.h"
 #include "NiagaraScriptSource.h"
 #include "NiagaraSettings.h"
 #include "NiagaraSimulationStageBase.h"
@@ -105,7 +106,6 @@ namespace NiagaraValidation
 
 	// --------------------------------------------------------------------------------------------------------------------------------------------
 	// Common fixes and links
-
 	void AddGoToFXTypeLink(FNiagaraValidationResult& Result, UNiagaraEffectType* FXType)
 	{
 		if (FXType == nullptr)
@@ -134,6 +134,33 @@ namespace NiagaraValidation
 			});
 	}
 
+	FNiagaraValidationFix MakeDisableGPUSimulationFix(FVersionedNiagaraEmitterWeakPtr WeakEmitterPtr)
+	{
+		return FNiagaraValidationFix(
+			LOCTEXT("GpuUsageInfoFix_SwitchToCput", "Set emitter to CPU"),
+			FNiagaraValidationFixDelegate::CreateLambda(
+				[WeakEmitterPtr]()
+				{
+					FVersionedNiagaraEmitter VersionedEmitter = WeakEmitterPtr.ResolveWeakPtr();
+					if (FVersionedNiagaraEmitterData* VersionedEmitterData = VersionedEmitter.GetEmitterData())
+					{
+						const FScopedTransaction Transaction(LOCTEXT("SetCPUSim", "Set CPU Simulation"));
+
+						VersionedEmitter.Emitter->Modify();
+						VersionedEmitterData->SimTarget = ENiagaraSimTarget::CPUSim;
+
+						FProperty* SimTargetProperty = FindFProperty<FProperty>(FVersionedNiagaraEmitterData::StaticStruct(), GET_MEMBER_NAME_CHECKED(FVersionedNiagaraEmitterData, SimTarget));
+						FPropertyChangedEvent PropertyChangedEvent(SimTargetProperty);
+						VersionedEmitter.Emitter->PostEditChangeVersionedProperty(PropertyChangedEvent, VersionedEmitter.Version);
+
+						UNiagaraSystem::RequestCompileForEmitter(VersionedEmitter);
+					}
+				}
+			)
+		);
+	}
+
+
 	FString GetPlatformConflictsString(const FNiagaraPlatformSet& PlatformSetA, const FNiagaraPlatformSet& PlatformSetB, int MaxPlatformsToShow = 4)
 	{
 		TArray<const FNiagaraPlatformSet*> CheckSets;
@@ -720,21 +747,83 @@ void UNiagaraValidationRule_GpuUsage::CheckValidity(const FNiagaraValidationCont
 		);
 			
 		ValidationResult.Fixes.Emplace(
-			LOCTEXT("GpuUsageInfoFix_SwitchToCput", "Set emitter to CPU"),
-			FNiagaraValidationFixDelegate::CreateLambda(
-				[WeakEmitterPtr=EmitterHandle->GetInstance().ToWeakPtr()]()
+			NiagaraValidation::MakeDisableGPUSimulationFix(EmitterHandle->GetInstance().ToWeakPtr())
+		);
+	}
+}
+
+void UNiagaraValidationRule_RibbonRenderer::CheckValidity(const FNiagaraValidationContext& Context, TArray<FNiagaraValidationResult>& Results)  const
+{
+	UNiagaraSystem& System = Context.ViewModel->GetSystem();
+	for (TSharedRef<FNiagaraEmitterHandleViewModel> EmitterHandleModel : Context.ViewModel->GetEmitterHandleViewModels())
+	{
+		FNiagaraEmitterHandle* EmitterHandle = EmitterHandleModel.Get().GetEmitterHandle();
+		FVersionedNiagaraEmitterData* EmitterData = EmitterHandle->GetEmitterData();
+
+		if (NiagaraValidation::GetPlatformConflictsString(Platforms, EmitterData->Platforms).IsEmpty())
+		{
+			continue;
+		}
+
+		EmitterData->ForEachRenderer(
+			[this, &Context, &Results, &EmitterData, &EmitterHandleModel, EmitterHandle](UNiagaraRendererProperties* RendererProperties)
+			{
+				UNiagaraRibbonRendererProperties* RibbonRenderer = Cast<UNiagaraRibbonRendererProperties>(RendererProperties);
+				UNiagaraStackRendererItem* StackItem = NiagaraValidation::GetRendererStackItem(EmitterHandleModel.Get().GetEmitterStackViewModel(), RendererProperties);
+				if (!RibbonRenderer || !StackItem)
 				{
-					FVersionedNiagaraEmitter VersionedEmitter = WeakEmitterPtr.ResolveWeakPtr();
-					if (FVersionedNiagaraEmitterData* VersionedEmitterData = VersionedEmitter.GetEmitterData())
-					{
-						VersionedEmitterData->SimTarget = ENiagaraSimTarget::CPUSim;
+					return;
+				}
 
-						FProperty* SimTargetProperty = FindFProperty<FProperty>(FVersionedNiagaraEmitterData::StaticStruct(), GET_MEMBER_NAME_CHECKED(FVersionedNiagaraEmitterData, SimTarget));
-						FPropertyChangedEvent PropertyChangedEvent(SimTargetProperty);
-						VersionedEmitter.Emitter->PostEditChangeVersionedProperty(PropertyChangedEvent, VersionedEmitter.Version);
-					}
+				const FString PlatformConflictsString = NiagaraValidation::GetPlatformConflictsString(Platforms, RendererProperties->Platforms);
+				if (PlatformConflictsString.IsEmpty())
+				{
+					return;
 				}
-			)
+
+				if (bFailIfUsedByGPUSimulation && EmitterData->SimTarget == ENiagaraSimTarget::GPUComputeSim)
+				{
+					FNiagaraValidationResult& ValidationResult = Results.Emplace_GetRef(
+						Severity,
+						LOCTEXT("RibbonRenderer_GpuSimulationError", "Ribbon Renderer is used with GPU simulation"),
+						FText::Format(LOCTEXT("RibbonRenderer_GpuSimulationErrorDetails", "Ribbon Renderer is used with GPU simulation and may not function as expected on '{0}'."), FText::FromString(PlatformConflictsString)),
+						StackItem
+					);
+
+					ValidationResult.Fixes.Emplace(
+						NiagaraValidation::MakeDisableGPUSimulationFix(EmitterHandle->GetInstance().ToWeakPtr())
+					);
+				}
+
+				if (bFailIfUsedByGPUInit && EmitterData->SimTarget != ENiagaraSimTarget::GPUComputeSim && RibbonRenderer->bUseGPUInit)
+				{
+					FNiagaraValidationResult& ValidationResult = Results.Emplace_GetRef(
+						Severity,
+						LOCTEXT("RibbonRenderer_GpuInitError", "Ribbon Renderer is used with GPU init"),
+						FText::Format(LOCTEXT("RibbonRenderer_GpuInitErrorDetails", "Ribbon Renderer is used with GPU init and may not function as expected on '{0}'."), FText::FromString(PlatformConflictsString)),
+						StackItem
+					);
+
+					ValidationResult.Fixes.Emplace(
+						LOCTEXT("RibbonRenderer_GpuInitErrorFix", "Disable GPU init"),
+						FNiagaraValidationFixDelegate::CreateLambda(
+							[WeakRibbonRenderer=MakeWeakObjectPtr(RibbonRenderer)]()
+							{
+								if (WeakRibbonRenderer.IsValid())
+								{
+									const FScopedTransaction Transaction(LOCTEXT("RibbonRenderer_GpuInitErrorApplyFix", "Disable GPU Init"));
+									WeakRibbonRenderer.Get()->Modify();
+									WeakRibbonRenderer.Get()->bUseGPUInit = false;
+									
+									FProperty* Property = FindFProperty<FProperty>(UNiagaraRibbonRendererProperties::StaticClass(), GET_MEMBER_NAME_CHECKED(UNiagaraRibbonRendererProperties, bUseGPUInit));
+									FPropertyChangedEvent PropertyChangedEvent(Property);
+									WeakRibbonRenderer.Get()->PostEditChangeProperty(PropertyChangedEvent);
+								}
+							}
+						)
+					);
+				}
+			}
 		);
 	}
 }
diff --git a/Engine/Plugins/FX/Niagara/Source/NiagaraEditor/Public/NiagaraValidationRules.h b/Engine/Plugins/FX/Niagara/Source/NiagaraEditor/Public/NiagaraValidationRules.h
--- a/Engine/Plugins/FX/Niagara/Source/NiagaraEditor/Public/NiagaraValidationRules.h
+++ b/Engine/Plugins/FX/Niagara/Source/NiagaraEditor/Public/NiagaraValidationRules.h
@@ -125,6 +125,33 @@ public:
 	virtual void CheckValidity(const FNiagaraValidationContext& Context, TArray<FNiagaraValidationResult>& OutResults) const override;
 };
 
+/**
+This validation rule is for ribbon renderers to ensure they are not used in situations that can cause compatability or performance issues.
+i.e. Don't use a ribbon renderer with a GPU emitter / enable GPU ribbon init on lower end devices.
+*/
+UCLASS(Category = "Validation", DisplayName = "Gpu Ribbons")
+class UNiagaraValidationRule_RibbonRenderer : public UNiagaraValidationRule
+{
+	GENERATED_BODY()
+
+public:
+	UPROPERTY(EditAnywhere, Category = Validation)
+	ENiagaraValidationSeverity Severity = ENiagaraValidationSeverity::Warning;
+
+	/** When enable validation will fail if used by a GPU emitter. */
+	UPROPERTY(EditAnywhere, Category = Validation)
+	bool bFailIfUsedByGPUSimulation = true;
+
+	/** When enable validation will fail if used by a CPU emitter and GPU init is enabled on the renderer. */
+	UPROPERTY(EditAnywhere, Category = Validation)
+	bool bFailIfUsedByGPUInit = true;
+
+	UPROPERTY(EditAnywhere, Category = Validation)
+	FNiagaraPlatformSet Platforms;
+
+	virtual void CheckValidity(const FNiagaraValidationContext& Context, TArray<FNiagaraValidationResult>& OutResults) const override;
+};
+
 /** This validation rule always fails and can be used to mark a default/test effect type as stand-in that must be changed. Effectively forces the user to choose a correct effect type for a system. */
 UCLASS(Category = "Validation", DisplayName = "Invalid Effect Type")
 class UNiagaraValidationRule_InvalidEffectType : public UNiagaraValidationRule
