- Change rule display name to make more sense

diff --git a/Engine/Plugins/FX/Niagara/Source/NiagaraEditor/Public/NiagaraValidationRules.h b/Engine/Plugins/FX/Niagara/Source/NiagaraEditor/Public/NiagaraValidationRules.h
--- a/Engine/Plugins/FX/Niagara/Source/NiagaraEditor/Public/NiagaraValidationRules.h
+++ b/Engine/Plugins/FX/Niagara/Source/NiagaraEditor/Public/NiagaraValidationRules.h
@@ -129,7 +129,7 @@ public:
 This validation rule is for ribbon renderers to ensure they are not used in situations that can cause compatability or performance issues.
 i.e. Don't use a ribbon renderer with a GPU emitter / enable GPU ribbon init on lower end devices.
 */
-UCLASS(Category = "Validation", DisplayName = "Gpu Ribbons")
+UCLASS(Category = "Validation", DisplayName = "Ribbon Renderer")
 class UNiagaraValidationRule_RibbonRenderer : public UNiagaraValidationRule
 {
 	GENERATED_BODY()
