caching 2D render targets

#virtualized

diff --git a/Engine/Plugins/FX/Niagara/Source/Niagara/Classes/NiagaraDataInterfaceRenderTarget2D.h b/Engine/Plugins/FX/Niagara/Source/Niagara/Classes/NiagaraDataInterfaceRenderTarget2D.h
--- a/Engine/Plugins/FX/Niagara/Source/Niagara/Classes/NiagaraDataInterfaceRenderTarget2D.h
+++ b/Engine/Plugins/FX/Niagara/Source/Niagara/Classes/NiagaraDataInterfaceRenderTarget2D.h
@@ -9,6 +9,7 @@
 #include "NiagaraDataInterfaceRW.h"
 #include "NiagaraComponent.h"
 #include "NiagaraGenerateMips.h"
+#include "NiagaraSimCacheCustomStorageInterface.h"
 #include "NiagaraDataInterfaceRenderTarget2D.generated.h"
 
 class FNiagaraSystemInstance;
@@ -73,7 +74,7 @@ struct FNiagaraDataInterfaceProxyRenderTarget2DProxy : public FNiagaraDataInterf
 };
 
 UCLASS(EditInlineNew, Category = "Rendering", CollapseCategories, meta = (DisplayName = "Render Target 2D"), Blueprintable, BlueprintType, MinimalAPI)
-class UNiagaraDataInterfaceRenderTarget2D : public UNiagaraDataInterfaceRWBase
+class UNiagaraDataInterfaceRenderTarget2D : public UNiagaraDataInterfaceRWBase, public INiagaraSimCacheCustomStorageInterface
 {
 	GENERATED_UCLASS_BODY()
 
@@ -119,6 +120,13 @@ public:
 	NIAGARA_API virtual bool RenderVariableToCanvas(FNiagaraSystemInstanceID SystemInstanceID, FName VariableName, class FCanvas* Canvas, const FIntRect& DrawRect) const override;
 	//~ UNiagaraDataInterface interface END
 
+	//~ INiagaraSimCacheCustomStorageInterface interface BEGIN
+	NIAGARA_API virtual UObject* SimCacheBeginWrite(UObject* SimCache, FNiagaraSystemInstance* NiagaraSystemInstance, const void* OptionalPerInstanceData, FNiagaraSimCacheFeedbackContext& FeedbackContext) const override;
+	NIAGARA_API virtual bool SimCacheWriteFrame(UObject* StorageObject, int FrameIndex, FNiagaraSystemInstance* SystemInstance, const void* OptionalPerInstanceData, FNiagaraSimCacheFeedbackContext& FeedbackContext) const override;
+	NIAGARA_API virtual bool SimCacheEndWrite(UObject* StorageObject) const override;
+	NIAGARA_API virtual bool SimCacheReadFrame(UObject* StorageObject, int FrameA, int FrameB, float Interp, FNiagaraSystemInstance* SystemInstance, void* OptionalPerInstanceData) override;
+	//~ UNiagaraDataInterface interface END
+
 	NIAGARA_API void VMGetSize(FVectorVMExternalFunctionContext& Context);
 	NIAGARA_API void VMSetSize(FVectorVMExternalFunctionContext& Context);
 	NIAGARA_API void VMGetNumMipLevels(FVectorVMExternalFunctionContext& Context);
diff --git a/Engine/Plugins/FX/Niagara/Source/Niagara/Private/NiagaraDataInterfaceRenderTarget2D.cpp b/Engine/Plugins/FX/Niagara/Source/Niagara/Private/NiagaraDataInterfaceRenderTarget2D.cpp
--- a/Engine/Plugins/FX/Niagara/Source/Niagara/Private/NiagaraDataInterfaceRenderTarget2D.cpp
+++ b/Engine/Plugins/FX/Niagara/Source/Niagara/Private/NiagaraDataInterfaceRenderTarget2D.cpp
@@ -6,6 +6,8 @@
 #include "CanvasTypes.h"
 #include "Engine/TextureRenderTarget2D.h"
 
+#include "NDIRenderTargetVolumeSimCacheData.h"
+#include "NiagaraSimCache.h"
 #include "NiagaraDataInterfaceRenderTargetCommon.h"
 #include "NiagaraSystemInstance.h"
 #include "NiagaraGpuComputeDebugInterface.h"
@@ -22,6 +24,16 @@
 
 namespace NDIRenderTarget2DLocal
 {
+	int32 GSimCacheCompressed = true;
+	static FAutoConsoleVariableRef CVarSimCacheCompressed(
+		TEXT("fx.Niagara.RenderTarget2D.SimCacheCompressed"),
+		GSimCacheCompressed,
+		TEXT("When enabled compression is used for the sim cache data."),
+		ECVF_Default
+	);
+
+	static const FName GetSimCacheCompressionType() { return GSimCacheCompressed ? NAME_Oodle : NAME_None; }
+
 	BEGIN_SHADER_PARAMETER_STRUCT(FShaderParameters, )
 		SHADER_PARAMETER(FIntPoint,								TextureSize)
 		SHADER_PARAMETER(int,									MipLevels)
@@ -562,6 +574,195 @@ bool UNiagaraDataInterfaceRenderTarget2D::RenderVariableToCanvas(FNiagaraSystemI
 	return true;
 }
 
+NIAGARA_API UObject* UNiagaraDataInterfaceRenderTarget2D::SimCacheBeginWrite(UObject* InSimCache, FNiagaraSystemInstance* NiagaraSystemInstance, const void* OptionalPerInstanceData, FNiagaraSimCacheFeedbackContext& FeedbackContext) const
+{
+	UNDIRenderTargetVolumeSimCacheData* SimCacheData = nullptr;
+	SimCacheData = NewObject<UNDIRenderTargetVolumeSimCacheData>(InSimCache);
+	SimCacheData->CompressionType = NDIRenderTarget2DLocal::GetSimCacheCompressionType();
+
+	return SimCacheData;	
+}
+
+NIAGARA_API bool UNiagaraDataInterfaceRenderTarget2D::SimCacheWriteFrame(UObject* StorageObject, int FrameIndex, FNiagaraSystemInstance* SystemInstance, const void* OptionalPerInstanceData, FNiagaraSimCacheFeedbackContext& FeedbackContext) const
+{
+	check(OptionalPerInstanceData);
+	const FRenderTarget2DRWInstanceData_GameThread* InstanceData_GT = reinterpret_cast<const FRenderTarget2DRWInstanceData_GameThread*>(OptionalPerInstanceData);
+
+	if (InstanceData_GT->TargetTexture)
+	{
+		const FNiagaraDataInterfaceProxyRenderTarget2DProxy* RT_Proxy = GetProxyAs<FNiagaraDataInterfaceProxyRenderTarget2DProxy>();
+		const FTextureRenderTargetResource* RT_TargetTexture = InstanceData_GT->TargetTexture->GameThread_GetRenderTargetResource();
+
+		TArray<uint8> TextureData;
+		ENQUEUE_RENDER_COMMAND(NDIRenderTargetVolumeRead)
+		(
+			[RT_Proxy, RT_InstanceID = SystemInstance->GetId(), RT_TargetTexture, RT_TextureData = &TextureData, RT_VolumeResolution = InstanceData_GT->Size](FRHICommandListImmediate& RHICmdList)
+			{
+				if (const FRenderTarget2DRWInstanceData_RenderThread* InstanceData_RT = RT_Proxy->SystemInstancesToProxyData_RT.Find(RT_InstanceID))
+				{
+					FRHIGPUTextureReadback RenderTargetReadback("ReadRTTexture");
+
+					RenderTargetReadback.EnqueueCopy(RHICmdList, InstanceData_RT->TextureRHI, FIntVector(0, 0, 0), 0, FIntVector(RT_VolumeResolution.X, RT_VolumeResolution.Y, 1));
+
+					// Sync the GPU. Unfortunately we can't use the fences because not all RHIs implement them yet.
+					RHICmdList.BlockUntilGPUIdle();
+					RHICmdList.FlushResources();
+
+					//Lock the readback staging texture
+					int32 RowPitchInPixels;
+					int32 BufferHeight;
+					const uint8* LockedData = (const uint8*)RenderTargetReadback.Lock(RowPitchInPixels, &BufferHeight);
+
+					uint32 BlockBytes = GPixelFormats[InstanceData_RT->TextureRHI->GetFormat()].BlockBytes;
+					int32 Count = InstanceData_RT->Size.X * InstanceData_RT->Size.Y * BlockBytes;
+					RT_TextureData->AddUninitialized(Count);
+
+					const uint8* SliceStart = LockedData;
+
+					const uint8* RowStart = SliceStart;
+					for (int32 Y = 0; Y < InstanceData_RT->Size.Y; ++Y)
+					{
+						int32 Offset = 0 + Y * InstanceData_RT->Size.X;
+						FMemory::Memcpy(RT_TextureData->GetData() + Offset * BlockBytes, RowStart, BlockBytes * InstanceData_RT->Size.X);
+
+						RowStart += RowPitchInPixels * BlockBytes;
+					}
+										
+					//Unlock the staging texture
+					RenderTargetReadback.Unlock();
+				}
+			}
+		);
+		FlushRenderingCommands();
+
+		if (TextureData.Num() > 0)
+		{
+			UNDIRenderTargetVolumeSimCacheData* SimCacheData = CastChecked<UNDIRenderTargetVolumeSimCacheData>(StorageObject);
+			SimCacheData->Frames.SetNum(FMath::Max(SimCacheData->Frames.Num(), FrameIndex + 1));
+
+			FNDIRenderTargetVolumeSimCacheFrame* CacheFrame = &SimCacheData->Frames[FrameIndex];
+
+			CacheFrame->Size = FIntVector(InstanceData_GT->Size.X, InstanceData_GT->Size.Y, 1);
+			CacheFrame->Format = GetPixelFormatFromRenderTargetFormat(InstanceData_GT->Format);
+
+			const FName CompressionType = SimCacheData->CompressionType;
+			const bool bUseCompression = CompressionType.IsNone() == false;
+			const int TextureSizeBytes = TextureData.Num() * TextureData.GetTypeSize();
+			int CompressedSize = bUseCompression ? FCompression::CompressMemoryBound(CompressionType, TextureSizeBytes) : TextureSizeBytes;
+
+			CacheFrame->UncompressedSize = TextureSizeBytes;
+			CacheFrame->CompressedSize = 0;
+			uint8* FinalPixelData = reinterpret_cast<uint8*>(FMemory::Malloc(CompressedSize));
+			if (bUseCompression)
+			{
+				if (FCompression::CompressMemory(CompressionType, FinalPixelData, CompressedSize, TextureData.GetData(), TextureSizeBytes, COMPRESS_BiasMemory))
+				{
+					CacheFrame->CompressedSize = CompressedSize;
+				}
+				else
+				{
+					UE_LOG(LogNiagara, Error, TEXT("Error compressing render target data"));
+					return false;
+				}
+			}
+
+			if (CacheFrame->CompressedSize == 0)
+			{
+				FMemory::Memcpy(FinalPixelData, TextureData.GetData(), TextureSizeBytes);
+			}
+
+			// Update bulk data
+			CacheFrame->BulkData.Lock(LOCK_READ_WRITE);
+			{
+				const int32 FinalByteSize = CacheFrame->CompressedSize > 0 ? CacheFrame->CompressedSize : CacheFrame->UncompressedSize;
+				uint8* BulkDataPtr = reinterpret_cast<uint8*>(CacheFrame->BulkData.Realloc(FinalByteSize));
+				FMemory::Memcpy(BulkDataPtr, FinalPixelData, FinalByteSize);
+			}
+			CacheFrame->BulkData.Unlock();
+
+			FMemory::Free(FinalPixelData);
+		}
+	}
+
+	return true;
+}
+
+NIAGARA_API bool UNiagaraDataInterfaceRenderTarget2D::SimCacheEndWrite(UObject* StorageObject) const
+{
+	return true;
+}
+
+NIAGARA_API bool UNiagaraDataInterfaceRenderTarget2D::SimCacheReadFrame(UObject* StorageObject, int FrameA, int FrameB, float Interp, FNiagaraSystemInstance* SystemInstance, void* OptionalPerInstanceData)
+{
+	UNDIRenderTargetVolumeSimCacheData* SimCacheData = CastChecked<UNDIRenderTargetVolumeSimCacheData>(StorageObject);
+
+	FRenderTarget2DRWInstanceData_GameThread* InstanceData_GT = reinterpret_cast<FRenderTarget2DRWInstanceData_GameThread*>(OptionalPerInstanceData);
+
+
+	const int FrameIndex = Interp >= 0.5f ? FrameB : FrameA;
+	if (!SimCacheData->Frames.IsValidIndex(FrameIndex))
+	{
+		return false;
+	}
+
+	FNDIRenderTargetVolumeSimCacheFrame* CacheFrame = &SimCacheData->Frames[FrameIndex];
+
+	InstanceData_GT->Size = FIntPoint(CacheFrame->Size.X, CacheFrame->Size.Y);
+	InstanceData_GT->Format = NiagaraDataInterfaceRenderTargetCommon::GetRenderTargetFormatFromPixelFormat(CacheFrame->Format);
+	InstanceData_GT->Filter = TextureFilter::TF_Default;
+
+	PerInstanceTick(InstanceData_GT, SystemInstance, 0.0f);
+	PerInstanceTickPostSimulate(InstanceData_GT, SystemInstance, 0.0f);
+
+	if (InstanceData_GT->TargetTexture && CacheFrame->GetPixelData() != nullptr)
+	{
+		FNiagaraDataInterfaceProxyRenderTarget2DProxy* RT_Proxy = GetProxyAs<FNiagaraDataInterfaceProxyRenderTarget2DProxy>();
+		FTextureRenderTargetResource* RT_TargetTexture = InstanceData_GT->TargetTexture->GameThread_GetRenderTargetResource();
+		ENQUEUE_RENDER_COMMAND(NDIRenderTargetVolumeUpdate)
+			(
+				[RT_PixelFormat = CacheFrame->Format, RT_Proxy, RT_InstanceID = SystemInstance->GetId(), RT_TargetTexture, RT_CacheFrame = CacheFrame, RT_CompressionType = SimCacheData->CompressionType, RT_Format = InstanceData_GT->Format](FRHICommandListImmediate& RHICmdList)
+				{
+					if (FRenderTarget2DRWInstanceData_RenderThread* InstanceData_RT = RT_Proxy->SystemInstancesToProxyData_RT.Find(RT_InstanceID))
+					{
+						uint32 BlockBytes = GPixelFormats[RT_PixelFormat].BlockBytes;
+						FUpdateTextureRegion2D UpdateRegion(0, 0, 0, 0, InstanceData_RT->Size.X, InstanceData_RT->Size.Y);
+						const int32 DestPitch = InstanceData_RT->Size.X * BlockBytes;
+									
+						TArray<uint8> Decompressed;
+					
+						if (RT_CacheFrame->CompressedSize > 0)
+						{
+							Decompressed.AddUninitialized(BlockBytes * InstanceData_RT->Size.X * InstanceData_RT->Size.Y);
+							bool Success = FCompression::UncompressMemory(RT_CompressionType, Decompressed.GetData(), Decompressed.Num(), RT_CacheFrame->GetPixelData(), RT_CacheFrame->CompressedSize);
+							check(Success);
+						}
+
+						uint32 Stride;
+						uint8* OutBuffer = static_cast<uint8*>(RHICmdList.LockTexture2D(InstanceData_RT->TextureRHI, 0, RLM_WriteOnly, Stride, true));
+
+						const int32 SourcePitch = InstanceData_RT->Size.X * BlockBytes;
+
+						const uint8* SrcData = Decompressed.Num() > 0 ? Decompressed.GetData() : RT_CacheFrame->GetPixelData();
+							
+						uint8* DstData = OutBuffer;
+							
+						for (int32 y = 0; y < InstanceData_RT->Size.Y; ++y)
+						{
+							FMemory::Memcpy(DstData, SrcData, SourcePitch);
+							SrcData += SourcePitch;
+							DstData += Stride;
+						}
+
+						RHICmdList.UnlockTexture2D(InstanceData_RT->TextureRHI, 0, true);						
+						//RHICmdList.UpdateTexture2D(InstanceData_RT->TextureRHI, 0, UpdateRegion, DestPitch, RT_CacheFrame->PixelData);													
+					}
+				}
+		);
+	}
+
+	return true;
+}
+
 void UNiagaraDataInterfaceRenderTarget2D::VMSetSize(FVectorVMExternalFunctionContext& Context)
 {
 	// This should only be called from a system or emitter script due to a need for only setting up initially.
@@ -849,5 +1050,4 @@ void FNiagaraDataInterfaceProxyRenderTarget2DProxy::GetDispatchArgs(const FNDIGp
 	}
 }
 
-#undef LOCTEXT_NAMESPACE
-
+#undef LOCTEXT_NAMESPACE
\ No newline at end of file
diff --git a/Engine/Plugins/FX/Niagara/Source/Niagara/Private/NiagaraDataInterfaceRenderTargetCommon.cpp b/Engine/Plugins/FX/Niagara/Source/Niagara/Private/NiagaraDataInterfaceRenderTargetCommon.cpp
--- a/Engine/Plugins/FX/Niagara/Source/Niagara/Private/NiagaraDataInterfaceRenderTargetCommon.cpp
+++ b/Engine/Plugins/FX/Niagara/Source/Niagara/Private/NiagaraDataInterfaceRenderTargetCommon.cpp
@@ -100,4 +100,27 @@ namespace NiagaraDataInterfaceRenderTargetCommon
 
 		return true;
 	}
+
+	const ETextureRenderTargetFormat GetRenderTargetFormatFromPixelFormat(EPixelFormat InPixelFormat)
+	{
+		switch (InPixelFormat)
+		{
+		case PF_G8:				return ETextureRenderTargetFormat::RTF_R8;
+		case PF_R8G8:			return ETextureRenderTargetFormat::RTF_RG8;
+		case PF_B8G8R8A8:		return ETextureRenderTargetFormat::RTF_RGBA8;
+		case PF_R16F:			return ETextureRenderTargetFormat::RTF_R16f;
+		case PF_G16R16F:		return ETextureRenderTargetFormat::RTF_RG16f;
+		case PF_R32_FLOAT:		return ETextureRenderTargetFormat::RTF_R32f;
+		case PF_G32R32F:		return ETextureRenderTargetFormat::RTF_RG32f;
+		case PF_A32B32G32R32F:	return ETextureRenderTargetFormat::RTF_RGBA32f;
+		case PF_A2B10G10R10:	return ETextureRenderTargetFormat::RTF_RGB10A2;
+		case PF_FloatRGBA:		return ETextureRenderTargetFormat::RTF_RGBA16f;
+		default:
+			break;
+		}
+
+		UE_LOG(LogNiagara, Error, TEXT("No mapping from pixel format to render target format is possible"));
+		return ETextureRenderTargetFormat::RTF_RGBA8;
+	}
+
 }
diff --git a/Engine/Plugins/FX/Niagara/Source/Niagara/Private/NiagaraDataInterfaceRenderTargetCommon.h b/Engine/Plugins/FX/Niagara/Source/Niagara/Private/NiagaraDataInterfaceRenderTargetCommon.h
--- a/Engine/Plugins/FX/Niagara/Source/Niagara/Private/NiagaraDataInterfaceRenderTargetCommon.h
+++ b/Engine/Plugins/FX/Niagara/Source/Niagara/Private/NiagaraDataInterfaceRenderTargetCommon.h
@@ -15,6 +15,8 @@ namespace NiagaraDataInterfaceRenderTargetCommon
 
 	extern bool GetRenderTargetFormat(bool bOverrideFormat, ETextureRenderTargetFormat OverrideFormat, ETextureRenderTargetFormat& OutRenderTargetFormat);
 
+	extern const ETextureRenderTargetFormat GetRenderTargetFormatFromPixelFormat(EPixelFormat InPixelFormat);
+
 	template<typename TRenderTargetType, typename TInstanceDataType>
 	static bool CreateRenderTarget(FNiagaraSystemInstance* SystemInstance, TInstanceDataType* InstanceData)
 	{
diff --git a/Engine/Plugins/FX/Niagara/Source/Niagara/Private/NiagaraDataInterfaceRenderTargetVolume.cpp b/Engine/Plugins/FX/Niagara/Source/Niagara/Private/NiagaraDataInterfaceRenderTargetVolume.cpp
--- a/Engine/Plugins/FX/Niagara/Source/Niagara/Private/NiagaraDataInterfaceRenderTargetVolume.cpp
+++ b/Engine/Plugins/FX/Niagara/Source/Niagara/Private/NiagaraDataInterfaceRenderTargetVolume.cpp
@@ -934,7 +934,7 @@ bool UNiagaraDataInterfaceRenderTargetVolume::SimCacheReadFrame(UObject* Storage
 							TArray<uint8> Decompressed;
 							if (RT_CacheFrame->CompressedSize > 0)
 							{
-								Decompressed.AddUninitialized(sizeof(BlockBytes) * InstanceData_RT->Size.X * InstanceData_RT->Size.Y * InstanceData_RT->Size.Z);
+								Decompressed.AddUninitialized(BlockBytes * InstanceData_RT->Size.X * InstanceData_RT->Size.Y * InstanceData_RT->Size.Z);
 								FCompression::UncompressMemory(RT_CompressionType, Decompressed.GetData(), Decompressed.Num(), RT_CacheFrame->GetPixelData(), RT_CacheFrame->CompressedSize);
 							}
 
