2d cubic interpolation

#virtualized

diff --git a/Engine/Plugins/FX/Niagara/Shaders/Private/NiagaraDataInterfaceGrid2DCollection.ush b/Engine/Plugins/FX/Niagara/Shaders/Private/NiagaraDataInterfaceGrid2DCollection.ush
new file mode 100644
--- /dev/null
+++ b/Engine/Plugins/FX/Niagara/Shaders/Private/NiagaraDataInterfaceGrid2DCollection.ush
@@ -0,0 +1,53 @@
+// Copyright Epic Games, Inc. All Rights Reserved.
+
+float SampleBiCubic_{ParameterName}(SamplerState Sampler, float3 UVW, int MipLevel)
+{		    
+	const int AttributeIndex = round(UVW[2]);
+
+	const float2 GridPos = UVW.xy * {NumCellsName}.xy - .5;
+								
+	// identify the lower-left-hand corner of the cell
+	const int2 GridCell = floor(GridPos);
+								
+	const int2 MaxCell = {NumCellsName} - int2(2,2);
+
+	if (any((GridCell < int2(1,1)) || (GridCell >= MaxCell)))
+	{
+		// revert to bilinear hardware sampling at the boundary cells.
+		return {GridName}.SampleLevel(Sampler, UVW, MipLevel);
+	}
+	else 
+	{
+		// sample point offset from lower left
+		float2 delta = GridPos - GridCell;
+		
+		float4x2 Weights;
+
+		Weights[0] = -1./3 * delta + 1./2 * delta * delta - 1./6 * delta * delta * delta;
+		Weights[1] = 1. - delta * delta + 1./2 * (delta * delta * delta - delta);
+		Weights[2] = delta + 1./2 * (delta * delta - delta * delta * delta);
+		Weights[3] = 1./6 * (delta * delta * delta - delta);		
+
+		float minv = 3.402823466e+38;
+		float maxv = -3.402823466e+38;
+
+		float BiCubicValue = 0;		
+		for (int j = 0; j < 4; ++j)
+		{
+			float qj = 0;
+
+			for (int i = 0; i < 4; ++i)
+			{
+				float Val = {GridName}.Load(int4(GridCell.x + i - 1, GridCell.y + j - 1, AttributeIndex, MipLevel)); 
+				qj += Weights[i][0] * Val;
+
+				minv = min(Val, minv);
+				maxv = max(Val, maxv);
+			}
+
+			BiCubicValue += Weights[j][1] * qj;
+		}
+
+		return min(max(BiCubicValue, minv), maxv);
+	}
+}
\ No newline at end of file
diff --git a/Engine/Plugins/FX/Niagara/Source/Niagara/Classes/NiagaraDataInterfaceGrid2DCollection.h b/Engine/Plugins/FX/Niagara/Source/Niagara/Classes/NiagaraDataInterfaceGrid2DCollection.h
--- a/Engine/Plugins/FX/Niagara/Source/Niagara/Classes/NiagaraDataInterfaceGrid2DCollection.h
+++ b/Engine/Plugins/FX/Niagara/Source/Niagara/Classes/NiagaraDataInterfaceGrid2DCollection.h
@@ -223,10 +223,12 @@ public:
 	static NIAGARA_API const FName SetValueFunctionName;
 	static NIAGARA_API const FName GetValueFunctionName;
 	static NIAGARA_API const FName SampleGridFunctionName;
+	static NIAGARA_API const FName CubicSampleGridFunctionName;
 
 	static NIAGARA_API const FName SetValueAtIndexFunctionName;
 	static NIAGARA_API const FName GetPreviousValueAtIndexFunctionName;
 	static NIAGARA_API const FName SamplePreviousGridAtIndexFunctionName;
+	static NIAGARA_API const FName CubicSamplePreviousGridAtIndexFunctionName;
 
 	static NIAGARA_API const FName SetVector4ValueFunctionName;
 	static NIAGARA_API const FName GetVector4ValueFunctionName;
@@ -246,14 +248,21 @@ public:
 
 	static NIAGARA_API const FName GetPreviousVector4ValueFunctionName;
 	static NIAGARA_API const FName SamplePreviousGridVector4FunctionName;
+	static NIAGARA_API const FName CubicSamplePreviousGridVector4FunctionName;
+
 	static NIAGARA_API const FName SetVectorValueFunctionName;
 	static NIAGARA_API const FName GetPreviousVectorValueFunctionName;
 	static NIAGARA_API const FName SamplePreviousGridVectorFunctionName;
+	static NIAGARA_API const FName CubicSamplePreviousGridVectorFunctionName;
+
 	static NIAGARA_API const FName SetVector2DValueFunctionName;
 	static NIAGARA_API const FName GetPreviousVector2DValueFunctionName;
 	static NIAGARA_API const FName SamplePreviousGridVector2DFunctionName;
+	static NIAGARA_API const FName CubicSamplePreviousGridVector2DFunctionName;
+
 	static NIAGARA_API const FName GetPreviousFloatValueFunctionName;
 	static NIAGARA_API const FName SamplePreviousGridFloatFunctionName;
+	static NIAGARA_API const FName CubicSamplePreviousGridFloatFunctionName;
 	
 	static NIAGARA_API const FString AttributeIndicesBaseName;
 	static NIAGARA_API const TCHAR* VectorComponentNames[];
@@ -267,6 +276,8 @@ public:
 
 	static NIAGARA_API const FString AnonymousAttributeString;
 
+	static NIAGARA_API const TCHAR* TemplateShaderFilePath;
+
 #if WITH_EDITOR
 	virtual bool SupportsSetupAndTeardownHLSL() const { return true; }
 	NIAGARA_API virtual bool GenerateSetupHLSL(FNiagaraDataInterfaceGPUParamInfo& DIInstanceInfo, TConstArrayView<FNiagaraVariable> InArguments, bool bSpawnOnly, bool bPartialWrites, TArray<FText>& OutErrors, FString& OutHLSL) const;
@@ -285,7 +296,7 @@ protected:
 #if WITH_EDITORONLY_DATA
 	NIAGARA_API void WriteSetHLSL(const FNiagaraDataInterfaceRWAttributeHelper &AttributeHelper, const FNiagaraDataInterfaceGPUParamInfo& ParamInfo, const FNiagaraDataInterfaceGeneratedFunction& FunctionInfo, int FunctionInstanceIndex, int32 InNumChannels, FString& OutHLSL);
 	NIAGARA_API void WriteGetHLSL(const FNiagaraDataInterfaceRWAttributeHelper& AttributeHelper, const FNiagaraDataInterfaceGPUParamInfo& ParamInfo, const FNiagaraDataInterfaceGeneratedFunction& FunctionInfo, int FunctionInstanceIndex, int32 InNumChannels, FString& OutHLSL);
-	NIAGARA_API void WriteSampleHLSL(const FNiagaraDataInterfaceRWAttributeHelper& AttributeHelper, const FNiagaraDataInterfaceGPUParamInfo& ParamInfo, const FNiagaraDataInterfaceGeneratedFunction& FunctionInfo, int FunctionInstanceIndex, int32 InNumChannels, FString& OutHLSL);
+	NIAGARA_API void WriteSampleHLSL(const FNiagaraDataInterfaceRWAttributeHelper& AttributeHelper, const FNiagaraDataInterfaceGPUParamInfo& ParamInfo, const FNiagaraDataInterfaceGeneratedFunction& FunctionInfo, int FunctionInstanceIndex, int32 InNumChannels, bool IsCubic, FString& OutHLSL);
 	NIAGARA_API void WriteAttributeGetIndexHLSL(const FNiagaraDataInterfaceRWAttributeHelper& AttributeHelper, const FNiagaraDataInterfaceGPUParamInfo& ParamInfo, const FNiagaraDataInterfaceGeneratedFunction& FunctionInfo, int FunctionInstanceIndex, int32 InNumChannels, FString& OutHLSL);
 
 	NIAGARA_API FString GenerateAttributeIndexHLSL(const FNiagaraDataInterfaceRWAttributeHelper& AttributeHelper, const FNiagaraDataInterfaceGeneratedFunction& FunctionInfo);
diff --git a/Engine/Plugins/FX/Niagara/Source/Niagara/Private/NiagaraDataInterfaceGrid2DCollection.cpp b/Engine/Plugins/FX/Niagara/Source/Niagara/Private/NiagaraDataInterfaceGrid2DCollection.cpp
--- a/Engine/Plugins/FX/Niagara/Source/Niagara/Private/NiagaraDataInterfaceGrid2DCollection.cpp
+++ b/Engine/Plugins/FX/Niagara/Source/Niagara/Private/NiagaraDataInterfaceGrid2DCollection.cpp
@@ -73,17 +73,22 @@ const FName UNiagaraDataInterfaceGrid2DCollection::SampleGridFloatFunctionName("
 const FName UNiagaraDataInterfaceGrid2DCollection::SetValueAtIndexFunctionName("SetValueAtIndex");
 const FName UNiagaraDataInterfaceGrid2DCollection::GetPreviousValueAtIndexFunctionName("GetPreviousValueAtIndex");
 const FName UNiagaraDataInterfaceGrid2DCollection::SamplePreviousGridAtIndexFunctionName("SamplePreviousGridAtIndex");
+const FName UNiagaraDataInterfaceGrid2DCollection::CubicSamplePreviousGridAtIndexFunctionName("CubicSamplePreviousGridAtIndex");
 
 const FName UNiagaraDataInterfaceGrid2DCollection::GetPreviousVector4ValueFunctionName("GetPreviousVector4Value");
 const FName UNiagaraDataInterfaceGrid2DCollection::SamplePreviousGridVector4FunctionName("SamplePreviousGridVector4Value");
+const FName UNiagaraDataInterfaceGrid2DCollection::CubicSamplePreviousGridVector4FunctionName("CubicSamplePreviousGridVector4Value");
 const FName UNiagaraDataInterfaceGrid2DCollection::SetVectorValueFunctionName("SetVectorValue");
 const FName UNiagaraDataInterfaceGrid2DCollection::GetPreviousVectorValueFunctionName("GetPreviousVectorValue");
 const FName UNiagaraDataInterfaceGrid2DCollection::SamplePreviousGridVectorFunctionName("SamplePreviousGridVector3Value");
+const FName UNiagaraDataInterfaceGrid2DCollection::CubicSamplePreviousGridVectorFunctionName("CubicSamplePreviousGridVector3Value");
 const FName UNiagaraDataInterfaceGrid2DCollection::SetVector2DValueFunctionName("SetVector2DValue");
 const FName UNiagaraDataInterfaceGrid2DCollection::GetPreviousVector2DValueFunctionName("GetPreviousVector2DValue");
 const FName UNiagaraDataInterfaceGrid2DCollection::SamplePreviousGridVector2DFunctionName("SamplePreviousGridVector2DValue");
+const FName UNiagaraDataInterfaceGrid2DCollection::CubicSamplePreviousGridVector2DFunctionName("CubicSamplePreviousGridVector2DValue");
 const FName UNiagaraDataInterfaceGrid2DCollection::GetPreviousFloatValueFunctionName("GetPreviousFloatValue");
 const FName UNiagaraDataInterfaceGrid2DCollection::SamplePreviousGridFloatFunctionName("SamplePreviousGridFloatValue");
+const FName UNiagaraDataInterfaceGrid2DCollection::CubicSamplePreviousGridFloatFunctionName("CubicSamplePreviousGridFloatValue");
 
 const FName UNiagaraDataInterfaceGrid2DCollection::GetVector4AttributeIndexFunctionName("GetVector4AttributeIndex");
 const FName UNiagaraDataInterfaceGrid2DCollection::GetVectorAttributeIndexFunctionName("GetVectorAttributeIndex");
@@ -102,6 +107,8 @@ const FName UNiagaraDataInterfaceGrid2DCollection::SampleGridFunctionName("Sampl
 
 FNiagaraVariableBase UNiagaraDataInterfaceGrid2DCollection::ExposedRTVar;
 
+const TCHAR* UNiagaraDataInterfaceGrid2DCollection::TemplateShaderFilePath = TEXT("/Plugin/FX/Niagara/Private/NiagaraDataInterfaceGrid2DCollection.ush");
+
 struct FNiagaraGridCollection2DDIFunctionVersion
 {
 	enum Type
@@ -117,26 +124,26 @@ struct FNiagaraGridCollection2DDIFunctionVersion
 
 bool UNiagaraDataInterfaceGrid2DCollection::CanCreateVarFromFuncName(const FName& FuncName)
 {
-	if (FuncName == UNiagaraDataInterfaceGrid2DCollection::SetVector4ValueFunctionName	|| FuncName == UNiagaraDataInterfaceGrid2DCollection::GetVector4ValueFunctionName || FuncName == UNiagaraDataInterfaceGrid2DCollection::SampleGridVector4FunctionName || FuncName == UNiagaraDataInterfaceGrid2DCollection::GetPreviousVector4ValueFunctionName || FuncName == UNiagaraDataInterfaceGrid2DCollection::SamplePreviousGridVector4FunctionName)
+	if (FuncName == UNiagaraDataInterfaceGrid2DCollection::SetVector4ValueFunctionName	|| FuncName == UNiagaraDataInterfaceGrid2DCollection::GetVector4ValueFunctionName || FuncName == UNiagaraDataInterfaceGrid2DCollection::SampleGridVector4FunctionName || FuncName == UNiagaraDataInterfaceGrid2DCollection::GetPreviousVector4ValueFunctionName || FuncName == UNiagaraDataInterfaceGrid2DCollection::SamplePreviousGridVector4FunctionName || FuncName == UNiagaraDataInterfaceGrid2DCollection::CubicSamplePreviousGridVector4FunctionName)
 		return true;
-	else if (FuncName == UNiagaraDataInterfaceGrid2DCollection::SetVector3ValueFunctionName || FuncName == UNiagaraDataInterfaceGrid2DCollection::SetVectorValueFunctionName || FuncName == UNiagaraDataInterfaceGrid2DCollection::GetVector3ValueFunctionName || FuncName == UNiagaraDataInterfaceGrid2DCollection::SampleGridVector3FunctionName || FuncName == UNiagaraDataInterfaceGrid2DCollection::GetPreviousVectorValueFunctionName || FuncName == UNiagaraDataInterfaceGrid2DCollection::SamplePreviousGridVectorFunctionName)
+	else if (FuncName == UNiagaraDataInterfaceGrid2DCollection::SetVector3ValueFunctionName || FuncName == UNiagaraDataInterfaceGrid2DCollection::SetVectorValueFunctionName || FuncName == UNiagaraDataInterfaceGrid2DCollection::GetVector3ValueFunctionName || FuncName == UNiagaraDataInterfaceGrid2DCollection::SampleGridVector3FunctionName || FuncName == UNiagaraDataInterfaceGrid2DCollection::GetPreviousVectorValueFunctionName || FuncName == UNiagaraDataInterfaceGrid2DCollection::SamplePreviousGridVectorFunctionName || FuncName == UNiagaraDataInterfaceGrid2DCollection::CubicSamplePreviousGridVectorFunctionName)
 		return true;
-	else if (FuncName == UNiagaraDataInterfaceGrid2DCollection::SetVector2ValueFunctionName || FuncName == UNiagaraDataInterfaceGrid2DCollection::SetVector2DValueFunctionName || FuncName == UNiagaraDataInterfaceGrid2DCollection::GetVector2ValueFunctionName || FuncName == UNiagaraDataInterfaceGrid2DCollection::SampleGridVector2FunctionName || FuncName == UNiagaraDataInterfaceGrid2DCollection::GetPreviousVector2DValueFunctionName || FuncName == UNiagaraDataInterfaceGrid2DCollection::SamplePreviousGridVector2DFunctionName)
+	else if (FuncName == UNiagaraDataInterfaceGrid2DCollection::SetVector2ValueFunctionName || FuncName == UNiagaraDataInterfaceGrid2DCollection::SetVector2DValueFunctionName || FuncName == UNiagaraDataInterfaceGrid2DCollection::GetVector2ValueFunctionName || FuncName == UNiagaraDataInterfaceGrid2DCollection::SampleGridVector2FunctionName || FuncName == UNiagaraDataInterfaceGrid2DCollection::GetPreviousVector2DValueFunctionName || FuncName == UNiagaraDataInterfaceGrid2DCollection::SamplePreviousGridVector2DFunctionName || FuncName == UNiagaraDataInterfaceGrid2DCollection::CubicSamplePreviousGridVector2DFunctionName)
 		return true;
-	else if (FuncName == UNiagaraDataInterfaceGrid2DCollection::SetFloatValueFunctionName || FuncName == UNiagaraDataInterfaceGrid2DCollection::GetFloatValueFunctionName || FuncName == UNiagaraDataInterfaceGrid2DCollection::SampleGridFloatFunctionName       || FuncName == UNiagaraDataInterfaceGrid2DCollection::GetPreviousFloatValueFunctionName || FuncName == UNiagaraDataInterfaceGrid2DCollection::SamplePreviousGridFloatFunctionName )
+	else if (FuncName == UNiagaraDataInterfaceGrid2DCollection::SetFloatValueFunctionName || FuncName == UNiagaraDataInterfaceGrid2DCollection::GetFloatValueFunctionName || FuncName == UNiagaraDataInterfaceGrid2DCollection::SampleGridFloatFunctionName       || FuncName == UNiagaraDataInterfaceGrid2DCollection::GetPreviousFloatValueFunctionName || FuncName == UNiagaraDataInterfaceGrid2DCollection::SamplePreviousGridFloatFunctionName || FuncName == UNiagaraDataInterfaceGrid2DCollection::CubicSamplePreviousGridFloatFunctionName)
 		return true;
 	return false;
 }
 
 FNiagaraTypeDefinition UNiagaraDataInterfaceGrid2DCollection::GetValueTypeFromFuncName(const FName& FuncName)
 {
-	if (FuncName == UNiagaraDataInterfaceGrid2DCollection::SetVector4ValueFunctionName || FuncName == UNiagaraDataInterfaceGrid2DCollection::GetVector4ValueFunctionName || FuncName == UNiagaraDataInterfaceGrid2DCollection::SampleGridVector4FunctionName || FuncName == UNiagaraDataInterfaceGrid2DCollection::GetPreviousVector4ValueFunctionName || FuncName == UNiagaraDataInterfaceGrid2DCollection::SamplePreviousGridVector4FunctionName)
+	if (FuncName == UNiagaraDataInterfaceGrid2DCollection::SetVector4ValueFunctionName || FuncName == UNiagaraDataInterfaceGrid2DCollection::GetVector4ValueFunctionName || FuncName == UNiagaraDataInterfaceGrid2DCollection::SampleGridVector4FunctionName || FuncName == UNiagaraDataInterfaceGrid2DCollection::GetPreviousVector4ValueFunctionName || FuncName == UNiagaraDataInterfaceGrid2DCollection::SamplePreviousGridVector4FunctionName || FuncName == UNiagaraDataInterfaceGrid2DCollection::CubicSamplePreviousGridVector4FunctionName)
 		return FNiagaraTypeDefinition::GetVec4Def();
-	else if (FuncName == UNiagaraDataInterfaceGrid2DCollection::SetVector3ValueFunctionName || FuncName == UNiagaraDataInterfaceGrid2DCollection::SetVectorValueFunctionName || FuncName == UNiagaraDataInterfaceGrid2DCollection::GetVector3ValueFunctionName || FuncName == UNiagaraDataInterfaceGrid2DCollection::SampleGridVector3FunctionName || FuncName == UNiagaraDataInterfaceGrid2DCollection::GetPreviousVectorValueFunctionName || FuncName == UNiagaraDataInterfaceGrid2DCollection::SamplePreviousGridVectorFunctionName)
+	else if (FuncName == UNiagaraDataInterfaceGrid2DCollection::SetVector3ValueFunctionName || FuncName == UNiagaraDataInterfaceGrid2DCollection::SetVectorValueFunctionName || FuncName == UNiagaraDataInterfaceGrid2DCollection::GetVector3ValueFunctionName || FuncName == UNiagaraDataInterfaceGrid2DCollection::SampleGridVector3FunctionName || FuncName == UNiagaraDataInterfaceGrid2DCollection::GetPreviousVectorValueFunctionName || FuncName == UNiagaraDataInterfaceGrid2DCollection::SamplePreviousGridVectorFunctionName || FuncName == UNiagaraDataInterfaceGrid2DCollection::CubicSamplePreviousGridVectorFunctionName)
 		return FNiagaraTypeDefinition::GetVec3Def();
-	else if (FuncName == UNiagaraDataInterfaceGrid2DCollection::SetVector2ValueFunctionName || FuncName == UNiagaraDataInterfaceGrid2DCollection::SetVector2DValueFunctionName || FuncName == UNiagaraDataInterfaceGrid2DCollection::GetVector2ValueFunctionName || FuncName == UNiagaraDataInterfaceGrid2DCollection::SampleGridVector2FunctionName || FuncName == UNiagaraDataInterfaceGrid2DCollection::GetPreviousVector2DValueFunctionName || FuncName == UNiagaraDataInterfaceGrid2DCollection::SamplePreviousGridVector2DFunctionName)
+	else if (FuncName == UNiagaraDataInterfaceGrid2DCollection::SetVector2ValueFunctionName || FuncName == UNiagaraDataInterfaceGrid2DCollection::SetVector2DValueFunctionName || FuncName == UNiagaraDataInterfaceGrid2DCollection::GetVector2ValueFunctionName || FuncName == UNiagaraDataInterfaceGrid2DCollection::SampleGridVector2FunctionName || FuncName == UNiagaraDataInterfaceGrid2DCollection::GetPreviousVector2DValueFunctionName || FuncName == UNiagaraDataInterfaceGrid2DCollection::SamplePreviousGridVector2DFunctionName || FuncName == UNiagaraDataInterfaceGrid2DCollection::CubicSamplePreviousGridVector2DFunctionName)
 		return FNiagaraTypeDefinition::GetVec2Def();
-	else if (FuncName == UNiagaraDataInterfaceGrid2DCollection::SetFloatValueFunctionName || FuncName == UNiagaraDataInterfaceGrid2DCollection::GetFloatValueFunctionName || FuncName == UNiagaraDataInterfaceGrid2DCollection::SampleGridFloatFunctionName || FuncName == UNiagaraDataInterfaceGrid2DCollection::GetPreviousFloatValueFunctionName || FuncName == UNiagaraDataInterfaceGrid2DCollection::SamplePreviousGridFloatFunctionName)
+	else if (FuncName == UNiagaraDataInterfaceGrid2DCollection::SetFloatValueFunctionName || FuncName == UNiagaraDataInterfaceGrid2DCollection::GetFloatValueFunctionName || FuncName == UNiagaraDataInterfaceGrid2DCollection::SampleGridFloatFunctionName || FuncName == UNiagaraDataInterfaceGrid2DCollection::GetPreviousFloatValueFunctionName || FuncName == UNiagaraDataInterfaceGrid2DCollection::SamplePreviousGridFloatFunctionName || FuncName == UNiagaraDataInterfaceGrid2DCollection::CubicSamplePreviousGridFloatFunctionName)
 		return FNiagaraTypeDefinition::GetFloatDef();
 
 	return FNiagaraTypeDefinition();
@@ -627,6 +634,31 @@ void UNiagaraDataInterfaceGrid2DCollection::GetFunctions(TArray<FNiagaraFunction
 		OutFunctions.Add(Sig);
 	}
 
+	{
+		FNiagaraFunctionSignature Sig;
+		Sig.Name = CubicSamplePreviousGridVector4FunctionName;
+		Sig.Inputs.Add(FNiagaraVariable(FNiagaraTypeDefinition(GetClass()), TEXT("Grid")));
+		Sig.Inputs.Add(FNiagaraVariable(FNiagaraTypeDefinition::GetVec2Def(), TEXT("Unit")));
+		Sig.Outputs.Add(FNiagaraVariable(FNiagaraTypeDefinition::GetVec4Def(), TEXT("Value")));
+		Sig.FunctionSpecifiers.Add(FName("Attribute"));
+		Sig.bMemberFunction = true;
+		Sig.bRequiresContext = false;
+		Sig.bExperimental = true;
+		Sig.ModuleUsageBitmask = ENiagaraScriptUsageMask::Particle;
+		Sig.bSupportsCPU = false;
+		Sig.bSupportsGPU = true;
+		Sig.bReadFunction = true;
+#if WITH_EDITORONLY_DATA
+		Sig.Description = NSLOCTEXT("Niagara", "NiagaraDataInterfaceGridColl2D_CubicSampleVector4", "Cubic sample a Vector4 value on the Grid by Attribute name. Note that this is the value from the previous execution stage.");
+#endif
+		OutFunctions.Add(Sig);
+
+		// Add older version for back compat with CustomHLSL.
+		Sig.Name = SampleGridVector4FunctionName;
+		Sig.bSoftDeprecatedFunction = true;
+		OutFunctions.Add(Sig);
+	}
+
 	{
 		FNiagaraFunctionSignature Sig;
 		Sig.Name = SetVectorValueFunctionName;
@@ -705,6 +737,31 @@ void UNiagaraDataInterfaceGrid2DCollection::GetFunctions(TArray<FNiagaraFunction
 		OutFunctions.Add(Sig);
 	}
 
+	{
+		FNiagaraFunctionSignature Sig;
+		Sig.Name = CubicSamplePreviousGridVectorFunctionName;
+		Sig.Inputs.Add(FNiagaraVariable(FNiagaraTypeDefinition(GetClass()), TEXT("Grid")));
+		Sig.Inputs.Add(FNiagaraVariable(FNiagaraTypeDefinition::GetVec2Def(), TEXT("Unit")));
+		Sig.Outputs.Add(FNiagaraVariable(FNiagaraTypeDefinition::GetVec3Def(), TEXT("Value")));
+		Sig.FunctionSpecifiers.Add(FName("Attribute"));
+		Sig.bMemberFunction = true;
+		Sig.bRequiresContext = false;
+		Sig.bExperimental = true;
+		Sig.ModuleUsageBitmask = ENiagaraScriptUsageMask::Particle;
+		Sig.bSupportsCPU = false;
+		Sig.bSupportsGPU = true;
+		Sig.bReadFunction = true;
+#if WITH_EDITORONLY_DATA
+		Sig.Description = NSLOCTEXT("Niagara", "NiagaraDataInterfaceGridColl2D_CubicSampleVector3", "Cubic sample a Vector3 value on the Grid by Attribute name. Note that this is the value from the previous execution stage.");
+#endif
+		OutFunctions.Add(Sig);
+
+		// Add older version for back compat with CustomHLSL.
+		Sig.Name = SampleGridVector3FunctionName;
+		Sig.bSoftDeprecatedFunction = true;
+		OutFunctions.Add(Sig);
+	}
+
 	{
 		FNiagaraFunctionSignature Sig;
 		Sig.Name = SetVector2DValueFunctionName;
@@ -783,6 +840,31 @@ void UNiagaraDataInterfaceGrid2DCollection::GetFunctions(TArray<FNiagaraFunction
 		OutFunctions.Add(Sig);
 	}
 
+	{
+		FNiagaraFunctionSignature Sig;
+		Sig.Name = CubicSamplePreviousGridVector2DFunctionName;
+		Sig.Inputs.Add(FNiagaraVariable(FNiagaraTypeDefinition(GetClass()), TEXT("Grid")));
+		Sig.Inputs.Add(FNiagaraVariable(FNiagaraTypeDefinition::GetVec2Def(), TEXT("Unit")));
+		Sig.Outputs.Add(FNiagaraVariable(FNiagaraTypeDefinition::GetVec2Def(), TEXT("Value")));
+		Sig.FunctionSpecifiers.Add(FName("Attribute"));
+		Sig.bMemberFunction = true;
+		Sig.bRequiresContext = false;
+		Sig.bExperimental = true;
+		Sig.ModuleUsageBitmask = ENiagaraScriptUsageMask::Particle;
+		Sig.bSupportsCPU = false;
+		Sig.bSupportsGPU = true;
+		Sig.bReadFunction = true;
+#if WITH_EDITORONLY_DATA
+		Sig.Description = NSLOCTEXT("Niagara", "NiagaraDataInterfaceGridColl2D_CubicSampleVector2", "Cubic sample a Vector2 value on the Grid by Attribute name. Note that this is the value from the previous execution stage.");
+#endif
+		OutFunctions.Add(Sig);
+
+		// Add older version for back compat with CustomHLSL.
+		Sig.Name = SampleGridVector2FunctionName;
+		Sig.bSoftDeprecatedFunction = true;
+		OutFunctions.Add(Sig);
+	}
+
 	{
 		FNiagaraFunctionSignature Sig;
 		Sig.Name = SetFloatValueFunctionName;
@@ -856,6 +938,31 @@ void UNiagaraDataInterfaceGrid2DCollection::GetFunctions(TArray<FNiagaraFunction
 		OutFunctions.Add(Sig);
 	}
 
+	{
+		FNiagaraFunctionSignature Sig;
+		Sig.Name = CubicSamplePreviousGridFloatFunctionName;
+		Sig.Inputs.Add(FNiagaraVariable(FNiagaraTypeDefinition(GetClass()), TEXT("Grid")));
+		Sig.Inputs.Add(FNiagaraVariable(FNiagaraTypeDefinition::GetVec2Def(), TEXT("Unit")));
+		Sig.Outputs.Add(FNiagaraVariable(FNiagaraTypeDefinition::GetFloatDef(), TEXT("Value")));
+		Sig.FunctionSpecifiers.Add(FName("Attribute"));
+		Sig.bMemberFunction = true;
+		Sig.bRequiresContext = false;
+		Sig.bExperimental = true;
+		Sig.ModuleUsageBitmask = ENiagaraScriptUsageMask::Particle;
+		Sig.bSupportsCPU = false;
+		Sig.bSupportsGPU = true;
+		Sig.bReadFunction = true;
+#if WITH_EDITORONLY_DATA
+		Sig.Description = NSLOCTEXT("Niagara", "NiagaraDataInterfaceGridColl2D_CubicSampleFloat", "Cubic sample a float value on the Grid by Attribute name.");
+#endif
+		OutFunctions.Add(Sig);
+
+		// Add older version for back compat with CustomHLSL.
+		Sig.Name = SampleGridFloatFunctionName;
+		Sig.bSoftDeprecatedFunction = true;
+		OutFunctions.Add(Sig);
+	}
+
 	{
 		FNiagaraFunctionSignature Sig;
 		Sig.Name = SampleGridFunctionName;
@@ -895,6 +1002,25 @@ void UNiagaraDataInterfaceGrid2DCollection::GetFunctions(TArray<FNiagaraFunction
 		OutFunctions.Add(Sig);
 	}
 
+	{
+		FNiagaraFunctionSignature Sig;
+		Sig.Name = CubicSamplePreviousGridAtIndexFunctionName;
+		Sig.Inputs.Add(FNiagaraVariable(FNiagaraTypeDefinition(GetClass()), TEXT("Grid")));
+		Sig.Inputs.Add(FNiagaraVariable(FNiagaraTypeDefinition::GetFloatDef(), TEXT("UnitX")));
+		Sig.Inputs.Add(FNiagaraVariable(FNiagaraTypeDefinition::GetFloatDef(), TEXT("UnitY")));
+		Sig.Inputs.Add(FNiagaraVariable(FNiagaraTypeDefinition::GetIntDef(), TEXT("AttributeIndex")));
+		Sig.Outputs.Add(FNiagaraVariable(FNiagaraTypeDefinition::GetFloatDef(), TEXT("Value")));
+
+		Sig.bExperimental = true;
+		Sig.bMemberFunction = true;
+		Sig.bRequiresContext = false;
+		Sig.ModuleUsageBitmask = ENiagaraScriptUsageMask::Particle;
+		Sig.bSupportsCPU = false;
+		Sig.bSupportsGPU = true;
+		Sig.bReadFunction = true;
+		OutFunctions.Add(Sig);
+	}
+
 	{
 		FNiagaraFunctionSignature Sig;
 		Sig.Name = GetVector4AttributeIndexFunctionName;
@@ -1065,14 +1191,19 @@ void UNiagaraDataInterfaceGrid2DCollection::GetParameterDefinitionHLSL(const FNi
 	const int AttributeInt4Count = FMath::Max(1, FMath::DivideAndRoundUp(ParamInfo.GeneratedFunctions.Num(), 4));
 
 	TMap<FString, FStringFormatArg> ArgsDeclarations = {
+		{ TEXT("ParameterName"),		ParamInfo.DataInterfaceHLSLSymbol},
 		{ TEXT("GridName"),				ParamInfo.DataInterfaceHLSLSymbol + GridName},
+		{ TEXT("Grid"),					ParamInfo.DataInterfaceHLSLSymbol + GridName},
 		{ TEXT("SamplerName"),			ParamInfo.DataInterfaceHLSLSymbol + SamplerName},
 		{ TEXT("OutputGridName"),		ParamInfo.DataInterfaceHLSLSymbol + OutputGridName},
 		{ TEXT("AttributeIndicesName"),	ParamInfo.DataInterfaceHLSLSymbol + TEXT("_") + AttributeIndicesBaseName},
 		{ TEXT("AttributeInt4Count"),	AttributeInt4Count},
 		{ TEXT("NumAttributesName"),	ParamInfo.DataInterfaceHLSLSymbol + UNiagaraDataInterfaceRWBase::NumAttributesName},
+		{ TEXT("NumCellsName"),			ParamInfo.DataInterfaceHLSLSymbol + UNiagaraDataInterfaceRWBase::NumCellsName},
 	};
 	OutHLSL += FString::Format(FormatDeclarations, ArgsDeclarations);
+
+	AppendTemplateHLSL(OutHLSL, TemplateShaderFilePath, ArgsDeclarations);
 }
 
 void UNiagaraDataInterfaceGrid2DCollection::WriteSetHLSL(const FNiagaraDataInterfaceRWAttributeHelper& AttributeHelper, const FNiagaraDataInterfaceGPUParamInfo& ParamInfo, const FNiagaraDataInterfaceGeneratedFunction& FunctionInfo, int FunctionInstanceIndex, int32 InNumChannels, FString& OutHLSL)
@@ -1107,6 +1238,7 @@ void UNiagaraDataInterfaceGrid2DCollection::WriteSetHLSL(const FNiagaraDataInter
 	TMap<FString, FStringFormatArg> ArgsBounds = {
 		{TEXT("FunctionName"),				FunctionInfo.InstanceName},
 		{TEXT("OutputGrid"),				ParamInfo.DataInterfaceHLSLSymbol + OutputGridName},
+		{TEXT("Grid"),						ParamInfo.DataInterfaceHLSLSymbol + GridName},
 		{TEXT("NumCellsName"),				ParamInfo.DataInterfaceHLSLSymbol + UNiagaraDataInterfaceRWBase::NumCellsName},
 		{TEXT("UnitToUVName"),				ParamInfo.DataInterfaceHLSLSymbol + UNiagaraDataInterfaceRWBase::UnitToUVName},
 		{TEXT("AttributeIndicesName"),		ParamInfo.DataInterfaceHLSLSymbol + TEXT("_") + AttributeIndicesBaseName},
@@ -1163,8 +1295,15 @@ void UNiagaraDataInterfaceGrid2DCollection::WriteGetHLSL(const FNiagaraDataInter
 	OutHLSL += FString::Format(*FormatBounds, ArgsBounds);
 }
 
-void UNiagaraDataInterfaceGrid2DCollection::WriteSampleHLSL(const FNiagaraDataInterfaceRWAttributeHelper& AttributeHelper, const FNiagaraDataInterfaceGPUParamInfo& ParamInfo, const FNiagaraDataInterfaceGeneratedFunction& FunctionInfo, int FunctionInstanceIndex, int32 InNumChannels, FString& OutHLSL)
+void UNiagaraDataInterfaceGrid2DCollection::WriteSampleHLSL(const FNiagaraDataInterfaceRWAttributeHelper& AttributeHelper, const FNiagaraDataInterfaceGPUParamInfo& ParamInfo, const FNiagaraDataInterfaceGeneratedFunction& FunctionInfo, int FunctionInstanceIndex, int32 InNumChannels, bool IsCubic, FString& OutHLSL)
 {
+	FString SampleFunctionPlaceholder = "{SampleFunction}";
+	FString SampleFunction = "{Grid}.SampleLevel";
+	if (IsCubic)
+	{
+		SampleFunction = "SampleBiCubic_{ParameterName}";
+	}
+
 
 	FString FormatBounds = TEXT(R"(
 			void {FunctionName}(float2 In_Unit, out float{NumChannelsVariableSuffix} Out_Val)
@@ -1174,16 +1313,16 @@ void UNiagaraDataInterfaceGrid2DCollection::WriteSampleHLSL(const FNiagaraDataIn
 					)");
 	if (InNumChannels == 1)
 	{
-		FormatBounds += FString("	Out_Val = {Grid}.SampleLevel({SamplerName}, float3(In_Unit, In_AttributeIndex), 0);\n");
+		FormatBounds += FString("	Out_Val = {SampleFunction}({SamplerName}, float3(In_Unit, In_AttributeIndex), 0);\n");
 	}
 	else if (InNumChannels > 1)
 	{
 		for (int i = 0; i < InNumChannels; ++i)
 					{
 			FormatBounds += FString("			Out_Val.").Append(FNiagaraDataInterfaceRWAttributeHelper::Channels[i]).
-				Append(" = {Grid}.SampleLevel({SamplerName}, float3(In_Unit, In_AttributeIndex + ").Append(FString::FromInt(i)).Append("), 0); \n");
+				Append(" = {SampleFunction}({SamplerName}, float3(In_Unit, In_AttributeIndex + ").Append(FString::FromInt(i)).Append("), 0); \n");
 		}
-}
+	}
 
 	FormatBounds += TEXT(R"(
 				}
@@ -1192,11 +1331,14 @@ void UNiagaraDataInterfaceGrid2DCollection::WriteSampleHLSL(const FNiagaraDataIn
 	FString AttributeIndexPlaceholder = "{AttributeIndex}";
 	FString AttributeQueryString = GenerateAttributeIndexHLSL(AttributeHelper, FunctionInfo);
 	FormatBounds.ReplaceInline(*AttributeIndexPlaceholder, *AttributeQueryString);
+	FormatBounds.ReplaceInline(*SampleFunctionPlaceholder, *SampleFunction);
 
 	TMap<FString, FStringFormatArg> ArgsBounds = {
+		{TEXT("ParameterName"),				ParamInfo.DataInterfaceHLSLSymbol},
 		{TEXT("FunctionName"),				FunctionInfo.InstanceName},
 		{TEXT("Grid"),						ParamInfo.DataInterfaceHLSLSymbol + GridName},
 		{TEXT("SamplerName"),				ParamInfo.DataInterfaceHLSLSymbol + SamplerName},
+		{TEXT("SampleFunction"),			SampleFunction },
 		{TEXT("NumCellsName"),				ParamInfo.DataInterfaceHLSLSymbol + UNiagaraDataInterfaceRWBase::NumCellsName},
 		{TEXT("UnitToUVName"),				ParamInfo.DataInterfaceHLSLSymbol + UNiagaraDataInterfaceRWBase::UnitToUVName},
 		{TEXT("NumChannels"),				FString::FromInt(InNumChannels)},
@@ -1204,6 +1346,7 @@ void UNiagaraDataInterfaceGrid2DCollection::WriteSampleHLSL(const FNiagaraDataIn
 		{TEXT("AttributeIndicesName"),		ParamInfo.DataInterfaceHLSLSymbol + TEXT("_") + AttributeIndicesBaseName},
 		{TEXT("AttributeIndexGroup"),		FunctionInstanceIndex / 4},
 		{TEXT("AttributeIndexComponent"),	VectorComponentNames[FunctionInstanceIndex % 4]},
+
 	};
 	OutHLSL += FString::Format(*FormatBounds, ArgsBounds);
 }
@@ -1307,6 +1450,7 @@ bool UNiagaraDataInterfaceGrid2DCollection::GetFunctionHLSL(const FNiagaraDataIn
 
 	TMap<FString, FStringFormatArg> ArgsBounds =
 	{
+		{TEXT("ParameterName"),	ParamInfo.DataInterfaceHLSLSymbol},
 		{TEXT("FunctionName"),	FunctionInfo.InstanceName},
 		{TEXT("Grid"),			ParamInfo.DataInterfaceHLSLSymbol + GridName},
 		{TEXT("OutputGrid"),	ParamInfo.DataInterfaceHLSLSymbol + OutputGridName},
@@ -1392,7 +1536,12 @@ bool UNiagaraDataInterfaceGrid2DCollection::GetFunctionHLSL(const FNiagaraDataIn
 	}
 	else if (FunctionInfo.DefinitionName == SampleGridVector4FunctionName || FunctionInfo.DefinitionName == SamplePreviousGridVector4FunctionName)
 	{
-		WriteSampleHLSL(AttributeHelper, ParamInfo, FunctionInfo, FunctionInstanceIndex, 4, OutHLSL);
+		WriteSampleHLSL(AttributeHelper, ParamInfo, FunctionInfo, FunctionInstanceIndex, 4, false, OutHLSL);
+		return true;
+	}
+	else if (FunctionInfo.DefinitionName == CubicSamplePreviousGridVector4FunctionName)
+	{
+		WriteSampleHLSL(AttributeHelper, ParamInfo, FunctionInfo, FunctionInstanceIndex, 4, true, OutHLSL);
 		return true;
 	}
 	else if (FunctionInfo.DefinitionName == SetVector3ValueFunctionName || FunctionInfo.DefinitionName == SetVectorValueFunctionName)
@@ -1407,7 +1556,12 @@ bool UNiagaraDataInterfaceGrid2DCollection::GetFunctionHLSL(const FNiagaraDataIn
 	}
 	else if (FunctionInfo.DefinitionName == SampleGridVector3FunctionName || FunctionInfo.DefinitionName == SamplePreviousGridVectorFunctionName)
 	{
-		WriteSampleHLSL(AttributeHelper, ParamInfo, FunctionInfo, FunctionInstanceIndex, 3, OutHLSL);
+		WriteSampleHLSL(AttributeHelper, ParamInfo, FunctionInfo, FunctionInstanceIndex, 3, false, OutHLSL);
+		return true;
+	}
+	else if (FunctionInfo.DefinitionName == CubicSamplePreviousGridVectorFunctionName)
+	{
+		WriteSampleHLSL(AttributeHelper, ParamInfo, FunctionInfo, FunctionInstanceIndex, 3, true, OutHLSL);
 		return true;
 	}
 	else if (FunctionInfo.DefinitionName == SetVector2ValueFunctionName || FunctionInfo.DefinitionName == SetVector2DValueFunctionName)
@@ -1422,7 +1576,12 @@ bool UNiagaraDataInterfaceGrid2DCollection::GetFunctionHLSL(const FNiagaraDataIn
 	}
 	else if (FunctionInfo.DefinitionName == SampleGridVector2FunctionName || FunctionInfo.DefinitionName == SamplePreviousGridVector2DFunctionName)
 	{
-		WriteSampleHLSL(AttributeHelper, ParamInfo, FunctionInfo, FunctionInstanceIndex, 2, OutHLSL);
+		WriteSampleHLSL(AttributeHelper, ParamInfo, FunctionInfo, FunctionInstanceIndex, 2, false, OutHLSL);
+		return true;
+	}
+	else if (FunctionInfo.DefinitionName == CubicSamplePreviousGridVector2DFunctionName)
+	{
+		WriteSampleHLSL(AttributeHelper, ParamInfo, FunctionInfo, FunctionInstanceIndex, 2, true, OutHLSL);
 		return true;
 	}
 	else if (FunctionInfo.DefinitionName == SetFloatValueFunctionName)
@@ -1437,7 +1596,12 @@ bool UNiagaraDataInterfaceGrid2DCollection::GetFunctionHLSL(const FNiagaraDataIn
 	}
 	else if (FunctionInfo.DefinitionName == SampleGridFloatFunctionName || FunctionInfo.DefinitionName == SamplePreviousGridFloatFunctionName)
 	{
-		WriteSampleHLSL(AttributeHelper, ParamInfo, FunctionInfo, FunctionInstanceIndex, 1, OutHLSL);
+		WriteSampleHLSL(AttributeHelper, ParamInfo, FunctionInfo, FunctionInstanceIndex, 1, false, OutHLSL);
+		return true;
+	}
+	else if (FunctionInfo.DefinitionName == CubicSamplePreviousGridFloatFunctionName)
+	{
+		WriteSampleHLSL(AttributeHelper, ParamInfo, FunctionInfo, FunctionInstanceIndex, 1, true, OutHLSL);
 		return true;
 	}
 	else if (FunctionInfo.DefinitionName == GetVector4AttributeIndexFunctionName)
@@ -1472,6 +1636,18 @@ bool UNiagaraDataInterfaceGrid2DCollection::GetFunctionHLSL(const FNiagaraDataIn
 		OutHLSL += FString::Format(FormatBounds, ArgsBounds);
 		return true;
 	}
+	else if (FunctionInfo.DefinitionName == CubicSamplePreviousGridAtIndexFunctionName)
+	{
+		static const TCHAR* FormatBounds = TEXT(R"(
+				void {FunctionName}(float In_UnitX, float In_UnitY, int In_AttributeIndex, out float Out_Val)
+				{
+					float3 UVW = float3(In_UnitX, In_UnitY, In_AttributeIndex);
+					Out_Val = SampleBiCubic_{ParameterName}({SamplerName}, UVW, 0);
+				}
+			)");
+		OutHLSL += FString::Format(FormatBounds, ArgsBounds);
+		return true;
+	}
 	return false;
 }
 #endif
@@ -2542,7 +2718,7 @@ void FGrid2DCollectionRWInstanceData_RenderThread::BeginSimulate(FRDGBuilder& Gr
 		const FRDGTextureDesc TextureDesc = FRDGTextureDesc::Create2DArray(NumCells, PixelFormat.GetValue(), FClearValueBinding::Black, ETextureCreateFlags::ShaderResource | ETextureCreateFlags::UAV, uint16(NumAttributes));
 
 		const TCHAR* GridTextureName = TEXT("Grid2D::GridTexture");
-	#if 0
+	#if 1
 		TStringBuilder<128> StringBuilder;
 		SourceDIName.AppendString(StringBuilder);
 		StringBuilder.Append("_");
