UE package version incorporated into the DDC data for niagara scripts as we use standard struct/object serialization for the data

diff --git a/Engine/Plugins/FX/Niagara/Source/Niagara/Private/NiagaraScript.cpp b/Engine/Plugins/FX/Niagara/Source/Niagara/Private/NiagaraScript.cpp
--- a/Engine/Plugins/FX/Niagara/Source/Niagara/Private/NiagaraScript.cpp
+++ b/Engine/Plugins/FX/Niagara/Source/Niagara/Private/NiagaraScript.cpp
@@ -2736,7 +2736,20 @@ bool UNiagaraScript::BinaryToExecData(const UNiagaraScript* Script, const TArray
 	}
 
 	FMemoryReader Ar(InBinaryData, true);
+
+	// Read the archive version from the header of the payload so we can setup our reader with the right version
+	FPackageFileVersion ArchiveVersion;
+	Ar << ArchiveVersion;
+
+	if (!ArchiveVersion.IsCompatible(GOldestLoadablePackageFileUEVersion))
+	{
+		UE_LOG(LogNiagara, Display, TEXT("Failed to validate FNiagaraVMExecutableData received from DDC, rejecting!  Reasons:\nDeprecated object version"));
+		return false;
+	}
+
 	FObjectAndNameAsStringProxyArchive SafeAr(Ar, false);
+	SafeAr.SetUEVer(ArchiveVersion);
+
 	OutExecData.SerializeData(SafeAr, true);
 
 	FString ValidationErrors;
@@ -2766,6 +2779,11 @@ bool UNiagaraScript::ExecToBinaryData(const UNiagaraScript* Script, TArray<uint8
 	}
 
 	FMemoryWriter Ar(OutBinaryData, true);
+
+	// include the archive version into the payload since we're using struct serialization for the ExecData
+	FPackageFileVersion ArchiveVersion = GPackageFileUEVersion;
+	Ar << ArchiveVersion;
+
 	FObjectAndNameAsStringProxyArchive SafeAr(Ar, false);
 	InExecData.SerializeData(SafeAr, true);
 
diff --git a/Engine/Source/Runtime/Core/Private/UObject/FortniteReleaseBranchCustomObjectVersion.cpp b/Engine/Source/Runtime/Core/Private/UObject/FortniteReleaseBranchCustomObjectVersion.cpp
--- a/Engine/Source/Runtime/Core/Private/UObject/FortniteReleaseBranchCustomObjectVersion.cpp
+++ b/Engine/Source/Runtime/Core/Private/UObject/FortniteReleaseBranchCustomObjectVersion.cpp
@@ -6,7 +6,7 @@ TMap<FGuid, FGuid> FFortniteReleaseBranchCustomObjectVersion::GetSystemGuids()
 	TMap<FGuid, FGuid> SystemGuids;
 	const FDevSystemGuids& DevGuids = FDevSystemGuids::Get();
 
-	SystemGuids.Add(DevGuids.Niagara_LatestScriptCompileVersion, FGuid("CB0D96CCD66AC742ABB0744654ED5FCE"));
+	SystemGuids.Add(DevGuids.Niagara_LatestScriptCompileVersion, FGuid("C7E3EAD0DBD441B2911434B5CDA377D2"));
 	SystemGuids.Add(DevGuids.SkeletalMeshDerivedDataVersion, FGuid("5705956EC7134274A5A1BD3FC74DB3A9"));
 	return SystemGuids;
 }
