- Add FNiagaraSimCacheCompare which can be used to compare simulation caches
- New comparison allows us to exclude attributes and supply a custom tolerance when comparing caches
- Removed the old comparison method from the UNiagaraSimCache

diff --git a/Engine/Plugins/FX/Niagara/Source/Niagara/Classes/NiagaraSimCache.h b/Engine/Plugins/FX/Niagara/Source/Niagara/Classes/NiagaraSimCache.h
--- a/Engine/Plugins/FX/Niagara/Source/Niagara/Classes/NiagaraSimCache.h
+++ b/Engine/Plugins/FX/Niagara/Source/Niagara/Classes/NiagaraSimCache.h
@@ -231,6 +231,8 @@ struct FNiagaraSimCacheVariable
 
 	UPROPERTY()
 	uint16 Int32Count = 0;
+
+	NIAGARA_API bool operator==(const FNiagaraSimCacheVariable& Other) const;
 };
 
 USTRUCT()
@@ -371,6 +373,7 @@ UCLASS(Experimental, BlueprintType, MinimalAPI)
 class UNiagaraSimCache : public UObject
 {
 	friend struct FNiagaraSimCacheAttributeReaderHelper;
+	friend struct FNiagaraSimCacheCompare;
 	friend struct FNiagaraSimCacheHelper;
 	friend struct FNiagaraSimCacheGpuResource;
 
@@ -559,9 +562,6 @@ public:
 	*/
 	UFUNCTION(BlueprintCallable, Category = NiagaraSimCache)
 	NIAGARA_API void ReadQuatAttributeWithRebase(TArray<FQuat>& OutValues, FQuat Quat, FName AttributeName = FName("MeshOrientation"), FName EmitterName = NAME_None, int FrameIndex = 0) const;
-
-	UE_DEPRECATED(5.3, "For test purposes only, don't use in production code.")
-	NIAGARA_API bool IsDataEqual(const UNiagaraSimCache& OtherCache, float ErrorTolerance, FString& OutDifference) const;
 	
 private:
 	UPROPERTY(VisibleAnywhere, Category=SimCache)
diff --git a/Engine/Plugins/FX/Niagara/Source/Niagara/Classes/NiagaraSimCacheCompare.h b/Engine/Plugins/FX/Niagara/Source/Niagara/Classes/NiagaraSimCacheCompare.h
new file mode 100644
--- /dev/null
+++ b/Engine/Plugins/FX/Niagara/Source/Niagara/Classes/NiagaraSimCacheCompare.h
@@ -0,0 +1,46 @@
+// Copyright Epic Games, Inc. All Rights Reserved.
+
+#pragma  once
+
+#include "NiagaraSimCache.h"
+
+/**
+Allows you to compare two simulation caches for equality, or approximate equality.
+Due to platform differences simulations may not be equal when comparing one generates from one platform to another.
+Therefore you may need to be more flexible with some attributes, like positions, velocity, when comparing caches and adjust the tolerance accordingly.
+Note: This class is Experimental and the API may change
+*/
+struct FNiagaraSimCacheCompare
+{
+public:
+	/** Default float tolerance when comparing caches */
+	float DefaultFloatTolerance = 0.01f;
+
+	/** Should we exclude known none deterministic variables, i.e. ID / MaterialRandom. */
+	bool bExcludeNoneDeterministicVariables = true;
+
+	/** Add an attribute to exclude from the comparison. */
+	NIAGARA_API void AddExcludeAttribute(FNiagaraVariableBase Variable);
+
+	/** Setup comparison tolerances for variables that match the provided type. */
+	NIAGARA_API void SetToleranceByType(const UNiagaraSimCache& SimCache, const FNiagaraTypeDefinition& TypeDef, float Tolerance);
+
+	/** Setup comparison tolerances for variables that match the provided name wildcard. */
+	NIAGARA_API void SetToleranceByName(const UNiagaraSimCache& SimCache, const FString& NameWildcard, float Tolerance);
+
+	/**
+	Compare the provided caches.
+	returns true if they match, false if they do not.
+	If they do not match OutDifferences will contain a subset of the information.
+	*/
+	NIAGARA_API bool Compare(const UNiagaraSimCache& LhsCache, const UNiagaraSimCache& RhsCache, FString& OutDifferences);
+
+private:
+	bool CompareEmitter(const UNiagaraSimCache& LhsCache, const UNiagaraSimCache& Rhs, int EmitterIndex, int FrameIndex, FString& OutDifferences);
+
+private:
+	TSet<FNiagaraVariableBase>			AttributesToExclude_SystemScript;
+	TSet<FNiagaraVariableBase>			AttributesToExclude_ParticleScript;
+	TMap<FNiagaraVariableBase, float>	VariableToFloatTolerances_SystemScript;
+	TMap<FNiagaraVariableBase, float>	VariableToFloatTolerances_ParticleScript;
+};
diff --git a/Engine/Plugins/FX/Niagara/Source/Niagara/Private/NiagaraSimCache.cpp b/Engine/Plugins/FX/Niagara/Source/Niagara/Private/NiagaraSimCache.cpp
--- a/Engine/Plugins/FX/Niagara/Source/Niagara/Private/NiagaraSimCache.cpp
+++ b/Engine/Plugins/FX/Niagara/Source/Niagara/Private/NiagaraSimCache.cpp
@@ -24,6 +24,20 @@ UNiagaraSimCache::FOnCacheEndWrite		UNiagaraSimCache::OnCacheEndWrite;
 
 ////////////////////////////////////////////////////////////////////////////////////////////////////////////////////////////////////////////////////
 
+bool FNiagaraSimCacheVariable::operator==(const FNiagaraSimCacheVariable& Other) const
+{
+	return
+		Variable	== Other.Variable &&
+		FloatOffset == Other.FloatOffset &&
+		FloatCount	== Other.FloatCount &&
+		HalfOffset	== Other.HalfOffset &&
+		HalfCount	== Other.HalfCount &&
+		Int32Offset	== Other.Int32Offset &&
+		Int32Count	== Other.Int32Count;
+}
+
+////////////////////////////////////////////////////////////////////////////////////////////////////////////////////////////////////////////////////
+
 FNiagaraSimCacheFeedbackContext::~FNiagaraSimCacheFeedbackContext()
 {
 	if (bAutoLogIssues)
@@ -1320,250 +1334,3 @@ void UNiagaraSimCache::ReadQuatAttributeWithRebase(TArray<FQuat>& OutValues, FQu
 		}
 	}
 }
-
-namespace CacheCompare
-{
-	constexpr int32 ErrorStrLen = 8000;
-
-	bool IsEqual(const float& Num1, const float& Num2, float Tolerance)
-	{
-		return FMath::IsNearlyEqual(Num1, Num2, Tolerance);
-	}
-	
-	bool IsEqual(const int32& Num1, const int32& Num2, float)
-	{
-		return Num1 == Num2;
-	}
-	
-	FString NumberToString(const float& Num)
-	{
-		return FString::SanitizeFloat(Num);
-	}
-
-	FString NumberToString(const int32& Num)
-	{
-		return FString::FromInt(Num);
-	}
-
-	template<typename T>
-	bool CompareAttributeData(TArray<int32> OtherIDValues, TMap<int32, int32> ExpectedIDToRow, TArray<T> ExpectedData, TArray<T> OtherData, int32 DataCount, TStringBuilder<1024>& ValueDiffString, float ErrorTolerance)
-	{
-		if (DataCount == 0)
-		{
-			return true;
-		}
-		for (int Index = 0; Index < OtherIDValues.Num(); Index++) {
-			int32& OtherID = OtherIDValues[Index];
-			int32 ExpectedDataRow = ExpectedIDToRow.FindOrAdd(OtherID, INDEX_NONE);
-			for (int32 Offset = 0; Offset < DataCount; Offset++)
-			{
-				int32 OtherDataIndex = Index + Offset * OtherIDValues.Num();
-				int32 ExpectedDataIndex = ExpectedDataRow + Offset * OtherIDValues.Num();
-				if (ExpectedDataRow == INDEX_NONE || !OtherData.IsValidIndex(OtherDataIndex) || !ExpectedData.IsValidIndex(ExpectedDataIndex))
-				{
-					ValueDiffString.Append(TEXT("Unable to find source data for particle ID "));
-					ValueDiffString.Append(FString::FromInt(OtherID));
-					return false;
-				}
-				
-				if (!IsEqual(OtherData[OtherDataIndex], ExpectedData[ExpectedDataIndex], ErrorTolerance))
-				{
-					ValueDiffString.Append(TEXT("Particle ID "));
-					ValueDiffString.Append(FString::FromInt(OtherID));
-					ValueDiffString.Append(TEXT(", Expected: "));
-					ValueDiffString.Append(NumberToString(ExpectedData[ExpectedDataIndex]));
-					ValueDiffString.Append(TEXT(", Actual: "));
-					ValueDiffString.Append(NumberToString(OtherData[OtherDataIndex]));
-					return false;
-				}
-			}
-		}
-		return true;
-	}
-
-	void AddError(TStringBuilder<ErrorStrLen>& Builder, const FString& Str)
-	{
-		if (Builder.Len() < ErrorStrLen)
-		{
-			Builder.Append(Str);
-			if (Builder.Len() >= ErrorStrLen)
-			{
-				Builder.Append("...\n");
-			}
-		}
-	}
-}
-
-bool UNiagaraSimCache::IsDataEqual(const UNiagaraSimCache& OtherCache, float ErrorTolerance, FString& OutDifference) const
-{
-	const TArray<FNiagaraSimCacheFrame>& OtherFrames = OtherCache.CacheFrames;
-	if (GetNumFrames() != OtherCache.GetNumFrames())
-	{
-		OutDifference = FText::Format(FText::FromString("Cache Frame count different. Expected: {0}, Actual: {1}"), FText::AsNumber(GetNumFrames()), FText::AsNumber(OtherCache.GetNumFrames())).ToString();
-		return false;
-	}
-
-	bool bEqual = true;
-	TStringBuilder<CacheCompare::ErrorStrLen> Errors;
-	for (int FrameIndex = 0; FrameIndex < CacheFrames.Num(); FrameIndex++)
-	{
-		FString Frame = FString::FromInt(FrameIndex);
-		const FNiagaraSimCacheFrame& ExpectedFrame = CacheFrames[FrameIndex];
-		const FNiagaraSimCacheFrame& OtherFrame = OtherFrames[FrameIndex];
-
-		if (!FMath::IsNearlyEqual(OtherFrame.SimulationAge, ExpectedFrame.SimulationAge, UE_KINDA_SMALL_NUMBER))
-		{
-			CacheCompare::AddError(Errors, TEXT("Simulation age different - frame ") + Frame + '\n');
-			bEqual = false;
-		}
-		if (OtherFrame.EmitterData.Num() != ExpectedFrame.EmitterData.Num())
-		{
-			CacheCompare::AddError(Errors, TEXT("Emitter count different - frame ") + Frame + '\n');
-			bEqual = false;
-			continue;
-		}
-
-		// compare system data buffers
-		if (OtherFrame.SystemData.SystemDataBuffers.NumInstances != ExpectedFrame.SystemData.SystemDataBuffers.NumInstances)
-		{
-			CacheCompare::AddError(Errors, TEXT("SystemDataBuffers.NumInstances different - frame ") + Frame + TEXT(": expected ") + FString::FromInt(ExpectedFrame.SystemData.SystemDataBuffers.NumInstances) + TEXT(", got ") + FString::FromInt(OtherFrame.SystemData.SystemDataBuffers.NumInstances) + '\n');
-			bEqual = false;
-		}
-
-		if (OtherFrame.SystemData.SystemDataBuffers.Int32Data.Num() == ExpectedFrame.SystemData.SystemDataBuffers.Int32Data.Num())
-		{
-			if (FMemory::Memcmp(OtherFrame.SystemData.SystemDataBuffers.Int32Data.GetData(), ExpectedFrame.SystemData.SystemDataBuffers.Int32Data.GetData(), ExpectedFrame.SystemData.SystemDataBuffers.Int32Data.Num()) != 0)
-			{
-				CacheCompare::AddError(Errors, TEXT("SystemDataBuffers.Int32Data different - frame ") + Frame + '\n');
-				bEqual = false;
-			}
-		}
-		else
-		{
-			CacheCompare::AddError(Errors, TEXT("SystemDataBuffers.Int32Data count different - frame ") + Frame + '\n');
-			bEqual = false;
-		}
-
-		if (OtherFrame.SystemData.SystemDataBuffers.IDToIndexTable.Num() == ExpectedFrame.SystemData.SystemDataBuffers.IDToIndexTable.Num())
-		{
-			if (FMemory::Memcmp(OtherFrame.SystemData.SystemDataBuffers.IDToIndexTable.GetData(), ExpectedFrame.SystemData.SystemDataBuffers.IDToIndexTable.GetData(), ExpectedFrame.SystemData.SystemDataBuffers.IDToIndexTable.Num()) != 0)
-			{
-				CacheCompare::AddError(Errors, TEXT("SystemDataBuffers.IDToIndexTable different - frame ") + Frame + '\n');
-				bEqual = false;
-			}
-		}
-		else
-		{
-			CacheCompare::AddError(Errors, TEXT("SystemDataBuffers.IDToIndexTable count different - frame ") + Frame + '\n');
-			bEqual = false;
-		}
-
-		if (OtherFrame.SystemData.SystemDataBuffers.FloatData.Num() == ExpectedFrame.SystemData.SystemDataBuffers.FloatData.Num())
-		{
-			for (int i = 0; i < OtherFrame.SystemData.SystemDataBuffers.FloatData.Num(); i += sizeof(float))
-			{
-				float ExpectedVal = *reinterpret_cast<const float*>(ExpectedFrame.SystemData.SystemDataBuffers.FloatData.GetData() + i);
-				float OtherVal    = *reinterpret_cast<const float*>(OtherFrame.SystemData.SystemDataBuffers.FloatData.GetData() + i);
-				if (!FMath::IsNearlyEqual(OtherVal, ExpectedVal, ErrorTolerance))
-				{
-					CacheCompare::AddError(Errors, TEXT("SystemDataBuffers.FloatData different - frame ") + Frame + '\n');
-					bEqual = false;
-					break;
-				}
-			}
-		}
-		else
-		{
-			CacheCompare::AddError(Errors, TEXT("SystemDataBuffers.FloatData count different - frame ") + Frame + '\n');
-			bEqual = false;
-		}
-
-		if (OtherFrame.SystemData.SystemDataBuffers.HalfData.Num() == ExpectedFrame.SystemData.SystemDataBuffers.HalfData.Num())
-		{
-			for (int i = 0; i < OtherFrame.SystemData.SystemDataBuffers.HalfData.Num(); i += sizeof(FFloat16))
-			{
-				FFloat16 ExpectedVal = *reinterpret_cast<const FFloat16*>(ExpectedFrame.SystemData.SystemDataBuffers.HalfData.GetData() + i);
-				FFloat16 OtherVal    = *reinterpret_cast<const FFloat16*>(OtherFrame.SystemData.SystemDataBuffers.HalfData.GetData() + i);
-				if (!FMath::IsNearlyEqual(OtherVal, ExpectedVal, ErrorTolerance))
-				{
-					CacheCompare::AddError(Errors, TEXT("SystemDataBuffers.HalfData different - frame ") + Frame + '\n');
-					bEqual = false;
-					break;
-				}
-			}
-		}
-		else
-		{
-			CacheCompare::AddError(Errors, TEXT("SystemDataBuffers.HalfData count different - frame ") + Frame + '\n');
-			bEqual = false;
-		}
-
-		// compare emitter data
-		for (int EmitterIndex = 0; EmitterIndex < ExpectedFrame.EmitterData.Num(); EmitterIndex++)
-		{
-			FName EmitterName = CacheLayout.EmitterLayouts[EmitterIndex].LayoutName;
-			const FNiagaraSimCacheEmitterFrame& ExpectedEmitterFrame = ExpectedFrame.EmitterData[EmitterIndex];
-			const FNiagaraSimCacheEmitterFrame& OtherEmitterFrame = OtherFrame.EmitterData[EmitterIndex];
-			if (OtherEmitterFrame.ParticleDataBuffers.NumInstances != ExpectedEmitterFrame.ParticleDataBuffers.NumInstances)
-			{
-				CacheCompare::AddError(Errors, TEXT("Particle count different - frame ") + Frame + TEXT(" emitter ") + EmitterName.ToString() + '\n');
-				bEqual = false;
-				continue;
-			}
-
-			FName IdName("UniqueID");
-			TArray<int32> ExpectedIDValues;
-			TArray<int32> OtherIDValues;
-			ReadIntAttribute(ExpectedIDValues, IdName, EmitterName, FrameIndex);
-			OtherCache.ReadIntAttribute(OtherIDValues, IdName, EmitterName, FrameIndex);
-
-			if (ExpectedIDValues.Num() != OtherIDValues.Num() || ExpectedIDValues.Num() != ExpectedEmitterFrame.ParticleDataBuffers.NumInstances)
-			{
-				CacheCompare::AddError(Errors, TEXT("Invalid particle IDs - frame ") + Frame + TEXT(" emitter ") + EmitterName.ToString() + '\n');
-				bEqual = false;
-				continue;
-			}
-			if (ExpectedIDValues.Num() == 0)
-			{
-				continue;
-			}
-
-			TMap<int32, int32> ExpectedIDToRow;
-			for (int Index = 0; Index < ExpectedIDValues.Num(); Index++) {
-				ExpectedIDToRow.Add(ExpectedIDValues[Index], Index);
-			}
-
-			for (const FNiagaraSimCacheVariable& Var : CacheLayout.EmitterLayouts[EmitterIndex].Variables)
-			{
-				if (Var.Variable.GetName() == FName("ID") || Var.Variable.GetName() == FName("MaterialRandom"))
-				{
-					// persistent ID and material random are not deterministic, skip
-					continue;
-				}
-				TArray<float> ExpectedFloats;
-				TArray<float> OtherFloats;
-				TArray<FFloat16> ExpectedHalfs;
-				TArray<FFloat16> OtherHalfs;
-				TArray<int32> ExpectedInts;
-				TArray<int32> OtherInts;
-				ReadAttribute(ExpectedFloats, ExpectedHalfs, ExpectedInts, Var.Variable.GetName(), EmitterName, FrameIndex);
-				OtherCache.ReadAttribute(OtherFloats, OtherHalfs, OtherInts, Var.Variable.GetName(), EmitterName, FrameIndex);
-
-				TStringBuilder<1024> ValueDiff;
-				if (!CacheCompare::CompareAttributeData( OtherIDValues, ExpectedIDToRow, ExpectedFloats,OtherFloats, Var.FloatCount, ValueDiff, ErrorTolerance) ||
-					!CacheCompare::CompareAttributeData( OtherIDValues, ExpectedIDToRow, ExpectedHalfs,OtherHalfs, Var.HalfCount, ValueDiff, ErrorTolerance) ||
-					!CacheCompare::CompareAttributeData( OtherIDValues, ExpectedIDToRow, ExpectedInts,OtherInts, Var.Int32Count, ValueDiff, ErrorTolerance))
-				{
-					CacheCompare::AddError(Errors, TEXT("Different particle data - frame ") + Frame + TEXT(" emitter ") + EmitterName.ToString() + TEXT(" attribute ") + Var.Variable.GetName().ToString() + ": ");
-					CacheCompare::AddError(Errors, ValueDiff.ToString());
-					CacheCompare::AddError(Errors, "\n");
-					bEqual = false;
-				}
-			}
-		}
-	}
-	OutDifference = Errors.ToString();
-	return bEqual;
-}
-
-
diff --git a/Engine/Plugins/FX/Niagara/Source/Niagara/Private/NiagaraSimCacheCompare.cpp b/Engine/Plugins/FX/Niagara/Source/Niagara/Private/NiagaraSimCacheCompare.cpp
new file mode 100644
--- /dev/null
+++ b/Engine/Plugins/FX/Niagara/Source/Niagara/Private/NiagaraSimCacheCompare.cpp
@@ -0,0 +1,369 @@
+// Copyright Epic Games, Inc. All Rights Reserved.
+
+#include "NiagaraSimCacheCompare.h"
+#include "NiagaraConstants.h"
+#include "NiagaraModule.h"
+
+namespace NiagaraSimCacheCompareInternal
+{
+	constexpr int32 MaxErrorStrLen = 8*1024;
+
+	bool IsEqual(const float& Num1, const float& Num2, float Tolerance)
+	{
+		return FMath::IsNearlyEqual(Num1, Num2, Tolerance);
+	}
+
+	bool IsEqual(const int32& Num1, const int32& Num2, float)
+	{
+		return Num1 == Num2;
+	}
+
+	FString NumberToString(const float& Num)
+	{
+		return FString::SanitizeFloat(Num);
+	}
+
+	FString NumberToString(const int32& Num)
+	{
+		return FString::FromInt(Num);
+	}
+
+	TArray<FNiagaraSimCacheVariable> GatherEmitterVariables(const UNiagaraSimCache& SimCache, int32 EmitterIndex, const TSet<FNiagaraVariableBase>& ExcludeVariables)
+	{
+		TArray<FNiagaraSimCacheVariable> OutVariables;
+		SimCache.ForEachEmitterAttribute(
+			EmitterIndex,
+			[&OutVariables, &ExcludeVariables](const FNiagaraSimCacheVariable& Variable) -> bool
+			{
+				if (!ExcludeVariables.Contains(Variable.Variable))
+				{
+					OutVariables.Add(Variable);
+				}
+				return true;
+			}
+		);
+		return OutVariables;
+	}
+
+	template<typename TType, typename TValueString>
+	bool CompareAttributeData(const TArray<int32>& InstanceMapping, const TArray<TType>& LhsData, const TArray<TType>& RhsData, int32 ComponentCount, TValueString& ValueDiffString, float ErrorTolerance)
+	{
+		const int32 NumInstances = InstanceMapping.Num();
+		for (int32 i = 0; i < NumInstances; ++i)
+		{
+			const int32 LhsInstance = i;
+			const int32 RhsInstance = InstanceMapping[i];
+
+			for (int32 Component = 0; Component < ComponentCount; ++Component)
+			{
+				const int32 LhsComponent = LhsInstance + (Component * NumInstances);
+				const int32 RhsComponent = RhsInstance + (Component * NumInstances);
+
+				if (!LhsData.IsValidIndex(LhsComponent) || !RhsData.IsValidIndex(RhsComponent))
+				{
+					ValueDiffString.Append(TEXT("Unable to find source data for particle ID "));
+					ValueDiffString.Append(FString::FromInt(LhsInstance));
+					return false;
+				}
+
+				const TType LhsValue = LhsData[LhsComponent];
+				const TType RhsValue = RhsData[RhsComponent];
+				if (!IsEqual(LhsValue, RhsValue, ErrorTolerance))
+				{
+					ValueDiffString.Append(TEXT("Particle ID "));
+					ValueDiffString.Append(FString::FromInt(LhsInstance));
+					ValueDiffString.Append(TEXT(", Expected: "));
+					ValueDiffString.Append(NumberToString(LhsValue));
+					ValueDiffString.Append(TEXT(", Actual: "));
+					ValueDiffString.Append(NumberToString(RhsValue));
+					return false;
+				}
+			}
+		}
+
+		return true;
+	}
+
+	void AddError(FString& OutErrors, const FString& Str)
+	{
+		if (OutErrors.Len() < MaxErrorStrLen)
+		{
+			OutErrors.Append(Str);
+			if (OutErrors.Len() >= MaxErrorStrLen)
+			{
+				OutErrors.Append(TEXT("...\n"));
+			}
+			else
+			{
+				OutErrors.Append(TEXT("\n"));
+			}
+		}
+	}
+
+	void InternalSetToleranceByType(const UNiagaraSimCache& SimCache, const FNiagaraTypeDefinition& TypeDef, float Tolerance, int EmitterIndex, TMap<FNiagaraVariableBase, float>& VariableToFloatTolerances)
+	{
+		SimCache.ForEachEmitterAttribute(
+			EmitterIndex,
+			[&TypeDef, Tolerance, &VariableToFloatTolerances](const FNiagaraSimCacheVariable& CacheVariable) -> bool
+			{
+				if (CacheVariable.Variable.GetType() == TypeDef)
+				{
+					VariableToFloatTolerances.FindOrAdd(CacheVariable.Variable) = Tolerance;
+				}
+				return true;
+			}
+		);
+	}
+
+	void InternalSetToleranceByName(const UNiagaraSimCache& SimCache, const FString& NameWildcard, float Tolerance, int EmitterIndex, TMap<FNiagaraVariableBase, float>& VariableToFloatTolerances)
+	{
+		SimCache.ForEachEmitterAttribute(
+			EmitterIndex,
+			[&NameWildcard, Tolerance, &VariableToFloatTolerances](const FNiagaraSimCacheVariable& CacheVariable) -> bool
+			{
+				if (CacheVariable.Variable.GetName().ToString().MatchesWildcard(NameWildcard))
+				{
+					VariableToFloatTolerances.FindOrAdd(CacheVariable.Variable) = Tolerance;
+				}
+				return true;
+			}
+		);
+	}
+}
+
+void FNiagaraSimCacheCompare::AddExcludeAttribute(FNiagaraVariableBase Variable)
+{
+	if (Variable.RemoveRootNamespace(FNiagaraConstants::ParticleAttributeNamespaceString))
+	{
+		AttributesToExclude_ParticleScript.Add(Variable);
+	}
+	else
+	{
+		AttributesToExclude_SystemScript.Add(Variable);
+	}
+}
+
+void FNiagaraSimCacheCompare::SetToleranceByType(const UNiagaraSimCache& SimCache, const FNiagaraTypeDefinition& TypeDef, float Tolerance)
+{
+	using namespace NiagaraSimCacheCompareInternal;
+
+	InternalSetToleranceByType(SimCache, TypeDef, Tolerance, INDEX_NONE, VariableToFloatTolerances_SystemScript);
+	for (int32 i=0; i < SimCache.GetNumEmitters(); ++i)
+	{
+		InternalSetToleranceByType(SimCache, TypeDef, Tolerance, i, VariableToFloatTolerances_ParticleScript);
+	}
+}
+
+void FNiagaraSimCacheCompare::SetToleranceByName(const UNiagaraSimCache& SimCache, const FString& NameWildcard, float Tolerance)
+{
+	using namespace NiagaraSimCacheCompareInternal;
+
+	InternalSetToleranceByName(SimCache, NameWildcard, Tolerance, INDEX_NONE, VariableToFloatTolerances_SystemScript);
+	for (int32 i = 0; i < SimCache.GetNumEmitters(); ++i)
+	{
+		InternalSetToleranceByName(SimCache, NameWildcard, Tolerance, i, VariableToFloatTolerances_ParticleScript);
+	}
+}
+
+bool FNiagaraSimCacheCompare::Compare(const UNiagaraSimCache& LhsCache, const UNiagaraSimCache& RhsCache, FString& OutDifferences)
+{
+	using namespace NiagaraSimCacheCompareInternal;
+
+	// Check Frame Count Matches
+	if (LhsCache.GetNumFrames() != RhsCache.GetNumFrames())
+	{
+		OutDifferences = FString::Printf(TEXT("Cache Frame count is different. Expected: %d, Actual: %d"), LhsCache.GetNumFrames(), RhsCache.GetNumFrames());
+		return false;
+	}
+
+	// Check Emitter Count Matches
+	if (LhsCache.GetNumEmitters() != RhsCache.GetNumEmitters())
+	{
+		OutDifferences = FString::Printf(TEXT("Cache Emitter count is different. Expected: %d, Actual: %d"), LhsCache.GetNumEmitters(), RhsCache.GetNumEmitters());
+		return false;
+	}
+
+	// Check Emitter Names Match
+	for (int iEmitter=0; iEmitter < LhsCache.GetNumEmitters(); ++iEmitter)
+	{
+		if (LhsCache.GetEmitterName(iEmitter) != RhsCache.GetEmitterName(iEmitter))
+		{
+			OutDifferences = FString::Printf(TEXT("Cache Emitter name is different. Expected: %s, Actual: %s"), *LhsCache.GetEmitterName(iEmitter).ToString(), *RhsCache.GetEmitterName(iEmitter).ToString());
+			return false;
+		}
+	}
+
+	// Add exclude variables
+	if (bExcludeNoneDeterministicVariables)
+	{
+		for (FNiagaraVariableBase ParticleVariable : {SYS_PARAM_PARTICLES_ID, SYS_PARAM_PARTICLES_MATERIAL_RANDOM})
+		{
+			ParticleVariable.RemoveRootNamespace(FNiagaraConstants::ParticleAttributeNamespaceString);
+			AttributesToExclude_ParticleScript.Add(ParticleVariable);
+		}
+	}
+
+	// Check Frames Match
+	OutDifferences.Reset(MaxErrorStrLen);
+
+	bool bEqual = true;
+	for (int FrameIndex=0; FrameIndex < LhsCache.GetNumFrames(); FrameIndex++)
+	{
+		const FNiagaraSimCacheFrame& LhsFrame = LhsCache.CacheFrames[FrameIndex];
+		const FNiagaraSimCacheFrame& RhsFrame = RhsCache.CacheFrames[FrameIndex];
+
+		if (!FMath::IsNearlyEqual(LhsFrame.SimulationAge, RhsFrame.SimulationAge, UE_KINDA_SMALL_NUMBER))
+		{
+			AddError(OutDifferences, FString::Printf(TEXT("Simulation age different - frame %d"), FrameIndex));
+			bEqual = false;
+		}
+
+		bEqual &= CompareEmitter(LhsCache, RhsCache, INDEX_NONE, FrameIndex, OutDifferences);
+		for (int EmitterIndex=0; EmitterIndex < LhsCache.GetNumEmitters(); EmitterIndex++)
+		{
+			bEqual &= CompareEmitter(LhsCache, RhsCache, EmitterIndex, FrameIndex, OutDifferences);
+		}
+
+		if (bEqual == false)
+		{
+			break;
+		}
+	}
+
+	return bEqual;
+}
+
+bool FNiagaraSimCacheCompare::CompareEmitter(const UNiagaraSimCache& LhsCache, const UNiagaraSimCache& RhsCache, int EmitterIndex, int FrameIndex, FString& OutDifferences)
+{
+	using namespace NiagaraSimCacheCompareInternal;
+
+	const FName EmitterName = EmitterIndex == INDEX_NONE ? NAME_None : LhsCache.GetEmitterName(EmitterIndex);
+	const FString EmitterNameString = EmitterIndex == INDEX_NONE ? FString(TEXT("SystemInstance")) : EmitterName.ToString();
+
+	const TSet<FNiagaraVariableBase>& AttributesToExclude = EmitterIndex == INDEX_NONE ? AttributesToExclude_SystemScript : AttributesToExclude_ParticleScript;
+	TMap<FNiagaraVariableBase, float>& VariableToFloatTolerances = EmitterIndex == INDEX_NONE ? VariableToFloatTolerances_SystemScript : VariableToFloatTolerances_ParticleScript;
+
+	// Compare instance counts
+	const int LhsNumInstances = LhsCache.GetEmitterNumInstances(EmitterIndex, FrameIndex);
+	const int RhsNumInstances = RhsCache.GetEmitterNumInstances(EmitterIndex, FrameIndex);
+	if (LhsNumInstances != RhsNumInstances)
+	{
+		AddError(OutDifferences, FString::Printf(TEXT("%s instance count different frame %d.  Expected: %d, Actual %d"), *EmitterNameString, FrameIndex, LhsNumInstances, RhsNumInstances));
+		return false;
+	}
+
+	if (LhsNumInstances == 0)
+	{
+		return true;
+	}
+
+	// Gather all the attributes we want to compare
+	const TArray<FNiagaraSimCacheVariable> CompareVariables = GatherEmitterVariables(LhsCache, EmitterIndex, AttributesToExclude);
+	{
+		bool bAttributesValid = true;
+		const TArray<FNiagaraSimCacheVariable> RhsCompareVariables = GatherEmitterVariables(RhsCache, EmitterIndex, AttributesToExclude);
+		if (CompareVariables.Num() != RhsCompareVariables.Num())
+		{
+			AddError(OutDifferences, FString::Printf(TEXT("%s attribute count different frame %d.  Expected: %d, Actual %d"), *EmitterNameString, FrameIndex, CompareVariables.Num(), RhsCompareVariables.Num()));
+			bAttributesValid = false;
+		}
+
+		auto CompareVariableFunc = 
+			[&OutDifferences, &EmitterNameString](const TArray<FNiagaraSimCacheVariable>& LhsVariables, const TArray<FNiagaraSimCacheVariable>& RhsVariables) -> bool
+			{
+				bool IsValid = true;
+				for (const FNiagaraSimCacheVariable& Var : LhsVariables)
+				{
+					if (!RhsVariables.Contains(Var))
+					{
+						AddError(OutDifferences, FString::Printf(TEXT("%s attribute %s does not exist in both caches."), *EmitterNameString, *Var.Variable.GetName().ToString()));
+						IsValid = false;
+					}
+				}
+				return IsValid;
+			};
+		bAttributesValid &= CompareVariableFunc(CompareVariables, RhsCompareVariables);
+		bAttributesValid &= CompareVariableFunc(RhsCompareVariables, CompareVariables);
+
+		if (!bAttributesValid)
+		{
+			return false;
+		}
+	}
+
+	// Attempt to build an instance mapping between the two caches as ordering might be inconsistent due to wave size on GPU, etc, otherwise we will have to assume consistent ordering
+	TArray<int32> InstanceMapping;
+	InstanceMapping.AddZeroed(LhsNumInstances);
+	{
+		const FName IDAttributeName("UniqueID");
+
+		TArray<int32> LhsIDValues;
+		TArray<int32> RhsIDValues;
+		LhsCache.ReadIntAttribute(LhsIDValues, IDAttributeName, EmitterName, FrameIndex);
+		RhsCache.ReadIntAttribute(RhsIDValues, IDAttributeName, EmitterName, FrameIndex);
+		if (LhsIDValues.Num() > 0 || RhsIDValues.Num() > 0)
+		{
+			if (LhsIDValues.Num() != RhsIDValues.Num() || LhsIDValues.Num() != LhsNumInstances)
+			{
+				AddError(OutDifferences, FString::Printf(TEXT("%s has invalid particle IDs - frame %d."), *EmitterNameString, FrameIndex));
+				return false;
+			}
+
+			bool bIDsValid = true;
+			for (int32 i=0; i < LhsNumInstances; ++i)
+			{
+				const int32 LhsIDValue = LhsIDValues[i];
+				const int32 RhsIndex = RhsIDValues.IndexOfByKey(LhsIDValue);
+				if (RhsIndex == INDEX_NONE)
+				{
+					AddError(OutDifferences, FString::Printf(TEXT("%s has inconsistent particle IDs - frame %d.  ID %d not found."), *EmitterNameString, FrameIndex, LhsIDValue));
+					bIDsValid = false;
+				}
+				InstanceMapping[i] = RhsIndex;
+			}
+			if (!bIDsValid)
+			{
+				return false;
+			}
+		}
+	}
+
+	if (InstanceMapping.Num() == 0)
+	{
+		for (int32 i=0; i < LhsNumInstances; ++i)
+		{
+			InstanceMapping[i] = i;
+		}
+	}
+
+	bool bAttributesEqual = true;
+	for (const FNiagaraSimCacheVariable& CompareVariable : CompareVariables)
+	{
+		const FName AtributeName = CompareVariable.Variable.GetName();
+		
+		TArray<float> LhsFloats;
+		TArray<float> RhsFloats;
+		TArray<FFloat16> LhsHalfs;
+		TArray<FFloat16> RhsHalfs;
+		TArray<int32> LhsInts;
+		TArray<int32> RhsInts;
+		LhsCache.ReadAttribute(LhsFloats, LhsHalfs, LhsInts, AtributeName, EmitterName, FrameIndex);
+		RhsCache.ReadAttribute(RhsFloats, RhsHalfs, RhsInts, AtributeName, EmitterName, FrameIndex);
+
+		const float ErrorTolerance = VariableToFloatTolerances.FindRef(CompareVariable.Variable, DefaultFloatTolerance);
+
+		TStringBuilder<1024> ValueDiff;
+		if (!CompareAttributeData(InstanceMapping, LhsFloats, RhsFloats, CompareVariable.FloatCount, ValueDiff, ErrorTolerance) ||
+			!CompareAttributeData(InstanceMapping, LhsHalfs,  RhsHalfs,  CompareVariable.HalfCount,  ValueDiff, ErrorTolerance) ||
+			!CompareAttributeData(InstanceMapping, LhsInts,   RhsInts,   CompareVariable.Int32Count, ValueDiff, ErrorTolerance))
+		{
+			AddError(
+				OutDifferences,
+				FString::Printf(TEXT("%s different particle data - frame %d attribute %s - %s"), *EmitterNameString, FrameIndex, *AtributeName.ToString(), ValueDiff.ToString())
+			);
+			bAttributesEqual = false;
+		}
+	}
+
+	return bAttributesEqual;
+}
