Non unity build fix

diff --git a/Engine/Plugins/FX/Niagara/Source/Niagara/Private/NiagaraRendererProperties.cpp b/Engine/Plugins/FX/Niagara/Source/Niagara/Private/NiagaraRendererProperties.cpp
--- a/Engine/Plugins/FX/Niagara/Source/Niagara/Private/NiagaraRendererProperties.cpp
+++ b/Engine/Plugins/FX/Niagara/Source/Niagara/Private/NiagaraRendererProperties.cpp
@@ -15,6 +15,7 @@
 #include "Materials/MaterialInterface.h"
 #include "Materials/MaterialInstanceConstant.h"
 #include "SceneInterface.h"
+#include "SceneUtils.h"
 #include "Styling/SlateIconFinder.h"
 
 #include UE_INLINE_GENERATED_CPP_BY_NAME(NiagaraRendererProperties)
