- Fix for emitter mode approximate motion vectors being incorrect

diff --git a/Engine/Plugins/FX/Niagara/Source/Niagara/Private/NiagaraRendererMeshes.cpp b/Engine/Plugins/FX/Niagara/Source/Niagara/Private/NiagaraRendererMeshes.cpp
--- a/Engine/Plugins/FX/Niagara/Source/Niagara/Private/NiagaraRendererMeshes.cpp
+++ b/Engine/Plugins/FX/Niagara/Source/Niagara/Private/NiagaraRendererMeshes.cpp
@@ -1210,8 +1210,18 @@ void FNiagaraRendererMeshes::SetupElementForGPUScene(
 		
 		if (bNeedsPrevTransform)
 		{
-			const FQuat PrevRot = FQuat(CommonParameters.DefaultPrevRotation.X, CommonParameters.DefaultPrevRotation.Y, CommonParameters.DefaultPrevRotation.Z, CommonParameters.DefaultPrevRotation.W).GetNormalized();
-			FMatrix PrevLocalToPrimitive = FTransform(PrevRot, FVector(CommonParameters.DefaultPrevPosition), FVector(CommonParameters.DefaultPrevScale)).ToMatrixWithScale();
+			FQuat PrevRot;
+			FMatrix PrevLocalToPrimitive;
+			if (CommonParameters.AccurateMotionVectors)
+			{
+				PrevRot = FQuat(CommonParameters.DefaultPrevRotation.X, CommonParameters.DefaultPrevRotation.Y, CommonParameters.DefaultPrevRotation.Z, CommonParameters.DefaultPrevRotation.W).GetNormalized();
+				PrevLocalToPrimitive = FTransform(PrevRot, FVector(CommonParameters.DefaultPrevPosition), FVector(CommonParameters.DefaultPrevScale)).ToMatrixWithScale();
+			}
+			else
+			{
+				PrevRot = FQuat(CommonParameters.DefaultRotation.X, CommonParameters.DefaultRotation.Y, CommonParameters.DefaultRotation.Z, CommonParameters.DefaultRotation.W).GetNormalized();
+				PrevLocalToPrimitive = FTransform(Rot, FVector(CommonParameters.DefaultPosition), FVector(CommonParameters.DefaultScale)).ToMatrixWithScale();
+			}
 
 			if (!bLocalSpace)
 			{
