GetAssetRegistryTags deprecation: Make the old GetAssetRegistryTags and related functions deprecated.

diff --git a/Engine/Plugins/Experimental/VirtualCamera/Source/VirtualCamera/Private/LevelSequence/VirtualCameraClipsMetaData.cpp b/Engine/Plugins/Experimental/VirtualCamera/Source/VirtualCamera/Private/LevelSequence/VirtualCameraClipsMetaData.cpp
--- a/Engine/Plugins/Experimental/VirtualCamera/Source/VirtualCamera/Private/LevelSequence/VirtualCameraClipsMetaData.cpp
+++ b/Engine/Plugins/Experimental/VirtualCamera/Source/VirtualCamera/Private/LevelSequence/VirtualCameraClipsMetaData.cpp
@@ -2,6 +2,8 @@
 
 #include "LevelSequence/VirtualCameraClipsMetaData.h"
 
+#include "UObject/AssetRegistryTagsContext.h"
+
 const FName UVirtualCameraClipsMetaData::AssetRegistryTag_FocalLength = "ClipsMetaData_FocalLength";
 const FName UVirtualCameraClipsMetaData::AssetRegistryTag_bIsSelected = "ClipsMetaData_bIsSelected";
 const FName UVirtualCameraClipsMetaData::AssetRegistryTag_RecordedLevelName = "ClipsMetaData_RecordedLevelName";
@@ -41,23 +43,25 @@ TSet<FName> UVirtualCameraClipsMetaData::GetAllClipsMetaDataTags()
 	};
 }
 
-void UVirtualCameraClipsMetaData::ExtendAssetRegistryTags(TArray<FAssetRegistryTag>& OutTags) const
+void UVirtualCameraClipsMetaData::ExtendAssetRegistryTags(FAssetRegistryTagsContext Context) const
 {
+	IMovieSceneMetaDataInterface::ExtendAssetRegistryTags(Context);
+
 	FString IsSelected = bIsSelected ? "true" : "false";
 	FString IsRecordedFromACineCameraActor = bIsACineCameraRecording ? "true" : "false";
 
-	OutTags.Emplace(AssetRegistryTag_FocalLength, FString::FromInt(FocalLength), FAssetRegistryTag::ETagType::TT_Numerical, FAssetRegistryTag::TD_None);
-	OutTags.Emplace(AssetRegistryTag_bIsSelected, IsSelected, FAssetRegistryTag::ETagType::TT_Alphabetical, FAssetRegistryTag::TD_None);
-	OutTags.Emplace(AssetRegistryTag_RecordedLevelName, RecordedLevelName, FAssetRegistryTag::ETagType::TT_Alphabetical, FAssetRegistryTag::TD_None);
-	OutTags.Emplace(AssetRegistryTag_FrameCountStart, FString::FromInt(FrameCountStart), FAssetRegistryTag::ETagType::TT_Numerical, FAssetRegistryTag::TD_None);
-	OutTags.Emplace(AssetRegistryTag_FrameCountEnd, FString::FromInt(FrameCountEnd), FAssetRegistryTag::ETagType::TT_Numerical, FAssetRegistryTag::TD_None);
-	OutTags.Emplace(AssetRegistryTag_LengthInFrames, FString::FromInt(LengthInFrames), FAssetRegistryTag::ETagType::TT_Numerical, FAssetRegistryTag::TD_None);
-	OutTags.Emplace(AssetRegistryTag_DisplayRate, DisplayRate.ToPrettyText().ToString(), FAssetRegistryTag::ETagType::TT_Alphabetical, FAssetRegistryTag::TD_None);
-	OutTags.Emplace(AssetRegistryTag_bIsACineCameraRecording, IsRecordedFromACineCameraActor, FAssetRegistryTag::ETagType::TT_Alphabetical, FAssetRegistryTag::TD_None);
-	OutTags.Emplace(AssetRegistryTag_bIsNoGood, FString::FromInt(bIsNoGood), FAssetRegistryTag::ETagType::TT_Numerical, FAssetRegistryTag::TD_None);
-	OutTags.Emplace(AssetRegistryTag_bIsFlagged, FString::FromInt(bIsFlagged), FAssetRegistryTag::ETagType::TT_Numerical, FAssetRegistryTag::TD_None);
-	OutTags.Emplace(AssetRegistryTag_FavoriteLevel, FString::FromInt(FavoriteLevel), FAssetRegistryTag::ETagType::TT_Alphabetical, FAssetRegistryTag::TD_None);
-	OutTags.Emplace(AssetRegistryTag_bIsCreatedFromVCam, FString::FromInt(bIsCreatedFromVCam), FAssetRegistryTag::ETagType::TT_Numerical, FAssetRegistryTag::TD_None);
+	Context.AddTag(FAssetRegistryTag(AssetRegistryTag_FocalLength, FString::FromInt(FocalLength), FAssetRegistryTag::ETagType::TT_Numerical, FAssetRegistryTag::TD_None));
+	Context.AddTag(FAssetRegistryTag(AssetRegistryTag_bIsSelected, IsSelected, FAssetRegistryTag::ETagType::TT_Alphabetical, FAssetRegistryTag::TD_None));
+	Context.AddTag(FAssetRegistryTag(AssetRegistryTag_RecordedLevelName, RecordedLevelName, FAssetRegistryTag::ETagType::TT_Alphabetical, FAssetRegistryTag::TD_None));
+	Context.AddTag(FAssetRegistryTag(AssetRegistryTag_FrameCountStart, FString::FromInt(FrameCountStart), FAssetRegistryTag::ETagType::TT_Numerical, FAssetRegistryTag::TD_None));
+	Context.AddTag(FAssetRegistryTag(AssetRegistryTag_FrameCountEnd, FString::FromInt(FrameCountEnd), FAssetRegistryTag::ETagType::TT_Numerical, FAssetRegistryTag::TD_None));
+	Context.AddTag(FAssetRegistryTag(AssetRegistryTag_LengthInFrames, FString::FromInt(LengthInFrames), FAssetRegistryTag::ETagType::TT_Numerical, FAssetRegistryTag::TD_None));
+	Context.AddTag(FAssetRegistryTag(AssetRegistryTag_DisplayRate, DisplayRate.ToPrettyText().ToString(), FAssetRegistryTag::ETagType::TT_Alphabetical, FAssetRegistryTag::TD_None));
+	Context.AddTag(FAssetRegistryTag(AssetRegistryTag_bIsACineCameraRecording, IsRecordedFromACineCameraActor, FAssetRegistryTag::ETagType::TT_Alphabetical, FAssetRegistryTag::TD_None));
+	Context.AddTag(FAssetRegistryTag(AssetRegistryTag_bIsNoGood, FString::FromInt(bIsNoGood), FAssetRegistryTag::ETagType::TT_Numerical, FAssetRegistryTag::TD_None));
+	Context.AddTag(FAssetRegistryTag(AssetRegistryTag_bIsFlagged, FString::FromInt(bIsFlagged), FAssetRegistryTag::ETagType::TT_Numerical, FAssetRegistryTag::TD_None));
+	Context.AddTag(FAssetRegistryTag(AssetRegistryTag_FavoriteLevel, FString::FromInt(FavoriteLevel), FAssetRegistryTag::ETagType::TT_Alphabetical, FAssetRegistryTag::TD_None));
+	Context.AddTag(FAssetRegistryTag(AssetRegistryTag_bIsCreatedFromVCam, FString::FromInt(bIsCreatedFromVCam), FAssetRegistryTag::ETagType::TT_Numerical, FAssetRegistryTag::TD_None));
 }
 
 #if WITH_EDITOR
diff --git a/Engine/Plugins/Experimental/VirtualCamera/Source/VirtualCamera/Public/LevelSequence/VirtualCameraClipsMetaData.h b/Engine/Plugins/Experimental/VirtualCamera/Source/VirtualCamera/Public/LevelSequence/VirtualCameraClipsMetaData.h
--- a/Engine/Plugins/Experimental/VirtualCamera/Source/VirtualCamera/Public/LevelSequence/VirtualCameraClipsMetaData.h
+++ b/Engine/Plugins/Experimental/VirtualCamera/Source/VirtualCamera/Public/LevelSequence/VirtualCameraClipsMetaData.h
@@ -118,7 +118,9 @@ public:
 	/**
 	* Extend the default ULevelSequence asset registry tags
 	*/
-	virtual void ExtendAssetRegistryTags(TArray<FAssetRegistryTag>& OutTags) const override;
+	virtual void ExtendAssetRegistryTags(FAssetRegistryTagsContext Context) const override;
+	UE_DEPRECATED(5.4, "Implement the version that takes FAssetRegistryTagsContext instead.")
+	virtual void ExtendAssetRegistryTags(TArray<FAssetRegistryTag>& OutTags) const override {}
 #if WITH_EDITOR
 	virtual void ExtendAssetRegistryTagMetaData(TMap<FName, FAssetRegistryTagMetadata>& OutMetadata) const override;
 #endif
diff --git a/Engine/Plugins/FX/Niagara/Source/Niagara/Classes/NiagaraSystem.h b/Engine/Plugins/FX/Niagara/Source/Niagara/Classes/NiagaraSystem.h
--- a/Engine/Plugins/FX/Niagara/Source/Niagara/Classes/NiagaraSystem.h
+++ b/Engine/Plugins/FX/Niagara/Source/Niagara/Classes/NiagaraSystem.h
@@ -376,7 +376,9 @@ public:
 
 	FORCEINLINE bool HasFixedTickDelta() const { return bFixedTickDelta; }
 	FORCEINLINE float GetFixedTickDeltaTime()const { return FixedTickDeltaTime; }
-	NIAGARA_API virtual void GetAssetRegistryTags(TArray<FAssetRegistryTag>& OutTags)  const override;
+	NIAGARA_API virtual void GetAssetRegistryTags(FAssetRegistryTagsContext Context) const override;
+	UE_DEPRECATED(5.4, "Implement the version that takes FAssetRegistryTagsContext instead.")
+	NIAGARA_API virtual void GetAssetRegistryTags(TArray<FAssetRegistryTag>& OutTags) const override;
 
 	FORCEINLINE bool NeedsDeterminism() const { return bDeterminism; }
 	FORCEINLINE int32 GetRandomSeed() const { return RandomSeed; }
diff --git a/Engine/Plugins/FX/Niagara/Source/Niagara/Private/NiagaraEmitter.cpp b/Engine/Plugins/FX/Niagara/Source/Niagara/Private/NiagaraEmitter.cpp
--- a/Engine/Plugins/FX/Niagara/Source/Niagara/Private/NiagaraEmitter.cpp
+++ b/Engine/Plugins/FX/Niagara/Source/Niagara/Private/NiagaraEmitter.cpp
@@ -26,6 +26,7 @@
 #include "Interfaces/ITargetPlatform.h"
 #include "Modules/ModuleManager.h"
 #include "Templates/Greater.h"
+#include "UObject/AssetRegistryTagsContext.h"
 #include "UObject/UE5MainStreamObjectVersion.h"
 #include "UObject/LinkerLoad.h"
 #include "UObject/ObjectSaveContext.h"
@@ -904,17 +905,24 @@ bool UNiagaraEmitter::IsEditorOnly() const
 
 
 void UNiagaraEmitter::GetAssetRegistryTags(TArray<FAssetRegistryTag>& OutTags) const
+{
+	PRAGMA_DISABLE_DEPRECATION_WARNINGS;
+	Super::GetAssetRegistryTags(OutTags);
+	PRAGMA_ENABLE_DEPRECATION_WARNINGS;
+}
+
+void UNiagaraEmitter::GetAssetRegistryTags(FAssetRegistryTagsContext Context) const
 {
 #if WITH_EDITOR
-	OutTags.Add(FAssetRegistryTag("VersioningEnabled", bVersioningEnabled ? TEXT("True") : TEXT("False"), FAssetRegistryTag::TT_Alphabetical));
+	Context.AddTag(FAssetRegistryTag("VersioningEnabled", bVersioningEnabled ? TEXT("True") : TEXT("False"), FAssetRegistryTag::TT_Alphabetical));
 
 	const FVersionedNiagaraEmitterData* AssetData = GetLatestEmitterData();
 	FVersionedNiagaraEmitterData DefaultData;
 	const FVersionedNiagaraEmitterData& EmitterData = AssetData ? *AssetData : DefaultData; // the CDO does not have any version data, so just use the default struct values in that case 
-	OutTags.Add(FAssetRegistryTag("HasGPUEmitter", ( EmitterData.SimTarget == ENiagaraSimTarget::GPUComputeSim) ? TEXT("True") : TEXT("False"), FAssetRegistryTag::TT_Alphabetical));
+	Context.AddTag(FAssetRegistryTag("HasGPUEmitter", ( EmitterData.SimTarget == ENiagaraSimTarget::GPUComputeSim) ? TEXT("True") : TEXT("False"), FAssetRegistryTag::TT_Alphabetical));
 
 	const float BoundsSize =  float(EmitterData.FixedBounds.GetSize().GetMax());
-	OutTags.Add(FAssetRegistryTag("FixedBoundsSize",  EmitterData.CalculateBoundsMode == ENiagaraEmitterCalculateBoundMode::Fixed ? FString::Printf(TEXT("%.2f"), BoundsSize) : FString(TEXT("None")), FAssetRegistryTag::TT_Numerical));
+	Context.AddTag(FAssetRegistryTag("FixedBoundsSize",  EmitterData.CalculateBoundsMode == ENiagaraEmitterCalculateBoundMode::Fixed ? FString::Printf(TEXT("%.2f"), BoundsSize) : FString(TEXT("None")), FAssetRegistryTag::TT_Numerical));
 
 
 	uint32 NumActiveRenderers = 0;
@@ -929,7 +937,7 @@ void UNiagaraEmitter::GetAssetRegistryTags(TArray<FAssetRegistryTag>& OutTags) c
 		}
 	}
 
-	OutTags.Add(FAssetRegistryTag("ActiveRenderers", LexToString(NumActiveRenderers), FAssetRegistryTag::TT_Numerical));
+	Context.AddTag(FAssetRegistryTag("ActiveRenderers", LexToString(NumActiveRenderers), FAssetRegistryTag::TT_Numerical));
 
 	// Gather up NumActive emitters based off of quality level.
 	if (const UNiagaraSettings* Settings = GetDefault<UNiagaraSettings>())
@@ -950,7 +958,7 @@ void UNiagaraEmitter::GetAssetRegistryTags(TArray<FAssetRegistryTag>& OutTags) c
 		for (int32 i = 0; i < NumQualityLevels; i++)
 		{
 			FString QualityLevelKey = Settings->QualityLevels[i].ToString() + TEXT("Emitters");
-			OutTags.Add(FAssetRegistryTag(*QualityLevelKey, LexToString(QualityLevelsNumActive[i]), FAssetRegistryTag::TT_Numerical));
+			Context.AddTag(FAssetRegistryTag(*QualityLevelKey, LexToString(QualityLevelsNumActive[i]), FAssetRegistryTag::TT_Numerical));
 		}
 	}
 
@@ -1005,7 +1013,7 @@ void UNiagaraEmitter::GetAssetRegistryTags(TArray<FAssetRegistryTag>& OutTags) c
 				PropDefault->GetAssetTagsForContext(this, EmitterData.Version.VersionGuid, DataInterfaces, NumericKeys, StringKeys);
 			}
 		}
-		OutTags.Add(FAssetRegistryTag("ActiveDIs", LexToString(DataInterfaces.Num()), FAssetRegistryTag::TT_Numerical));
+		Context.AddTag(FAssetRegistryTag("ActiveDIs", LexToString(DataInterfaces.Num()), FAssetRegistryTag::TT_Numerical));
 	}
 
 
@@ -1014,7 +1022,7 @@ void UNiagaraEmitter::GetAssetRegistryTags(TArray<FAssetRegistryTag>& OutTags) c
 	while (NumericIter)
 	{
 
-		OutTags.Add(FAssetRegistryTag(NumericIter.Key(), LexToString(NumericIter.Value()), FAssetRegistryTag::TT_Numerical));
+		Context.AddTag(FAssetRegistryTag(NumericIter.Key(), LexToString(NumericIter.Value()), FAssetRegistryTag::TT_Numerical));
 		++NumericIter;
 	}
 
@@ -1022,19 +1030,19 @@ void UNiagaraEmitter::GetAssetRegistryTags(TArray<FAssetRegistryTag>& OutTags) c
 	while (StringIter)
 	{
 
-		OutTags.Add(FAssetRegistryTag(StringIter.Key(), LexToString(StringIter.Value()), FAssetRegistryTag::TT_Alphabetical));
+		Context.AddTag(FAssetRegistryTag(StringIter.Key(), LexToString(StringIter.Value()), FAssetRegistryTag::TT_Alphabetical));
 		++StringIter;
 	}
 
 	// TemplateSpecialization
 	FName TemplateSpecificationName = GET_MEMBER_NAME_CHECKED(UNiagaraEmitter, TemplateSpecification);
 	FText TemplateSpecializationValueString = StaticEnum<ENiagaraScriptTemplateSpecification>()->GetDisplayNameTextByValue((int64) TemplateSpecification);
-	OutTags.Add(FAssetRegistryTag(TemplateSpecificationName, TemplateSpecializationValueString.ToString(), FAssetRegistryTag::TT_Alphabetical));
+	Context.AddTag(FAssetRegistryTag(TemplateSpecificationName, TemplateSpecializationValueString.ToString(), FAssetRegistryTag::TT_Alphabetical));
 	
 	INiagaraModule& NiagaraModule = FModuleManager::GetModuleChecked<INiagaraModule>("Niagara");
-	OutTags.Add(NiagaraModule.GetEditorOnlyDataUtilities().CreateClassUsageAssetRegistryTag(this));
+	Context.AddTag(NiagaraModule.GetEditorOnlyDataUtilities().CreateClassUsageAssetRegistryTag(this));
 #endif
-	Super::GetAssetRegistryTags(OutTags);
+	Super::GetAssetRegistryTags(Context);
 }
 
 bool UNiagaraEmitter::NeedsLoadForTargetPlatform(const ITargetPlatform* TargetPlatform)const
diff --git a/Engine/Plugins/FX/Niagara/Source/Niagara/Private/NiagaraScript.cpp b/Engine/Plugins/FX/Niagara/Source/Niagara/Private/NiagaraScript.cpp
--- a/Engine/Plugins/FX/Niagara/Source/Niagara/Private/NiagaraScript.cpp
+++ b/Engine/Plugins/FX/Niagara/Source/Niagara/Private/NiagaraScript.cpp
@@ -29,6 +29,7 @@
 #include "Serialization/MemoryReader.h"
 #include "Serialization/ObjectAndNameAsStringProxyArchive.h"
 #include "Stats/Stats.h"
+#include "UObject/AssetRegistryTagsContext.h"
 #include "UObject/EditorObjectVersion.h"
 #include "UObject/FortniteMainBranchObjectVersion.h"
 #include "UObject/ObjectSaveContext.h"
@@ -3373,7 +3374,14 @@ void UNiagaraScript::RaiseOnGPUCompilationComplete()
 
 void UNiagaraScript::GetAssetRegistryTags(TArray<FAssetRegistryTag>& OutTags) const
 {
+	PRAGMA_DISABLE_DEPRECATION_WARNINGS;
 	Super::GetAssetRegistryTags(OutTags);
+	PRAGMA_ENABLE_DEPRECATION_WARNINGS;
+}
+
+void UNiagaraScript::GetAssetRegistryTags(FAssetRegistryTagsContext Context) const
+{
+	Super::GetAssetRegistryTags(Context);
 #if WITH_EDITORONLY_DATA
 
 	if (CustomAssetRegistryTagCache.IsSet() == false)
@@ -3398,7 +3406,7 @@ void UNiagaraScript::GetAssetRegistryTags(TArray<FAssetRegistryTag>& OutTags) co
 			}
 		}
 
-		OutTags.Add(FAssetRegistryTag(ProvidedDependenciesName, *ProvidedDependenciesTags, FAssetRegistryTag::TT_Hidden));
+		Context.AddTag(FAssetRegistryTag(ProvidedDependenciesName, *ProvidedDependenciesTags, FAssetRegistryTag::TT_Hidden));
 	}
 
 	// Category
@@ -3409,7 +3417,7 @@ void UNiagaraScript::GetAssetRegistryTags(TArray<FAssetRegistryTag>& OutTags) co
 		FString& CategoryTag = CustomAssetRegistryTagCache->FindOrAdd(CategoryName);
 		CategoryTag = CategoryText.ToString();
 
-		OutTags.Add(FAssetRegistryTag(CategoryName, CategoryTag, FAssetRegistryTag::TT_Alphabetical));
+		Context.AddTag(FAssetRegistryTag(CategoryName, CategoryTag, FAssetRegistryTag::TT_Alphabetical));
 	}
 
 	// Description
@@ -3420,7 +3428,7 @@ void UNiagaraScript::GetAssetRegistryTags(TArray<FAssetRegistryTag>& OutTags) co
 		FString& DescriptionTag = CustomAssetRegistryTagCache->FindOrAdd(DescriptionName);
 		DescriptionTag = DescriptionText.ToString();
 
-		OutTags.Add(FAssetRegistryTag(DescriptionName, DescriptionTag, FAssetRegistryTag::TT_Alphabetical));
+		Context.AddTag(FAssetRegistryTag(DescriptionName, DescriptionTag, FAssetRegistryTag::TT_Alphabetical));
 	}
 
 	// Keywords
@@ -3431,7 +3439,7 @@ void UNiagaraScript::GetAssetRegistryTags(TArray<FAssetRegistryTag>& OutTags) co
 		FString& KeywordsTag = CustomAssetRegistryTagCache->FindOrAdd(KeywordsName);
 		KeywordsTag = KeywordsText.ToString();
 
-		OutTags.Add(FAssetRegistryTag(KeywordsName, KeywordsTag, FAssetRegistryTag::TT_Alphabetical));
+		Context.AddTag(FAssetRegistryTag(KeywordsName, KeywordsTag, FAssetRegistryTag::TT_Alphabetical));
 	}
 
 	// Visibility
@@ -3440,36 +3448,36 @@ void UNiagaraScript::GetAssetRegistryTags(TArray<FAssetRegistryTag>& OutTags) co
 	FString& VisibilityTag = CustomAssetRegistryTagCache->FindOrAdd(VisibilityName);
 	UEnum* VisibilityEnum = StaticEnum<ENiagaraScriptLibraryVisibility>();
 	VisibilityTag = VisibilityEnum->GetNameStringByValue((int64)Visibility);
-	OutTags.Add(FAssetRegistryTag(VisibilityName, VisibilityTag, FAssetRegistryTag::TT_Alphabetical));
+	Context.AddTag(FAssetRegistryTag(VisibilityName, VisibilityTag, FAssetRegistryTag::TT_Alphabetical));
 
 	// Usage bitmask
 	int32 UsageBitmask = ScriptData ? ScriptData->ModuleUsageBitmask : ModuleUsageBitmask_DEPRECATED;
 	FName UsageBitmaskName = GET_MEMBER_NAME_CHECKED(FVersionedNiagaraScriptData, ModuleUsageBitmask);
 	FString& UsageBitmaskTag = CustomAssetRegistryTagCache->FindOrAdd(UsageBitmaskName);
 	UsageBitmaskTag = FString::FromInt(UsageBitmask);
-	OutTags.Add(FAssetRegistryTag(UsageBitmaskName, UsageBitmaskTag, FAssetRegistryTag::TT_Hidden));
+	Context.AddTag(FAssetRegistryTag(UsageBitmaskName, UsageBitmaskTag, FAssetRegistryTag::TT_Hidden));
 
 	// Deprecation
 	bool bDeprecated = ScriptData ? ScriptData->bDeprecated : bDeprecated_DEPRECATED;
 	FName DeprecatedName = GET_MEMBER_NAME_CHECKED(FVersionedNiagaraScriptData, bDeprecated);
 	FString& DeprecatedTag = CustomAssetRegistryTagCache->FindOrAdd(DeprecatedName);
 	DeprecatedTag = FString::FromInt(bDeprecated);
-	OutTags.Add(FAssetRegistryTag(DeprecatedName, DeprecatedTag, FAssetRegistryTag::TT_Hidden));
+	Context.AddTag(FAssetRegistryTag(DeprecatedName, DeprecatedTag, FAssetRegistryTag::TT_Hidden));
 
 	// Suggested
 	bool bSuggested = ScriptData ? ScriptData->bSuggested : false;
 	FName SuggestedName = GET_MEMBER_NAME_CHECKED(FVersionedNiagaraScriptData, bSuggested);
 	FString& SuggestedTag = CustomAssetRegistryTagCache->FindOrAdd(SuggestedName);
 	SuggestedTag = FString::FromInt(bSuggested);
-	OutTags.Add(FAssetRegistryTag(SuggestedName, SuggestedTag, FAssetRegistryTag::TT_Hidden));
+	Context.AddTag(FAssetRegistryTag(SuggestedName, SuggestedTag, FAssetRegistryTag::TT_Hidden));
 
 	// Add the current custom version to the tags so that tags can be fixed up in the future without having to load
 	// the whole asset.
 	const int32 NiagaraVer = GetLinkerCustomVersion(FNiagaraCustomVersion::GUID);
-	OutTags.Add(FAssetRegistryTag(NiagaraCustomVersionTagName, FString::FromInt(NiagaraVer), FAssetRegistryTag::TT_Hidden));
+	Context.AddTag(FAssetRegistryTag(NiagaraCustomVersionTagName, FString::FromInt(NiagaraVer), FAssetRegistryTag::TT_Hidden));
 
 	INiagaraModule& NiagaraModule = FModuleManager::GetModuleChecked<INiagaraModule>("Niagara");
-	OutTags.Add(NiagaraModule.GetEditorOnlyDataUtilities().CreateClassUsageAssetRegistryTag(this));
+	Context.AddTag(NiagaraModule.GetEditorOnlyDataUtilities().CreateClassUsageAssetRegistryTag(this));
 #endif
 }
 
diff --git a/Engine/Plugins/FX/Niagara/Source/Niagara/Private/NiagaraSystem.cpp b/Engine/Plugins/FX/Niagara/Source/Niagara/Private/NiagaraSystem.cpp
--- a/Engine/Plugins/FX/Niagara/Source/Niagara/Private/NiagaraSystem.cpp
+++ b/Engine/Plugins/FX/Niagara/Source/Niagara/Private/NiagaraSystem.cpp
@@ -38,6 +38,7 @@
 #include "Modules/ModuleManager.h"
 #include "ProfilingDebugging/CookStats.h"
 #include "ProfilingDebugging/ScopedTimers.h"
+#include "UObject/AssetRegistryTagsContext.h"
 #include "UObject/ObjectSaveContext.h"
 #include "UObject/LinkerLoad.h"
 #include "UObject/Package.h"
@@ -1383,6 +1384,13 @@ bool UNiagaraSystem::IsReadyToRun() const
 }
 
 void UNiagaraSystem::GetAssetRegistryTags(TArray<FAssetRegistryTag>& OutTags) const
+{
+	PRAGMA_DISABLE_DEPRECATION_WARNINGS;
+	Super::GetAssetRegistryTags(OutTags);
+	PRAGMA_ENABLE_DEPRECATION_WARNINGS;
+}
+
+void UNiagaraSystem::GetAssetRegistryTags(FAssetRegistryTagsContext Context) const
 {
 	if (HasAnyFlags(RF_ClassDefaultObject) == false)
 	{
@@ -1390,12 +1398,12 @@ void UNiagaraSystem::GetAssetRegistryTags(TArray<FAssetRegistryTag>& OutTags) co
 	}
 
 #if WITH_EDITOR
-	OutTags.Add(FAssetRegistryTag("HasGPUEmitter", HasAnyGPUEmitters() ? TEXT("True") : TEXT("False"), FAssetRegistryTag::TT_Alphabetical));
+	Context.AddTag(FAssetRegistryTag("HasGPUEmitter", HasAnyGPUEmitters() ? TEXT("True") : TEXT("False"), FAssetRegistryTag::TT_Alphabetical));
 
 	const float BoundsSize = float(FixedBounds.GetSize().GetMax());
-	OutTags.Add(FAssetRegistryTag("FixedBoundsSize", bFixedBounds ? FString::Printf(TEXT("%.2f"), BoundsSize) : FString(TEXT("None")), FAssetRegistryTag::TT_Numerical));
+	Context.AddTag(FAssetRegistryTag("FixedBoundsSize", bFixedBounds ? FString::Printf(TEXT("%.2f"), BoundsSize) : FString(TEXT("None")), FAssetRegistryTag::TT_Numerical));
 
-	OutTags.Add(FAssetRegistryTag("NumEmitters", LexToString(EmitterHandles.Num()), FAssetRegistryTag::TT_Numerical));
+	Context.AddTag(FAssetRegistryTag("NumEmitters", LexToString(EmitterHandles.Num()), FAssetRegistryTag::TT_Numerical));
 
 	uint32 GPUSimsMissingFixedBounds = 0;
 
@@ -1429,14 +1437,14 @@ void UNiagaraSystem::GetAssetRegistryTags(TArray<FAssetRegistryTag>& OutTags) co
 		}
 	}
 
-	OutTags.Add(FAssetRegistryTag("ActiveEmitters", LexToString(NumActiveEmitters), FAssetRegistryTag::TT_Numerical));
-	OutTags.Add(FAssetRegistryTag("ActiveRenderers", LexToString(NumActiveRenderers), FAssetRegistryTag::TT_Numerical));
-	OutTags.Add(FAssetRegistryTag("GPUSimsMissingFixedBounds", LexToString(GPUSimsMissingFixedBounds), FAssetRegistryTag::TT_Numerical));
-	OutTags.Add(FAssetRegistryTag("EffectType", EffectType != nullptr ? EffectType->GetName() : FString(TEXT("None")), FAssetRegistryTag::TT_Alphabetical));
-	OutTags.Add(FAssetRegistryTag("WarmupTime", LexToString(WarmupTime), FAssetRegistryTag::TT_Numerical));
-	OutTags.Add(FAssetRegistryTag("HasOverrideScalabilityForSystem", bOverrideScalabilitySettings ? TEXT("True") : TEXT("False"), FAssetRegistryTag::TT_Alphabetical));
-	OutTags.Add(FAssetRegistryTag("HasDIsWithPostSimulateTick", bHasDIsWithPostSimulateTick ? TEXT("True") : TEXT("False"), FAssetRegistryTag::TT_Alphabetical));
-	OutTags.Add(FAssetRegistryTag("NeedsSortedSignificanceCull", bNeedsSortedSignificanceCull ? TEXT("True") : TEXT("False"), FAssetRegistryTag::TT_Alphabetical));
+	Context.AddTag(FAssetRegistryTag("ActiveEmitters", LexToString(NumActiveEmitters), FAssetRegistryTag::TT_Numerical));
+	Context.AddTag(FAssetRegistryTag("ActiveRenderers", LexToString(NumActiveRenderers), FAssetRegistryTag::TT_Numerical));
+	Context.AddTag(FAssetRegistryTag("GPUSimsMissingFixedBounds", LexToString(GPUSimsMissingFixedBounds), FAssetRegistryTag::TT_Numerical));
+	Context.AddTag(FAssetRegistryTag("EffectType", EffectType != nullptr ? EffectType->GetName() : FString(TEXT("None")), FAssetRegistryTag::TT_Alphabetical));
+	Context.AddTag(FAssetRegistryTag("WarmupTime", LexToString(WarmupTime), FAssetRegistryTag::TT_Numerical));
+	Context.AddTag(FAssetRegistryTag("HasOverrideScalabilityForSystem", bOverrideScalabilitySettings ? TEXT("True") : TEXT("False"), FAssetRegistryTag::TT_Alphabetical));
+	Context.AddTag(FAssetRegistryTag("HasDIsWithPostSimulateTick", bHasDIsWithPostSimulateTick ? TEXT("True") : TEXT("False"), FAssetRegistryTag::TT_Alphabetical));
+	Context.AddTag(FAssetRegistryTag("NeedsSortedSignificanceCull", bNeedsSortedSignificanceCull ? TEXT("True") : TEXT("False"), FAssetRegistryTag::TT_Alphabetical));
 
 	// Gather up NumActive emitters based off of quality level.
 	const UNiagaraSettings* Settings = GetDefault<UNiagaraSettings>();
@@ -1468,7 +1476,7 @@ void UNiagaraSystem::GetAssetRegistryTags(TArray<FAssetRegistryTag>& OutTags) co
 		for (int32 i = 0; i < NumQualityLevels; i++)
 		{
 			FString QualityLevelKey = Settings->QualityLevels[i].ToString() + TEXT("Emitters");
-			OutTags.Add(FAssetRegistryTag(*QualityLevelKey, LexToString(QualityLevelsNumActive[i]), FAssetRegistryTag::TT_Numerical));
+			Context.AddTag(FAssetRegistryTag(*QualityLevelKey, LexToString(QualityLevelsNumActive[i]), FAssetRegistryTag::TT_Numerical));
 		}
 	}
 
@@ -1535,7 +1543,7 @@ void UNiagaraSystem::GetAssetRegistryTags(TArray<FAssetRegistryTag>& OutTags) co
 				PropDefault->GetAssetTagsForContext(this, FGuid(), DataInterfaces, NumericKeys, StringKeys);
 			}
 		}
-		OutTags.Add(FAssetRegistryTag("ActiveDIs", LexToString(DataInterfaces.Num()), FAssetRegistryTag::TT_Numerical));
+		Context.AddTag(FAssetRegistryTag("ActiveDIs", LexToString(DataInterfaces.Num()), FAssetRegistryTag::TT_Numerical));
 	}
 
 
@@ -1544,7 +1552,7 @@ void UNiagaraSystem::GetAssetRegistryTags(TArray<FAssetRegistryTag>& OutTags) co
 	while (NumericIter)
 	{
 
-		OutTags.Add(FAssetRegistryTag(NumericIter.Key(), LexToString(NumericIter.Value()), FAssetRegistryTag::TT_Numerical));
+		Context.AddTag(FAssetRegistryTag(NumericIter.Key(), LexToString(NumericIter.Value()), FAssetRegistryTag::TT_Numerical));
 		++NumericIter;
 	}
 
@@ -1552,33 +1560,33 @@ void UNiagaraSystem::GetAssetRegistryTags(TArray<FAssetRegistryTag>& OutTags) co
 	while (StringIter)
 	{
 
-		OutTags.Add(FAssetRegistryTag(StringIter.Key(), LexToString(StringIter.Value()), FAssetRegistryTag::TT_Alphabetical));
+		Context.AddTag(FAssetRegistryTag(StringIter.Key(), LexToString(StringIter.Value()), FAssetRegistryTag::TT_Alphabetical));
 		++StringIter;
 	}
 
 	// TemplateSpecialization
 	FName TemplateSpecificationName = GET_MEMBER_NAME_CHECKED(UNiagaraSystem, TemplateSpecification);
 	FText TemplateSpecializationValueString = StaticEnum<ENiagaraScriptTemplateSpecification>()->GetDisplayNameTextByValue((int64) TemplateSpecification);
-	OutTags.Add(FAssetRegistryTag(TemplateSpecificationName, TemplateSpecializationValueString.ToString(), FAssetRegistryTag::TT_Alphabetical));
+	Context.AddTag(FAssetRegistryTag(TemplateSpecificationName, TemplateSpecializationValueString.ToString(), FAssetRegistryTag::TT_Alphabetical));
 
 	/*for (const UNiagaraDataInterface* DI : DataInterfaces)
 	{
 		FString ClassName;
 		DI->GetClass()->GetName(ClassName);
-		OutTags.Add(FAssetRegistryTag(*(TEXT("bHas")+ClassName), TEXT("True"), FAssetRegistryTag::TT_Alphabetical));
+		Context.AddTag(FAssetRegistryTag(*(TEXT("bHas")+ClassName), TEXT("True"), FAssetRegistryTag::TT_Alphabetical));
 	}*/
 
 
-	//OutTags.Add(FAssetRegistryTag("CPUCollision", UsesCPUCollision() ? TEXT("True") : TEXT("False"), FAssetRegistryTag::TT_Alphabetical));
-	//OutTags.Add(FAssetRegistryTag("Looping", bAnyEmitterLoopsForever ? TEXT("True") : TEXT("False"), FAssetRegistryTag::TT_Alphabetical));
-	//OutTags.Add(FAssetRegistryTag("Immortal", IsImmortal() ? TEXT("True") : TEXT("False"), FAssetRegistryTag::TT_Alphabetical));
-	//OutTags.Add(FAssetRegistryTag("Becomes Zombie", WillBecomeZombie() ? TEXT("True") : TEXT("False"), FAssetRegistryTag::TT_Alphabetical));
-	//OutTags.Add(FAssetRegistryTag("CanBeOccluded", OcclusionBoundsMethod == EParticleSystemOcclusionBoundsMethod::EPSOBM_None ? TEXT("False") : TEXT("True"), FAssetRegistryTag::TT_Alphabetical));
+	//Context.AddTag(FAssetRegistryTag("CPUCollision", UsesCPUCollision() ? TEXT("True") : TEXT("False"), FAssetRegistryTag::TT_Alphabetical));
+	//Context.AddTag(FAssetRegistryTag("Looping", bAnyEmitterLoopsForever ? TEXT("True") : TEXT("False"), FAssetRegistryTag::TT_Alphabetical));
+	//Context.AddTag(FAssetRegistryTag("Immortal", IsImmortal() ? TEXT("True") : TEXT("False"), FAssetRegistryTag::TT_Alphabetical));
+	//Context.AddTag(FAssetRegistryTag("Becomes Zombie", WillBecomeZombie() ? TEXT("True") : TEXT("False"), FAssetRegistryTag::TT_Alphabetical));
+	//Context.AddTag(FAssetRegistryTag("CanBeOccluded", OcclusionBoundsMethod == EParticleSystemOcclusionBoundsMethod::EPSOBM_None ? TEXT("False") : TEXT("True"), FAssetRegistryTag::TT_Alphabetical));
 
 	INiagaraModule& NiagaraModule = FModuleManager::GetModuleChecked<INiagaraModule>("Niagara");
-	OutTags.Add(NiagaraModule.GetEditorOnlyDataUtilities().CreateClassUsageAssetRegistryTag(this));
+	Context.AddTag(NiagaraModule.GetEditorOnlyDataUtilities().CreateClassUsageAssetRegistryTag(this));
 #endif
-	Super::GetAssetRegistryTags(OutTags);
+	Super::GetAssetRegistryTags(Context);
 }
 
 #if WITH_EDITORONLY_DATA
diff --git a/Engine/Plugins/Interchange/Runtime/Source/FactoryNodes/Private/InterchangeSkeletalMeshFactoryNode.cpp b/Engine/Plugins/Interchange/Runtime/Source/FactoryNodes/Private/InterchangeSkeletalMeshFactoryNode.cpp
--- a/Engine/Plugins/Interchange/Runtime/Source/FactoryNodes/Private/InterchangeSkeletalMeshFactoryNode.cpp
+++ b/Engine/Plugins/Interchange/Runtime/Source/FactoryNodes/Private/InterchangeSkeletalMeshFactoryNode.cpp
@@ -4,6 +4,7 @@
 
 #include "Engine/SkeletalMesh.h"
 #include "Engine/SkinnedAssetCommon.h"
+#include "UObject/AssetRegistryTagsContext.h"
 
 #include UE_INLINE_GENERATED_CPP_BY_NAME(InterchangeSkeletalMeshFactoryNode)
 
@@ -244,7 +245,14 @@ bool UInterchangeSkeletalMeshFactoryNode::SetCustomImportContentType(const EInte
 
 void UInterchangeSkeletalMeshFactoryNode::AppendAssetRegistryTags(TArray<FAssetRegistryTag>& OutTags) const
 {
+	PRAGMA_DISABLE_DEPRECATION_WARNINGS;
 	Super::AppendAssetRegistryTags(OutTags);
+	PRAGMA_ENABLE_DEPRECATION_WARNINGS;
+}
+
+void UInterchangeSkeletalMeshFactoryNode::AppendAssetRegistryTags(FAssetRegistryTagsContext Context) const
+{
+	Super::AppendAssetRegistryTags(Context);
 #if WITH_EDITORONLY_DATA
 	EInterchangeSkeletalMeshContentType ContentType;
 	if (GetCustomImportContentType(ContentType))
@@ -264,7 +272,7 @@ void UInterchangeSkeletalMeshFactoryNode::AppendAssetRegistryTags(TArray<FAssetR
 		};
 
 		FString EnumString = ImportContentTypeToString(ContentType);
-		OutTags.Add(FAssetRegistryTag(NSSkeletalMeshSourceFileLabels::GetSkeletalMeshLastImportContentTypeMetadataKey(), EnumString, FAssetRegistryTag::TT_Hidden));
+		Context.AddTag(FAssetRegistryTag(NSSkeletalMeshSourceFileLabels::GetSkeletalMeshLastImportContentTypeMetadataKey(), EnumString, FAssetRegistryTag::TT_Hidden));
 	}
 #endif
 }
diff --git a/Engine/Plugins/Interchange/Runtime/Source/FactoryNodes/Public/InterchangeSkeletalMeshFactoryNode.h b/Engine/Plugins/Interchange/Runtime/Source/FactoryNodes/Public/InterchangeSkeletalMeshFactoryNode.h
--- a/Engine/Plugins/Interchange/Runtime/Source/FactoryNodes/Public/InterchangeSkeletalMeshFactoryNode.h
+++ b/Engine/Plugins/Interchange/Runtime/Source/FactoryNodes/Public/InterchangeSkeletalMeshFactoryNode.h
@@ -155,6 +155,8 @@ public:
 	/**
 	 * The skeletal mesh thumbnail can have an overlay if the last re-import was geometry only. This thumbnail overlay feature use the metadata to find out if the last import was geometry only.
 	 */
+	virtual void AppendAssetRegistryTags(FAssetRegistryTagsContext Context) const override;
+	UE_DEPRECATED(5.4, "Implement the version that takes FAssetRegistryTagsContext instead.")
 	virtual void AppendAssetRegistryTags(TArray<FAssetRegistryTag>& OutTags) const override;
 
 	virtual void CopyWithObject(const UInterchangeFactoryBaseNode* SourceNode, UObject* Object) override;
diff --git a/Engine/Plugins/MovieScene/CustomizableSequencerTracks/Source/CustomizableSequencerTracks/Private/SequencerTrackBP.cpp b/Engine/Plugins/MovieScene/CustomizableSequencerTracks/Source/CustomizableSequencerTracks/Private/SequencerTrackBP.cpp
--- a/Engine/Plugins/MovieScene/CustomizableSequencerTracks/Source/CustomizableSequencerTracks/Private/SequencerTrackBP.cpp
+++ b/Engine/Plugins/MovieScene/CustomizableSequencerTracks/Source/CustomizableSequencerTracks/Private/SequencerTrackBP.cpp
@@ -2,7 +2,7 @@
 
 #include "SequencerTrackBP.h"
 
-
+#include "UObject/AssetRegistryTagsContext.h"
 
 UMovieSceneSection* USequencerTrackBP::CreateNewSection()
 {
@@ -25,5 +25,12 @@ UMovieSceneSection* USequencerTrackBP::CreateNewSection()
 
 void USequencerTrackBP::GetAssetRegistryTags(TArray<FAssetRegistryTag>& OutTags) const
 {
+	PRAGMA_DISABLE_DEPRECATION_WARNINGS;
 	Super::GetAssetRegistryTags(OutTags);
+	PRAGMA_ENABLE_DEPRECATION_WARNINGS;
+}
+
+void USequencerTrackBP::GetAssetRegistryTags(FAssetRegistryTagsContext Context) const
+{
+	Super::GetAssetRegistryTags(Context);
 }
diff --git a/Engine/Plugins/MovieScene/TemplateSequence/Source/TemplateSequence/Private/TemplateSequence.cpp b/Engine/Plugins/MovieScene/TemplateSequence/Source/TemplateSequence/Private/TemplateSequence.cpp
--- a/Engine/Plugins/MovieScene/TemplateSequence/Source/TemplateSequence/Private/TemplateSequence.cpp
+++ b/Engine/Plugins/MovieScene/TemplateSequence/Source/TemplateSequence/Private/TemplateSequence.cpp
@@ -10,6 +10,7 @@
 #include "Tracks/MovieSceneSkeletalAnimationTrack.h"
 #include "MovieSceneSpawnable.h"
 #include "Tracks/MovieSceneSpawnTrack.h"
+#include "UObject/AssetRegistryTagsContext.h"
 
 #include "TemplateSequenceActor.h"
 #include "TemplateSequencePlayer.h"
@@ -245,16 +246,23 @@ UObject* UTemplateSequence::MakeSpawnableTemplateFromInstance(UObject& InSourceO
 
 void UTemplateSequence::GetAssetRegistryTags(TArray<FAssetRegistryTag>& OutTags) const
 {
+	PRAGMA_DISABLE_DEPRECATION_WARNINGS;
 	Super::GetAssetRegistryTags(OutTags);
+	PRAGMA_ENABLE_DEPRECATION_WARNINGS;
+}
+
+void UTemplateSequence::GetAssetRegistryTags(FAssetRegistryTagsContext Context) const
+{
+	Super::GetAssetRegistryTags(Context);
 
 	if (BoundActorClass.IsValid())
 	{
 		FAssetRegistryTag Tag("BoundActorClass", BoundActorClass->GetPathName(), FAssetRegistryTag::TT_Alphabetical);
-		OutTags.Add(Tag);
+		Context.AddTag(Tag);
 	}
 	else
 	{
-		OutTags.Emplace("BoundActorClass", "(None)", FAssetRegistryTag::TT_Alphabetical);
+		Context.AddTag(FAssetRegistryTag("BoundActorClass", "(None)", FAssetRegistryTag::TT_Alphabetical));
 	}
 }
 
diff --git a/Engine/Plugins/Runtime/GeometryCache/Source/GeometryCache/Private/GeometryCache.cpp b/Engine/Plugins/Runtime/GeometryCache/Source/GeometryCache/Private/GeometryCache.cpp
--- a/Engine/Plugins/Runtime/GeometryCache/Source/GeometryCache/Private/GeometryCache.cpp
+++ b/Engine/Plugins/Runtime/GeometryCache/Source/GeometryCache/Private/GeometryCache.cpp
@@ -5,6 +5,7 @@
 #include "Engine/AssetUserData.h"
 #include "GeometryCacheMeshData.h"
 #include "GeometryCacheTrack.h"
+#include "UObject/AssetRegistryTagsContext.h"
 #include "UObject/FrameworkObjectVersion.h"
 #include "UObject/AnimPhysObjectVersion.h"
 #include "Interfaces/ITargetPlatform.h"
@@ -94,19 +95,26 @@ FString UGeometryCache::GetDesc()
 }
 
 void UGeometryCache::GetAssetRegistryTags(TArray<FAssetRegistryTag>& OutTags) const
+{
+	PRAGMA_DISABLE_DEPRECATION_WARNINGS;
+	Super::GetAssetRegistryTags(OutTags);
+	PRAGMA_ENABLE_DEPRECATION_WARNINGS;
+}
+
+void UGeometryCache::GetAssetRegistryTags(FAssetRegistryTagsContext Context) const
 {
 	// Information on number total of (per type) tracks
 	const int32 NumTracks = Tracks.Num();	
-	OutTags.Add(FAssetRegistryTag("Total Tracks", FString::FromInt(NumTracks), FAssetRegistryTag::TT_Numerical));
+	Context.AddTag(FAssetRegistryTag("Total Tracks", FString::FromInt(NumTracks), FAssetRegistryTag::TT_Numerical));
 
 #if WITH_EDITORONLY_DATA
 	if (AssetImportData)
 	{
-		OutTags.Add(FAssetRegistryTag(SourceFileTagName(), AssetImportData->GetSourceData().ToJson(), FAssetRegistryTag::TT_Hidden));
+		Context.AddTag(FAssetRegistryTag(SourceFileTagName(), AssetImportData->GetSourceData().ToJson(), FAssetRegistryTag::TT_Hidden));
 	}
 #endif
 
-	Super::GetAssetRegistryTags(OutTags);
+	Super::GetAssetRegistryTags(Context);
 }
 
 void UGeometryCache::GetResourceSizeEx(FResourceSizeEx& CumulativeResourceSize)
diff --git a/Engine/Plugins/Runtime/HairStrands/Source/HairStrandsCore/Private/GroomAsset.cpp b/Engine/Plugins/Runtime/HairStrands/Source/HairStrandsCore/Private/GroomAsset.cpp
--- a/Engine/Plugins/Runtime/HairStrands/Source/HairStrandsCore/Private/GroomAsset.cpp
+++ b/Engine/Plugins/Runtime/HairStrands/Source/HairStrandsCore/Private/GroomAsset.cpp
@@ -23,6 +23,7 @@
 #include "Serialization/MemoryReader.h"
 #include "Serialization/MemoryWriter.h"
 #include "UObject/AnimObjectVersion.h"
+#include "UObject/AssetRegistryTagsContext.h"
 #include "UObject/Package.h"
 #include "UObject/PhysicsObjectVersion.h"
 #include "UObject/ReleaseObjectVersion.h"
@@ -1779,13 +1780,20 @@ void UGroomAsset::MarkMaterialsHasChanged()
 
 #if WITH_EDITORONLY_DATA
 void UGroomAsset::GetAssetRegistryTags(TArray<FAssetRegistryTag>& OutTags) const
+{
+	PRAGMA_DISABLE_DEPRECATION_WARNINGS;
+	Super::GetAssetRegistryTags(OutTags);
+	PRAGMA_ENABLE_DEPRECATION_WARNINGS;
+}
+
+void UGroomAsset::GetAssetRegistryTags(FAssetRegistryTagsContext Context) const
 {
 	if (AssetImportData)
 	{
-		OutTags.Add(FAssetRegistryTag(SourceFileTagName(), AssetImportData->GetSourceData().ToJson(), FAssetRegistryTag::TT_Hidden));
+		Context.AddTag(FAssetRegistryTag(SourceFileTagName(), AssetImportData->GetSourceData().ToJson(), FAssetRegistryTag::TT_Hidden));
 	}
 
-	Super::GetAssetRegistryTags(OutTags);
+	Super::GetAssetRegistryTags(Context);
 }
 
 void UGroomAsset::PostInitProperties()
diff --git a/Engine/Plugins/Runtime/HairStrands/Source/HairStrandsCore/Private/GroomCache.cpp b/Engine/Plugins/Runtime/HairStrands/Source/HairStrandsCore/Private/GroomCache.cpp
--- a/Engine/Plugins/Runtime/HairStrands/Source/HairStrandsCore/Private/GroomCache.cpp
+++ b/Engine/Plugins/Runtime/HairStrands/Source/HairStrandsCore/Private/GroomCache.cpp
@@ -7,6 +7,7 @@
 #include "GroomAsset.h"
 #include "Serialization/MemoryReader.h"
 #include "Serialization/MemoryWriter.h"
+#include "UObject/AssetRegistryTagsContext.h"
 
 #include UE_INLINE_GENERATED_CPP_BY_NAME(GroomCache)
 
@@ -34,15 +35,22 @@ void UGroomCache::Serialize(FArchive& Ar)
 }
 
 void UGroomCache::GetAssetRegistryTags(TArray<FAssetRegistryTag>& OutTags) const
+{
+	PRAGMA_DISABLE_DEPRECATION_WARNINGS;
+	Super::GetAssetRegistryTags(OutTags);
+	PRAGMA_ENABLE_DEPRECATION_WARNINGS;
+}
+
+void UGroomCache::GetAssetRegistryTags(FAssetRegistryTagsContext Context) const
 {
 #if WITH_EDITORONLY_DATA
 	if (AssetImportData)
 	{
-		OutTags.Add(FAssetRegistryTag(SourceFileTagName(), AssetImportData->GetSourceData().ToJson(), FAssetRegistryTag::TT_Hidden));
+		Context.AddTag(FAssetRegistryTag(SourceFileTagName(), AssetImportData->GetSourceData().ToJson(), FAssetRegistryTag::TT_Hidden));
 	}
 #endif
 
-	Super::GetAssetRegistryTags(OutTags);
+	Super::GetAssetRegistryTags(Context);
 }
 
 void UGroomCache::Initialize(EGroomCacheType Type)
diff --git a/Engine/Plugins/Runtime/HairStrands/Source/HairStrandsCore/Public/GroomAsset.h b/Engine/Plugins/Runtime/HairStrands/Source/HairStrandsCore/Public/GroomAsset.h
--- a/Engine/Plugins/Runtime/HairStrands/Source/HairStrandsCore/Public/GroomAsset.h
+++ b/Engine/Plugins/Runtime/HairStrands/Source/HairStrandsCore/Public/GroomAsset.h
@@ -619,7 +619,9 @@ public:
 	UPROPERTY(VisibleAnywhere, Instanced, Category = ImportSettings)
 	TObjectPtr<class UAssetImportData> AssetImportData;
 
-	/** Retrievde the asset tags*/
+	/** Retrieve the asset tags*/
+	virtual void GetAssetRegistryTags(FAssetRegistryTagsContext Context) const override;
+	UE_DEPRECATED(5.4, "Implement the version that takes FAssetRegistryTagsContext instead.")
 	virtual void GetAssetRegistryTags(TArray<FAssetRegistryTag>& OutTags) const override;
 
 	/** Part of Uobject interface */
diff --git a/Engine/Plugins/Runtime/HairStrands/Source/HairStrandsCore/Public/GroomCache.h b/Engine/Plugins/Runtime/HairStrands/Source/HairStrandsCore/Public/GroomCache.h
--- a/Engine/Plugins/Runtime/HairStrands/Source/HairStrandsCore/Public/GroomCache.h
+++ b/Engine/Plugins/Runtime/HairStrands/Source/HairStrandsCore/Public/GroomCache.h
@@ -21,7 +21,9 @@ public:
 
 	//~ Begin UObject Interface.
 	virtual void Serialize(FArchive& Ar) override;
-	virtual void GetAssetRegistryTags(TArray<FAssetRegistryTag>& OutTags) const;
+	virtual void GetAssetRegistryTags(FAssetRegistryTagsContext Context) const override;
+	UE_DEPRECATED(5.4, "Implement the version that takes FAssetRegistryTagsContext instead.")
+	virtual void GetAssetRegistryTags(TArray<FAssetRegistryTag>& OutTags) const override;
 	//~ End UObject Interface.
 
 	void Initialize(EGroomCacheType Type);
diff --git a/Engine/Plugins/Runtime/ModelViewViewModel/Source/ModelViewViewModelBlueprint/Private/MVVMBlueprintView.cpp b/Engine/Plugins/Runtime/ModelViewViewModel/Source/ModelViewViewModelBlueprint/Private/MVVMBlueprintView.cpp
--- a/Engine/Plugins/Runtime/ModelViewViewModel/Source/ModelViewViewModelBlueprint/Private/MVVMBlueprintView.cpp
+++ b/Engine/Plugins/Runtime/ModelViewViewModel/Source/ModelViewViewModelBlueprint/Private/MVVMBlueprintView.cpp
@@ -9,6 +9,7 @@
 #include "MVVMBlueprintViewEvent.h"
 #include "MVVMWidgetBlueprintExtension_View.h"
 #include "WidgetBlueprint.h"
+#include "UObject/AssetRegistryTagsContext.h"
 #include "UObject/FortniteMainBranchObjectVersion.h"
 
 #include UE_INLINE_GENERATED_CPP_BY_NAME(MVVMBlueprintView)
@@ -535,25 +536,29 @@ void UMVVMBlueprintView::PostEditChangeChainProperty(FPropertyChangedChainEvent&
 }
 
 void UMVVMBlueprintView::AddAssetTags(TArray<FAssetRegistryTag>& OutTags) const
+{
+}
+
+void UMVVMBlueprintView::AddAssetTags(FAssetRegistryTagsContext Context) const
 {
 	if (AvailableViewModels.Num() > 0)
 	{
 		TStringBuilder<512> Builder;
-		for (const FMVVMBlueprintViewModelContext& Context : AvailableViewModels)
+		for (const FMVVMBlueprintViewModelContext& ViewModelContext : AvailableViewModels)
 		{
-			if (Context.IsValid())
+			if (ViewModelContext.IsValid())
 			{
 				if (Builder.Len() > 0)
 				{
 					Builder << TEXT(',');
 				}
-				Builder << Context.GetViewModelClass()->GetPathName();
+				Builder << ViewModelContext.GetViewModelClass()->GetPathName();
 			}
 		}
 
 		if (Builder.Len() > 0)
 		{
-			OutTags.Emplace(FName("Viewmodels"), Builder.ToString(), FAssetRegistryTag::TT_Hidden);
+			Context.AddTag(FAssetRegistryTag(FName("Viewmodels"), Builder.ToString(), FAssetRegistryTag::TT_Hidden));
 		}
 	}
 }
diff --git a/Engine/Plugins/Runtime/ModelViewViewModel/Source/ModelViewViewModelBlueprint/Public/MVVMBlueprintView.h b/Engine/Plugins/Runtime/ModelViewViewModel/Source/ModelViewViewModelBlueprint/Public/MVVMBlueprintView.h
--- a/Engine/Plugins/Runtime/ModelViewViewModel/Source/ModelViewViewModelBlueprint/Public/MVVMBlueprintView.h
+++ b/Engine/Plugins/Runtime/ModelViewViewModel/Source/ModelViewViewModelBlueprint/Public/MVVMBlueprintView.h
@@ -154,6 +154,8 @@ public:
 	virtual void PostEditChangeProperty(FPropertyChangedEvent& PropertyChangedEvent) override;
 	virtual void PostEditChangeChainProperty(FPropertyChangedChainEvent& PropertyChainEvent) override;
 
+	void AddAssetTags(FAssetRegistryTagsContext Context) const;
+	UE_DEPRECATED(5.4, "Implement the version that takes FAssetRegistryTagsContext instead.")
 	void AddAssetTags(TArray<FAssetRegistryTag>& OutTags) const;
 	void WidgetRenamed(FName OldObjectName, FName NewObjectName);
 #endif
diff --git a/Engine/Plugins/Runtime/ModelViewViewModel/Source/ModelViewViewModelEditor/Private/ModelViewViewModelEditorModule.cpp b/Engine/Plugins/Runtime/ModelViewViewModel/Source/ModelViewViewModelEditor/Private/ModelViewViewModelEditorModule.cpp
--- a/Engine/Plugins/Runtime/ModelViewViewModel/Source/ModelViewViewModelEditor/Private/ModelViewViewModelEditorModule.cpp
+++ b/Engine/Plugins/Runtime/ModelViewViewModel/Source/ModelViewViewModelEditor/Private/ModelViewViewModelEditorModule.cpp
@@ -22,6 +22,7 @@
 #include "Tabs/MVVMPreviewSourceSummoner.h"
 #include "Tabs/MVVMViewModelSummoner.h"
 #include "ToolMenus.h"
+#include "UObject/AssetRegistryTagsContext.h"
 #include "UMGEditorModule.h"
 #include "WidgetBlueprintEditor.h"
 #include "WidgetDrawerConfig.h"
@@ -60,8 +61,8 @@ void FModelViewViewModelEditorModule::StartupModule()
 	}
 
 	FMVVMEditorCommands::Register();
-	FWidgetBlueprintDelegates::GetAssetTags.AddRaw(this, &FModelViewViewModelEditorModule::HandleWidgetBlueprintAssetTags);
-	FWidgetBlueprintGeneratedClassDelegates::GetAssetTags.AddRaw(this, &FModelViewViewModelEditorModule::HandleClassBlueprintAssetTags);
+	FWidgetBlueprintDelegates::GetAssetTagsWithContext.AddRaw(this, &FModelViewViewModelEditorModule::HandleWidgetBlueprintAssetTags);
+	FWidgetBlueprintGeneratedClassDelegates::GetAssetTagsWithContext.AddRaw(this, &FModelViewViewModelEditorModule::HandleClassBlueprintAssetTags);
 
 	UToolMenus::RegisterStartupCallback(FSimpleMulticastDelegate::FDelegate::CreateRaw(this, &FModelViewViewModelEditorModule::HandleRegisterMenus));
 }
@@ -71,8 +72,8 @@ void FModelViewViewModelEditorModule::ShutdownModule()
 {
 	UnregisterMenus();
 
-	FWidgetBlueprintGeneratedClassDelegates::GetAssetTags.RemoveAll(this);
-	FWidgetBlueprintDelegates::GetAssetTags.RemoveAll(this);
+	FWidgetBlueprintGeneratedClassDelegates::GetAssetTagsWithContext.RemoveAll(this);
+	FWidgetBlueprintDelegates::GetAssetTagsWithContext.RemoveAll(this);
 	if (FMessageLogModule* MessageLogModule = FModuleManager::GetModulePtr<FMessageLogModule>("MessageLog"))
 	{
 		MessageLogModule->UnregisterLogListing("Model View Viewmodel");
@@ -207,7 +208,7 @@ void FModelViewViewModelEditorModule::HandleActivateMode(FWidgetBlueprintApplica
 	}
 }
 
-void FModelViewViewModelEditorModule::HandleWidgetBlueprintAssetTags(const UWidgetBlueprint* WidgetBlueprint, TArray<UObject::FAssetRegistryTag>& OutTags)
+void FModelViewViewModelEditorModule::HandleWidgetBlueprintAssetTags(const UWidgetBlueprint* WidgetBlueprint, FAssetRegistryTagsContext Context)
 {
 	if (WidgetBlueprint && GEditor)
 	{
@@ -215,19 +216,19 @@ void FModelViewViewModelEditorModule::HandleWidgetBlueprintAssetTags(const UWidg
 		{
 			if (UMVVMBlueprintView* BlueprintView = Subsystem->GetView(WidgetBlueprint))
 			{
-				BlueprintView->AddAssetTags(OutTags);
+				BlueprintView->AddAssetTags(Context);
 			}
 		}
 	}
 }
 
-void FModelViewViewModelEditorModule::HandleClassBlueprintAssetTags(const UWidgetBlueprintGeneratedClass* GeneratedClass, TArray<UObject::FAssetRegistryTag>& OutTags)
+void FModelViewViewModelEditorModule::HandleClassBlueprintAssetTags(const UWidgetBlueprintGeneratedClass* GeneratedClass, FAssetRegistryTagsContext Context)
 {
 	if (GeneratedClass && GEditor && GeneratedClass->ClassGeneratedBy)
 	{
 		if (UWidgetBlueprint* WidgetBlueprint = Cast<UWidgetBlueprint>(GeneratedClass->ClassGeneratedBy))
 		{
-			HandleWidgetBlueprintAssetTags(WidgetBlueprint, OutTags);
+			HandleWidgetBlueprintAssetTags(WidgetBlueprint, Context);
 		}
 	}
 }
diff --git a/Engine/Plugins/Runtime/ModelViewViewModel/Source/ModelViewViewModelEditor/Private/ModelViewViewModelEditorModule.h b/Engine/Plugins/Runtime/ModelViewViewModel/Source/ModelViewViewModelEditor/Private/ModelViewViewModelEditorModule.h
--- a/Engine/Plugins/Runtime/ModelViewViewModel/Source/ModelViewViewModelEditor/Private/ModelViewViewModelEditorModule.h
+++ b/Engine/Plugins/Runtime/ModelViewViewModel/Source/ModelViewViewModelEditor/Private/ModelViewViewModelEditorModule.h
@@ -36,8 +36,8 @@ private:
 	void HandleRenameVariableReferences(UBlueprint* Blueprint, UClass* VariableClass, const FName& OldVarName, const FName& NewVarName);
 	void HandleDeactiveMode(FWidgetBlueprintApplicationMode& InDesignerMode);
 	void HandleActivateMode(FWidgetBlueprintApplicationMode& InDesignerMode);
-	void HandleWidgetBlueprintAssetTags(const UWidgetBlueprint* Widget, TArray<UObject::FAssetRegistryTag>& OutTags);
-	void HandleClassBlueprintAssetTags(const UWidgetBlueprintGeneratedClass* GeneratedClass, TArray<UObject::FAssetRegistryTag>& OutTags);
+	void HandleWidgetBlueprintAssetTags(const UWidgetBlueprint* Widget, FAssetRegistryTagsContext Context);
+	void HandleClassBlueprintAssetTags(const UWidgetBlueprintGeneratedClass* GeneratedClass, FAssetRegistryTagsContext Context);
 	void HandleRegisterMenus();
 	void UnregisterMenus();
 
diff --git a/Engine/Plugins/Runtime/RigVM/Source/RigVMDeveloper/Private/RigVMBlueprint.cpp b/Engine/Plugins/Runtime/RigVM/Source/RigVMDeveloper/Private/RigVMBlueprint.cpp
--- a/Engine/Plugins/Runtime/RigVM/Source/RigVMDeveloper/Private/RigVMBlueprint.cpp
+++ b/Engine/Plugins/Runtime/RigVM/Source/RigVMDeveloper/Private/RigVMBlueprint.cpp
@@ -9,6 +9,7 @@
 #include "EdGraph/RigVMEdGraph.h"
 #include "EdGraph/RigVMEdGraphNode.h"
 #include "EdGraph/RigVMEdGraphSchema.h"
+#include "UObject/AssetRegistryTagsContext.h"
 #include "UObject/ObjectSaveContext.h"
 #include "UObject/UObjectGlobals.h"
 #include "RigVMObjectVersion.h"
@@ -2630,15 +2631,29 @@ void URigVMBlueprint::PostDuplicate(bool bDuplicateForPIE)
 }
 
 void URigVMBlueprint::GetAssetRegistryTags(TArray<FAssetRegistryTag>& OutTags) const
+{
+	PRAGMA_DISABLE_DEPRECATION_WARNINGS;
+	Super::GetAssetRegistryTags(OutTags);
+	PRAGMA_ENABLE_DEPRECATION_WARNINGS;
+}
+
+void URigVMBlueprint::GetAssetRegistryTags(FAssetRegistryTagsContext Context) const
 {
 	if (CachedAssetTags.IsEmpty())
 	{
-		Super::GetAssetRegistryTags(OutTags);
-		CachedAssetTags = OutTags;
+		Super::GetAssetRegistryTags(Context);
+		CachedAssetTags.Reset(Context.GetNumTags());
+		Context.EnumerateTags([this](const FAssetRegistryTag& Tag)
+			{
+				CachedAssetTags.Add(Tag);
+			});
 	}
 	else
 	{
-		OutTags = CachedAssetTags;
+		for (const FAssetRegistryTag& Tag : CachedAssetTags)
+		{
+			Context.AddTag(Tag);
+		}
 	}
 }
 
diff --git a/Engine/Plugins/Runtime/StateTree/Source/StateTreeModule/Private/StateTree.cpp b/Engine/Plugins/Runtime/StateTree/Source/StateTreeModule/Private/StateTree.cpp
--- a/Engine/Plugins/Runtime/StateTree/Source/StateTreeModule/Private/StateTree.cpp
+++ b/Engine/Plugins/Runtime/StateTree/Source/StateTreeModule/Private/StateTree.cpp
@@ -13,6 +13,7 @@
 #include "Misc/DataValidation.h"
 #include "StructUtilsDelegates.h"
 #include "Misc/EnumerateRange.h"
+#include "UObject/AssetRegistryTagsContext.h"
 #if WITH_EDITOR
 #include "Engine/UserDefinedStruct.h"
 #endif
@@ -241,11 +242,18 @@ void UStateTree::BeginDestroy()
 }
 
 void UStateTree::GetAssetRegistryTags(TArray<FAssetRegistryTag>& OutTags) const
+{
+	PRAGMA_DISABLE_DEPRECATION_WARNINGS;
+	Super::GetAssetRegistryTags(OutTags);
+	PRAGMA_ENABLE_DEPRECATION_WARNINGS;
+}
+
+void UStateTree::GetAssetRegistryTags(FAssetRegistryTagsContext Context) const
 {
 	const FString SchemaClassName = Schema ? Schema->GetClass()->GetPathName() : TEXT("");
-	OutTags.Add(FAssetRegistryTag(UE::StateTree::SchemaTag, SchemaClassName, FAssetRegistryTag::TT_Alphabetical));
+	Context.AddTag(FAssetRegistryTag(UE::StateTree::SchemaTag, SchemaClassName, FAssetRegistryTag::TT_Alphabetical));
 
-	Super::GetAssetRegistryTags(OutTags);
+	Super::GetAssetRegistryTags(Context);
 }
 
 void UStateTree::PostLoadAssetRegistryTags(const FAssetData& InAssetData, TArray<FAssetRegistryTag>& OutTagsAndValuesToUpdate) const
diff --git a/Engine/Plugins/ScriptPlugin/Source/ScriptPlugin/Classes/ScriptBlueprint.h b/Engine/Plugins/ScriptPlugin/Source/ScriptPlugin/Classes/ScriptBlueprint.h
--- a/Engine/Plugins/ScriptPlugin/Source/ScriptPlugin/Classes/ScriptBlueprint.h
+++ b/Engine/Plugins/ScriptPlugin/Source/ScriptPlugin/Classes/ScriptBlueprint.h
@@ -35,7 +35,9 @@ public:
 	
 #if WITH_EDITORONLY_DATA
 	/** Override to ensure we write out the asset import data */
-	virtual void GetAssetRegistryTags(TArray<FAssetRegistryTag>& OutTags) const;
+	virtual void GetAssetRegistryTags(FAssetRegistryTagsContext Context) const override;
+	UE_DEPRECATED(5.4, "Implement the version that takes FAssetRegistryTagsContext instead.")
+	virtual void GetAssetRegistryTags(TArray<FAssetRegistryTag>& OutTags) const override;
 	virtual void PostLoad() override;
 #endif
 
diff --git a/Engine/Plugins/ScriptPlugin/Source/ScriptPlugin/Private/ScriptBlueprint.cpp b/Engine/Plugins/ScriptPlugin/Source/ScriptPlugin/Private/ScriptBlueprint.cpp
--- a/Engine/Plugins/ScriptPlugin/Source/ScriptPlugin/Private/ScriptBlueprint.cpp
+++ b/Engine/Plugins/ScriptPlugin/Source/ScriptPlugin/Private/ScriptBlueprint.cpp
@@ -4,6 +4,7 @@
 #include "ScriptPluginLog.h"
 #include "HAL/FileManager.h"
 #include "Misc/FileHelper.h"
+#include "UObject/AssetRegistryTagsContext.h"
 #if WITH_EDITOR
 #include "EditorFramework/AssetImportData.h"
 #include "Kismet2/BlueprintEditorUtils.h"
@@ -22,13 +23,20 @@ UScriptBlueprint::UScriptBlueprint(const FObjectInitializer& ObjectInitializer)
 
 #if WITH_EDITORONLY_DATA
 void UScriptBlueprint::GetAssetRegistryTags(TArray<FAssetRegistryTag>& OutTags) const
+{
+	PRAGMA_DISABLE_DEPRECATION_WARNINGS;
+	Super::GetAssetRegistryTags(OutTags);
+	PRAGMA_ENABLE_DEPRECATION_WARNINGS;
+}
+
+void UScriptBlueprint::GetAssetRegistryTags(FAssetRegistryTagsContext Context) const
 {
 	if (AssetImportData)
 	{
-		OutTags.Add( FAssetRegistryTag(SourceFileTagName(), AssetImportData->GetSourceData().ToJson(), FAssetRegistryTag::TT_Hidden) );
+		Context.AddTag( FAssetRegistryTag(SourceFileTagName(), AssetImportData->GetSourceData().ToJson(), FAssetRegistryTag::TT_Hidden) );
 	}
 
-	Super::GetAssetRegistryTags(OutTags);
+	Super::GetAssetRegistryTags(Context);
 }
 void UScriptBlueprint::PostLoad()
 {
diff --git a/Engine/Plugins/VirtualProduction/Takes/Source/TakesCore/Private/TakeMetaData.cpp b/Engine/Plugins/VirtualProduction/Takes/Source/TakesCore/Private/TakeMetaData.cpp
--- a/Engine/Plugins/VirtualProduction/Takes/Source/TakesCore/Private/TakeMetaData.cpp
+++ b/Engine/Plugins/VirtualProduction/Takes/Source/TakesCore/Private/TakeMetaData.cpp
@@ -2,6 +2,7 @@
 
 #include "TakeMetaData.h"
 #include "TakePreset.h"
+#include "UObject/AssetRegistryTagsContext.h"
 #include "UObject/Package.h"
 #include "TakesCoreBlueprintLibrary.h"
 #include "ObjectTools.h"
@@ -269,14 +270,22 @@ void UTakeMetaData::SetFrameRateFromTimecode(bool InFromTimecode)
 }
 
 
-void UTakeMetaData::ExtendAssetRegistryTags(TArray<FAssetRegistryTag>& OutTags) const
+void UTakeMetaData::ExtendAssetRegistryTags(FAssetRegistryTagsContext Context) const
 {
-	OutTags.Emplace(AssetRegistryTag_Slate,       Slate,                     FAssetRegistryTag::ETagType::TT_Alphabetical,  FAssetRegistryTag::TD_None);
-	OutTags.Emplace(AssetRegistryTag_TakeNumber,  LexToString(TakeNumber),   FAssetRegistryTag::ETagType::TT_Numerical,     FAssetRegistryTag::TD_None);
-	OutTags.Emplace(AssetRegistryTag_Timestamp, Timestamp.ToString(), FAssetRegistryTag::ETagType::TT_Chronological, FAssetRegistryTag::TD_Date | FAssetRegistryTag::TD_Time);
-	OutTags.Emplace(AssetRegistryTag_TimecodeIn, TimecodeIn.ToString(), FAssetRegistryTag::ETagType::TT_Numerical, FAssetRegistryTag::TD_None);
-	OutTags.Emplace(AssetRegistryTag_TimecodeOut, TimecodeOut.ToString(), FAssetRegistryTag::ETagType::TT_Numerical, FAssetRegistryTag::TD_None);
-	OutTags.Emplace(AssetRegistryTag_Description, Description,               FAssetRegistryTag::ETagType::TT_Alphabetical,  FAssetRegistryTag::TD_None);
+	IMovieSceneMetaDataInterface::ExtendAssetRegistryTags(Context);
+
+	Context.AddTag(FAssetRegistryTag(	AssetRegistryTag_Slate,							Slate,                 
+										FAssetRegistryTag::ETagType::TT_Alphabetical,	FAssetRegistryTag::TD_None));
+	Context.AddTag(FAssetRegistryTag(	AssetRegistryTag_TakeNumber,					LexToString(TakeNumber),
+										FAssetRegistryTag::ETagType::TT_Numerical,		FAssetRegistryTag::TD_None));
+	Context.AddTag(FAssetRegistryTag(	AssetRegistryTag_Timestamp,						Timestamp.ToString(),
+										FAssetRegistryTag::ETagType::TT_Chronological,	FAssetRegistryTag::TD_Date | FAssetRegistryTag::TD_Time));
+	Context.AddTag(FAssetRegistryTag(	AssetRegistryTag_TimecodeIn,					TimecodeIn.ToString(),
+										FAssetRegistryTag::ETagType::TT_Numerical,		FAssetRegistryTag::TD_None));
+	Context.AddTag(FAssetRegistryTag(	AssetRegistryTag_TimecodeOut,					TimecodeOut.ToString(),
+										FAssetRegistryTag::ETagType::TT_Numerical,		FAssetRegistryTag::TD_None));
+	Context.AddTag(FAssetRegistryTag(	AssetRegistryTag_Description,					Description,
+										FAssetRegistryTag::ETagType::TT_Alphabetical,	FAssetRegistryTag::TD_None));
 }
 
 void UTakeMetaData::ExtendAssetRegistryTagMetaData(TMap<FName, FAssetRegistryTagMetadata>& OutMetadata) const
diff --git a/Engine/Plugins/VirtualProduction/Takes/Source/TakesCore/Public/TakeMetaData.h b/Engine/Plugins/VirtualProduction/Takes/Source/TakesCore/Public/TakeMetaData.h
--- a/Engine/Plugins/VirtualProduction/Takes/Source/TakesCore/Public/TakeMetaData.h
+++ b/Engine/Plugins/VirtualProduction/Takes/Source/TakesCore/Public/TakeMetaData.h
@@ -64,7 +64,9 @@ public:
 	/**
 	 * Extend the default ULevelSequence asset registry tags
 	 */
-	virtual void ExtendAssetRegistryTags(TArray<FAssetRegistryTag>& OutTags) const override;
+	virtual void ExtendAssetRegistryTags(FAssetRegistryTagsContext Context) const override;
+	UE_DEPRECATED(5.4, "Implement the version that takes FAssetRegistryTagsContext instead.")
+	virtual void ExtendAssetRegistryTags(TArray<FAssetRegistryTag>& OutTags) const override {}
 
 	/**
 	 * Extend the default ULevelSequence asset registry tag meta-data
diff --git a/Engine/Source/Developer/FunctionalTesting/Private/FunctionalTest.cpp b/Engine/Source/Developer/FunctionalTesting/Private/FunctionalTest.cpp
--- a/Engine/Source/Developer/FunctionalTesting/Private/FunctionalTest.cpp
+++ b/Engine/Source/Developer/FunctionalTesting/Private/FunctionalTest.cpp
@@ -9,6 +9,7 @@
 #include "Components/BillboardComponent.h"
 #include "HAL/FileManager.h"
 #include "Misc/FileHelper.h"
+#include "UObject/AssetRegistryTagsContext.h"
 #include "UObject/ConstructorHelpers.h"
 #include "ProfilingDebugging/ProfilingHelpers.h"
 #include "Misc/AutomationTest.h"
@@ -595,13 +596,20 @@ void AFunctionalTest::PostEditChangeProperty( struct FPropertyChangedEvent& Prop
 
 void AFunctionalTest::GetAssetRegistryTags(TArray<FAssetRegistryTag>& OutTags) const
 {
+	PRAGMA_DISABLE_DEPRECATION_WARNINGS;
 	Super::GetAssetRegistryTags(OutTags);
+	PRAGMA_ENABLE_DEPRECATION_WARNINGS;
+}
+
+void AFunctionalTest::GetAssetRegistryTags(FAssetRegistryTagsContext Context) const
+{
+	Super::GetAssetRegistryTags(Context);
 
 	if (IsPackageExternal() && IsEnabled())
 	{
 		const FString TestActor = GetActorLabel() + TEXT("|") + GetName();
 		const TCHAR* TestCategory = IsEditorOnlyObject(this) ? TEXT("TestNameEditor") : TEXT("TestName");
-		OutTags.Add(UObject::FAssetRegistryTag(TestCategory, TestActor, UObject::FAssetRegistryTag::TT_Hidden));
+		Context.AddTag(UObject::FAssetRegistryTag(TestCategory, TestActor, UObject::FAssetRegistryTag::TT_Hidden));
 	}
 }
 
diff --git a/Engine/Source/Developer/FunctionalTesting/Private/FunctionalTestingModule.cpp b/Engine/Source/Developer/FunctionalTesting/Private/FunctionalTestingModule.cpp
--- a/Engine/Source/Developer/FunctionalTesting/Private/FunctionalTestingModule.cpp
+++ b/Engine/Source/Developer/FunctionalTesting/Private/FunctionalTestingModule.cpp
@@ -14,6 +14,7 @@
 #include "AssetRegistry/AssetRegistryModule.h"
 #include "Misc/CommandLine.h"
 #include "IAutomationControllerModule.h"
+#include "UObject/AssetRegistryTagsContext.h"
 
 #define LOCTEXT_NAMESPACE "FunctionalTesting"
 
@@ -38,7 +39,7 @@ public:
 
 private:
 	UWorld* GetTestWorld();
-	void OnGetAssetTagsForWorld(const UWorld* World, TArray<UObject::FAssetRegistryTag>& OutTags);
+	void OnGetAssetTagsForWorld(const UWorld* World, FAssetRegistryTagsContext Context);
 
 
 	TWeakObjectPtr<class UFunctionalTestingManager> TestManager;
@@ -49,18 +50,18 @@ void FFunctionalTestingModule::StartupModule()
 {
 	bPendingActivation = false;
 #if WITH_EDITOR
-	FWorldDelegates::GetAssetTags.AddRaw(this, &FFunctionalTestingModule::OnGetAssetTagsForWorld);
+	FWorldDelegates::GetAssetTagsWithContext.AddRaw(this, &FFunctionalTestingModule::OnGetAssetTagsForWorld);
 #endif
 }
 
 void FFunctionalTestingModule::ShutdownModule() 
 {
 #if WITH_EDITOR
-	FWorldDelegates::GetAssetTags.RemoveAll(this);
+	FWorldDelegates::GetAssetTagsWithContext.RemoveAll(this);
 #endif
 }
 
-void FFunctionalTestingModule::OnGetAssetTagsForWorld(const UWorld* World, TArray<UObject::FAssetRegistryTag>& OutTags)
+void FFunctionalTestingModule::OnGetAssetTagsForWorld(const UWorld* World, FAssetRegistryTagsContext Context)
 {
 #if WITH_EDITOR
 	TArray<FString> TestNamesRuntime;
@@ -80,13 +81,13 @@ void FFunctionalTestingModule::OnGetAssetTagsForWorld(const UWorld* World, TArra
 		}
 	}
 
-	auto AddTestNames = [&OutTags](const TCHAR* TagName, TArray<FString>& TestNames)
+	auto AddTestNames = [&Context](const TCHAR* TagName, TArray<FString>& TestNames)
 	{
 		if (!TestNames.IsEmpty())
 		{
 			TestNames.Sort();
 			FString TestNamesStr = FString::Join(TestNames, TEXT(""));
-			OutTags.Add(UObject::FAssetRegistryTag(TagName, MoveTemp(TestNamesStr), UObject::FAssetRegistryTag::TT_Hidden));
+			Context.AddTag(UObject::FAssetRegistryTag(TagName, MoveTemp(TestNamesStr), UObject::FAssetRegistryTag::TT_Hidden));
 		}
 	};
 	AddTestNames(TEXT("TestNames"), TestNamesRuntime);
diff --git a/Engine/Source/Editor/AnimationModifiers/Private/AnimationModifier.cpp b/Engine/Source/Editor/AnimationModifiers/Private/AnimationModifier.cpp
--- a/Engine/Source/Editor/AnimationModifiers/Private/AnimationModifier.cpp
+++ b/Engine/Source/Editor/AnimationModifiers/Private/AnimationModifier.cpp
@@ -25,6 +25,7 @@
 #include "ScopedTransaction.h"
 #include "Serialization/Archive.h"
 #include "Templates/Casts.h"
+#include "UObject/AssetRegistryTagsContext.h"
 #include "UObject/Class.h"
 #include "UObject/NameTypes.h"
 #include "UObject/ObjectKey.h"
@@ -303,8 +304,9 @@ FGuid UAnimationModifier::GetLatestRevisionGuid() const
 	return GetDefault<UAnimationModifier>(GetClass())->RevisionGuid;
 }
 
-void UAnimationModifier::GetAssetRegistryTagsForAppliedModifiersFromSkeleton(const UObject* Object, TArray<UObject::FAssetRegistryTag>& OutTags)
+void UAnimationModifier::GetAssetRegistryTagsForAppliedModifiersFromSkeleton(FAssetRegistryTagsContext Context)
 {
+	const UObject* Object = Context.GetObject();
 	FString ModifiersList;
 	if (UAnimSequence* AnimSequence = const_cast<UAnimSequence*>(Cast<UAnimSequence>(Object)))
 	{
@@ -324,7 +326,7 @@ void UAnimationModifier::GetAssetRegistryTagsForAppliedModifiersFromSkeleton(con
 				}
 				if (!ModifiersList.IsEmpty())
 				{
-					OutTags.Emplace(AnimationModifiersTag, ModifiersList, UObject::FAssetRegistryTag::TT_Hidden);
+					Context.AddTag(FAssetRegistryTag(AnimationModifiersTag, ModifiersList, UObject::FAssetRegistryTag::TT_Hidden));
 				}
 			}
 		}
diff --git a/Engine/Source/Editor/AnimationModifiers/Private/AnimationModifiersModule.cpp b/Engine/Source/Editor/AnimationModifiers/Private/AnimationModifiersModule.cpp
--- a/Engine/Source/Editor/AnimationModifiers/Private/AnimationModifiersModule.cpp
+++ b/Engine/Source/Editor/AnimationModifiers/Private/AnimationModifiersModule.cpp
@@ -28,6 +28,7 @@
 #include "Subsystems/ImportSubsystem.h"
 #include "Templates/Casts.h"
 #include "Templates/SubclassOf.h"
+#include "UObject/AssetRegistryTagsContext.h"
 #include "UObject/Object.h"
 #include "UObject/UObjectGlobals.h"
 #include "Widgets/DeclarativeSyntaxSupport.h"
@@ -80,7 +81,7 @@ void FAnimationModifiersModule::StartupModule()
 	});
 
 	// Register extra asset registry tags for UAnimSequence
-	OnGetExtraObjectTagsHandle = UObject::FAssetRegistryTag::OnGetExtraObjectTags.AddStatic(&UAnimationModifier::GetAssetRegistryTagsForAppliedModifiersFromSkeleton);
+	OnGetExtraObjectTagsHandle = UObject::FAssetRegistryTag::OnGetExtraObjectTagsWithContext.AddStatic(&UAnimationModifier::GetAssetRegistryTagsForAppliedModifiersFromSkeleton);
 }
 
 void FAnimationModifiersModule::ShutdownModule()
@@ -95,7 +96,7 @@ void FAnimationModifiersModule::ShutdownModule()
 	UToolMenus::UnregisterOwner(this);
 	FCoreDelegates::OnPostEngineInit.Remove(DelegateHandle);
 
-	UObject::FAssetRegistryTag::OnGetExtraObjectTags.Remove(OnGetExtraObjectTagsHandle);
+	UObject::FAssetRegistryTag::OnGetExtraObjectTagsWithContext.Remove(OnGetExtraObjectTagsHandle);
 	
 	// Remove extender delegate
 	FWorkflowCentricApplication::GetModeExtenderList().RemoveAll([this](FWorkflowApplicationModeExtender& StoredExtender) { return StoredExtender.GetHandle() == Extender.GetHandle(); });
diff --git a/Engine/Source/Editor/AnimationModifiers/Public/AnimationModifier.h b/Engine/Source/Editor/AnimationModifiers/Public/AnimationModifier.h
--- a/Engine/Source/Editor/AnimationModifiers/Public/AnimationModifier.h
+++ b/Engine/Source/Editor/AnimationModifiers/Public/AnimationModifier.h
@@ -61,7 +61,7 @@ public:
 
 	//! @brief Get Asset Registry tags for applied modifiers from _skeleton_
 	//! @param[out] OutTags AnimationModifiers = %PATH%=%REVISION%;...
-	static void GetAssetRegistryTagsForAppliedModifiersFromSkeleton(const UObject* Object, TArray<UObject::FAssetRegistryTag>& OutTags);
+	static void GetAssetRegistryTagsForAppliedModifiersFromSkeleton(FAssetRegistryTagsContext Context);
 
 	static const FName AnimationModifiersTag;
 	static constexpr TCHAR AnimationModifiersDelimiter = ';';
diff --git a/Engine/Source/Editor/Blutility/Classes/ActorActionUtility.h b/Engine/Source/Editor/Blutility/Classes/ActorActionUtility.h
--- a/Engine/Source/Editor/Blutility/Classes/ActorActionUtility.h
+++ b/Engine/Source/Editor/Blutility/Classes/ActorActionUtility.h
@@ -32,8 +32,10 @@ public:
 
 public:
     // Begin UObject
-    virtual void GetAssetRegistryTags(TArray<FAssetRegistryTag>& OutTags) const override;
-    // End UObject
+	virtual void GetAssetRegistryTags(FAssetRegistryTagsContext Context) const override;
+	UE_DEPRECATED(5.4, "Implement the version that takes FAssetRegistryTagsContext instead.")
+	virtual void GetAssetRegistryTags(TArray<FAssetRegistryTag>& OutTags) const override;
+	// End UObject
 	
 private:
 	/**
diff --git a/Engine/Source/Editor/Blutility/Private/ActorActionUtility.cpp b/Engine/Source/Editor/Blutility/Private/ActorActionUtility.cpp
--- a/Engine/Source/Editor/Blutility/Private/ActorActionUtility.cpp
+++ b/Engine/Source/Editor/Blutility/Private/ActorActionUtility.cpp
@@ -4,23 +4,31 @@
 
 #include "EditorUtilityAssetPrototype.h"
 #include "EditorUtilityBlueprint.h"
+#include "UObject/AssetRegistryTagsContext.h"
 
 #define LOCTEXT_NAMESPACE "ActorActionUtility"
 
 void UActorActionUtility::GetAssetRegistryTags(TArray<FAssetRegistryTag>& OutTags) const
 {
+	PRAGMA_DISABLE_DEPRECATION_WARNINGS;
 	Super::GetAssetRegistryTags(OutTags);
+	PRAGMA_ENABLE_DEPRECATION_WARNINGS;
+}
+
+void UActorActionUtility::GetAssetRegistryTags(FAssetRegistryTagsContext Context) const
+{
+	Super::GetAssetRegistryTags(Context);
 	
 	const bool IsUpToDate =
 		!GetClass()->IsFunctionImplementedInScript("GetSupportedClass");
 
 	if (IsUpToDate)
 	{
-		FAssetActionUtilityPrototype::AddTagsFor_Version(OutTags);
+		FAssetActionUtilityPrototype::AddTagsFor_Version(Context);
 	}
 	
-	FAssetActionUtilityPrototype::AddTagsFor_SupportedClasses(SupportedClasses, OutTags);
-	FAssetActionUtilityPrototype::AddTagsFor_CallableFunctions(this, OutTags);
+	FAssetActionUtilityPrototype::AddTagsFor_SupportedClasses(SupportedClasses, Context);
+	FAssetActionUtilityPrototype::AddTagsFor_CallableFunctions(this, Context);
 }
 
 #undef LOCTEXT_NAMESPACE
\ No newline at end of file
diff --git a/Engine/Source/Editor/Blutility/Private/AssetActionUtility.cpp b/Engine/Source/Editor/Blutility/Private/AssetActionUtility.cpp
--- a/Engine/Source/Editor/Blutility/Private/AssetActionUtility.cpp
+++ b/Engine/Source/Editor/Blutility/Private/AssetActionUtility.cpp
@@ -8,6 +8,7 @@
 #include "JsonObjectConverter.h"
 #include "Misc/DataValidation.h"
 #include "Serialization/JsonSerializerMacros.h"
+#include "UObject/AssetRegistryTagsContext.h"
 #include "UObject/UObjectThreadContext.h"
 
 #define LOCTEXT_NAMESPACE "AssetActionUtility"
@@ -19,23 +20,30 @@ bool UAssetActionUtility::IsActionForBlueprints_Implementation() const
 
 void UAssetActionUtility::GetAssetRegistryTags(TArray<FAssetRegistryTag>& OutTags) const
 {
+	PRAGMA_DISABLE_DEPRECATION_WARNINGS;
 	Super::GetAssetRegistryTags(OutTags);
+	PRAGMA_ENABLE_DEPRECATION_WARNINGS;
+}
+
+void UAssetActionUtility::GetAssetRegistryTags(FAssetRegistryTagsContext Context) const
+{
+	Super::GetAssetRegistryTags(Context);
 	
 	const bool IsUpToDate =
 		!GetClass()->IsFunctionImplementedInScript("GetSupportedClass");
 	
 	if (IsUpToDate)
 	{
-		FAssetActionUtilityPrototype::AddTagsFor_Version(OutTags);
+		FAssetActionUtilityPrototype::AddTagsFor_Version(Context);
 	}
 	
-	FAssetActionUtilityPrototype::AddTagsFor_SupportedClasses(SupportedClasses, OutTags);
-	FAssetActionUtilityPrototype::AddTagsFor_SupportedConditions(SupportedConditions, OutTags);
+	FAssetActionUtilityPrototype::AddTagsFor_SupportedClasses(SupportedClasses, Context);
+	FAssetActionUtilityPrototype::AddTagsFor_SupportedConditions(SupportedConditions, Context);
 	if (!FUObjectThreadContext::Get().IsRoutingPostLoad)
 	{
-		FAssetActionUtilityPrototype::AddTagsFor_IsActionForBlueprints(IsActionForBlueprints(), OutTags);
+		FAssetActionUtilityPrototype::AddTagsFor_IsActionForBlueprints(IsActionForBlueprints(), Context);
 	}
-	FAssetActionUtilityPrototype::AddTagsFor_CallableFunctions(this, OutTags);
+	FAssetActionUtilityPrototype::AddTagsFor_CallableFunctions(this, Context);
 }
 			
 #undef LOCTEXT_NAMESPACE
\ No newline at end of file
diff --git a/Engine/Source/Editor/Blutility/Private/EditorUtilityAssetPrototype.cpp b/Engine/Source/Editor/Blutility/Private/EditorUtilityAssetPrototype.cpp
--- a/Engine/Source/Editor/Blutility/Private/EditorUtilityAssetPrototype.cpp
+++ b/Engine/Source/Editor/Blutility/Private/EditorUtilityAssetPrototype.cpp
@@ -8,6 +8,7 @@
 #include "JsonObjectConverter.h"
 #include "AssetRegistry/AssetData.h"
 #include "Serialization/JsonSerializerMacros.h"
+#include "UObject/AssetRegistryTagsContext.h"
 
 const FName AssetActionUtilityTags::BlutilityTagVersion(TEXT("BlutilityTagVersion"));
 const FName AssetActionUtilityTags::SupportedClasses(TEXT("SupportedClasses"));
@@ -104,19 +105,19 @@ TArray<FAssetActionSupportCondition> FAssetActionUtilityPrototype::GetAssetActio
 	return Results;
 }
 
-void FAssetActionUtilityPrototype::AddTagsFor_Version(TArray<UObject::FAssetRegistryTag>& OutTags)
+void FAssetActionUtilityPrototype::AddTagsFor_Version(FAssetRegistryTagsContext Context)
 {
 	// Adding a version to the tags just in case we need to know 'this blutility is out of date and we can't go based off the cached data because the format isn't followed, or this is a pre-tagged version'
-	OutTags.Add(UObject::FAssetRegistryTag(AssetActionUtilityTags::BlutilityTagVersion, LexToString(AssetActionUtilityTags::TagVersion), UObject::FAssetRegistryTag::TT_Hidden));	
+	Context.AddTag(UObject::FAssetRegistryTag(AssetActionUtilityTags::BlutilityTagVersion, LexToString(AssetActionUtilityTags::TagVersion), UObject::FAssetRegistryTag::TT_Hidden));	
 }
 
-void FAssetActionUtilityPrototype::AddTagsFor_SupportedClasses(const TArray<TSoftClassPtr<UObject>>& SupportedClasses, TArray<UObject::FAssetRegistryTag>& OutTags)
+void FAssetActionUtilityPrototype::AddTagsFor_SupportedClasses(const TArray<TSoftClassPtr<UObject>>& SupportedClasses, FAssetRegistryTagsContext Context)
 {
 	const FString SupportedClassesString = FString::JoinBy(SupportedClasses, TEXT(","), [](const TSoftClassPtr<UObject>& SupportedClass) { return SupportedClass.ToString(); });
-	OutTags.Add(UObject::FAssetRegistryTag(AssetActionUtilityTags::SupportedClasses, SupportedClassesString, UObject::FAssetRegistryTag::TT_Hidden));
+	Context.AddTag(UObject::FAssetRegistryTag(AssetActionUtilityTags::SupportedClasses, SupportedClassesString, UObject::FAssetRegistryTag::TT_Hidden));
 }
 
-void FAssetActionUtilityPrototype::AddTagsFor_SupportedConditions(const TArray<FAssetActionSupportCondition>& SupportedConditions, TArray<UObject::FAssetRegistryTag>& OutTags)
+void FAssetActionUtilityPrototype::AddTagsFor_SupportedConditions(const TArray<FAssetActionSupportCondition>& SupportedConditions, FAssetRegistryTagsContext Context)
 {
 	TArray<TSharedPtr<FJsonValue>> SupportedConditionJsonValues;
 	for (const FAssetActionSupportCondition& Condition : SupportedConditions)
@@ -137,17 +138,17 @@ void FAssetActionUtilityPrototype::AddTagsFor_SupportedConditions(const TArray<F
 			FJsonSerializer::Serialize(JsonArray, FString(), JsonWriter);
 		}
 
-		OutTags.Add(UObject::FAssetRegistryTag(AssetActionUtilityTags::SupportedConditions, SupportedConditionsJson, UObject::FAssetRegistryTag::TT_Hidden));
+		Context.AddTag(UObject::FAssetRegistryTag(AssetActionUtilityTags::SupportedConditions, SupportedConditionsJson, UObject::FAssetRegistryTag::TT_Hidden));
 	}
 }
 
-void FAssetActionUtilityPrototype::AddTagsFor_IsActionForBlueprints(bool IsActionForBlueprints, TArray<UObject::FAssetRegistryTag>& OutTags)
+void FAssetActionUtilityPrototype::AddTagsFor_IsActionForBlueprints(bool IsActionForBlueprints, FAssetRegistryTagsContext Context)
 {
 	const bool ActionForBlueprint = IsActionForBlueprints;
-	OutTags.Add(UObject::FAssetRegistryTag(AssetActionUtilityTags::IsActionForBlueprint, ActionForBlueprint ? TEXT("True") : TEXT("False"), UObject::FAssetRegistryTag::TT_Hidden));
+	Context.AddTag(UObject::FAssetRegistryTag(AssetActionUtilityTags::IsActionForBlueprint, ActionForBlueprint ? TEXT("True") : TEXT("False"), UObject::FAssetRegistryTag::TT_Hidden));
 }
 
-void FAssetActionUtilityPrototype::AddTagsFor_CallableFunctions(const UObject* FunctionsSource, TArray<UObject::FAssetRegistryTag>& OutTags)
+void FAssetActionUtilityPrototype::AddTagsFor_CallableFunctions(const UObject* FunctionsSource, FAssetRegistryTagsContext Context)
 {
 	check(FunctionsSource);
 	
@@ -192,6 +193,6 @@ void FAssetActionUtilityPrototype::AddTagsFor_CallableFunctions(const UObject* F
 			FJsonSerializer::Serialize(JsonArray, FString(), JsonWriter);
 		}
 
-		OutTags.Add(UObject::FAssetRegistryTag(AssetActionUtilityTags::CallableFunctions, FunctionDataJson, UObject::FAssetRegistryTag::TT_Hidden));
+		Context.AddTag(UObject::FAssetRegistryTag(AssetActionUtilityTags::CallableFunctions, FunctionDataJson, UObject::FAssetRegistryTag::TT_Hidden));
 	}
 }
\ No newline at end of file
diff --git a/Engine/Source/Editor/Blutility/Private/EditorUtilityAssetPrototype.h b/Engine/Source/Editor/Blutility/Private/EditorUtilityAssetPrototype.h
--- a/Engine/Source/Editor/Blutility/Private/EditorUtilityAssetPrototype.h
+++ b/Engine/Source/Editor/Blutility/Private/EditorUtilityAssetPrototype.h
@@ -80,11 +80,11 @@ public:
 	}
 
 public:
-	static void AddTagsFor_Version(TArray<UObject::FAssetRegistryTag>& OutTags);
-	static void AddTagsFor_SupportedClasses(const TArray<TSoftClassPtr<UObject>>& SupportedClasses, TArray<UObject::FAssetRegistryTag>& OutTags);
-	static void AddTagsFor_SupportedConditions(const TArray<FAssetActionSupportCondition>& SupportedConditions, TArray<UObject::FAssetRegistryTag>& OutTags);
-	static void AddTagsFor_IsActionForBlueprints(bool IsActionForBlueprints, TArray<UObject::FAssetRegistryTag>& OutTags);
-	static void AddTagsFor_CallableFunctions(const UObject* FunctionsSource, TArray<UObject::FAssetRegistryTag>& OutTags);
+	static void AddTagsFor_Version(FAssetRegistryTagsContext Context);
+	static void AddTagsFor_SupportedClasses(const TArray<TSoftClassPtr<UObject>>& SupportedClasses, FAssetRegistryTagsContext Context);
+	static void AddTagsFor_SupportedConditions(const TArray<FAssetActionSupportCondition>& SupportedConditions, FAssetRegistryTagsContext Context);
+	static void AddTagsFor_IsActionForBlueprints(bool IsActionForBlueprints, FAssetRegistryTagsContext Context);
+	static void AddTagsFor_CallableFunctions(const UObject* FunctionsSource, FAssetRegistryTagsContext Context);
 
 private:
 	FAssetData UtilityBlueprintAsset;
diff --git a/Engine/Source/Editor/Persona/Private/SAnimAssetFindReplace.cpp b/Engine/Source/Editor/Persona/Private/SAnimAssetFindReplace.cpp
--- a/Engine/Source/Editor/Persona/Private/SAnimAssetFindReplace.cpp
+++ b/Engine/Source/Editor/Persona/Private/SAnimAssetFindReplace.cpp
@@ -21,6 +21,7 @@
 #include "Misc/ScopedSlowTask.h"
 #include "ScopedTransaction.h"
 #include "Widgets/Views/SListView.h"
+#include "UObject/AssetRegistryTagsContext.h"
 #include "UObject/AnimPhysObjectVersion.h"
 #include "UObject/ObjectVersion.h"
 #include "Widgets/Input/SHyperlink.h"
@@ -137,15 +138,17 @@ void SAnimAssetFindReplace::Construct(const FArguments& InArgs)
 	}
 	ClassesWithAssetRegistryTags.Append(DerivedClassesWithAssetRegistryTags);
 
-	TArray<UObject::FAssetRegistryTag> AssetRegistryTags;
 	for(UClass* Class : ClassesWithAssetRegistryTags)
 	{
-		Class->GetDefaultObject()->GetAssetRegistryTags(AssetRegistryTags);
-		for(UObject::FAssetRegistryTag& AssetRegistryTag : AssetRegistryTags)
+		UObject* DefaultObject = Class->GetDefaultObject();
+		FAssetRegistryTagsContextData TagsContext(DefaultObject, EAssetRegistryTagsCaller::Uncategorized);
+		DefaultObject->GetAssetRegistryTags(TagsContext);
+
+		for(const TPair<FName,UObject::FAssetRegistryTag>& TagPair : TagsContext.Tags)
 		{
-			if(AssetRegistryTag.Type != UObject::FAssetRegistryTag::TT_Hidden)
+			if(TagPair.Value.Type != UObject::FAssetRegistryTag::TT_Hidden)
 			{
-				AssetPickerConfig.HiddenColumnNames.AddUnique(AssetRegistryTag.Name.ToString());
+				AssetPickerConfig.HiddenColumnNames.AddUnique(TagPair.Key.ToString());
 			}
 		}
 	}
diff --git a/Engine/Source/Editor/Persona/Private/SAnimationSequenceBrowser.cpp b/Engine/Source/Editor/Persona/Private/SAnimationSequenceBrowser.cpp
--- a/Engine/Source/Editor/Persona/Private/SAnimationSequenceBrowser.cpp
+++ b/Engine/Source/Editor/Persona/Private/SAnimationSequenceBrowser.cpp
@@ -49,6 +49,7 @@
 #include "AnimationSequenceBrowserMenuContexts.h"
 #include "Viewports.h"
 #include "SceneInterface.h"
+#include "UObject/AssetRegistryTagsContext.h"
 
 #define LOCTEXT_NAMESPACE "SequenceBrowser"
 
@@ -1041,11 +1042,12 @@ void SAnimationSequenceBrowser::Construct(const FArguments& InArgs, const TShare
 	Config.OnAssetToolTipClosing = FOnAssetToolTipClosing::CreateSP( this, &SAnimationSequenceBrowser::OnAssetToolTipClosing );
 
 	// hide all asset registry columns by default (we only really want the name and path)
-	TArray<UObject::FAssetRegistryTag> AssetRegistryTags;
-	UAnimSequence::StaticClass()->GetDefaultObject()->GetAssetRegistryTags(AssetRegistryTags);
-	for(UObject::FAssetRegistryTag& AssetRegistryTag : AssetRegistryTags)
+	UObject* AnimSequenceDefaultObject = UAnimSequence::StaticClass()->GetDefaultObject();
+	FAssetRegistryTagsContextData TagsContext(AnimSequenceDefaultObject, EAssetRegistryTagsCaller::Uncategorized);
+	AnimSequenceDefaultObject->GetAssetRegistryTags(TagsContext);
+	for (const TPair<FName,UObject::FAssetRegistryTag>& TagPair : TagsContext.Tags)
 	{
-		Config.HiddenColumnNames.Add(AssetRegistryTag.Name.ToString());
+		Config.HiddenColumnNames.Add(TagPair.Key.ToString());
 	}
 
 	// Also hide the type column by default (but allow users to enable it, so don't use bShowTypeInColumnView)
@@ -1380,14 +1382,17 @@ TSharedRef<SToolTip> SAnimationSequenceBrowser::CreateCustomAssetToolTip(FAssetD
 	TArray<UObject::FAssetRegistryTag> Tags;
 	UClass* AssetClass = FindObject<UClass>(AssetData.AssetClassPath);
 	check(AssetClass);
-	AssetClass->GetDefaultObject()->GetAssetRegistryTags(Tags);
+	UObject* DefaultObject = AssetClass->GetDefaultObject();
+	FAssetRegistryTagsContextData TagsContext(DefaultObject, EAssetRegistryTagsCaller::Uncategorized);
+	DefaultObject->GetAssetRegistryTags(TagsContext);
 
 	TArray<FName> TagsToShow;
-	for(UObject::FAssetRegistryTag& TagEntry : Tags)
+	FName NameSkeleton(TEXT("Skeleton"));
+	for (const TPair<FName, UObject::FAssetRegistryTag>& TagPair : TagsContext.Tags)
 	{
-		if(TagEntry.Name != FName(TEXT("Skeleton")) && TagEntry.Type != UObject::FAssetRegistryTag::TT_Hidden)
+		if(TagPair.Key != NameSkeleton && TagPair.Value.Type != UObject::FAssetRegistryTag::TT_Hidden)
 		{
-			TagsToShow.Add(TagEntry.Name);
+			TagsToShow.Add(TagPair.Key);
 		}
 	}
 
diff --git a/Engine/Source/Editor/SkeletonEditor/Private/EditableSkeleton.cpp b/Engine/Source/Editor/SkeletonEditor/Private/EditableSkeleton.cpp
--- a/Engine/Source/Editor/SkeletonEditor/Private/EditableSkeleton.cpp
+++ b/Engine/Source/Editor/SkeletonEditor/Private/EditableSkeleton.cpp
@@ -5,6 +5,7 @@
 #include "Misc/MessageDialog.h"
 #include "Misc/FeedbackContext.h"
 #include "Modules/ModuleManager.h"
+#include "UObject/AssetRegistryTagsContext.h"
 #include "UObject/UObjectHash.h"
 #include "UObject/UObjectIterator.h"
 #include "AssetRegistry/AssetData.h"
@@ -295,14 +296,10 @@ void FEditableSkeleton::GetAssetsContainingCurves(const FName& InContainerName,
 				// using the selected name. We only load what we have to here.
 				UObject* Asset = Data.GetAsset();
 				check(Asset);
-				TArray<UObject::FAssetRegistryTag> Tags;
-				Asset->GetAssetRegistryTags(Tags);
+				FAssetRegistryTagsContextData TagsContext(Asset, EAssetRegistryTagsCaller::Uncategorized);
+				Asset->GetAssetRegistryTags(TagsContext);
 
-				UObject::FAssetRegistryTag* CurveTag = Tags.FindByPredicate([](const UObject::FAssetRegistryTag& InTag)
-				{
-					return InTag.Name == USkeleton::CurveNameTag;
-				});
-				
+				UObject::FAssetRegistryTag* CurveTag = TagsContext.Tags.Find(USkeleton::CurveNameTag);
 				if (CurveTag)
 				{
 					CurveData = CurveTag->Value;
diff --git a/Engine/Source/Editor/UMGEditor/Private/WidgetBlueprint.cpp b/Engine/Source/Editor/UMGEditor/Private/WidgetBlueprint.cpp
--- a/Engine/Source/Editor/UMGEditor/Private/WidgetBlueprint.cpp
+++ b/Engine/Source/Editor/UMGEditor/Private/WidgetBlueprint.cpp
@@ -16,6 +16,7 @@
 #include "Kismet2/CompilerResultsLog.h"
 #include "Binding/PropertyBinding.h"
 #include "Blueprint/WidgetBlueprintGeneratedClass.h"
+#include "UObject/AssetRegistryTagsContext.h"
 #include "UObject/PropertyTag.h"
 #include "WidgetBlueprintCompiler.h"
 #include "UObject/EditorObjectVersion.h"
@@ -41,7 +42,10 @@
 
 #define LOCTEXT_NAMESPACE "UMG"
 
+PRAGMA_DISABLE_DEPRECATION_WARNINGS;
 FWidgetBlueprintDelegates::FGetAssetTags FWidgetBlueprintDelegates::GetAssetTags;
+PRAGMA_ENABLE_DEPRECATION_WARNINGS;
+FWidgetBlueprintDelegates::FGetAssetTagsWithContext FWidgetBlueprintDelegates::GetAssetTagsWithContext;
 
 FEditorPropertyPathSegment::FEditorPropertyPathSegment()
 	: Struct(nullptr)
@@ -720,9 +724,24 @@ void UWidgetBlueprint::PreSave(FObjectPreSaveContext ObjectSaveContext)
 
 void UWidgetBlueprint::GetAssetRegistryTags(TArray<FAssetRegistryTag>& OutTags) const
 {
+	PRAGMA_DISABLE_DEPRECATION_WARNINGS;
 	Super::GetAssetRegistryTags(OutTags);
+	PRAGMA_ENABLE_DEPRECATION_WARNINGS;
+}
+
+void UWidgetBlueprint::GetAssetRegistryTags(FAssetRegistryTagsContext Context) const
+{
+	Super::GetAssetRegistryTags(Context);
 
-	FWidgetBlueprintDelegates::GetAssetTags.Broadcast(this, OutTags);
+	PRAGMA_DISABLE_DEPRECATION_WARNINGS;
+	TArray<UObject::FAssetRegistryTag> DeprecatedFunctionTags;
+	FWidgetBlueprintDelegates::GetAssetTags.Broadcast(this, DeprecatedFunctionTags);
+	for (UObject::FAssetRegistryTag& Tag : DeprecatedFunctionTags)
+	{
+		Context.AddTag(MoveTemp(Tag));
+	}
+	PRAGMA_ENABLE_DEPRECATION_WARNINGS;
+	FWidgetBlueprintDelegates::GetAssetTagsWithContext.Broadcast(this, Context);
 }
 
 void UWidgetBlueprint::NotifyGraphRenamed(class UEdGraph* Graph, FName OldName, FName NewName)
diff --git a/Engine/Source/Editor/UMGEditor/Public/WidgetBlueprint.h b/Engine/Source/Editor/UMGEditor/Public/WidgetBlueprint.h
--- a/Engine/Source/Editor/UMGEditor/Public/WidgetBlueprint.h
+++ b/Engine/Source/Editor/UMGEditor/Public/WidgetBlueprint.h
@@ -33,9 +33,12 @@ class UMGEDITOR_API FWidgetBlueprintDelegates
 {
 public:
 	// delegate for generating widget asset registry tags.
+	DECLARE_MULTICAST_DELEGATE_TwoParams(FGetAssetTagsWithContext, const UWidgetBlueprint*, FAssetRegistryTagsContext);
 	DECLARE_MULTICAST_DELEGATE_TwoParams(FGetAssetTags, const UWidgetBlueprint*, TArray<UObject::FAssetRegistryTag>&);
 
 	// called by UWdgetBlueprint::GetAssetRegistryTags()
+	static FGetAssetTagsWithContext GetAssetTagsWithContext;
+	UE_DEPRECATED(5.4, "Subscribe to GetAssetTagsWithContext instead.")
 	static FGetAssetTags GetAssetTags;
 };
 
@@ -283,6 +286,8 @@ public:
 #endif // WITH_EDITORONLY_DATA
 
 #if WITH_EDITOR
+	virtual void GetAssetRegistryTags(FAssetRegistryTagsContext Context) const override;
+	UE_DEPRECATED(5.4, "Implement the version that takes FAssetRegistryTagsContext instead.")
 	virtual void GetAssetRegistryTags(TArray<FAssetRegistryTag>& OutTags) const override;
 	virtual void NotifyGraphRenamed(class UEdGraph* Graph, FName OldName, FName NewName) override;
 	virtual EDataValidationResult IsDataValid(class FDataValidationContext& Context) const override;
diff --git a/Engine/Source/Editor/UnrealEd/Classes/Factories/FbxSceneImportData.h b/Engine/Source/Editor/UnrealEd/Classes/Factories/FbxSceneImportData.h
--- a/Engine/Source/Editor/UnrealEd/Classes/Factories/FbxSceneImportData.h
+++ b/Engine/Source/Editor/UnrealEd/Classes/Factories/FbxSceneImportData.h
@@ -60,6 +60,8 @@ public:
 	/** Attempt to parse an asset import structure from the specified json string. */
 	UNREALED_API void FromJson(FString InJsonString);
 
+	UNREALED_API virtual void GetAssetRegistryTags(FAssetRegistryTagsContext Context) const override;
+	UE_DEPRECATED(5.4, "Implement the version that takes FAssetRegistryTagsContext instead.")
 	UNREALED_API virtual void GetAssetRegistryTags(TArray<FAssetRegistryTag>& OutTags) const override;
 
 protected:
diff --git a/Engine/Source/Editor/UnrealEd/Classes/Factories/FbxSkeletalMeshImportData.h b/Engine/Source/Editor/UnrealEd/Classes/Factories/FbxSkeletalMeshImportData.h
--- a/Engine/Source/Editor/UnrealEd/Classes/Factories/FbxSkeletalMeshImportData.h
+++ b/Engine/Source/Editor/UnrealEd/Classes/Factories/FbxSkeletalMeshImportData.h
@@ -116,5 +116,7 @@ public:
 	bool GetImportContentFilename(FString& OutFilename, FString& OutFilenameLabel) const;
 
 	/** This function add the last import content type to the asset registry which is use by the thumbnail overlay of the skeletal mesh */
-	virtual void AppendAssetRegistryTags(TArray<FAssetRegistryTag>& OutTags);
+	virtual void AppendAssetRegistryTags(FAssetRegistryTagsContext Context) override;
+	UE_DEPRECATED(5.4, "Implement the version that takes FAssetRegistryTagsContext instead.")
+	virtual void AppendAssetRegistryTags(TArray<FAssetRegistryTag>& OutTags) override;
 };
\ No newline at end of file
diff --git a/Engine/Source/Editor/UnrealEd/Private/AutoReimport/AutoReimportUtilities.cpp b/Engine/Source/Editor/UnrealEd/Private/AutoReimport/AutoReimportUtilities.cpp
--- a/Engine/Source/Editor/UnrealEd/Private/AutoReimport/AutoReimportUtilities.cpp
+++ b/Engine/Source/Editor/UnrealEd/Private/AutoReimport/AutoReimportUtilities.cpp
@@ -7,6 +7,7 @@
 #include "HAL/Platform.h"
 #include "HAL/PlatformCrt.h"
 #include "Misc/Optional.h"
+#include "UObject/AssetRegistryTagsContext.h"
 #include "UObject/NameTypes.h"
 #include "UObject/Object.h"
 
@@ -30,12 +31,13 @@ namespace Utils
 
 	void ExtractSourceFilePaths(UObject* Object, TArray<FString>& OutSourceFiles)
 	{
-		TArray<UObject::FAssetRegistryTag> TagList;
+		FAssetRegistryTagsContextData TagList(Object, EAssetRegistryTagsCaller::Uncategorized);
 		Object->GetAssetRegistryTags(TagList);
 
 		const FName TagName = UObject::SourceFileTagName();
-		for (const auto& Tag : TagList)
+		for (const TPair<FName, UObject::FAssetRegistryTag>& Pair : TagList.Tags)
 		{
+			const UObject::FAssetRegistryTag& Tag = Pair.Value;
 			if (Tag.Name == TagName)
 			{
 				int32 PreviousNum = OutSourceFiles.Num();
diff --git a/Engine/Source/Editor/UnrealEd/Private/Fbx/FbxSceneImportData.cpp b/Engine/Source/Editor/UnrealEd/Private/Fbx/FbxSceneImportData.cpp
--- a/Engine/Source/Editor/UnrealEd/Private/Fbx/FbxSceneImportData.cpp
+++ b/Engine/Source/Editor/UnrealEd/Private/Fbx/FbxSceneImportData.cpp
@@ -568,12 +568,19 @@ void UFbxSceneImportData::FromJson(FString InJsonString)
 }
 
 void UFbxSceneImportData::GetAssetRegistryTags(TArray<FAssetRegistryTag>& OutTags) const
+{
+	PRAGMA_DISABLE_DEPRECATION_WARNINGS;
+	Super::GetAssetRegistryTags(OutTags);
+	PRAGMA_ENABLE_DEPRECATION_WARNINGS;
+}
+
+void UFbxSceneImportData::GetAssetRegistryTags(FAssetRegistryTagsContext Context) const
 {
 	FAssetImportInfo AssetImportInfo;
 	FAssetImportInfo::FSourceFile SourceFile(SourceFbxFile);
 	AssetImportInfo.Insert(SourceFile);
-	OutTags.Add(FAssetRegistryTag(SourceFileTagName(), AssetImportInfo.ToJson(), FAssetRegistryTag::TT_Hidden));
-	Super::GetAssetRegistryTags(OutTags);
+	Context.AddTag(FAssetRegistryTag(SourceFileTagName(), AssetImportInfo.ToJson(), FAssetRegistryTag::TT_Hidden));
+	Super::GetAssetRegistryTags(Context);
 }
 
 void UFbxSceneImportData::Serialize(FArchive& Ar)
diff --git a/Engine/Source/Editor/UnrealEd/Private/Fbx/FbxSkeletalMeshImportData.cpp b/Engine/Source/Editor/UnrealEd/Private/Fbx/FbxSkeletalMeshImportData.cpp
--- a/Engine/Source/Editor/UnrealEd/Private/Fbx/FbxSkeletalMeshImportData.cpp
+++ b/Engine/Source/Editor/UnrealEd/Private/Fbx/FbxSkeletalMeshImportData.cpp
@@ -2,6 +2,7 @@
 
 #include "Factories/FbxSkeletalMeshImportData.h"
 #include "Engine/SkeletalMesh.h"
+#include "UObject/AssetRegistryTagsContext.h"
 #include "UObject/Package.h"
 #include "UObject/EditorObjectVersion.h"
 
@@ -83,6 +84,13 @@ bool UFbxSkeletalMeshImportData::GetImportContentFilename(FString& OutFilename,
 }
 
 void UFbxSkeletalMeshImportData::AppendAssetRegistryTags(TArray<FAssetRegistryTag>& OutTags)
+{
+	PRAGMA_DISABLE_DEPRECATION_WARNINGS;
+	Super::AppendAssetRegistryTags(OutTags);
+	PRAGMA_ENABLE_DEPRECATION_WARNINGS;
+}
+
+void UFbxSkeletalMeshImportData::AppendAssetRegistryTags(FAssetRegistryTagsContext Context)
 {
 	auto EFBXImportContentTypeToString = [](const EFBXImportContentType value)-> FString
 	{
@@ -101,7 +109,7 @@ void UFbxSkeletalMeshImportData::AppendAssetRegistryTags(TArray<FAssetRegistryTa
 	};
 
 	FString EnumString = EFBXImportContentTypeToString(LastImportContentType);
-	OutTags.Add(FAssetRegistryTag("LastImportContentType", EnumString, FAssetRegistryTag::TT_Hidden));
+	Context.AddTag(FAssetRegistryTag("LastImportContentType", EnumString, FAssetRegistryTag::TT_Hidden));
 
-	Super::AppendAssetRegistryTags(OutTags);
+	Super::AppendAssetRegistryTags(Context);
 }
\ No newline at end of file
diff --git a/Engine/Source/Runtime/AssetRegistry/Private/AssetRegistry.cpp b/Engine/Source/Runtime/AssetRegistry/Private/AssetRegistry.cpp
--- a/Engine/Source/Runtime/AssetRegistry/Private/AssetRegistry.cpp
+++ b/Engine/Source/Runtime/AssetRegistry/Private/AssetRegistry.cpp
@@ -7031,7 +7031,7 @@ void UAssetRegistryImpl::GetInheritanceContextAfterVerifyingLock(uint64 CurrentG
 }
 
 #if WITH_EDITOR
-void UAssetRegistryImpl::OnGetExtraObjectTags(FAssetRegistryTagsContext& Context)
+void UAssetRegistryImpl::OnGetExtraObjectTags(FAssetRegistryTagsContext Context)
 {
 	if (bAddMetaDataTagsToOnGetExtraObjectTags)
 	{
diff --git a/Engine/Source/Runtime/AssetRegistry/Private/AssetRegistry.h b/Engine/Source/Runtime/AssetRegistry/Private/AssetRegistry.h
--- a/Engine/Source/Runtime/AssetRegistry/Private/AssetRegistry.h
+++ b/Engine/Source/Runtime/AssetRegistry/Private/AssetRegistry.h
@@ -292,7 +292,7 @@ private:
 	 * Callback for FObject::FAssetRegistryTag::OnGetExtraObjectTags
 	 * If bAddMetaDataTagsToOnGetExtraObjectTags is true, this function will add missing UMetaData tags to cooked assets
 	 */
-	void OnGetExtraObjectTags(FAssetRegistryTagsContext& Context);
+	void OnGetExtraObjectTags(FAssetRegistryTagsContext Context);
 #endif
 
 private:
diff --git a/Engine/Source/Runtime/CoreUObject/Private/UObject/AssetRegistryTagsContext.cpp b/Engine/Source/Runtime/CoreUObject/Private/UObject/AssetRegistryTagsContext.cpp
--- a/Engine/Source/Runtime/CoreUObject/Private/UObject/AssetRegistryTagsContext.cpp
+++ b/Engine/Source/Runtime/CoreUObject/Private/UObject/AssetRegistryTagsContext.cpp
@@ -17,7 +17,7 @@ FAssetRegistryTagsContext::FAssetRegistryTagsContext(FAssetRegistryTagsContextDa
 {
 }
 
-FAssetRegistryTagsContext::FAssetRegistryTagsContext(FAssetRegistryTagsContext& Other)
+FAssetRegistryTagsContext::FAssetRegistryTagsContext(const FAssetRegistryTagsContext& Other)
 	: Data(Other.Data)
 {
 }
@@ -102,6 +102,19 @@ void FAssetRegistryTagsContext::AddTag(UObject::FAssetRegistryTag TagResult)
 	AddTagInternal(MoveTemp(TagResult), FStringView());
 }
 
+int32 FAssetRegistryTagsContext::GetNumTags() const
+{
+	return Data.Tags.Num();
+}
+
+void FAssetRegistryTagsContext::EnumerateTags(TFunctionRef<void(const UObject::FAssetRegistryTag&)> Visitor) const
+{
+	for (const TPair<FName, UObject::FAssetRegistryTag>& Pair : Data.Tags)
+	{
+		Visitor(Pair.Value);
+	}
+}
+
 bool FAssetRegistryTagsContext::WantsBundleResult() const
 {
 	return Data.bWantsBundleResult;
diff --git a/Engine/Source/Runtime/CoreUObject/Private/UObject/Class.cpp b/Engine/Source/Runtime/CoreUObject/Private/UObject/Class.cpp
--- a/Engine/Source/Runtime/CoreUObject/Private/UObject/Class.cpp
+++ b/Engine/Source/Runtime/CoreUObject/Private/UObject/Class.cpp
@@ -17,6 +17,7 @@
 #include "Misc/OutputDeviceNull.h"
 #include "UObject/CoreNet.h"
 #include "Modules/ModuleManager.h"
+#include "UObject/AssetRegistryTagsContext.h"
 #include "UObject/UObjectAllocator.h"
 #include "UObject/UObjectHash.h"
 #include "UObject/UObjectIterator.h"
@@ -4863,20 +4864,27 @@ FString UClass::GetDesc()
 
 void UClass::GetAssetRegistryTags(TArray<FAssetRegistryTag>& OutTags) const
 {
+	PRAGMA_DISABLE_DEPRECATION_WARNINGS;
 	Super::GetAssetRegistryTags(OutTags);
+	PRAGMA_ENABLE_DEPRECATION_WARNINGS;
+}
+
+void UClass::GetAssetRegistryTags(FAssetRegistryTagsContext Context) const
+{
+	Super::GetAssetRegistryTags(Context);
 
 #if WITH_EDITOR
 	static const FName ParentClassFName = "ParentClass";
 	const UClass* const ParentClass = GetSuperClass();
-	OutTags.Add( FAssetRegistryTag(ParentClassFName, ((ParentClass) ? ParentClass->GetPathName() : FString()), FAssetRegistryTag::TT_Alphabetical) );
+	Context.AddTag( FAssetRegistryTag(ParentClassFName, ((ParentClass) ? ParentClass->GetPathName() : FString()), FAssetRegistryTag::TT_Alphabetical) );
 
 	static const FName ModuleNameFName = "ModuleName";
 	const UPackage* const ClassPackage = GetOuterUPackage();
-	OutTags.Add( FAssetRegistryTag(ModuleNameFName, ((ClassPackage) ? FPackageName::GetShortFName(ClassPackage->GetFName()) : NAME_None).ToString(), FAssetRegistryTag::TT_Alphabetical) );
+	Context.AddTag( FAssetRegistryTag(ModuleNameFName, ((ClassPackage) ? FPackageName::GetShortFName(ClassPackage->GetFName()) : NAME_None).ToString(), FAssetRegistryTag::TT_Alphabetical) );
 
 	static const FName ModuleRelativePathFName = "ModuleRelativePath";
 	const FString& ClassModuleRelativeIncludePath = GetMetaData(ModuleRelativePathFName);
-	OutTags.Add( FAssetRegistryTag(ModuleRelativePathFName, ClassModuleRelativeIncludePath, FAssetRegistryTag::TT_Alphabetical) );
+	Context.AddTag( FAssetRegistryTag(ModuleRelativePathFName, ClassModuleRelativeIncludePath, FAssetRegistryTag::TT_Alphabetical) );
 #endif
 }
 
diff --git a/Engine/Source/Runtime/CoreUObject/Private/UObject/Obj.cpp b/Engine/Source/Runtime/CoreUObject/Private/UObject/Obj.cpp
--- a/Engine/Source/Runtime/CoreUObject/Private/UObject/Obj.cpp
+++ b/Engine/Source/Runtime/CoreUObject/Private/UObject/Obj.cpp
@@ -132,9 +132,11 @@ static UPackage*			GObjTransientPkg								= NULL;
 		}
 	};
 	static FUObjectAnnotationSparse<FPropagatedEditChangeAnnotation, true> PropagatedEditChangeAnnotation;
+	PRAGMA_DISABLE_DEPRECATION_WARNINGS;
 	UObject::FAssetRegistryTag::FOnGetObjectAssetRegistryTags UObject::FAssetRegistryTag::OnGetExtraObjectTags;
-	UObject::FAssetRegistryTag::FOnGetObjectAssetRegistryTagsWithContext UObject::FAssetRegistryTag::OnGetExtraObjectTagsWithContext;
 	UObject::FAssetRegistryTag::FOnGetExtendedAssetRegistryTagsForSave UObject::FAssetRegistryTag::OnGetExtendedAssetRegistryTagsForSave;
+	PRAGMA_ENABLE_DEPRECATION_WARNINGS;
+	UObject::FAssetRegistryTag::FOnGetObjectAssetRegistryTagsWithContext UObject::FAssetRegistryTag::OnGetExtraObjectTagsWithContext;
 	UObject::FOnGetPreviewPlatform UObject::OnGetPreviewPlatform;
 #endif // WITH_EDITOR
 
@@ -2305,7 +2307,9 @@ void UObject::GetExtendedAssetRegistryTagsForSave(const ITargetPlatform* TargetP
 	// EAssetRegistryTagsCaller::AssetRegistryLoad. Because this function was designed only to be called from SavePackage,
 	// and we have removed SavePackage's dependence on it, we decided not to make that extra effort. Any licensee calling
 	// this function should instead call GetAssetRegistryTags with EAssetRegistryTagsCaller::SavePackage.
+	PRAGMA_DISABLE_DEPRECATION_WARNINGS;
 	FAssetRegistryTag::OnGetExtendedAssetRegistryTagsForSave.Broadcast(this, TargetPlatform, OutTags);
+	PRAGMA_ENABLE_DEPRECATION_WARNINGS;
 }
 #endif // WITH_EDITOR
 
diff --git a/Engine/Source/Runtime/CoreUObject/Private/UObject/ObjectRedirector.cpp b/Engine/Source/Runtime/CoreUObject/Private/UObject/ObjectRedirector.cpp
--- a/Engine/Source/Runtime/CoreUObject/Private/UObject/ObjectRedirector.cpp
+++ b/Engine/Source/Runtime/CoreUObject/Private/UObject/ObjectRedirector.cpp
@@ -6,6 +6,7 @@
 
 #include "UObject/ObjectRedirector.h"
 
+#include "UObject/AssetRegistryTagsContext.h"
 #include "UObject/ObjectSaveContext.h"
 #include "UObject/Package.h"
 #include "Templates/Casts.h"
@@ -62,7 +63,14 @@ bool UObjectRedirector::NeedsLoadForEditorGame() const
 
 void UObjectRedirector::GetAssetRegistryTags(TArray<FAssetRegistryTag>& OutTags) const
 {
+	PRAGMA_DISABLE_DEPRECATION_WARNINGS;
 	Super::GetAssetRegistryTags(OutTags);
+	PRAGMA_ENABLE_DEPRECATION_WARNINGS;
+}
+
+void UObjectRedirector::GetAssetRegistryTags(FAssetRegistryTagsContext Context) const
+{
+	Super::GetAssetRegistryTags(Context);
 
 	FString DestVal;
 	if ( DestinationObject != nullptr )
@@ -74,7 +82,7 @@ void UObjectRedirector::GetAssetRegistryTags(TArray<FAssetRegistryTag>& OutTags)
 		DestVal = TEXT("None");
 	}
 
-	OutTags.Add(FAssetRegistryTag("DestinationObject", DestVal, UObject::FAssetRegistryTag::TT_Alphabetical));
+	Context.AddTag(FAssetRegistryTag("DestinationObject", DestVal, UObject::FAssetRegistryTag::TT_Alphabetical));
 }
 
 /**
diff --git a/Engine/Source/Runtime/CoreUObject/Public/UObject/AssetRegistryTagsContext.h b/Engine/Source/Runtime/CoreUObject/Public/UObject/AssetRegistryTagsContext.h
--- a/Engine/Source/Runtime/CoreUObject/Public/UObject/AssetRegistryTagsContext.h
+++ b/Engine/Source/Runtime/CoreUObject/Public/UObject/AssetRegistryTagsContext.h
@@ -6,6 +6,7 @@
 #include "Containers/Map.h"
 #include "Containers/StringView.h"
 #include "Misc/Optional.h"
+#include "Templates/Function.h"
 #include "Templates/UnrealTemplate.h"
 #include "UObject/CookEnums.h"
 #include "UObject/Object.h" // UObject::FAssetRegistryTag
@@ -98,7 +99,7 @@ class FAssetRegistryTagsContext
 public:
 	// This is an implicit constructor to reduce the boilerplate for calling GetAssetRegistryTags
 	COREUOBJECT_API FAssetRegistryTagsContext(FAssetRegistryTagsContextData& InData);
-	COREUOBJECT_API FAssetRegistryTagsContext(FAssetRegistryTagsContext& Other);
+	COREUOBJECT_API FAssetRegistryTagsContext(const FAssetRegistryTagsContext& Other);
 
 	COREUOBJECT_API EAssetRegistryTagsCaller GetCaller() const;
 	COREUOBJECT_API const UObject* GetObject() const;
@@ -155,6 +156,10 @@ public:
 	COREUOBJECT_API bool ContainsTag(FName TagName) const;
 	/** Output function: move the given tag into the results, overwriting previous result if it exists. */
 	COREUOBJECT_API void AddTag(UObject::FAssetRegistryTag TagResult);
+	/** Return the number of tags so far reported to the context, including CookTags. */
+	COREUOBJECT_API int32 GetNumTags() const;
+	/** Pass each tag so far reported to the context into the provided Visitor, including CookTags. */
+	COREUOBJECT_API void EnumerateTags(TFunctionRef<void(const UObject::FAssetRegistryTag&)> Visitor) const;
 
 	/**
 	 * True if the caller wants the FAssetBundleData structure (if present) to be stored as a pointer.
