- Rigid Mesh DI / SM DI move from RequiresDistanceFieldData to RequiresEarlyViewData as mesh distance fields are ready once we have a view to read the data
- Rename RequiresDistanceFieldData to RequiresGlobalDistanceField to be more explicit

diff --git a/Engine/Plugins/FX/Niagara/Shaders/Private/NiagaraDataInterfaceStaticMeshTemplate.ush b/Engine/Plugins/FX/Niagara/Shaders/Private/NiagaraDataInterfaceStaticMeshTemplate.ush
--- a/Engine/Plugins/FX/Niagara/Shaders/Private/NiagaraDataInterfaceStaticMeshTemplate.ush
+++ b/Engine/Plugins/FX/Niagara/Shaders/Private/NiagaraDataInterfaceStaticMeshTemplate.ush
@@ -766,21 +766,21 @@ void GetMeshBoundsWS_{ParameterName}(out float3 Center, out float3 ExtentsMin, o
 // Experimental Distance Field Functions
 void QueryDistanceField_{ParameterName}(bool bExecute, float3 WorldPosition, bool bUseMaxDistance, float MaxDistance, out bool OutIsValid, out float OutDistance)
 {
-	const FLWCVector3 LWCWorldPosition = MakeLWCVector3({ParameterName}_SystemLWCTile, WorldPosition);
-
 	if ( !bUseMaxDistance )
 	{
 		MaxDistance = 1e38;
 	}
 
 #if DISTATICMESH_ALLOWDISTANCEFIELD
-	OutIsValid = bExecute && {ParameterName}_InstanceDistanceFieldIndex != int(-1);
+	const int DFIndex = {ParameterName}_InstanceDistanceFieldIndex;
+	OutIsValid = bExecute && DFIndex >= 0 && DFIndex < NumSceneObjects;
 	OutDistance = MaxDistance;
 	if ( OutIsValid )
 	{
 		// Needs to be a smidge over 0.0 for DistanceToNearestSurfaceForObject to query the SDF
-		MaxDistance = max(0.001, MaxDistance);
-		OutDistance = DistanceToNearestSurfaceForObject({ParameterName}_InstanceDistanceFieldIndex, LWCWorldPosition, MaxDistance);
+		MaxDistance = max(0.001f, MaxDistance);
+		const FLWCVector3 LWCWorldPosition = MakeLWCVector3({ParameterName}_SystemLWCTile, WorldPosition);
+		OutDistance = DistanceToNearestSurfaceForObject(DFIndex, LWCWorldPosition, MaxDistance);
 	}
 #else
 	OutIsValid = false;
@@ -788,7 +788,6 @@ void QueryDistanceField_{ParameterName}(bool bExecute, float3 WorldPosition, boo
 #endif
 }
 
-
 ////////////////////////////////////////////////////////////////////////////////////////////////////////////////////////////////////////////////////////////////
 // Uv Mapping
 void GetTriangleCoordAtUV_{ParameterName}(bool Enabled, float2 InUv, float InTolerance, out int OutTriangle, out float3 OutBaryCoord, out bool OutIsValid)
diff --git a/Engine/Plugins/FX/Niagara/Source/Niagara/Classes/NiagaraDataInterface.h b/Engine/Plugins/FX/Niagara/Source/Niagara/Classes/NiagaraDataInterface.h
--- a/Engine/Plugins/FX/Niagara/Source/Niagara/Classes/NiagaraDataInterface.h
+++ b/Engine/Plugins/FX/Niagara/Source/Niagara/Classes/NiagaraDataInterface.h
@@ -683,7 +683,15 @@ public:
 	*/
 	virtual bool PostStageCanOverlapTickGroups() const { return true; }
 
+	UE_DEPRECATED(5.4, "RequiresDistanceFieldData was renamed to RequiresGlobalDistanceField.")
 	virtual bool RequiresDistanceFieldData() const { return false; }
+
+	virtual bool RequiresGlobalDistanceField() const
+	{
+		PRAGMA_DISABLE_DEPRECATION_WARNINGS
+		return RequiresDistanceFieldData();
+		PRAGMA_ENABLE_DEPRECATION_WARNINGS
+	}
 	virtual bool RequiresDepthBuffer() const { return false; }
 	virtual bool RequiresEarlyViewData() const { return false; }
 	virtual bool RequiresRayTracingScene() const { return false; }
diff --git a/Engine/Plugins/FX/Niagara/Source/Niagara/Classes/NiagaraDataInterfaceAsyncGpuTrace.h b/Engine/Plugins/FX/Niagara/Source/Niagara/Classes/NiagaraDataInterfaceAsyncGpuTrace.h
--- a/Engine/Plugins/FX/Niagara/Source/Niagara/Classes/NiagaraDataInterfaceAsyncGpuTrace.h
+++ b/Engine/Plugins/FX/Niagara/Source/Niagara/Classes/NiagaraDataInterfaceAsyncGpuTrace.h
@@ -47,7 +47,7 @@ public:
 	NIAGARA_API virtual void GetFunctions(TArray<FNiagaraFunctionSignature>& OutFunctions) override;
 
 	virtual bool CanExecuteOnTarget(ENiagaraSimTarget Target) const override { return Target == ENiagaraSimTarget::GPUComputeSim; }
-	NIAGARA_API virtual bool RequiresDistanceFieldData() const override;
+	NIAGARA_API virtual bool RequiresGlobalDistanceField() const override;
 	NIAGARA_API virtual bool RequiresRayTracingScene() const override;
 
 #if WITH_EDITORONLY_DATA
diff --git a/Engine/Plugins/FX/Niagara/Source/Niagara/Classes/NiagaraDataInterfaceAudioOscilloscope.h b/Engine/Plugins/FX/Niagara/Source/Niagara/Classes/NiagaraDataInterfaceAudioOscilloscope.h
--- a/Engine/Plugins/FX/Niagara/Source/Niagara/Classes/NiagaraDataInterfaceAudioOscilloscope.h
+++ b/Engine/Plugins/FX/Niagara/Source/Niagara/Classes/NiagaraDataInterfaceAudioOscilloscope.h
@@ -125,8 +125,6 @@ public:
 		return true;
 	}
 
-	virtual bool RequiresDistanceFieldData() const override { return false; }
-
 #if WITH_EDITORONLY_DATA
 	NIAGARA_API virtual bool AppendCompileHash(FNiagaraCompileHashVisitor* InVisitor) const override;
 	NIAGARA_API virtual bool GetFunctionHLSL(const FNiagaraDataInterfaceGPUParamInfo& ParamInfo, const FNiagaraDataInterfaceGeneratedFunction& FunctionInfo, int FunctionInstanceIndex, FString& OutHLSL) override;
diff --git a/Engine/Plugins/FX/Niagara/Source/Niagara/Classes/NiagaraDataInterfaceCollisionQuery.h b/Engine/Plugins/FX/Niagara/Source/Niagara/Classes/NiagaraDataInterfaceCollisionQuery.h
--- a/Engine/Plugins/FX/Niagara/Source/Niagara/Classes/NiagaraDataInterfaceCollisionQuery.h
+++ b/Engine/Plugins/FX/Niagara/Source/Niagara/Classes/NiagaraDataInterfaceCollisionQuery.h
@@ -55,7 +55,7 @@ public:
 	NIAGARA_API void PerformQueryAsyncCPU(FVectorVMExternalFunctionContext& Context);
 
 	virtual bool CanExecuteOnTarget(ENiagaraSimTarget Target) const override { return true; }
-	virtual bool RequiresDistanceFieldData() const override { return true; }
+	virtual bool RequiresGlobalDistanceField() const override { return true; }
 	virtual bool RequiresDepthBuffer() const override { return true; }
 
 #if WITH_EDITORONLY_DATA
diff --git a/Engine/Plugins/FX/Niagara/Source/Niagara/Internal/DataInterface/NiagaraDataInterfaceStaticMesh.h b/Engine/Plugins/FX/Niagara/Source/Niagara/Internal/DataInterface/NiagaraDataInterfaceStaticMesh.h
--- a/Engine/Plugins/FX/Niagara/Source/Niagara/Internal/DataInterface/NiagaraDataInterfaceStaticMesh.h
+++ b/Engine/Plugins/FX/Niagara/Source/Niagara/Internal/DataInterface/NiagaraDataInterfaceStaticMesh.h
@@ -162,7 +162,7 @@ public:
 #endif
 	NIAGARA_API virtual void GetVMExternalFunction(const FVMExternalFunctionBindingInfo& BindingInfo, void* InstanceData, FVMExternalFunction &OutFunc) override;
 	virtual bool CanExecuteOnTarget(ENiagaraSimTarget Target) const override { return true; }
-	NIAGARA_API virtual bool RequiresDistanceFieldData() const override;
+	virtual bool RequiresEarlyViewData() const override;
 	virtual bool HasPreSimulateTick() const override { return true; }
 
 #if WITH_EDITORONLY_DATA
diff --git a/Engine/Plugins/FX/Niagara/Source/Niagara/Private/DataInterface/NiagaraDataInterfaceAsyncGpuTrace.cpp b/Engine/Plugins/FX/Niagara/Source/Niagara/Private/DataInterface/NiagaraDataInterfaceAsyncGpuTrace.cpp
--- a/Engine/Plugins/FX/Niagara/Source/Niagara/Private/DataInterface/NiagaraDataInterfaceAsyncGpuTrace.cpp
+++ b/Engine/Plugins/FX/Niagara/Source/Niagara/Private/DataInterface/NiagaraDataInterfaceAsyncGpuTrace.cpp
@@ -327,9 +327,9 @@ void UNiagaraDataInterfaceAsyncGpuTrace::GetFunctions(TArray<FNiagaraFunctionSig
 	}
 }
 
-bool UNiagaraDataInterfaceAsyncGpuTrace::RequiresDistanceFieldData() const
+bool UNiagaraDataInterfaceAsyncGpuTrace::RequiresGlobalDistanceField() const
 {
-	return FNiagaraAsyncGpuTraceHelper::RequiresDistanceFieldData(TraceProvider);
+	return FNiagaraAsyncGpuTraceHelper::RequiresGlobalDistanceField(TraceProvider);
 }
 
 bool UNiagaraDataInterfaceAsyncGpuTrace::RequiresRayTracingScene() const
diff --git a/Engine/Plugins/FX/Niagara/Source/Niagara/Private/DataInterface/NiagaraDataInterfaceStaticMesh.cpp b/Engine/Plugins/FX/Niagara/Source/Niagara/Private/DataInterface/NiagaraDataInterfaceStaticMesh.cpp
--- a/Engine/Plugins/FX/Niagara/Source/Niagara/Private/DataInterface/NiagaraDataInterfaceStaticMesh.cpp
+++ b/Engine/Plugins/FX/Niagara/Source/Niagara/Private/DataInterface/NiagaraDataInterfaceStaticMesh.cpp
@@ -3009,7 +3009,7 @@ void UNiagaraDataInterfaceStaticMesh::GetVMExternalFunction(const FVMExternalFun
 	}
 }
 
-bool UNiagaraDataInterfaceStaticMesh::RequiresDistanceFieldData() const
+bool UNiagaraDataInterfaceStaticMesh::RequiresEarlyViewData() const
 {
 	return GetDefault<UNiagaraSettings>()->NDIStaticMesh_AllowDistanceFields;
 }
diff --git a/Engine/Plugins/FX/Niagara/Source/Niagara/Private/NiagaraAsyncGpuTraceHelper.cpp b/Engine/Plugins/FX/Niagara/Source/Niagara/Private/NiagaraAsyncGpuTraceHelper.cpp
--- a/Engine/Plugins/FX/Niagara/Source/Niagara/Private/NiagaraAsyncGpuTraceHelper.cpp
+++ b/Engine/Plugins/FX/Niagara/Source/Niagara/Private/NiagaraAsyncGpuTraceHelper.cpp
@@ -289,9 +289,9 @@ ENDICollisionQuery_AsyncGpuTraceProvider::Type FNiagaraAsyncGpuTraceHelper::Reso
 	return FNiagaraAsyncGpuTraceProvider::ResolveSupportedType(InType, GetDefault<UNiagaraSettings>()->NDICollisionQuery_AsyncGpuTraceProviderOrder);
 }
 
-bool FNiagaraAsyncGpuTraceHelper::RequiresDistanceFieldData(ENDICollisionQuery_AsyncGpuTraceProvider::Type InType)
+bool FNiagaraAsyncGpuTraceHelper::RequiresGlobalDistanceField(ENDICollisionQuery_AsyncGpuTraceProvider::Type InType)
 {
-	return FNiagaraAsyncGpuTraceProvider::RequiresDistanceFieldData(InType, GetDefault<UNiagaraSettings>()->NDICollisionQuery_AsyncGpuTraceProviderOrder);
+	return FNiagaraAsyncGpuTraceProvider::RequiresGlobalDistanceField(InType, GetDefault<UNiagaraSettings>()->NDICollisionQuery_AsyncGpuTraceProviderOrder);
 }
 
 bool FNiagaraAsyncGpuTraceHelper::RequiresRayTracingScene(ENDICollisionQuery_AsyncGpuTraceProvider::Type InType)
diff --git a/Engine/Plugins/FX/Niagara/Source/Niagara/Private/NiagaraAsyncGpuTraceHelper.h b/Engine/Plugins/FX/Niagara/Source/Niagara/Private/NiagaraAsyncGpuTraceHelper.h
--- a/Engine/Plugins/FX/Niagara/Source/Niagara/Private/NiagaraAsyncGpuTraceHelper.h
+++ b/Engine/Plugins/FX/Niagara/Source/Niagara/Private/NiagaraAsyncGpuTraceHelper.h
@@ -88,7 +88,7 @@ public:
 	FNiagaraGpuScratchPad GetTraceCounts() { return TraceCounts; }
 
 	static ENDICollisionQuery_AsyncGpuTraceProvider::Type ResolveSupportedType(ENDICollisionQuery_AsyncGpuTraceProvider::Type InType);
-	static bool RequiresDistanceFieldData(ENDICollisionQuery_AsyncGpuTraceProvider::Type InType);
+	static bool RequiresGlobalDistanceField(ENDICollisionQuery_AsyncGpuTraceProvider::Type InType);
 	static bool RequiresRayTracingScene(ENDICollisionQuery_AsyncGpuTraceProvider::Type InType);
 
 #if NIAGARA_ASYNC_GPU_TRACE_COLLISION_GROUPS
diff --git a/Engine/Plugins/FX/Niagara/Source/Niagara/Private/NiagaraDebugHud.cpp b/Engine/Plugins/FX/Niagara/Source/Niagara/Private/NiagaraDebugHud.cpp
--- a/Engine/Plugins/FX/Niagara/Source/Niagara/Private/NiagaraDebugHud.cpp
+++ b/Engine/Plugins/FX/Niagara/Source/Niagara/Private/NiagaraDebugHud.cpp
@@ -693,9 +693,9 @@ namespace NiagaraDebugLocal
 		if (Settings.SystemDebugVerbosity == ENiagaraDebugHudVerbosity::Verbose)
 		{
 			TStringBuilder<128> GpuFeaturesBuilder;
-			if (SystemInstance->RequiresDistanceFieldData())
+			if (SystemInstance->RequiresGlobalDistanceField())
 			{
-				GpuFeaturesBuilder.Append(TEXT(" DistanceFieldData"));
+				GpuFeaturesBuilder.Append(TEXT(" GlobalDistanceField"));
 			}
 			if (SystemInstance->RequiresDepthBuffer())
 			{
diff --git a/Engine/Plugins/FX/Niagara/Source/Niagara/Private/NiagaraDebuggerClient.cpp b/Engine/Plugins/FX/Niagara/Source/Niagara/Private/NiagaraDebuggerClient.cpp
--- a/Engine/Plugins/FX/Niagara/Source/Niagara/Private/NiagaraDebuggerClient.cpp
+++ b/Engine/Plugins/FX/Niagara/Source/Niagara/Private/NiagaraDebuggerClient.cpp
@@ -432,7 +432,7 @@ bool FNiagaraDebuggerClient::UpdateOutliner(float DeltaSeconds)
 					}
 					InstData.LWCTile = Inst->GetLWCTile();
 					InstData.bIsSolo = Inst->IsSolo();
-					InstData.bRequiresDistanceFieldData = Inst->RequiresDistanceFieldData();
+					InstData.bRequiresGlobalDistanceField = Inst->RequiresGlobalDistanceField();
 					InstData.bRequiresDepthBuffer = Inst->RequiresDepthBuffer();
 					InstData.bRequiresEarlyViewData = Inst->RequiresEarlyViewData();
 					InstData.bRequiresViewUniformBuffer = Inst->RequiresViewUniformBuffer();
diff --git a/Engine/Plugins/FX/Niagara/Source/Niagara/Private/NiagaraGpuComputeDispatch.cpp b/Engine/Plugins/FX/Niagara/Source/Niagara/Private/NiagaraGpuComputeDispatch.cpp
--- a/Engine/Plugins/FX/Niagara/Source/Niagara/Private/NiagaraGpuComputeDispatch.cpp
+++ b/Engine/Plugins/FX/Niagara/Source/Niagara/Private/NiagaraGpuComputeDispatch.cpp
@@ -246,10 +246,10 @@ void FNiagaraGpuComputeDispatch::AddGpuComputeProxy(FNiagaraSystemGpuComputeProx
 	ComputeProxy->ComputeDispatchIndex = ProxiesPerStage[TickStage].Num();
 	ProxiesPerStage[TickStage].Add(ComputeProxy);
 
-	NumProxiesThatRequireDistanceFieldData	+= ComputeProxy->RequiresDistanceFieldData() ? 1 : 0;
-	NumProxiesThatRequireDepthBuffer		+= ComputeProxy->RequiresDepthBuffer() ? 1 : 0;
-	NumProxiesThatRequireEarlyViewData		+= ComputeProxy->RequiresEarlyViewData() ? 1 : 0;
-	NumProxiesThatRequireRayTracingScene	+= ComputeProxy->RequiresRayTracingScene() ? 1 : 0;
+	NumProxiesThatRequireGlobalDistanceField	+= ComputeProxy->RequiresGlobalDistanceField() ? 1 : 0;
+	NumProxiesThatRequireDepthBuffer			+= ComputeProxy->RequiresDepthBuffer() ? 1 : 0;
+	NumProxiesThatRequireEarlyViewData			+= ComputeProxy->RequiresEarlyViewData() ? 1 : 0;
+	NumProxiesThatRequireRayTracingScene		+= ComputeProxy->RequiresRayTracingScene() ? 1 : 0;
 }
 
 void FNiagaraGpuComputeDispatch::RemoveGpuComputeProxy(FNiagaraSystemGpuComputeProxy* ComputeProxy)
@@ -267,10 +267,10 @@ void FNiagaraGpuComputeDispatch::RemoveGpuComputeProxy(FNiagaraSystemGpuComputeP
 	}
 	ComputeProxy->ComputeDispatchIndex = INDEX_NONE;
 
-	NumProxiesThatRequireDistanceFieldData	-= ComputeProxy->RequiresDistanceFieldData() ? 1 : 0;
-	NumProxiesThatRequireDepthBuffer		-= ComputeProxy->RequiresDepthBuffer() ? 1 : 0;
-	NumProxiesThatRequireEarlyViewData		-= ComputeProxy->RequiresEarlyViewData() ? 1 : 0;
-	NumProxiesThatRequireRayTracingScene	-= ComputeProxy->RequiresRayTracingScene() ? 1 : 0;
+	NumProxiesThatRequireGlobalDistanceField	-= ComputeProxy->RequiresGlobalDistanceField() ? 1 : 0;
+	NumProxiesThatRequireDepthBuffer			-= ComputeProxy->RequiresDepthBuffer() ? 1 : 0;
+	NumProxiesThatRequireEarlyViewData			-= ComputeProxy->RequiresEarlyViewData() ? 1 : 0;
+	NumProxiesThatRequireRayTracingScene		-= ComputeProxy->RequiresRayTracingScene() ? 1 : 0;
 
 #if NIAGARA_COMPUTEDEBUG_ENABLED
 	if (FNiagaraGpuComputeDebug* GpuComputeDebug = GpuComputeDebugPtr.Get())
@@ -2151,7 +2151,7 @@ void FNiagaraGpuComputeDispatch::ProcessDebugReadbacks(FRHICommandListImmediate&
 
 bool FNiagaraGpuComputeDispatch::UsesGlobalDistanceField() const
 {
-	return NumProxiesThatRequireDistanceFieldData > 0;
+	return NumProxiesThatRequireGlobalDistanceField > 0;
 }
 
 bool FNiagaraGpuComputeDispatch::UsesDepthBuffer() const
diff --git a/Engine/Plugins/FX/Niagara/Source/Niagara/Private/NiagaraGpuComputeDispatch.h b/Engine/Plugins/FX/Niagara/Source/Niagara/Private/NiagaraGpuComputeDispatch.h
--- a/Engine/Plugins/FX/Niagara/Source/Niagara/Private/NiagaraGpuComputeDispatch.h
+++ b/Engine/Plugins/FX/Niagara/Source/Niagara/Private/NiagaraGpuComputeDispatch.h
@@ -226,7 +226,7 @@ private:
 	FRWBuffer FreeIDListSizesBuffer;
 	uint32 NumAllocatedFreeIDListSizes = 0;
 
-	uint32 NumProxiesThatRequireDistanceFieldData = 0;
+	uint32 NumProxiesThatRequireGlobalDistanceField = 0;
 	uint32 NumProxiesThatRequireDepthBuffer = 0;
 	uint32 NumProxiesThatRequireEarlyViewData = 0;
 	uint32 NumProxiesThatRequireRayTracingScene = 0;
diff --git a/Engine/Plugins/FX/Niagara/Source/Niagara/Private/NiagaraSystemGpuComputeProxy.cpp b/Engine/Plugins/FX/Niagara/Source/Niagara/Private/NiagaraSystemGpuComputeProxy.cpp
--- a/Engine/Plugins/FX/Niagara/Source/Niagara/Private/NiagaraSystemGpuComputeProxy.cpp
+++ b/Engine/Plugins/FX/Niagara/Source/Niagara/Private/NiagaraSystemGpuComputeProxy.cpp
@@ -15,7 +15,7 @@ FNiagaraSystemGpuComputeProxy::FNiagaraSystemGpuComputeProxy(FNiagaraSystemInsta
 	, SystemLWCTile(OwnerInstance->GetLWCTile())
 	, SystemInstanceID(OwnerInstance->GetId())
 {
-	bRequiresDistanceFieldData = OwnerInstance->RequiresDistanceFieldData();
+	bRequiresGlobalDistanceField = OwnerInstance->RequiresGlobalDistanceField();
 	bRequiresDepthBuffer = OwnerInstance->RequiresDepthBuffer();
 	bRequiresEarlyViewData = OwnerInstance->RequiresEarlyViewData();
 	bRequiresViewUniformBuffer = OwnerInstance->RequiresViewUniformBuffer();
@@ -31,7 +31,7 @@ FNiagaraSystemGpuComputeProxy::FNiagaraSystemGpuComputeProxy(FNiagaraSystemInsta
 	}
 
 	// Calculate Tick Stage
-	if (bRequiresDistanceFieldData || bRequiresDepthBuffer || bRequiresRayTracingScene)
+	if (bRequiresGlobalDistanceField || bRequiresDepthBuffer || bRequiresRayTracingScene)
 	{
 		ComputeTickStage = ENiagaraGpuComputeTickStage::PostOpaqueRender;
 	}
diff --git a/Engine/Plugins/FX/Niagara/Source/Niagara/Private/NiagaraSystemInstance.cpp b/Engine/Plugins/FX/Niagara/Source/Niagara/Private/NiagaraSystemInstance.cpp
--- a/Engine/Plugins/FX/Niagara/Source/Niagara/Private/NiagaraSystemInstance.cpp
+++ b/Engine/Plugins/FX/Niagara/Source/Niagara/Private/NiagaraSystemInstance.cpp
@@ -1223,7 +1223,7 @@ FNiagaraWorldManager* FNiagaraSystemInstance::GetWorldManager() const
 	return FNiagaraWorldManager::Get(World);
 }
 
-bool FNiagaraSystemInstance::RequiresDistanceFieldData() const
+bool FNiagaraSystemInstance::RequiresGlobalDistanceField() const
 {
 	if (!bHasGPUEmitters)
 	{
@@ -1237,7 +1237,7 @@ bool FNiagaraSystemInstance::RequiresDistanceFieldData() const
 		{
 			for (UNiagaraDataInterface* DataInterface : GPUContext->CombinedParamStore.GetDataInterfaces())
 			{
-				if (DataInterface && DataInterface->RequiresDistanceFieldData())
+				if (DataInterface && DataInterface->RequiresGlobalDistanceField())
 				{
 					return true;
 				}
diff --git a/Engine/Plugins/FX/Niagara/Source/Niagara/Public/NiagaraDataInterfaceRigidMeshCollisionQuery.h b/Engine/Plugins/FX/Niagara/Source/Niagara/Public/NiagaraDataInterfaceRigidMeshCollisionQuery.h
--- a/Engine/Plugins/FX/Niagara/Source/Niagara/Public/NiagaraDataInterfaceRigidMeshCollisionQuery.h
+++ b/Engine/Plugins/FX/Niagara/Source/Niagara/Public/NiagaraDataInterfaceRigidMeshCollisionQuery.h
@@ -242,7 +242,7 @@ public:
 	virtual bool HasTickGroupPrereqs() const override { return true; }
 	virtual ETickingGroup CalculateTickGroup(const void* PerInstanceData) const override;
 
-	virtual bool RequiresDistanceFieldData() const override { return true; }
+	virtual bool RequiresEarlyViewData() const override { return true; }
 
 	/** GPU simulation  functionality */
 #if WITH_EDITORONLY_DATA
diff --git a/Engine/Plugins/FX/Niagara/Source/Niagara/Public/NiagaraDebuggerCommon.h b/Engine/Plugins/FX/Niagara/Source/Niagara/Public/NiagaraDebuggerCommon.h
--- a/Engine/Plugins/FX/Niagara/Source/Niagara/Public/NiagaraDebuggerCommon.h
+++ b/Engine/Plugins/FX/Niagara/Source/Niagara/Public/NiagaraDebuggerCommon.h
@@ -108,7 +108,7 @@ struct FNiagaraOutlinerSystemInstanceData
 	uint32 bIsSolo : 1;
 
 	UPROPERTY(VisibleAnywhere, Category = "Gpu")
-	uint32 bRequiresDistanceFieldData : 1;
+	uint32 bRequiresGlobalDistanceField : 1;
 
 	UPROPERTY(VisibleAnywhere, Category = "Gpu")
 	uint32 bRequiresDepthBuffer : 1;
@@ -128,7 +128,7 @@ struct FNiagaraOutlinerSystemInstanceData
 		, TickGroup(0)
 		, GpuTickStage(ENiagaraGpuComputeTickStage::First)
 		, bIsSolo(false)
-		, bRequiresDistanceFieldData(false)
+		, bRequiresGlobalDistanceField(false)
 		, bRequiresDepthBuffer(false)
 		, bRequiresEarlyViewData(false)
 		, bRequiresViewUniformBuffer(false)
diff --git a/Engine/Plugins/FX/Niagara/Source/Niagara/Public/NiagaraSystemGpuComputeProxy.h b/Engine/Plugins/FX/Niagara/Source/Niagara/Public/NiagaraSystemGpuComputeProxy.h
--- a/Engine/Plugins/FX/Niagara/Source/Niagara/Public/NiagaraSystemGpuComputeProxy.h
+++ b/Engine/Plugins/FX/Niagara/Source/Niagara/Public/NiagaraSystemGpuComputeProxy.h
@@ -28,7 +28,7 @@ public:
 	void QueueTick(const FNiagaraGPUSystemTick& Tick);
 	void ReleaseTicks(FNiagaraGPUInstanceCountManager& GPUInstanceCountManager, int32 NumTicksToRelease, bool bLastViewFamily);
 
-	bool RequiresDistanceFieldData() const { return bRequiresDistanceFieldData; }
+	bool RequiresGlobalDistanceField() const { return bRequiresGlobalDistanceField; }
 	bool RequiresDepthBuffer() const { return bRequiresDepthBuffer; }
 	bool RequiresEarlyViewData() const { return bRequiresEarlyViewData; }
 	bool RequiresViewUniformBuffer() const { return bRequiresViewUniformBuffer; }
@@ -43,7 +43,7 @@ private:
 
 	FNiagaraSystemInstanceID					SystemInstanceID = FNiagaraSystemInstanceID();
 	ENiagaraGpuComputeTickStage::Type			ComputeTickStage = ENiagaraGpuComputeTickStage::PostOpaqueRender;
-	uint32										bRequiresDistanceFieldData : 1;
+	uint32										bRequiresGlobalDistanceField : 1;
 	uint32										bRequiresDepthBuffer : 1;
 	uint32										bRequiresEarlyViewData : 1;
 	uint32										bRequiresViewUniformBuffer : 1;
diff --git a/Engine/Plugins/FX/Niagara/Source/Niagara/Public/NiagaraSystemInstance.h b/Engine/Plugins/FX/Niagara/Source/Niagara/Public/NiagaraSystemInstance.h
--- a/Engine/Plugins/FX/Niagara/Source/Niagara/Public/NiagaraSystemInstance.h
+++ b/Engine/Plugins/FX/Niagara/Source/Niagara/Public/NiagaraSystemInstance.h
@@ -162,7 +162,7 @@ public:
 	FORCEINLINE FNiagaraEmitterParameters& EditEmitterParameters(int32 EmitterIdx) { return EmitterParameters[EmitterIdx * 2 + GetParameterIndex()]; }
 	
 	NIAGARA_API FNiagaraWorldManager* GetWorldManager()const;
-	NIAGARA_API bool RequiresDistanceFieldData() const;
+	NIAGARA_API bool RequiresGlobalDistanceField() const;
 	NIAGARA_API bool RequiresDepthBuffer() const;
 	NIAGARA_API bool RequiresEarlyViewData() const;
 	NIAGARA_API bool RequiresViewUniformBuffer() const;
diff --git a/Engine/Plugins/FX/Niagara/Source/NiagaraEditor/Private/Widgets/SNiagaraSystemViewport.cpp b/Engine/Plugins/FX/Niagara/Source/NiagaraEditor/Private/Widgets/SNiagaraSystemViewport.cpp
--- a/Engine/Plugins/FX/Niagara/Source/NiagaraEditor/Private/Widgets/SNiagaraSystemViewport.cpp
+++ b/Engine/Plugins/FX/Niagara/Source/NiagaraEditor/Private/Widgets/SNiagaraSystemViewport.cpp
@@ -435,9 +435,9 @@ void FNiagaraSystemViewportClient::DrawGpuTickInformation(UNiagaraComponent* Com
 			Canvas->DrawShadowedString(CurrentX + 5.0f, CurrentY, TEXT("No GPU Emitters"), Font, FLinearColor::White);
 			CurrentY += FontHeight;
 		}
-		if (SystemInstance->RequiresDistanceFieldData())
+		if (SystemInstance->RequiresGlobalDistanceField())
 		{
-			Canvas->DrawShadowedString(CurrentX + 5.0f, CurrentY, TEXT("RequiresDistanceFieldData"), Font, FLinearColor::White);
+			Canvas->DrawShadowedString(CurrentX + 5.0f, CurrentY, TEXT("RequiresGlobalDistanceField"), Font, FLinearColor::White);
 			CurrentY += FontHeight;
 		}
 		if (SystemInstance->RequiresDepthBuffer())
diff --git a/Engine/Plugins/FX/Niagara/Source/NiagaraShader/Private/NiagaraAsyncGpuTraceProvider.cpp b/Engine/Plugins/FX/Niagara/Source/NiagaraShader/Private/NiagaraAsyncGpuTraceProvider.cpp
--- a/Engine/Plugins/FX/Niagara/Source/NiagaraShader/Private/NiagaraAsyncGpuTraceProvider.cpp
+++ b/Engine/Plugins/FX/Niagara/Source/NiagaraShader/Private/NiagaraAsyncGpuTraceProvider.cpp
@@ -130,7 +130,7 @@ bool CheckProviderIsRequestedAndSupported(FNiagaraAsyncGpuTraceProvider::EProvid
 	return false;
 }
 
-bool FNiagaraAsyncGpuTraceProvider::RequiresDistanceFieldData(EProviderType InType, const FProviderPriorityArray& Priorities)
+bool FNiagaraAsyncGpuTraceProvider::RequiresGlobalDistanceField(EProviderType InType, const FProviderPriorityArray& Priorities)
 {
 	return CheckProviderIsRequestedAndSupported<FNiagaraAsyncGpuTraceProviderGsdf>(InType, Priorities);
 }
diff --git a/Engine/Plugins/FX/Niagara/Source/NiagaraShader/Public/NiagaraAsyncGpuTraceProvider.h b/Engine/Plugins/FX/Niagara/Source/NiagaraShader/Public/NiagaraAsyncGpuTraceProvider.h
--- a/Engine/Plugins/FX/Niagara/Source/NiagaraShader/Public/NiagaraAsyncGpuTraceProvider.h
+++ b/Engine/Plugins/FX/Niagara/Source/NiagaraShader/Public/NiagaraAsyncGpuTraceProvider.h
@@ -59,7 +59,7 @@ public:
 	};
 
 	static NIAGARASHADER_API EProviderType ResolveSupportedType(EProviderType InType, const FProviderPriorityArray& Priorities);
-	static NIAGARASHADER_API bool RequiresDistanceFieldData(EProviderType InType, const FProviderPriorityArray& Priorities);
+	static NIAGARASHADER_API bool RequiresGlobalDistanceField(EProviderType InType, const FProviderPriorityArray& Priorities);
 	static NIAGARASHADER_API bool RequiresRayTracingScene(EProviderType InType, const FProviderPriorityArray& Priorities);
 
 	/** Hash table.
